Added system to attach additional data to animation notifies

- Added animation notify mirroring functionality
- Extended IGraphMessage to allow users to create context data that is attached to notifies
- Notify, NotifyBegin, NotifyTick, and NotifyEnd methods have been extedned to take an FAnimNotifyEventReference that can provide generic data
- StateMachines can optionally attach active state information to notifies
- Added function libraries to query if a notify was triggered by a particular state or state machine
- Added support for detecting if a notify was generated by a mirrored animation

#github https://github.com/EpicGames/UnrealEngine/pull/7204 Fredrik Lindh

diff --git a/Engine/Plugins/FX/Niagara/Source/NiagaraAnimNotifies/Private/AnimNotifyState_TimedNiagaraEffect.cpp b/Engine/Plugins/FX/Niagara/Source/NiagaraAnimNotifies/Private/AnimNotifyState_TimedNiagaraEffect.cpp
--- a/Engine/Plugins/FX/Niagara/Source/NiagaraAnimNotifies/Private/AnimNotifyState_TimedNiagaraEffect.cpp
+++ b/Engine/Plugins/FX/Niagara/Source/NiagaraAnimNotifies/Private/AnimNotifyState_TimedNiagaraEffect.cpp
@@ -53,6 +53,10 @@ UFXSystemComponent* UAnimNotifyState_TimedNiagaraEffect::GetSpawnedEffect(UMeshC
 }
 
 void UAnimNotifyState_TimedNiagaraEffect::NotifyBegin(USkeletalMeshComponent * MeshComp, class UAnimSequenceBase * Animation, float TotalDuration)
+{
+}
+
+void UAnimNotifyState_TimedNiagaraEffect::NotifyBegin(USkeletalMeshComponent * MeshComp, class UAnimSequenceBase * Animation, float TotalDuration, const FAnimNotifyEventReference& EventReference)
 {
 	if (UFXSystemComponent* Component = SpawnEffect(MeshComp, Animation))
 	{
@@ -60,10 +64,14 @@ void UAnimNotifyState_TimedNiagaraEffect::NotifyBegin(USkeletalMeshComponent * M
 		Component->ComponentTags.AddUnique(GetSpawnedComponentTag());
 	}
 
-	Super::NotifyBegin(MeshComp, Animation, TotalDuration);
+	Super::NotifyBegin(MeshComp, Animation, TotalDuration, EventReference);
 }
 
 void UAnimNotifyState_TimedNiagaraEffect::NotifyEnd(USkeletalMeshComponent * MeshComp, class UAnimSequenceBase * Animation)
+{
+}
+
+void UAnimNotifyState_TimedNiagaraEffect::NotifyEnd(USkeletalMeshComponent * MeshComp, class UAnimSequenceBase * Animation, const FAnimNotifyEventReference& EventReference)
 {
 	if (UFXSystemComponent* FXComponent = GetSpawnedEffect(MeshComp))
 	{
@@ -82,7 +90,7 @@ void UAnimNotifyState_TimedNiagaraEffect::NotifyEnd(USkeletalMeshComponent * Mes
 		}
 	}
 
-	Super::NotifyEnd(MeshComp, Animation);
+	Super::NotifyEnd(MeshComp, Animation, EventReference);
 }
 
 bool UAnimNotifyState_TimedNiagaraEffect::ValidateParameters(USkeletalMeshComponent* MeshComp) const
@@ -120,7 +128,11 @@ UAnimNotifyState_TimedNiagaraEffectAdvanced::UAnimNotifyState_TimedNiagaraEffect
 
 void UAnimNotifyState_TimedNiagaraEffectAdvanced::NotifyBegin(USkeletalMeshComponent* MeshComp, class UAnimSequenceBase* Animation, float TotalDuration)
 {
-	Super::NotifyBegin(MeshComp, Animation, TotalDuration);
+}
+
+void UAnimNotifyState_TimedNiagaraEffectAdvanced::NotifyBegin(USkeletalMeshComponent* MeshComp, class UAnimSequenceBase* Animation, float TotalDuration, const FAnimNotifyEventReference& EventReference)
+{
+	Super::NotifyBegin(MeshComp, Animation, TotalDuration, EventReference);
 
 	FInstanceProgressInfo& NewInfo = ProgressInfoMap.Add(MeshComp);
 	NewInfo.Duration = TotalDuration;
@@ -129,13 +141,21 @@ void UAnimNotifyState_TimedNiagaraEffectAdvanced::NotifyBegin(USkeletalMeshCompo
 
 void UAnimNotifyState_TimedNiagaraEffectAdvanced::NotifyEnd(USkeletalMeshComponent* MeshComp, class UAnimSequenceBase* Animation)
 {
-	Super::NotifyEnd(MeshComp, Animation);
+}
+
+void UAnimNotifyState_TimedNiagaraEffectAdvanced::NotifyEnd(USkeletalMeshComponent* MeshComp, class UAnimSequenceBase* Animation, const FAnimNotifyEventReference& EventReference)
+{
+	Super::NotifyEnd(MeshComp, Animation, EventReference);
 	ProgressInfoMap.Remove(MeshComp);
 }
 
 void UAnimNotifyState_TimedNiagaraEffectAdvanced::NotifyTick(USkeletalMeshComponent* MeshComp, UAnimSequenceBase* Animation, float FrameDeltaTime)
 {
-	Super::NotifyTick(MeshComp, Animation, FrameDeltaTime);
+}
+
+void UAnimNotifyState_TimedNiagaraEffectAdvanced::NotifyTick(USkeletalMeshComponent* MeshComp, UAnimSequenceBase* Animation, float FrameDeltaTime, const FAnimNotifyEventReference& EventReference)
+{
+	Super::NotifyTick(MeshComp, Animation, FrameDeltaTime, EventReference);
 
 	//Advance the progress.
 	//TODO: There must be some way to avoid this map and lookup. The information about the current elapsed time and a mapping onto the notify range should be available somewhere in the mesh comp and notify.
diff --git a/Engine/Plugins/FX/Niagara/Source/NiagaraAnimNotifies/Private/AnimNotify_PlayNiagaraEffect.cpp b/Engine/Plugins/FX/Niagara/Source/NiagaraAnimNotifies/Private/AnimNotify_PlayNiagaraEffect.cpp
--- a/Engine/Plugins/FX/Niagara/Source/NiagaraAnimNotifies/Private/AnimNotify_PlayNiagaraEffect.cpp
+++ b/Engine/Plugins/FX/Niagara/Source/NiagaraAnimNotifies/Private/AnimNotify_PlayNiagaraEffect.cpp
@@ -75,14 +75,16 @@ void UAnimNotify_PlayNiagaraEffect::ValidateAssociatedAssets()
 #endif
 
 void UAnimNotify_PlayNiagaraEffect::Notify(class USkeletalMeshComponent* MeshComp, class UAnimSequenceBase* Animation)
+{
+}
+
+void UAnimNotify_PlayNiagaraEffect::Notify(class USkeletalMeshComponent* MeshComp, class UAnimSequenceBase* Animation, const FAnimNotifyEventReference& EventReference)
 {
 	//Store the spawned effect in a protected variable
 	SpawnedEffect = SpawnEffect(MeshComp, Animation);
 	
 	//Call to BP to allows setting of Niagara User Variables
-	Super::Notify(MeshComp, Animation);
-	
-	
+	Super::Notify(MeshComp, Animation, EventReference);
 }
 
 FString UAnimNotify_PlayNiagaraEffect::GetNotifyName_Implementation() const
diff --git a/Engine/Plugins/FX/Niagara/Source/NiagaraAnimNotifies/Public/AnimNotifyState_TimedNiagaraEffect.h b/Engine/Plugins/FX/Niagara/Source/NiagaraAnimNotifies/Public/AnimNotifyState_TimedNiagaraEffect.h
--- a/Engine/Plugins/FX/Niagara/Source/NiagaraAnimNotifies/Public/AnimNotifyState_TimedNiagaraEffect.h
+++ b/Engine/Plugins/FX/Niagara/Source/NiagaraAnimNotifies/Public/AnimNotifyState_TimedNiagaraEffect.h
@@ -41,8 +41,14 @@ class NIAGARAANIMNOTIFIES_API UAnimNotifyState_TimedNiagaraEffect : public UAnim
 	UPROPERTY(EditAnywhere, Category = NiagaraSystem, meta = (DisplayName = "Destroy Immediately", ToolTip = "Whether the Niagara system should be immediately destroyed at the end of the notify state or be allowed to finish"))
 	bool bDestroyAtEnd;
 
+	UE_DEPRECATED(5.0, "Please use the other NotifyBegin function instead")
 	virtual void NotifyBegin(class USkeletalMeshComponent * MeshComp, class UAnimSequenceBase * Animation, float TotalDuration) override;
+
+	UE_DEPRECATED(5.0, "Please use the other NotifyEnd function instead")
 	virtual void NotifyEnd(class USkeletalMeshComponent * MeshComp, class UAnimSequenceBase * Animation) override;
+	
+	virtual void NotifyBegin(class USkeletalMeshComponent * MeshComp, class UAnimSequenceBase * Animation, float TotalDuration, const FAnimNotifyEventReference& EventReference) override;
+	virtual void NotifyEnd(class USkeletalMeshComponent * MeshComp, class UAnimSequenceBase * Animation, const FAnimNotifyEventReference& EventReference) override;
 
 	// Overridden from UAnimNotifyState to provide custom notify name.
 	FString GetNotifyName_Implementation() const override;
@@ -69,9 +75,18 @@ class NIAGARAANIMNOTIFIES_API UAnimNotifyState_TimedNiagaraEffectAdvanced : publ
 	GENERATED_UCLASS_BODY()
 
 public:
+	UE_DEPRECATED(5.0, "Please use the other NotifyBegin function instead")
 	virtual void NotifyBegin(class USkeletalMeshComponent* MeshComp, class UAnimSequenceBase* Animation, float TotalDuration) override;
+
+	UE_DEPRECATED(5.0, "Please use the other NotifyEnd function instead")
 	virtual void NotifyEnd(class USkeletalMeshComponent* MeshComp, class UAnimSequenceBase* Animation) override;
+
+	UE_DEPRECATED(5.0, "Please use the other NotifyTick function instead")
 	virtual void NotifyTick(USkeletalMeshComponent * MeshComp, UAnimSequenceBase * Animation, float FrameDeltaTime) override;
+	
+	virtual void NotifyBegin(class USkeletalMeshComponent* MeshComp, class UAnimSequenceBase* Animation, float TotalDuration, const FAnimNotifyEventReference& EventReference) override;
+	virtual void NotifyEnd(class USkeletalMeshComponent* MeshComp, class UAnimSequenceBase* Animation, const FAnimNotifyEventReference& EventReference) override;
+	virtual void NotifyTick(USkeletalMeshComponent * MeshComp, UAnimSequenceBase * Animation, float FrameDeltaTime, const FAnimNotifyEventReference& EventReference) override;
 
 	// Returns a 0 to 1 value for the progress of this component along the notify.
 	UFUNCTION(BlueprintCallable, Category = "AnimNotify")
diff --git a/Engine/Plugins/FX/Niagara/Source/NiagaraAnimNotifies/Public/AnimNotify_PlayNiagaraEffect.h b/Engine/Plugins/FX/Niagara/Source/NiagaraAnimNotifies/Public/AnimNotify_PlayNiagaraEffect.h
--- a/Engine/Plugins/FX/Niagara/Source/NiagaraAnimNotifies/Public/AnimNotify_PlayNiagaraEffect.h
+++ b/Engine/Plugins/FX/Niagara/Source/NiagaraAnimNotifies/Public/AnimNotify_PlayNiagaraEffect.h
@@ -35,7 +35,10 @@ public:
 
 	// Begin UAnimNotify interface
 	virtual FString GetNotifyName_Implementation() const override;
+	
+	UE_DEPRECATED(5.0, "Please use the other Notify function instead")
 	virtual void Notify(USkeletalMeshComponent* MeshComp, UAnimSequenceBase* Animation) override;
+	virtual void Notify(USkeletalMeshComponent* MeshComp, UAnimSequenceBase* Animation, const FAnimNotifyEventReference& EventReference) override;
 #if WITH_EDITOR
 	virtual void ValidateAssociatedAssets() override;
 #endif
diff --git a/Engine/Plugins/Runtime/GameplayAbilities/Source/GameplayAbilities/Private/AnimNotify_GameplayCue.cpp b/Engine/Plugins/Runtime/GameplayAbilities/Source/GameplayAbilities/Private/AnimNotify_GameplayCue.cpp
--- a/Engine/Plugins/Runtime/GameplayAbilities/Source/GameplayAbilities/Private/AnimNotify_GameplayCue.cpp
+++ b/Engine/Plugins/Runtime/GameplayAbilities/Source/GameplayAbilities/Private/AnimNotify_GameplayCue.cpp
@@ -74,7 +74,11 @@ UAnimNotify_GameplayCue::UAnimNotify_GameplayCue()
 
 void UAnimNotify_GameplayCue::Notify(USkeletalMeshComponent* MeshComp, UAnimSequenceBase* Animation)
 {
-	Super::Notify(MeshComp, Animation);
+}
+
+void UAnimNotify_GameplayCue::Notify(USkeletalMeshComponent* MeshComp, UAnimSequenceBase* Animation, const FAnimNotifyEventReference& EventReference)
+{
+	Super::Notify(MeshComp, Animation, EventReference);
 
 	ProcessGameplayCue(&UGameplayCueManager::ExecuteGameplayCue_NonReplicated, MeshComp, GameplayCue.GameplayCueTag, Animation);
 }
@@ -109,7 +113,11 @@ UAnimNotify_GameplayCueState::UAnimNotify_GameplayCueState()
 
 void UAnimNotify_GameplayCueState::NotifyBegin(USkeletalMeshComponent* MeshComp, UAnimSequenceBase* Animation, float TotalDuration)
 {
-	Super::NotifyBegin(MeshComp, Animation, TotalDuration);
+}
+
+void UAnimNotify_GameplayCueState::NotifyBegin(USkeletalMeshComponent* MeshComp, UAnimSequenceBase* Animation, float TotalDuration, const FAnimNotifyEventReference& EventReference)
+{
+	Super::NotifyBegin(MeshComp, Animation, TotalDuration, EventReference);
 
 	ProcessGameplayCue(&UGameplayCueManager::AddGameplayCue_NonReplicated, MeshComp, GameplayCue.GameplayCueTag, Animation);
 
@@ -124,6 +132,10 @@ void UAnimNotify_GameplayCueState::NotifyBegin(USkeletalMeshComponent* MeshComp,
 }
 
 void UAnimNotify_GameplayCueState::NotifyTick(USkeletalMeshComponent * MeshComp, UAnimSequenceBase * Animation, float FrameDeltaTime)
+{
+}
+
+void UAnimNotify_GameplayCueState::NotifyTick(USkeletalMeshComponent * MeshComp, UAnimSequenceBase * Animation, float FrameDeltaTime, const FAnimNotifyEventReference& EventReference)
 {
 	AActor* OwnerActor = MeshComp->GetOwner();
 
@@ -140,7 +152,7 @@ void UAnimNotify_GameplayCueState::NotifyTick(USkeletalMeshComponent * MeshComp,
 	}
 #endif // #if WITH_EDITOR
 
-	Super::NotifyTick(MeshComp, Animation, FrameDeltaTime);
+	Super::NotifyTick(MeshComp, Animation, FrameDeltaTime, EventReference);
 
 #if WITH_EDITOR
 	if (GIsEditor)
@@ -166,7 +178,11 @@ FString UAnimNotify_GameplayCueState::GetNotifyName_Implementation() const
 
 void UAnimNotify_GameplayCueState::NotifyEnd(USkeletalMeshComponent* MeshComp, UAnimSequenceBase* Animation)
 {
-	Super::NotifyEnd(MeshComp, Animation);
+}
+
+void UAnimNotify_GameplayCueState::NotifyEnd(USkeletalMeshComponent* MeshComp, UAnimSequenceBase* Animation, const FAnimNotifyEventReference& EventReference)
+{
+	Super::NotifyEnd(MeshComp, Animation, EventReference);
 
 	ProcessGameplayCue(&UGameplayCueManager::RemoveGameplayCue_NonReplicated, MeshComp, GameplayCue.GameplayCueTag, Animation);
 }
diff --git a/Engine/Plugins/Runtime/GameplayAbilities/Source/GameplayAbilities/Public/AnimNotify_GameplayCue.h b/Engine/Plugins/Runtime/GameplayAbilities/Source/GameplayAbilities/Public/AnimNotify_GameplayCue.h
--- a/Engine/Plugins/Runtime/GameplayAbilities/Source/GameplayAbilities/Public/AnimNotify_GameplayCue.h
+++ b/Engine/Plugins/Runtime/GameplayAbilities/Source/GameplayAbilities/Public/AnimNotify_GameplayCue.h
@@ -24,8 +24,11 @@ public:
 
 	UAnimNotify_GameplayCue();
 
+	UE_DEPRECATED(5.0, "Please use the other Notify function instead")
 	virtual void Notify(USkeletalMeshComponent* MeshComp, UAnimSequenceBase* Animation) override;
 
+	virtual void Notify(USkeletalMeshComponent* MeshComp, UAnimSequenceBase* Animation, const FAnimNotifyEventReference& EventReference) override;
+
 	FString GetNotifyName_Implementation() const override;
 
 #if WITH_EDITOR
@@ -53,11 +56,19 @@ class UAnimNotify_GameplayCueState : public UAnimNotifyState
 public:
 
 	UAnimNotify_GameplayCueState();
-
+	UE_DEPRECATED(5.0, "Please use the other NotifyBegin function instead")
 	virtual void NotifyBegin(USkeletalMeshComponent* MeshComp, UAnimSequenceBase* Animation, float TotalDuration) override;
+
+	UE_DEPRECATED(5.0, "Please use the other NotifyTick function instead")
 	virtual void NotifyTick(USkeletalMeshComponent* MeshComp, UAnimSequenceBase* Animation, float FrameDeltaTime) override;
+
+	UE_DEPRECATED(5.0, "Please use the other NotifyEnd function instead")
 	virtual void NotifyEnd(USkeletalMeshComponent* MeshComp, UAnimSequenceBase* Animation) override;
 
+	virtual void NotifyBegin(USkeletalMeshComponent* MeshComp, UAnimSequenceBase* Animation, float TotalDuration, const FAnimNotifyEventReference& EventReference) override;
+	virtual void NotifyTick(USkeletalMeshComponent* MeshComp, UAnimSequenceBase* Animation, float FrameDeltaTime, const FAnimNotifyEventReference& EventReference) override;
+	virtual void NotifyEnd(USkeletalMeshComponent* MeshComp, UAnimSequenceBase* Animation, const FAnimNotifyEventReference& EventReference) override;
+
 	FString GetNotifyName_Implementation() const override;
 
 #if WITH_EDITOR
diff --git a/Engine/Source/Editor/AnimationBlueprintEditor/Private/AnimationBlueprintEditor.cpp b/Engine/Source/Editor/AnimationBlueprintEditor/Private/AnimationBlueprintEditor.cpp
--- a/Engine/Source/Editor/AnimationBlueprintEditor/Private/AnimationBlueprintEditor.cpp
+++ b/Engine/Source/Editor/AnimationBlueprintEditor/Private/AnimationBlueprintEditor.cpp
@@ -1579,9 +1579,10 @@ void FAnimationBlueprintEditor::OnBlueprintPreCompile(UBlueprint* BlueprintToCom
 			for(int32 Idx = Instance->ActiveAnimNotifyState.Num() - 1 ; Idx >= 0 ; --Idx)
 			{
 				FAnimNotifyEvent& Event = Instance->ActiveAnimNotifyState[Idx];
+				const FAnimNotifyEventReference& EventReference = Instance->ActiveAnimNotifyEventReference[Idx];
 				if(Event.NotifyStateClass->GetClass() == BlueprintToCompile->GeneratedClass)
 				{
-					Event.NotifyStateClass->NotifyEnd(SkelMeshComp, Cast<UAnimSequenceBase>(Event.NotifyStateClass->GetOuter()));
+					Event.NotifyStateClass->NotifyEnd(SkelMeshComp, Cast<UAnimSequenceBase>(Event.NotifyStateClass->GetOuter()), EventReference);
 					Instance->ActiveAnimNotifyState.RemoveAt(Idx);
 				}
 			}
diff --git a/Engine/Source/Runtime/AnimGraphRuntime/Private/AnimNodes/AnimNode_BlendSpacePlayer.cpp b/Engine/Source/Runtime/AnimGraphRuntime/Private/AnimNodes/AnimNode_BlendSpacePlayer.cpp
--- a/Engine/Source/Runtime/AnimGraphRuntime/Private/AnimNodes/AnimNode_BlendSpacePlayer.cpp
+++ b/Engine/Source/Runtime/AnimGraphRuntime/Private/AnimNodes/AnimNode_BlendSpacePlayer.cpp
@@ -88,6 +88,10 @@ void FAnimNode_BlendSpacePlayer::UpdateInternal(const FAnimationUpdateContext& C
 		FAnimTickRecord TickRecord(CurrentBlendSpace, Position, BlendSampleDataCache, BlendFilter, GetLoop(), GetPlayRate(), Context.GetFinalBlendWeight(), /*inout*/ InternalTimeAccumulator, MarkerTickRecord);
 		TickRecord.RootMotionWeightModifier = Context.GetRootMotionWeightModifier();
 		UE::Anim::FAnimSyncParams SyncParams(GetGroupName(), GetGroupRole(), GetGroupMethod());
+		if(Context.GetSharedContext())
+		{
+			Context.GetSharedContext()->MessageStack.MakeEventContextData(TickRecord.ContextData);
+		}
 
 		SyncScope.AddTickRecord(TickRecord, SyncParams, UE::Anim::FAnimSyncDebugInfo(Context));
 
diff --git a/Engine/Source/Runtime/AnimGraphRuntime/Private/AnimNodes/AnimNode_RandomPlayer.cpp b/Engine/Source/Runtime/AnimGraphRuntime/Private/AnimNodes/AnimNode_RandomPlayer.cpp
--- a/Engine/Source/Runtime/AnimGraphRuntime/Private/AnimNodes/AnimNode_RandomPlayer.cpp
+++ b/Engine/Source/Runtime/AnimGraphRuntime/Private/AnimNodes/AnimNode_RandomPlayer.cpp
@@ -231,7 +231,11 @@ void FAnimNode_RandomPlayer::Update_AnyThread(const FAnimationUpdateContext& Con
 	}
 
 	FAnimTickRecord TickRecord(CurrentData->Entry->Sequence, true, CurrentData->PlayRate, CurrentData->BlendWeight, CurrentData->CurrentPlayTime, CurrentData->MarkerTickRecord);
-
+	
+	if(Context.GetSharedContext())
+	{
+		Context.GetSharedContext()->MessageStack.MakeEventContextData(TickRecord.ContextData);
+	}
 	UE::Anim::FAnimSyncGroupScope& SyncScope = Context.GetMessageChecked<UE::Anim::FAnimSyncGroupScope>();
 	SyncScope.AddTickRecord(TickRecord, UE::Anim::FAnimSyncParams(), UE::Anim::FAnimSyncDebugInfo(Context));
 
@@ -240,6 +244,11 @@ void FAnimNode_RandomPlayer::Update_AnyThread(const FAnimationUpdateContext& Con
 	if (FAnimationRuntime::HasWeight(NextData->BlendWeight))
 	{
 		FAnimTickRecord NextTickRecord(NextData->Entry->Sequence, true, NextData->PlayRate, NextData->BlendWeight, NextData->CurrentPlayTime, NextData->MarkerTickRecord);
+
+		if(Context.GetSharedContext())
+		{
+			Context.GetSharedContext()->MessageStack.MakeEventContextData(NextTickRecord.ContextData);
+		}
 		SyncScope.AddTickRecord(NextTickRecord, UE::Anim::FAnimSyncParams(), UE::Anim::FAnimSyncDebugInfo(Context));
 
 		TRACE_ANIM_TICK_RECORD(Context, NextTickRecord);
diff --git a/Engine/Source/Runtime/Engine/Classes/Animation/AnimComposite.h b/Engine/Source/Runtime/Engine/Classes/Animation/AnimComposite.h
--- a/Engine/Source/Runtime/Engine/Classes/Animation/AnimComposite.h
+++ b/Engine/Source/Runtime/Engine/Classes/Animation/AnimComposite.h
@@ -40,7 +40,7 @@ public:
 	virtual bool IsValidAdditive() const override { return GetAdditiveAnimType() != AAT_None; }
 	virtual void EnableRootMotionSettingFromMontage(bool bInEnableRootMotion, const ERootMotionRootLock::Type InRootMotionRootLock) override;
 	virtual bool HasRootMotion() const override;
-	virtual void GetAnimNotifiesFromDeltaPositions(const float& PreviousPosition, const float & CurrentPosition, TArray<FAnimNotifyEventReference>& OutActiveNotifies) const override;
+	virtual void GetAnimNotifiesFromDeltaPositions(const float& PreviousPosition, const float & CurrentPosition, FAnimNotifyContext& NotifyContext) const override;
 	virtual bool IsNotifyAvailable() const override;
 	//~ End UAnimSequenceBase Interface
 	//~ Begin UAnimSequence Interface
diff --git a/Engine/Source/Runtime/Engine/Classes/Animation/AnimCompositeBase.h b/Engine/Source/Runtime/Engine/Classes/Animation/AnimCompositeBase.h
--- a/Engine/Source/Runtime/Engine/Classes/Animation/AnimCompositeBase.h
+++ b/Engine/Source/Runtime/Engine/Classes/Animation/AnimCompositeBase.h
@@ -156,7 +156,7 @@ struct FAnimSegment
 	 * Supports playing backwards (CurrentTrackPosition<PreviousTrackPosition).
 	 * Only supports contiguous range, does NOT support looping and wrapping over.
 	 */
-	UE_DEPRECATED(4.19, "Use the GetAnimNotifiesFromTrackPositions that takes FAnimNotifyEventReferences instead")
+	UE_DEPRECATED(4.19, "Use the GetAnimNotifiesFromTrackPositions that takes FAnimNotifyContext instead")
 	void GetAnimNotifiesFromTrackPositions(const float& PreviousTrackPosition, const float& CurrentTrackPosition, TArray<const FAnimNotifyEvent *> & OutActiveNotifies) const;
 	
 	/**
@@ -165,8 +165,11 @@ struct FAnimSegment
 	* Supports playing backwards (CurrentTrackPosition<PreviousTrackPosition).
 	* Only supports contiguous range, does NOT support looping and wrapping over.
 	*/
+	UE_DEPRECATED(5.0, "Use the GetAnimNotifiesFromTrackPositions that takes FAnimNotifyContext instead")
 	ENGINE_API void GetAnimNotifiesFromTrackPositions(const float& PreviousTrackPosition, const float& CurrentTrackPosition, TArray<FAnimNotifyEventReference> & OutActiveNotifies) const;
 
+	ENGINE_API void GetAnimNotifiesFromTrackPositions(const float& PreviousTrackPosition, const float& CurrentTrackPosition, FAnimNotifyContext& NotifyContext) const;
+
 	/** 
 	 * Given a Track delta position [StartTrackPosition, EndTrackPosition]
 	 * See if this AnimSegment overlaps any of it, and if it does, break it up into RootMotionExtractionSteps.
@@ -284,8 +287,17 @@ struct FAnimTrack
 	* Supports playing backwards (CurrentTrackPosition<PreviousTrackPosition).
 	* Only supports contiguous range, does NOT support looping and wrapping over.
 	*/
+	UE_DEPRECATED(5.0, "Use the GetAnimNotifiesFromTrackPositions that takes FAnimNotifyContext instead")
 	ENGINE_API void GetAnimNotifiesFromTrackPositions(const float& PreviousTrackPosition, const float& CurrentTrackPosition, TArray<FAnimNotifyEventReference>& OutActiveNotifies) const;
 
+	/**
+	* Retrieves AnimNotifies between two Track time positions. ]PreviousTrackPosition, CurrentTrackPosition]
+	* Between PreviousTrackPosition (exclusive) and CurrentTrackPosition (inclusive).
+	* Supports playing backwards (CurrentTrackPosition<PreviousTrackPosition).
+	* Only supports contiguous range, does NOT support looping and wrapping over.
+	*/
+	ENGINE_API void GetAnimNotifiesFromTrackPositions(const float& PreviousTrackPosition, const float& CurrentTrackPosition, FAnimNotifyContext& NotifyContext) const;
+
 	/** return true if anim notify is available */
 	bool IsNotifyAvailable() const;
 };
diff --git a/Engine/Source/Runtime/Engine/Classes/Animation/AnimInstance.h b/Engine/Source/Runtime/Engine/Classes/Animation/AnimInstance.h
--- a/Engine/Source/Runtime/Engine/Classes/Animation/AnimInstance.h
+++ b/Engine/Source/Runtime/Engine/Classes/Animation/AnimInstance.h
@@ -1014,6 +1014,42 @@ public:
 	UFUNCTION(BlueprintPure, Category = "Asset Player", meta = (BlueprintInternalUseOnly = "true", AnimGetter = "true", GetterContext = "Transition"))
 	float GetRelevantAnimTimeFraction(int32 MachineIndex, int32 StateIndex);
 
+	/** Get whether a particular notify state was active in any state machine last tick.*/
+	UFUNCTION(BlueprintPure, Category = "Asset Player", meta = (Keywords = "Window,TransitionWindow,NotifyState,NotifyStateTransition"))
+	bool WasAnimNotifyStateActiveInAnyState(TSubclassOf<UAnimNotifyState> AnimNotifyStateType);
+
+	/** Get whether a particular notify state is active in a specific state machine last tick.  */
+	UFUNCTION(BlueprintPure, Category = "Asset Player", meta = (BlueprintInternalUseOnly = "true", AnimGetter = "true", GetterContext = "Transition", Keywords = "Window,TransitionWindow,NotifyState,NotifyStateTransition"))
+	bool WasAnimNotifyStateActiveInStateMachine(int32 MachineIndex, TSubclassOf<UAnimNotifyState> AnimNotifyStateType);
+
+	/** Get whether a particular notify state is active in a specific state last tick. */
+	UFUNCTION(BlueprintPure, Category = "Asset Player", meta = (BlueprintInternalUseOnly = "true", AnimGetter = "true", GetterContext = "Transition", Keywords = "Window,TransitionWindow,NotifyState,NotifyStateTransition"))
+	bool WasAnimNotifyStateActiveInSourceState(int32 MachineIndex, int32 StateIndex, TSubclassOf<UAnimNotifyState> AnimNotifyStateType);
+
+	/** Get whether the most relevant animation was in a particular notify state last tick. */
+	UFUNCTION(BlueprintPure, Category = "Asset Player", meta = (BlueprintInternalUseOnly = "true", AnimGetter = "true", GetterContext = "Transition", Keywords = "Window,TransitionWindow,NotifyState,NotifyStateTransition"))
+    bool WasAnimNotifyTriggeredInSourceState(int32 MachineIndex, int32 StateIndex,  TSubclassOf<UAnimNotify> AnimNotifyType);
+
+	/** Get whether the most relevant animation triggered the animation notify with the specified name last tick.. */
+	UFUNCTION(BlueprintPure, Category = "Asset Player", meta = (BlueprintInternalUseOnly = "true", AnimGetter = "true", GetterContext = "Transition", Keywords = "Window,TransitionWindow,NotifyState,NotifyStateTransition"))
+    bool WasAnimNotifyNameTriggeredInSourceState(int32 MachineIndex, int32 StateIndex, FName NotifyName);
+
+	/** Get whether a particular notify type was active in a specific state machine last tick.  */
+	UFUNCTION(BlueprintPure, Category = "Asset Player", meta = (BlueprintInternalUseOnly = "true", AnimGetter = "true", GetterContext = "Transition", Keywords = "Window,TransitionWindow,NotifyState,NotifyStateTransition"))
+    bool WasAnimNotifyTriggeredInStateMachine(int32 MachineIndex, TSubclassOf<UAnimNotify> AnimNotifyType);
+
+	/** Get whether the given state machine triggered the animation notify with the specified name last tick. */
+	UFUNCTION(BlueprintPure, Category = "Asset Player", meta = (BlueprintInternalUseOnly = "true", AnimGetter = "true", GetterContext = "Transition", Keywords = "Window,TransitionWindow,NotifyState,NotifyStateTransition"))
+    bool WasAnimNotifyNameTriggeredInStateMachine(int32 MachineIndex, FName NotifyName);
+	
+	/**  Get whether an animation notify of a given type was triggered last tick. */
+	UFUNCTION(BlueprintPure, Category = "Asset Player", meta = (BlueprintInternalUseOnly = "true", AnimGetter = "true", GetterContext = "Transition", Keywords = "Window,TransitionWindow,NotifyState,NotifyStateTransition"))
+    bool WasAnimNotifyTriggeredInAnyState(TSubclassOf<UAnimNotify> AnimNotifyType);
+
+	/** Get whether the animation notify with the specified name triggered last tick. */
+	UFUNCTION(BlueprintPure, Category = "Asset Player", meta = (BlueprintInternalUseOnly = "true", AnimGetter = "true", GetterContext = "Transition", Keywords = "Window,TransitionWindow,NotifyState,NotifyStateTransition"))
+    bool WasAnimNotifyNameTriggeredInAnyState(FName NotifyName);
+	
 	/** Gets the runtime instance of the specified state machine by Name */
 	const FAnimNode_StateMachine* GetStateMachineInstanceFromName(FName MachineName) const;
 
@@ -1263,6 +1299,9 @@ public:
 	UPROPERTY(transient)
 	TArray<FAnimNotifyEvent> ActiveAnimNotifyState;
 
+	UPROPERTY(transient)
+	TArray<FAnimNotifyEventReference> ActiveAnimNotifyEventReference;
+	
 private:
 	/** Reset Animation Curves */
 	void ResetAnimationCurves();
@@ -1339,8 +1378,13 @@ public:
 
 	/** Trigger AnimNotifies **/
 	void TriggerAnimNotifies(float DeltaSeconds);
+
+	/** Trigger an AnimNotify.  Note that this version does not provide any context for the event **/ 
 	void TriggerSingleAnimNotify(const FAnimNotifyEvent* AnimNotifyEvent);
 
+	/** Trigger an AnimNotify using an EventReference that provide context used in derived notify events**/ 
+	void TriggerSingleAnimNotify(FAnimNotifyEventReference& EventReference);
+
 	/** Triggers end on active notify states and clears the array */
 	void EndNotifyStates();
 
diff --git a/Engine/Source/Runtime/Engine/Classes/Animation/AnimMontage.h b/Engine/Source/Runtime/Engine/Classes/Animation/AnimMontage.h
--- a/Engine/Source/Runtime/Engine/Classes/Animation/AnimMontage.h
+++ b/Engine/Source/Runtime/Engine/Classes/Animation/AnimMontage.h
@@ -738,7 +738,9 @@ public:
 	virtual EAnimEventTriggerOffsets::Type CalculateOffsetForNotify(float NotifyDisplayTime) const override;
 #endif // WITH_EDITOR
 	virtual void GetMarkerIndicesForTime(float CurrentTime, bool bLooping, const TArray<FName>& ValidMarkerNames, FMarkerPair& OutPrevMarker, FMarkerPair& OutNextMarker) const override;
-	virtual FMarkerSyncAnimPosition GetMarkerSyncPositionfromMarkerIndicies(int32 PrevMarker, int32 NextMarker, float CurrentTime, const UMirrorDataTable* MirrorTable) const override;
+	UE_DEPRECATED(5.0, "Use other GetMarkerSyncPositionfromMarkerIndicies signature")
+    virtual FMarkerSyncAnimPosition GetMarkerSyncPositionfromMarkerIndicies(int32 PrevMarker, int32 NextMarker, float CurrentTime) const { return UAnimMontage::GetMarkerSyncPositionFromMarkerIndicies(PrevMarker, NextMarker, CurrentTime, nullptr); }
+	virtual FMarkerSyncAnimPosition GetMarkerSyncPositionFromMarkerIndicies(int32 PrevMarker, int32 NextMarker, float CurrentTime, const UMirrorDataTable* MirrorTable) const override;
 	virtual void TickAssetPlayer(FAnimTickRecord& Instance, struct FAnimNotifyQueue& NotifyQueue, FAnimAssetTickContext& Context) const override;
 	virtual TArray<FName>* GetUniqueMarkerNames() override { return &MarkerData.UniqueMarkerNames; }
 	virtual void RefreshCacheData() override;
diff --git a/Engine/Source/Runtime/Engine/Classes/Animation/AnimNode_StateMachine.h b/Engine/Source/Runtime/Engine/Classes/Animation/AnimNode_StateMachine.h
--- a/Engine/Source/Runtime/Engine/Classes/Animation/AnimNode_StateMachine.h
+++ b/Engine/Source/Runtime/Engine/Classes/Animation/AnimNode_StateMachine.h
@@ -141,6 +141,10 @@ public:
 	UPROPERTY(EditAnywhere, Category = Settings)
 	bool bReinitializeOnBecomingRelevant;
 
+	// Tag Notifies with meta data such as the active state and mirroring state.  Producing this
+	// data has a  slight performance cost.
+	UPROPERTY(EditAnywhere, Category = Settings)
+	bool bCreateNotifyMetaData;
 private:
 	// true if it is the first update.
 	bool bFirstUpdate;
@@ -202,6 +206,7 @@ public:
 		, MaxTransitionsPerFrame(3)
 		, bSkipFirstUpdateTransition(true)
 		, bReinitializeOnBecomingRelevant(true)
+		, bCreateNotifyMetaData(true)
 		, bFirstUpdate(true)
 		, CurrentState(INDEX_NONE)
 		, ElapsedTime(0.0f)
diff --git a/Engine/Source/Runtime/Engine/Classes/Animation/AnimNotifies/AnimNotify.h b/Engine/Source/Runtime/Engine/Classes/Animation/AnimNotifies/AnimNotify.h
--- a/Engine/Source/Runtime/Engine/Classes/Animation/AnimNotifies/AnimNotify.h
+++ b/Engine/Source/Runtime/Engine/Classes/Animation/AnimNotifies/AnimNotify.h
@@ -7,12 +7,12 @@
 #include "UObject/ObjectMacros.h"
 #include "UObject/Object.h"
 #include "UObject/ScriptMacros.h"
+#include "Animation/AnimNotifyQueue.h"
 #include "AnimNotify.generated.h"
 
 class UAnimSequenceBase;
 class USkeletalMeshComponent;
 struct FAnimNotifyEvent;
-
 USTRUCT()
 struct FBranchingPointNotifyPayload
 {
@@ -51,7 +51,7 @@ class ENGINE_API UAnimNotify : public UObject
 	FString GetNotifyName() const;
 
 	UFUNCTION(BlueprintImplementableEvent)
-	bool Received_Notify(USkeletalMeshComponent* MeshComp, UAnimSequenceBase* Animation) const;
+	bool Received_Notify(USkeletalMeshComponent* MeshComp, UAnimSequenceBase* Animation, const FAnimNotifyEventReference& EventReference) const;
 
 #if WITH_EDITORONLY_DATA
 	/** Color of Notify in editor */
@@ -65,7 +65,9 @@ class ENGINE_API UAnimNotify : public UObject
 	virtual void ValidateAssociatedAssets() {}
 #endif
 
+	UE_DEPRECATED(5.0, "Please use the other Notify function instead")
 	virtual void Notify(USkeletalMeshComponent* MeshComp, UAnimSequenceBase* Animation);
+	virtual void Notify(USkeletalMeshComponent* MeshComp, UAnimSequenceBase* Animation, const FAnimNotifyEventReference& EventReference);
 	virtual void BranchingPointNotify(FBranchingPointNotifyPayload& BranchingPointPayload);
 
 	// @todo document 
diff --git a/Engine/Source/Runtime/Engine/Classes/Animation/AnimNotifies/AnimNotifyState.h b/Engine/Source/Runtime/Engine/Classes/Animation/AnimNotifies/AnimNotifyState.h
--- a/Engine/Source/Runtime/Engine/Classes/Animation/AnimNotifies/AnimNotifyState.h
+++ b/Engine/Source/Runtime/Engine/Classes/Animation/AnimNotifies/AnimNotifyState.h
@@ -7,6 +7,7 @@
 #include "UObject/ObjectMacros.h"
 #include "UObject/Object.h"
 #include "UObject/ScriptMacros.h"
+#include "Animation/AnimNotifyQueue.h"
 #include "AnimNotifyState.generated.h"
 
 class UAnimSequenceBase;
@@ -26,13 +27,13 @@ class ENGINE_API UAnimNotifyState : public UObject
 	FString GetNotifyName() const;
 
 	UFUNCTION(BlueprintImplementableEvent)
-	bool Received_NotifyBegin(USkeletalMeshComponent * MeshComp, UAnimSequenceBase * Animation, float TotalDuration) const;
+	bool Received_NotifyBegin(USkeletalMeshComponent * MeshComp, UAnimSequenceBase * Animation, float TotalDuration, const FAnimNotifyEventReference& EventReference) const;
 	
 	UFUNCTION(BlueprintImplementableEvent)
-	bool Received_NotifyTick(USkeletalMeshComponent * MeshComp, UAnimSequenceBase * Animation, float FrameDeltaTime) const;
+	bool Received_NotifyTick(USkeletalMeshComponent * MeshComp, UAnimSequenceBase * Animation, float FrameDeltaTime, const FAnimNotifyEventReference& EventReference) const;
 
 	UFUNCTION(BlueprintImplementableEvent)
-	bool Received_NotifyEnd(USkeletalMeshComponent * MeshComp, UAnimSequenceBase * Animation) const;
+	bool Received_NotifyEnd(USkeletalMeshComponent * MeshComp, UAnimSequenceBase * Animation, const FAnimNotifyEventReference& EventReference) const;
 
 #if WITH_EDITORONLY_DATA
 	/** Color of Notify in editor */
@@ -46,10 +47,17 @@ class ENGINE_API UAnimNotifyState : public UObject
 	virtual bool CanBePlaced(UAnimSequenceBase* Animation) const { return true; }
 #endif
 
+	UE_DEPRECATED(5.0, "This function is deprecated. Use the other NotifyBegin instead.")
 	virtual void NotifyBegin(USkeletalMeshComponent * MeshComp, UAnimSequenceBase * Animation, float TotalDuration);
+	UE_DEPRECATED(5.0, "This function is deprecated. Use the other NotifyTick instead.")
 	virtual void NotifyTick(USkeletalMeshComponent * MeshComp, UAnimSequenceBase * Animation, float FrameDeltaTime);
+	UE_DEPRECATED(5.0, "This function is deprecated. Use the other NotifyEnd instead.")
 	virtual void NotifyEnd(USkeletalMeshComponent * MeshComp, UAnimSequenceBase * Animation);
 
+	virtual void NotifyBegin(USkeletalMeshComponent * MeshComp, UAnimSequenceBase * Animation, float TotalDuration, const FAnimNotifyEventReference& EventReference);
+	virtual void NotifyTick(USkeletalMeshComponent * MeshComp, UAnimSequenceBase * Animation, float FrameDeltaTime, const FAnimNotifyEventReference& EventReference);
+	virtual void NotifyEnd(USkeletalMeshComponent * MeshComp, UAnimSequenceBase * Animation, const FAnimNotifyEventReference& EventReference);
+	
 	virtual void BranchingPointNotifyBegin(FBranchingPointNotifyPayload& BranchingPointPayload);
 	virtual void BranchingPointNotifyTick(FBranchingPointNotifyPayload& BranchingPointPayload, float FrameDeltaTime);
 	virtual void BranchingPointNotifyEnd(FBranchingPointNotifyPayload& BranchingPointPayload);
diff --git a/Engine/Source/Runtime/Engine/Classes/Animation/AnimNotifies/AnimNotifyState_TimedParticleEffect.h b/Engine/Source/Runtime/Engine/Classes/Animation/AnimNotifies/AnimNotifyState_TimedParticleEffect.h
--- a/Engine/Source/Runtime/Engine/Classes/Animation/AnimNotifies/AnimNotifyState_TimedParticleEffect.h
+++ b/Engine/Source/Runtime/Engine/Classes/Animation/AnimNotifies/AnimNotifyState_TimedParticleEffect.h
@@ -57,9 +57,16 @@ class ENGINE_API UAnimNotifyState_TimedParticleEffect : public UAnimNotifyState
 	virtual void PreEditChange(FProperty* PropertyAboutToChange) override;
 #endif
 
+	UE_DEPRECATED(5.0, "Please use the other NotifyBegin function instead")
 	virtual void NotifyBegin(class USkeletalMeshComponent * MeshComp, class UAnimSequenceBase * Animation, float TotalDuration) override;
+	UE_DEPRECATED(5.0, "Please use the other NotifyTick function instead")
 	virtual void NotifyTick(class USkeletalMeshComponent * MeshComp, class UAnimSequenceBase * Animation, float FrameDeltaTime) override;
+	UE_DEPRECATED(5.0, "Please use the other NotifyEnd function instead")
 	virtual void NotifyEnd(class USkeletalMeshComponent * MeshComp, class UAnimSequenceBase * Animation) override;
+	
+	virtual void NotifyBegin(class USkeletalMeshComponent * MeshComp, class UAnimSequenceBase * Animation, float TotalDuration, const FAnimNotifyEventReference& EventReference) override;
+	virtual void NotifyTick(class USkeletalMeshComponent * MeshComp, class UAnimSequenceBase * Animation, float FrameDeltaTime, const FAnimNotifyEventReference& EventReference) override;
+	virtual void NotifyEnd(class USkeletalMeshComponent * MeshComp, class UAnimSequenceBase * Animation, const FAnimNotifyEventReference& EventReference) override;
 
 	// Overridden from UAnimNotifyState to provide custom notify name.
 	FString GetNotifyName_Implementation() const override;
diff --git a/Engine/Source/Runtime/Engine/Classes/Animation/AnimNotifies/AnimNotifyState_Trail.h b/Engine/Source/Runtime/Engine/Classes/Animation/AnimNotifies/AnimNotifyState_Trail.h
--- a/Engine/Source/Runtime/Engine/Classes/Animation/AnimNotifies/AnimNotifyState_Trail.h
+++ b/Engine/Source/Runtime/Engine/Classes/Animation/AnimNotifies/AnimNotifyState_Trail.h
@@ -76,9 +76,16 @@ class ENGINE_API UAnimNotifyState_Trail : public UAnimNotifyState
 	/** Helper function for outside code to get PSC that we are using */
 	UParticleSystemComponent* GetParticleSystemComponent(USkeletalMeshComponent* MeshComp) const;
 
-	virtual void NotifyBegin(USkeletalMeshComponent* MeshComp, UAnimSequenceBase* Animation, float TotalDuration) override;
-	virtual void NotifyTick(USkeletalMeshComponent* MeshComp, UAnimSequenceBase* Animation, float FrameDeltaTime) override;
-	virtual void NotifyEnd(USkeletalMeshComponent* MeshComp, UAnimSequenceBase* Animation) override;
+	UE_DEPRECATED(5.0, "Please use the other NotifyBegin function instead")
+	virtual void NotifyBegin(class USkeletalMeshComponent * MeshComp, class UAnimSequenceBase * Animation, float TotalDuration) override;
+	UE_DEPRECATED(5.0, "Please use the other NotifyTick function instead")
+	virtual void NotifyTick(class USkeletalMeshComponent * MeshComp, class UAnimSequenceBase * Animation, float FrameDeltaTime) override;
+	UE_DEPRECATED(5.0, "Please use the other NotifyEnd function instead")
+	virtual void NotifyEnd(class USkeletalMeshComponent * MeshComp, class UAnimSequenceBase * Animation) override;
+
+	virtual void NotifyBegin(USkeletalMeshComponent* MeshComp, UAnimSequenceBase* Animation, float TotalDuration, const FAnimNotifyEventReference& EventReference) override;
+	virtual void NotifyTick(USkeletalMeshComponent* MeshComp, UAnimSequenceBase* Animation, float FrameDeltaTime, const FAnimNotifyEventReference& EventReference) override;
+	virtual void NotifyEnd(USkeletalMeshComponent* MeshComp, UAnimSequenceBase* Animation, const FAnimNotifyEventReference& EventReference) override;
 
 	bool ValidateInput(USkeletalMeshComponent* MeshComp, bool bReportErrors = false);
 };
diff --git a/Engine/Source/Runtime/Engine/Classes/Animation/AnimNotifies/AnimNotify_PauseClothingSimulation.h b/Engine/Source/Runtime/Engine/Classes/Animation/AnimNotifies/AnimNotify_PauseClothingSimulation.h
--- a/Engine/Source/Runtime/Engine/Classes/Animation/AnimNotifies/AnimNotify_PauseClothingSimulation.h
+++ b/Engine/Source/Runtime/Engine/Classes/Animation/AnimNotifies/AnimNotify_PauseClothingSimulation.h
@@ -17,7 +17,10 @@ public:
 
 	// Begin UAnimNotify interface
 	virtual FString GetNotifyName_Implementation() const override;
+	
+	UE_DEPRECATED(5.0, "Please use the other Notify function instead")
 	virtual void Notify(USkeletalMeshComponent* MeshComp, UAnimSequenceBase* Animation) override;
+	virtual void Notify(USkeletalMeshComponent* MeshComp, UAnimSequenceBase* Animation, const FAnimNotifyEventReference& EventReference) override;
 	// End UAnimNotify interface
 
 };
diff --git a/Engine/Source/Runtime/Engine/Classes/Animation/AnimNotifies/AnimNotify_PlayParticleEffect.h b/Engine/Source/Runtime/Engine/Classes/Animation/AnimNotifies/AnimNotify_PlayParticleEffect.h
--- a/Engine/Source/Runtime/Engine/Classes/Animation/AnimNotifies/AnimNotify_PlayParticleEffect.h
+++ b/Engine/Source/Runtime/Engine/Classes/Animation/AnimNotifies/AnimNotify_PlayParticleEffect.h
@@ -31,7 +31,10 @@ public:
 
 	// Begin UAnimNotify interface
 	virtual FString GetNotifyName_Implementation() const override;
+
+	UE_DEPRECATED(5.0, "Please use the other Notify function instead")
 	virtual void Notify(USkeletalMeshComponent* MeshComp, UAnimSequenceBase* Animation) override;
+	virtual void Notify(USkeletalMeshComponent* MeshComp, UAnimSequenceBase* Animation, const FAnimNotifyEventReference& EventReference) override;
 #if WITH_EDITOR
 	virtual void ValidateAssociatedAssets() override;
 #endif
diff --git a/Engine/Source/Runtime/Engine/Classes/Animation/AnimNotifies/AnimNotify_PlaySound.h b/Engine/Source/Runtime/Engine/Classes/Animation/AnimNotifies/AnimNotify_PlaySound.h
--- a/Engine/Source/Runtime/Engine/Classes/Animation/AnimNotifies/AnimNotify_PlaySound.h
+++ b/Engine/Source/Runtime/Engine/Classes/Animation/AnimNotifies/AnimNotify_PlaySound.h
@@ -23,7 +23,9 @@ public:
 
 	// Begin UAnimNotify interface
 	virtual FString GetNotifyName_Implementation() const override;
+	UE_DEPRECATED(5.0, "Please use the other Notify function instead")
 	virtual void Notify(USkeletalMeshComponent* MeshComp, UAnimSequenceBase* Animation) override;
+	virtual void Notify(USkeletalMeshComponent* MeshComp, UAnimSequenceBase* Animation, const FAnimNotifyEventReference& EventReference) override;
 #if WITH_EDITOR
 	virtual void ValidateAssociatedAssets() override;
 #endif
diff --git a/Engine/Source/Runtime/Engine/Classes/Animation/AnimNotifies/AnimNotify_ResetClothingSimulation.h b/Engine/Source/Runtime/Engine/Classes/Animation/AnimNotifies/AnimNotify_ResetClothingSimulation.h
--- a/Engine/Source/Runtime/Engine/Classes/Animation/AnimNotifies/AnimNotify_ResetClothingSimulation.h
+++ b/Engine/Source/Runtime/Engine/Classes/Animation/AnimNotifies/AnimNotify_ResetClothingSimulation.h
@@ -17,7 +17,9 @@ public:
 
 	// Begin UAnimNotify interface
 	virtual FString GetNotifyName_Implementation() const override;
+	UE_DEPRECATED(5.0, "Please use the other Notify function instead")
 	virtual void Notify(USkeletalMeshComponent* MeshComp, UAnimSequenceBase* Animation) override;
+	virtual void Notify(USkeletalMeshComponent* MeshComp, UAnimSequenceBase* Animation, const FAnimNotifyEventReference& EventReference) override;
 	// End UAnimNotify interface
 
 };
diff --git a/Engine/Source/Runtime/Engine/Classes/Animation/AnimNotifies/AnimNotify_ResetDynamics.h b/Engine/Source/Runtime/Engine/Classes/Animation/AnimNotifies/AnimNotify_ResetDynamics.h
--- a/Engine/Source/Runtime/Engine/Classes/Animation/AnimNotifies/AnimNotify_ResetDynamics.h
+++ b/Engine/Source/Runtime/Engine/Classes/Animation/AnimNotifies/AnimNotify_ResetDynamics.h
@@ -17,7 +17,9 @@ public:
 
 	// Begin UAnimNotify interface
 	virtual FString GetNotifyName_Implementation() const override;
+	UE_DEPRECATED(5.0, "Please use the other Notify function instead")
 	virtual void Notify(USkeletalMeshComponent* MeshComp, UAnimSequenceBase* Animation) override;
+	virtual void Notify(USkeletalMeshComponent* MeshComp, UAnimSequenceBase* Animation, const FAnimNotifyEventReference& EventReference) override;
 	// End UAnimNotify interface
 
 };
diff --git a/Engine/Source/Runtime/Engine/Classes/Animation/AnimNotifies/AnimNotify_ResumeClothingSimulation.h b/Engine/Source/Runtime/Engine/Classes/Animation/AnimNotifies/AnimNotify_ResumeClothingSimulation.h
--- a/Engine/Source/Runtime/Engine/Classes/Animation/AnimNotifies/AnimNotify_ResumeClothingSimulation.h
+++ b/Engine/Source/Runtime/Engine/Classes/Animation/AnimNotifies/AnimNotify_ResumeClothingSimulation.h
@@ -17,7 +17,9 @@ public:
 
 	// Begin UAnimNotify interface
 	virtual FString GetNotifyName_Implementation() const override;
+	UE_DEPRECATED(5.0, "Please use the other Notify function instead")
 	virtual void Notify(USkeletalMeshComponent* MeshComp, UAnimSequenceBase* Animation) override;
+	virtual void Notify(USkeletalMeshComponent* MeshComp, UAnimSequenceBase* Animation, const FAnimNotifyEventReference& EventReference) override;
 	// End UAnimNotify interface
 
 };
diff --git a/Engine/Source/Runtime/Engine/Classes/Animation/AnimSequence.h b/Engine/Source/Runtime/Engine/Classes/Animation/AnimSequence.h
--- a/Engine/Source/Runtime/Engine/Classes/Animation/AnimSequence.h
+++ b/Engine/Source/Runtime/Engine/Classes/Animation/AnimSequence.h
@@ -871,7 +871,10 @@ public:
 	virtual void AdvanceMarkerPhaseAsLeader(bool bLooping, float MoveDelta, const TArray<FName>& ValidMarkerNames, float& CurrentTime, FMarkerPair& PrevMarker, FMarkerPair& NextMarker, TArray<FPassedMarker>& MarkersPassed, const UMirrorDataTable* MirrorTable) const;
 	virtual void AdvanceMarkerPhaseAsFollower(const FMarkerTickContext& Context, float DeltaRemaining, bool bLooping, float& CurrentTime, FMarkerPair& PreviousMarker, FMarkerPair& NextMarker, const UMirrorDataTable* MirrorTable) const;
 	virtual void GetMarkerIndicesForTime(float CurrentTime, bool bLooping, const TArray<FName>& ValidMarkerNames, FMarkerPair& OutPrevMarker, FMarkerPair& OutNextMarker) const;
-	virtual FMarkerSyncAnimPosition GetMarkerSyncPositionfromMarkerIndicies(int32 PrevMarker, int32 NextMarker, float CurrentTime, const UMirrorDataTable* MirrorTable) const;
+
+	UE_DEPRECATED(5.0, "Use other GetMarkerSyncPositionfromMarkerIndicies signature")
+	virtual FMarkerSyncAnimPosition GetMarkerSyncPositionfromMarkerIndicies(int32 PrevMarker, int32 NextMarker, float CurrentTime) const { return UAnimSequence::GetMarkerSyncPositionFromMarkerIndicies(PrevMarker, NextMarker, CurrentTime, nullptr); }
+	virtual FMarkerSyncAnimPosition GetMarkerSyncPositionFromMarkerIndicies(int32 PrevMarker, int32 NextMarker, float CurrentTime, const UMirrorDataTable* MirrorTable) const;
 	virtual void GetMarkerIndicesForPosition(const FMarkerSyncAnimPosition& SyncPosition, bool bLooping, FMarkerPair& OutPrevMarker, FMarkerPair& OutNextMarker, float& CurrentTime, const UMirrorDataTable* MirrorTable) const;
 	
 	virtual float GetFirstMatchingPosFromMarkerSyncPos(const FMarkerSyncAnimPosition& InMarkerSyncGroupPosition) const override;
diff --git a/Engine/Source/Runtime/Engine/Classes/Animation/AnimSequenceBase.h b/Engine/Source/Runtime/Engine/Classes/Animation/AnimSequenceBase.h
--- a/Engine/Source/Runtime/Engine/Classes/Animation/AnimSequenceBase.h
+++ b/Engine/Source/Runtime/Engine/Classes/Animation/AnimSequenceBase.h
@@ -87,16 +87,19 @@ public:
 	 * Supports playing backwards (DeltaTime<0).
 	 * Returns notifies between StartTime (exclusive) and StartTime+DeltaTime (inclusive)
 	 */
-	UE_DEPRECATED(4.19, "Use the GetAnimNotifiesFromTrackPositions that takes FAnimNotifyEventReferences instead")
+	UE_DEPRECATED(4.19, "Use the GetAnimNotifies that takes FAnimNotifyEventReferences instead")
 	void GetAnimNotifies(const float& StartTime, const float& DeltaTime, const bool bAllowLooping, TArray<const FAnimNotifyEvent *>& OutActiveNotifies) const;
 
+	UE_DEPRECATED(5.0, "Use the other GetAnimNotifies that takes FAnimNotifyContext instead")
+	void GetAnimNotifies(const float& StartTime, const float& DeltaTime, const bool bAllowLooping, TArray<FAnimNotifyEventReference>& OutActiveNotifies) const;
+
 	/**
 	* Retrieves AnimNotifies given a StartTime and a DeltaTime.
 	* Time will be advanced and support looping if bAllowLooping is true.
 	* Supports playing backwards (DeltaTime<0).
 	* Returns notifies between StartTime (exclusive) and StartTime+DeltaTime (inclusive)
 	*/
-	void GetAnimNotifies(const float& StartTime, const float& DeltaTime, const bool bAllowLooping, TArray<FAnimNotifyEventReference>& OutActiveNotifies) const;
+	void GetAnimNotifies(const float& StartTime, const float& DeltaTime, FAnimNotifyContext& NotifyContext) const;
 
 	/** 
 	 * Retrieves AnimNotifies between two time positions. ]PreviousPosition, CurrentPosition]
@@ -104,16 +107,18 @@ public:
 	 * Supports playing backwards (CurrentPosition<PreviousPosition).
 	 * Only supports contiguous range, does NOT support looping and wrapping over.
 	 */
-	UE_DEPRECATED(4.19, "Use the GetAnimNotifiesFromTrackPositions that takes FAnimNotifyEventReferences instead")
+	UE_DEPRECATED(4.19, "Use the GetAnimNotifiesFromDeltaPositions that takes FAnimNotifyEventReferences instead")
 	void GetAnimNotifiesFromDeltaPositions(const float& PreviousPosition, const float & CurrentPosition, TArray<const FAnimNotifyEvent *>& OutActiveNotifies) const;
 
+	UE_DEPRECATED(5.0, "Use the other GetAnimNotifiesFromDeltaPositions that takes FAnimNotifyContext instead")
+	virtual void GetAnimNotifiesFromDeltaPositions(const float& PreviousPosition, const float & CurrentPosition, TArray<FAnimNotifyEventReference>& OutActiveNotifies) const;
 	/**
 	* Retrieves AnimNotifies between two time positions. ]PreviousPosition, CurrentPosition]
 	* Between PreviousPosition (exclusive) and CurrentPosition (inclusive).
 	* Supports playing backwards (CurrentPosition<PreviousPosition).
 	* Only supports contiguous range, does NOT support looping and wrapping over.
 	*/
-	virtual void GetAnimNotifiesFromDeltaPositions(const float& PreviousPosition, const float & CurrentPosition, TArray<FAnimNotifyEventReference>& OutActiveNotifies) const;
+	virtual void GetAnimNotifiesFromDeltaPositions(const float& PreviousPosition, const float & CurrentPosition, FAnimNotifyContext& NotifyContext) const;
 
 	/** Evaluate curve data to Instance at the time of CurrentTime **/
 	virtual void EvaluateCurveData(FBlendedCurve& OutCurve, float CurrentTime, bool bForceUseRawData = false) const;
@@ -209,8 +214,8 @@ public:
 	virtual void GetMarkerIndicesForTime(float CurrentTime, bool bLooping, const TArray<FName>& ValidMarkerNames, FMarkerPair& OutPrevMarker, FMarkerPair& OutNextMarker) const { check(false); /*Should never call this (either missing override or calling on unsupported asset */ }
 
 	UE_DEPRECATED(5.0, "Use other GetMarkerSyncPositionfromMarkerIndicies signature")
-	virtual FMarkerSyncAnimPosition GetMarkerSyncPositionfromMarkerIndicies(int32 PrevMarker, int32 NextMarker, float CurrentTime) const { return GetMarkerSyncPositionfromMarkerIndicies(PrevMarker, NextMarker, CurrentTime, nullptr); }
-	virtual FMarkerSyncAnimPosition GetMarkerSyncPositionfromMarkerIndicies(int32 PrevMarker, int32 NextMarker, float CurrentTime, const UMirrorDataTable* MirrorTable) const { check(false); return FMarkerSyncAnimPosition(); /*Should never call this (either missing override or calling on unsupported asset */ }
+	virtual FMarkerSyncAnimPosition GetMarkerSyncPositionfromMarkerIndicies(int32 PrevMarker, int32 NextMarker, float CurrentTime) const { return UAnimSequenceBase::GetMarkerSyncPositionFromMarkerIndicies(PrevMarker, NextMarker, CurrentTime, nullptr); }
+	virtual FMarkerSyncAnimPosition GetMarkerSyncPositionFromMarkerIndicies(int32 PrevMarker, int32 NextMarker, float CurrentTime, const UMirrorDataTable* MirrorTable) const { check(false); return FMarkerSyncAnimPosition(); /*Should never call this (either missing override or calling on unsupported asset */ }
 
 	UE_DEPRECATED(5.0, "Use other GetMarkerIndicesForPosition signature")
 	virtual void GetMarkerIndicesForPosition(const FMarkerSyncAnimPosition& SyncPosition, bool bLooping, FMarkerPair& OutPrevMarker, FMarkerPair& OutNextMarker, float& CurrentTime) const { GetMarkerIndicesForPosition(SyncPosition, bLooping, OutPrevMarker, OutNextMarker, CurrentTime, nullptr); }
diff --git a/Engine/Source/Runtime/Engine/Classes/Animation/AnimationAsset.h b/Engine/Source/Runtime/Engine/Classes/Animation/AnimationAsset.h
--- a/Engine/Source/Runtime/Engine/Classes/Animation/AnimationAsset.h
+++ b/Engine/Source/Runtime/Engine/Classes/Animation/AnimationAsset.h
@@ -26,7 +26,11 @@ class USkeleton;
 class UAnimSequenceBase;
 class UBlendSpace;
 class UPoseAsset;
-class UMirrorDataTable; 
+class UMirrorDataTable;
+
+namespace UE { namespace Anim {
+class IAnimNotifyEventContextDataInterface;
+}}
 
 namespace MarkerIndexSpecialValues
 {
@@ -307,6 +311,7 @@ struct FAnimTickRecord
 	bool bLooping = false;
 	const UMirrorDataTable* MirrorDataTable = nullptr;
 
+	TArray<TSharedPtr<const UE::Anim::IAnimNotifyEventContextDataInterface>> ContextData;	
 	union
 	{
 		struct
diff --git a/Engine/Source/Runtime/Engine/Classes/Animation/MirrorDataTable.h b/Engine/Source/Runtime/Engine/Classes/Animation/MirrorDataTable.h
--- a/Engine/Source/Runtime/Engine/Classes/Animation/MirrorDataTable.h
+++ b/Engine/Source/Runtime/Engine/Classes/Animation/MirrorDataTable.h
@@ -8,25 +8,33 @@
 
 /** Type referenced by a row in the mirror data table */
 UENUM()
-enum EMirrorRowType
+namespace EMirrorRowType
 {
-	Bone,
-	AnimationNotify,
-	Curve,
-	SyncMarker
-};
+	enum Type
+	{
+		Bone,
+		AnimationNotify,
+		Curve,
+		SyncMarker,
+		Custom
+	};
+}
+
 
 /** Find and Replace Method for FMirrorFindReplaceExpression. */
 UENUM()
-enum EMirrorFindReplaceMethod
+namespace EMirrorFindReplaceMethod
 {
-	/** Only find and replace matching strings at the start of the name  */
-	Prefix,
-	/** Only find and replace matching strings at the end of the name  */
-	Suffix,
-	/** Use regular expressions for find and replace, including support for captures $1 - $10 */
-	RegularExpression
-};
+	enum Type
+	{
+		/** Only find and replace matching strings at the start of the name  */
+		Prefix,
+        /** Only find and replace matching strings at the end of the name  */
+        Suffix,
+        /** Use regular expressions for find and replace, including support for captures $1 - $10 */
+        RegularExpression
+    };
+}
 
 /**  Base Mirror Table containing all data required by the animation mirroring system. */
 USTRUCT()
@@ -41,7 +49,7 @@ struct FMirrorTableRow : public FTableRowBase
 	FName MirroredName;
 
 	UPROPERTY(EditAnywhere, BlueprintReadOnly, Category = Mirroring)
-	TEnumAsByte<EMirrorRowType> MirrorEntryType;
+	TEnumAsByte<EMirrorRowType::Type> MirrorEntryType;
 
 	FMirrorTableRow()
 		: Name(NAME_None)
@@ -68,14 +76,14 @@ struct FMirrorFindReplaceExpression
 	FName ReplaceExpression;
 
 	UPROPERTY(EditAnywhere, Category = Mirroring)
-	TEnumAsByte<EMirrorFindReplaceMethod> FindReplaceMethod;
+	TEnumAsByte<EMirrorFindReplaceMethod::Type> FindReplaceMethod;
 
 	FMirrorFindReplaceExpression() 
 		: FindExpression(NAME_None)
 		, ReplaceExpression(NAME_None)
 		, FindReplaceMethod(EMirrorFindReplaceMethod::Prefix) {}
 
-	FMirrorFindReplaceExpression(FName InFindExpression, FName InReplaceExpression, EMirrorFindReplaceMethod Method)
+	FMirrorFindReplaceExpression(FName InFindExpression, FName InReplaceExpression, EMirrorFindReplaceMethod::Type Method)
 		: FindExpression(InFindExpression)
 		, ReplaceExpression(InReplaceExpression)
 		, FindReplaceMethod(Method)
diff --git a/Engine/Source/Runtime/Engine/Private/Animation/ActiveStateMachineScope.cpp b/Engine/Source/Runtime/Engine/Private/Animation/ActiveStateMachineScope.cpp
new file mode 100644
--- /dev/null
+++ b/Engine/Source/Runtime/Engine/Private/Animation/ActiveStateMachineScope.cpp
@@ -0,0 +1,106 @@
+// Copyright Epic Games, Inc. All Rights Reserved.
+
+#include "Animation/ActiveStateMachineScope.h"
+#include "Animation/AnimInstanceProxy.h"
+#include "Animation/AnimNode_StateMachine.h"
+
+FEncounteredStateMachineStack::FEncounteredStateMachineStack(const FEncounteredStateMachineStack& ParentStack, int32 InStateMachineIndex, int32 InStateIndex) :
+	StateStack(ParentStack.StateStack)
+{
+	StateStack.Emplace(InStateMachineIndex, InStateIndex);
+}
+
+FEncounteredStateMachineStack::FEncounteredStateMachineStack(int32 InStateMachineIndex, int32 InStateIndex)
+{
+	StateStack.Emplace(InStateMachineIndex, InStateIndex);
+}
+
+TSharedPtr<const FEncounteredStateMachineStack> FEncounteredStateMachineStack::InitStack(const FEncounteredStateMachineStack& ParentStack, int32 InStateMachineIndex, int32 InStateIndex)
+{
+	return MakeShared<const FEncounteredStateMachineStack>(ParentStack, InStateMachineIndex, InStateIndex); 
+}
+
+TSharedPtr<const FEncounteredStateMachineStack> FEncounteredStateMachineStack::InitStack(int32 InStateMachineIndex, int32 InStateIndex)
+{
+	return MakeShared<const FEncounteredStateMachineStack>(InStateMachineIndex, InStateIndex);
+}
+
+
+
+namespace UE { namespace Anim {
+
+IMPLEMENT_NOTIFY_CONTEXT_INTERFACE(FAnimNotifyStateMachineContext)
+
+FAnimNotifyStateMachineContext::FAnimNotifyStateMachineContext(const TSharedPtr<const FEncounteredStateMachineStack>& InEncounteredStateMachines)
+{
+	EncounteredStateMachines = InEncounteredStateMachines;
+}
+
+bool FAnimNotifyStateMachineContext::IsStateMachineInContext(int32 StateMachineIndex) const
+{
+	if (!EncounteredStateMachines.IsValid())
+	{
+		return false; 
+	}
+
+	for (const FEncounteredStateMachineStack::FStateMachineEntry& Entry : EncounteredStateMachines->StateStack)
+	{
+		if (Entry.StateMachineIndex == StateMachineIndex)
+		{
+			return true; 
+		}
+	}
+	return false; 
+}
+
+bool FAnimNotifyStateMachineContext::IsStateInStateMachineInContext(int32 StateMachineIndex, int32 StateIndex) const
+{
+	if (!EncounteredStateMachines.IsValid())
+	{
+		return false;
+	}
+
+	for (const FEncounteredStateMachineStack::FStateMachineEntry& Entry : EncounteredStateMachines->StateStack)
+	{
+		if (Entry.StateMachineIndex == StateMachineIndex && Entry.StateIndex == StateIndex)
+		{
+			return true;
+		}
+	}
+	return false;
+}
+
+IMPLEMENT_ANIMGRAPH_MESSAGE(FActiveStateMachineScope);
+
+FActiveStateMachineScope::FActiveStateMachineScope(const FAnimationBaseContext& InContext, FAnimNode_StateMachine* StateMachine, int32 InStateIndex)
+{
+	int32 StateMachineIndex = GetStateMachineIndex(StateMachine, InContext);
+	if (FActiveStateMachineScope* ParentStateMachineScope = InContext.GetMessage<FActiveStateMachineScope>())
+	{
+		ActiveStateMachines = FEncounteredStateMachineStack::InitStack(*(ParentStateMachineScope->ActiveStateMachines.Get()), StateMachineIndex, InStateIndex); 
+	}
+	else
+	{
+		ActiveStateMachines = FEncounteredStateMachineStack::InitStack(StateMachineIndex, InStateIndex);
+	}
+}
+
+FActiveStateMachineScope::~FActiveStateMachineScope()
+{
+}
+
+int32 FActiveStateMachineScope::GetStateMachineIndex(FAnimNode_StateMachine* StateMachine, const FAnimationBaseContext& Context)
+{
+	if (Context.AnimInstanceProxy)
+	{
+		return Context.AnimInstanceProxy->GetStateMachineIndex(StateMachine);
+	}
+	return INDEX_NONE;
+}
+
+TSharedPtr<const IAnimNotifyEventContextDataInterface> FActiveStateMachineScope::MakeEventContextData() const
+{
+	return MakeShared<const FAnimNotifyStateMachineContext>(ActiveStateMachines);
+}
+	
+}}	// namespace UE::Anim
\ No newline at end of file
diff --git a/Engine/Source/Runtime/Engine/Private/Animation/AnimComposite.cpp b/Engine/Source/Runtime/Engine/Private/Animation/AnimComposite.cpp
--- a/Engine/Source/Runtime/Engine/Private/Animation/AnimComposite.cpp
+++ b/Engine/Source/Runtime/Engine/Private/Animation/AnimComposite.cpp
@@ -31,34 +31,34 @@ bool UAnimComposite::IsNotifyAvailable() const
 	return (GetPlayLength() > 0.f && (Super::IsNotifyAvailable() || AnimationTrack.IsNotifyAvailable()));
 }
 
-void UAnimComposite::GetAnimNotifiesFromDeltaPositions(const float& PreviousPosition, const float & CurrentPosition, TArray<FAnimNotifyEventReference>& OutActiveNotifies) const
+void UAnimComposite::GetAnimNotifiesFromDeltaPositions(const float& PreviousPosition, const float & CurrentPosition, FAnimNotifyContext& NotifyContext) const
 {
 	const bool bMovingForward = (RateScale >= 0.f);
 
-	Super::GetAnimNotifiesFromDeltaPositions(PreviousPosition, CurrentPosition, OutActiveNotifies);
+	Super::GetAnimNotifiesFromDeltaPositions(PreviousPosition, CurrentPosition, NotifyContext);
 
 	if (bMovingForward)
 	{
 		if (PreviousPosition <= CurrentPosition)
 		{
-			AnimationTrack.GetAnimNotifiesFromTrackPositions(PreviousPosition, CurrentPosition, OutActiveNotifies);
+			AnimationTrack.GetAnimNotifiesFromTrackPositions(PreviousPosition, CurrentPosition, NotifyContext);
 		}
 		else
 		{
-			AnimationTrack.GetAnimNotifiesFromTrackPositions(PreviousPosition, GetPlayLength(), OutActiveNotifies);
-			AnimationTrack.GetAnimNotifiesFromTrackPositions(0.f, CurrentPosition, OutActiveNotifies);
+			AnimationTrack.GetAnimNotifiesFromTrackPositions(PreviousPosition, GetPlayLength(), NotifyContext);
+			AnimationTrack.GetAnimNotifiesFromTrackPositions(0.f, CurrentPosition, NotifyContext);
 		}
 	}
 	else
 	{
 		if (PreviousPosition >= CurrentPosition)
 		{
-			AnimationTrack.GetAnimNotifiesFromTrackPositions(PreviousPosition, CurrentPosition, OutActiveNotifies);
+			AnimationTrack.GetAnimNotifiesFromTrackPositions(PreviousPosition, CurrentPosition, NotifyContext);
 		}
 		else
 		{
-			AnimationTrack.GetAnimNotifiesFromTrackPositions(PreviousPosition, 0.f, OutActiveNotifies);
-			AnimationTrack.GetAnimNotifiesFromTrackPositions(GetPlayLength(), CurrentPosition, OutActiveNotifies);
+			AnimationTrack.GetAnimNotifiesFromTrackPositions(PreviousPosition, 0.f, NotifyContext);
+			AnimationTrack.GetAnimNotifiesFromTrackPositions(GetPlayLength(), CurrentPosition, NotifyContext);
 		}
 	}
 }
diff --git a/Engine/Source/Runtime/Engine/Private/Animation/AnimCompositeBase.cpp b/Engine/Source/Runtime/Engine/Private/Animation/AnimCompositeBase.cpp
--- a/Engine/Source/Runtime/Engine/Private/Animation/AnimCompositeBase.cpp
+++ b/Engine/Source/Runtime/Engine/Private/Animation/AnimCompositeBase.cpp
@@ -70,11 +70,12 @@ float FAnimSegment::ConvertTrackPosToAnimPos(const float& TrackPosition) const
 
 void FAnimSegment::GetAnimNotifiesFromTrackPositions(const float& PreviousTrackPosition, const float& CurrentTrackPosition, TArray<const FAnimNotifyEvent *> & OutActiveNotifies) const
 {
-	TArray<FAnimNotifyEventReference> NotifyRefs;
-	GetAnimNotifiesFromTrackPositions(PreviousTrackPosition, CurrentTrackPosition, NotifyRefs);
+	FAnimTickRecord TickRecord;
+	FAnimNotifyContext NotifyContext(TickRecord);
+	GetAnimNotifiesFromTrackPositions(PreviousTrackPosition, CurrentTrackPosition, NotifyContext);
 
-	OutActiveNotifies.Reset(NotifyRefs.Num());
-	for (FAnimNotifyEventReference& NotifyRef : NotifyRefs)
+	OutActiveNotifies.Reset(NotifyContext.ActiveNotifies.Num());
+	for (FAnimNotifyEventReference& NotifyRef : NotifyContext.ActiveNotifies)
 	{
 		if (const FAnimNotifyEvent* Notify = NotifyRef.GetNotify())
 		{
@@ -84,6 +85,15 @@ void FAnimSegment::GetAnimNotifiesFromTrackPositions(const float& PreviousTrackP
 }
 
 void FAnimSegment::GetAnimNotifiesFromTrackPositions(const float& PreviousTrackPosition, const float& CurrentTrackPosition, TArray<FAnimNotifyEventReference> & OutActiveNotifies) const
+{
+	FAnimTickRecord TickRecord;
+	FAnimNotifyContext NotifyContext(TickRecord);
+	GetAnimNotifiesFromTrackPositions(PreviousTrackPosition, CurrentTrackPosition, NotifyContext);
+	// Slow copy due assumption of calling code that OutActiveNotifies is only extended
+	OutActiveNotifies.Append(NotifyContext.ActiveNotifies);
+}
+
+void FAnimSegment::GetAnimNotifiesFromTrackPositions(const float& PreviousTrackPosition, const float& CurrentTrackPosition, FAnimNotifyContext& NotifyContext) const
 {
 	if( PreviousTrackPosition == CurrentTrackPosition )
 	{
@@ -130,14 +140,14 @@ void FAnimSegment::GetAnimNotifiesFromTrackPositions(const float& PreviousTrackP
 				{
 					const float PlayRate = ValidPlayRate * (bTrackPlayingBackwards ? -1.f : 1.f);
 					const float AnimEndPosition = (TrackTimeToGo * PlayRate) + AnimStartPosition;
-					AnimSequenceBase->GetAnimNotifiesFromDeltaPositions(AnimStartPosition, AnimEndPosition, OutActiveNotifies);
+					AnimSequenceBase->GetAnimNotifiesFromDeltaPositions(AnimStartPosition, AnimEndPosition, NotifyContext);
 					break;
 				}
 				// Otherwise we hit the end point of the animation first...
 				else
 				{
 					// Add that piece for extraction.
-					AnimSequenceBase->GetAnimNotifiesFromDeltaPositions(AnimStartPosition, AnimEndPoint, OutActiveNotifies);
+					AnimSequenceBase->GetAnimNotifiesFromDeltaPositions(AnimStartPosition, AnimEndPoint, NotifyContext);
 
 					// decrease our TrackTimeToGo if we have to do another iteration.
 					// and put ourselves back at the beginning of the animation.
@@ -644,11 +654,12 @@ bool FAnimTrack::ContainRecursive(const TArray<UAnimCompositeBase*>& CurrentAccu
 
 void FAnimTrack::GetAnimNotifiesFromTrackPositions(const float& PreviousTrackPosition, const float& CurrentTrackPosition, TArray<const FAnimNotifyEvent *> & OutActiveNotifies) const
 {
-	TArray<FAnimNotifyEventReference> NotifyRefs;
-	GetAnimNotifiesFromTrackPositions(PreviousTrackPosition, CurrentTrackPosition, NotifyRefs);
+	FAnimTickRecord TickRecord;
+	FAnimNotifyContext NotifyContext(TickRecord);
+	GetAnimNotifiesFromTrackPositions(PreviousTrackPosition, CurrentTrackPosition, NotifyContext);
 
-	OutActiveNotifies.Reset(NotifyRefs.Num());
-	for (FAnimNotifyEventReference& NotifyRef : NotifyRefs)
+	OutActiveNotifies.Reset(NotifyContext.ActiveNotifies.Num());
+	for (FAnimNotifyEventReference& NotifyRef : NotifyContext.ActiveNotifies)
 	{
 		if (const FAnimNotifyEvent* Notify = NotifyRef.GetNotify())
 		{
@@ -658,12 +669,26 @@ void FAnimTrack::GetAnimNotifiesFromTrackPositions(const float& PreviousTrackPos
 }
 
 void FAnimTrack::GetAnimNotifiesFromTrackPositions(const float& PreviousTrackPosition, const float& CurrentTrackPosition, TArray<FAnimNotifyEventReference> & OutActiveNotifies) const
+{
+	FAnimTickRecord TickRecord;
+	FAnimNotifyContext NotifyContext(TickRecord);
+	for (int32 SegmentIndex = 0; SegmentIndex<AnimSegments.Num(); ++SegmentIndex)
+	{
+		if (AnimSegments[SegmentIndex].IsValid())
+		{
+			AnimSegments[SegmentIndex].GetAnimNotifiesFromTrackPositions(PreviousTrackPosition, CurrentTrackPosition, NotifyContext);
+		}
+	}
+	Swap(OutActiveNotifies, NotifyContext.ActiveNotifies);
+}
+
+void FAnimTrack::GetAnimNotifiesFromTrackPositions(const float& PreviousTrackPosition, const float& CurrentTrackPosition, FAnimNotifyContext& NotifyContext) const
 {
 	for (int32 SegmentIndex = 0; SegmentIndex<AnimSegments.Num(); ++SegmentIndex)
 	{
 		if (AnimSegments[SegmentIndex].IsValid())
 		{
-			AnimSegments[SegmentIndex].GetAnimNotifiesFromTrackPositions(PreviousTrackPosition, CurrentTrackPosition, OutActiveNotifies);
+			AnimSegments[SegmentIndex].GetAnimNotifiesFromTrackPositions(PreviousTrackPosition, CurrentTrackPosition, NotifyContext);
 		}
 	}
 }
diff --git a/Engine/Source/Runtime/Engine/Private/Animation/AnimInstance.cpp b/Engine/Source/Runtime/Engine/Private/Animation/AnimInstance.cpp
--- a/Engine/Source/Runtime/Engine/Private/Animation/AnimInstance.cpp
+++ b/Engine/Source/Runtime/Engine/Private/Animation/AnimInstance.cpp
@@ -290,16 +290,18 @@ void UAnimInstance::UninitializeAnimation()
 			for(int32 Index=0; Index<ActiveAnimNotifyState.Num(); Index++)
 			{
 				const FAnimNotifyEvent& AnimNotifyEvent = ActiveAnimNotifyState[Index];
+				const FAnimNotifyEventReference& EventReference = ActiveAnimNotifyEventReference[Index];
 				if (ShouldTriggerAnimNotifyState(AnimNotifyEvent.NotifyStateClass))
 				{
 					TRACE_ANIM_NOTIFY(this, AnimNotifyEvent, End);
-					AnimNotifyEvent.NotifyStateClass->NotifyEnd(SkelMeshComp, Cast<UAnimSequenceBase>(AnimNotifyEvent.NotifyStateClass->GetOuter()));
+					AnimNotifyEvent.NotifyStateClass->NotifyEnd(SkelMeshComp, Cast<UAnimSequenceBase>(AnimNotifyEvent.NotifyStateClass->GetOuter()), EventReference);
 				}
 			}
 		}
 	}
 
 	ActiveAnimNotifyState.Reset();
+	ActiveAnimNotifyEventReference.Reset(); 
 	NotifyQueue.Reset(SkelMeshComp);
 
 	SlotGroupInertializationRequestMap.Reset();
@@ -542,6 +544,8 @@ void UAnimInstance::PreUpdateAnimation(float DeltaSeconds)
 
 	bNeedsUpdate = true;
 
+	GetProxyOnGameThread<FAnimInstanceProxy>().UpdateActiveAnimNotifiesSinceLastTick(NotifyQueue);
+	
 	NotifyQueue.Reset(GetSkelMeshComponent());
 	RootMotionBlendQueue.Reset();
 
@@ -1355,8 +1359,13 @@ void UAnimInstance::TriggerAnimNotifies(float DeltaSeconds)
 	TArray<FAnimNotifyEvent> NewActiveAnimNotifyState;
 	NewActiveAnimNotifyState.Reserve(NotifyQueue.AnimNotifies.Num());
 
+	TArray<FAnimNotifyEventReference> NewActiveAnimNotifyEventReference;
+	NewActiveAnimNotifyEventReference.Reserve(NotifyQueue.AnimNotifies.Num());
+
+	
 	// AnimNotifyState freshly added that need their 'NotifyBegin' event called.
 	TArray<const FAnimNotifyEvent *> NotifyStateBeginEvent;
+	TArray<const FAnimNotifyEventReference *> NotifyStateBeginEventReference;
 
 	for (int32 Index=0; Index<NotifyQueue.AnimNotifies.Num(); Index++)
 	{
@@ -1369,13 +1378,23 @@ void UAnimInstance::TriggerAnimNotifies(float DeltaSeconds)
 				{
 					// Queue up calls to 'NotifyBegin', so they happen after 'NotifyEnd'.
 					NotifyStateBeginEvent.Add(AnimNotifyEvent);
+					NotifyStateBeginEventReference.Add(&NotifyQueue.AnimNotifies[Index]);
+				}
+				else if(!ActiveAnimNotifyEventReference.IsEmpty())
+				{
+					// Remove the matching Event Reference
+					bool bRemovedEvent = ActiveAnimNotifyEventReference.RemoveSingleSwap(NotifyQueue.AnimNotifies[Index]) ==1;
+					check(bRemovedEvent);
 				}
+
 				NewActiveAnimNotifyState.Add(*AnimNotifyEvent);
+				FAnimNotifyEventReference& EventRef = NewActiveAnimNotifyEventReference.Add_GetRef(NotifyQueue.AnimNotifies[Index]);
+				EventRef.SetNotify(&NewActiveAnimNotifyState.Top());
 				continue;
 			}
 
 			// Trigger non 'state' AnimNotifies
-			TriggerSingleAnimNotify(AnimNotifyEvent);
+			TriggerSingleAnimNotify(NotifyQueue.AnimNotifies[Index]);
 		}
 	}
 
@@ -1383,10 +1402,11 @@ void UAnimInstance::TriggerAnimNotifies(float DeltaSeconds)
 	for (int32 Index = 0; Index < ActiveAnimNotifyState.Num(); ++Index)
 	{
 		const FAnimNotifyEvent& AnimNotifyEvent = ActiveAnimNotifyState[Index];
+		const FAnimNotifyEventReference& EventReference = ActiveAnimNotifyEventReference[Index];
 		if (AnimNotifyEvent.NotifyStateClass && ShouldTriggerAnimNotifyState(AnimNotifyEvent.NotifyStateClass))
 		{
 			TRACE_ANIM_NOTIFY(this, AnimNotifyEvent, End);
-			AnimNotifyEvent.NotifyStateClass->NotifyEnd(SkelMeshComp, Cast<UAnimSequenceBase>(AnimNotifyEvent.NotifyStateClass->GetOuter()));
+			AnimNotifyEvent.NotifyStateClass->NotifyEnd(SkelMeshComp, Cast<UAnimSequenceBase>(AnimNotifyEvent.NotifyStateClass->GetOuter()), EventReference);
 		}
 		// The NotifyEnd callback above may have triggered actor destruction and the tear down
 		// of this instance via UninitializeAnimation which empties ActiveAnimNotifyState.
@@ -1398,33 +1418,45 @@ void UAnimInstance::TriggerAnimNotifies(float DeltaSeconds)
 		}
 	}
 
-	// Call 'NotifyBegin' event on freshly added AnimNotifyState.
-	for (const FAnimNotifyEvent* AnimNotifyEvent : NotifyStateBeginEvent)
+	check(NotifyStateBeginEventReference.Num() == NotifyStateBeginEvent.Num());
+	for (int32 Index = 0; Index < NotifyStateBeginEvent.Num(); Index++)
 	{
+		const FAnimNotifyEvent* AnimNotifyEvent = NotifyStateBeginEvent[Index];
+		const FAnimNotifyEventReference * AnimNotifyEventReference = NotifyStateBeginEventReference[Index];
 		if (ShouldTriggerAnimNotifyState(AnimNotifyEvent->NotifyStateClass))
 		{
 			TRACE_ANIM_NOTIFY(this, *AnimNotifyEvent, Begin);
-			AnimNotifyEvent->NotifyStateClass->NotifyBegin(SkelMeshComp, Cast<UAnimSequenceBase>(AnimNotifyEvent->NotifyStateClass->GetOuter()), AnimNotifyEvent->GetDuration());
+			AnimNotifyEvent->NotifyStateClass->NotifyBegin(SkelMeshComp, Cast<UAnimSequenceBase>(AnimNotifyEvent->NotifyStateClass->GetOuter()), AnimNotifyEvent->GetDuration(), *AnimNotifyEventReference);
 		}
 	}
 
 	// Switch our arrays.
 	ActiveAnimNotifyState = MoveTemp(NewActiveAnimNotifyState);
-
+	ActiveAnimNotifyEventReference = MoveTemp(NewActiveAnimNotifyEventReference);
 	// Tick currently active AnimNotifyState
-	for(const FAnimNotifyEvent& AnimNotifyEvent : ActiveAnimNotifyState)
+	for (int32 Index = 0; Index < ActiveAnimNotifyState.Num(); Index++)
 	{
+		const FAnimNotifyEvent& AnimNotifyEvent = ActiveAnimNotifyState[Index];
+		const FAnimNotifyEventReference& EventReference = ActiveAnimNotifyEventReference[Index];
 		if (ShouldTriggerAnimNotifyState(AnimNotifyEvent.NotifyStateClass))
 		{
 			TRACE_ANIM_NOTIFY(this, AnimNotifyEvent, Tick);
-			AnimNotifyEvent.NotifyStateClass->NotifyTick(SkelMeshComp, Cast<UAnimSequenceBase>(AnimNotifyEvent.NotifyStateClass->GetOuter()), DeltaSeconds);
+			AnimNotifyEvent.NotifyStateClass->NotifyTick(SkelMeshComp, Cast<UAnimSequenceBase>(AnimNotifyEvent.NotifyStateClass->GetOuter()), DeltaSeconds, EventReference);
 		}
 	}
 }
 
 void UAnimInstance::TriggerSingleAnimNotify(const FAnimNotifyEvent* AnimNotifyEvent)
+{
+	FAnimNotifyEventReference EventReference(AnimNotifyEvent, this);
+	
+	TriggerSingleAnimNotify(EventReference); 
+}
+
+void UAnimInstance::TriggerSingleAnimNotify(FAnimNotifyEventReference& EventReference)
 {
 	// This is for non 'state' anim notifies.
+	const FAnimNotifyEvent* AnimNotifyEvent = EventReference.GetNotify();
 	if (AnimNotifyEvent && (AnimNotifyEvent->NotifyStateClass == NULL))
 	{
 		if (HandleNotify(*AnimNotifyEvent))
@@ -1436,12 +1468,12 @@ void UAnimInstance::TriggerSingleAnimNotify(const FAnimNotifyEvent* AnimNotifyEv
 		{	
 			// Implemented notify: just call Notify. UAnimNotify will forward this to the event which will do the work.
 			TRACE_ANIM_NOTIFY(this, *AnimNotifyEvent, Event);
-			AnimNotifyEvent->Notify->Notify(GetSkelMeshComponent(), Cast<UAnimSequenceBase>(AnimNotifyEvent->Notify->GetOuter()));
+			AnimNotifyEvent->Notify->Notify(GetSkelMeshComponent(), Cast<UAnimSequenceBase>(AnimNotifyEvent->Notify->GetOuter()), EventReference);
 		}
 		else if (AnimNotifyEvent->NotifyName != NAME_None)
 		{
 			// Custom Event based notifies. These will call a AnimNotify_* function on the AnimInstance.
-			const FName FuncName = AnimNotifyEvent->GetNotifyEventName();
+			FName FuncName = AnimNotifyEvent->GetNotifyEventName(EventReference.GetMirrorDataTable());
 
 			auto NotifyAnimInstance = [this, AnimNotifyEvent, FuncName](UAnimInstance* InAnimInstance)
 			{
@@ -1501,15 +1533,18 @@ void UAnimInstance::EndNotifyStates()
 {
 	USkeletalMeshComponent* SkelMeshComp = GetSkelMeshComponent();
 
-	for (FAnimNotifyEvent& Event : ActiveAnimNotifyState)
+	for (int32 Index = 0; Index < ActiveAnimNotifyState.Num(); Index++)
 	{
+		const FAnimNotifyEvent& Event = ActiveAnimNotifyState[Index];
+		const FAnimNotifyEventReference& EventReference = ActiveAnimNotifyEventReference[Index];
 		if (UAnimNotifyState* NotifyState = Event.NotifyStateClass)
 		{
 			TRACE_ANIM_NOTIFY(this, Event, End);
-			NotifyState->NotifyEnd(SkelMeshComp, Cast<UAnimSequenceBase>(NotifyState->GetOuter()));
+			NotifyState->NotifyEnd(SkelMeshComp, Cast<UAnimSequenceBase>(NotifyState->GetOuter()), EventReference);
 		}
 	}
 	ActiveAnimNotifyState.Reset();
+	ActiveAnimNotifyEventReference.Reset();
 }
 
 //to debug montage weight
@@ -1727,6 +1762,7 @@ void UAnimInstance::TriggerMontageEndedEvent(const FQueuedMontageEndedEvent& Mon
 		for (int32 Index = ActiveAnimNotifyState.Num() - 1; ActiveAnimNotifyState.IsValidIndex(Index); --Index)
 		{
 			const FAnimNotifyEvent& AnimNotifyEvent = ActiveAnimNotifyState[Index];
+			const FAnimNotifyEventReference& EventReference = ActiveAnimNotifyEventReference[Index];
 			UAnimMontage* NotifyMontage = Cast<UAnimMontage>(AnimNotifyEvent.NotifyStateClass->GetOuter());
 
 			if (NotifyMontage && (NotifyMontage == MontageEndedEvent.Montage))
@@ -1734,7 +1770,7 @@ void UAnimInstance::TriggerMontageEndedEvent(const FQueuedMontageEndedEvent& Mon
 				if (ShouldTriggerAnimNotifyState(AnimNotifyEvent.NotifyStateClass))
 				{
 					TRACE_ANIM_NOTIFY(this, AnimNotifyEvent, End);
-					AnimNotifyEvent.NotifyStateClass->NotifyEnd(SkelMeshComp, NotifyMontage);
+					AnimNotifyEvent.NotifyStateClass->NotifyEnd(SkelMeshComp, NotifyMontage, EventReference);
 				}
 
 				if (ActiveAnimNotifyState.IsValidIndex(Index))
@@ -3384,6 +3420,51 @@ float UAnimInstance::GetRelevantAnimTimeFraction(int32 MachineIndex, int32 State
 	return GetProxyOnAnyThread<FAnimInstanceProxy>().GetRelevantAnimTimeFraction(MachineIndex, StateIndex);
 }
 
+bool UAnimInstance::WasAnimNotifyStateActiveInAnyState(TSubclassOf<UAnimNotifyState> AnimNotifyStateType)
+{
+	return GetProxyOnAnyThread<FAnimInstanceProxy>().WasAnimNotifyStateActiveInAnyState(AnimNotifyStateType);
+}
+
+bool UAnimInstance::WasAnimNotifyStateActiveInStateMachine(int32 MachineIndex, TSubclassOf<UAnimNotifyState> AnimNotifyStateType)
+{
+	return GetProxyOnAnyThread<FAnimInstanceProxy>().WasAnimNotifyStateActiveInStateMachine(MachineIndex, AnimNotifyStateType);
+}
+
+bool UAnimInstance::WasAnimNotifyStateActiveInSourceState(int32 MachineIndex, int32 StateIndex, TSubclassOf<UAnimNotifyState> AnimNotifyStateType)
+{
+	return GetProxyOnAnyThread<FAnimInstanceProxy>().WasAnimNotifyStateActiveInSourceState(MachineIndex, StateIndex, AnimNotifyStateType);
+}
+
+bool UAnimInstance::WasAnimNotifyTriggeredInSourceState(int32 MachineIndex, int32 StateIndex,  TSubclassOf<UAnimNotify> AnimNotifyType)
+{
+	return GetProxyOnAnyThread<FAnimInstanceProxy>().WasAnimNotifyTriggeredInSourceState(MachineIndex, StateIndex, AnimNotifyType);
+}
+
+bool UAnimInstance::WasAnimNotifyNameTriggeredInSourceState(int32 MachineIndex, int32 StateIndex, FName NotifyName)
+{
+	return GetProxyOnAnyThread<FAnimInstanceProxy>().WasAnimNotifyNameTriggeredInSourceState(MachineIndex, StateIndex, NotifyName);
+}
+
+bool UAnimInstance::WasAnimNotifyTriggeredInStateMachine(int32 MachineIndex, TSubclassOf<UAnimNotify> AnimNotifyType)
+{
+	return GetProxyOnAnyThread<FAnimInstanceProxy>().WasAnimNotifyTriggeredInStateMachine(MachineIndex, AnimNotifyType);
+}
+
+bool UAnimInstance::WasAnimNotifyNameTriggeredInStateMachine(int32 MachineIndex, FName NotifyName)
+{
+	return GetProxyOnAnyThread<FAnimInstanceProxy>().WasAnimNotifyNameTriggeredInStateMachine(MachineIndex, NotifyName);
+}
+
+bool UAnimInstance::WasAnimNotifyTriggeredInAnyState(TSubclassOf<UAnimNotify> AnimNotifyType)
+{
+	return GetProxyOnAnyThread<FAnimInstanceProxy>().WasAnimNotifyTriggeredInAnyState(AnimNotifyType);
+}
+
+bool UAnimInstance::WasAnimNotifyNameTriggeredInAnyState(FName NotifyName)
+{
+	return GetProxyOnAnyThread<FAnimInstanceProxy>().WasAnimNotifyNameTriggeredInAnyState(NotifyName);
+}
+
 const FAnimNode_StateMachine* UAnimInstance::GetStateMachineInstance(int32 MachineIndex) const
 {
 	return GetProxyOnAnyThread<FAnimInstanceProxy>().GetStateMachineInstance(MachineIndex);
diff --git a/Engine/Source/Runtime/Engine/Private/Animation/AnimInstanceProxy.cpp b/Engine/Source/Runtime/Engine/Private/Animation/AnimInstanceProxy.cpp
--- a/Engine/Source/Runtime/Engine/Private/Animation/AnimInstanceProxy.cpp
+++ b/Engine/Source/Runtime/Engine/Private/Animation/AnimInstanceProxy.cpp
@@ -5,6 +5,7 @@
 #include "Animation/PoseAsset.h"
 #include "AnimationRuntime.h"
 #include "Animation/BlendSpace.h"
+#include "Animation/AnimNotifies/AnimNotifyState.h"
 #include "AnimationUtils.h"
 #include "Logging/MessageLog.h"
 #include "Animation/AnimNode_AssetPlayerBase.h"
@@ -22,6 +23,8 @@
 #include "Animation/AnimMontage.h"
 #include "Animation/AnimInstance.h"
 #include "Animation/AnimSyncScope.h"
+#include "Animation/AnimNotifyStateMachineInspectionLibrary.h"
+#include "Animation/MirrorDataTable.h"
 
 #define DO_ANIMSTAT_PROCESSING(StatName) DEFINE_STAT(STAT_ ## StatName)
 #include "Animation/AnimMTStats.h"
@@ -345,6 +348,13 @@ void FAnimInstanceProxy::Uninitialize(UAnimInstance* InAnimInstance)
 	DefaultLinkedInstanceInputNode = nullptr;
 	ResetAnimationCurves();
 	MaterialParametersToClear.Reset();
+	ActiveAnimNotifiesSinceLastTick.Reset();
+}
+
+void FAnimInstanceProxy::UpdateActiveAnimNotifiesSinceLastTick(const FAnimNotifyQueue& AnimInstanceQueue)
+{
+	ActiveAnimNotifiesSinceLastTick.Reset(AnimInstanceQueue.AnimNotifies.Num());
+	ActiveAnimNotifiesSinceLastTick.Append(AnimInstanceQueue.AnimNotifies);
 }
 
 void FAnimInstanceProxy::PreUpdate(UAnimInstance* InAnimInstance, float DeltaSeconds)
@@ -2252,6 +2262,160 @@ const TArray<FMontageEvaluationState>& FAnimInstanceProxy::GetMontageEvaluationD
 	check(MainMontageEvaluationData);
 	return *MainMontageEvaluationData;
 }
+bool FAnimInstanceProxy::WasAnimNotifyStateActiveInAnyState(TSubclassOf<UAnimNotifyState> AnimNotifyStateType) const
+{
+	for (const FAnimNotifyEventReference& Ref : ActiveAnimNotifiesSinceLastTick)
+	{
+		const FAnimNotifyEvent* NotifyEvent = Ref.GetNotify();
+		if (NotifyEvent && NotifyEvent->NotifyStateClass 
+			&& NotifyEvent->NotifyStateClass->IsA(AnimNotifyStateType))
+		{
+			return true;
+		}
+	}
+
+	return false;
+}
+
+bool FAnimInstanceProxy::WasAnimNotifyStateActiveInStateMachine(int32 MachineIndex, TSubclassOf<UAnimNotifyState> AnimNotifyStateType) const
+{
+	for (const FAnimNotifyEventReference& Ref : ActiveAnimNotifiesSinceLastTick)
+	{
+		const FAnimNotifyEvent* NotifyEvent = Ref.GetNotify();
+		if (NotifyEvent && NotifyEvent->NotifyStateClass && NotifyEvent->NotifyStateClass->IsA(AnimNotifyStateType))
+		{
+			return UAnimNotifyStateMachineInspectionLibrary::IsStateMachineInEventContext(Ref, MachineIndex);
+		}
+	}
+	return false;
+}
+
+bool FAnimInstanceProxy::WasAnimNotifyStateActiveInSourceState(int32 MachineIndex, int32 StateIndex, TSubclassOf<UAnimNotifyState> AnimNotifyStateType) const
+{
+	for (const FAnimNotifyEventReference& Ref : ActiveAnimNotifiesSinceLastTick)
+	{
+		const FAnimNotifyEvent* NotifyEvent = Ref.GetNotify();
+		if (NotifyEvent && NotifyEvent->NotifyStateClass && NotifyEvent->NotifyStateClass->IsA(AnimNotifyStateType))
+		{
+			return UAnimNotifyStateMachineInspectionLibrary::IsStateInStateMachineInEventContext(Ref, StateIndex, MachineIndex);
+		}
+	}
+	return false;
+}
+
+bool FAnimInstanceProxy::WasAnimNotifyTriggeredInSourceState(int32 MachineIndex, int32 StateIndex, TSubclassOf<UAnimNotify> AnimNotifyType) const
+{
+	for (const FAnimNotifyEventReference& Ref : ActiveAnimNotifiesSinceLastTick)
+	{
+		const FAnimNotifyEvent* NotifyEvent = Ref.GetNotify();
+		if (NotifyEvent && NotifyEvent->Notify && NotifyEvent->Notify->IsA(AnimNotifyType))
+		{
+			return UAnimNotifyStateMachineInspectionLibrary::IsStateInStateMachineInEventContext(Ref, StateIndex, MachineIndex);
+		}
+	}
+	return false;
+}
+
+bool FAnimInstanceProxy::WasAnimNotifyNameTriggeredInSourceState(int32 MachineIndex, int32 StateIndex, FName NotifyName) const
+{
+	for (const FAnimNotifyEventReference& Ref : ActiveAnimNotifiesSinceLastTick)
+	{
+		const FAnimNotifyEvent* NotifyEvent = Ref.GetNotify();
+		if (NotifyEvent)
+		{
+			FName LookupName = NotifyEvent->NotifyName;
+			if (Ref.GetMirrorDataTable())
+			{
+				const FName* MirroredName = Ref.GetMirrorDataTable()->AnimNotifyToMirrorAnimNotifyMap.Find(LookupName);
+				if (MirroredName)
+				{
+					LookupName = *MirroredName; 
+				}
+			}
+			if(LookupName == NotifyName)
+			{
+				return UAnimNotifyStateMachineInspectionLibrary::IsStateInStateMachineInEventContext(Ref, StateIndex, MachineIndex);
+			}
+		}
+	}
+	return false;
+}
+
+bool FAnimInstanceProxy::WasAnimNotifyTriggeredInStateMachine(int32 MachineIndex,TSubclassOf<UAnimNotify> AnimNotifyType) const
+{
+	for (const FAnimNotifyEventReference& Ref : ActiveAnimNotifiesSinceLastTick)
+	{
+		const FAnimNotifyEvent* NotifyEvent = Ref.GetNotify();
+		if (NotifyEvent && NotifyEvent->Notify && NotifyEvent->Notify->IsA(AnimNotifyType))
+		{
+			return UAnimNotifyStateMachineInspectionLibrary::IsStateMachineInEventContext(Ref, MachineIndex);
+		}
+	}
+	return false;
+}
+
+bool FAnimInstanceProxy::WasAnimNotifyTriggeredInAnyState(TSubclassOf<UAnimNotify> AnimNotifyType) const
+{
+	for (const FAnimNotifyEventReference& Ref : ActiveAnimNotifiesSinceLastTick)
+	{
+		const FAnimNotifyEvent* NotifyEvent = Ref.GetNotify();
+		if (NotifyEvent && NotifyEvent->Notify && NotifyEvent->Notify->IsA(AnimNotifyType))
+		{
+			return true;
+		}
+	}
+	return false;
+}
+
+bool FAnimInstanceProxy::WasAnimNotifyNameTriggeredInAnyState(FName NotifyName) const
+{
+	for (const FAnimNotifyEventReference& Ref : ActiveAnimNotifiesSinceLastTick)
+	{
+		const FAnimNotifyEvent* NotifyEvent = Ref.GetNotify();
+		if (NotifyEvent)
+		{
+			FName LookupName = NotifyEvent->NotifyName;
+			if (Ref.GetMirrorDataTable())
+			{
+				const FName* MirroredName = Ref.GetMirrorDataTable()->AnimNotifyToMirrorAnimNotifyMap.Find(LookupName);
+				if (MirroredName)
+				{
+					LookupName = *MirroredName; 
+				}
+			}
+			if(LookupName == NotifyName)
+			{
+				return true;
+			}
+		}
+	}
+	return false; 
+}
+
+bool FAnimInstanceProxy::WasAnimNotifyNameTriggeredInStateMachine(int32 MachineIndex, FName NotifyName)
+{
+	for (const FAnimNotifyEventReference& Ref : ActiveAnimNotifiesSinceLastTick)
+	{
+		const FAnimNotifyEvent* NotifyEvent = Ref.GetNotify();
+		if (NotifyEvent)
+		{
+			FName LookupName = NotifyEvent->NotifyName;
+			if (Ref.GetMirrorDataTable())
+			{
+				const FName* MirroredName = Ref.GetMirrorDataTable()->AnimNotifyToMirrorAnimNotifyMap.Find(LookupName);
+				if (MirroredName)
+				{
+					LookupName = *MirroredName; 
+				}
+			}
+			if(LookupName == NotifyName)
+			{
+				return UAnimNotifyStateMachineInspectionLibrary::IsStateMachineInEventContext(Ref, MachineIndex);
+			}
+		}
+	}
+	return false;
+}
 
 const FAnimNode_AssetPlayerBase* FAnimInstanceProxy::GetRelevantAssetPlayerFromState(int32 MachineIndex, int32 StateIndex) const
 {
@@ -2547,6 +2711,27 @@ int32 FAnimInstanceProxy::GetStateMachineIndex(FName MachineName) const
 	return INDEX_NONE;
 }
 
+int32  FAnimInstanceProxy::GetStateMachineIndex(FAnimNode_StateMachine* StateMachine) const
+{
+	if (AnimClassInterface)
+	{
+		const TArray<FStructProperty*>& AnimNodeProperties = AnimClassInterface->GetAnimNodeProperties();
+		for (int32 MachineIndex = 0; MachineIndex < AnimNodeProperties.Num(); MachineIndex++)
+		{
+			FStructProperty* Property = AnimNodeProperties[AnimNodeProperties.Num() - 1 - MachineIndex];
+			if (Property && Property->Struct->IsChildOf(FAnimNode_StateMachine::StaticStruct()))
+			{
+				FAnimNode_StateMachine* CurStateMachine = Property->ContainerPtrToValuePtr<FAnimNode_StateMachine>(AnimInstanceObject);
+				if (CurStateMachine == StateMachine)
+				{
+					return MachineIndex;
+				}
+			}
+		}
+	}
+	return INDEX_NONE; 	
+}
+
 void FAnimInstanceProxy::GetStateMachineIndexAndDescription(FName InMachineName, int32& OutMachineIndex, const FBakedAnimationStateMachine** OutMachineDescription) const
 {
 	if (AnimClassInterface)
diff --git a/Engine/Source/Runtime/Engine/Private/Animation/AnimMontage.cpp b/Engine/Source/Runtime/Engine/Private/Animation/AnimMontage.cpp
--- a/Engine/Source/Runtime/Engine/Private/Animation/AnimMontage.cpp
+++ b/Engine/Source/Runtime/Engine/Private/Animation/AnimMontage.cpp
@@ -1172,7 +1172,7 @@ void UAnimMontage::TickAssetPlayer(FAnimTickRecord& Instance, struct FAnimNotify
 
 				if (MarkerTickRecord->IsValid(Instance.bLooping))
 				{
-					MarkerTickContext.SetMarkerSyncStartPosition(GetMarkerSyncPositionfromMarkerIndicies(MarkerTickRecord->PreviousMarker.MarkerIndex, MarkerTickRecord->NextMarker.MarkerIndex, PreviousTime, nullptr));
+					MarkerTickContext.SetMarkerSyncStartPosition(GetMarkerSyncPositionFromMarkerIndicies(MarkerTickRecord->PreviousMarker.MarkerIndex, MarkerTickRecord->NextMarker.MarkerIndex, PreviousTime, nullptr));
 
 				}
 				else
@@ -1181,14 +1181,14 @@ void UAnimMontage::TickAssetPlayer(FAnimTickRecord& Instance, struct FAnimNotify
 					FMarkerPair PreviousMarker;
 					FMarkerPair NextMarker;
 					GetMarkerIndicesForTime(PreviousTime, false, MarkerTickContext.GetValidMarkerNames(), PreviousMarker, NextMarker);
-					MarkerTickContext.SetMarkerSyncStartPosition(GetMarkerSyncPositionfromMarkerIndicies(PreviousMarker.MarkerIndex, NextMarker.MarkerIndex, PreviousTime, nullptr));
+					MarkerTickContext.SetMarkerSyncStartPosition(GetMarkerSyncPositionFromMarkerIndicies(PreviousMarker.MarkerIndex, NextMarker.MarkerIndex, PreviousTime, nullptr));
 				}
 
 				// @todo this won't work well once we start jumping
 				// only thing is that passed markers won't work in this frame. To do that, I have to figure out how it jumped from where to where, 
 				GetMarkerIndicesForTime(CurrentTime, false, MarkerTickContext.GetValidMarkerNames(), MarkerTickRecord->PreviousMarker, MarkerTickRecord->NextMarker);
 				bRecordNeedsResetting = false; // we have updated it now, no need to reset
-				MarkerTickContext.SetMarkerSyncEndPosition(GetMarkerSyncPositionfromMarkerIndicies(MarkerTickRecord->PreviousMarker.MarkerIndex, MarkerTickRecord->NextMarker.MarkerIndex, CurrentTime, nullptr));
+				MarkerTickContext.SetMarkerSyncEndPosition(GetMarkerSyncPositionFromMarkerIndicies(MarkerTickRecord->PreviousMarker.MarkerIndex, MarkerTickRecord->NextMarker.MarkerIndex, CurrentTime, nullptr));
 
 				MarkerTickContext.MarkersPassedThisTick = *Instance.Montage.MarkersPassedThisTick;
 
@@ -1261,9 +1261,9 @@ void UAnimMontage::GetMarkerIndicesForTime(float CurrentTime, bool bLooping, con
 	MarkerData.GetMarkerIndicesForTime(CurrentTime, bLooping, ValidMarkerNames, OutPrevMarker, OutNextMarker, GetPlayLength());
 }
 
-FMarkerSyncAnimPosition UAnimMontage::GetMarkerSyncPositionfromMarkerIndicies(int32 PrevMarker, int32 NextMarker, float CurrentTime, const UMirrorDataTable* MirrorTable) const
+FMarkerSyncAnimPosition UAnimMontage::GetMarkerSyncPositionFromMarkerIndicies(int32 PrevMarker, int32 NextMarker, float CurrentTime, const UMirrorDataTable* MirrorTable) const
 {
-	return MarkerData.GetMarkerSyncPositionfromMarkerIndicies(PrevMarker, NextMarker, CurrentTime, GetPlayLength(), MirrorTable);
+	return MarkerData.GetMarkerSyncPositionFromMarkerIndicies(PrevMarker, NextMarker, CurrentTime, GetPlayLength(), MirrorTable);
 }
 
 void UAnimMontage::InvalidateRecursiveAsset()
@@ -2532,26 +2532,27 @@ void FAnimMontageInstance::HandleEvents(float PreviousTrackPos, float CurrentTra
 	// now get active Notifies based on how it advanced
 	if (AnimInstance.IsValid())
 	{
-		TArray<FAnimNotifyEventReference> NotitfyRefs;
 		TMap<FName, TArray<FAnimNotifyEventReference>> NotifyMap;
-
+		FAnimTickRecord TickRecord;
+		FAnimNotifyContext NotifyContext(TickRecord);
 		// We already break up AnimMontage update to handle looping, so we guarantee that PreviousPos and CurrentPos are contiguous.
-		Montage->GetAnimNotifiesFromDeltaPositions(PreviousTrackPos, CurrentTrackPos, NotitfyRefs);
+		Montage->GetAnimNotifiesFromDeltaPositions(PreviousTrackPos, CurrentTrackPos, NotifyContext);
 
 		// For Montage only, remove notifies marked as 'branching points'. They are not queued and are handled separately.
-		Montage->FilterOutNotifyBranchingPoints(NotitfyRefs);
+		Montage->FilterOutNotifyBranchingPoints(NotifyContext.ActiveNotifies);
 
 		// now trigger notifies for all animations within montage
 		// we'll do this for all slots for now
 		for (auto SlotTrack = Montage->SlotAnimTracks.CreateIterator(); SlotTrack; ++SlotTrack)
 		{
 			TArray<FAnimNotifyEventReference>& MapNotifies = NotifyMap.FindOrAdd(SlotTrack->SlotName);
-
+			PRAGMA_DISABLE_DEPRECATION_WARNINGS
 			SlotTrack->AnimTrack.GetAnimNotifiesFromTrackPositions(PreviousTrackPos, CurrentTrackPos, MapNotifies);
+			PRAGMA_ENABLE_DEPRECATION_WARNINGS
 		}
 
 		// Queue all these notifies.
-		AnimInstance->NotifyQueue.AddAnimNotifies(NotitfyRefs, NotifyWeight);
+		AnimInstance->NotifyQueue.AddAnimNotifies(NotifyContext.ActiveNotifies, NotifyWeight);
 		AnimInstance->NotifyQueue.AddAnimNotifies(NotifyMap, NotifyWeight);
 	}
 
diff --git a/Engine/Source/Runtime/Engine/Private/Animation/AnimNodeMessages.cpp b/Engine/Source/Runtime/Engine/Private/Animation/AnimNodeMessages.cpp
--- a/Engine/Source/Runtime/Engine/Private/Animation/AnimNodeMessages.cpp
+++ b/Engine/Source/Runtime/Engine/Private/Animation/AnimNodeMessages.cpp
@@ -5,8 +5,9 @@
 #include "Animation/AnimInstanceProxy.h"
 
 namespace UE { namespace Anim {
-
 const FName IGraphMessage::BaseTypeName("IGraphMessage");
+	
+const FName IAnimNotifyEventContextDataInterface::BaseTypeName("IAnimNotifyEventContextDataInterface");
 
 FScopedGraphTag::FScopedGraphTag(const FAnimationBaseContext& InContext, FName InTag)
 	: SharedContext(InContext.SharedContext)
@@ -188,4 +189,22 @@ void FMessageStack::CopyForCachedUpdate(const FMessageStack& InStack)
 	}
 }
 
+void FMessageStack::MakeEventContextData(TArray<TSharedPtr<const IAnimNotifyEventContextDataInterface>>& ContextData) const
+{
+	ContextData.Empty();
+	for (const TPair<FName, FMessageStackEntry>& MessageStackPair : MessageStacks)
+	{
+		const FMessageStackEntry& StackEntry = MessageStackPair.Value;
+		if(StackEntry.Num())
+		{
+			const FMessageEntry& MessageEntry = StackEntry.Top();
+			TSharedPtr<const IAnimNotifyEventContextDataInterface> NotifyData = MessageEntry.Message->MakeEventContextData();
+			if(NotifyData)
+			{
+				ContextData.Emplace(NotifyData);
+			}
+		}
+	}
+}
+	
 }}	// namespace UE::Anim
\ No newline at end of file
diff --git a/Engine/Source/Runtime/Engine/Private/Animation/AnimNode_AssetPlayerBase.cpp b/Engine/Source/Runtime/Engine/Private/Animation/AnimNode_AssetPlayerBase.cpp
--- a/Engine/Source/Runtime/Engine/Private/Animation/AnimNode_AssetPlayerBase.cpp
+++ b/Engine/Source/Runtime/Engine/Private/Animation/AnimNode_AssetPlayerBase.cpp
@@ -40,6 +40,10 @@ void FAnimNode_AssetPlayerBase::CreateTickRecordForNode(const FAnimationUpdateCo
 
 	UE::Anim::FAnimSyncParams SyncParams(GroupNameToUse, SyncGroupRole, MethodToUse);
 	FAnimTickRecord TickRecord(Sequence, bLooping, PlayRate, FinalBlendWeight, /*inout*/ InternalTimeAccumulator, MarkerTickRecord);
+	if(Context.GetSharedContext())
+	{
+		Context.GetSharedContext()->MessageStack.MakeEventContextData(TickRecord.ContextData);
+	}
 	TickRecord.RootMotionWeightModifier = Context.GetRootMotionWeightModifier();
 
 	SyncScope.AddTickRecord(TickRecord, SyncParams, UE::Anim::FAnimSyncDebugInfo(Context));
diff --git a/Engine/Source/Runtime/Engine/Private/Animation/AnimNode_StateMachine.cpp b/Engine/Source/Runtime/Engine/Private/Animation/AnimNode_StateMachine.cpp
--- a/Engine/Source/Runtime/Engine/Private/Animation/AnimNode_StateMachine.cpp
+++ b/Engine/Source/Runtime/Engine/Private/Animation/AnimNode_StateMachine.cpp
@@ -10,6 +10,7 @@
 #include "Animation/AnimNode_LinkedAnimLayer.h"
 #include "Animation/AnimInstance.h"
 #include "Animation/AnimTrace.h"
+#include "Animation/ActiveStateMachineScope.h"
 
 #if WITH_EDITORONLY_DATA
 #include "Animation/AnimBlueprintGeneratedClass.h"
@@ -345,6 +346,7 @@ void FAnimNode_StateMachine::LogInertializationRequestError(const FAnimationUpda
 // Temporarily turned off while we track down and fix https://jira.ol.epicgames.net/browse/OR-17066
 TAutoConsoleVariable<int32> CVarAnimStateMachineRelevancyReset(TEXT("a.AnimNode.StateMachine.EnableRelevancyReset"), 1, TEXT("Reset State Machine when it becomes relevant"));
 
+
 void FAnimNode_StateMachine::Update_AnyThread(const FAnimationUpdateContext& Context)
 {
 	Context.AnimInstanceProxy->RecordMachineWeight(StateMachineIndexInClass, Context.GetFinalBlendWeight());
@@ -481,6 +483,8 @@ void FAnimNode_StateMachine::Update_AnyThread(const FAnimationUpdateContext& Con
 	}
 	while (bFoundValidTransition && (TransitionCountThisFrame < MaxTransitionsPerFrame));
 
+	UE::Anim::TOptionalScopedGraphMessage<UE::Anim::FActiveStateMachineScope> Message(bCreateNotifyMetaData, Context, Context, this, CurrentState);
+	
 	if (bFirstUpdate)
 	{
 		if (bSkipFirstUpdateTransition)
diff --git a/Engine/Source/Runtime/Engine/Private/Animation/AnimNotify.cpp b/Engine/Source/Runtime/Engine/Private/Animation/AnimNotify.cpp
--- a/Engine/Source/Runtime/Engine/Private/Animation/AnimNotify.cpp
+++ b/Engine/Source/Runtime/Engine/Private/Animation/AnimNotify.cpp
@@ -23,16 +23,24 @@ UAnimNotify::UAnimNotify(const FObjectInitializer& ObjectInitializer)
 }
 
 void UAnimNotify::Notify(class USkeletalMeshComponent* MeshComp, class UAnimSequenceBase* Animation)
+{
+}
+
+void UAnimNotify::Notify(USkeletalMeshComponent* MeshComp, UAnimSequenceBase* Animation, const FAnimNotifyEventReference& EventReference)
 {
 	USkeletalMeshComponent* PrevContext = MeshContext;
 	MeshContext = MeshComp;
-	Received_Notify(MeshComp, Animation);
+	PRAGMA_DISABLE_DEPRECATION_WARNINGS
+	Notify(MeshComp, Animation);
+	PRAGMA_ENABLE_DEPRECATION_WARNINGS
+	Received_Notify(MeshComp, Animation, EventReference);
 	MeshContext = PrevContext;
 }
 
 void UAnimNotify::BranchingPointNotify(FBranchingPointNotifyPayload& BranchingPointPayload)
 {
-	Notify(BranchingPointPayload.SkelMeshComponent, BranchingPointPayload.SequenceAsset);
+	const FAnimNotifyEventReference EventReference;
+	Notify(BranchingPointPayload.SkelMeshComponent, BranchingPointPayload.SequenceAsset, EventReference);
 }
 
 class UWorld* UAnimNotify::GetWorld() const
diff --git a/Engine/Source/Runtime/Engine/Private/Animation/AnimNotifyMirrorInspectionLibrary.cpp b/Engine/Source/Runtime/Engine/Private/Animation/AnimNotifyMirrorInspectionLibrary.cpp
new file mode 100644
--- /dev/null
+++ b/Engine/Source/Runtime/Engine/Private/Animation/AnimNotifyMirrorInspectionLibrary.cpp
@@ -0,0 +1,29 @@
+// Copyright Epic Games, Inc. All Rights Reserved.
+
+#include "Animation/AnimNotifyMirrorInspectionLibrary.h"
+#include "Animation/AnimNotifyQueue.h"
+#include "Animation/AnimationAsset.h"
+#include "Animation/AnimClassInterface.h"
+#include "Animation/AnimNode_StateMachine.h"
+#include "Animation/MirrorSyncScope.h"
+#include "Animation/ActiveStateMachineScope.h"
+
+bool UAnimNotifyMirrorInspectionLibrary::IsTriggeredByMirroredAnimation(const FAnimNotifyEventReference& EventReference)
+{
+	const UE::Anim::FAnimNotifyMirrorContext * MirrorContext = EventReference.GetContextData<UE::Anim::FAnimNotifyMirrorContext>();
+	if (MirrorContext)
+	{
+		return MirrorContext->bAnimationMirrored; 
+	}
+	return false; 
+}
+
+const UMirrorDataTable* UAnimNotifyMirrorInspectionLibrary::GetMirrorDataTable(const FAnimNotifyEventReference& EventReference)
+{
+	const UE::Anim::FAnimNotifyMirrorContext * MirrorContext =  EventReference.GetContextData<UE::Anim::FAnimNotifyMirrorContext>();
+	if (MirrorContext)
+	{
+		return MirrorContext->MirrorTable.Get(); 
+	}
+	return nullptr;
+}
\ No newline at end of file
diff --git a/Engine/Source/Runtime/Engine/Private/Animation/AnimNotifyQueue.cpp b/Engine/Source/Runtime/Engine/Private/Animation/AnimNotifyQueue.cpp
--- a/Engine/Source/Runtime/Engine/Private/Animation/AnimNotifyQueue.cpp
+++ b/Engine/Source/Runtime/Engine/Private/Animation/AnimNotifyQueue.cpp
@@ -2,7 +2,6 @@
 
 #include "Animation/AnimNotifyQueue.h"
 #include "Animation/AnimInstanceProxy.h"
-#include "Animation/AnimTypes.h"
 
 bool operator==(const FAnimNotifyEventReference& Lhs, const FAnimNotifyEvent& Rhs)
 {
@@ -13,6 +12,11 @@ bool operator==(const FAnimNotifyEventReference& Lhs, const FAnimNotifyEvent& Rh
 	return false;
 }
 
+void FAnimNotifyEventReference::SetContextDataInterfaces(const TArray<TSharedPtr<const UE::Anim::IAnimNotifyEventContextDataInterface>>& InContextData)
+{
+	ContextData = InContextData; 
+}
+
 bool FAnimNotifyQueue::PassesFiltering(const FAnimNotifyEvent* Notify) const
 {
 	switch (Notify->NotifyFilterType)
diff --git a/Engine/Source/Runtime/Engine/Private/Animation/AnimNotifyState.cpp b/Engine/Source/Runtime/Engine/Private/Animation/AnimNotifyState.cpp
--- a/Engine/Source/Runtime/Engine/Private/Animation/AnimNotifyState.cpp
+++ b/Engine/Source/Runtime/Engine/Private/Animation/AnimNotifyState.cpp
@@ -18,35 +18,58 @@ UAnimNotifyState::UAnimNotifyState(const FObjectInitializer& ObjectInitializer)
 	bIsNativeBranchingPoint = false;
 }
 
-
 void UAnimNotifyState::NotifyBegin(class USkeletalMeshComponent * MeshComp, class UAnimSequenceBase * Animation, float TotalDuration)
 {
-	Received_NotifyBegin(MeshComp, Animation, TotalDuration);
 }
 
 void UAnimNotifyState::NotifyTick(class USkeletalMeshComponent * MeshComp, class UAnimSequenceBase * Animation, float FrameDeltaTime)
 {
-	Received_NotifyTick(MeshComp, Animation, FrameDeltaTime);
 }
 
 void UAnimNotifyState::NotifyEnd(class USkeletalMeshComponent * MeshComp, class UAnimSequenceBase * Animation)
 {
-	Received_NotifyEnd(MeshComp, Animation);
+}
+
+void UAnimNotifyState::NotifyBegin(USkeletalMeshComponent* MeshComp, UAnimSequenceBase* Animation, float TotalDuration, const FAnimNotifyEventReference& EventReference)
+{
+	PRAGMA_DISABLE_DEPRECATION_WARNINGS
+	NotifyBegin(MeshComp, Animation, TotalDuration);
+	PRAGMA_ENABLE_DEPRECATION_WARNINGS
+	Received_NotifyBegin(MeshComp, Animation, TotalDuration, EventReference);
+}
+
+void UAnimNotifyState::NotifyTick(USkeletalMeshComponent* MeshComp, UAnimSequenceBase* Animation, float FrameDeltaTime, const FAnimNotifyEventReference& EventReference)
+{
+	PRAGMA_DISABLE_DEPRECATION_WARNINGS
+	NotifyTick(MeshComp, Animation, FrameDeltaTime);
+	PRAGMA_ENABLE_DEPRECATION_WARNINGS
+	Received_NotifyTick(MeshComp, Animation, FrameDeltaTime, EventReference);
+}
+
+void UAnimNotifyState::NotifyEnd(USkeletalMeshComponent* MeshComp, UAnimSequenceBase* Animation, const FAnimNotifyEventReference& EventReference)
+{
+	PRAGMA_DISABLE_DEPRECATION_WARNINGS
+	NotifyEnd(MeshComp, Animation);
+	PRAGMA_ENABLE_DEPRECATION_WARNINGS
+	Received_NotifyEnd(MeshComp, Animation, EventReference);
 }
 
 void UAnimNotifyState::BranchingPointNotifyBegin(FBranchingPointNotifyPayload& BranchingPointPayload)
 {
-	NotifyBegin(BranchingPointPayload.SkelMeshComponent, BranchingPointPayload.SequenceAsset, BranchingPointPayload.NotifyEvent ? BranchingPointPayload.NotifyEvent->GetDuration() : 0.f);
+	const FAnimNotifyEventReference EventReference;
+	NotifyBegin(BranchingPointPayload.SkelMeshComponent, BranchingPointPayload.SequenceAsset, BranchingPointPayload.NotifyEvent ? BranchingPointPayload.NotifyEvent->GetDuration() : 0.f, EventReference);
 }
 
 void UAnimNotifyState::BranchingPointNotifyTick(FBranchingPointNotifyPayload& BranchingPointPayload, float FrameDeltaTime)
 {
-	NotifyTick(BranchingPointPayload.SkelMeshComponent, BranchingPointPayload.SequenceAsset, FrameDeltaTime);
+	const FAnimNotifyEventReference EventReference;
+	NotifyTick(BranchingPointPayload.SkelMeshComponent, BranchingPointPayload.SequenceAsset, FrameDeltaTime, EventReference);
 }
 
 void UAnimNotifyState::BranchingPointNotifyEnd(FBranchingPointNotifyPayload& BranchingPointPayload)
 {
-	NotifyEnd(BranchingPointPayload.SkelMeshComponent, BranchingPointPayload.SequenceAsset);
+	const FAnimNotifyEventReference EventReference;
+	NotifyEnd(BranchingPointPayload.SkelMeshComponent, BranchingPointPayload.SequenceAsset, EventReference);
 }
 
 /// @cond DOXYGEN_WARNINGS
diff --git a/Engine/Source/Runtime/Engine/Private/Animation/AnimNotifyStateMachineInspectionLibrary.cpp b/Engine/Source/Runtime/Engine/Private/Animation/AnimNotifyStateMachineInspectionLibrary.cpp
new file mode 100644
--- /dev/null
+++ b/Engine/Source/Runtime/Engine/Private/Animation/AnimNotifyStateMachineInspectionLibrary.cpp
@@ -0,0 +1,86 @@
+// Copyright Epic Games, Inc. All Rights Reserved.
+
+#include "Animation/AnimNotifyStateMachineInspectionLibrary.h"
+#include "Animation/AnimNotifyQueue.h"
+#include "Animation/AnimationAsset.h"
+#include "Animation/AnimClassInterface.h"
+#include "Animation/AnimNode_StateMachine.h"
+#include "Animation/MirrorSyncScope.h"
+#include "Animation/ActiveStateMachineScope.h"
+
+bool UAnimNotifyStateMachineInspectionLibrary::IsStateMachineInEventContext(const FAnimNotifyEventReference& Reference, int32 StateMachineIndex)
+{
+	const UE::Anim::FAnimNotifyStateMachineContext * StateMachineContext = Reference.GetContextData<UE::Anim::FAnimNotifyStateMachineContext>();
+	if (StateMachineContext)
+	{
+		if (StateMachineContext->IsStateMachineInContext(StateMachineIndex))
+		{
+			return true; 
+		}
+	}
+	return false; 
+}
+
+bool UAnimNotifyStateMachineInspectionLibrary::IsStateInStateMachineInEventContext(const FAnimNotifyEventReference& Reference, int32 StateMachineIndex, int32 StateIndex) 
+{
+	const UE::Anim::FAnimNotifyStateMachineContext * StateMachineContext = Reference.GetContextData<UE::Anim::FAnimNotifyStateMachineContext>();
+	if (StateMachineContext)
+	{
+		if (StateMachineContext->IsStateInStateMachineInContext(StateMachineIndex, StateIndex))
+		{
+			return true; 
+		}
+	}
+	return false;
+}
+
+bool UAnimNotifyStateMachineInspectionLibrary::IsTriggeredByStateMachine(const FAnimNotifyEventReference& EventReference, UAnimInstance* AnimInstance, FName StateMachineName)
+{
+	if (AnimInstance)
+	{
+		int32 MachineIndex;
+		const FBakedAnimationStateMachine* BakedAnimationStateMachine = nullptr;
+		AnimInstance->GetStateMachineIndexAndDescription(StateMachineName,MachineIndex, &BakedAnimationStateMachine);
+		if(MachineIndex != INDEX_NONE)
+		{
+			return IsStateMachineInEventContext(EventReference, MachineIndex);
+		}
+	}
+	return false; 
+}
+
+bool UAnimNotifyStateMachineInspectionLibrary::IsTriggeredByStateInStateMachine(const FAnimNotifyEventReference& EventReference, UAnimInstance* AnimInstance, FName StateMachineName, FName StateName)
+{
+	if (AnimInstance)
+	{
+		int32 MachineIndex;
+		const FBakedAnimationStateMachine* BakedAnimationStateMachine = nullptr;
+		AnimInstance->GetStateMachineIndexAndDescription(StateMachineName,MachineIndex, &BakedAnimationStateMachine);
+		if(MachineIndex != INDEX_NONE && BakedAnimationStateMachine)
+		{
+			int32 StateIndex = BakedAnimationStateMachine->FindStateIndex(StateName);
+			return IsStateInStateMachineInEventContext(EventReference, StateIndex, MachineIndex);
+		}
+	}
+	return false; 
+}
+
+bool  UAnimNotifyStateMachineInspectionLibrary::IsTriggeredByState(const FAnimNotifyEventReference& EventReference, UAnimInstance* AnimInstance, FName StateName)
+{
+	if (AnimInstance)
+	{
+		const UE::Anim::FAnimNotifyStateMachineContext * StateMachineContext = EventReference.GetContextData<UE::Anim::FAnimNotifyStateMachineContext>();
+		if (StateMachineContext)
+		{
+			for (const FEncounteredStateMachineStack::FStateMachineEntry& Entry: StateMachineContext->EncounteredStateMachines->StateStack)
+			{
+				const FAnimNode_StateMachine* StateMachineInstance = AnimInstance->GetStateMachineInstance(Entry.StateMachineIndex);
+				if(StateMachineInstance->GetStateInfo(Entry.StateIndex).StateName == StateName)
+				{
+					return true; 
+				}
+			}
+		}
+	}
+	return false;
+}
diff --git a/Engine/Source/Runtime/Engine/Private/Animation/AnimNotifyState_TimedParticleEffect.cpp b/Engine/Source/Runtime/Engine/Private/Animation/AnimNotifyState_TimedParticleEffect.cpp
--- a/Engine/Source/Runtime/Engine/Private/Animation/AnimNotifyState_TimedParticleEffect.cpp
+++ b/Engine/Source/Runtime/Engine/Private/Animation/AnimNotifyState_TimedParticleEffect.cpp
@@ -14,22 +14,35 @@ UAnimNotifyState_TimedParticleEffect::UAnimNotifyState_TimedParticleEffect(const
 }
 
 void UAnimNotifyState_TimedParticleEffect::NotifyBegin(USkeletalMeshComponent * MeshComp, class UAnimSequenceBase * Animation, float TotalDuration)
+{
+
+}
+
+
+void UAnimNotifyState_TimedParticleEffect::NotifyBegin(USkeletalMeshComponent * MeshComp, class UAnimSequenceBase * Animation, float TotalDuration, const FAnimNotifyEventReference& EventReference)
 {
 	// Only spawn if we've got valid params
 	if(ValidateParameters(MeshComp))
 	{
 		UParticleSystemComponent* NewComponent = UGameplayStatics::SpawnEmitterAttached(PSTemplate, MeshComp, SocketName, LocationOffset, RotationOffset, EAttachLocation::KeepRelativeOffset, !bDestroyAtEnd);
 	}
-
-	Received_NotifyBegin(MeshComp, Animation, TotalDuration);
+	Received_NotifyBegin(MeshComp, Animation, TotalDuration, EventReference);
 }
 
 void UAnimNotifyState_TimedParticleEffect::NotifyTick(USkeletalMeshComponent * MeshComp, class UAnimSequenceBase * Animation, float FrameDeltaTime)
 {
-	Received_NotifyTick(MeshComp, Animation, FrameDeltaTime);
+}
+
+void UAnimNotifyState_TimedParticleEffect::NotifyTick(USkeletalMeshComponent * MeshComp, class UAnimSequenceBase * Animation, float FrameDeltaTime, const FAnimNotifyEventReference& EventReference)
+{
+	Received_NotifyTick(MeshComp, Animation, FrameDeltaTime, EventReference);
 }
 
 void UAnimNotifyState_TimedParticleEffect::NotifyEnd(USkeletalMeshComponent * MeshComp, class UAnimSequenceBase * Animation)
+{
+}
+
+void UAnimNotifyState_TimedParticleEffect::NotifyEnd(USkeletalMeshComponent * MeshComp, class UAnimSequenceBase * Animation, const FAnimNotifyEventReference& EventReference)
 {
 	TArray<USceneComponent*> Children;
 	MeshComp->GetChildrenComponents(false, Children);
@@ -73,7 +86,7 @@ void UAnimNotifyState_TimedParticleEffect::NotifyEnd(USkeletalMeshComponent * Me
 		}
 	}
 
-	Received_NotifyEnd(MeshComp, Animation);
+	Received_NotifyEnd(MeshComp, Animation, EventReference);
 }
 
 bool UAnimNotifyState_TimedParticleEffect::ValidateParameters(USkeletalMeshComponent* MeshComp)
diff --git a/Engine/Source/Runtime/Engine/Private/Animation/AnimNotifyState_Trail.cpp b/Engine/Source/Runtime/Engine/Private/Animation/AnimNotifyState_Trail.cpp
--- a/Engine/Source/Runtime/Engine/Private/Animation/AnimNotifyState_Trail.cpp
+++ b/Engine/Source/Runtime/Engine/Private/Animation/AnimNotifyState_Trail.cpp
@@ -75,7 +75,11 @@ float UAnimNotifyState_Trail::GetCurveWidth(USkeletalMeshComponent* MeshComp) co
 	return Width;
 }
 
-void UAnimNotifyState_Trail::NotifyBegin(class USkeletalMeshComponent * MeshComp, class UAnimSequenceBase * Animation, float TotalDuration)
+void UAnimNotifyState_Trail::NotifyBegin(class USkeletalMeshComponent* MeshComp, class UAnimSequenceBase* Animation, float TotalDuration)
+{
+}
+
+void UAnimNotifyState_Trail::NotifyBegin(class USkeletalMeshComponent * MeshComp, class UAnimSequenceBase * Animation, float TotalDuration, const FAnimNotifyEventReference& EventReference)
 {
 	bool bError = ValidateInput(MeshComp);
 
@@ -195,10 +199,14 @@ void UAnimNotifyState_Trail::NotifyBegin(class USkeletalMeshComponent * MeshComp
 		}
 	}
 
-	Received_NotifyBegin(MeshComp, Animation, TotalDuration);
+	Received_NotifyBegin(MeshComp, Animation, TotalDuration, EventReference);
 }
 
-void UAnimNotifyState_Trail::NotifyTick(class USkeletalMeshComponent * MeshComp, class UAnimSequenceBase * Animation, float FrameDeltaTime)
+void UAnimNotifyState_Trail::NotifyTick(class USkeletalMeshComponent* MeshComp, class UAnimSequenceBase* Animation, float FrameDeltaTime)
+{
+}
+
+void UAnimNotifyState_Trail::NotifyTick(class USkeletalMeshComponent * MeshComp, class UAnimSequenceBase * Animation, float FrameDeltaTime, const FAnimNotifyEventReference& EventReference)
 {
 	bool bError = ValidateInput(MeshComp, true);
 
@@ -236,10 +244,14 @@ void UAnimNotifyState_Trail::NotifyTick(class USkeletalMeshComponent * MeshComp,
 		}
 	}
 
-	Received_NotifyTick(MeshComp, Animation, FrameDeltaTime);
+	Received_NotifyTick(MeshComp, Animation, FrameDeltaTime, EventReference);
+}
+
+void UAnimNotifyState_Trail::NotifyEnd(class USkeletalMeshComponent* MeshComp, class UAnimSequenceBase* Animation)
+{
 }
 
-void UAnimNotifyState_Trail::NotifyEnd(class USkeletalMeshComponent * MeshComp, class UAnimSequenceBase * Animation)
+void UAnimNotifyState_Trail::NotifyEnd(class USkeletalMeshComponent * MeshComp, class UAnimSequenceBase * Animation, const FAnimNotifyEventReference& EventReference)
 {
 	if (MeshComp->GetWorld()->GetNetMode() == NM_DedicatedServer)
 	{
@@ -262,7 +274,7 @@ void UAnimNotifyState_Trail::NotifyEnd(class USkeletalMeshComponent * MeshComp,
 		}
 	}
 
-	Received_NotifyEnd(MeshComp, Animation);
+	Received_NotifyEnd(MeshComp, Animation, EventReference);
 }
 
 UParticleSystemComponent* UAnimNotifyState_Trail::GetParticleSystemComponent(USkeletalMeshComponent* MeshComp) const
diff --git a/Engine/Source/Runtime/Engine/Private/Animation/AnimNotify_PauseClothingSimulation.cpp b/Engine/Source/Runtime/Engine/Private/Animation/AnimNotify_PauseClothingSimulation.cpp
--- a/Engine/Source/Runtime/Engine/Private/Animation/AnimNotify_PauseClothingSimulation.cpp
+++ b/Engine/Source/Runtime/Engine/Private/Animation/AnimNotify_PauseClothingSimulation.cpp
@@ -16,6 +16,11 @@ void UAnimNotify_PauseClothingSimulation::Notify(class USkeletalMeshComponent* M
 	MeshComp->SuspendClothingSimulation();
 }
 
+void UAnimNotify_PauseClothingSimulation::Notify(class USkeletalMeshComponent* MeshComp, class UAnimSequenceBase* Animation, const FAnimNotifyEventReference& EventReference)
+{
+	MeshComp->SuspendClothingSimulation();
+}
+
 FString UAnimNotify_PauseClothingSimulation::GetNotifyName_Implementation() const
 {
 	return TEXT("Pause Clothing Sim");
diff --git a/Engine/Source/Runtime/Engine/Private/Animation/AnimNotify_PlayParticleEffect.cpp b/Engine/Source/Runtime/Engine/Private/Animation/AnimNotify_PlayParticleEffect.cpp
--- a/Engine/Source/Runtime/Engine/Private/Animation/AnimNotify_PlayParticleEffect.cpp
+++ b/Engine/Source/Runtime/Engine/Private/Animation/AnimNotify_PlayParticleEffect.cpp
@@ -71,6 +71,10 @@ void UAnimNotify_PlayParticleEffect::ValidateAssociatedAssets()
 #endif
 
 void UAnimNotify_PlayParticleEffect::Notify(class USkeletalMeshComponent* MeshComp, class UAnimSequenceBase* Animation)
+{
+}
+
+void UAnimNotify_PlayParticleEffect::Notify(class USkeletalMeshComponent* MeshComp, class UAnimSequenceBase* Animation, const FAnimNotifyEventReference& EventReference)
 {
 	// Don't call super to avoid unnecessary call in to blueprints
 	SpawnParticleSystem(MeshComp, Animation);
diff --git a/Engine/Source/Runtime/Engine/Private/Animation/AnimNotify_PlaySound.cpp b/Engine/Source/Runtime/Engine/Private/Animation/AnimNotify_PlaySound.cpp
--- a/Engine/Source/Runtime/Engine/Private/Animation/AnimNotify_PlaySound.cpp
+++ b/Engine/Source/Runtime/Engine/Private/Animation/AnimNotify_PlaySound.cpp
@@ -27,7 +27,11 @@ UAnimNotify_PlaySound::UAnimNotify_PlaySound()
 #endif // WITH_EDITORONLY_DATA
 }
 
-void UAnimNotify_PlaySound::Notify(class USkeletalMeshComponent* MeshComp, class UAnimSequenceBase* Animation)
+void UAnimNotify_PlaySound::Notify(USkeletalMeshComponent* MeshComp, UAnimSequenceBase* Animation)
+{
+}
+
+void UAnimNotify_PlaySound::Notify(class USkeletalMeshComponent* MeshComp, class UAnimSequenceBase* Animation, const FAnimNotifyEventReference& EventReference)
 {
 	// Don't call super to avoid call back in to blueprints
 	if (Sound && MeshComp)
diff --git a/Engine/Source/Runtime/Engine/Private/Animation/AnimNotify_ResetClothingSimulation.cpp b/Engine/Source/Runtime/Engine/Private/Animation/AnimNotify_ResetClothingSimulation.cpp
--- a/Engine/Source/Runtime/Engine/Private/Animation/AnimNotify_ResetClothingSimulation.cpp
+++ b/Engine/Source/Runtime/Engine/Private/Animation/AnimNotify_ResetClothingSimulation.cpp
@@ -12,6 +12,10 @@ UAnimNotify_ResetClothingSimulation::UAnimNotify_ResetClothingSimulation()
 }
 
 void UAnimNotify_ResetClothingSimulation::Notify(class USkeletalMeshComponent* MeshComp, class UAnimSequenceBase* Animation)
+{
+}
+
+void UAnimNotify_ResetClothingSimulation::Notify(class USkeletalMeshComponent* MeshComp, class UAnimSequenceBase* Animation, const FAnimNotifyEventReference& EventReference)
 {
 	MeshComp->ForceClothNextUpdateTeleportAndReset();
 }
diff --git a/Engine/Source/Runtime/Engine/Private/Animation/AnimNotify_ResetDynamics.cpp b/Engine/Source/Runtime/Engine/Private/Animation/AnimNotify_ResetDynamics.cpp
--- a/Engine/Source/Runtime/Engine/Private/Animation/AnimNotify_ResetDynamics.cpp
+++ b/Engine/Source/Runtime/Engine/Private/Animation/AnimNotify_ResetDynamics.cpp
@@ -12,6 +12,10 @@ UAnimNotify_ResetDynamics::UAnimNotify_ResetDynamics()
 }
 
 void UAnimNotify_ResetDynamics::Notify(class USkeletalMeshComponent* MeshComp, class UAnimSequenceBase* Animation)
+{
+}
+
+void UAnimNotify_ResetDynamics::Notify(class USkeletalMeshComponent* MeshComp, class UAnimSequenceBase* Animation, const FAnimNotifyEventReference& EventReference)
 {
 	MeshComp->ResetAnimInstanceDynamics();
 }
diff --git a/Engine/Source/Runtime/Engine/Private/Animation/AnimNotify_ResumeClothingSimulation.cpp b/Engine/Source/Runtime/Engine/Private/Animation/AnimNotify_ResumeClothingSimulation.cpp
--- a/Engine/Source/Runtime/Engine/Private/Animation/AnimNotify_ResumeClothingSimulation.cpp
+++ b/Engine/Source/Runtime/Engine/Private/Animation/AnimNotify_ResumeClothingSimulation.cpp
@@ -12,6 +12,10 @@ UAnimNotify_ResumeClothingSimulation::UAnimNotify_ResumeClothingSimulation()
 }
 
 void UAnimNotify_ResumeClothingSimulation::Notify(class USkeletalMeshComponent* MeshComp, class UAnimSequenceBase* Animation)
+{
+}
+
+void UAnimNotify_ResumeClothingSimulation::Notify(class USkeletalMeshComponent* MeshComp, class UAnimSequenceBase* Animation, const FAnimNotifyEventReference& EventReference)
 {
 	MeshComp->ResumeClothingSimulation();
 }
diff --git a/Engine/Source/Runtime/Engine/Private/Animation/AnimSequence.cpp b/Engine/Source/Runtime/Engine/Private/Animation/AnimSequence.cpp
--- a/Engine/Source/Runtime/Engine/Private/Animation/AnimSequence.cpp
+++ b/Engine/Source/Runtime/Engine/Private/Animation/AnimSequence.cpp
@@ -5074,7 +5074,7 @@ void UAnimSequence::GetMarkerIndicesForTime(float CurrentTime, bool bLooping, co
 	}
 }
 
-FMarkerSyncAnimPosition UAnimSequence::GetMarkerSyncPositionfromMarkerIndicies(int32 PrevMarker, int32 NextMarker, float CurrentTime, const UMirrorDataTable* MirrorTable) const
+FMarkerSyncAnimPosition UAnimSequence::GetMarkerSyncPositionFromMarkerIndicies(int32 PrevMarker, int32 NextMarker, float CurrentTime, const UMirrorDataTable* MirrorTable) const
 {
 	FMarkerSyncAnimPosition SyncPosition;
 	float PrevTime, NextTime;
diff --git a/Engine/Source/Runtime/Engine/Private/Animation/AnimSequenceBase.cpp b/Engine/Source/Runtime/Engine/Private/Animation/AnimSequenceBase.cpp
--- a/Engine/Source/Runtime/Engine/Private/Animation/AnimSequenceBase.cpp
+++ b/Engine/Source/Runtime/Engine/Private/Animation/AnimSequenceBase.cpp
@@ -12,6 +12,7 @@
 #include "UObject/FortniteMainBranchObjectVersion.h"
 #include "Animation/AnimationPoseData.h"
 #include "Animation/AnimSequenceHelpers.h"
+#include "Animation/MirrorDataTable.h"
 #include "Animation/AnimData/AnimDataModel.h"
 #include "Modules/ModuleManager.h"
 
@@ -259,11 +260,13 @@ bool UAnimSequenceBase::IsNotifyAvailable() const
 
 void UAnimSequenceBase::GetAnimNotifies(const float& StartTime, const float& DeltaTime, const bool bAllowLooping, TArray<const FAnimNotifyEvent *> & OutActiveNotifies) const
 {
-	TArray<FAnimNotifyEventReference> NotifyRefs;
-	GetAnimNotifies(StartTime, DeltaTime, bAllowLooping, NotifyRefs);
+	FAnimTickRecord TickRecord;
+	TickRecord.bLooping = bAllowLooping; 
+	FAnimNotifyContext NotifyContext(TickRecord);
+	GetAnimNotifies(StartTime, DeltaTime, NotifyContext);
 
-	OutActiveNotifies.Reset(NotifyRefs.Num());
-	for (FAnimNotifyEventReference& NotifyRef : NotifyRefs)
+	OutActiveNotifies.Reset(NotifyContext.ActiveNotifies.Num());
+	for (FAnimNotifyEventReference NotifyRef : NotifyContext.ActiveNotifies)
 	{
 		if (const FAnimNotifyEvent* Notify = NotifyRef.GetNotify())
 		{
@@ -272,7 +275,16 @@ void UAnimSequenceBase::GetAnimNotifies(const float& StartTime, const float& Del
 	}
 }
 
-void UAnimSequenceBase::GetAnimNotifies(const float& StartTime, const float& DeltaTime, const bool bAllowLooping, TArray<FAnimNotifyEventReference> & OutActiveNotifies) const
+void UAnimSequenceBase::GetAnimNotifies(const float& StartTime, const float& DeltaTime, const bool bAllowLooping, TArray<FAnimNotifyEventReference>& OutActiveNotifies) const
+{
+	FAnimTickRecord TickRecord;
+	TickRecord.bLooping = bAllowLooping;
+	FAnimNotifyContext NotifyContext(TickRecord);
+	GetAnimNotifies(StartTime, DeltaTime, NotifyContext);
+	Swap(NotifyContext.ActiveNotifies, OutActiveNotifies);
+}
+
+void UAnimSequenceBase::GetAnimNotifies(const float& StartTime, const float& DeltaTime, FAnimNotifyContext& NotifyContext) const
 {
 	if(DeltaTime == 0.f)
 	{
@@ -299,10 +311,10 @@ void UAnimSequenceBase::GetAnimNotifies(const float& StartTime, const float& Del
 		ensureMsgf(bPlayingBackwards ? (CurrentPosition <= PreviousPosition) : (CurrentPosition >= PreviousPosition), TEXT("in Animation %s(Skeleton %s) : bPlayingBackwards(%d), PreviousPosition(%0.2f), Current Position(%0.2f)"), 
 			*GetName(), *GetNameSafe(GetSkeleton()), bPlayingBackwards, PreviousPosition, CurrentPosition);
 		
-		GetAnimNotifiesFromDeltaPositions(PreviousPosition, CurrentPosition, OutActiveNotifies);
+		GetAnimNotifiesFromDeltaPositions(PreviousPosition, CurrentPosition, NotifyContext);
 	
 		// If we've hit the end of the animation, and we're allowed to loop, keep going.
-		if( (AdvanceType == ETAA_Finished) &&  bAllowLooping )
+		if( (AdvanceType == ETAA_Finished) &&  NotifyContext.TickRecord && NotifyContext.TickRecord->bLooping )
 		{
 			const float ActualDeltaMove = (CurrentPosition - PreviousPosition);
 			DesiredDeltaMove -= ActualDeltaMove; 
@@ -320,11 +332,12 @@ void UAnimSequenceBase::GetAnimNotifies(const float& StartTime, const float& Del
 
 void UAnimSequenceBase::GetAnimNotifiesFromDeltaPositions(const float& PreviousPosition, const float& CurrentPosition, TArray<const FAnimNotifyEvent *> & OutActiveNotifies) const
 {
-	TArray<FAnimNotifyEventReference> NotifyRefs;
-	GetAnimNotifiesFromDeltaPositions(PreviousPosition, CurrentPosition, NotifyRefs);
+	FAnimTickRecord TickRecord;
+	FAnimNotifyContext NotifyContext(TickRecord);
+	GetAnimNotifiesFromDeltaPositions(PreviousPosition, CurrentPosition, NotifyContext);
 
-	OutActiveNotifies.Reset(NotifyRefs.Num());
-	for (FAnimNotifyEventReference NotifyRef : NotifyRefs)
+	OutActiveNotifies.Reset(NotifyContext.ActiveNotifies.Num());
+	for (FAnimNotifyEventReference NotifyRef : NotifyContext.ActiveNotifies)
 	{
 		if (const FAnimNotifyEvent* Notify = NotifyRef.GetNotify())
 		{
@@ -333,7 +346,15 @@ void UAnimSequenceBase::GetAnimNotifiesFromDeltaPositions(const float& PreviousP
 	}
 }
 
-void UAnimSequenceBase::GetAnimNotifiesFromDeltaPositions(const float& PreviousPosition, const float& CurrentPosition, TArray<FAnimNotifyEventReference> & OutActiveNotifies) const
+void UAnimSequenceBase::GetAnimNotifiesFromDeltaPositions(const float& PreviousPosition, const float & CurrentPosition, TArray<FAnimNotifyEventReference>& OutActiveNotifies) const
+{
+	FAnimTickRecord TickRecord;
+	FAnimNotifyContext NotifyContext(TickRecord);
+	GetAnimNotifiesFromDeltaPositions(PreviousPosition, CurrentPosition,NotifyContext);
+	Swap(NotifyContext.ActiveNotifies, OutActiveNotifies);
+}
+
+void UAnimSequenceBase::GetAnimNotifiesFromDeltaPositions(const float& PreviousPosition, const float& CurrentPosition,  FAnimNotifyContext& NotifyContext) const
 {
 	// Early out if we have no notifies
 	if( (Notifies.Num() == 0) || (PreviousPosition == CurrentPosition) )
@@ -354,7 +375,15 @@ void UAnimSequenceBase::GetAnimNotifiesFromDeltaPositions(const float& PreviousP
 
 			if( (NotifyStartTime < PreviousPosition) && (NotifyEndTime >= CurrentPosition) )
 			{
-				OutActiveNotifies.Emplace(&AnimNotifyEvent, this);
+				if (NotifyContext.TickRecord)
+				{
+					NotifyContext.ActiveNotifies.Emplace(&AnimNotifyEvent, this, NotifyContext.TickRecord->MirrorDataTable);
+					NotifyContext.ActiveNotifies.Top().SetContextDataInterfaces(NotifyContext.TickRecord->ContextData);
+				}
+				else
+				{
+					NotifyContext.ActiveNotifies.Emplace(&AnimNotifyEvent, this, nullptr);
+				}
 			}
 		}
 	}
@@ -368,7 +397,16 @@ void UAnimSequenceBase::GetAnimNotifiesFromDeltaPositions(const float& PreviousP
 
 			if( (NotifyStartTime <= CurrentPosition) && (NotifyEndTime > PreviousPosition) )
 			{
-				OutActiveNotifies.Emplace(&AnimNotifyEvent, this);
+				if (NotifyContext.TickRecord)
+				{
+					NotifyContext.ActiveNotifies.Emplace(&AnimNotifyEvent, this, NotifyContext.TickRecord->MirrorDataTable);
+					NotifyContext.ActiveNotifies.Top().SetContextDataInterfaces(NotifyContext.TickRecord->ContextData); 
+				}
+				else
+				{
+					NotifyContext.ActiveNotifies.Emplace(&AnimNotifyEvent, this, nullptr);
+				}
+
 			}
 		}
 	}
@@ -487,13 +525,13 @@ void UAnimSequenceBase::TickByMarkerAsLeader(FMarkerTickRecord& Instance, FMarke
 		
 	}
 
-	MarkerContext.SetMarkerSyncStartPosition(GetMarkerSyncPositionfromMarkerIndicies(Instance.PreviousMarker.MarkerIndex, Instance.NextMarker.MarkerIndex, CurrentTime, MirrorTable));
+	MarkerContext.SetMarkerSyncStartPosition(GetMarkerSyncPositionFromMarkerIndicies(Instance.PreviousMarker.MarkerIndex, Instance.NextMarker.MarkerIndex, CurrentTime, MirrorTable));
 
 	OutPreviousTime = CurrentTime;
 
 	AdvanceMarkerPhaseAsLeader(bLooping, MoveDelta, MarkerContext.GetValidMarkerNames(), CurrentTime, Instance.PreviousMarker, Instance.NextMarker, MarkerContext.MarkersPassedThisTick, MirrorTable);
 
-	MarkerContext.SetMarkerSyncEndPosition(GetMarkerSyncPositionfromMarkerIndicies(Instance.PreviousMarker.MarkerIndex, Instance.NextMarker.MarkerIndex, CurrentTime, MirrorTable));
+	MarkerContext.SetMarkerSyncEndPosition(GetMarkerSyncPositionFromMarkerIndicies(Instance.PreviousMarker.MarkerIndex, Instance.NextMarker.MarkerIndex, CurrentTime, MirrorTable));
 
 	UE_LOG(LogAnimMarkerSync, Log, TEXT("Leader (%s) (TickByMarker) PreviousTime(%0.2f) CurrentTime(%0.2f) MoveDelta(%0.2f) Looping(%d) %s"), *GetName(), OutPreviousTime, CurrentTime, MoveDelta, bLooping ? 1 : 0, *MarkerContext.ToString());
 }
@@ -908,9 +946,9 @@ void UAnimSequenceBase::GetAnimationPose(struct FCompactPose& OutPose, FBlendedC
 void UAnimSequenceBase::HandleAssetPlayerTickedInternal(FAnimAssetTickContext &Context, const float PreviousTime, const float MoveDelta, const FAnimTickRecord &Instance, struct FAnimNotifyQueue& NotifyQueue) const
 {
 	// Harvest and record notifies
-	TArray<FAnimNotifyEventReference> AnimNotifies;
-	GetAnimNotifies(PreviousTime, MoveDelta, Instance.bLooping, AnimNotifies);
-	NotifyQueue.AddAnimNotifies(Context.ShouldGenerateNotifies(),AnimNotifies, Instance.EffectiveBlendWeight);
+	FAnimNotifyContext NotifyContext(Instance);
+	GetAnimNotifies(PreviousTime, MoveDelta, NotifyContext);
+	NotifyQueue.AddAnimNotifies(Context.ShouldGenerateNotifies(), NotifyContext.ActiveNotifies, Instance.EffectiveBlendWeight);
 }
 
 #if WITH_EDITOR
diff --git a/Engine/Source/Runtime/Engine/Private/Animation/AnimSingleNodeInstance.cpp b/Engine/Source/Runtime/Engine/Private/Animation/AnimSingleNodeInstance.cpp
--- a/Engine/Source/Runtime/Engine/Private/Animation/AnimSingleNodeInstance.cpp
+++ b/Engine/Source/Runtime/Engine/Private/Animation/AnimSingleNodeInstance.cpp
@@ -350,14 +350,15 @@ void UAnimSingleNodeInstance::SetPositionWithPreviousTime(float InPosition, floa
 		UAnimSequenceBase * SequenceBase = Cast<UAnimSequenceBase> (CurrentAsset);
 		if (SequenceBase)
 		{
+			FAnimTickRecord TickRecord;
+			FAnimNotifyContext NotifyContext(TickRecord);
 			NotifyQueue.Reset(GetSkelMeshComponent());
 
-			TArray<FAnimNotifyEventReference> Notifies;
-			SequenceBase->GetAnimNotifiesFromDeltaPositions(InPreviousTime, Proxy.GetCurrentTime(), Notifies);
-			if ( Notifies.Num() > 0 )
+			SequenceBase->GetAnimNotifiesFromDeltaPositions(InPreviousTime, Proxy.GetCurrentTime(), NotifyContext);
+			if ( NotifyContext.ActiveNotifies.Num() > 0 )
 			{
 				// single node instance only has 1 asset at a time
-				NotifyQueue.AddAnimNotifies(Notifies, 1.0f);
+				NotifyQueue.AddAnimNotifies(NotifyContext.ActiveNotifies, 1.0f);
 			}
 
 			TriggerAnimNotifies(0.f);
diff --git a/Engine/Source/Runtime/Engine/Private/Animation/AnimSync.cpp b/Engine/Source/Runtime/Engine/Private/Animation/AnimSync.cpp
--- a/Engine/Source/Runtime/Engine/Private/Animation/AnimSync.cpp
+++ b/Engine/Source/Runtime/Engine/Private/Animation/AnimSync.cpp
@@ -34,6 +34,7 @@ void FAnimSync::AddTickRecord(const FAnimTickRecord& InTickRecord, const FAnimSy
 	else
 	{
 		UngroupedActivePlayerArrays[GetSyncGroupWriteIndex()].Add(InTickRecord);
+		UngroupedActivePlayerArrays[GetSyncGroupWriteIndex()].Top().MirrorDataTable = MirrorDataTable;
 	}
 }
 
diff --git a/Engine/Source/Runtime/Engine/Private/Animation/AnimTypes.cpp b/Engine/Source/Runtime/Engine/Private/Animation/AnimTypes.cpp
--- a/Engine/Source/Runtime/Engine/Private/Animation/AnimTypes.cpp
+++ b/Engine/Source/Runtime/Engine/Private/Animation/AnimTypes.cpp
@@ -134,6 +134,24 @@ FName FAnimNotifyEvent::GetNotifyEventName() const
 	return NAME_None;
 }
 
+FName FAnimNotifyEvent::GetNotifyEventName(const UMirrorDataTable* MirrorDataTable) const
+{
+	if (MirrorDataTable)
+	{
+		if(NotifyName == NAME_None)
+		{
+			return NAME_None;
+		}
+		const FName* MirroredName = MirrorDataTable->AnimNotifyToMirrorAnimNotifyMap.Find(NotifyName);
+		if (MirroredName)
+		{
+			const FString EventName = FString::Printf(TEXT("AnimNotify_%s"), *MirroredName->ToString());
+			return FName(*EventName);
+		}
+	}
+	return GetNotifyEventName(); 
+}
+
 ////////////////////////////
 //
 // FAnimSyncMarker
@@ -210,7 +228,12 @@ void FMarkerSyncData::GetMarkerIndicesForTime(float CurrentTime, bool bLooping,
 	}
 }
 
-FMarkerSyncAnimPosition FMarkerSyncData::GetMarkerSyncPositionfromMarkerIndicies(int32 PrevMarker, int32 NextMarker, float CurrentTime, float SequenceLength, const UMirrorDataTable* MirrorTable) const
+FMarkerSyncAnimPosition FMarkerSyncData::GetMarkerSyncPositionfromMarkerIndicies(int32 PrevMarker, int32 NextMarker, float CurrentTime, float SequenceLength) const
+{
+	return GetMarkerSyncPositionFromMarkerIndicies(PrevMarker, NextMarker, CurrentTime, SequenceLength, nullptr);
+}
+
+FMarkerSyncAnimPosition FMarkerSyncData::GetMarkerSyncPositionFromMarkerIndicies(int32 PrevMarker, int32 NextMarker, float CurrentTime, float SequenceLength, const UMirrorDataTable* MirrorTable) const
 {
 	FMarkerSyncAnimPosition SyncPosition;
 	float PrevTime, NextTime;
diff --git a/Engine/Source/Runtime/Engine/Private/Animation/BlendSpace.cpp b/Engine/Source/Runtime/Engine/Private/Animation/BlendSpace.cpp
--- a/Engine/Source/Runtime/Engine/Private/Animation/BlendSpace.cpp
+++ b/Engine/Source/Runtime/Engine/Private/Animation/BlendSpace.cpp
@@ -486,8 +486,7 @@ void UBlendSpace::TickAssetPlayer(FAnimTickRecord& Instance, struct FAnimNotifyQ
 
 			// generate notifies and sets time
 			{
-				TArray<FAnimNotifyEventReference> Notifies;
-
+				FAnimNotifyContext NotifyContext(Instance);
 				const float ClampedNormalizedPreviousTime = FMath::Clamp<float>(NormalizedPreviousTime, 0.f, 1.f);
 				const float ClampedNormalizedCurrentTime = FMath::Clamp<float>(NormalizedCurrentTime, 0.f, 1.f);
 				const bool bGenerateNotifies = (NormalizedCurrentTime != NormalizedPreviousTime) && NotifyTriggerMode != ENotifyTriggerMode::None;
@@ -537,7 +536,7 @@ void UBlendSpace::TickAssetPlayer(FAnimTickRecord& Instance, struct FAnimNotifyQ
 							if (bGenerateNotifies && (!bTriggerNotifyHighestWeightedAnim || (I == HighestWeightIndex)))
 							{
 								// Harvest and record notifies
-								Sample.Animation->GetAnimNotifies(PrevSampleDataTime, DeltaTimePosition, Instance.bLooping, Notifies);
+								Sample.Animation->GetAnimNotifies(PrevSampleDataTime, DeltaTimePosition, NotifyContext);
 							}
 
 							if (Context.RootMotionMode == ERootMotionMode::RootMotionFromEverything && Sample.Animation->bEnableRootMotion)
@@ -551,9 +550,9 @@ void UBlendSpace::TickAssetPlayer(FAnimTickRecord& Instance, struct FAnimNotifyQ
 					}
 				}
 
-				if (bGenerateNotifies && Notifies.Num() > 0)
+				if (bGenerateNotifies && NotifyContext.ActiveNotifies.Num() > 0)
 				{
-					NotifyQueue.AddAnimNotifies(Context.ShouldGenerateNotifies(), Notifies, Instance.EffectiveBlendWeight);
+					NotifyQueue.AddAnimNotifies(Context.ShouldGenerateNotifies(), NotifyContext.ActiveNotifies, Instance.EffectiveBlendWeight);
 				}
 			}
 		}
diff --git a/Engine/Source/Runtime/Engine/Private/Animation/MirrorDataTable.cpp b/Engine/Source/Runtime/Engine/Private/Animation/MirrorDataTable.cpp
--- a/Engine/Source/Runtime/Engine/Private/Animation/MirrorDataTable.cpp
+++ b/Engine/Source/Runtime/Engine/Private/Animation/MirrorDataTable.cpp
@@ -217,7 +217,7 @@ void UMirrorDataTable::FindReplaceMirroredNames()
 
 	bool bChangedTable = false;
 
-	auto AddMirrorRow = [this, &bChangedTable, &NamesByCategory](const FName& Name, const FName& MirroredName, EMirrorRowType RowType)
+	auto AddMirrorRow = [this, &bChangedTable, &NamesByCategory](const FName& Name, const FName& MirroredName, EMirrorRowType::Type RowType)
 	{
 		// directly add rows to avoid using FDataTableEditorUtils, which is not appropriate at this point
 		// equivalent to FDataTableEditorUtils::AddRow(DataTable, BoneName);
diff --git a/Engine/Source/Runtime/Engine/Private/Animation/MirrorSyncScope.cpp b/Engine/Source/Runtime/Engine/Private/Animation/MirrorSyncScope.cpp
--- a/Engine/Source/Runtime/Engine/Private/Animation/MirrorSyncScope.cpp
+++ b/Engine/Source/Runtime/Engine/Private/Animation/MirrorSyncScope.cpp
@@ -7,6 +7,8 @@
 
 IMPLEMENT_ANIMGRAPH_MESSAGE(UE::Anim::FMirrorSyncScope);
 
+IMPLEMENT_NOTIFY_CONTEXT_INTERFACE(UE::Anim::FAnimNotifyMirrorContext)
+
 namespace UE { namespace Anim {
 
 bool IsMirrorSyncIdentical(const UMirrorDataTable* MirrorTableA, const UMirrorDataTable* MirrorTableB)
@@ -69,5 +71,10 @@ int32 FMirrorSyncScope::GetMirrorScopeDepth() const
 {
 	return MirrorScopeDepth;
 }
-	
+
+TSharedPtr<const IAnimNotifyEventContextDataInterface> FMirrorSyncScope::MakeEventContextData() const
+{
+	return MakeShared<FAnimNotifyMirrorContext>(MirrorDataTable);
+}
+
 }}	// namespace UE::Anim
\ No newline at end of file
diff --git a/Engine/Source/Runtime/Engine/Public/Animation/ActiveStateMachineScope.h b/Engine/Source/Runtime/Engine/Public/Animation/ActiveStateMachineScope.h
new file mode 100644
--- /dev/null
+++ b/Engine/Source/Runtime/Engine/Public/Animation/ActiveStateMachineScope.h
@@ -0,0 +1,62 @@
+// Copyright Epic Games, Inc. All Rights Reserved.
+
+#pragma once
+
+#include "CoreMinimal.h"
+#include "AnimNodeMessages.h"
+#include "Animation/AnimationAsset.h"
+#include "Animation/AnimSync.h"
+#include "ActiveStateMachineScope.generated.h"
+
+struct FAnimationBaseContext;
+struct FAnimInstanceProxy;
+struct FAnimNode_StateMachine;
+
+USTRUCT()
+struct ENGINE_API FEncounteredStateMachineStack
+{
+	GENERATED_BODY()
+	FEncounteredStateMachineStack() {} 
+	FEncounteredStateMachineStack(const FEncounteredStateMachineStack& ParentStack, int32 InStateMachineIndex, int32 InStateIndex);
+	FEncounteredStateMachineStack(int32 InStateMachineIndex, int32 InStateIndex);
+	struct ENGINE_API FStateMachineEntry
+	{
+		FStateMachineEntry() {}
+		FStateMachineEntry(int32 InStateMachineIndex, int32 InStateIndex)
+		: StateMachineIndex(InStateMachineIndex)
+		, StateIndex(InStateIndex)
+		{}
+		
+		int32 StateMachineIndex = INDEX_NONE;
+		int32 StateIndex = INDEX_NONE;;
+	};
+	TArray<FStateMachineEntry, TInlineAllocator<4>> StateStack;
+	static TSharedPtr<const FEncounteredStateMachineStack> InitStack(const FEncounteredStateMachineStack& ParentStack, int32 InStateMachineIndex, int32 InStateIndex);
+	static TSharedPtr<const FEncounteredStateMachineStack> InitStack(int32 InStateMachineIndex, int32 InStateIndex);
+};
+
+namespace UE { namespace Anim {
+	
+class ENGINE_API FAnimNotifyStateMachineContext : public UE::Anim::IAnimNotifyEventContextDataInterface
+{
+	DECLARE_NOTIFY_CONTEXT_INTERFACE(FAnimNotifyStateMachineContext)
+	public:
+	FAnimNotifyStateMachineContext(const TSharedPtr<const FEncounteredStateMachineStack>& InEncounteredStateMachines);
+	bool IsStateMachineInContext(int32 StateMachineIndex) const;
+	bool IsStateInStateMachineInContext(int32 StateMachineIndex, int32 StateIndex) const; 
+	TSharedPtr<const FEncounteredStateMachineStack> EncounteredStateMachines; 
+};
+
+class ENGINE_API FActiveStateMachineScope : public UE::Anim::IGraphMessage
+{
+	DECLARE_ANIMGRAPH_MESSAGE(FActiveStateMachineScope);
+	public:
+	FActiveStateMachineScope(const FAnimationBaseContext& InContext, FAnimNode_StateMachine* StateMachine, int32 InStateIndex);
+	virtual ~FActiveStateMachineScope();
+	static int32 GetStateMachineIndex(FAnimNode_StateMachine* StateMachine, const FAnimationBaseContext& Context);
+	virtual TSharedPtr<const IAnimNotifyEventContextDataInterface> MakeEventContextData() const override;
+	
+	TSharedPtr<const FEncounteredStateMachineStack> ActiveStateMachines; 
+};
+
+}}	// namespace UE::Anim
diff --git a/Engine/Source/Runtime/Engine/Public/Animation/AnimInstanceProxy.h b/Engine/Source/Runtime/Engine/Public/Animation/AnimInstanceProxy.h
--- a/Engine/Source/Runtime/Engine/Public/Animation/AnimInstanceProxy.h
+++ b/Engine/Source/Runtime/Engine/Public/Animation/AnimInstanceProxy.h
@@ -20,6 +20,7 @@
 #include "Animation/AnimTrace.h"
 #include "Animation/AnimSync.h"
 #include "Animation/AnimSyncScope.h"
+#include "Animation/ActiveStateMachineScope.h"
 #include "AnimInstanceProxy.generated.h"
 
 class UBlendSpace;
@@ -576,7 +577,8 @@ public:
 	friend struct FAnimTrace;
 	friend struct UE::Anim::FAnimSync;
 	friend class UE::Anim::FAnimSyncGroupScope;
-
+	friend class UE::Anim::FActiveStateMachineScope;
+	
 protected:
 	/** Called when our anim instance is being initialized */
 	virtual void Initialize(UAnimInstance* InAnimInstance);
@@ -584,6 +586,9 @@ protected:
 	/** Called when our anim instance is being uninitialized */
 	virtual void Uninitialize(UAnimInstance* InAnimInstance);
 
+	/** Let us copy the full notify queue from the last frame */
+	void UpdateActiveAnimNotifiesSinceLastTick(const FAnimNotifyQueue& AnimInstanceQueue);
+	
 	/** Called before update so we can copy any data we need */
 	virtual void PreUpdate(UAnimInstance* InAnimInstance, float DeltaSeconds);
 
@@ -858,6 +863,33 @@ protected:
 	/** Get the current accumulated time as a fraction of the length of the most relevant animation in the source state */
 	float GetRelevantAnimTimeFraction(int32 MachineIndex, int32 StateIndex) const;
 
+	/** Get whether a particular notify state was active in any state machine last tick. */
+	bool WasAnimNotifyStateActiveInAnyState(TSubclassOf<UAnimNotifyState> AnimNotifyStateType) const;
+
+	/** Get whether a particular notify state is active in a specific state machine last tick. */
+	bool WasAnimNotifyStateActiveInStateMachine(int32 MachineIndex, TSubclassOf<UAnimNotifyState> AnimNotifyStateType) const;
+
+	/** Get whether the most relevant animation was in a particular notify state last tick. */
+	bool WasAnimNotifyStateActiveInSourceState(int32 MachineIndex, int32 StateIndex, TSubclassOf<UAnimNotifyState> AnimNotifyStateType) const;
+	
+	/** Get whether the most relevant animation triggered the given notify last tick. */
+	bool WasAnimNotifyTriggeredInSourceState(int32 MachineIndex, int32 StateIndex, TSubclassOf<UAnimNotify> AnimNotifyType) const;
+
+	/** Get whether the most relevant animation triggered the given notify last tick. */
+	bool WasAnimNotifyNameTriggeredInSourceState(int32 MachineIndex, int32 StateIndex, FName NotifyName) const;
+
+	/** Get whether a particular notify type was active in a specific state machine last tick.  */
+	bool WasAnimNotifyTriggeredInStateMachine(int32 MachineIndex, TSubclassOf<UAnimNotify> AnimNotifyType) const;
+
+	/**  Get whether an animation notify of a given type was triggered last tick. */
+    bool WasAnimNotifyTriggeredInAnyState(TSubclassOf<UAnimNotify> AnimNotifyType) const;
+
+	/** Get whether the animation notify with the specified name triggered last tick. */
+    bool WasAnimNotifyNameTriggeredInAnyState(FName NotifyName) const;
+	
+	/** Get whether the given state machine triggered the animation notify with the specified name last tick. */
+    bool WasAnimNotifyNameTriggeredInStateMachine(int32 MachineIndex, FName NotifyName);
+	
 	// Sets up a native transition delegate between states with PrevStateName and NextStateName, in the state machine with name MachineName.
 	// Note that a transition already has to exist for this to succeed
 	void AddNativeTransitionBinding(const FName& MachineName, const FName& PrevStateName, const FName& NextStateName, const FCanTakeTransition& NativeTransitionDelegate, const FName& TransitionName = NAME_None);
@@ -886,6 +918,9 @@ protected:
 	/** Gets the index of the state machine matching MachineName */
 	int32 GetStateMachineIndex(FName MachineName) const;
 
+	/** Gets the index of the state machine */
+	int32 GetStateMachineIndex(FAnimNode_StateMachine* StateMachine) const;
+
 	void GetStateMachineIndexAndDescription(FName InMachineName, int32& OutMachineIndex, const FBakedAnimationStateMachine** OutMachineDescription) const;
 
 	/** Initialize the root node - split into a separate function for backwards compatibility (initialization order) reasons */
@@ -1025,6 +1060,9 @@ private:
 	TArray<FName> MaterialParametersToClear;
 
 protected:
+	// Animation Notifies that has been triggered since the last tick. These can be safely consumed at any point.
+	TArray<FAnimNotifyEventReference> ActiveAnimNotifiesSinceLastTick;
+	
 	// Counters for synchronization
 	FGraphTraversalCounter InitializationCounter;
 	FGraphTraversalCounter CachedBonesCounter;
diff --git a/Engine/Source/Runtime/Engine/Public/Animation/AnimNodeMessages.h b/Engine/Source/Runtime/Engine/Public/Animation/AnimNodeMessages.h
--- a/Engine/Source/Runtime/Engine/Public/Animation/AnimNodeMessages.h
+++ b/Engine/Source/Runtime/Engine/Public/Animation/AnimNodeMessages.h
@@ -8,6 +8,7 @@
 struct FAnimationUpdateSharedContext;
 struct FAnimationBaseContext;
 struct FAnimNode_Base;
+class IAnimNotifyEventContextDataInterface;
 
 // Simple RTTI implementation for graph messages
 
@@ -23,8 +24,47 @@ struct FAnimNode_Base;
 #define IMPLEMENT_ANIMGRAPH_MESSAGE(ClassName) \
 	const FName ClassName::TypeName = TEXT(#ClassName);
 
+#define IMPLEMENT_NOTIFY_CONTEXT_INTERFACE(ClassName) \
+	IMPLEMENT_ANIMGRAPH_MESSAGE(ClassName);
+
+#define DECLARE_NOTIFY_CONTEXT_INTERFACE(ClassName) \
+	DECLARE_ANIMGRAPH_MESSAGE(ClassName)
+
 namespace UE { namespace Anim {
 
+class ENGINE_API IAnimNotifyEventContextDataInterface
+{
+public:
+	virtual ~IAnimNotifyEventContextDataInterface() = default;
+	// RTTI support functions
+	static FName GetStaticTypeName()
+	{ 
+		return BaseTypeName;
+	}
+
+	virtual FName GetTypeName() const = 0;
+	
+	template<typename Type> 
+	bool Is() const 
+	{ 
+		return GetTypeName() == Type::GetStaticTypeName(); 
+	}
+
+	template<typename Type> 
+	const Type& As() const 
+	{ 
+		return *static_cast<const Type*>(this);
+	}
+
+	template<typename Type> 
+    Type& As() 
+	{ 
+		return *static_cast<Type*>(this);
+	}
+private:
+	static const FName BaseTypeName;
+};
+	
 // Base class for all messages/events/scopes that are fired during execution
 // Note that these messages are events/callbacks only, no blending is involved.
 class ENGINE_API IGraphMessage
@@ -59,6 +99,12 @@ public:
 		return *static_cast<Type*>(this);
 	}
 
+	virtual TSharedPtr<const IAnimNotifyEventContextDataInterface> MakeEventContextData() const
+	{
+		TSharedPtr<const IAnimNotifyEventContextDataInterface> NullPtr;
+		return NullPtr; 
+	}
+
 private:
 	static const FName BaseTypeName;
 };
@@ -224,6 +270,10 @@ public:
 	// @param	InStack		The stack to copy from
 	void CopyForCachedUpdate(const FMessageStack& InStack);
 
+	// Call MakeEventContextData for the top entry of each MessageType, returning interfaces for types that return event data
+	// @param	ContextData		An array of valid IAnimNotifyEventContextDataInterface.  Only one entry will exist per message type
+	void MakeEventContextData(TArray<TSharedPtr<const IAnimNotifyEventContextDataInterface>>& ContextData) const;
+
 private:
 	// Push a message onto the stack
 	// @param	InNode		The node that this message was pushed from
diff --git a/Engine/Source/Runtime/Engine/Public/Animation/AnimNotifyMirrorInspectionLibrary.h b/Engine/Source/Runtime/Engine/Public/Animation/AnimNotifyMirrorInspectionLibrary.h
new file mode 100644
--- /dev/null
+++ b/Engine/Source/Runtime/Engine/Public/Animation/AnimNotifyMirrorInspectionLibrary.h
@@ -0,0 +1,34 @@
+// Copyright Epic Games, Inc. All Rights Reserved.
+
+#pragma once
+
+#include "CoreMinimal.h"
+#include "Kismet/BlueprintFunctionLibrary.h"
+#include "AnimNotifyMirrorInspectionLibrary.generated.h"
+
+struct FAnimNotifyEventReference;
+/**
+*	A library of commonly used functionality for Notifies related to mirroring, exposed to blueprint.
+*/
+UCLASS(meta = (ScriptName = "UAnimNotifyMirrorInspectionLibrary"))
+class ENGINE_API UAnimNotifyMirrorInspectionLibrary : public UBlueprintFunctionLibrary
+{
+	GENERATED_BODY()
+public:
+
+	/** Get whether the animation which triggered this notify was mirrored.
+	*
+	* @param EventReference		The event to inspect
+	*/
+	UFUNCTION(BlueprintPure, Category = "Utilities|Animation|Notifies")
+    static bool IsTriggeredByMirroredAnimation(const FAnimNotifyEventReference& EventReference);
+
+	/** If the notify is mirrored, return the mirror data table that was active.
+	*
+	* @param EventReference		The event to inspect
+	*/
+	UFUNCTION(BlueprintPure, Category = "Utilities|Animation|Notifies")
+    static const UMirrorDataTable* GetMirrorDataTable(const FAnimNotifyEventReference& EventReference);
+    
+
+};
diff --git a/Engine/Source/Runtime/Engine/Public/Animation/AnimNotifyQueue.h b/Engine/Source/Runtime/Engine/Public/Animation/AnimNotifyQueue.h
--- a/Engine/Source/Runtime/Engine/Public/Animation/AnimNotifyQueue.h
+++ b/Engine/Source/Runtime/Engine/Public/Animation/AnimNotifyQueue.h
@@ -4,51 +4,104 @@
 
 #include "CoreMinimal.h"
 #include "Math/RandomStream.h"
-
+#include "Animation/AnimTypes.h"
+#include "Animation/AnimNodeMessages.h"
 #include "AnimNotifyQueue.generated.h"
 
+
 class USkeletalMeshComponent;
 struct FAnimInstanceProxy;
 struct FAnimNotifyEvent;
+class UMirrorDataTable;
+struct FAnimTickRecord;
 
-USTRUCT()
+
+USTRUCT(BlueprintType)
 struct FAnimNotifyEventReference
 {
 	GENERATED_BODY()
 
 	FAnimNotifyEventReference()
 		: Notify(nullptr)
+		, MirrorTable(nullptr)
 		, NotifySource(nullptr)
+
 	{}
 
 	FAnimNotifyEventReference(const FAnimNotifyEventReference& rhs)
-		: Notify(rhs.Notify)
+		:
+		 ContextData(rhs.ContextData)
+		, Notify(rhs.Notify)
+		, MirrorTable(rhs.MirrorTable)
 		, NotifySource(rhs.NotifySource)
-	{
 
+	{
 	}
 
 	FAnimNotifyEventReference(const FAnimNotifyEvent* InNotify, const UObject* InNotifySource)
 		: Notify(InNotify)
+		, MirrorTable(nullptr)
 		, NotifySource(InNotifySource)
 	{}
 
+	FAnimNotifyEventReference(const FAnimNotifyEvent* InNotify, const UObject* InNotifySource, const UMirrorDataTable* MirrorDataTable)
+    : Notify(InNotify)
+	, MirrorTable(MirrorDataTable)
+    , NotifySource(InNotifySource)
+	{}
+	
 	const FAnimNotifyEvent* GetNotify() const
 	{
 		return NotifySource ? Notify : nullptr;
 	}
 
+	void SetNotify(const FAnimNotifyEvent* InNotify)
+	{
+		if (NotifySource)
+		{
+			Notify = InNotify;
+		}
+	}
+
+	const UMirrorDataTable* GetMirrorDataTable() const
+	{
+		return MirrorTable.Get();
+	}
+	
 	friend bool operator==(const FAnimNotifyEventReference& Lhs, const FAnimNotifyEventReference& Rhs)
 	{
-		return Lhs.Notify == Rhs.Notify;
+		return(
+			(Lhs.Notify == Rhs.Notify) ||
+			(Lhs.Notify && Rhs.Notify && *Lhs.Notify == *Rhs.Notify)
+		);
 	}
 
 	friend bool operator==(const FAnimNotifyEventReference& Lhs, const FAnimNotifyEvent& Rhs);
+	
+	template<typename Type> 
+	const Type* GetContextData() const 
+	{
+		for(const TSharedPtr<const UE::Anim::IAnimNotifyEventContextDataInterface>& DataInterface :  ContextData)
+        {
+        	if (DataInterface->Is<Type>())
+        	{
+        		return &(DataInterface->As<Type>()); 
+        	}
+        }
+        return nullptr; 
+	}
 
+	void SetContextDataInterfaces(const TArray<TSharedPtr<const UE::Anim::IAnimNotifyEventContextDataInterface>>& InContextData);
 private:
 
+	TArray<TSharedPtr<const UE::Anim::IAnimNotifyEventContextDataInterface>> ContextData;
+	
 	const FAnimNotifyEvent* Notify;
 
+	// If set, the Notify has been mirrored.  The mirrored name can be found in MirrorTable->AnimNotifyToMirrorAnimNotifyMap
+	UPROPERTY(transient)
+	TObjectPtr<const UMirrorDataTable> MirrorTable; 
+
 	UPROPERTY(transient)
 	TObjectPtr<const UObject> NotifySource;
 };
@@ -62,6 +115,18 @@ struct FAnimNotifyArray
 	TArray<FAnimNotifyEventReference> Notifies;
 };
 
+USTRUCT()
+struct FAnimNotifyContext
+{
+	GENERATED_BODY()
+	FAnimNotifyContext() {}
+	FAnimNotifyContext(const FAnimTickRecord& InTickRecord)
+    : TickRecord(&InTickRecord)
+	{}
+	const FAnimTickRecord* TickRecord = nullptr;
+	TArray<FAnimNotifyEventReference> ActiveNotifies;
+};
+
 USTRUCT()
 struct FAnimNotifyQueue
 {
diff --git a/Engine/Source/Runtime/Engine/Public/Animation/AnimNotifyStateMachineInspectionLibrary.h b/Engine/Source/Runtime/Engine/Public/Animation/AnimNotifyStateMachineInspectionLibrary.h
new file mode 100644
--- /dev/null
+++ b/Engine/Source/Runtime/Engine/Public/Animation/AnimNotifyStateMachineInspectionLibrary.h
@@ -0,0 +1,61 @@
+// Copyright Epic Games, Inc. All Rights Reserved.
+
+#pragma once
+
+#include "CoreMinimal.h"
+#include "Kismet/BlueprintFunctionLibrary.h"
+#include "AnimNotifyStateMachineInspectionLibrary.generated.h"
+
+struct FAnimNotifyEventReference;
+class USkeletalMeshComponent;
+/**
+*	A library of commonly used functionality for Notifies related to state machines, exposed to blueprint.
+*/
+UCLASS(meta = (ScriptName = "UAnimNotifyStateMachineInspectionLibrary"))
+class ENGINE_API UAnimNotifyStateMachineInspectionLibrary : public UBlueprintFunctionLibrary
+{
+	GENERATED_BODY()
+public:
+	/** Get whether the notify was triggered from the specified state machine
+	*
+    * @param EventReference		The event to inspect
+    * @param MeshComp			The skeletal mesh that contains the animation instance that produced the event
+    * @param StateMachineName	The Name of a state machine to test
+    */
+	UFUNCTION(BlueprintPure, Category = "Utilities|Animation|Notifies")
+    static bool IsTriggeredByStateMachine(const FAnimNotifyEventReference& EventReference,UAnimInstance* AnimInstance, FName StateMachineName);
+	
+	/** Get whether a particular state in a specific state machine triggered the notify
+	*
+	* @param EventReference		The event to inspect
+	* @param MeshComp			The skeletal mesh that contains the animation instance that produced the event
+	* @param StateMachineName	The name of a state machine to test
+	* @param StateName			The name of a state to test
+	*/
+	UFUNCTION(BlueprintPure, Category = "Utilities|Animation|Notifies")
+	static bool IsTriggeredByStateInStateMachine(const FAnimNotifyEventReference& EventReference, UAnimInstance* AnimInstance, FName StateMachineName, FName StateName);
+
+	/** Get whether a state with the given name in any state machine triggered the notify
+	*
+	* @param EventReference		The event to inspect
+	* @param MeshComp			The skeletal mesh that contains the animation instance that produced the event
+	* @param StateName			The name of a state to test
+	*/
+	UFUNCTION(BlueprintPure, Category = "Utilities|Animation|Notifies")
+    static bool IsTriggeredByState(const FAnimNotifyEventReference& EventReference, UAnimInstance* AnimInstance, FName StateName);
+    
+	/** Get whether the Reference ContextData has the given machine index in a UAnimNotifyStateMachineContext.  This function should not be exposed to blueprint
+	*
+	* @param Reference		The event to inspect
+	* @param StateMachineIndex	The MachineIndex as defined in UAnimInstance
+	*/
+	static bool IsStateMachineInEventContext(const FAnimNotifyEventReference& Reference, int32 StateMachineIndex);
+
+	/** Get whether the Reference ContextData has the given state and machine index in a UAnimNotifyStateMachineContext.
+	*	
+	* @param Reference			The event to inspect
+	* @param StateIndex			Index of a state inside a state machine as defined in UAnimInstance
+	* @param StateMachineIndex	The MachineIndex as defined in UAnimInstance
+	*/
+	static bool IsStateInStateMachineInEventContext(const FAnimNotifyEventReference& Reference, int32 StateMachineIndex, int32 StateIndex); 
+};
diff --git a/Engine/Source/Runtime/Engine/Public/Animation/AnimTypes.h b/Engine/Source/Runtime/Engine/Public/Animation/AnimTypes.h
--- a/Engine/Source/Runtime/Engine/Public/Animation/AnimTypes.h
+++ b/Engine/Source/Runtime/Engine/Public/Animation/AnimTypes.h
@@ -442,6 +442,9 @@ public:
 
 	/** Get the NotifyName pre-appended to "AnimNotify_", for calling the event */
 	ENGINE_API FName GetNotifyEventName() const;
+	
+	/** Get the mirrored NotifyName pre-appended to "AnimNotify_", for calling the event. If Notify is not mirrored returns result of GetNotifyEventName()*/
+	ENGINE_API FName GetNotifyEventName(const UMirrorDataTable* MirrorDataTable) const;
 };
 
 #if WITH_EDITORONLY_DATA
@@ -643,7 +646,10 @@ struct FMarkerSyncData
 	TArray<FName>				UniqueMarkerNames;
 
 	void GetMarkerIndicesForTime(float CurrentTime, bool bLooping, const TArray<FName>& ValidMarkerNames, FMarkerPair& OutPrevMarker, FMarkerPair& OutNextMarker, float SequenceLength) const;
-	FMarkerSyncAnimPosition GetMarkerSyncPositionfromMarkerIndicies(int32 PrevMarker, int32 NextMarker, float CurrentTime, float SequenceLength, const UMirrorDataTable* MirrorTable = nullptr) const;
+	
+	UE_DEPRECATED(5.0, "Use other GetMarkerSyncPositionFromMarkerIndicies signature")
+	FMarkerSyncAnimPosition GetMarkerSyncPositionfromMarkerIndicies(int32 PrevMarker, int32 NextMarker, float CurrentTime, float SequenceLength) const;
+	FMarkerSyncAnimPosition GetMarkerSyncPositionFromMarkerIndicies(int32 PrevMarker, int32 NextMarker, float CurrentTime, float SequenceLength, const UMirrorDataTable* MirrorTable = nullptr) const;
 	void CollectUniqueNames();
 	void CollectMarkersInRange(float PrevPosition, float NewPosition, TArray<FPassedMarker>& OutMarkersPassedThisTick, float TotalDeltaMove);
 };
diff --git a/Engine/Source/Runtime/Engine/Public/Animation/MirrorSyncScope.h b/Engine/Source/Runtime/Engine/Public/Animation/MirrorSyncScope.h
--- a/Engine/Source/Runtime/Engine/Public/Animation/MirrorSyncScope.h
+++ b/Engine/Source/Runtime/Engine/Public/Animation/MirrorSyncScope.h
@@ -11,6 +11,17 @@ namespace UE { namespace Anim {
 
 class FAnimSyncGroupScope;
 
+class ENGINE_API FAnimNotifyMirrorContext : public IAnimNotifyEventContextDataInterface
+{
+	DECLARE_NOTIFY_CONTEXT_INTERFACE(FAnimNotifyMirrorContext)
+	public:
+	FAnimNotifyMirrorContext() {} 
+	FAnimNotifyMirrorContext(const UMirrorDataTable* MirrorDataTable) : MirrorTable(MirrorDataTable) {bAnimationMirrored = MirrorDataTable != nullptr;} 
+	bool bAnimationMirrored = true; 
+	UPROPERTY(transient)
+	TObjectPtr<const UMirrorDataTable> MirrorTable; 
+};
+	
 // Scoped graph message used to synchronize mirroring 
 class ENGINE_API FMirrorSyncScope : public IGraphMessage
 {
@@ -19,7 +30,7 @@ public:
 	FMirrorSyncScope(const FAnimationBaseContext& InContext, const UMirrorDataTable* InMirrorDataTable);
 	virtual ~FMirrorSyncScope();
 	int32 GetMirrorScopeDepth() const;
-	
+	TSharedPtr<const IAnimNotifyEventContextDataInterface> MakeEventContextData() const override;
 private:
 	const UMirrorDataTable* MirrorDataTable = nullptr;
 	int32 MirrorScopeDepth = 1; 
