Fix an issue where notify event references were not removed correctly

diff --git a/Engine/Source/Runtime/Engine/Private/Animation/AnimInstance.cpp b/Engine/Source/Runtime/Engine/Private/Animation/AnimInstance.cpp
--- a/Engine/Source/Runtime/Engine/Private/Animation/AnimInstance.cpp
+++ b/Engine/Source/Runtime/Engine/Private/Animation/AnimInstance.cpp
@@ -1374,17 +1374,18 @@ void UAnimInstance::TriggerAnimNotifies(float DeltaSeconds)
 			// AnimNotifyState
 			if (AnimNotifyEvent->NotifyStateClass)
 			{
-				if (!ActiveAnimNotifyState.RemoveSingleSwap(*AnimNotifyEvent, false))
+				int32 ExistingItemIndex = INDEX_NONE;
+
+				if (ActiveAnimNotifyState.Find(*AnimNotifyEvent, ExistingItemIndex))
 				{
-					// Queue up calls to 'NotifyBegin', so they happen after 'NotifyEnd'.
-					NotifyStateBeginEvent.Add(AnimNotifyEvent);
-					NotifyStateBeginEventReference.Add(&NotifyQueue.AnimNotifies[Index]);
+					check(ActiveAnimNotifyState.Num() == ActiveAnimNotifyEventReference.Num());
+					ActiveAnimNotifyState.RemoveAtSwap(ExistingItemIndex, 1, false); 
+					ActiveAnimNotifyEventReference.RemoveAtSwap(ExistingItemIndex, 1, false);
 				}
-				else if(!ActiveAnimNotifyEventReference.IsEmpty())
+				else
 				{
-					// Remove the matching Event Reference
-					bool bRemovedEvent = ActiveAnimNotifyEventReference.RemoveSingleSwap(NotifyQueue.AnimNotifies[Index]) ==1;
-					check(bRemovedEvent);
+					NotifyStateBeginEvent.Add(AnimNotifyEvent);
+					NotifyStateBeginEventReference.Add(&NotifyQueue.AnimNotifies[Index]);
 				}
 
 				NewActiveAnimNotifyState.Add(*AnimNotifyEvent);
