Fix crash in cooked games with standalone asset players

Also fix blendspace sync groups not working correctly

diff --git a/Engine/Source/Editor/AnimGraph/Private/AnimGraphNode_AssetPlayerBase.cpp b/Engine/Source/Editor/AnimGraph/Private/AnimGraphNode_AssetPlayerBase.cpp
--- a/Engine/Source/Editor/AnimGraph/Private/AnimGraphNode_AssetPlayerBase.cpp
+++ b/Engine/Source/Editor/AnimGraph/Private/AnimGraphNode_AssetPlayerBase.cpp
@@ -31,7 +31,7 @@ void UAnimGraphNode_AssetPlayerBase::Serialize(FArchive& Ar)
 
 	Ar.UsingCustomVersion(FUE5MainStreamObjectVersion::GUID);
 
-	if(Ar.CustomVer(FUE5MainStreamObjectVersion::GUID) < FUE5MainStreamObjectVersion::AnimSyncGroupsExplicitSyncMethod)
+	if(Ar.IsLoading() && Ar.CustomVer(FUE5MainStreamObjectVersion::GUID) < FUE5MainStreamObjectVersion::AnimSyncGroupsExplicitSyncMethod)
 	{
 		if(SyncGroup_DEPRECATED.GroupName != NAME_None)
 		{
@@ -39,14 +39,14 @@ void UAnimGraphNode_AssetPlayerBase::Serialize(FArchive& Ar)
 		}
 	}
 
-	if(Ar.CustomVer(FUE5MainStreamObjectVersion::GUID) < FUE5MainStreamObjectVersion::AnimNodeConstantDataRefactorPhase0)
+	if(Ar.IsLoading() && Ar.CustomVer(FUE5MainStreamObjectVersion::GUID) < FUE5MainStreamObjectVersion::AnimNodeConstantDataRefactorPhase0)
 	{
 		FStructProperty* NodeProperty = GetFNodeProperty();
 		if(NodeProperty->Struct->IsChildOf(FAnimNode_AssetPlayerBase::StaticStruct()))
 		{
 			FAnimNode_AssetPlayerBase* Node = NodeProperty->ContainerPtrToValuePtr<FAnimNode_AssetPlayerBase>(this);
 			Node->SetGroupName(SyncGroup_DEPRECATED.GroupName);
-			Node->GetGroupRole(SyncGroup_DEPRECATED.GroupRole);
+			Node->SetGroupRole(SyncGroup_DEPRECATED.GroupRole);
 			Node->SetGroupMethod(SyncGroup_DEPRECATED.Method);
 		}
 	}
@@ -59,13 +59,13 @@ void UAnimGraphNode_AssetPlayerBase::PostEditChangeProperty(FPropertyChangedEven
 	FStructProperty* NodeProperty = GetFNodeProperty();
 	if(NodeProperty->Struct->IsChildOf(FAnimNode_AssetPlayerBase::StaticStruct()))
 	{
-		if(PropertyChangedEvent.GetPropertyName() == GET_MEMBER_NAME_CHECKED(FAnimNode_AssetPlayerBase, Method))
+		if(PropertyChangedEvent.GetPropertyName() == TEXT("Method"))
 		{
 			FAnimNode_AssetPlayerBase* Node = NodeProperty->ContainerPtrToValuePtr<FAnimNode_AssetPlayerBase>(this);
 			if(Node->GetGroupMethod() != EAnimSyncMethod::SyncGroup)
 			{
 				Node->SetGroupName(NAME_None);
-				Node->GetGroupRole(EAnimGroupRole::CanBeLeader);
+				Node->SetGroupRole(EAnimGroupRole::CanBeLeader);
 			}
 		}
 	}
diff --git a/Engine/Source/Runtime/AnimGraphRuntime/Private/AnimNodes/AnimNode_BlendSpacePlayer.cpp b/Engine/Source/Runtime/AnimGraphRuntime/Private/AnimNodes/AnimNode_BlendSpacePlayer.cpp
--- a/Engine/Source/Runtime/AnimGraphRuntime/Private/AnimNodes/AnimNode_BlendSpacePlayer.cpp
+++ b/Engine/Source/Runtime/AnimGraphRuntime/Private/AnimNodes/AnimNode_BlendSpacePlayer.cpp
@@ -234,3 +234,51 @@ UBlendSpace* FAnimNode_BlendSpacePlayer::GetBlendSpace() const
 {
 	return GET_ANIM_NODE_DATA(TObjectPtr<UBlendSpace>, BlendSpace);
 }
+
+FName FAnimNode_BlendSpacePlayer::GetGroupName() const
+{
+	return GET_ANIM_NODE_DATA(FName, GroupName);
+}
+
+EAnimGroupRole::Type FAnimNode_BlendSpacePlayer::GetGroupRole() const
+{
+	return GET_ANIM_NODE_DATA(EAnimGroupRole::Type, GroupRole);
+}
+
+EAnimSyncMethod FAnimNode_BlendSpacePlayer::GetGroupMethod() const
+{
+	return GET_ANIM_NODE_DATA(EAnimSyncMethod, Method);
+}
+
+bool FAnimNode_BlendSpacePlayer::GetIgnoreForRelevancyTest() const
+{
+	return GET_ANIM_NODE_DATA(bool, bIgnoreForRelevancyTest);
+}
+
+void FAnimNode_BlendSpacePlayer::SetGroupName(FName InGroupName)
+{
+#if WITH_EDITORONLY_DATA	
+	GroupName = InGroupName;
+#endif
+}
+
+void FAnimNode_BlendSpacePlayer::SetGroupRole(EAnimGroupRole::Type InRole)
+{
+#if WITH_EDITORONLY_DATA
+	GroupRole = InRole;
+#endif
+}
+
+void FAnimNode_BlendSpacePlayer::SetGroupMethod(EAnimSyncMethod InMethod)
+{
+#if WITH_EDITORONLY_DATA
+	Method = InMethod;
+#endif
+}
+
+void FAnimNode_BlendSpacePlayer::SetIgnoreForRelevancyTest(bool bInIgnoreForRelevancyTest)
+{
+#if WITH_EDITORONLY_DATA
+	bIgnoreForRelevancyTest = bInIgnoreForRelevancyTest;
+#endif
+}
\ No newline at end of file
diff --git a/Engine/Source/Runtime/AnimGraphRuntime/Private/AnimNodes/AnimNode_SequenceEvaluator.cpp b/Engine/Source/Runtime/AnimGraphRuntime/Private/AnimNodes/AnimNode_SequenceEvaluator.cpp
--- a/Engine/Source/Runtime/AnimGraphRuntime/Private/AnimNodes/AnimNode_SequenceEvaluator.cpp
+++ b/Engine/Source/Runtime/AnimGraphRuntime/Private/AnimNodes/AnimNode_SequenceEvaluator.cpp
@@ -157,3 +157,51 @@ float FAnimNode_SequenceEvaluator::GetStartPosition() const
 {
 	return GET_ANIM_NODE_DATA(float, StartPosition);
 }
+
+FName FAnimNode_SequenceEvaluator::GetGroupName() const
+{
+	return GET_ANIM_NODE_DATA(FName, GroupName);
+}
+
+EAnimGroupRole::Type FAnimNode_SequenceEvaluator::GetGroupRole() const
+{
+	return GET_ANIM_NODE_DATA(EAnimGroupRole::Type, GroupRole);
+}
+
+EAnimSyncMethod FAnimNode_SequenceEvaluator::GetGroupMethod() const
+{
+	return GET_ANIM_NODE_DATA(EAnimSyncMethod, Method);
+}
+
+bool FAnimNode_SequenceEvaluator::GetIgnoreForRelevancyTest() const
+{
+	return GET_ANIM_NODE_DATA(bool, bIgnoreForRelevancyTest);
+}
+
+void FAnimNode_SequenceEvaluator::SetGroupName(FName InGroupName)
+{
+#if WITH_EDITORONLY_DATA	
+	GroupName = InGroupName;
+#endif
+}
+
+void FAnimNode_SequenceEvaluator::SetGroupRole(EAnimGroupRole::Type InRole)
+{
+#if WITH_EDITORONLY_DATA
+	GroupRole = InRole;
+#endif
+}
+
+void FAnimNode_SequenceEvaluator::SetGroupMethod(EAnimSyncMethod InMethod)
+{
+#if WITH_EDITORONLY_DATA
+	Method = InMethod;
+#endif
+}
+
+void FAnimNode_SequenceEvaluator::SetIgnoreForRelevancyTest(bool bInIgnoreForRelevancyTest)
+{
+#if WITH_EDITORONLY_DATA
+	bIgnoreForRelevancyTest = bInIgnoreForRelevancyTest;
+#endif
+}
\ No newline at end of file
diff --git a/Engine/Source/Runtime/AnimGraphRuntime/Public/AnimNodes/AnimNode_BlendSpacePlayer.h b/Engine/Source/Runtime/AnimGraphRuntime/Public/AnimNodes/AnimNode_BlendSpacePlayer.h
--- a/Engine/Source/Runtime/AnimGraphRuntime/Public/AnimNodes/AnimNode_BlendSpacePlayer.h
+++ b/Engine/Source/Runtime/AnimGraphRuntime/Public/AnimNodes/AnimNode_BlendSpacePlayer.h
@@ -26,6 +26,22 @@ struct ANIMGRAPHRUNTIME_API FAnimNode_BlendSpacePlayer : public FAnimNode_AssetP
 private:
 
 #if WITH_EDITORONLY_DATA
+	// The group name (NAME_None if it is not part of any group)
+	UPROPERTY(EditAnywhere, Category=Sync, meta=(FoldProperty))
+	FName GroupName = NAME_None;
+
+	// The role this player can assume within the group (ignored if GroupIndex is INDEX_NONE)
+	UPROPERTY(EditAnywhere, Category=Sync, meta=(FoldProperty))
+	TEnumAsByte<EAnimGroupRole::Type> GroupRole = EAnimGroupRole::CanBeLeader;
+
+	// How synchronization is determined
+	UPROPERTY(EditAnywhere, Category=Sync, meta=(FoldProperty))
+	EAnimSyncMethod Method = EAnimSyncMethod::DoNotSync;
+
+	// If true, "Relevant anim" nodes that look for the highest weighted animation in a state will ignore this node
+	UPROPERTY(EditAnywhere, Category=Relevancy, meta=(FoldProperty, PinHiddenByDefault))
+	bool bIgnoreForRelevancyTest = false;
+	
 	// The X coordinate to sample in the blendspace
 	UPROPERTY(EditAnywhere, Category=Coordinates, meta=(PinShownByDefault, FoldProperty))
 	float X = 0.0f;
@@ -81,6 +97,14 @@ public:
 	virtual float GetCurrentAssetTimePlayRateAdjusted() const override;
 	virtual float GetCurrentAssetLength() const override;
 	virtual UAnimationAsset* GetAnimAsset() const override;
+	virtual FName GetGroupName() const override;
+	virtual EAnimGroupRole::Type GetGroupRole() const override;
+	virtual EAnimSyncMethod GetGroupMethod() const override;
+	virtual bool GetIgnoreForRelevancyTest() const override;
+	virtual void SetGroupName(FName InGroupName) override;
+	virtual void SetGroupRole(EAnimGroupRole::Type InRole) override;
+	virtual void SetGroupMethod(EAnimSyncMethod InMethod) override;
+	virtual void SetIgnoreForRelevancyTest(bool bInIgnoreForRelevancyTest) override;
 	// End of FAnimNode_AssetPlayerBase interface
 
 	// FAnimNode_Base interface
diff --git a/Engine/Source/Runtime/AnimGraphRuntime/Public/AnimNodes/AnimNode_SequenceEvaluator.h b/Engine/Source/Runtime/AnimGraphRuntime/Public/AnimNodes/AnimNode_SequenceEvaluator.h
--- a/Engine/Source/Runtime/AnimGraphRuntime/Public/AnimNodes/AnimNode_SequenceEvaluator.h
+++ b/Engine/Source/Runtime/AnimGraphRuntime/Public/AnimNodes/AnimNode_SequenceEvaluator.h
@@ -98,6 +98,22 @@ private:
 	friend class UAnimGraphNode_SequenceEvaluator;
 
 #if WITH_EDITORONLY_DATA
+	// The group name (NAME_None if it is not part of any group)
+	UPROPERTY(EditAnywhere, Category=Sync, meta=(FoldProperty))
+	FName GroupName = NAME_None;
+
+	// The role this player can assume within the group (ignored if GroupIndex is INDEX_NONE)
+	UPROPERTY(EditAnywhere, Category=Sync, meta=(FoldProperty))
+	TEnumAsByte<EAnimGroupRole::Type> GroupRole = EAnimGroupRole::CanBeLeader;
+
+	// How synchronization is determined
+	UPROPERTY(EditAnywhere, Category=Sync, meta=(FoldProperty))
+	EAnimSyncMethod Method = EAnimSyncMethod::DoNotSync;
+
+	// If true, "Relevant anim" nodes that look for the highest weighted animation in a state will ignore this node
+	UPROPERTY(EditAnywhere, Category=Relevancy, meta=(FoldProperty, PinHiddenByDefault))
+	bool bIgnoreForRelevancyTest = false;
+	
 	// The animation sequence asset to evaluate
 	UPROPERTY(EditAnywhere, Category = Settings, meta = (PinHiddenByDefault, FoldProperty))
 	TObjectPtr<UAnimSequenceBase> Sequence = nullptr;
@@ -134,6 +150,16 @@ public:
 	virtual bool GetTeleportToExplicitTime() const override;
 	virtual TEnumAsByte<ESequenceEvalReinit::Type> GetReinitializationBehavior() const override;
 	virtual float GetStartPosition() const override;
+
+	// FAnimNode_AssetPlayerBase interface
+	virtual FName GetGroupName() const override;
+	virtual EAnimGroupRole::Type GetGroupRole() const override;
+	virtual EAnimSyncMethod GetGroupMethod() const override;
+	virtual bool GetIgnoreForRelevancyTest() const override;
+	virtual void SetGroupName(FName InGroupName) override;
+	virtual void SetGroupRole(EAnimGroupRole::Type InRole) override;
+	virtual void SetGroupMethod(EAnimSyncMethod InMethod) override;
+	virtual void SetIgnoreForRelevancyTest(bool bInIgnoreForRelevancyTest) override;
 };
 
 // Sequence evaluator node that can be used standalone (without constant folding)
@@ -143,6 +169,22 @@ struct ANIMGRAPHRUNTIME_API FAnimNode_SequenceEvaluator_Standalone : public FAni
 	GENERATED_BODY()
 
 private:
+	// The group name (NAME_None if it is not part of any group)
+	UPROPERTY(EditAnywhere, Category=Sync)
+	FName GroupName = NAME_None;
+
+	// The role this player can assume within the group (ignored if GroupIndex is INDEX_NONE)
+	UPROPERTY(EditAnywhere, Category=Sync)
+	TEnumAsByte<EAnimGroupRole::Type> GroupRole = EAnimGroupRole::CanBeLeader;
+
+	// How synchronization is determined
+	UPROPERTY(EditAnywhere, Category=Sync)
+	EAnimSyncMethod Method = EAnimSyncMethod::DoNotSync;
+
+	// If true, "Relevant anim" nodes that look for the highest weighted animation in a state will ignore this node
+	UPROPERTY(EditAnywhere, Category=Relevancy, meta=(PinHiddenByDefault))
+	bool bIgnoreForRelevancyTest = false;
+	
 	// The animation sequence asset to evaluate
 	UPROPERTY(EditAnywhere, Category = Settings, meta = (PinHiddenByDefault))
 	TObjectPtr<UAnimSequenceBase> Sequence = nullptr;
@@ -181,4 +223,14 @@ public:
 	virtual bool GetTeleportToExplicitTime() const override { return bTeleportToExplicitTime; }
 	virtual TEnumAsByte<ESequenceEvalReinit::Type> GetReinitializationBehavior() const override { return ReinitializationBehavior; }
 	virtual float GetStartPosition() const override { return StartPosition; }
+
+	// FAnimNode_AssetPlayerBase interface
+	virtual FName GetGroupName() const override { return GroupName; }
+	virtual EAnimGroupRole::Type GetGroupRole() const override { return GroupRole; }
+	virtual EAnimSyncMethod GetGroupMethod() const override { return Method; }
+	virtual bool GetIgnoreForRelevancyTest() const override { return bIgnoreForRelevancyTest; }
+	virtual void SetGroupName(FName InGroupName) override { GroupName = InGroupName; }
+	virtual void SetGroupRole(EAnimGroupRole::Type InRole) override { GroupRole = InRole; }
+	virtual void SetGroupMethod(EAnimSyncMethod InMethod) override { Method = InMethod; }
+	virtual void SetIgnoreForRelevancyTest(bool bInIgnoreForRelevancyTest) override { bIgnoreForRelevancyTest = bInIgnoreForRelevancyTest; }
 };
\ No newline at end of file
diff --git a/Engine/Source/Runtime/Engine/Classes/Animation/AnimNode_AssetPlayerBase.h b/Engine/Source/Runtime/Engine/Classes/Animation/AnimNode_AssetPlayerBase.h
--- a/Engine/Source/Runtime/Engine/Classes/Animation/AnimNode_AssetPlayerBase.h
+++ b/Engine/Source/Runtime/Engine/Classes/Animation/AnimNode_AssetPlayerBase.h
@@ -12,7 +12,7 @@
 USTRUCT(BlueprintInternalUseOnly)
 struct ENGINE_API FAnimNode_AssetPlayerBase : public FAnimNode_Base
 {
-	GENERATED_BODY();
+	GENERATED_BODY()
 
 	friend class UAnimGraphNode_AssetPlayerBase;
 
@@ -39,7 +39,6 @@ struct ENGINE_API FAnimNode_AssetPlayerBase : public FAnimNode_Base
 	/** Update the node, marked final so we can always handle blendweight caching.
 	 *  Derived classes should implement UpdateAssetPlayer
 	 */
-
 	virtual void Update_AnyThread(const FAnimationUpdateContext& Context) final override;
 
 	/** Update method for the asset player, to be implemented by derived classes */
@@ -54,59 +53,38 @@ struct ENGINE_API FAnimNode_AssetPlayerBase : public FAnimNode_Base
 	virtual float GetCurrentAssetTimePlayRateAdjusted() const { return GetCurrentAssetTime(); }
 
 	// Get the sync group name we are using
-	FName GetGroupName() const;
+	virtual FName GetGroupName() const { return NAME_None; }
 
 	// Get the sync group role we are using
-	EAnimGroupRole::Type GetGroupRole() const;
+	virtual EAnimGroupRole::Type GetGroupRole() const { return EAnimGroupRole::CanBeLeader; }
 
 	// Get the sync group method we are using
-	EAnimSyncMethod GetGroupMethod() const;
+	virtual EAnimSyncMethod GetGroupMethod() const { return EAnimSyncMethod::DoNotSync; }
 
 	// Check whether this node should be ignored when testing for relevancy in state machines
-	bool GetIgnoreForRelevancyTest() const;
+	virtual bool GetIgnoreForRelevancyTest() const { return false; }
 
-#if WITH_EDITORONLY_DATA
 	// Set the sync group name we are using
-	void SetGroupName(FName InGroupName) { GroupName = InGroupName; }
+	virtual void SetGroupName(FName InGroupName) { }
 
 	// Set the sync group role we are using
-	void GetGroupRole(EAnimGroupRole::Type InRole) { GroupRole = InRole; }
+	virtual void SetGroupRole(EAnimGroupRole::Type InRole) { }
 
 	// Set the sync group method we are using
-	void SetGroupMethod(EAnimSyncMethod InMethod) { Method = InMethod; }
+	virtual void SetGroupMethod(EAnimSyncMethod InMethod) { }
 
 	// Set whether this node should be ignored when testing for relevancy in state machines
-	void SetIgnoreForRelevancyTest(bool bInIgnoreForRelevancyTest) { bIgnoreForRelevancyTest = bInIgnoreForRelevancyTest; }
-#endif
+	virtual void SetIgnoreForRelevancyTest(bool bInIgnoreForRelevancyTest) { }
 
 private:
-
 #if WITH_EDITORONLY_DATA
-	// The group name (NAME_None if it is not part of any group)
-	UPROPERTY(EditAnywhere, Category=Sync, meta=(FoldProperty))
-	FName GroupName = NAME_None;
-
 	UPROPERTY()
 	int32 GroupIndex_DEPRECATED = INDEX_NONE;
 
 	UPROPERTY()
 	EAnimSyncGroupScope GroupScope_DEPRECATED = EAnimSyncGroupScope::Local;
-
-	// The role this player can assume within the group (ignored if GroupIndex is INDEX_NONE)
-	UPROPERTY(EditAnywhere, Category=Sync, meta=(FoldProperty))
-	TEnumAsByte<EAnimGroupRole::Type> GroupRole = EAnimGroupRole::CanBeLeader;
-
-	// How synchronization is determined
-	UPROPERTY(EditAnywhere, Category=Sync, meta=(FoldProperty))
-	EAnimSyncMethod Method = EAnimSyncMethod::DoNotSync;
-
-	/** If true, "Relevant anim" nodes that look for the highest weighted animation in a state will ignore
-	 *  this node
-	 */
-	UPROPERTY(EditAnywhere, Category=Relevancy, meta=(FoldProperty, PinHiddenByDefault))
-	bool bIgnoreForRelevancyTest = false;
 #endif
-
+	
 protected:
 	/** Store data about current marker position when using marker based syncing*/
 	FMarkerTickRecord MarkerTickRecord;
diff --git a/Engine/Source/Runtime/Engine/Classes/Animation/AnimNode_SequencePlayer.h b/Engine/Source/Runtime/Engine/Classes/Animation/AnimNode_SequencePlayer.h
--- a/Engine/Source/Runtime/Engine/Classes/Animation/AnimNode_SequencePlayer.h
+++ b/Engine/Source/Runtime/Engine/Classes/Animation/AnimNode_SequencePlayer.h
@@ -82,6 +82,22 @@ protected:
 	friend class UAnimGraphNode_SequencePlayer;
 
 #if WITH_EDITORONLY_DATA
+	// The group name (NAME_None if it is not part of any group)
+	UPROPERTY(EditAnywhere, Category=Sync, meta=(FoldProperty))
+	FName GroupName = NAME_None;
+
+	// The role this player can assume within the group (ignored if GroupIndex is INDEX_NONE)
+	UPROPERTY(EditAnywhere, Category=Sync, meta=(FoldProperty))
+	TEnumAsByte<EAnimGroupRole::Type> GroupRole = EAnimGroupRole::CanBeLeader;
+
+	// How synchronization is determined
+	UPROPERTY(EditAnywhere, Category=Sync, meta=(FoldProperty))
+	EAnimSyncMethod Method = EAnimSyncMethod::DoNotSync;
+
+	// If true, "Relevant anim" nodes that look for the highest weighted animation in a state will ignore this node
+	UPROPERTY(EditAnywhere, Category=Relevancy, meta=(FoldProperty, PinHiddenByDefault))
+	bool bIgnoreForRelevancyTest = false;
+	
 	// The animation sequence asset to play
 	UPROPERTY(EditAnywhere, Category = Settings, meta = (PinHiddenByDefault, DisallowedClasses="AnimMontage", FoldProperty))
 	TObjectPtr<UAnimSequenceBase> Sequence = nullptr;
@@ -128,6 +144,16 @@ public:
 	virtual float GetStartPosition() const override;
 	virtual bool GetLoopAnimation() const override;
 	virtual bool GetStartFromMatchingPose() const override;
+
+	// FAnimNode_AssetPlayerBase interface
+	virtual FName GetGroupName() const override;
+	virtual EAnimGroupRole::Type GetGroupRole() const override;
+	virtual EAnimSyncMethod GetGroupMethod() const override;
+	virtual bool GetIgnoreForRelevancyTest() const override;
+	virtual void SetGroupName(FName InGroupName) override;
+	virtual void SetGroupRole(EAnimGroupRole::Type InRole) override;
+	virtual void SetGroupMethod(EAnimSyncMethod InMethod) override;
+	virtual void SetIgnoreForRelevancyTest(bool bInIgnoreForRelevancyTest) override;	
 };
 
 // Sequence player node that can be used standalone (without constant folding)
@@ -137,6 +163,22 @@ struct ENGINE_API FAnimNode_SequencePlayer_Standalone : public FAnimNode_Sequenc
 	GENERATED_BODY()
 
 protected:
+	// The group name (NAME_None if it is not part of any group)
+	UPROPERTY(EditAnywhere, Category=Sync)
+	FName GroupName = NAME_None;
+
+	// The role this player can assume within the group (ignored if GroupIndex is INDEX_NONE)
+	UPROPERTY(EditAnywhere, Category=Sync)
+	TEnumAsByte<EAnimGroupRole::Type> GroupRole = EAnimGroupRole::CanBeLeader;
+
+	// How synchronization is determined
+	UPROPERTY(EditAnywhere, Category=Sync)
+	EAnimSyncMethod Method = EAnimSyncMethod::DoNotSync;
+
+	// If true, "Relevant anim" nodes that look for the highest weighted animation in a state will ignore this node
+	UPROPERTY(EditAnywhere, Category=Relevancy, meta=(PinHiddenByDefault))
+	bool bIgnoreForRelevancyTest = false;
+
 	// The animation sequence asset to play
 	UPROPERTY(EditAnywhere, Category = Settings, meta = (PinHiddenByDefault, DisallowedClasses="AnimMontage"))
 	TObjectPtr<UAnimSequenceBase> Sequence = nullptr;
@@ -170,13 +212,23 @@ protected:
 
 public:
 	// FAnimNode_SequencePlayerBase interface
-	virtual void SetSequence(UAnimSequenceBase* InSequence) override;
-	virtual void SetLoopAnimation(bool bInLoopAnimation) override;
-	virtual UAnimSequenceBase* GetSequence() const override;
-	virtual float GetPlayRateBasis() const override;
-	virtual float GetPlayRate() const override;
-	virtual const FInputScaleBiasClampConstants& GetPlayRateScaleBiasClampConstants() const override;
-	virtual float GetStartPosition() const override;
-	virtual bool GetLoopAnimation() const override;
-	virtual bool GetStartFromMatchingPose() const override;
+	virtual void SetSequence(UAnimSequenceBase* InSequence) override { Sequence = InSequence; }
+	virtual void SetLoopAnimation(bool bInLoopAnimation) override { bLoopAnimation = bInLoopAnimation; }
+	virtual UAnimSequenceBase* GetSequence() const override { return Sequence; }
+	virtual float GetPlayRateBasis() const override { return PlayRateBasis; }
+	virtual float GetPlayRate() const override { return PlayRate; }
+	virtual const FInputScaleBiasClampConstants& GetPlayRateScaleBiasClampConstants() const override { return PlayRateScaleBiasClampConstants; }
+	virtual float GetStartPosition() const override { return StartPosition; }
+	virtual bool GetLoopAnimation() const override { return bLoopAnimation; }
+	virtual bool GetStartFromMatchingPose() const override { return bStartFromMatchingPose; }
+
+	// FAnimNode_AssetPlayerBase interface
+	virtual FName GetGroupName() const override { return GroupName; }
+	virtual EAnimGroupRole::Type GetGroupRole() const override { return GroupRole; }
+	virtual EAnimSyncMethod GetGroupMethod() const override { return Method; }
+	virtual bool GetIgnoreForRelevancyTest() const override { return bIgnoreForRelevancyTest; }
+	virtual void SetGroupName(FName InGroupName) override { GroupName = InGroupName; }
+	virtual void SetGroupRole(EAnimGroupRole::Type InRole) override { GroupRole = InRole; }
+	virtual void SetGroupMethod(EAnimSyncMethod InMethod) override { Method = InMethod; }
+	virtual void SetIgnoreForRelevancyTest(bool bInIgnoreForRelevancyTest) override { bIgnoreForRelevancyTest = bInIgnoreForRelevancyTest; }
 };
diff --git a/Engine/Source/Runtime/Engine/Private/Animation/AnimNode_AssetPlayerBase.cpp b/Engine/Source/Runtime/Engine/Private/Animation/AnimNode_AssetPlayerBase.cpp
--- a/Engine/Source/Runtime/Engine/Private/Animation/AnimNode_AssetPlayerBase.cpp
+++ b/Engine/Source/Runtime/Engine/Private/Animation/AnimNode_AssetPlayerBase.cpp
@@ -31,14 +31,14 @@ void FAnimNode_AssetPlayerBase::CreateTickRecordForNode(const FAnimationUpdateCo
 	const EAnimGroupRole::Type SyncGroupRole = GetGroupRole();
 	const FName SyncGroupName = GetGroupName();
 
-	FName GroupNameToUse = ((SyncGroupRole < EAnimGroupRole::TransitionLeader) || bHasBeenFullWeight) ? SyncGroupName : NAME_None;
+	const FName GroupNameToUse = ((SyncGroupRole < EAnimGroupRole::TransitionLeader) || bHasBeenFullWeight) ? SyncGroupName : NAME_None;
 	EAnimSyncMethod MethodToUse = GetGroupMethod();
 	if(GroupNameToUse == NAME_None && MethodToUse == EAnimSyncMethod::SyncGroup)
 	{
 		MethodToUse = EAnimSyncMethod::DoNotSync;
 	}
 
-	UE::Anim::FAnimSyncParams SyncParams(GroupNameToUse, SyncGroupRole, MethodToUse);
+	const UE::Anim::FAnimSyncParams SyncParams(GroupNameToUse, SyncGroupRole, MethodToUse);
 	FAnimTickRecord TickRecord(Sequence, bLooping, PlayRate, FinalBlendWeight, /*inout*/ InternalTimeAccumulator, MarkerTickRecord);
 	if(Context.GetSharedContext())
 	{
@@ -75,23 +75,3 @@ void FAnimNode_AssetPlayerBase::ClearCachedBlendWeight()
 {
 	BlendWeight = 0.0f;
 }
-
-FName FAnimNode_AssetPlayerBase::GetGroupName() const
-{
-	return GET_ANIM_NODE_DATA(FName, GroupName);
-}
-
-EAnimGroupRole::Type FAnimNode_AssetPlayerBase::GetGroupRole() const
-{
-	return GET_ANIM_NODE_DATA(EAnimGroupRole::Type, GroupRole);
-}
-
-EAnimSyncMethod FAnimNode_AssetPlayerBase::GetGroupMethod() const
-{
-	return GET_ANIM_NODE_DATA(EAnimSyncMethod, Method);
-}
-
-bool FAnimNode_AssetPlayerBase::GetIgnoreForRelevancyTest() const
-{
-	return GET_ANIM_NODE_DATA(bool, bIgnoreForRelevancyTest);
-}
\ No newline at end of file
diff --git a/Engine/Source/Runtime/Engine/Private/Animation/AnimNode_SequencePlayer.cpp b/Engine/Source/Runtime/Engine/Private/Animation/AnimNode_SequencePlayer.cpp
--- a/Engine/Source/Runtime/Engine/Private/Animation/AnimNode_SequencePlayer.cpp
+++ b/Engine/Source/Runtime/Engine/Private/Animation/AnimNode_SequencePlayer.cpp
@@ -170,7 +170,7 @@ float FAnimNode_SequencePlayerBase::GetEffectiveStartPosition(const FAnimationBa
 
 void FAnimNode_SequencePlayer::SetSequence(UAnimSequenceBase* InSequence)
 {
-#if WITH_EDITORONLY_DATA	
+#if WITH_EDITORONLY_DATA
 	TObjectPtr<UAnimSequenceBase>& SequenceToSet = GET_MUTABLE_ANIM_NODE_DATA(TObjectPtr<UAnimSequenceBase>, Sequence);
 	SequenceToSet = InSequence;
 #endif
@@ -178,9 +178,8 @@ void FAnimNode_SequencePlayer::SetSequence(UAnimSequenceBase* InSequence)
 
 void FAnimNode_SequencePlayer::SetLoopAnimation(bool bInLoopAnimation)
 {
-#if WITH_EDITORONLY_DATA	
-	bool& bValueToSet = GET_MUTABLE_ANIM_NODE_DATA(bool, bLoopAnimation);
-	bValueToSet = bInLoopAnimation;
+#if WITH_EDITORONLY_DATA
+	bLoopAnimation = bInLoopAnimation;
 #endif
 }
 
@@ -219,49 +218,52 @@ bool FAnimNode_SequencePlayer::GetStartFromMatchingPose() const
 	return GET_ANIM_NODE_DATA(bool, bStartFromMatchingPose);
 }
 
-void FAnimNode_SequencePlayer_Standalone::SetSequence(UAnimSequenceBase* InSequence)
-{
-	Sequence = InSequence;
-}
-
-void FAnimNode_SequencePlayer_Standalone::SetLoopAnimation(bool bInLoopAnimation)
+FName FAnimNode_SequencePlayer::GetGroupName() const
 {
-	bLoopAnimation = bInLoopAnimation;
+	return GET_ANIM_NODE_DATA(FName, GroupName);
 }
 
-UAnimSequenceBase* FAnimNode_SequencePlayer_Standalone::GetSequence() const
+EAnimGroupRole::Type FAnimNode_SequencePlayer::GetGroupRole() const
 {
-	return Sequence;
+	return GET_ANIM_NODE_DATA(EAnimGroupRole::Type, GroupRole);
 }
 
-float FAnimNode_SequencePlayer_Standalone::GetPlayRateBasis() const
+EAnimSyncMethod FAnimNode_SequencePlayer::GetGroupMethod() const
 {
-	return PlayRateBasis;
+	return GET_ANIM_NODE_DATA(EAnimSyncMethod, Method);
 }
 
-float FAnimNode_SequencePlayer_Standalone::GetPlayRate() const
+bool FAnimNode_SequencePlayer::GetIgnoreForRelevancyTest() const
 {
-	return PlayRate;
+	return GET_ANIM_NODE_DATA(bool, bIgnoreForRelevancyTest);
 }
 
-const FInputScaleBiasClampConstants& FAnimNode_SequencePlayer_Standalone::GetPlayRateScaleBiasClampConstants() const
+void FAnimNode_SequencePlayer::SetGroupName(FName InGroupName)
 {
-	return PlayRateScaleBiasClampConstants;
+#if WITH_EDITORONLY_DATA
+	GroupName = InGroupName;
+#endif
 }
 
-float FAnimNode_SequencePlayer_Standalone::GetStartPosition() const
+void FAnimNode_SequencePlayer::SetGroupRole(EAnimGroupRole::Type InRole)
 {
-	return StartPosition;
+#if WITH_EDITORONLY_DATA
+	GroupRole = InRole;
+#endif
 }
 
-bool FAnimNode_SequencePlayer_Standalone::GetLoopAnimation() const
+void FAnimNode_SequencePlayer::SetGroupMethod(EAnimSyncMethod InMethod)
 {
-	return bLoopAnimation;
+#if WITH_EDITORONLY_DATA
+	Method = InMethod;
+#endif
 }
 
-bool FAnimNode_SequencePlayer_Standalone::GetStartFromMatchingPose() const
+void FAnimNode_SequencePlayer::SetIgnoreForRelevancyTest(bool bInIgnoreForRelevancyTest)
 {
-	return bStartFromMatchingPose;
+#if WITH_EDITORONLY_DATA
+	bIgnoreForRelevancyTest = bInIgnoreForRelevancyTest;
+#endif
 }
 
 #undef LOCTEXT_NAMESPACE
