Fix notify functions to ensure blueprints compile correctly with EventReference parameter

diff --git a/Engine/Source/Runtime/Engine/Classes/Animation/AnimNotifies/AnimNotify.h b/Engine/Source/Runtime/Engine/Classes/Animation/AnimNotifies/AnimNotify.h
--- a/Engine/Source/Runtime/Engine/Classes/Animation/AnimNotifies/AnimNotify.h
+++ b/Engine/Source/Runtime/Engine/Classes/Animation/AnimNotifies/AnimNotify.h
@@ -50,7 +50,7 @@ class ENGINE_API UAnimNotify : public UObject
 	UFUNCTION(BlueprintNativeEvent)
 	FString GetNotifyName() const;
 
-	UFUNCTION(BlueprintImplementableEvent)
+	UFUNCTION(BlueprintImplementableEvent, meta=(AutoCreateRefTerm="EventReference"))
 	bool Received_Notify(USkeletalMeshComponent* MeshComp, UAnimSequenceBase* Animation, const FAnimNotifyEventReference& EventReference) const;
 
 #if WITH_EDITORONLY_DATA
diff --git a/Engine/Source/Runtime/Engine/Classes/Animation/AnimNotifies/AnimNotifyState.h b/Engine/Source/Runtime/Engine/Classes/Animation/AnimNotifies/AnimNotifyState.h
--- a/Engine/Source/Runtime/Engine/Classes/Animation/AnimNotifies/AnimNotifyState.h
+++ b/Engine/Source/Runtime/Engine/Classes/Animation/AnimNotifies/AnimNotifyState.h
@@ -26,13 +26,13 @@ class ENGINE_API UAnimNotifyState : public UObject
 	UFUNCTION(BlueprintNativeEvent)
 	FString GetNotifyName() const;
 
-	UFUNCTION(BlueprintImplementableEvent)
+	UFUNCTION(BlueprintImplementableEvent, meta=(AutoCreateRefTerm="EventReference"))
 	bool Received_NotifyBegin(USkeletalMeshComponent * MeshComp, UAnimSequenceBase * Animation, float TotalDuration, const FAnimNotifyEventReference& EventReference) const;
 	
-	UFUNCTION(BlueprintImplementableEvent)
+	UFUNCTION(BlueprintImplementableEvent, meta=(AutoCreateRefTerm="EventReference"))
 	bool Received_NotifyTick(USkeletalMeshComponent * MeshComp, UAnimSequenceBase * Animation, float FrameDeltaTime, const FAnimNotifyEventReference& EventReference) const;
 
-	UFUNCTION(BlueprintImplementableEvent)
+	UFUNCTION(BlueprintImplementableEvent, meta=(AutoCreateRefTerm="EventReference"))
 	bool Received_NotifyEnd(USkeletalMeshComponent * MeshComp, UAnimSequenceBase * Animation, const FAnimNotifyEventReference& EventReference) const;
 
 #if WITH_EDITORONLY_DATA
