Better logging in Actor clustering

diff --git a/Engine/Source/Runtime/Engine/Private/WorldPartition/RuntimeSpatialHash/RuntimeSpatialHashGridHelper.cpp b/Engine/Source/Runtime/Engine/Private/WorldPartition/RuntimeSpatialHash/RuntimeSpatialHashGridHelper.cpp
--- a/Engine/Source/Runtime/Engine/Private/WorldPartition/RuntimeSpatialHash/RuntimeSpatialHashGridHelper.cpp
+++ b/Engine/Source/Runtime/Engine/Private/WorldPartition/RuntimeSpatialHash/RuntimeSpatialHashGridHelper.cpp
@@ -1,7 +1,7 @@
 // Copyright Epic Games, Inc. All Rights Reserved.
 
 #include "WorldPartition/RuntimeSpatialHash/RuntimeSpatialHashGridHelper.h"
-
+#include "Algo/Transform.h"
 #include "ProfilingDebugging/ScopedTimers.h"
 
 FSquare2DGridHelper::FSquare2DGridHelper(const FBox& InWorldBounds, const FVector& InOrigin, int32 InCellSize)
@@ -242,6 +242,13 @@ FSquare2DGridHelper GetPartitionedActors(const UWorldPartition* WorldPartition,
 					(int)(0.01f * ClusterInstance->Bounds.GetSize().X),
 					(int)(0.01f * ClusterInstance->Bounds.GetSize().Y));
 
+				if (ClusterInstance->DataLayers.Num())
+				{
+					TArray<FString> DataLayers;
+					Algo::Transform(ClusterInstance->DataLayers, DataLayers, [](const UDataLayer* DataLayer) { return DataLayer->GetDataLayerLabel().ToString(); });
+					UE_LOG(LogWorldPartition, Verbose, TEXT("   - DataLayers: %s"), *FString::Join(DataLayers, TEXT(", ")));
+				}
+
 				for (const FGuid& ActorGuid : ActorCluster->Actors)
 				{
 					const FWorldPartitionActorDescView& ActorDescView = ClusterInstance->ContainerInstance->GetActorDescView(ActorGuid);
diff --git a/Engine/Source/Runtime/Engine/Private/WorldPartition/WorldPartitionActorCluster.cpp b/Engine/Source/Runtime/Engine/Private/WorldPartition/WorldPartitionActorCluster.cpp
--- a/Engine/Source/Runtime/Engine/Private/WorldPartition/WorldPartitionActorCluster.cpp
+++ b/Engine/Source/Runtime/Engine/Private/WorldPartition/WorldPartitionActorCluster.cpp
@@ -44,7 +44,7 @@ FActorCluster::FActorCluster(UWorld* InWorld, const FWorldPartitionActorDescView
 	DataLayersID = FDataLayersID(DataLayers.Array());
 }
 
-void FActorCluster::Add(const FActorCluster& InActorCluster)
+void FActorCluster::Add(const FActorCluster& InActorCluster, const FActorContainerInstance& ContainerInstance)
 {
 	// Merge RuntimeGrid
 	RuntimeGrid = RuntimeGrid == InActorCluster.RuntimeGrid ? RuntimeGrid : NAME_None;
@@ -79,28 +79,44 @@ void FActorCluster::Add(const FActorCluster& InActorCluster)
 
 	if (DataLayersID != InActorCluster.DataLayersID)
 	{
+		auto LogActorGuid = [&ContainerInstance](const FGuid& ActorGuid)
+		{
+			const FWorldPartitionActorDescView* ActorDescView = ContainerInstance.ActorDescViewMap.Find(ActorGuid);
+			UE_LOG(LogWorldPartition, Verbose, TEXT("   - Actor: %s (%s)"), ActorDescView ? *ActorDescView->GetActorPath().ToString() : TEXT("None"), *ActorGuid.ToString(EGuidFormats::UniqueObjectGuid));
+		};
+
+		auto LogDataLayers = [](const TSet<const UDataLayer*>& InDataLayers)
+		{
+			TArray<FString> DataLayerLabels;
+			Algo::Transform(InDataLayers, DataLayerLabels, [](const UDataLayer* DataLayer) { return DataLayer->GetDataLayerLabel().ToString(); });
+			UE_LOG(LogWorldPartition, Verbose, TEXT("   - DataLayers: %s"), *FString::Join(DataLayerLabels, TEXT(", ")));
+		};
+
 		if (DataLayers.Num() && InActorCluster.DataLayers.Num())
 		{
-			// Merge Data Layers
-			for (const UDataLayer* DataLayer : InActorCluster.DataLayers)
-			{
-				check(DataLayer->IsDynamicallyLoaded());
-				DataLayers.Add(DataLayer);
-			}
 			UE_SUPPRESS(LogWorldPartition, Verbose,
 			{
 				UE_LOG(LogWorldPartition, Verbose, TEXT("Merging Data Layers for clustered actors with different sets of Data Layers."));
-				UE_LOG(LogWorldPartition, Verbose, TEXT("1st clustered actors :"));
+				UE_LOG(LogWorldPartition, Verbose, TEXT("1st cluster :"));
+				LogDataLayers(DataLayers);
 				for (const FGuid& ActorGuid : Actors)
 				{
-					UE_LOG(LogWorldPartition, Verbose, TEXT("   - Actor: (%s)"), *ActorGuid.ToString(EGuidFormats::UniqueObjectGuid));
+					LogActorGuid(ActorGuid);
 				}
-				UE_LOG(LogWorldPartition, Verbose, TEXT("2nd clustered actors :"));
+				UE_LOG(LogWorldPartition, Verbose, TEXT("2nd cluster :"));
+				LogDataLayers(InActorCluster.DataLayers);
 				for (const FGuid& ActorGuid : InActorCluster.Actors)
 				{
-					UE_LOG(LogWorldPartition, Verbose, TEXT("   - Actor: (%s)"), *ActorGuid.ToString(EGuidFormats::UniqueObjectGuid));
+					LogActorGuid(ActorGuid);
 				}
 			});
+
+			// Merge Data Layers
+			for (const UDataLayer* DataLayer : InActorCluster.DataLayers)
+			{
+				check(DataLayer->IsDynamicallyLoaded());
+				DataLayers.Add(DataLayer);
+			}
 		}
 		else
 		{
@@ -113,13 +129,13 @@ void FActorCluster::Add(const FActorCluster& InActorCluster)
 				const FActorCluster * ActorClusterWithDataLayer = DataLayers.Num() ? this : &InActorCluster;
 				for (const FGuid& ActorGuid : ActorClusterWithDataLayer->Actors)
 				{
-					UE_LOG(LogWorldPartition, Verbose, TEXT("   - Actor: (%s)"), *ActorGuid.ToString(EGuidFormats::UniqueObjectGuid));
+					LogActorGuid(ActorGuid);
 				}
 				UE_LOG(LogWorldPartition, Verbose, TEXT("Clustered actors without Data Layer :"));
 				const FActorCluster * ActorClusterWithoutDataLayer = DataLayers.Num() ? &InActorCluster : this;
 				for (const FGuid& ActorGuid : ActorClusterWithoutDataLayer->Actors)
 				{
-					UE_LOG(LogWorldPartition, Verbose, TEXT("   - Actor: (%s)"), *ActorGuid.ToString(EGuidFormats::UniqueObjectGuid));
+					LogActorGuid(ActorGuid);
 				}
 			});
 			DataLayers.Empty();
@@ -255,7 +271,7 @@ void CreateActorCluster(const FWorldPartitionActorDescView& ActorDescView, TMap<
 						if (ReferenceCluster != ActorCluster)
 						{
 							// Merge reference cluster in Actor's cluster
-							ActorCluster->Add(*ReferenceCluster);
+							ActorCluster->Add(*ReferenceCluster, ContainerInstance);
 							for (const FGuid& ReferenceClusterActorGuid : ReferenceCluster->Actors)
 							{
 								ActorToActorCluster[ReferenceClusterActorGuid] = ActorCluster;
@@ -267,7 +283,7 @@ void CreateActorCluster(const FWorldPartitionActorDescView& ActorDescView, TMap<
 					else
 					{
 						// Put Reference in Actor's cluster
-						ActorCluster->Add(FActorCluster(World, *ReferenceActorDescView));
+						ActorCluster->Add(FActorCluster(World, *ReferenceActorDescView), ContainerInstance);
 					}
 
 					// Map its cluster
diff --git a/Engine/Source/Runtime/Engine/Public/WorldPartition/WorldPartitionActorCluster.h b/Engine/Source/Runtime/Engine/Public/WorldPartition/WorldPartitionActorCluster.h
--- a/Engine/Source/Runtime/Engine/Public/WorldPartition/WorldPartitionActorCluster.h
+++ b/Engine/Source/Runtime/Engine/Public/WorldPartition/WorldPartitionActorCluster.h
@@ -26,7 +26,7 @@ struct FActorCluster
 	FDataLayersID				DataLayersID;
 
 	FActorCluster(UWorld* InWorld, const FWorldPartitionActorDescView& InActorDescView);
-	void Add(const FActorCluster& InActorCluster);
+	void Add(const FActorCluster& InActorCluster, const struct FActorContainerInstance& ContainerInstance);
 
 };
 
