non-unity fix

#rnx

diff --git a/Engine/Source/Runtime/Engine/Private/OnlineReplStructs.cpp b/Engine/Source/Runtime/Engine/Private/OnlineReplStructs.cpp
--- a/Engine/Source/Runtime/Engine/Private/OnlineReplStructs.cpp
+++ b/Engine/Source/Runtime/Engine/Private/OnlineReplStructs.cpp
@@ -12,6 +12,7 @@
 #include "Dom/JsonValue.h"
 #include "EngineLogs.h"
 #include "Net/OnlineEngineInterface.h"
+#include "Misc/ConfigCacheIni.h"
 
 namespace
 {
