Removal of UE4 references

diff --git a/Engine/Source/Runtime/Engine/Private/EngineSessionManager.cpp b/Engine/Source/Runtime/Engine/Private/EngineSessionManager.cpp
--- a/Engine/Source/Runtime/Engine/Private/EngineSessionManager.cpp
+++ b/Engine/Source/Runtime/Engine/Private/EngineSessionManager.cpp
@@ -68,7 +68,7 @@ namespace SessionManagerDefs
 	static const FString WasDebuggerStoreKey(TEXT("WasEverDebugger"));
 	static const FString UserActivityStoreKey(TEXT("CurrentUserActivity"));
 	static const FString VanillaStoreKey(TEXT("IsVanilla"));
-	static const FString GlobalLockName(TEXT("UE4_SessionManager_Lock"));
+	static const FString GlobalLockName(TEXT("UE_SessionManager_Lock"));
 	static const FString FalseValueString(TEXT("0"));
 	static const FString TrueValueString(TEXT("1"));
 	static const FString EditorValueString(TEXT("Editor"));
diff --git a/Engine/Source/Runtime/Engine/Private/GameEngine.cpp b/Engine/Source/Runtime/Engine/Private/GameEngine.cpp
--- a/Engine/Source/Runtime/Engine/Private/GameEngine.cpp
+++ b/Engine/Source/Runtime/Engine/Private/GameEngine.cpp
@@ -811,7 +811,7 @@ public:
 
 		FEmbeddedDelegates::GetNativeToEmbeddedParamsDelegateForSubsystem(TEXT("engine")).AddLambda([](const FEmbeddedCallParamsHelper& Message)
 		{
-			if (Message.Command == TEXT("StartUE4Live"))
+			if (Message.Command == TEXT("StartUELive"))
 			{
 				FName Requester = *Message.Parameters.FindRef(TEXT("requester"));
 				bool bTickOnly = Message.Parameters.FindRef(TEXT("tickonly")) == TEXT("true");
@@ -819,7 +819,7 @@ public:
 				FEmbeddedCommunication::KeepAwake(Requester, !bTickOnly);
 				Message.OnCompleteDelegate({}, TEXT(""));
 			}
-			else if (Message.Command == TEXT("StopUE4Live"))
+			else if (Message.Command == TEXT("StopUELive"))
 			{
 				FName Requester = *Message.Parameters.FindRef(TEXT("requester"));
 				
