Fix handling of deprecated notify functions in derived classes

diff --git a/Engine/Source/Runtime/Engine/Private/Animation/AnimNotifyState_TimedParticleEffect.cpp b/Engine/Source/Runtime/Engine/Private/Animation/AnimNotifyState_TimedParticleEffect.cpp
--- a/Engine/Source/Runtime/Engine/Private/Animation/AnimNotifyState_TimedParticleEffect.cpp
+++ b/Engine/Source/Runtime/Engine/Private/Animation/AnimNotifyState_TimedParticleEffect.cpp
@@ -21,6 +21,10 @@ void UAnimNotifyState_TimedParticleEffect::NotifyBegin(USkeletalMeshComponent *
 
 void UAnimNotifyState_TimedParticleEffect::NotifyBegin(USkeletalMeshComponent * MeshComp, class UAnimSequenceBase * Animation, float TotalDuration, const FAnimNotifyEventReference& EventReference)
 {
+	// ensure deprecated path is called because a call to Super is not made
+	PRAGMA_DISABLE_DEPRECATION_WARNINGS
+	NotifyBegin(MeshComp, Animation, TotalDuration);
+	PRAGMA_ENABLE_DEPRECATION_WARNINGS
 	// Only spawn if we've got valid params
 	if(ValidateParameters(MeshComp))
 	{
@@ -35,6 +39,9 @@ void UAnimNotifyState_TimedParticleEffect::NotifyTick(USkeletalMeshComponent * M
 
 void UAnimNotifyState_TimedParticleEffect::NotifyTick(USkeletalMeshComponent * MeshComp, class UAnimSequenceBase * Animation, float FrameDeltaTime, const FAnimNotifyEventReference& EventReference)
 {
+	PRAGMA_DISABLE_DEPRECATION_WARNINGS
+	NotifyTick(MeshComp, Animation, FrameDeltaTime);
+	PRAGMA_ENABLE_DEPRECATION_WARNINGS
 	Received_NotifyTick(MeshComp, Animation, FrameDeltaTime, EventReference);
 }
 
@@ -44,6 +51,9 @@ void UAnimNotifyState_TimedParticleEffect::NotifyEnd(USkeletalMeshComponent * Me
 
 void UAnimNotifyState_TimedParticleEffect::NotifyEnd(USkeletalMeshComponent * MeshComp, class UAnimSequenceBase * Animation, const FAnimNotifyEventReference& EventReference)
 {
+	PRAGMA_DISABLE_DEPRECATION_WARNINGS
+	NotifyEnd(MeshComp, Animation);
+	PRAGMA_ENABLE_DEPRECATION_WARNINGS
 	TArray<USceneComponent*> Children;
 	MeshComp->GetChildrenComponents(false, Children);
 
diff --git a/Engine/Source/Runtime/Engine/Private/Animation/AnimNotifyState_Trail.cpp b/Engine/Source/Runtime/Engine/Private/Animation/AnimNotifyState_Trail.cpp
--- a/Engine/Source/Runtime/Engine/Private/Animation/AnimNotifyState_Trail.cpp
+++ b/Engine/Source/Runtime/Engine/Private/Animation/AnimNotifyState_Trail.cpp
@@ -81,6 +81,9 @@ void UAnimNotifyState_Trail::NotifyBegin(class USkeletalMeshComponent* MeshComp,
 
 void UAnimNotifyState_Trail::NotifyBegin(class USkeletalMeshComponent * MeshComp, class UAnimSequenceBase * Animation, float TotalDuration, const FAnimNotifyEventReference& EventReference)
 {
+	PRAGMA_DISABLE_DEPRECATION_WARNINGS
+    NotifyBegin(MeshComp, Animation, TotalDuration);
+	PRAGMA_ENABLE_DEPRECATION_WARNINGS
 	bool bError = ValidateInput(MeshComp);
 
 	if (MeshComp->GetWorld()->GetNetMode() == NM_DedicatedServer)
@@ -208,6 +211,9 @@ void UAnimNotifyState_Trail::NotifyTick(class USkeletalMeshComponent* MeshComp,
 
 void UAnimNotifyState_Trail::NotifyTick(class USkeletalMeshComponent * MeshComp, class UAnimSequenceBase * Animation, float FrameDeltaTime, const FAnimNotifyEventReference& EventReference)
 {
+	PRAGMA_DISABLE_DEPRECATION_WARNINGS
+    NotifyTick(MeshComp, Animation, FrameDeltaTime);
+	PRAGMA_ENABLE_DEPRECATION_WARNINGS
 	bool bError = ValidateInput(MeshComp, true);
 
 	if (MeshComp->GetWorld()->GetNetMode() == NM_DedicatedServer)
@@ -253,6 +259,10 @@ void UAnimNotifyState_Trail::NotifyEnd(class USkeletalMeshComponent* MeshComp, c
 
 void UAnimNotifyState_Trail::NotifyEnd(class USkeletalMeshComponent * MeshComp, class UAnimSequenceBase * Animation, const FAnimNotifyEventReference& EventReference)
 {
+	PRAGMA_DISABLE_DEPRECATION_WARNINGS
+	NotifyEnd(MeshComp, Animation);
+	PRAGMA_ENABLE_DEPRECATION_WARNINGS
+
 	if (MeshComp->GetWorld()->GetNetMode() == NM_DedicatedServer)
 	{
 		return;
diff --git a/Engine/Source/Runtime/Engine/Private/Animation/AnimNotify_PauseClothingSimulation.cpp b/Engine/Source/Runtime/Engine/Private/Animation/AnimNotify_PauseClothingSimulation.cpp
--- a/Engine/Source/Runtime/Engine/Private/Animation/AnimNotify_PauseClothingSimulation.cpp
+++ b/Engine/Source/Runtime/Engine/Private/Animation/AnimNotify_PauseClothingSimulation.cpp
@@ -13,11 +13,14 @@ UAnimNotify_PauseClothingSimulation::UAnimNotify_PauseClothingSimulation()
 
 void UAnimNotify_PauseClothingSimulation::Notify(class USkeletalMeshComponent* MeshComp, class UAnimSequenceBase* Animation)
 {
-	MeshComp->SuspendClothingSimulation();
+
 }
 
 void UAnimNotify_PauseClothingSimulation::Notify(class USkeletalMeshComponent* MeshComp, class UAnimSequenceBase* Animation, const FAnimNotifyEventReference& EventReference)
 {
+	PRAGMA_DISABLE_DEPRECATION_WARNINGS
+    Notify(MeshComp, Animation);
+    PRAGMA_ENABLE_DEPRECATION_WARNINGS
 	MeshComp->SuspendClothingSimulation();
 }
 
diff --git a/Engine/Source/Runtime/Engine/Private/Animation/AnimNotify_PlayParticleEffect.cpp b/Engine/Source/Runtime/Engine/Private/Animation/AnimNotify_PlayParticleEffect.cpp
--- a/Engine/Source/Runtime/Engine/Private/Animation/AnimNotify_PlayParticleEffect.cpp
+++ b/Engine/Source/Runtime/Engine/Private/Animation/AnimNotify_PlayParticleEffect.cpp
@@ -76,6 +76,9 @@ void UAnimNotify_PlayParticleEffect::Notify(class USkeletalMeshComponent* MeshCo
 
 void UAnimNotify_PlayParticleEffect::Notify(class USkeletalMeshComponent* MeshComp, class UAnimSequenceBase* Animation, const FAnimNotifyEventReference& EventReference)
 {
+	PRAGMA_DISABLE_DEPRECATION_WARNINGS
+    Notify(MeshComp, Animation);
+	PRAGMA_ENABLE_DEPRECATION_WARNINGS
 	// Don't call super to avoid unnecessary call in to blueprints
 	SpawnParticleSystem(MeshComp, Animation);
 }
diff --git a/Engine/Source/Runtime/Engine/Private/Animation/AnimNotify_PlaySound.cpp b/Engine/Source/Runtime/Engine/Private/Animation/AnimNotify_PlaySound.cpp
--- a/Engine/Source/Runtime/Engine/Private/Animation/AnimNotify_PlaySound.cpp
+++ b/Engine/Source/Runtime/Engine/Private/Animation/AnimNotify_PlaySound.cpp
@@ -33,6 +33,9 @@ void UAnimNotify_PlaySound::Notify(USkeletalMeshComponent* MeshComp, UAnimSequen
 
 void UAnimNotify_PlaySound::Notify(class USkeletalMeshComponent* MeshComp, class UAnimSequenceBase* Animation, const FAnimNotifyEventReference& EventReference)
 {
+	PRAGMA_DISABLE_DEPRECATION_WARNINGS
+    Notify(MeshComp, Animation);
+	PRAGMA_ENABLE_DEPRECATION_WARNINGS
 	// Don't call super to avoid call back in to blueprints
 	if (Sound && MeshComp)
 	{
diff --git a/Engine/Source/Runtime/Engine/Private/Animation/AnimNotify_ResetClothingSimulation.cpp b/Engine/Source/Runtime/Engine/Private/Animation/AnimNotify_ResetClothingSimulation.cpp
--- a/Engine/Source/Runtime/Engine/Private/Animation/AnimNotify_ResetClothingSimulation.cpp
+++ b/Engine/Source/Runtime/Engine/Private/Animation/AnimNotify_ResetClothingSimulation.cpp
@@ -17,6 +17,9 @@ void UAnimNotify_ResetClothingSimulation::Notify(class USkeletalMeshComponent* M
 
 void UAnimNotify_ResetClothingSimulation::Notify(class USkeletalMeshComponent* MeshComp, class UAnimSequenceBase* Animation, const FAnimNotifyEventReference& EventReference)
 {
+	PRAGMA_DISABLE_DEPRECATION_WARNINGS
+    Notify(MeshComp, Animation);
+	PRAGMA_ENABLE_DEPRECATION_WARNINGS
 	MeshComp->ForceClothNextUpdateTeleportAndReset();
 }
 
diff --git a/Engine/Source/Runtime/Engine/Private/Animation/AnimNotify_ResetDynamics.cpp b/Engine/Source/Runtime/Engine/Private/Animation/AnimNotify_ResetDynamics.cpp
--- a/Engine/Source/Runtime/Engine/Private/Animation/AnimNotify_ResetDynamics.cpp
+++ b/Engine/Source/Runtime/Engine/Private/Animation/AnimNotify_ResetDynamics.cpp
@@ -17,6 +17,9 @@ void UAnimNotify_ResetDynamics::Notify(class USkeletalMeshComponent* MeshComp, c
 
 void UAnimNotify_ResetDynamics::Notify(class USkeletalMeshComponent* MeshComp, class UAnimSequenceBase* Animation, const FAnimNotifyEventReference& EventReference)
 {
+	PRAGMA_DISABLE_DEPRECATION_WARNINGS
+    Notify(MeshComp, Animation);
+	PRAGMA_ENABLE_DEPRECATION_WARNINGS
 	MeshComp->ResetAnimInstanceDynamics();
 }
 
diff --git a/Engine/Source/Runtime/Engine/Private/Animation/AnimNotify_ResumeClothingSimulation.cpp b/Engine/Source/Runtime/Engine/Private/Animation/AnimNotify_ResumeClothingSimulation.cpp
--- a/Engine/Source/Runtime/Engine/Private/Animation/AnimNotify_ResumeClothingSimulation.cpp
+++ b/Engine/Source/Runtime/Engine/Private/Animation/AnimNotify_ResumeClothingSimulation.cpp
@@ -17,6 +17,9 @@ void UAnimNotify_ResumeClothingSimulation::Notify(class USkeletalMeshComponent*
 
 void UAnimNotify_ResumeClothingSimulation::Notify(class USkeletalMeshComponent* MeshComp, class UAnimSequenceBase* Animation, const FAnimNotifyEventReference& EventReference)
 {
+	PRAGMA_DISABLE_DEPRECATION_WARNINGS
+    Notify(MeshComp, Animation);
+	PRAGMA_ENABLE_DEPRECATION_WARNINGS
 	MeshComp->ResumeClothingSimulation();
 }
 
