Fixed NonUnity BuildErrors in and EngineTest deprecation warnings related to animation notifies

diff --git a/Engine/Source/Runtime/Engine/Public/Animation/AnimNotifyMirrorInspectionLibrary.h b/Engine/Source/Runtime/Engine/Public/Animation/AnimNotifyMirrorInspectionLibrary.h
--- a/Engine/Source/Runtime/Engine/Public/Animation/AnimNotifyMirrorInspectionLibrary.h
+++ b/Engine/Source/Runtime/Engine/Public/Animation/AnimNotifyMirrorInspectionLibrary.h
@@ -4,8 +4,10 @@
 
 #include "CoreMinimal.h"
 #include "Kismet/BlueprintFunctionLibrary.h"
+#include "Animation/AnimNotifyQueue.h"
 #include "AnimNotifyMirrorInspectionLibrary.generated.h"
 
+
 struct FAnimNotifyEventReference;
 /**
 *	A library of commonly used functionality for Notifies related to mirroring, exposed to blueprint.
diff --git a/Engine/Source/Runtime/Engine/Public/Animation/AnimNotifyStateMachineInspectionLibrary.h b/Engine/Source/Runtime/Engine/Public/Animation/AnimNotifyStateMachineInspectionLibrary.h
--- a/Engine/Source/Runtime/Engine/Public/Animation/AnimNotifyStateMachineInspectionLibrary.h
+++ b/Engine/Source/Runtime/Engine/Public/Animation/AnimNotifyStateMachineInspectionLibrary.h
@@ -4,6 +4,7 @@
 
 #include "CoreMinimal.h"
 #include "Kismet/BlueprintFunctionLibrary.h"
+#include "Animation/AnimNotifyQueue.h"
 #include "AnimNotifyStateMachineInspectionLibrary.generated.h"
 
 struct FAnimNotifyEventReference;
diff --git a/Engine/Source/Runtime/Engine/Public/Animation/MirrorSyncScope.h b/Engine/Source/Runtime/Engine/Public/Animation/MirrorSyncScope.h
--- a/Engine/Source/Runtime/Engine/Public/Animation/MirrorSyncScope.h
+++ b/Engine/Source/Runtime/Engine/Public/Animation/MirrorSyncScope.h
@@ -4,8 +4,7 @@
 
 #include "CoreMinimal.h"
 #include "AnimNodeMessages.h"
-
-class UMirrorDataTable;
+#include "Animation/MirrorDataTable.h"
 
 namespace UE { namespace Anim {
 
@@ -18,8 +17,8 @@ class ENGINE_API FAnimNotifyMirrorContext : public IAnimNotifyEventContextDataIn
 	FAnimNotifyMirrorContext() {} 
 	FAnimNotifyMirrorContext(const UMirrorDataTable* MirrorDataTable) : MirrorTable(MirrorDataTable) {bAnimationMirrored = MirrorDataTable != nullptr;} 
 	bool bAnimationMirrored = true; 
-	UPROPERTY(transient)
-	TObjectPtr<const UMirrorDataTable> MirrorTable; 
+
+	TWeakObjectPtr<const UMirrorDataTable> MirrorTable; 
 };
 	
 // Scoped graph message used to synchronize mirroring 
