Expose some material connected pins to runtime.
Adds BaseColor, Roughness and AO output pins to runtime structures with a small memory cost.
Required for logic determining decal blend states.

diff --git a/Engine/Source/Runtime/Engine/Classes/Materials/Material.h b/Engine/Source/Runtime/Engine/Classes/Materials/Material.h
--- a/Engine/Source/Runtime/Engine/Classes/Materials/Material.h
+++ b/Engine/Source/Runtime/Engine/Classes/Materials/Material.h
@@ -339,10 +339,10 @@ class UMaterial : public UMaterialInterface
 
 	UPROPERTY()
 	FColorMaterialInput SpecularColor_DEPRECATED;
+#endif
 
 	UPROPERTY()
 	FColorMaterialInput BaseColor;
-#endif
 
 	UPROPERTY()
 	FScalarMaterialInput Metallic;
@@ -350,10 +350,8 @@ class UMaterial : public UMaterialInterface
 	UPROPERTY()
 	FScalarMaterialInput Specular;
 
-#if WITH_EDITORONLY_DATA
 	UPROPERTY()
 	FScalarMaterialInput Roughness;
-#endif
 
 	UPROPERTY()
 	FScalarMaterialInput Anisotropy;
@@ -364,7 +362,6 @@ class UMaterial : public UMaterialInterface
 	UPROPERTY()
 	FVectorMaterialInput Tangent;
 
-	// Emission.
 	UPROPERTY()
 	FColorMaterialInput EmissiveColor;
 
@@ -445,11 +442,11 @@ public:
 	/**  */
 	UPROPERTY()
 	FScalarMaterialInput ClearCoatRoughness;
+#endif
 
 	/** output ambient occlusion to the GBuffer */
 	UPROPERTY()
 	FScalarMaterialInput AmbientOcclusion;
-#endif
 
 	/**
 	 * output refraction index for translucent rendering
@@ -1766,18 +1763,12 @@ public:
 	void SaveShaderStableKeys(const class ITargetPlatform* TP);
 	ENGINE_API virtual void SaveShaderStableKeysInner(const class ITargetPlatform* TP, const struct FStableShaderKeyAndValue& SaveKeyVal) override;
 
-#if WITH_EDITORONLY_DATA
 	bool HasBaseColorConnected() const { return BaseColor.IsConnected(); }
 	bool HasRoughnessConnected() const { return Roughness.IsConnected(); }
 	bool HasAmbientOcclusionConnected() const { return AmbientOcclusion.IsConnected(); }
-#else	
-	// Add to runtime data only if we need to call these at runtime
-	bool HasBaseColorConnected() const { check(0); return false; }
-	bool HasRoughnessConnected() const { check(0); return false; }
-	bool HasAmbientOcclusionConnected() const { check(0); return false; }
-#endif 	
 	bool HasNormalConnected() const { return Normal.IsConnected(); }
 	bool HasSpecularConnected() const { return Specular.IsConnected(); }
+	bool HasMetallicConnected() const { return Metallic.IsConnected(); }
 	bool HasEmissiveColorConnected() const { return EmissiveColor.IsConnected(); }
 	bool HasAnisotropyConnected() const { return Anisotropy.IsConnected(); }
 	bool HasStrataFrontMaterialConnected() const { return FrontMaterial.IsConnected(); }
diff --git a/Engine/Source/Runtime/Engine/Private/Materials/MaterialShared.cpp b/Engine/Source/Runtime/Engine/Private/Materials/MaterialShared.cpp
--- a/Engine/Source/Runtime/Engine/Private/Materials/MaterialShared.cpp
+++ b/Engine/Source/Runtime/Engine/Private/Materials/MaterialShared.cpp
@@ -1601,6 +1601,11 @@ bool FMaterialResource::HasSpecularConnected() const
 	return HasMaterialAttributesConnected() || Material->HasSpecularConnected();
 }
 
+bool FMaterialResource::HasMetallicConnected() const
+{
+	return HasMaterialAttributesConnected() || Material->HasMetallicConnected();
+}
+
 bool FMaterialResource::HasEmissiveColorConnected() const
 {
 	return HasMaterialAttributesConnected() || Material->HasEmissiveColorConnected();
diff --git a/Engine/Source/Runtime/Engine/Public/MaterialShared.h b/Engine/Source/Runtime/Engine/Public/MaterialShared.h
--- a/Engine/Source/Runtime/Engine/Public/MaterialShared.h
+++ b/Engine/Source/Runtime/Engine/Public/MaterialShared.h
@@ -1700,6 +1700,7 @@ public:
 	virtual bool HasNormalConnected() const { return false; }
 	virtual bool HasRoughnessConnected() const { return false; }
 	virtual bool HasSpecularConnected() const { return false; }
+	virtual bool HasMetallicConnected() const { return false; }
 	virtual bool HasEmissiveColorConnected() const { return false; }
 	virtual bool HasAnisotropyConnected() const { return false; }
 	virtual bool HasAmbientOcclusionConnected() const { return false; }
@@ -2549,7 +2550,8 @@ public:
 	ENGINE_API virtual bool HasBaseColorConnected() const override;
 	ENGINE_API virtual bool HasNormalConnected() const override;
 	ENGINE_API virtual bool HasRoughnessConnected() const override;
-	ENGINE_API virtual bool HasSpecularConnected() const override;	
+	ENGINE_API virtual bool HasSpecularConnected() const override;
+	ENGINE_API virtual bool HasMetallicConnected() const override;
 	ENGINE_API virtual bool HasEmissiveColorConnected() const override;
 	ENGINE_API virtual bool HasAnisotropyConnected() const override;
 	ENGINE_API virtual bool HasAmbientOcclusionConnected() const override;
@@ -3131,8 +3133,13 @@ struct FMaterialShaderParameters
 			uint64 bShouldCastDynamicShadows : 1;
 			uint64 bWritesEveryPixel : 1;
 			uint64 bWritesEveryPixelShadowPass : 1;
+			uint64 bHasBaseColorConnected : 1;
 			uint64 bHasNormalConnected : 1;
+			uint64 bHasRoughnessConnected : 1;
+			uint64 bHasSpecularConnected : 1;
+			uint64 bHasMetallicConnected : 1;
 			uint64 bHasEmissiveColorConnected : 1;
+			uint64 bHasAmbientOcclusionConnected : 1;
 			uint64 bHasAnisotropyConnected : 1;
 			uint64 bHasVertexPositionOffsetConnected : 1;
 			uint64 bHasPixelDepthOffsetConnected : 1;
@@ -3186,8 +3193,13 @@ struct FMaterialShaderParameters
 		bShouldCastDynamicShadows = InMaterial->ShouldCastDynamicShadows();
 		bWritesEveryPixel = InMaterial->WritesEveryPixel(false);
 		bWritesEveryPixelShadowPass = InMaterial->WritesEveryPixel(true);
+		bHasBaseColorConnected = InMaterial->HasBaseColorConnected();
 		bHasNormalConnected = InMaterial->HasNormalConnected();
+		bHasRoughnessConnected = InMaterial->HasRoughnessConnected();
+		bHasSpecularConnected = InMaterial->HasSpecularConnected();
+		bHasMetallicConnected = InMaterial->HasMetallicConnected();
 		bHasEmissiveColorConnected = InMaterial->HasEmissiveColorConnected();
+		bHasAmbientOcclusionConnected = InMaterial->HasAmbientOcclusionConnected();
 		bHasAnisotropyConnected = InMaterial->HasAnisotropyConnected();
 		bHasVertexPositionOffsetConnected = InMaterial->HasVertexPositionOffsetConnected();
 		bHasPixelDepthOffsetConnected = InMaterial->HasPixelDepthOffsetConnected();
