Require cast shadow for cast hidden shadow. It was possible to enable cast hidden shadow only, which didn't render any shadow when object is hidden, but was still picked up by various systems downstream as a "mesh casting visible shadows". E.g. causing it to be added to the distance field scene, calculating skeletal mesh bounds or updating transform. This also make this setting consistent with UI (CastShadow greys out CastHiddenShadow)

diff --git a/Engine/Source/Runtime/Engine/Private/PrimitiveSceneProxy.cpp b/Engine/Source/Runtime/Engine/Private/PrimitiveSceneProxy.cpp
--- a/Engine/Source/Runtime/Engine/Private/PrimitiveSceneProxy.cpp
+++ b/Engine/Source/Runtime/Engine/Private/PrimitiveSceneProxy.cpp
@@ -148,7 +148,7 @@ FPrimitiveSceneProxy::FPrimitiveSceneProxy(const UPrimitiveComponent* InComponen
 ,	bCastContactShadow(InComponent->CastShadow && InComponent->bCastContactShadow)
 ,	bCastCapsuleDirectShadow(false)
 ,	bCastsDynamicIndirectShadow(false)
-,	bCastHiddenShadow(InComponent->bCastHiddenShadow)
+,	bCastHiddenShadow(InComponent->CastShadow && InComponent->bCastHiddenShadow)
 ,	bCastShadowAsTwoSided(InComponent->bCastShadowAsTwoSided)
 ,	bSelfShadowOnly(InComponent->bSelfShadowOnly)
 ,	bCastInsetShadow(InComponent->bSelfShadowOnly ? true : InComponent->bCastInsetShadow)	// Assumed to be enabled if bSelfShadowOnly is enabled.
