Do not register components when reinstancing actors or when executing construction scripts in a world that is not yet initialized.
Those steps will be taken care of as part of the normal world initialization flow

diff --git a/Engine/Source/Editor/UnrealEd/Private/Kismet2/KismetReinstanceUtilities.cpp b/Engine/Source/Editor/UnrealEd/Private/Kismet2/KismetReinstanceUtilities.cpp
--- a/Engine/Source/Editor/UnrealEd/Private/Kismet2/KismetReinstanceUtilities.cpp
+++ b/Engine/Source/Editor/UnrealEd/Private/Kismet2/KismetReinstanceUtilities.cpp
@@ -2124,9 +2124,14 @@ static void ReplaceActorHelper(AActor* OldActor, UClass* OldClass, UObject*& New
 
 	// reset properties/streams
 	NewActor->ResetPropertiesForConstruction();
-	// register native components
-	NewActor->RegisterAllComponents();
 
+	// Only register the components if the world is already initialized
+	if (World->bIsWorldInitialized)
+	{
+		// register native components
+		NewActor->RegisterAllComponents();
+	}
+	
 	// 
 	// clean up the old actor (unselect it, remove it from the world, etc.)...
 
diff --git a/Engine/Source/Runtime/Engine/Private/ActorConstruction.cpp b/Engine/Source/Runtime/Engine/Private/ActorConstruction.cpp
--- a/Engine/Source/Runtime/Engine/Private/ActorConstruction.cpp
+++ b/Engine/Source/Runtime/Engine/Private/ActorConstruction.cpp
@@ -771,7 +771,7 @@ bool AActor::ExecuteConstruction(const FTransform& Transform, const FRotationCon
 			}
 
 			// Ensure that we've called RegisterAllComponents(), in case it was deferred and the SCS could not be fully executed.
-			if (HasDeferredComponentRegistration())
+			if (HasDeferredComponentRegistration() && GetWorld()->bIsWorldInitialized)
 			{
 				RegisterAllComponents();
 			}
@@ -782,7 +782,7 @@ bool AActor::ExecuteConstruction(const FTransform& Transform, const FRotationCon
 			for (UActorComponent* ActorComponent : PostSCSComponents)
 			{
 				// Limit registration to components that are known to have been created during SCS execution
-				if (!ActorComponent->IsRegistered() && ActorComponent->bAutoRegister && !ActorComponent->IsPendingKill()
+				if (!ActorComponent->IsRegistered() && ActorComponent->bAutoRegister && !ActorComponent->IsPendingKill() && GetWorld()->bIsWorldInitialized
 					&& (ActorComponent->CreationMethod == EComponentCreationMethod::SimpleConstructionScript || !PreSCSComponents.Contains(ActorComponent)))
 				{
 					USimpleConstructionScript::RegisterInstancedComponent(ActorComponent);
@@ -858,7 +858,7 @@ bool AActor::ExecuteConstruction(const FTransform& Transform, const FRotationCon
 			}
 
 			// Ensure that we've called RegisterAllComponents(), in case it was deferred and the SCS could not be executed (due to error).
-			if (HasDeferredComponentRegistration())
+			if (HasDeferredComponentRegistration() && GetWorld()->bIsWorldInitialized)
 			{
 				RegisterAllComponents();
 			}
diff --git a/Engine/Source/Runtime/Engine/Private/Components/ActorComponent.cpp b/Engine/Source/Runtime/Engine/Private/Components/ActorComponent.cpp
--- a/Engine/Source/Runtime/Engine/Private/Components/ActorComponent.cpp
+++ b/Engine/Source/Runtime/Engine/Private/Components/ActorComponent.cpp
@@ -1315,6 +1315,7 @@ void UActorComponent::RegisterComponent()
 	UWorld* MyOwnerWorld = (MyOwner ? MyOwner->GetWorld() : nullptr);
 	if (ensure(MyOwnerWorld))
 	{
+		ensure(MyOwnerWorld->bIsWorldInitialized);
 		RegisterComponentWithWorld(MyOwnerWorld);
 	}
 }
diff --git a/Engine/Source/Runtime/Engine/Private/SCS_Node.cpp b/Engine/Source/Runtime/Engine/Private/SCS_Node.cpp
--- a/Engine/Source/Runtime/Engine/Private/SCS_Node.cpp
+++ b/Engine/Source/Runtime/Engine/Private/SCS_Node.cpp
@@ -3,6 +3,7 @@
 #include "Engine/SCS_Node.h"
 #include "UObject/LinkerLoad.h"
 #include "Engine/Blueprint.h"
+#include "Engine/World.h"
 #include "Misc/SecureHash.h"
 #include "UObject/PropertyPortFlags.h"
 #include "Engine/InheritableComponentHandler.h"
@@ -116,6 +117,10 @@ UActorComponent* USCS_Node::ExecuteNodeOnActor(AActor* Actor, USceneComponent* P
 		USceneComponent* NewSceneComp = Cast<USceneComponent>(NewActorComp);
 		if (NewSceneComp != nullptr)
 		{
+			// Only register scene components if the world is initialized
+			UWorld* World = Actor->GetWorld();
+			bool bRegisterComponent = World && World->bIsWorldInitialized;
+
 			// If NULL is passed in, we are the root, so set transform and assign as RootComponent on Actor, similarly if the 
 			// NewSceneComp is the ParentComponent then we are the root component. This happens when the root component is recycled
 			// by StaticAllocateObject.
@@ -140,7 +145,7 @@ UActorComponent* USCS_Node::ExecuteNodeOnActor(AActor* Actor, USceneComponent* P
 				Actor->SetRootComponent(NewSceneComp);
 
 				// This will be true if we deferred the RegisterAllComponents() call at spawn time. In that case, we can call it now since we have set a scene root.
-				if (Actor->HasDeferredComponentRegistration())
+				if (Actor->HasDeferredComponentRegistration() && bRegisterComponent)
 				{
 					// Register the root component along with any components whose registration may have been deferred pending SCS execution in order to establish a root.
 					Actor->RegisterAllComponents();
@@ -153,7 +158,10 @@ UActorComponent* USCS_Node::ExecuteNodeOnActor(AActor* Actor, USceneComponent* P
 			}
 
 			// Register SCS scene components now (if necessary). Non-scene SCS component registration is deferred until after SCS execution, as there can be dependencies on the scene hierarchy.
-			USimpleConstructionScript::RegisterInstancedComponent(NewSceneComp);
+			if (bRegisterComponent)
+			{
+				USimpleConstructionScript::RegisterInstancedComponent(NewSceneComp);
+			}
 		}
 
 		// If we want to save this to a property, do it here
