Add replayID CSV metadata

diff --git a/Engine/Source/Runtime/Engine/Private/DemoNetDriver.cpp b/Engine/Source/Runtime/Engine/Private/DemoNetDriver.cpp
--- a/Engine/Source/Runtime/Engine/Private/DemoNetDriver.cpp
+++ b/Engine/Source/Runtime/Engine/Private/DemoNetDriver.cpp
@@ -2600,6 +2600,8 @@ void UDemoNetDriver::TickDemoPlayback(float DeltaSeconds)
 
 			FReplayPlaylistTracker* LocalPlaylistTracker = PlaylistTracker.Get();
 
+			CSV_METADATA(TEXT("ReplayID"), nullptr);
+
 			// checking against 1 so the count will mean total number of playthroughs, not additional loops
 			if (GDemoLoopCount > 1)
 			{
@@ -3031,6 +3033,8 @@ void UDemoNetDriver::ReplayStreamingReady(const FStartStreamingResult& Result)
 
 			UE_LOG(LogDemo, Log, TEXT("ReplayStreamingReady: playing back replay [%s] %s, which was recorded on engine version %s with flags [%s]"),
 				*ReplayHelper.GetPlaybackGuid().ToString(EGuidFormats::Digits), *ReplayHelper.DemoURL.Map, *ReplayHelper.PlaybackDemoHeader.EngineVersion.ToString(), *HeaderFlags);
+
+			CSV_METADATA(TEXT("ReplayID"), *ReplayHelper.GetPlaybackGuid().ToString(EGuidFormats::Digits));
 		}
 
 		// Notify all listeners that a demo is starting
