WorldPartition: Fix BP Reinstancing not doing fixup of loaded actor pointers in their ActorDesc, resulting in Actor not being in the PIE session until they are saved.

diff --git a/Engine/Source/Runtime/Engine/Private/WorldPartition/ActorDescContainer.cpp b/Engine/Source/Runtime/Engine/Private/WorldPartition/ActorDescContainer.cpp
--- a/Engine/Source/Runtime/Engine/Private/WorldPartition/ActorDescContainer.cpp
+++ b/Engine/Source/Runtime/Engine/Private/WorldPartition/ActorDescContainer.cpp
@@ -204,6 +204,24 @@ void UActorDescContainer::OnPackageDeleted(UPackage* Package)
 	}
 }
 
+void UActorDescContainer::OnObjectsReplaced(const TMap<UObject*, UObject*>& OldToNewObjectMap)
+{
+	// Patch up Actor pointers in ActorDescs
+	for (auto Iter = OldToNewObjectMap.CreateConstIterator(); Iter; ++Iter)
+	{
+		if (AActor* OldActor = Cast<AActor>(Iter->Key))
+		{
+			if (FWorldPartitionActorDesc* ActorDesc = GetActorDesc(OldActor->GetActorGuid()))
+			{
+				if (ActorDesc->GetActor() == OldActor)
+				{
+					ActorDesc->ActorPtr = Cast<AActor>(Iter->Value);
+				}
+			}
+		}
+	}
+}
+
 void UActorDescContainer::RemoveActor(const FGuid& ActorGuid)
 {
 	if (TUniquePtr<FWorldPartitionActorDesc>* ExistingActorDesc = GetActorDescriptor(ActorGuid))
@@ -222,6 +240,7 @@ void UActorDescContainer::RegisterEditorDelegates()
 	{
 		FCoreUObjectDelegates::OnObjectSaved.AddUObject(this, &UActorDescContainer::OnObjectPreSave);
 		FEditorDelegates::OnPackageDeleted.AddUObject(this, &UActorDescContainer::OnPackageDeleted);
+		FCoreUObjectDelegates::OnObjectsReplaced.AddUObject(this, &UActorDescContainer::OnObjectsReplaced);
 	}
 }
 
@@ -231,6 +250,7 @@ void UActorDescContainer::UnregisterEditorDelegates()
 	{
 		FCoreUObjectDelegates::OnObjectSaved.RemoveAll(this);
 		FEditorDelegates::OnPackageDeleted.RemoveAll(this);
+		FCoreUObjectDelegates::OnObjectsReplaced.RemoveAll(this);
 	}
 }
 
diff --git a/Engine/Source/Runtime/Engine/Public/WorldPartition/ActorDescContainer.h b/Engine/Source/Runtime/Engine/Public/WorldPartition/ActorDescContainer.h
--- a/Engine/Source/Runtime/Engine/Public/WorldPartition/ActorDescContainer.h
+++ b/Engine/Source/Runtime/Engine/Public/WorldPartition/ActorDescContainer.h
@@ -26,6 +26,7 @@ public:
 	// Events
 	virtual void OnObjectPreSave(UObject* Object);
 	virtual void OnPackageDeleted(UPackage* Package);
+	virtual void OnObjectsReplaced(const TMap<UObject*, UObject*>& OldToNewObjectMap);
 
 	FName GetContainerPackage() const { return ContainerPackageName; }
 
