Fixes for log spam

diff --git a/Engine/Source/Runtime/Engine/Private/AudioStreamingCache.cpp b/Engine/Source/Runtime/Engine/Private/AudioStreamingCache.cpp
--- a/Engine/Source/Runtime/Engine/Private/AudioStreamingCache.cpp
+++ b/Engine/Source/Runtime/Engine/Private/AudioStreamingCache.cpp
@@ -1066,7 +1066,7 @@ uint64 FAudioChunkCache::TrimMemory(uint64 BytesToFree, bool bInAllowRetainedChu
 		}
 	}
 
-	UE_LOG(LogAudioStreamCaching, Display, TEXT("TrimMemory: NumElements Evicted: %d. Bytes Freed: %d"), NumElementsEvicted, BytesFreed);
+	UE_LOG(LogAudioStreamCaching, Verbose, TEXT("TrimMemory: NumElements Evicted: %d. Bytes Freed: %d"), NumElementsEvicted, BytesFreed);
 
 	return BytesFreed;
 }
diff --git a/Engine/Source/Runtime/Engine/Private/SoundWave.cpp b/Engine/Source/Runtime/Engine/Private/SoundWave.cpp
--- a/Engine/Source/Runtime/Engine/Private/SoundWave.cpp
+++ b/Engine/Source/Runtime/Engine/Private/SoundWave.cpp
@@ -2331,17 +2331,17 @@ void USoundWave::Parse(FAudioDevice* AudioDevice, const UPTRINT NodeWaveInstance
 					if (UAudioComponent* AudioComponent = UAudioComponent::GetAudioComponentFromID(AudioComponentID))
 					{
 						AActor* SoundOwner = AudioComponent->GetOwner();
-						UE_LOG(LogAudio, Warning, TEXT("%s Actor: %s AudioComponent: %s"), *SoundWarningInfo, (SoundOwner ? *SoundOwner->GetName() : TEXT("None")), *AudioComponent->GetName());
+						UE_LOG(LogAudio, Verbose, TEXT("%s Actor: %s AudioComponent: %s"), *SoundWarningInfo, (SoundOwner ? *SoundOwner->GetName() : TEXT("None")), *AudioComponent->GetName());
 					}
 					else
 					{
-						UE_LOG(LogAudio, Warning, TEXT("%s"), *SoundWarningInfo);
+						UE_LOG(LogAudio, Verbose, TEXT("%s"), *SoundWarningInfo);
 					}
 				});
 			}
 			else
 			{
-				UE_LOG(LogAudio, Warning, TEXT("%s"), *SoundWarningInfo);
+				UE_LOG(LogAudio, Verbose, TEXT("%s"), *SoundWarningInfo);
 			}
 
 			ReportedSounds.Add(this);
