Added CVar r.RayTracing.Geometry.SupportSkeletalMeshes to entirely turn off skeletal mesh support in ray tracing effects. Unlike existing r.RayTracing.Geometry.SkeletalMeshes, the new CVar is read-only at runtime.

Changing this CVar does not currently affect cooked data and only prevents creation of ray tracing resources (including GPUSkinCache).

diff --git a/Engine/Source/Runtime/Engine/Private/GPUSkinCache.cpp b/Engine/Source/Runtime/Engine/Private/GPUSkinCache.cpp
--- a/Engine/Source/Runtime/Engine/Private/GPUSkinCache.cpp
+++ b/Engine/Source/Runtime/Engine/Private/GPUSkinCache.cpp
@@ -17,6 +17,7 @@ GPUSkinCache.cpp: Performs skinning on a compute shader into a buffer to avoid v
 #include "MeshMaterialShader.h"
 #include "RenderGraphResources.h"
 #include "Algo/Unique.h"
+#include "HAL/IConsoleManager.h"
 
 DEFINE_STAT(STAT_GPUSkinCache_TotalNumChunks);
 DEFINE_STAT(STAT_GPUSkinCache_TotalNumVertices);
@@ -175,7 +176,8 @@ FAutoConsoleVariableRef CVarGPUSkinCacheMaxRayTracingPrimitivesPerCmdList(
 	TEXT("Maximum amount of primitives which are batched together into a single command list to fix potential TDRs."),
 	ECVF_RenderThreadSafe
 );
-#endif
+
+#endif // RHI_RAYTRACING
 
 const float MBSize = 1048576.f; // 1024 x 1024 bytes
 
@@ -1672,6 +1674,14 @@ void FGPUSkinCache::ProcessEntry(
 }
 
 #if RHI_RAYTRACING
+
+static bool IsGPUSkinCacheRayTracingSupported()
+{
+	static const auto CVar = IConsoleManager::Get().FindConsoleVariable(TEXT("r.RayTracing.Geometry.SupportSkeletalMeshes"));
+	static const bool SupportSkeletalMeshes = CVar->GetInt() != 0;
+	return IsRayTracingEnabled() && SupportSkeletalMeshes && GEnableGPUSkinCache;
+}
+
 void FGPUSkinCache::ProcessRayTracingGeometryToUpdate(
 	FRHICommandListImmediate& RHICmdList,
 	FGPUSkinCacheEntry* SkinCacheEntry,
@@ -1680,7 +1690,7 @@ void FGPUSkinCache::ProcessRayTracingGeometryToUpdate(
 	bool bAnySegmentUsesWorldPositionOffset
 	)
 {
-	if (IsRayTracingEnabled() && GEnableGPUSkinCache && SkinCacheEntry && SkinCacheEntry->GPUSkin->bSupportRayTracing)
+	if (IsGPUSkinCacheRayTracingSupported() && SkinCacheEntry && SkinCacheEntry->GPUSkin->bSupportRayTracing)
 	{
 		FRayTracingGeometry& RayTracingGeometry = SkinCacheEntry->GPUSkin->RayTracingGeometry;
 
@@ -1802,7 +1812,7 @@ void FGPUSkinCache::EndBatchDispatch(FRHICommandListImmediate& RHICmdList)
 	DoDispatch(RHICmdList);
 
 #if RHI_RAYTRACING
-	if (IsRayTracingEnabled() && GEnableGPUSkinCache)
+	if (IsGPUSkinCacheRayTracingSupported())
 	{
 		TSet<FGPUSkinCacheEntry*> SkinCacheEntriesProcessed;
 
@@ -2259,7 +2269,7 @@ void FGPUSkinCache::CVarSinkFunction()
 
 	if (GEnableGPUSkinCacheShaders)
 	{
-		if (GIsRHIInitialized && IsRayTracingEnabled())
+		if (GIsRHIInitialized && IsGPUSkinCacheRayTracingSupported())
 		{
 			// Skin cache is *required* for ray tracing.
 			NewGPUSkinCacheValue = 1;
diff --git a/Engine/Source/Runtime/Engine/Private/SkeletalMesh.cpp b/Engine/Source/Runtime/Engine/Private/SkeletalMesh.cpp
--- a/Engine/Source/Runtime/Engine/Private/SkeletalMesh.cpp
+++ b/Engine/Source/Runtime/Engine/Private/SkeletalMesh.cpp
@@ -126,6 +126,14 @@ static TAutoConsoleVariable<int32> CVarRayTracingSkeletalMeshes(
 	1,
 	TEXT("Include skeletal meshes in ray tracing effects (default = 1 (skeletal meshes enabled in ray tracing))"));
 
+static TAutoConsoleVariable<int32> CVarRayTracingSupportSkeletalMeshes(
+	TEXT("r.RayTracing.Geometry.SupportSkeletalMeshes"),
+	1,
+	TEXT("Whether the project supports skeletal meshes in ray tracing effects. ")
+	TEXT("Turning this off disables creation of all skeletal mesh ray tracing GPU resources, saving GPU memory and time. ")
+	TEXT("This setting is read-only at runtime. (default: 1)"),
+	ECVF_ReadOnly);
+
 #if WITH_APEX_CLOTHING
 /*-----------------------------------------------------------------------------
 	utility functions for apex clothing 
@@ -6174,7 +6182,8 @@ void FSkeletalMeshSceneProxy::GetDynamicElementsSection(const TArray<const FScen
 #if RHI_RAYTRACING
 void FSkeletalMeshSceneProxy::GetDynamicRayTracingInstances(FRayTracingMaterialGatheringContext & Context, TArray<struct FRayTracingInstance>& OutRayTracingInstances)
 {
-	if (!CVarRayTracingSkeletalMeshes.GetValueOnRenderThread())
+	if (!CVarRayTracingSkeletalMeshes.GetValueOnRenderThread()
+		|| !CVarRayTracingSupportSkeletalMeshes.GetValueOnRenderThread())
 	{
 		return;
 	}
@@ -6190,10 +6199,10 @@ void FSkeletalMeshSceneProxy::GetDynamicRayTracingInstances(FRayTracingMaterialG
 			FRayTracingInstance RayTracingInstance;
 			RayTracingInstance.Geometry = MeshObject->GetRayTracingGeometry();
 
-				// Setup materials for each segment
-				const int32 LODIndex = MeshObject->GetLOD();
-				check(LODIndex < SkeletalMeshRenderData->LODRenderData.Num());
-				const FSkeletalMeshLODRenderData& LODData = SkeletalMeshRenderData->LODRenderData[LODIndex];
+			// Setup materials for each segment
+			const int32 LODIndex = MeshObject->GetLOD();
+			check(LODIndex < SkeletalMeshRenderData->LODRenderData.Num());
+			const FSkeletalMeshLODRenderData& LODData = SkeletalMeshRenderData->LODRenderData[LODIndex];
 
 			if (LODIndex < SkeletalMeshRenderData->CurrentFirstLODIdx)
 			{
@@ -6201,76 +6210,76 @@ void FSkeletalMeshSceneProxy::GetDynamicRayTracingInstances(FRayTracingMaterialG
 				return;
 			}
 
-				ensure(LODSections.Num() > 0);
-				const FLODSectionElements& LODSection = LODSections[LODIndex];
-				check(LODSection.SectionElements.Num() == LODData.RenderSections.Num());
+			ensure(LODSections.Num() > 0);
+			const FLODSectionElements& LODSection = LODSections[LODIndex];
+			check(LODSection.SectionElements.Num() == LODData.RenderSections.Num());
 				
-				//#dxr_todo (UE-113617): verify why this condition is not fulfilled sometimes
-				if(!ensure(LODSection.SectionElements.Num() == MeshObject->GetRayTracingGeometry()->Initializer.Segments.Num()))
-				{
-					return;
-				}
+			//#dxr_todo (UE-113617): verify why this condition is not fulfilled sometimes
+			if(!ensure(LODSection.SectionElements.Num() == MeshObject->GetRayTracingGeometry()->Initializer.Segments.Num()))
+			{
+				return;
+			}
 
-			#if WITH_EDITORONLY_DATA
-				int32 SectionIndexPreview = MeshObject->SectionIndexPreview;
-				int32 MaterialIndexPreview = MeshObject->MaterialIndexPreview;
-				MeshObject->SectionIndexPreview = INDEX_NONE;
-				MeshObject->MaterialIndexPreview = INDEX_NONE;
-			#endif
+		#if WITH_EDITORONLY_DATA
+			int32 SectionIndexPreview = MeshObject->SectionIndexPreview;
+			int32 MaterialIndexPreview = MeshObject->MaterialIndexPreview;
+			MeshObject->SectionIndexPreview = INDEX_NONE;
+			MeshObject->MaterialIndexPreview = INDEX_NONE;
+		#endif
 
-				for (FSkeletalMeshSectionIter Iter(LODIndex, *MeshObject, LODData, LODSection); Iter; ++Iter)
-				{
-					const FSkelMeshRenderSection& Section = Iter.GetSection();
-					const int32 SectionIndex = Iter.GetSectionElementIndex();
-					const FSectionElementInfo& SectionElementInfo = Iter.GetSectionElementInfo();
+			for (FSkeletalMeshSectionIter Iter(LODIndex, *MeshObject, LODData, LODSection); Iter; ++Iter)
+			{
+				const FSkelMeshRenderSection& Section = Iter.GetSection();
+				const int32 SectionIndex = Iter.GetSectionElementIndex();
+				const FSectionElementInfo& SectionElementInfo = Iter.GetSectionElementInfo();
 
-					FMeshBatch MeshBatch;
-					CreateBaseMeshBatch(Context.ReferenceView, LODData, LODIndex, SectionIndex, SectionElementInfo, MeshBatch);
+				FMeshBatch MeshBatch;
+				CreateBaseMeshBatch(Context.ReferenceView, LODData, LODIndex, SectionIndex, SectionElementInfo, MeshBatch);
 
-					RayTracingInstance.Materials.Add(MeshBatch);
-				}
+				RayTracingInstance.Materials.Add(MeshBatch);
+			}
 
-			#if WITH_EDITORONLY_DATA
-				MeshObject->SectionIndexPreview = SectionIndexPreview;
-				MeshObject->MaterialIndexPreview = MaterialIndexPreview;
-			#endif
-				if (bAnySegmentUsesWorldPositionOffset)
-				{
-					RayTracingInstance.InstanceTransforms.Add(FMatrix::Identity);
-				}
-				else
+		#if WITH_EDITORONLY_DATA
+			MeshObject->SectionIndexPreview = SectionIndexPreview;
+			MeshObject->MaterialIndexPreview = MaterialIndexPreview;
+		#endif
+			if (bAnySegmentUsesWorldPositionOffset)
+			{
+				RayTracingInstance.InstanceTransforms.Add(FMatrix::Identity);
+			}
+			else
+			{
+				RayTracingInstance.InstanceTransforms.Add(GetLocalToWorld());
+			}
+
+			if (bAnySegmentUsesWorldPositionOffset && MeshObject->GetSkinVertexFactory(Context.ReferenceView, LODIndex, 0)->GetType()->SupportsRayTracingDynamicGeometry())
+			{
+				TArray<FRayTracingGeometrySegment> GeometrySections;
+				GeometrySections.Reserve(LODData.RenderSections.Num());
+				for (const FSkelMeshRenderSection& Section : LODData.RenderSections)
 				{
-					RayTracingInstance.InstanceTransforms.Add(GetLocalToWorld());
+					FRayTracingGeometrySegment Segment;
+					Segment.FirstPrimitive = Section.BaseIndex / 3;
+					Segment.NumPrimitives = Section.NumTriangles;
+					Segment.bEnabled = !Section.bDisabled;
+					GeometrySections.Add(Segment);
 				}
+				MeshObject->GetRayTracingGeometry()->Initializer.Segments = GeometrySections;
 
-				if (bAnySegmentUsesWorldPositionOffset && MeshObject->GetSkinVertexFactory(Context.ReferenceView, LODIndex, 0)->GetType()->SupportsRayTracingDynamicGeometry())
-				{
-					TArray<FRayTracingGeometrySegment> GeometrySections;
-					GeometrySections.Reserve(LODData.RenderSections.Num());
-					for (const FSkelMeshRenderSection& Section : LODData.RenderSections)
+				Context.DynamicRayTracingGeometriesToUpdate.Add(
+					FRayTracingDynamicGeometryUpdateParams
 					{
-						FRayTracingGeometrySegment Segment;
-						Segment.FirstPrimitive = Section.BaseIndex / 3;
-						Segment.NumPrimitives = Section.NumTriangles;
-						Segment.bEnabled = !Section.bDisabled;
-						GeometrySections.Add(Segment);
+						RayTracingInstance.Materials,
+						false,
+						LODData.GetNumVertices(),
+						LODData.GetNumVertices() * (uint32)sizeof(FVector),
+						MeshObject->GetRayTracingGeometry()->Initializer.TotalPrimitiveCount,
+						MeshObject->GetRayTracingGeometry(),
+						MeshObject->GetRayTracingDynamicVertexBuffer(),
+						true
 					}
-					MeshObject->GetRayTracingGeometry()->Initializer.Segments = GeometrySections;
-
-					Context.DynamicRayTracingGeometriesToUpdate.Add(
-						FRayTracingDynamicGeometryUpdateParams
-						{
-							RayTracingInstance.Materials,
-							false,
-							LODData.GetNumVertices(),
-							LODData.GetNumVertices() * (uint32)sizeof(FVector),
-							MeshObject->GetRayTracingGeometry()->Initializer.TotalPrimitiveCount,
-							MeshObject->GetRayTracingGeometry(),
-							MeshObject->GetRayTracingDynamicVertexBuffer(),
-							true
-						}
-					);
-				}
+				);
+			}
 
 			RayTracingInstance.BuildInstanceMaskAndFlags();
 
diff --git a/Engine/Source/Runtime/Engine/Private/SkeletalRender.cpp b/Engine/Source/Runtime/Engine/Private/SkeletalRender.cpp
--- a/Engine/Source/Runtime/Engine/Private/SkeletalRender.cpp
+++ b/Engine/Source/Runtime/Engine/Private/SkeletalRender.cpp
@@ -19,6 +19,15 @@ const float MinMorphTargetBlendWeight = SMALL_NUMBER;
 // largest blend weight for vertex anims
 const float MaxMorphTargetBlendWeight = 5.0f;
 
+#if RHI_RAYTRACING
+static bool IsSkeletalMeshRayTracingSupported()
+{
+	static const auto CVar = IConsoleManager::Get().FindConsoleVariable(TEXT("r.RayTracing.Geometry.SupportSkeletalMeshes"));
+	static const bool SupportSkeletalMeshes = CVar->GetInt() != 0;
+	return SupportSkeletalMeshes;
+}
+#endif // RHI_RAYTRACING
+
 /*-----------------------------------------------------------------------------
 FSkeletalMeshObject
 -----------------------------------------------------------------------------*/
@@ -31,7 +40,7 @@ FSkeletalMeshObject::FSkeletalMeshObject(USkinnedMeshComponent* InMeshComponent,
 ,   bHasBeenUpdatedAtLeastOnce(false)
 #if RHI_RAYTRACING
 , bRequireRecreatingRayTracingGeometry(false)
-, bSupportRayTracing(InMeshComponent->SkeletalMesh->bSupportRayTracing)
+, bSupportRayTracing(IsSkeletalMeshRayTracingSupported() && InMeshComponent->SkeletalMesh->bSupportRayTracing)
 #endif
 #if !UE_BUILD_SHIPPING
 , DebugName(InMeshComponent->SkeletalMesh->GetFName())
