Fixed compile error due to IsGPUSkinCacheRayTracingSupported being used outside of RHI_RAYTRACING

diff --git a/Engine/Source/Runtime/Engine/Private/GPUSkinCache.cpp b/Engine/Source/Runtime/Engine/Private/GPUSkinCache.cpp
--- a/Engine/Source/Runtime/Engine/Private/GPUSkinCache.cpp
+++ b/Engine/Source/Runtime/Engine/Private/GPUSkinCache.cpp
@@ -1673,15 +1673,19 @@ void FGPUSkinCache::ProcessEntry(
 	}
 }
 
-#if RHI_RAYTRACING
-
 static bool IsGPUSkinCacheRayTracingSupported()
 {
+#if RHI_RAYTRACING
 	static const auto CVar = IConsoleManager::Get().FindConsoleVariable(TEXT("r.RayTracing.Geometry.SupportSkeletalMeshes"));
 	static const bool SupportSkeletalMeshes = CVar->GetInt() != 0;
 	return IsRayTracingEnabled() && SupportSkeletalMeshes && GEnableGPUSkinCache;
+#else
+	return false;
+#endif
 }
 
+#if RHI_RAYTRACING
+
 void FGPUSkinCache::ProcessRayTracingGeometryToUpdate(
 	FRHICommandListImmediate& RHICmdList,
 	FGPUSkinCacheEntry* SkinCacheEntry,
