Add interface to enforce a better contract for asset compilers
  - A single call to FAssetCompilingManager::Get().FinishAllCompilation() is now sufficient to ensure every asset in the editor has finished compiling.
  - Dependencies are now used to determine asset compiler ordering
  - Implement interface and unify some functionality between these assets
    - MeshCards
    - DistanceFields
    - SkeletalMeshes
    - StaticMeshes
    - Shaders
    - Textures

diff --git a/Engine/Plugins/MovieScene/MovieRenderPipeline/Source/MovieRenderPipelineCore/Private/MoviePipelineRendering.cpp b/Engine/Plugins/MovieScene/MovieRenderPipeline/Source/MovieRenderPipelineCore/Private/MoviePipelineRendering.cpp
--- a/Engine/Plugins/MovieScene/MovieRenderPipeline/Source/MovieRenderPipelineCore/Private/MoviePipelineRendering.cpp
+++ b/Engine/Plugins/MovieScene/MovieRenderPipeline/Source/MovieRenderPipelineCore/Private/MoviePipelineRendering.cpp
@@ -491,89 +491,10 @@ void UMoviePipeline::FlushAsyncEngineSystems()
 	{
 		GetWorld()->BlockTillLevelStreamingCompleted();
 	}
-	
-#if WITH_EDITOR
-	// Flush all assets handled by the asset compiling manager (i.e. textures, static meshes).
+
+	// Flush all assets still being compiled asynchronously.
 	// A progressbar is already in place so the user can get feedback while waiting for everything to settle.
 	FAssetCompilingManager::Get().FinishAllCompilation();
-#endif
-
-	// Now we can flush the shader compiler. ToDo: This should probably happen right before SendAllEndOfFrameUpdates() is normally called
-	if (GShaderCompilingManager)
-	{
-		bool bDidWork = false;
-		int32 NumShadersToCompile = GShaderCompilingManager->GetNumRemainingJobs();
-		if (NumShadersToCompile > 0)
-		{
-			UE_LOG(LogMovieRenderPipeline, Log, TEXT("[%d] Starting build for %d shaders."), GFrameCounter, NumShadersToCompile);
-		}
-
-		while (GShaderCompilingManager->GetNumRemainingJobs() > 0 || GShaderCompilingManager->HasShaderJobs())
-		{
-			UE_LOG(LogMovieRenderPipeline, Log, TEXT("[%d] Waiting for %d shaders [Has Shader Jobs: %d] to finish compiling..."), GFrameCounter, GShaderCompilingManager->GetNumRemainingJobs(), GShaderCompilingManager->HasShaderJobs());
-			GShaderCompilingManager->ProcessAsyncResults(false, true);
-
-			// Sleep for 1 second and then check again. This way we get an indication of progress as this works.
-			FPlatformProcess::Sleep(1.f);
-			bDidWork = true;
-		}
-
-		if (bDidWork)
-		{
-			UE_LOG(LogMovieRenderPipeline, Log, TEXT("[%d] Done building %d shaders."), GFrameCounter, NumShadersToCompile);
-		}
-	}
-
-	// Flush the Mesh Distance Field builder as well.
-	if (GDistanceFieldAsyncQueue)
-	{
-		bool bDidWork = false;
-		int32 NumDistanceFieldsToBuild = GDistanceFieldAsyncQueue->GetNumOutstandingTasks();
-		if (NumDistanceFieldsToBuild > 0)
-		{
-			UE_LOG(LogMovieRenderPipeline, Log, TEXT("[%d] Starting build for %d mesh distance fields."), GFrameCounter, NumDistanceFieldsToBuild);
-		}
-
-		while (GDistanceFieldAsyncQueue->GetNumOutstandingTasks() > 0)
-		{
-			UE_LOG(LogMovieRenderPipeline, Log, TEXT("[%d] Waiting for %d Mesh Distance Fields to finish building..."), GFrameCounter, GDistanceFieldAsyncQueue->GetNumOutstandingTasks());
-			GDistanceFieldAsyncQueue->ProcessAsyncTasks();
-
-			// Sleep for 1 second and then check again. This way we get an indication of progress as this works.
-			FPlatformProcess::Sleep(1.f);
-			bDidWork = true;
-		}
-
-		if (bDidWork)
-		{
-			UE_LOG(LogMovieRenderPipeline, Log, TEXT("[%d] Done building %d Mesh Distance Fields."), GFrameCounter, NumDistanceFieldsToBuild);
-		}
-	}
-
-	if (GCardRepresentationAsyncQueue)
-	{
-		bool bDidWork = false;
-		int32 NumTasks = GCardRepresentationAsyncQueue->GetNumOutstandingTasks();
-		if (NumTasks > 0)
-		{
-			UE_LOG(LogMovieRenderPipeline, Log, TEXT("[%d] Starting build for %d mesh cards."), GFrameCounter, NumTasks);
-		}
-
-		while (GCardRepresentationAsyncQueue->GetNumOutstandingTasks() > 0)
-		{
-			UE_LOG(LogMovieRenderPipeline, Log, TEXT("[%d] Waiting for %d Mesh Cards to finish building..."), GFrameCounter, GCardRepresentationAsyncQueue->GetNumOutstandingTasks());
-			GCardRepresentationAsyncQueue->ProcessAsyncTasks();
-
-			// Sleep for 1 second and then check again. This way we get an indication of progress as this works.
-			FPlatformProcess::Sleep(1.f);
-			bDidWork = true;
-		}
-
-		if (bDidWork)
-		{
-			UE_LOG(LogMovieRenderPipeline, Log, TEXT("[%d] Done building %d Mesh Cards."), GFrameCounter, NumTasks);
-		}
-	}
 
 	// Flush grass
 	{
diff --git a/Engine/Source/Runtime/Engine/Private/AssetCompilingManager.cpp b/Engine/Source/Runtime/Engine/Private/AssetCompilingManager.cpp
--- a/Engine/Source/Runtime/Engine/Private/AssetCompilingManager.cpp
+++ b/Engine/Source/Runtime/Engine/Private/AssetCompilingManager.cpp
@@ -9,8 +9,6 @@ DECLARE_LLM_MEMORY_STAT(TEXT("AssetCompilation"), STAT_AssetCompilationLLM, STAT
 DECLARE_LLM_MEMORY_STAT(TEXT("AssetCompilation"), STAT_AssetCompilationSummaryLLM, STATGROUP_LLM);
 LLM_DEFINE_TAG(AssetCompilation, NAME_None, NAME_None, GET_STATFNAME(STAT_AssetCompilationLLM), GET_STATFNAME(STAT_AssetCompilationSummaryLLM));
 
-#if WITH_EDITOR
-
 #include "Misc/QueuedThreadPoolWrapper.h"
 #include "Misc/CommandLine.h"
 #include "UObject/UObjectIterator.h"
@@ -19,12 +17,17 @@ LLM_DEFINE_TAG(AssetCompilation, NAME_None, NAME_None, GET_STATFNAME(STAT_AssetC
 #include "TextureCompiler.h"
 #include "Components/StaticMeshComponent.h"
 #include "Materials/MaterialInstance.h"
+#include "ObjectCacheContext.h"
 #include "AsyncCompilationHelpers.h"
 #include "SkeletalMeshCompiler.h"
+#include "Algo/TopologicalSort.h"
+#include "Algo/Find.h"
 #include "ProfilingDebugging/CountersTrace.h"
 
 #define LOCTEXT_NAMESPACE "AssetCompilingManager"
 
+#if WITH_EDITOR
+
 static AsyncCompilationHelpers::FAsyncCompilationStandardCVars CVarAsyncCompilationStandard(
 	TEXT("Asset"), 
 	TEXT("assets"),
@@ -51,10 +54,18 @@ static TAutoConsoleVariable<int32> CVarAsyncAssetCompilationMaxMemoryUsage(
 	TEXT("Try to stay under specified memory limit for asset compilation by reducing concurrency when under memory pressure.\n"),
 	ECVF_Default);
 
+TRACE_DECLARE_INT_COUNTER(AsyncCompilationMaxConcurrency, TEXT("AsyncCompilation/MaxConcurrency"));
+TRACE_DECLARE_INT_COUNTER(AsyncCompilationConcurrency, TEXT("AsyncCompilation/Concurrency"));
+TRACE_DECLARE_MEMORY_COUNTER(AsyncCompilationTotalMemoryLimit, TEXT("AsyncCompilation/TotalMemoryLimit"));
+TRACE_DECLARE_MEMORY_COUNTER(AsyncCompilationTotalEstimatedMemory, TEXT("AsyncCompilation/TotalEstimatedMemory"));
+
+#endif //#if WITH_EDITOR
+
 namespace AssetCompilingManagerImpl
 {
 	static void EnsureInitializedCVars()
 	{
+#if WITH_EDITOR
 		static bool bIsInitialized = false;
 
 		if (!bIsInitialized)
@@ -67,24 +78,11 @@ namespace AssetCompilingManagerImpl
 				CVarAsyncCompilationStandard.AsyncCompilationMaxConcurrency
 			);
 		}
+#endif
 	}
 }
 
-/**
- * Returns the number of outstanding asset compilations.
- */
-int32 FAssetCompilingManager::GetNumRemainingAssets() const
-{
-	return 
-		FStaticMeshCompilingManager::Get().GetNumRemainingMeshes() +
-		FTextureCompilingManager::Get().GetNumRemainingTextures() + 
-		FSkeletalMeshCompilingManager::Get().GetNumRemainingJobs();
-}
-
-TRACE_DECLARE_INT_COUNTER(AsyncCompilationMaxConcurrency, TEXT("AsyncCompilation/MaxConcurrency"));
-TRACE_DECLARE_INT_COUNTER(AsyncCompilationConcurrency, TEXT("AsyncCompilation/Concurrency"));
-TRACE_DECLARE_MEMORY_COUNTER(AsyncCompilationTotalMemoryLimit, TEXT("AsyncCompilation/TotalMemoryLimit"));
-TRACE_DECLARE_MEMORY_COUNTER(AsyncCompilationTotalEstimatedMemory, TEXT("AsyncCompilation/TotalEstimatedMemory"));
+#if WITH_EDITOR
 
 class FMemoryBoundQueuedThreadPoolWrapper : public FQueuedThreadPoolWrapper
 {
@@ -186,8 +184,25 @@ private:
 	std::atomic<int64> TotalEstimatedMemory {0};
 };
 
+#endif // #if WITH_EDITOR
+
+TArrayView<IAssetCompilingManager* const> FAssetCompilingManager::GetRegisteredManagers() const
+{
+	return TArrayView<IAssetCompilingManager* const>(AssetCompilingManagers);
+}
+
+FAssetCompilingManager::FAssetCompilingManager()
+{
+#if WITH_EDITOR
+	RegisterManager(&FStaticMeshCompilingManager::Get());
+	RegisterManager(&FSkeletalMeshCompilingManager::Get());
+	RegisterManager(&FTextureCompilingManager::Get());
+#endif
+}
+
 FQueuedThreadPool* FAssetCompilingManager::GetThreadPool() const
 {
+#if WITH_EDITOR
 	static FQueuedThreadPoolWrapper* GAssetThreadPool = nullptr;
 	if (GAssetThreadPool == nullptr)
 	{
@@ -209,6 +224,75 @@ FQueuedThreadPool* FAssetCompilingManager::GetThreadPool() const
 	}
 
 	return GAssetThreadPool;
+#else
+	return GThreadPool;
+#endif
+}
+
+/**
+ * Register an asset compiling manager.
+ */
+bool FAssetCompilingManager::RegisterManager(IAssetCompilingManager* InAssetCompilingManager)
+{
+	if (AssetCompilingManagers.Contains(InAssetCompilingManager))
+	{
+		return false;
+	}
+
+	AssetCompilingManagers.Add(InAssetCompilingManager);
+	AssetCompilingManagersWithValidDependencies.Add(InAssetCompilingManager);
+
+	auto GetVertexDependencies =
+		[this](IAssetCompilingManager* CurrentManager)
+	{
+		TArray<IAssetCompilingManager*> Connections;
+		for (FName DependentTypeName : CurrentManager->GetDependentTypeNames())
+		{
+			IAssetCompilingManager** DependentManager =
+				Algo::FindBy(AssetCompilingManagersWithValidDependencies, DependentTypeName, [](const IAssetCompilingManager* Manager) { return Manager->GetAssetTypeName();	});
+
+			if (DependentManager)
+			{
+				Connections.Add(*DependentManager);
+			}
+		}
+		
+		return Connections;
+	};
+
+	if (!Algo::TopologicalSort(AssetCompilingManagers, GetVertexDependencies))
+	{
+		AssetCompilingManagersWithValidDependencies.Remove(InAssetCompilingManager);
+		UE_LOG(LogAsyncCompilation, Error, TEXT("AssetCompilingManager %s introduced a circular dependency, it's dependencies will be ignored"), *InAssetCompilingManager->GetAssetTypeName().ToString());
+	}
+
+	return true;
+}
+
+bool FAssetCompilingManager::UnregisterManager(IAssetCompilingManager* InAssetCompilingManager)
+{
+	if (!AssetCompilingManagers.Contains(InAssetCompilingManager))
+	{
+		return false;
+	}
+
+	AssetCompilingManagers.Remove(InAssetCompilingManager);
+	AssetCompilingManagersWithValidDependencies.Remove(InAssetCompilingManager);
+	return true;
+}
+
+/**
+ * Returns the number of outstanding asset compilations.
+ */
+int32 FAssetCompilingManager::GetNumRemainingAssets() const
+{
+	int32 RemainingAssets = 0;
+	for (IAssetCompilingManager* AssetCompilingManager : AssetCompilingManagers)
+	{
+		RemainingAssets += AssetCompilingManager->GetNumRemainingAssets();
+	}
+
+	return RemainingAssets;
 }
 
 /**
@@ -216,9 +300,10 @@ FQueuedThreadPool* FAssetCompilingManager::GetThreadPool() const
  */
 void FAssetCompilingManager::FinishAllCompilation()
 {
-	FTextureCompilingManager::Get().FinishAllCompilation();
-	FStaticMeshCompilingManager::Get().FinishAllCompilation();
-	FSkeletalMeshCompilingManager::Get().FinishAllCompilation();
+	for (IAssetCompilingManager* AssetCompilingManager : AssetCompilingManagers)
+	{
+		AssetCompilingManager->FinishAllCompilation();
+	}
 }
 
 /**
@@ -226,14 +311,17 @@ void FAssetCompilingManager::FinishAllCompilation()
  */
 void FAssetCompilingManager::Shutdown()
 {
-	FStaticMeshCompilingManager::Get().Shutdown();
-	FTextureCompilingManager::Get().Shutdown();
-	FSkeletalMeshCompilingManager::Get().Shutdown();
+	for (IAssetCompilingManager* AssetCompilingManager : AssetCompilingManagers)
+	{
+		AssetCompilingManager->Shutdown();
+	}
 
+#if WITH_EDITOR
 	if (FParse::Param(FCommandLine::Get(), TEXT("DumpAsyncStallsOnExit")))
 	{
 		AsyncCompilationHelpers::DumpStallStacks();
 	}
+#endif // #if WITH_EDITOR
 }
 
 FAssetCompilingManager& FAssetCompilingManager::Get()
@@ -244,14 +332,13 @@ FAssetCompilingManager& FAssetCompilingManager::Get()
 
 void FAssetCompilingManager::ProcessAsyncTasks(bool bLimitExecutionTime)
 {
-	// Update textures first to avoid having to update the render state
-	// of static mesh in the same frame we created them...
-	FTextureCompilingManager::Get().ProcessAsyncTasks(bLimitExecutionTime);
-
-	FStaticMeshCompilingManager::Get().ProcessAsyncTasks(bLimitExecutionTime);
+	// Reuse ObjectIterator Caching and Reverse lookups for the duration of all asset updates
+	FObjectCacheContextScope ObjectCacheScope;
 
-	FSkeletalMeshCompilingManager::Get().ProcessAsyncTasks(bLimitExecutionTime);
+	for (IAssetCompilingManager* AssetCompilingManager : AssetCompilingManagers)
+	{
+		AssetCompilingManager->ProcessAsyncTasks(bLimitExecutionTime);
+	}
 }
 
-#undef LOCTEXT_NAMESPACE
-#endif // #if WITH_EDITOR
\ No newline at end of file
+#undef LOCTEXT_NAMESPACE
\ No newline at end of file
diff --git a/Engine/Source/Runtime/Engine/Private/AsyncCompilationHelpers.cpp b/Engine/Source/Runtime/Engine/Private/AsyncCompilationHelpers.cpp
--- a/Engine/Source/Runtime/Engine/Private/AsyncCompilationHelpers.cpp
+++ b/Engine/Source/Runtime/Engine/Private/AsyncCompilationHelpers.cpp
@@ -19,10 +19,14 @@
 #include "Misc/ScopeRWLock.h"
 #include "ProfilingDebugging/StallDetector.h"
 
-#define LOCTEXT_NAMESPACE "AsyncCompilation"
+#endif // #if WITH_EDITOR
 
 DEFINE_LOG_CATEGORY(LogAsyncCompilation);
 
+#if WITH_EDITOR
+
+#define LOCTEXT_NAMESPACE "AsyncCompilation"
+
 namespace AsyncCompilationHelpers
 {
 	void DumpStallStacks();
@@ -41,16 +45,22 @@ void FAsyncCompilationNotification::Update(int32 NumJobs)
 {
 	check(IsInGameThread());
 
-	FFormatNamedArguments Args;
-	Args.Add(TEXT("AssetType"), AssetType);
-	Args.Add(TEXT("NumJobs"), FText::AsNumber(NumJobs));
-	FText ProgressMessage = FText::Format(LOCTEXT("AsyncCompilationProgress", "Preparing {AssetType} ({NumJobs})"), Args);
+	// Use a lambda to only for progress message in code-path where it's needed
+	auto GetProgressMessage =
+		[this, NumJobs]()
+		{
+			FFormatNamedArguments Args;
+			// Always use the plural form when displaying the notification
+			Args.Add(TEXT("AssetType"), FText::Format(AssetNameFormat, FText::AsNumber(100)));
+			Args.Add(TEXT("NumJobs"), FText::AsNumber(NumJobs));
+			return FText::Format(LOCTEXT("AsyncCompilationProgress", "Preparing {AssetType} ({NumJobs})"), Args);
+		};
 
 	if (NumJobs == 0)
 	{
 		if (NotificationHandle.IsValid())
 		{
-			FSlateNotificationManager::Get().UpdateProgressNotification(NotificationHandle, StartNumJobs, StartNumJobs, ProgressMessage);
+			FSlateNotificationManager::Get().UpdateProgressNotification(NotificationHandle, StartNumJobs, StartNumJobs, GetProgressMessage());
 		}
 		StartNumJobs = 0;
 		NotificationHandle = FProgressNotificationHandle();
@@ -60,7 +70,7 @@ void FAsyncCompilationNotification::Update(int32 NumJobs)
 		if (!NotificationHandle.IsValid())
 		{
 			StartNumJobs = NumJobs;
-			NotificationHandle = FSlateNotificationManager::Get().StartProgressNotification(ProgressMessage, StartNumJobs);
+			NotificationHandle = FSlateNotificationManager::Get().StartProgressNotification(GetProgressMessage(), StartNumJobs);
 		}
 		else
 		{
@@ -68,7 +78,7 @@ void FAsyncCompilationNotification::Update(int32 NumJobs)
 			{
 				StartNumJobs = NumJobs;
 			}
-			FSlateNotificationManager::Get().UpdateProgressNotification(NotificationHandle, StartNumJobs - NumJobs, StartNumJobs, ProgressMessage);
+			FSlateNotificationManager::Get().UpdateProgressNotification(NotificationHandle, StartNumJobs - NumJobs, StartNumJobs, GetProgressMessage());
 		}
 	}
 };
@@ -377,4 +387,4 @@ namespace AsyncCompilationHelpers
 
 #undef LOCTEXT_NAMESPACE
 
-#endif // #if WITH_EDITOR
\ No newline at end of file
+#endif // #if WITH_EDITOR
diff --git a/Engine/Source/Runtime/Engine/Private/DistanceFieldAtlas.cpp b/Engine/Source/Runtime/Engine/Private/DistanceFieldAtlas.cpp
--- a/Engine/Source/Runtime/Engine/Private/DistanceFieldAtlas.cpp
+++ b/Engine/Source/Runtime/Engine/Private/DistanceFieldAtlas.cpp
@@ -36,6 +36,8 @@
 #include "IMeshBuilderModule.h"
 #endif
 
+#define LOCTEXT_NAMESPACE "DistanceField"
+
 CSV_DEFINE_CATEGORY(DistanceField, false);
 
 #if ENABLE_COOK_STATS
@@ -257,8 +259,23 @@ FAsyncDistanceFieldTask::FAsyncDistanceFieldTask()
 {
 }
 
+void FDistanceFieldAsyncQueue::OnAssetPostCompile(const TArray<FAssetCompileData>& CompiledAssets)
+{
+	for (const FAssetCompileData& CompileData : CompiledAssets)
+	{
+		if (CompileData.Asset.IsValid() && CompileData.Asset.Get()->IsA<UStaticMesh>())
+		{
+			ProcessPendingTasks();
+			return;
+		}
+	}
+}
+
 FDistanceFieldAsyncQueue::FDistanceFieldAsyncQueue() 
 {
+	FAssetCompilingManager::Get().RegisterManager(this);
+	FAssetCompilingManager::Get().OnAssetPostCompileEvent().AddRaw(this, &FDistanceFieldAsyncQueue::OnAssetPostCompile);
+
 #if WITH_EDITOR
 	MeshUtilities = NULL;
 
@@ -273,6 +290,8 @@ FDistanceFieldAsyncQueue::FDistanceFieldAsyncQueue()
 
 FDistanceFieldAsyncQueue::~FDistanceFieldAsyncQueue()
 {
+	FAssetCompilingManager::Get().OnAssetPostCompileEvent().RemoveAll(this);
+	FAssetCompilingManager::Get().UnregisterManager(this);
 }
 
 void FAsyncDistanceFieldTaskWorker::DoWork()
@@ -282,6 +301,46 @@ void FAsyncDistanceFieldTaskWorker::DoWork()
 	GDistanceFieldAsyncQueue->Build(&Task, TaskGraphWrapper);
 }
 
+FName FDistanceFieldAsyncQueue::GetStaticAssetTypeName()
+{
+	return TEXT("UE-MeshDistanceField");
+}
+
+FName FDistanceFieldAsyncQueue::GetAssetTypeName() const 
+{ 
+	return GetStaticAssetTypeName();
+}
+
+FTextFormat FDistanceFieldAsyncQueue::GetAssetNameFormat() const
+{
+	return LOCTEXT("MeshDistanceFieldNameFormat", "{0}|plural(one=Mesh Distance Field,other=Mesh Distance Fields)");
+}
+
+TArrayView<FName> FDistanceFieldAsyncQueue::GetDependentTypeNames() const 
+{
+#if WITH_EDITOR
+	static FName DependentTypeNames[] = 
+	{
+		// DistanceField are a byproduct of StaticMesh so they have a natural dependencies toward them.
+		FStaticMeshCompilingManager::GetStaticAssetTypeName() 
+	};
+
+	return TArrayView<FName>(DependentTypeNames);
+#else
+	return TArrayView<FName>();
+#endif
+}
+
+int32 FDistanceFieldAsyncQueue::GetNumRemainingAssets() const 
+{
+	return GetNumOutstandingTasks(); 
+}
+
+void FDistanceFieldAsyncQueue::FinishAllCompilation() 
+{
+	BlockUntilAllBuildsComplete(); 
+}
+
 void FDistanceFieldAsyncQueue::CancelBackgroundTask(TArray<FAsyncDistanceFieldTask*> Tasks)
 {
 	// Do all the cancellation first to make sure none of these tasks
@@ -1303,3 +1362,5 @@ FIntPoint FLandscapeTextureAtlas::FPendingUpload::SetCommonShaderParameters(void
 	const uint32 NumGroupsY = FMath::DivideAndRoundUp(UpdateRegionSizeY, FUploadLandscapeTextureToAtlasCS::ThreadGroupSizeY);
 	return FIntPoint(NumGroupsX, NumGroupsY);
 }
+
+#undef LOCTEXT_NAMESPACE
\ No newline at end of file
diff --git a/Engine/Source/Runtime/Engine/Private/MeshCardRepresentation.cpp b/Engine/Source/Runtime/Engine/Private/MeshCardRepresentation.cpp
--- a/Engine/Source/Runtime/Engine/Private/MeshCardRepresentation.cpp
+++ b/Engine/Source/Runtime/Engine/Private/MeshCardRepresentation.cpp
@@ -17,6 +17,7 @@
 #include "Engine/StaticMesh.h"
 #include "Misc/AutomationTest.h"
 #include "Async/ParallelFor.h"
+#include "DistanceFieldAtlas.h"
 #include "Misc/QueuedThreadPoolWrapper.h"
 #include "Async/Async.h"
 #include "ObjectCacheContext.h"
@@ -28,6 +29,8 @@
 #include "StaticMeshCompiler.h"
 #endif
 
+#define LOCTEXT_NAMESPACE "MeshCardRepresentation"
+
 #if WITH_EDITORONLY_DATA
 #include "IMeshBuilderModule.h"
 #endif
@@ -177,10 +180,13 @@ FCardRepresentationAsyncQueue::FCardRepresentationAsyncQueue()
 	const int32 MaxConcurrency = 1;
 	ThreadPool = MakeUnique<FQueuedThreadPoolWrapper>(GThreadPool, MaxConcurrency, [](EQueuedWorkPriority) { return EQueuedWorkPriority::Lowest; });
 #endif
+
+	FAssetCompilingManager::Get().RegisterManager(this);
 }
 
 FCardRepresentationAsyncQueue::~FCardRepresentationAsyncQueue()
 {
+	FAssetCompilingManager::Get().UnregisterManager(this);
 }
 
 void FAsyncCardRepresentationTaskWorker::DoWork()
@@ -190,6 +196,37 @@ void FAsyncCardRepresentationTaskWorker::DoWork()
 	GCardRepresentationAsyncQueue->Build(&Task, TaskGraphWrapper);
 }
 
+FName FCardRepresentationAsyncQueue::GetStaticAssetTypeName()
+{
+	return TEXT("UE-MeshCard");
+}
+
+FName FCardRepresentationAsyncQueue::GetAssetTypeName() const 
+{
+	return GetStaticAssetTypeName();
+}
+
+FTextFormat FCardRepresentationAsyncQueue::GetAssetNameFormat() const
+{
+	return LOCTEXT("MeshCardNameFormat", "{0}|plural(one=Mesh Card,other=Mesh Cards)");
+}
+
+TArrayView<FName> FCardRepresentationAsyncQueue::GetDependentTypeNames() const 
+{
+	static FName DependentTypeNames[] = { FDistanceFieldAsyncQueue::GetStaticAssetTypeName() };
+	return TArrayView<FName>(DependentTypeNames);
+}
+
+int32 FCardRepresentationAsyncQueue::GetNumRemainingAssets() const
+{
+	return GetNumOutstandingTasks();
+}
+
+void FCardRepresentationAsyncQueue::FinishAllCompilation()
+{
+	BlockUntilAllBuildsComplete();
+}
+
 void FCardRepresentationAsyncQueue::CancelBackgroundTask(TArray<FAsyncCardRepresentationTask*> Tasks)
 {
 	// Do all the cancellation first to make sure none of these tasks
@@ -560,4 +597,6 @@ void FCardRepresentationAsyncQueue::Shutdown()
 
 	UE_LOG(LogStaticMesh, Log, TEXT("Abandoning remaining async card representation tasks for shutdown"));
 	ThreadPool->Destroy();
-}
\ No newline at end of file
+}
+
+#undef LOCTEXT_NAMESPACE
\ No newline at end of file
diff --git a/Engine/Source/Runtime/Engine/Private/ShaderCompiler/ShaderCompiler.cpp b/Engine/Source/Runtime/Engine/Private/ShaderCompiler/ShaderCompiler.cpp
--- a/Engine/Source/Runtime/Engine/Private/ShaderCompiler/ShaderCompiler.cpp
+++ b/Engine/Source/Runtime/Engine/Private/ShaderCompiler/ShaderCompiler.cpp
@@ -3020,6 +3020,61 @@ FShaderCompilingManager::FShaderCompilingManager() :
 	{
 		Thread->StartThread();
 	}
+
+	FAssetCompilingManager::Get().RegisterManager(this);
+}
+
+FShaderCompilingManager::~FShaderCompilingManager()
+{
+	PrintStats();
+
+	for (const auto& Thread : Threads)
+	{
+		Thread->Stop();
+		Thread->WaitForCompletion();
+	}
+
+	FAssetCompilingManager::Get().UnregisterManager(this);
+}
+
+FName FShaderCompilingManager::GetStaticAssetTypeName()
+{
+	return TEXT("UE-Shader");
+}
+
+FName FShaderCompilingManager::GetAssetTypeName() const
+{
+	return GetStaticAssetTypeName();
+}
+
+FTextFormat FShaderCompilingManager::GetAssetNameFormat() const
+{
+	return LOCTEXT("ShaderNameFormat", "{0}|plural(one=Shader,other=Shaders)");
+}
+
+TArrayView<FName> FShaderCompilingManager::GetDependentTypeNames() const
+{
+#if WITH_EDITOR
+	static FName DependentTypeNames[] = 
+	{
+		// Texture can require materials to be updated,
+		// they should be processed first to avoid unecessary material updates.
+		FTextureCompilingManager::GetStaticAssetTypeName() 
+	};
+	return TArrayView<FName>(DependentTypeNames);
+#else
+	return TArrayView<FName>();
+#endif	
+}
+
+int32 FShaderCompilingManager::GetNumRemainingAssets() const
+{
+	return GetNumOutstandingJobs() + GetNumPendingJobs();
+}
+
+void FShaderCompilingManager::ProcessAsyncTasks(bool bLimitExecutionTime)
+{
+	ProcessAsyncResults(bLimitExecutionTime, false);
 }
 
 int32 FShaderCompilingManager::GetNumPendingJobs() const
@@ -3867,13 +3922,8 @@ void FShaderCompilingManager::PropagateMaterialChangesToPrimitives(const TMap<TR
  */
 void FShaderCompilingManager::Shutdown()
 {
-	PrintStats();
-
-	for (const auto& Thread : Threads)
-	{
-		Thread->Stop();
-		Thread->WaitForCompletion();
-	}
+	// Shutdown has been moved to the destructor because the shader compiler lifetime is expected to
+	// be longer than other asset compilers, otherwise niagara compilations might get stuck.
 }
 
 void FShaderCompilingManager::PrintStats(bool bForceLogIgnoringTimeInverval)
diff --git a/Engine/Source/Runtime/Engine/Private/SkeletalMeshCompiler.cpp b/Engine/Source/Runtime/Engine/Private/SkeletalMeshCompiler.cpp
--- a/Engine/Source/Runtime/Engine/Private/SkeletalMeshCompiler.cpp
+++ b/Engine/Source/Runtime/Engine/Private/SkeletalMeshCompiler.cpp
@@ -12,6 +12,8 @@
 #include "EngineModule.h"
 #include "Misc/ScopedSlowTask.h"
 #include "UObject/StrongObjectPtr.h"
+#include "ShaderCompiler.h"
+#include "TextureCompiler.h"
 #include "Misc/IQueuedWork.h"
 #include "LevelEditor.h"
 #include "SLevelViewport.h"
@@ -58,10 +60,37 @@ namespace SkeletalMeshCompilingManagerImpl
 }
 
 FSkeletalMeshCompilingManager::FSkeletalMeshCompilingManager()
-	: Notification(LOCTEXT("SkeletalMeshes", "Skeletal Meshes"))
+	: Notification(GetAssetNameFormat())
 {
 }
 
+FName FSkeletalMeshCompilingManager::GetAssetTypeName() const
+{
+	return TEXT("UE-SkeletalMesh");
+}
+
+FTextFormat FSkeletalMeshCompilingManager::GetAssetNameFormat() const
+{
+	return LOCTEXT("SkeletalMeshNameFormat", "{0}|plural(one=Skeletal Mesh,other=Skeletal Meshes)");
+}
+
+TArrayView<FName> FSkeletalMeshCompilingManager::GetDependentTypeNames() const
+{
+	// Texture and shaders can affect materials which can affect Skeletal Meshes once they are visible.
+	// Adding these dependencies can reduces the actual number of render state update we need to do in a frame
+	static FName DependentTypeNames[] = 
+	{
+		FTextureCompilingManager::GetStaticAssetTypeName(), 
+		FShaderCompilingManager::GetStaticAssetTypeName() 
+	};
+	return TArrayView<FName>(DependentTypeNames);
+}
+
+int32 FSkeletalMeshCompilingManager::GetNumRemainingAssets() const
+{
+	return GetNumRemainingJobs();
+}
+
 EQueuedWorkPriority FSkeletalMeshCompilingManager::GetBasePriority(USkeletalMesh* InSkeletalMesh) const
 {
 	return EQueuedWorkPriority::Low;
diff --git a/Engine/Source/Runtime/Engine/Private/StaticMeshCompiler.cpp b/Engine/Source/Runtime/Engine/Private/StaticMeshCompiler.cpp
--- a/Engine/Source/Runtime/Engine/Private/StaticMeshCompiler.cpp
+++ b/Engine/Source/Runtime/Engine/Private/StaticMeshCompiler.cpp
@@ -19,6 +19,8 @@
 #include "LevelEditor.h"
 #include "SLevelViewport.h"
 #include "Components/PrimitiveComponent.h"
+#include "TextureCompiler.h"
+#include "ShaderCompiler.h"
 #include "ContentStreaming.h"
 #include "DrawDebugHelpers.h"
 #include "GameFramework/PlayerController.h"
@@ -85,7 +87,7 @@ namespace StaticMeshCompilingManagerImpl
 }
 
 FStaticMeshCompilingManager::FStaticMeshCompilingManager()
-	: Notification(LOCTEXT("StaticMeshes", "Static Meshes"))
+	: Notification(GetAssetNameFormat())
 {
 	PostReachabilityAnalysisHandle = FCoreUObjectDelegates::PostReachabilityAnalysis.AddRaw(this, &FStaticMeshCompilingManager::OnPostReachabilityAnalysis);
 }
@@ -121,6 +123,33 @@ void FStaticMeshCompilingManager::OnPostReachabilityAnalysis()
 	}
 }
 
+FName FStaticMeshCompilingManager::GetStaticAssetTypeName()
+{
+	return TEXT("UE-StaticMesh");
+}
+
+FName FStaticMeshCompilingManager::GetAssetTypeName() const
+{
+	return GetStaticAssetTypeName();
+}
+
+FTextFormat FStaticMeshCompilingManager::GetAssetNameFormat() const
+{
+	return LOCTEXT("StaticMeshNameFormat", "{0}|plural(one=Static Mesh,other=Static Meshes)");
+}
+
+TArrayView<FName> FStaticMeshCompilingManager::GetDependentTypeNames() const
+{
+	// Texture and shaders can affect materials which can affect Static Meshes once they are visible.
+	// Adding these dependencies can reduces the actual number of render state update we need to do in a frame
+	static FName DependentTypeNames[] = 
+	{ 
+		FTextureCompilingManager::GetStaticAssetTypeName(),
+		FShaderCompilingManager::GetStaticAssetTypeName()
+	};
+	return TArrayView<FName>(DependentTypeNames);
+}
+
 EQueuedWorkPriority FStaticMeshCompilingManager::GetBasePriority(UStaticMesh* InStaticMesh) const
 {
 	return EQueuedWorkPriority::Low;
@@ -299,6 +328,11 @@ int32 FStaticMeshCompilingManager::GetNumRemainingMeshes() const
 	return RegisteredStaticMesh.Num();
 }
 
+int32 FStaticMeshCompilingManager::GetNumRemainingAssets() const
+{
+	return GetNumRemainingMeshes();
+}
+
 void FStaticMeshCompilingManager::AddStaticMeshes(TArrayView<UStaticMesh* const> InStaticMeshes)
 {
 	TRACE_CPUPROFILER_EVENT_SCOPE(FStaticMeshCompilingManager::AddStaticMeshes)
diff --git a/Engine/Source/Runtime/Engine/Private/TextureCompiler.cpp b/Engine/Source/Runtime/Engine/Private/TextureCompiler.cpp
--- a/Engine/Source/Runtime/Engine/Private/TextureCompiler.cpp
+++ b/Engine/Source/Runtime/Engine/Private/TextureCompiler.cpp
@@ -77,10 +77,30 @@ namespace TextureCompilingManagerImpl
 }
 
 FTextureCompilingManager::FTextureCompilingManager()
-	: Notification(LOCTEXT("Textures", "Textures"))
+	: Notification(GetAssetNameFormat())
 {
 }
 
+FName FTextureCompilingManager::GetStaticAssetTypeName()
+{
+	return TEXT("UE-Texture");
+}
+
+FName FTextureCompilingManager::GetAssetTypeName() const
+{
+	return GetStaticAssetTypeName();
+}
+
+TArrayView<FName> FTextureCompilingManager::GetDependentTypeNames() const
+{
+	return TArrayView<FName>{ };
+}
+
+FTextFormat FTextureCompilingManager::GetAssetNameFormat() const
+{
+	return LOCTEXT("TextureNameFormat", "{0}|plural(one=Texture,other=Textures)");
+}
+
 EQueuedWorkPriority FTextureCompilingManager::GetBasePriority(UTexture* InTexture) const
 {
 	return TextureCompilingManagerImpl::GetBasePriority(InTexture);
@@ -193,6 +213,11 @@ int32 FTextureCompilingManager::GetNumRemainingTextures() const
 	return Num;
 }
 
+int32 FTextureCompilingManager::GetNumRemainingAssets() const
+{
+	return GetNumRemainingTextures();
+}
+
 void FTextureCompilingManager::AddTextures(TArrayView<UTexture* const> InTextures)
 {
 	TRACE_CPUPROFILER_EVENT_SCOPE(FTextureCompilingManager::AddTextures)
diff --git a/Engine/Source/Runtime/Engine/Private/UnrealEngine.cpp b/Engine/Source/Runtime/Engine/Private/UnrealEngine.cpp
--- a/Engine/Source/Runtime/Engine/Private/UnrealEngine.cpp
+++ b/Engine/Source/Runtime/Engine/Private/UnrealEngine.cpp
@@ -10454,31 +10454,18 @@ float DrawMapWarnings(UWorld* World, FViewport* Viewport, FCanvas* Canvas, UCanv
 
 	SmallTextItem.SetColor(FLinearColor::White);
 
-	if (GShaderCompilingManager && GShaderCompilingManager->IsCompiling())
+	for (IAssetCompilingManager* AssetCompilingManager : FAssetCompilingManager::Get().GetRegisteredManagers())
 	{
-		SmallTextItem.SetColor(FLinearColor::White);
-		SmallTextItem.Text = FText::Format(LOCTEXT("ShadersCompilingFmt", "Shaders Compiling ({0})"), GShaderCompilingManager->GetNumRemainingJobs());
-		Canvas->DrawItem(SmallTextItem, FVector2D(MessageX, MessageY));
-		MessageY += FontSizeY;
-	}
-
-#if WITH_EDITOR
-	if (FTextureCompilingManager::Get().GetNumRemainingTextures() > 0)
-	{
-		SmallTextItem.SetColor(FLinearColor::White);
-		SmallTextItem.Text = FText::Format(LOCTEXT("TexturesCompilingFmt", "Preparing Textures ({0})"), FTextureCompilingManager::Get().GetNumRemainingTextures());
-		Canvas->DrawItem(SmallTextItem, FVector2D(MessageX, MessageY));
-		MessageY += FontSizeY;
-	}
-
-	if (FStaticMeshCompilingManager::Get().GetNumRemainingMeshes() > 0)
-	{
-		SmallTextItem.SetColor(FLinearColor::White);
-		SmallTextItem.Text = FText::Format(LOCTEXT("StaticMeshCompilingFmt", "Preparing Static Meshes ({0})"), FStaticMeshCompilingManager::Get().GetNumRemainingMeshes());
-		Canvas->DrawItem(SmallTextItem, FVector2D(MessageX, MessageY));
-		MessageY += FontSizeY;
+		if (AssetCompilingManager->GetNumRemainingAssets() > 0)
+		{
+			// Always show asset plural form when displaying preparation status
+			FText AssetTypePlural = FText::Format(AssetCompilingManager->GetAssetNameFormat(), FText::AsNumber(100));
+			SmallTextItem.SetColor(FLinearColor::White);
+			SmallTextItem.Text = FText::Format(LOCTEXT("AssetCompilingFmt", "Preparing {0} ({1})"), AssetTypePlural, AssetCompilingManager->GetNumRemainingAssets());
+			Canvas->DrawItem(SmallTextItem, FVector2D(MessageX, MessageY));
+			MessageY += FontSizeY;
+		}
 	}
-#endif
 
 	if (World->bIsLevelStreamingFrozen)
 	{
diff --git a/Engine/Source/Runtime/Engine/Public/AssetCompilingManager.h b/Engine/Source/Runtime/Engine/Public/AssetCompilingManager.h
--- a/Engine/Source/Runtime/Engine/Public/AssetCompilingManager.h
+++ b/Engine/Source/Runtime/Engine/Public/AssetCompilingManager.h
@@ -26,7 +26,51 @@ struct FAssetCompileData
 	}
 };
 
-#if WITH_EDITOR
+struct IAssetCompilingManager
+{
+	/**
+	 * A unique name among all asset compiling manager to identify the type of asset this manager handles.
+	 */
+	ENGINE_API virtual FName GetAssetTypeName() const = 0;
+
+	/**
+	 * Returns an FTextFormat representing a localized singular/plural formatter for this resource name.
+	 * 
+	 * Note: Should be in a similar form as "{0}|plural(one=Singular Name,other=Plural Name)"
+	 */
+	ENGINE_API virtual FTextFormat GetAssetNameFormat() const = 0;
+
+	/**
+	 * Return other asset types that should preferably be handled before this one.
+	 */
+	ENGINE_API virtual TArrayView<FName> GetDependentTypeNames() const = 0;
+
+	/**
+	 * Returns the number of remaining compilations.
+	 */
+	ENGINE_API virtual int32 GetNumRemainingAssets() const = 0;
+
+	/**
+	 * Blocks until completion of all assets.
+	 */
+	ENGINE_API virtual void FinishAllCompilation() = 0;
+
+	/**
+	 * Cancel any pending work and blocks until it is safe to shut down.
+	 */
+	ENGINE_API virtual void Shutdown() = 0;
+
+protected:
+	friend class FAssetCompilingManager;
+
+	/** 
+	 * Called once per frame, fetches completed tasks and applies them to the scene. 
+	 */
+	ENGINE_API virtual void ProcessAsyncTasks(bool bLimitExecutionTime = false) = 0;
+
+
+	ENGINE_API virtual ~IAssetCompilingManager() {}
+};
 
 DECLARE_MULTICAST_DELEGATE_OneParam(FAssetPostCompileEvent, const TArray<FAssetCompileData>&);
 
@@ -35,8 +79,23 @@ class FAssetCompilingManager
 public:
 	ENGINE_API static FAssetCompilingManager& Get();
 
+	/**
+	 * Register an asset compiling manager.
+	 */
+	ENGINE_API bool RegisterManager(IAssetCompilingManager* InAssetCompilingManager);
+
+	/**
+	 * Unregister an asset compiling manager.
+	 */
+	ENGINE_API bool UnregisterManager(IAssetCompilingManager* InAssetCompilingManager);
+
+	/**
+	 * Register the list of registered managers.
+	 */
+	ENGINE_API TArrayView<IAssetCompilingManager* const> GetRegisteredManagers() const;
+
 	/** 
-	 * Returns the number of outstanding texture compilations.
+	 * Returns the number of remaining compilations.
 	 */
 	ENGINE_API int32 GetNumRemainingAssets() const;
 
@@ -45,11 +104,6 @@ public:
 	 */
 	ENGINE_API void FinishAllCompilation();
 
-	/**
-	 * Event called after an asset finishes compilation.
-	 */
-	FAssetPostCompileEvent& OnAssetPostCompileEvent() { return AssetPostCompileEvent; }
-
 	/**
 	 * Cancel any pending work and blocks until it is safe to shut down.
 	 */
@@ -60,17 +114,24 @@ public:
 	 */
 	ENGINE_API FQueuedThreadPool* GetThreadPool() const;
 
-	/** Called once per frame, fetches completed tasks and applies them to the scene. */
+	/** 
+	 * Called once per frame, fetches completed tasks and applies them to the scene. 
+	 */
 	ENGINE_API void ProcessAsyncTasks(bool bLimitExecutionTime = false);
 
+	/**
+	 * Event called after an asset finishes compilation.
+	 */
+	FAssetPostCompileEvent& OnAssetPostCompileEvent() { return AssetPostCompileEvent; }
 private:
-	friend class FStaticMeshCompilingManager;
-	friend class FTextureCompilingManager;
+	FAssetCompilingManager();
+	friend struct IAssetCompilingManager;
 
 	bool bHasShutdown = false;
 	
+	TArray<IAssetCompilingManager*> AssetCompilingManagers;
+	TArray<IAssetCompilingManager*> AssetCompilingManagersWithValidDependencies;
+
 	/** Event issued at the end of the compile process */
 	FAssetPostCompileEvent AssetPostCompileEvent;
 };
-
-#endif //#if WITH_EDITOR
\ No newline at end of file
diff --git a/Engine/Source/Runtime/Engine/Public/AsyncCompilationHelpers.h b/Engine/Source/Runtime/Engine/Public/AsyncCompilationHelpers.h
--- a/Engine/Source/Runtime/Engine/Public/AsyncCompilationHelpers.h
+++ b/Engine/Source/Runtime/Engine/Public/AsyncCompilationHelpers.h
@@ -17,8 +17,8 @@ class FQueuedThreadPoolWrapper;
 class FAsyncCompilationNotification
 {
 public:
-	FAsyncCompilationNotification(FText InAssetType)
-		: AssetType(InAssetType)
+	FAsyncCompilationNotification(FTextFormat InAssetNameFormat)
+		: AssetNameFormat(InAssetNameFormat)
 	{
 	}
 
@@ -27,7 +27,7 @@ private:
 	FProgressNotificationHandle NotificationHandle;
 	/** The number of jobs when the notification began */
 	int32 StartNumJobs;
-	FText AssetType;
+	FTextFormat AssetNameFormat;
 };
 
 namespace AsyncCompilationHelpers
diff --git a/Engine/Source/Runtime/Engine/Public/DistanceFieldAtlas.h b/Engine/Source/Runtime/Engine/Public/DistanceFieldAtlas.h
--- a/Engine/Source/Runtime/Engine/Public/DistanceFieldAtlas.h
+++ b/Engine/Source/Runtime/Engine/Public/DistanceFieldAtlas.h
@@ -14,6 +14,7 @@
 #include "RenderResource.h"
 #include "RenderingThread.h"
 #include "TextureLayout3d.h"
+#include "AssetCompilingManager.h"
 #include "Templates/UniquePtr.h"
 #include "DerivedMeshDataTaskUtils.h"
 #include "Async/AsyncWork.h"
@@ -22,6 +23,7 @@
 #include "MeshUtilities.h"
 #endif
 
+struct FAssetCompileData;
 class FDistanceFieldVolumeData;
 class UStaticMesh;
 class UTexture2D;
@@ -334,7 +336,7 @@ public:
 };
 
 /** Class that manages asynchronous building of mesh distance fields. */
-class FDistanceFieldAsyncQueue : public FGCObject
+class FDistanceFieldAsyncQueue : public FGCObject, IAssetCompilingManager
 {
 public:
 
@@ -358,7 +360,7 @@ public:
 	ENGINE_API void BlockUntilAllBuildsComplete();
 
 	/** Called once per frame, fetches completed tasks and applies them to the scene. */
-	ENGINE_API void ProcessAsyncTasks(bool bLimitExecutionTime = false);
+	ENGINE_API void ProcessAsyncTasks(bool bLimitExecutionTime = false) override;
 
 	/** Exposes UObject references used by the async build. */
 	ENGINE_API virtual void AddReferencedObjects(FReferenceCollector& Collector) override;
@@ -367,7 +369,7 @@ public:
 	ENGINE_API virtual FString GetReferencerName() const override;
 
 	/** Blocks until it is safe to shut down (worker threads are idle). */
-	ENGINE_API void Shutdown();
+	ENGINE_API void Shutdown() override;
 
 	int32 GetNumOutstandingTasks() const
 	{
@@ -375,7 +377,16 @@ public:
 		return ReferencedTasks.Num();
 	}
 
+	/** Get the name of the asset type this compiler handles */
+	ENGINE_API static FName GetStaticAssetTypeName();
+
 private:
+	FName GetAssetTypeName() const override;
+	FTextFormat GetAssetNameFormat() const override;
+	TArrayView<FName> GetDependentTypeNames() const override;
+	int32 GetNumRemainingAssets() const override;
+	void FinishAllCompilation() override;
+
 	friend FAsyncDistanceFieldTaskWorker;
 	void ProcessPendingTasks();
 
@@ -393,6 +404,9 @@ private:
 	/** Cancel or finish any background work for the given task. */
 	void CancelBackgroundTask(TArray<FAsyncDistanceFieldTask*> Tasks);
 
+	/** Get notified when static mesh finish compiling */
+	void OnAssetPostCompile(const TArray<FAssetCompileData>& CompiledAssets);
+
 	/** Game-thread managed list of tasks in the async system. */
 	TArray<FAsyncDistanceFieldTask*> ReferencedTasks;
 
diff --git a/Engine/Source/Runtime/Engine/Public/MeshCardRepresentation.h b/Engine/Source/Runtime/Engine/Public/MeshCardRepresentation.h
--- a/Engine/Source/Runtime/Engine/Public/MeshCardRepresentation.h
+++ b/Engine/Source/Runtime/Engine/Public/MeshCardRepresentation.h
@@ -11,6 +11,7 @@
 #include "ProfilingDebugging/ResourceSize.h"
 #include "Engine/EngineTypes.h"
 #include "UObject/GCObject.h"
+#include "AssetCompilingManager.h"
 #include "RenderResource.h"
 #include "RenderingThread.h"
 #include "Templates/UniquePtr.h"
@@ -199,7 +200,7 @@ public:
 };
 
 /** Class that manages asynchronous building of mesh distance fields. */
-class FCardRepresentationAsyncQueue : public FGCObject
+class FCardRepresentationAsyncQueue : public FGCObject, IAssetCompilingManager
 {
 public:
 
@@ -223,7 +224,7 @@ public:
 	ENGINE_API void BlockUntilAllBuildsComplete();
 
 	/** Called once per frame, fetches completed tasks and applies them to the scene. */
-	ENGINE_API void ProcessAsyncTasks(bool bLimitExecutionTime = false);
+	ENGINE_API void ProcessAsyncTasks(bool bLimitExecutionTime = false) override;
 
 	/** Exposes UObject references used by the async build. */
 	ENGINE_API void AddReferencedObjects(FReferenceCollector& Collector);
@@ -231,7 +232,7 @@ public:
 	ENGINE_API virtual FString GetReferencerName() const override;
 
 	/** Blocks until it is safe to shut down (worker threads are idle). */
-	ENGINE_API void Shutdown();
+	ENGINE_API void Shutdown() override;
 
 	int32 GetNumOutstandingTasks() const
 	{
@@ -239,8 +240,18 @@ public:
 		return ReferencedTasks.Num();
 	}
 
+	/** Get the name of the asset type this compiler handles */
+	ENGINE_API static FName GetStaticAssetTypeName();
+
 private:
 	friend FAsyncCardRepresentationTaskWorker;
+
+	ENGINE_API FName GetAssetTypeName() const override;
+	ENGINE_API FTextFormat GetAssetNameFormat() const override;
+	ENGINE_API TArrayView<FName> GetDependentTypeNames() const override;
+	ENGINE_API int32 GetNumRemainingAssets() const override;
+	ENGINE_API void FinishAllCompilation() override;
+	
 	void ProcessPendingTasks();
 
 	TUniquePtr<FQueuedThreadPool> ThreadPool;
diff --git a/Engine/Source/Runtime/Engine/Public/ShaderCompiler.h b/Engine/Source/Runtime/Engine/Public/ShaderCompiler.h
--- a/Engine/Source/Runtime/Engine/Public/ShaderCompiler.h
+++ b/Engine/Source/Runtime/Engine/Public/ShaderCompiler.h
@@ -21,6 +21,7 @@
 #include "GBufferInfo.h"
 #include "ShaderMaterial.h"
 #include "Misc/ScopeRWLock.h"
+#include "AssetCompilingManager.h"
 #include "Containers/HashTable.h"
 #include "Containers/List.h"
 #include "Hash/Blake3.h"
@@ -921,7 +922,7 @@ private:
  * Manager of asynchronous and parallel shader compilation.
  * This class contains an interface to enqueue and retreive asynchronous shader jobs, and manages a FShaderCompileThreadRunnable.
  */
-class FShaderCompilingManager
+class FShaderCompilingManager : IAssetCompilingManager
 {
 	friend class FShaderCompileThreadRunnableBase;
 	friend class FShaderCompileThreadRunnable;
@@ -932,7 +933,16 @@ class FShaderCompilingManager
 	friend class FShaderCompileDistributedThreadRunnable_Interface;
 	friend class FShaderCompileFASTBuildThreadRunnable;
 
+public:
+	/** Get the name of the asset type this compiler handles */
+	ENGINE_API static FName GetStaticAssetTypeName();
+
 private:
+	FName GetAssetTypeName() const override;
+	FTextFormat GetAssetNameFormat() const override;
+	TArrayView<FName> GetDependentTypeNames() const override;
+	int32 GetNumRemainingAssets() const override;
+	void ProcessAsyncTasks(bool bLimitExecutionTime = false) override;
 
 	//////////////////////////////////////////////////////
 	// Thread shared properties: These variables can only be read from or written to when a lock on CompileQueueSection is obtained, since they are used by both threads.
@@ -1049,6 +1059,7 @@ private:
 public:
 	
 	ENGINE_API FShaderCompilingManager();
+	ENGINE_API ~FShaderCompilingManager();
 
 	ENGINE_API int32 GetNumPendingJobs() const;
 	ENGINE_API int32 GetNumOutstandingJobs() const;
@@ -1171,12 +1182,12 @@ public:
 	 * Blocks until completion of all async shader compiling, and assigns shader maps to relevant materials.
 	 * This should be called before exit if the DDC needs to be made up to date. 
 	 */
-	ENGINE_API void FinishAllCompilation();
+	ENGINE_API void FinishAllCompilation() override;
 
 	/** 
 	 * Shutdown the shader compiler manager, this will shutdown immediately and not process any more shader compile requests. 
 	 */
-	ENGINE_API void Shutdown();
+	ENGINE_API void Shutdown() override;
 
 	/**
 	 * Prints stats related to shader compilation to the log.
diff --git a/Engine/Source/Runtime/Engine/Public/SkeletalMeshCompiler.h b/Engine/Source/Runtime/Engine/Public/SkeletalMeshCompiler.h
--- a/Engine/Source/Runtime/Engine/Public/SkeletalMeshCompiler.h
+++ b/Engine/Source/Runtime/Engine/Public/SkeletalMeshCompiler.h
@@ -6,6 +6,7 @@
 #include "Containers/Set.h"
 #include "Tickable.h"
 #include "Stats/Stats.h"
+#include "AssetCompilingManager.h"
 #include "AsyncCompilationHelpers.h"
 
 #if WITH_EDITOR
@@ -17,7 +18,7 @@ class FQueuedThreadPool;
 struct FAssetCompileContext;
 enum class EQueuedWorkPriority : uint8;
 
-class FSkeletalMeshCompilingManager
+class FSkeletalMeshCompilingManager : IAssetCompilingManager
 {
 public:
 	ENGINE_API static FSkeletalMeshCompilingManager& Get();
@@ -45,7 +46,7 @@ public:
 	/** 
 	 * Blocks until completion of all async skeletal mesh compilation.
 	 */
-	ENGINE_API void FinishAllCompilation();
+	ENGINE_API void FinishAllCompilation() override;
 
 	/**
 	 * Returns if asynchronous compilation is allowed for this skeletal mesh.
@@ -65,11 +66,16 @@ public:
 	/**
 	 * Cancel any pending work and blocks until it is safe to shut down.
 	 */
-	ENGINE_API void Shutdown();
+	ENGINE_API void Shutdown() override;
 
 private:
 	FSkeletalMeshCompilingManager();
-	void ProcessAsyncTasks(bool bLimitExecutionTime = false);
+
+	FName GetAssetTypeName() const override;
+	FTextFormat GetAssetNameFormat() const override;
+	TArrayView<FName> GetDependentTypeNames() const override;
+	int32 GetNumRemainingAssets() const override;
+	void ProcessAsyncTasks(bool bLimitExecutionTime = false) override;
 
 	friend class FAssetCompilingManager;
 	
diff --git a/Engine/Source/Runtime/Engine/Public/StaticMeshCompiler.h b/Engine/Source/Runtime/Engine/Public/StaticMeshCompiler.h
--- a/Engine/Source/Runtime/Engine/Public/StaticMeshCompiler.h
+++ b/Engine/Source/Runtime/Engine/Public/StaticMeshCompiler.h
@@ -4,6 +4,7 @@
 #include "CoreMinimal.h"
 #include "UObject/WeakObjectPtr.h"
 #include "Containers/Set.h"
+#include "AssetCompilingManager.h"
 #include "AsyncCompilationHelpers.h"
 
 #if WITH_EDITOR
@@ -15,7 +16,7 @@ class FQueuedThreadPool;
 struct FAssetCompileContext;
 enum class EQueuedWorkPriority : uint8;
 
-class FStaticMeshCompilingManager
+class FStaticMeshCompilingManager : IAssetCompilingManager
 {
 public:
 	ENGINE_API static FStaticMeshCompilingManager& Get();
@@ -43,7 +44,7 @@ public:
 	/** 
 	 * Blocks until completion of all async static mesh compilation.
 	 */
-	ENGINE_API void FinishAllCompilation();
+	ENGINE_API void FinishAllCompilation() override;
 
 	/**
 	 * Returns if asynchronous compilation is allowed for this static mesh.
@@ -63,15 +64,23 @@ public:
 	/**
 	 * Cancel any pending work and blocks until it is safe to shut down.
 	 */
-	ENGINE_API void Shutdown();
+	ENGINE_API void Shutdown() override;
+
+	/** Get the name of the asset type this compiler handles */
+	ENGINE_API static FName GetStaticAssetTypeName();
 
 private:
 	friend class FAssetCompilingManager;
 
 	FStaticMeshCompilingManager();
 
+	FName GetAssetTypeName() const override;
+	FTextFormat GetAssetNameFormat() const override;
+	TArrayView<FName> GetDependentTypeNames() const override;
+	int32 GetNumRemainingAssets() const override;
+
 	/** Called once per frame, fetches completed tasks and applies them to the scene. */
-	void ProcessAsyncTasks(bool bLimitExecutionTime = false);
+	void ProcessAsyncTasks(bool bLimitExecutionTime = false) override;
 
 	bool bHasShutdown = false;
 	TSet<TWeakObjectPtr<UStaticMesh>> RegisteredStaticMesh;
diff --git a/Engine/Source/Runtime/Engine/Public/TextureCompiler.h b/Engine/Source/Runtime/Engine/Public/TextureCompiler.h
--- a/Engine/Source/Runtime/Engine/Public/TextureCompiler.h
+++ b/Engine/Source/Runtime/Engine/Public/TextureCompiler.h
@@ -4,6 +4,7 @@
 #include "CoreMinimal.h"
 #include "UObject/WeakObjectPtr.h"
 #include "Containers/Set.h"
+#include "AssetCompilingManager.h"
 #include "AsyncCompilationHelpers.h"
 
 #if WITH_EDITOR
@@ -12,7 +13,7 @@ class UTexture;
 class FQueuedThreadPool;
 enum class EQueuedWorkPriority : uint8;
 
-class FTextureCompilingManager
+class FTextureCompilingManager : IAssetCompilingManager
 {
 public:
 	ENGINE_API static FTextureCompilingManager& Get();
@@ -40,7 +41,7 @@ public:
 	/** 
 	 * Blocks until completion of all async texture compilation.
 	 */
-	ENGINE_API void FinishAllCompilation();
+	ENGINE_API void FinishAllCompilation() override;
 
 	/**
 	 * Returns if asynchronous compilation is allowed for this texture.
@@ -65,15 +66,21 @@ public:
 	/**
 	 * Cancel any pending work and blocks until it is safe to shut down.
 	 */
-	ENGINE_API void Shutdown();
+	ENGINE_API void Shutdown() override;
 
+	/** Get the name of the asset type this compiler handles */
+	ENGINE_API static FName GetStaticAssetTypeName();
 
 private:
 	friend class FAssetCompilingManager;
 
 	FTextureCompilingManager();
-	
-	void ProcessAsyncTasks(bool bLimitExecutionTime = false);
+
+	FName GetAssetTypeName() const override;
+	FTextFormat GetAssetNameFormat() const override;
+	TArrayView<FName> GetDependentTypeNames() const override;
+	int32 GetNumRemainingAssets() const override;
+	void ProcessAsyncTasks(bool bLimitExecutionTime = false) override;
 	
 	void FinishCompilationsForGame();
 	void ProcessTextures(bool bLimitExecutionTime, int32 MaximumPriority = -1);
diff --git a/Engine/Source/Runtime/Launch/Private/LaunchEngineLoop.cpp b/Engine/Source/Runtime/Launch/Private/LaunchEngineLoop.cpp
--- a/Engine/Source/Runtime/Launch/Private/LaunchEngineLoop.cpp
+++ b/Engine/Source/Runtime/Launch/Private/LaunchEngineLoop.cpp
@@ -119,7 +119,6 @@
 	#include "Features/IModularFeatures.h"
 	#include "GameFramework/WorldSettings.h"
 	#include "SystemSettings.h"
-	#include "ObjectCacheContext.h"
 	#include "EngineStats.h"
 	#include "EngineGlobals.h"
 	#include "AudioThread.h"
@@ -4281,9 +4280,7 @@ void FEngineLoop::Exit()
 	FVisualLogger::Get().Shutdown();
 #endif
 
-#if WITH_EDITOR
 	FAssetCompilingManager::Get().Shutdown();
-#endif
 
 #if WITH_ENGINE
 	// shut down messaging
@@ -4298,14 +4295,12 @@ void FEngineLoop::Exit()
 
 	if (GDistanceFieldAsyncQueue)
 	{
-		GDistanceFieldAsyncQueue->Shutdown();
 		delete GDistanceFieldAsyncQueue;
 		GDistanceFieldAsyncQueue = nullptr;
 	}
 
 	if (GCardRepresentationAsyncQueue)
 	{
-		GCardRepresentationAsyncQueue->Shutdown();
 		delete GCardRepresentationAsyncQueue;
 		GCardRepresentationAsyncQueue = nullptr;
 	}
@@ -5058,32 +5053,8 @@ void FEngineLoop::Tick()
 			}
 		}
 
-		{
-			// Reuse ObjectIterator Caching and Reverse lookups for the duration of all asset updates
-			FObjectCacheContextScope ObjectCacheScope;
-
-			if (GShaderCompilingManager)
-			{
-				// Process any asynchronous shader compile results that are ready, limit execution time
-				QUICK_SCOPE_CYCLE_COUNTER(STAT_FEngineLoop_Tick_GShaderCompilingManager);
-				GShaderCompilingManager->ProcessAsyncResults(true, false);
-			}
+		FAssetCompilingManager::Get().ProcessAsyncTasks(true);
 
-			if (GDistanceFieldAsyncQueue)
-			{
-				QUICK_SCOPE_CYCLE_COUNTER(STAT_FEngineLoop_Tick_GDistanceFieldAsyncQueue);
-				GDistanceFieldAsyncQueue->ProcessAsyncTasks();
-			}
-
-			if (GCardRepresentationAsyncQueue)
-			{
-				QUICK_SCOPE_CYCLE_COUNTER(STAT_FEngineLoop_Tick_GCardRepresentationAsyncQueue);
-				GCardRepresentationAsyncQueue->ProcessAsyncTasks();
-			}
-#if WITH_EDITOR
-			FAssetCompilingManager::Get().ProcessAsyncTasks(true);
-#endif
-		}
 		// Tick the platform and input portion of Slate application, we need to do this before we run things
 		// concurrent with networking.
 		if (FSlateApplication::IsInitialized() && !bIdleMode)
@@ -5968,8 +5939,6 @@ void FEngineLoop::AppPreExit( )
 #if WITH_ENGINE
 	if ( GShaderCompilingManager )
 	{
-		GShaderCompilingManager->Shutdown();
-
 		delete GShaderCompilingManager;
 		GShaderCompilingManager = nullptr;
 	}
