Fixed strata normal to take into account the mateiral two sided normal sign.
Added support for masked strata materials.

diff --git a/Engine/Source/Runtime/Engine/Private/Materials/HLSLMaterialTranslator.cpp b/Engine/Source/Runtime/Engine/Private/Materials/HLSLMaterialTranslator.cpp
--- a/Engine/Source/Runtime/Engine/Private/Materials/HLSLMaterialTranslator.cpp
+++ b/Engine/Source/Runtime/Engine/Private/Materials/HLSLMaterialTranslator.cpp
@@ -1017,7 +1017,7 @@ bool FHLSLMaterialTranslator::Translate()
 
 		if (Material->IsStrataMaterial())
 		{
-			if (Domain!= MD_Volume && !(BlendMode == BLEND_Translucent || BlendMode == BLEND_Opaque))
+			if (Domain!= MD_Volume && !(BlendMode == BLEND_Translucent || BlendMode == BLEND_Opaque || BlendMode == BLEND_Masked))
 			{
 				Errorf(TEXT("Strata materials must be an opaque or translucent surface, or a volume."));
 			}
@@ -1811,6 +1811,10 @@ void FHLSLMaterialTranslator::GetSharedInputsMaterialCode(FString& PixelMembersD
 			for (TMultiMap<uint64, FStrataSharedNormalInfo>::TIterator It(CodeChunkToStrataSharedNormal); It; ++It)
 			{
 				PixelInputInitializerValues += FString::Printf(TEXT("\tParameters.SharedNormals.Normals[%u] = %s;\n"), It->Value.SharedNormalIndex, *It->Value.SharedNormalCode);
+
+				PixelInputInitializerValues += FString::Printf(TEXT("\t#if MATERIAL_TANGENTSPACENORMAL\n"));
+				PixelInputInitializerValues += FString::Printf(TEXT("\tParameters.SharedNormals.Normals[%u] *= Parameters.TwoSidedSign;\n"), It->Value.SharedNormalIndex);
+				PixelInputInitializerValues += FString::Printf(TEXT("\t#endif\n"));
 			}
 			for (TMultiMap<uint64, FStrataSharedNormalInfo>::TIterator It(CodeChunkToStrataSharedNormal); It; ++It)
 			{
@@ -1822,6 +1826,10 @@ void FHLSLMaterialTranslator::GetSharedInputsMaterialCode(FString& PixelMembersD
 				{
 					PixelInputInitializerValues += FString::Printf(TEXT("\tParameters.SharedNormals.Tangents[%u] = float3(1,0,0);\n"), It->Value.SharedNormalIndex);
 				}
+
+				PixelInputInitializerValues += FString::Printf(TEXT("\t#if MATERIAL_TANGENTSPACENORMAL\n"));
+				PixelInputInitializerValues += FString::Printf(TEXT("\tParameters.SharedNormals.Tangents[%u] *= Parameters.TwoSidedSign;\n"), It->Value.SharedNormalIndex);
+				PixelInputInitializerValues += FString::Printf(TEXT("\t#endif\n"));
 			}
 			PixelInputInitializerValues += FString::Printf(TEXT("\tParameters.SharedNormals.NormalCount = %u;\n"), CodeChunkToStrataSharedNormal.Num());
 		}
