OFPA Child Actor Fixes:
- IsMainPackageActor checks for ParentComponent.IsExplicitelyNull() as the only use case for non main package actors are ChildActors.
- Fix Saving of a new level containing OFPA ChildActors crashing

diff --git a/Engine/Source/Editor/UnrealEd/Private/FileHelpers.cpp b/Engine/Source/Editor/UnrealEd/Private/FileHelpers.cpp
--- a/Engine/Source/Editor/UnrealEd/Private/FileHelpers.cpp
+++ b/Engine/Source/Editor/UnrealEd/Private/FileHelpers.cpp
@@ -3533,7 +3533,13 @@ bool FEditorFileUtils::SaveMapDataPackages(UWorld* WorldToSave, bool bCheckDirty
 
 		if (bSaveExternal)
 		{
-			PackagesToSave.Append(WorldPackage->GetExternalPackages());
+			for (UPackage* ExternalPackage : WorldPackage->GetExternalPackages())
+			{
+				if (!FPackageName::IsTempPackage(ExternalPackage->GetName()))
+				{
+					PackagesToSave.Add(ExternalPackage);
+				}
+			}
 		}
 	}
 			
diff --git a/Engine/Source/Runtime/Engine/Private/Actor.cpp b/Engine/Source/Runtime/Engine/Private/Actor.cpp
--- a/Engine/Source/Runtime/Engine/Private/Actor.cpp
+++ b/Engine/Source/Runtime/Engine/Private/Actor.cpp
@@ -1118,7 +1118,7 @@ float AActor::GetActorTickInterval() const
 
 bool AActor::IsMainPackageActor() const
 {
-	return IsPackageExternal() && (GetExternalPackage()->GetFName().ToString() == ULevel::GetActorPackageName(GetLevel()->GetPackage(), GetPathName()));
+	return IsPackageExternal() && ParentComponent.IsExplicitlyNull();
 }
 
 AActor* AActor::FindActorInPackage(UPackage* InPackage)
diff --git a/Engine/Source/Runtime/Engine/Private/World.cpp b/Engine/Source/Runtime/Engine/Private/World.cpp
--- a/Engine/Source/Runtime/Engine/Private/World.cpp
+++ b/Engine/Source/Runtime/Engine/Private/World.cpp
@@ -644,18 +644,21 @@ bool UWorld::Rename(const TCHAR* InName, UObject* NewOuter, ERenameFlags Flags)
 		}
 	}
 
+
 	if (!bTestRename)
 	{
 		// We also need to rename any external actor packages we have since the search for them is based on the world name
-		for (AActor* Actor : PersistentLevel->Actors)
+		// Make a Copy to iterate as this could modify PersistentLevel->Actors
+		TArray<AActor*> CopyActors;
+		CopyActors.Reserve(PersistentLevel->Actors.Num());
+		Algo::CopyIf(PersistentLevel->Actors, CopyActors, [](AActor* InActor) { return InActor && InActor->IsMainPackageActor(); });
+				
+		// Instead of just renaming the package, re-embed and re-externalize the actor
+		// this will leave dirty empty actor packages being which will be cleaned up though SaveAll, although SaveCurrentLevel won't pick them up
+		for (AActor* Actor : CopyActors)
 		{
-			if (Actor && Actor->IsPackageExternal())
-			{
-				// Instead of just renaming the package, re-embed and re-externalize the actor
-				// this will leave dirty empty actor packages being which will be cleaned up though SaveAll, although SaveCurrentLevel won't pick them up
-				Actor->SetPackageExternal(false);
-				Actor->SetPackageExternal(true);
-			}
+			Actor->SetPackageExternal(false);
+			Actor->SetPackageExternal(true);
 		}
 	}
 
