Fix buffer memory calculation mismatches buffer allocation

diff --git a/Engine/Source/Runtime/Engine/Public/GPUSkinCache.h b/Engine/Source/Runtime/Engine/Public/GPUSkinCache.h
--- a/Engine/Source/Runtime/Engine/Public/GPUSkinCache.h
+++ b/Engine/Source/Runtime/Engine/Public/GPUSkinCache.h
@@ -221,14 +221,14 @@ public:
 		{
 			for (int32 Index = 0; Index < NUM_BUFFERS; ++Index)
 			{
-				RWBuffers[Index].Initialize(TEXT("SkinCacheVertices"), 4, NumVertices * 3, PF_R32_FLOAT, BUF_Static);
+				RWBuffers[Index].Initialize(TEXT("SkinCacheVertices"), PosBufferBytesPerElement, NumVertices * 3, PF_R32_FLOAT, BUF_Static);
 			}
 			if (WithTangents)
 			{
-				Tangents.Initialize(TEXT("SkinCacheTangents"), 8, NumVertices * 2, PF_R16G16B16A16_SNORM, BUF_Static);
+				Tangents.Initialize(TEXT("SkinCacheTangents"), TangentBufferBytesPerElement, NumVertices * 2, PF_R16G16B16A16_SNORM, BUF_Static);
 				if (FGPUSkinCache::UseIntermediateTangents())
 				{
-					IntermediateTangents.Initialize(TEXT("SkinCacheIntermediateTangents"), 8, NumVertices * 2, PF_R16G16B16A16_SNORM, BUF_Static);
+					IntermediateTangents.Initialize(TEXT("SkinCacheIntermediateTangents"), TangentBufferBytesPerElement, NumVertices * 2, PF_R16G16B16A16_SNORM, BUF_Static);
 				}
 			}
 			if (NumTriangles > 0)
@@ -258,12 +258,12 @@ public:
 
 		static uint64 CalculateRequiredMemory(uint32 InNumVertices, bool InWithTangents, uint32 InNumTriangles)
 		{
-			uint64 PositionBufferSize = 4 * 3 * InNumVertices * NUM_BUFFERS;
-			uint64 TangentBufferSize = InWithTangents ? 2 * 4 * InNumVertices : 0;
+			uint64 PositionBufferSize = PosBufferBytesPerElement * InNumVertices * 3 * NUM_BUFFERS;
+			uint64 TangentBufferSize = InWithTangents ? TangentBufferBytesPerElement * InNumVertices * 2 : 0;
 			uint64 IntermediateTangentBufferSize = 0;
 			if (FGPUSkinCache::UseIntermediateTangents())
 			{
-				IntermediateTangentBufferSize = InWithTangents ? 2 * 4 * InNumVertices : 0;
+				IntermediateTangentBufferSize = InWithTangents ? TangentBufferBytesPerElement * InNumVertices * 2 : 0;
 			}
 			uint64 AccumulatedTangentBufferSize = InNumTriangles * 3 * FGPUSkinCache::IntermediateAccumBufferNumInts * sizeof(int32);
 			return TangentBufferSize + IntermediateTangentBufferSize + PositionBufferSize + AccumulatedTangentBufferSize;
@@ -302,6 +302,9 @@ public:
 		const uint32 NumVertices;
 		const bool WithTangents;
 		const uint32 NumTriangles;
+
+		static const uint32 PosBufferBytesPerElement = 4;
+		static const uint32 TangentBufferBytesPerElement = 8;
 	};
 
 	struct FRWBufferTracker
