Rename buffer for readability

diff --git a/Engine/Source/Runtime/Engine/Private/GPUSkinCache.cpp b/Engine/Source/Runtime/Engine/Private/GPUSkinCache.cpp
--- a/Engine/Source/Runtime/Engine/Private/GPUSkinCache.cpp
+++ b/Engine/Source/Runtime/Engine/Private/GPUSkinCache.cpp
@@ -2138,7 +2138,7 @@ void FGPUSkinCache::FRWBuffersAllocation::RemoveAllFromTransitionArray(TSet<FRHI
 {
 	for (uint32 i = 0; i < NUM_BUFFERS; i++)
 	{
-		FRWBuffer& RWBuffer = RWBuffers[i];
+		FRWBuffer& RWBuffer = PositionBuffers[i];
 		if (RWBuffer.UAV.IsValid())
 		{
 			InBuffersToTransition.Remove(RWBuffer.UAV);
diff --git a/Engine/Source/Runtime/Engine/Public/GPUSkinCache.h b/Engine/Source/Runtime/Engine/Public/GPUSkinCache.h
--- a/Engine/Source/Runtime/Engine/Public/GPUSkinCache.h
+++ b/Engine/Source/Runtime/Engine/Public/GPUSkinCache.h
@@ -221,7 +221,7 @@ public:
 		{
 			for (int32 Index = 0; Index < NUM_BUFFERS; ++Index)
 			{
-				RWBuffers[Index].Initialize(TEXT("SkinCacheVertices"), PosBufferBytesPerElement, NumVertices * 3, PF_R32_FLOAT, BUF_Static);
+				PositionBuffers[Index].Initialize(TEXT("SkinCachePositions"), PosBufferBytesPerElement, NumVertices * 3, PF_R32_FLOAT, BUF_Static);
 			}
 			if (WithTangents)
 			{
@@ -243,7 +243,7 @@ public:
 		{
 			for (int32 Index = 0; Index < NUM_BUFFERS; ++Index)
 			{
-				RWBuffers[Index].Release();
+				PositionBuffers[Index].Release();
 			}
 			if (WithTangents)
 			{
@@ -293,7 +293,7 @@ public:
 
 	private:
 		// Output of the GPU skinning (ie Pos, Normals)
-		FRWBuffer RWBuffers[NUM_BUFFERS];
+		FRWBuffer PositionBuffers[NUM_BUFFERS];
 
 		FRWBuffer Tangents;
 		FRWBuffer IntermediateTangents;
@@ -337,7 +337,7 @@ public:
 			{
 				if (Revisions[Index] == Revision && BoneBuffers[Index] == &BoneBuffer)
 				{
-					return &Allocation->RWBuffers[Index];
+					return &Allocation->PositionBuffers[Index];
 				}
 			}
 
