Fix a bug in WasAnimNotifyX source state functions

diff --git a/Engine/Source/Runtime/Engine/Private/Animation/AnimInstanceProxy.cpp b/Engine/Source/Runtime/Engine/Private/Animation/AnimInstanceProxy.cpp
--- a/Engine/Source/Runtime/Engine/Private/Animation/AnimInstanceProxy.cpp
+++ b/Engine/Source/Runtime/Engine/Private/Animation/AnimInstanceProxy.cpp
@@ -2297,7 +2297,7 @@ bool FAnimInstanceProxy::WasAnimNotifyStateActiveInSourceState(int32 MachineInde
 		const FAnimNotifyEvent* NotifyEvent = Ref.GetNotify();
 		if (NotifyEvent && NotifyEvent->NotifyStateClass && NotifyEvent->NotifyStateClass->IsA(AnimNotifyStateType))
 		{
-			return UAnimNotifyStateMachineInspectionLibrary::IsStateInStateMachineInEventContext(Ref, StateIndex, MachineIndex);
+			return UAnimNotifyStateMachineInspectionLibrary::IsStateInStateMachineInEventContext(Ref, MachineIndex, StateIndex);
 		}
 	}
 	return false;
@@ -2310,7 +2310,7 @@ bool FAnimInstanceProxy::WasAnimNotifyTriggeredInSourceState(int32 MachineIndex,
 		const FAnimNotifyEvent* NotifyEvent = Ref.GetNotify();
 		if (NotifyEvent && NotifyEvent->Notify && NotifyEvent->Notify->IsA(AnimNotifyType))
 		{
-			return UAnimNotifyStateMachineInspectionLibrary::IsStateInStateMachineInEventContext(Ref, StateIndex, MachineIndex);
+			return UAnimNotifyStateMachineInspectionLibrary::IsStateInStateMachineInEventContext(Ref, MachineIndex, StateIndex);
 		}
 	}
 	return false;
@@ -2334,7 +2334,7 @@ bool FAnimInstanceProxy::WasAnimNotifyNameTriggeredInSourceState(int32 MachineIn
 			}
 			if(LookupName == NotifyName)
 			{
-				return UAnimNotifyStateMachineInspectionLibrary::IsStateInStateMachineInEventContext(Ref, StateIndex, MachineIndex);
+				return UAnimNotifyStateMachineInspectionLibrary::IsStateInStateMachineInEventContext(Ref, MachineIndex, StateIndex);
 			}
 		}
 	}
