Deprecated StaticMesh IsBuiltAtRuntime property.
Runtime-built static meshes now have a flag determining whether they were built with the fast path, or the full editor-only path. They are then built and cached in the DDC as per any static mesh.
The default behaviour when calling BuildFromMeshDescriptions from C++ is to use the full static mesh build in editor builds, and the fast build in non-editor builds.
The default behaviour when calling BuildFromStaticMeshDescriptions (normally from BP) is to use the fast build as its general use case is to generate transient static meshes at runtime.
This fixes issues where meshes are continually rebuilt and assigned new keys on every load.

diff --git a/Engine/Plugins/Experimental/Enterprise/DatasmithRuntime/Source/Private/MeshImporter.cpp b/Engine/Plugins/Experimental/Enterprise/DatasmithRuntime/Source/Private/MeshImporter.cpp
--- a/Engine/Plugins/Experimental/Enterprise/DatasmithRuntime/Source/Private/MeshImporter.cpp
+++ b/Engine/Plugins/Experimental/Enterprise/DatasmithRuntime/Source/Private/MeshImporter.cpp
@@ -559,6 +559,7 @@ namespace DatasmithRuntime
 			Params.bBuildSimpleCollision = bEnableCollision;
 			// Do not commit since we only need the render data and commit is slow
 			Params.bCommitMeshDescription = false;
+			Params.bFastBuild = true;
 			StaticMesh->BuildFromMeshDescriptions(MeshDescriptionPointers, Params);
 		}
 
diff --git a/Engine/Plugins/Experimental/Text3D/Source/Text3D/Private/Glyph.cpp b/Engine/Plugins/Experimental/Text3D/Source/Text3D/Private/Glyph.cpp
--- a/Engine/Plugins/Experimental/Text3D/Source/Text3D/Private/Glyph.cpp
+++ b/Engine/Plugins/Experimental/Text3D/Source/Text3D/Private/Glyph.cpp
@@ -29,7 +29,11 @@ void FText3DGlyph::Build(UStaticMesh* StaticMesh, UMaterial* DefaultMaterial)
 
 	TArray<const FMeshDescription*> MeshDescriptions;
 	MeshDescriptions.Add(&MeshDescription);
-	StaticMesh->BuildFromMeshDescriptions(MeshDescriptions);
+
+	UStaticMesh::FBuildMeshDescriptionsParams Params;
+	Params.bCommitMeshDescription = true;
+	Params.bFastBuild = true;
+	StaticMesh->BuildFromMeshDescriptions(MeshDescriptions, Params);
 }
 
 FMeshDescription& FText3DGlyph::GetMeshDescription()
diff --git a/Engine/Source/Runtime/Engine/Classes/Engine/StaticMesh.h b/Engine/Source/Runtime/Engine/Classes/Engine/StaticMesh.h
--- a/Engine/Source/Runtime/Engine/Classes/Engine/StaticMesh.h
+++ b/Engine/Source/Runtime/Engine/Classes/Engine/StaticMesh.h
@@ -969,30 +969,25 @@ public:
 	UPROPERTY(EditAnywhere, Category = RayTracing)
 	uint8 bSupportRayTracing : 1;
 
+	UPROPERTY()
+	uint8 bDoFastBuild : 1;
+
 private:
 
-	/**
-	 * If true, StaticMesh has been built at runtime
-	 */
-	UE_DEPRECATED(5.00, "This must be protected for async build, always use the accessors even internally.")
 	UPROPERTY()
-	uint8 bIsBuiltAtRuntime : 1;
+	uint8 bIsBuiltAtRuntime_DEPRECATED : 1;
+
 public:
 
+	UE_DEPRECATED(5.00, "IsBuiltAtRuntime() is no longer used.")
 	bool IsBuiltAtRuntime() const
 	{
-		WaitUntilAsyncPropertyReleased(EStaticMeshAsyncProperties::IsBuiltAtRuntime);
-		PRAGMA_DISABLE_DEPRECATION_WARNINGS
-		return bIsBuiltAtRuntime;
-		PRAGMA_ENABLE_DEPRECATION_WARNINGS
+		return false;
 	}
 
+	UE_DEPRECATED(5.00, "SetIsBuiltAtRuntime() is no longer used.")
 	void SetIsBuiltAtRuntime(bool InIsBuiltAtRuntime)
 	{
-		WaitUntilAsyncPropertyReleased(EStaticMeshAsyncProperties::IsBuiltAtRuntime);
-		PRAGMA_DISABLE_DEPRECATION_WARNINGS
-		bIsBuiltAtRuntime = InIsBuiltAtRuntime;
-		PRAGMA_ENABLE_DEPRECATION_WARNINGS
 	}
 protected:
 	/** Tracks whether InitResources has been called, and rendering resources are initialized. */
@@ -1333,7 +1328,7 @@ public:
 
 	/** Builds static mesh LODs from the array of StaticMeshDescriptions passed in */
 	UFUNCTION(BlueprintCallable, Category="StaticMesh")
-	ENGINE_API void BuildFromStaticMeshDescriptions(const TArray<UStaticMeshDescription*>& StaticMeshDescriptions, bool bBuildSimpleCollision = false);
+	ENGINE_API void BuildFromStaticMeshDescriptions(const TArray<UStaticMeshDescription*>& StaticMeshDescriptions, bool bBuildSimpleCollision = false, bool bFastBuild = true);
 
 	/** Return a new StaticMeshDescription referencing the MeshDescription of the given LOD */
 	UFUNCTION(BlueprintCallable, Category="StaticMesh")
@@ -1347,6 +1342,7 @@ public:
 			, bUseHashAsGuid(false)
 			, bBuildSimpleCollision(false)
 			, bCommitMeshDescription(true)
+			, bFastBuild(false)
 		{}
 
 		/**
@@ -1370,6 +1366,12 @@ public:
 		 * Commits the MeshDescription as part of the building process. Set to true by default.
 		 */
 		bool bCommitMeshDescription;
+
+		/**
+		 * Specifies that the mesh will be built by the fast path (mandatory in non-editor builds).
+		 * Set to false by default.
+		 */
+		bool bFastBuild;
 	};
 
 	/**
diff --git a/Engine/Source/Runtime/Engine/Private/StaticMesh.cpp b/Engine/Source/Runtime/Engine/Private/StaticMesh.cpp
--- a/Engine/Source/Runtime/Engine/Private/StaticMesh.cpp
+++ b/Engine/Source/Runtime/Engine/Private/StaticMesh.cpp
@@ -2726,18 +2726,32 @@ void FStaticMeshRenderData::Cache(const ITargetPlatform* TargetPlatform, UStatic
 
 			checkf(Owner->IsMeshDescriptionValid(0), TEXT("Bad MeshDescription on %s"), *GetPathNameSafe(Owner));
 
-			IMeshBuilderModule& MeshBuilderModule = IMeshBuilderModule::GetForPlatform(TargetPlatform);
-			if (!MeshBuilderModule.BuildMesh(*this, Owner, LODGroup))
+			if (Owner->bDoFastBuild)
 			{
-				UE_LOG(LogStaticMesh, Error, TEXT("Failed to build static mesh. See previous line(s) for details."));
-				return;
+				// If the mesh is built via the FastBuild path, just build it each time here instead of using the DDC as a cache.
+				const int32 NumSourceModels = Owner->GetNumSourceModels();
+				AllocateLODResources(NumSourceModels);
+				for (int32 LodIndex = 0; LodIndex < NumSourceModels; LodIndex++)
+				{
+					Owner->BuildFromMeshDescription(*Owner->GetMeshDescription(LodIndex), LODResources[LodIndex]);
+				}
 			}
-
-			ComputeUVDensities();
-			if(Owner->bSupportUniformlyDistributedSampling)
+			else
 			{
-				BuildAreaWeighedSamplingData();
+				IMeshBuilderModule& MeshBuilderModule = IMeshBuilderModule::GetForPlatform(TargetPlatform);
+				if (!MeshBuilderModule.BuildMesh(*this, Owner, LODGroup))
+				{
+					UE_LOG(LogStaticMesh, Error, TEXT("Failed to build static mesh. See previous line(s) for details."));
+					return;
+				}
+
+				ComputeUVDensities();
+				if(Owner->bSupportUniformlyDistributedSampling)
+				{
+					BuildAreaWeighedSamplingData();
+				}
 			}
+			
 			bLODsShareStaticLighting = Owner->CanLODsShareStaticLighting();
 			FMemoryWriter Ar(DerivedData, /*bIsPersistent=*/ true);
 			Serialize(Ar, Owner, /*bCooked=*/ false);
@@ -2890,11 +2904,11 @@ UStaticMesh::UStaticMesh(const FObjectInitializer& ObjectInitializer)
 	SetLightMapResolution(4);
 	SetMinLOD(0);
 
+	bDoFastBuild = false;
 	bSupportUniformlyDistributedSampling = false;
 
 	bSupportRayTracing = true;
 
-	SetIsBuiltAtRuntime(false);
 	bRenderingResourcesInitialized = false;
 #if WITH_EDITOR
 	BuildCacheAutomationTestGuid.Invalidate();
@@ -5082,8 +5096,7 @@ void UStaticMesh::PostLoad()
 	BeginPostLoadInternal(Context);
 
 #if WITH_EDITOR
-	// IsBuiltAtRuntime must stay synchronous for now
-	if (!IsBuiltAtRuntime() && FStaticMeshCompilingManager::Get().IsAsyncCompilationAllowed(this))
+	if (FStaticMeshCompilingManager::Get().IsAsyncCompilationAllowed(this))
 	{
 		FModuleManager::Get().LoadModuleChecked<IMeshUtilities>(TEXT("MeshUtilities"));
  
@@ -5277,26 +5290,8 @@ void UStaticMesh::ExecutePostLoadInternal(FStaticMeshPostLoadContext& Context)
 
 	if (!Context.bIsCookedForEditor)
 	{
-		if (IsBuiltAtRuntime())
-		{
-			// If built at runtime, but an editor build, we cache the mesh descriptions so that they can be rebuilt within the editor if necessary.
-			// This is done through the fast build path for consistency
-			TArray<const FMeshDescription*> MeshDescriptions;
-			const int32 NumSourceModels = GetNumSourceModels();
-			MeshDescriptions.Reserve(NumSourceModels);
-			for (int32 SourceModelIndex = 0; SourceModelIndex < NumSourceModels; SourceModelIndex++)
-			{
-				MeshDescriptions.Add(GetMeshDescription(SourceModelIndex));
-			}
-			BuildFromMeshDescriptions(MeshDescriptions);
-		}
-		else
-		{
-			// This, among many other things, will build a MeshDescription from the legacy RawMesh if one has not already been serialized,
-			// or, failing that, if there is not already one in the DDC. This will remain cached until the end of PostLoad(), upon which it
-			// is then released, and can be reloaded on demand.
-			CacheDerivedData();
-		}
+		// Generate and cache render data
+		CacheDerivedData();
 	}
 
 	GetBodySetup()->CreatePhysicsMeshes();
@@ -5770,7 +5765,7 @@ UStaticMeshDescription* UStaticMesh::GetStaticMeshDescription(int32 LODIndex)
 }
 
 
-void UStaticMesh::BuildFromStaticMeshDescriptions(const TArray<UStaticMeshDescription*>& StaticMeshDescriptions, bool bBuildSimpleCollision)
+void UStaticMesh::BuildFromStaticMeshDescriptions(const TArray<UStaticMeshDescription*>& StaticMeshDescriptions, bool bBuildSimpleCollision, bool bFastBuild)
 {
 	TArray<const FMeshDescription*> MeshDescriptions;
 	MeshDescriptions.Reserve(StaticMeshDescriptions.Num());
@@ -5782,6 +5777,7 @@ void UStaticMesh::BuildFromStaticMeshDescriptions(const TArray<UStaticMeshDescri
 
 	FBuildMeshDescriptionsParams Params;
 	Params.bBuildSimpleCollision = bBuildSimpleCollision;
+	Params.bFastBuild = bFastBuild;
 	BuildFromMeshDescriptions(MeshDescriptions, Params);
 }
 
@@ -5791,9 +5787,19 @@ bool UStaticMesh::BuildFromMeshDescriptions(const TArray<const FMeshDescription*
 	TRACE_CPUPROFILER_EVENT_SCOPE(UStaticMesh::BuildFromMeshDescriptions);
 
 	// Set up
-	SetIsBuiltAtRuntime(true);
 	NeverStream = true;
 
+#if !WITH_EDITOR
+	// In non-Editor builds, we can only perform fast mesh builds
+	check(Params.bFastBuild);
+#endif
+	if (Params.bFastBuild)
+	{
+		bDoFastBuild = true;
+	}
+
+	check(Params.bCommitMeshDescription || Params.bFastBuild);
+
 	TOptional<FStaticMeshComponentRecreateRenderStateContext> RecreateRenderStateContext;
 	
 	bool bNewMesh = true;
@@ -5808,7 +5814,19 @@ bool UStaticMesh::BuildFromMeshDescriptions(const TArray<const FMeshDescription*
 	}
 
 	#if WITH_EDITOR
+	if (Params.bCommitMeshDescription)
+	{
+		FCommitMeshDescriptionParams CommitParams;
+		CommitParams.bMarkPackageDirty = Params.bMarkPackageDirty;
+		CommitParams.bUseHashAsGuid = Params.bUseHashAsGuid;
+
 		SetNumSourceModels(MeshDescriptions.Num());
+		for (int32 LODIndex = 0; LODIndex < MeshDescriptions.Num(); LODIndex++)
+		{
+			CreateMeshDescription(LODIndex, *MeshDescriptions[LODIndex]);
+			CommitMeshDescription(LODIndex, CommitParams);
+		}
+	}
 	#endif
 
 	SetRenderData(MakeUnique<FStaticMeshRenderData>());
@@ -5816,85 +5834,85 @@ bool UStaticMesh::BuildFromMeshDescriptions(const TArray<const FMeshDescription*
 
 	// Build render data from each mesh description
 
-	int32 LODIndex = 0;
-	for (const FMeshDescription* MeshDescriptionPtr : MeshDescriptions)
-	{
 #if WITH_EDITOR
-		// Editor builds cache the mesh description so that it can be preserved during map reloads etc
-		if (Params.bCommitMeshDescription)
-		{
-			CreateMeshDescription(LODIndex, *MeshDescriptionPtr);
-			FCommitMeshDescriptionParams CommitParams;
-			CommitParams.bMarkPackageDirty = Params.bMarkPackageDirty;
-			CommitParams.bUseHashAsGuid = Params.bUseHashAsGuid;
-			CommitMeshDescription(LODIndex, CommitParams);
-		}
+	if (Params.bFastBuild)
 #endif
-		check(MeshDescriptionPtr != nullptr);
-		FStaticMeshLODResources& LODResources = GetRenderData()->LODResources[LODIndex];
-
-		BuildFromMeshDescription(*MeshDescriptionPtr, LODResources);
-
-#if WITH_EDITOR
-		for (int32 SectionIndex = 0; SectionIndex < LODResources.Sections.Num(); SectionIndex++)
+	{
+		for (int32 LODIndex = 0; LODIndex < MeshDescriptions.Num(); LODIndex++)
 		{
-			const FStaticMeshSection& StaticMeshSection = LODResources.Sections[SectionIndex];
-			FMeshSectionInfo SectionInfo;
-			SectionInfo.MaterialIndex = StaticMeshSection.MaterialIndex;
-			SectionInfo.bEnableCollision = StaticMeshSection.bEnableCollision;
-			SectionInfo.bCastShadow = StaticMeshSection.bCastShadow;
-			GetSectionInfoMap().Set(LODIndex, SectionIndex, SectionInfo);
+			check(MeshDescriptions[LODIndex] != nullptr);
+			FStaticMeshLODResources& LODResources = GetRenderData()->LODResources[LODIndex];
+
+			BuildFromMeshDescription(*MeshDescriptions[LODIndex], LODResources);
 		}
-#endif
-		LODIndex++;
-	}
 
-	InitResources();
+		InitResources();
 
-	// Set up RenderData bounds and LOD data
-	GetRenderData()->Bounds = MeshDescriptions[0]->GetBounds();
-	CalculateExtendedBounds();
+		// Set up RenderData bounds and LOD data
+		GetRenderData()->Bounds = MeshDescriptions[0]->GetBounds();
+		CalculateExtendedBounds();
 
-	for (int32 LOD = 0; LOD < MeshDescriptions.Num(); ++LOD)
-	{
-		// @todo: some way of customizing LOD screen size and/or calculate it based on mesh bounds
-		if (true)
+		for (int32 LOD = 0; LOD < MeshDescriptions.Num(); ++LOD)
 		{
-			const float LODPowerBase = 0.75f;
-			GetRenderData()->ScreenSize[LOD].Default = FMath::Pow(LODPowerBase, LOD);
+			// @todo: some way of customizing LOD screen size and/or calculate it based on mesh bounds
+			if (true)
+			{
+				const float LODPowerBase = 0.75f;
+				GetRenderData()->ScreenSize[LOD].Default = FMath::Pow(LODPowerBase, LOD);
+			}
+			else
+			{
+				// Possible model for flexible LODs
+				const float MaxDeviation = 100.0f; // specify
+				const float PixelError = SMALL_NUMBER;
+				const float ViewDistance = (MaxDeviation * 960.0f) / PixelError;
+
+				// Generate a projection matrix.
+				const float HalfFOV = PI * 0.25f;
+				const float ScreenWidth = 1920.0f;
+				const float ScreenHeight = 1080.0f;
+				const FPerspectiveMatrix ProjMatrix(HalfFOV, ScreenWidth, ScreenHeight, 1.0f);
+
+				GetRenderData()->ScreenSize[LOD].Default = ComputeBoundsScreenSize(FVector::ZeroVector, GetRenderData()->Bounds.SphereRadius, FVector(0.0f, 0.0f, ViewDistance + GetRenderData()->Bounds.SphereRadius), ProjMatrix);
+			}
 		}
-		else
-		{
-			// Possible model for flexible LODs
-			const float MaxDeviation = 100.0f; // specify
-			const float PixelError = SMALL_NUMBER;
-			const float ViewDistance = (MaxDeviation * 960.0f) / PixelError;
 
-			// Generate a projection matrix.
-			const float HalfFOV = PI * 0.25f;
-			const float ScreenWidth = 1920.0f;
-			const float ScreenHeight = 1080.0f;
-			const FPerspectiveMatrix ProjMatrix(HalfFOV, ScreenWidth, ScreenHeight, 1.0f);
+		// Set up physics-related data
+		CreateBodySetup();
+		check(GetBodySetup());
+		GetBodySetup()->InvalidatePhysicsData();
 
-			GetRenderData()->ScreenSize[LOD].Default = ComputeBoundsScreenSize(FVector::ZeroVector, GetRenderData()->Bounds.SphereRadius, FVector(0.0f, 0.0f, ViewDistance + GetRenderData()->Bounds.SphereRadius), ProjMatrix);
+		if (Params.bBuildSimpleCollision)
+		{
+			FKBoxElem BoxElem;
+			BoxElem.Center = GetRenderData()->Bounds.Origin;
+			BoxElem.X = GetRenderData()->Bounds.BoxExtent.X * 2.0f;
+			BoxElem.Y = GetRenderData()->Bounds.BoxExtent.Y * 2.0f;
+			BoxElem.Z = GetRenderData()->Bounds.BoxExtent.Z * 2.0f;
+			GetBodySetup()->AggGeom.BoxElems.Add(BoxElem);
+			GetBodySetup()->CreatePhysicsMeshes();
 		}
 	}
+#if WITH_EDITOR
+	else
+	{
+		Build(true);
+	}
 
-	// Set up physics-related data
-	CreateBodySetup();
-	check(GetBodySetup());
-	GetBodySetup()->InvalidatePhysicsData();
-
-	if (Params.bBuildSimpleCollision)
+	for (int32 LODIndex = 0; LODIndex < MeshDescriptions.Num(); LODIndex++)
 	{
-		FKBoxElem BoxElem;
-		BoxElem.Center = GetRenderData()->Bounds.Origin;
-		BoxElem.X = GetRenderData()->Bounds.BoxExtent.X * 2.0f;
-		BoxElem.Y = GetRenderData()->Bounds.BoxExtent.Y * 2.0f;
-		BoxElem.Z = GetRenderData()->Bounds.BoxExtent.Z * 2.0f;
-		GetBodySetup()->AggGeom.BoxElems.Add(BoxElem);
-		GetBodySetup()->CreatePhysicsMeshes();
+		FStaticMeshLODResources& LODResources = GetRenderData()->LODResources[LODIndex];
+		for (int32 SectionIndex = 0; SectionIndex < LODResources.Sections.Num(); SectionIndex++)
+		{
+			const FStaticMeshSection& StaticMeshSection = LODResources.Sections[SectionIndex];
+			FMeshSectionInfo SectionInfo;
+			SectionInfo.MaterialIndex = StaticMeshSection.MaterialIndex;
+			SectionInfo.bEnableCollision = StaticMeshSection.bEnableCollision;
+			SectionInfo.bCastShadow = StaticMeshSection.bCastShadow;
+			GetSectionInfoMap().Set(LODIndex, SectionIndex, SectionInfo);
+		}
 	}
+#endif
 
 	if (!bNewMesh)
 	{
