Update LazyPrintf for new HLSL translator (no more tesselation)

diff --git a/Engine/Source/Runtime/Engine/Private/Materials/MaterialHLSLEmitter.cpp b/Engine/Source/Runtime/Engine/Private/Materials/MaterialHLSLEmitter.cpp
--- a/Engine/Source/Runtime/Engine/Private/Materials/MaterialHLSLEmitter.cpp
+++ b/Engine/Source/Runtime/Engine/Private/Materials/MaterialHLSLEmitter.cpp
@@ -206,8 +206,6 @@ static FString GenerateMaterialTemplateHLSL(EShaderPlatform ShaderPlatform,
 	LazyPrintf.PushParam(TEXT("return 0.0f"));
 	LazyPrintf.PushParam(TEXT("return 0.0f"));
 	LazyPrintf.PushParam(TEXT("return 0.0f"));
-	LazyPrintf.PushParam(TEXT("return 0.0f"));
-	LazyPrintf.PushParam(TEXT("return 0.0f"));
 
 	// Print custom texture coordinate assignments, should be fine with regular derivatives
 	FString CustomUVAssignments;
