Implemented UTypedElementCounterInterface for UObject-based element types to count the selected UClass types

diff --git a/Engine/Source/Runtime/Engine/Private/Elements/Actor/ActorElementCounterInterface.cpp b/Engine/Source/Runtime/Engine/Private/Elements/Actor/ActorElementCounterInterface.cpp
new file mode 100644
--- /dev/null
+++ b/Engine/Source/Runtime/Engine/Private/Elements/Actor/ActorElementCounterInterface.cpp
@@ -0,0 +1,23 @@
+// Copyright Epic Games, Inc. All Rights Reserved.
+
+#include "Elements/Actor/ActorElementCounterInterface.h"
+#include "Elements/Actor/ActorElementData.h"
+#include "GameFramework/Actor.h"
+
+#include "Elements/Object/ObjectElementCounterInterface.h"
+
+void UActorElementCounterInterface::IncrementCountersForElement(const FTypedElementHandle& InElementHandle, FTypedElementCounter& InOutCounter)
+{
+	if (const AActor* Actor = ActorElementDataUtil::GetActorFromHandle(InElementHandle))
+	{
+		UObjectElementCounterInterface::IncrementCounterForObjectClass(Actor, InOutCounter);
+	}
+}
+
+void UActorElementCounterInterface::DecrementCountersForElement(const FTypedElementHandle& InElementHandle, FTypedElementCounter& InOutCounter)
+{
+	if (const AActor* Actor = ActorElementDataUtil::GetActorFromHandle(InElementHandle))
+	{
+		UObjectElementCounterInterface::DecrementCounterForObjectClass(Actor, InOutCounter);
+	}
+}
diff --git a/Engine/Source/Runtime/Engine/Private/Elements/Component/ComponentElementCounterInterface.cpp b/Engine/Source/Runtime/Engine/Private/Elements/Component/ComponentElementCounterInterface.cpp
new file mode 100644
--- /dev/null
+++ b/Engine/Source/Runtime/Engine/Private/Elements/Component/ComponentElementCounterInterface.cpp
@@ -0,0 +1,23 @@
+// Copyright Epic Games, Inc. All Rights Reserved.
+
+#include "Elements/Component/ComponentElementCounterInterface.h"
+#include "Elements/Component/ComponentElementData.h"
+#include "Components/ActorComponent.h"
+
+#include "Elements/Object/ObjectElementCounterInterface.h"
+
+void UComponentElementCounterInterface::IncrementCountersForElement(const FTypedElementHandle& InElementHandle, FTypedElementCounter& InOutCounter)
+{
+	if (const UActorComponent* Component = ComponentElementDataUtil::GetComponentFromHandle(InElementHandle))
+	{
+		UObjectElementCounterInterface::IncrementCounterForObjectClass(Component, InOutCounter);
+	}
+}
+
+void UComponentElementCounterInterface::DecrementCountersForElement(const FTypedElementHandle& InElementHandle, FTypedElementCounter& InOutCounter)
+{
+	if (const UActorComponent* Component = ComponentElementDataUtil::GetComponentFromHandle(InElementHandle))
+	{
+		UObjectElementCounterInterface::DecrementCounterForObjectClass(Component, InOutCounter);
+	}
+}
diff --git a/Engine/Source/Runtime/Engine/Private/Elements/Framework/EngineElements.cpp b/Engine/Source/Runtime/Engine/Private/Elements/Framework/EngineElements.cpp
--- a/Engine/Source/Runtime/Engine/Private/Elements/Framework/EngineElements.cpp
+++ b/Engine/Source/Runtime/Engine/Private/Elements/Framework/EngineElements.cpp
@@ -6,16 +6,19 @@
 #include "Elements/Object/ObjectElementData.h"
 #include "Elements/Object/ObjectElementAssetDataInterface.h"
 #include "Elements/Object/ObjectElementObjectInterface.h"
+#include "Elements/Object/ObjectElementCounterInterface.h"
 #include "Elements/Object/ObjectElementSelectionInterface.h"
 
 #include "Elements/Actor/ActorElementData.h"
 #include "Elements/Actor/ActorElementAssetDataInterface.h"
 #include "Elements/Actor/ActorElementObjectInterface.h"
+#include "Elements/Actor/ActorElementCounterInterface.h"
 #include "Elements/Actor/ActorElementWorldInterface.h"
 #include "Elements/Actor/ActorElementSelectionInterface.h"
 
 #include "Elements/Component/ComponentElementData.h"
 #include "Elements/Component/ComponentElementObjectInterface.h"
+#include "Elements/Component/ComponentElementCounterInterface.h"
 #include "Elements/Component/ComponentElementWorldInterface.h"
 #include "Elements/Component/ComponentElementSelectionInterface.h"
 
@@ -34,6 +37,7 @@ void RegisterEngineObjectElements()
 	Registry->RegisterElementType<FObjectElementData>(NAME_Object);
 	Registry->RegisterElementInterface<UTypedElementAssetDataInterface>(NAME_Object, NewObject<UObjectElementAssetDataInterface>());
 	Registry->RegisterElementInterface<UTypedElementObjectInterface>(NAME_Object, NewObject<UObjectElementObjectInterface>());
+	Registry->RegisterElementInterface<UTypedElementCounterInterface>(NAME_Object, NewObject<UObjectElementCounterInterface>());
 	Registry->RegisterElementInterface<UTypedElementSelectionInterface>(NAME_Object, NewObject<UObjectElementSelectionInterface>());
 }
 
@@ -44,6 +48,7 @@ void RegisterEngineActorElements()
 	Registry->RegisterElementType<FActorElementData>(NAME_Actor);
 	Registry->RegisterElementInterface<UTypedElementAssetDataInterface>(NAME_Actor, NewObject<UActorElementAssetDataInterface>());
 	Registry->RegisterElementInterface<UTypedElementObjectInterface>(NAME_Actor, NewObject<UActorElementObjectInterface>());
+	Registry->RegisterElementInterface<UTypedElementCounterInterface>(NAME_Actor, NewObject<UActorElementCounterInterface>());
 	Registry->RegisterElementInterface<UTypedElementWorldInterface>(NAME_Actor, NewObject<UActorElementWorldInterface>());
 	Registry->RegisterElementInterface<UTypedElementSelectionInterface>(NAME_Actor, NewObject<UActorElementSelectionInterface>());
 }
@@ -54,6 +59,7 @@ void RegisterEngineComponentElements()
 
 	Registry->RegisterElementType<FComponentElementData>(NAME_Components);
 	Registry->RegisterElementInterface<UTypedElementObjectInterface>(NAME_Components, NewObject<UComponentElementObjectInterface>());
+	Registry->RegisterElementInterface<UTypedElementCounterInterface>(NAME_Components, NewObject<UComponentElementCounterInterface>());
 	Registry->RegisterElementInterface<UTypedElementWorldInterface>(NAME_Components, NewObject<UComponentElementWorldInterface>());
 	Registry->RegisterElementInterface<UTypedElementSelectionInterface>(NAME_Components, NewObject<UComponentElementSelectionInterface>());
 }
diff --git a/Engine/Source/Runtime/Engine/Private/Elements/Object/ObjectElementCounterInterface.cpp b/Engine/Source/Runtime/Engine/Private/Elements/Object/ObjectElementCounterInterface.cpp
new file mode 100644
--- /dev/null
+++ b/Engine/Source/Runtime/Engine/Private/Elements/Object/ObjectElementCounterInterface.cpp
@@ -0,0 +1,32 @@
+// Copyright Epic Games, Inc. All Rights Reserved.
+
+#include "Elements/Object/ObjectElementCounterInterface.h"
+#include "Elements/Object/ObjectElementData.h"
+
+void UObjectElementCounterInterface::IncrementCountersForElement(const FTypedElementHandle& InElementHandle, FTypedElementCounter& InOutCounter)
+{
+	if (const UObject* Object = ObjectElementDataUtil::GetObjectFromHandle(InElementHandle))
+	{
+		UObjectElementCounterInterface::IncrementCounterForObjectClass(Object, InOutCounter);
+	}
+}
+
+void UObjectElementCounterInterface::DecrementCountersForElement(const FTypedElementHandle& InElementHandle, FTypedElementCounter& InOutCounter)
+{
+	if (const UObject* Object = ObjectElementDataUtil::GetObjectFromHandle(InElementHandle))
+	{
+		UObjectElementCounterInterface::DecrementCounterForObjectClass(Object, InOutCounter);
+	}
+}
+
+void UObjectElementCounterInterface::IncrementCounterForObjectClass(const UObject* InObject, FTypedElementCounter& InOutCounter)
+{
+	check(InObject);
+	InOutCounter.IncrementCounter(NAME_Class, InObject->GetClass());
+}
+
+void UObjectElementCounterInterface::DecrementCounterForObjectClass(const UObject* InObject, FTypedElementCounter& InOutCounter)
+{
+	check(InObject);
+	InOutCounter.DecrementCounter(NAME_Class, InObject->GetClass());
+}
diff --git a/Engine/Source/Runtime/Engine/Public/Elements/Actor/ActorElementCounterInterface.h b/Engine/Source/Runtime/Engine/Public/Elements/Actor/ActorElementCounterInterface.h
new file mode 100644
--- /dev/null
+++ b/Engine/Source/Runtime/Engine/Public/Elements/Actor/ActorElementCounterInterface.h
@@ -0,0 +1,16 @@
+// Copyright Epic Games, Inc. All Rights Reserved.
+
+#pragma once
+
+#include "Elements/Framework/TypedElementCounter.h"
+#include "ActorElementCounterInterface.generated.h"
+
+UCLASS()
+class ENGINE_API UActorElementCounterInterface : public UTypedElementCounterInterface
+{
+	GENERATED_BODY()
+
+public:
+	virtual void IncrementCountersForElement(const FTypedElementHandle& InElementHandle, FTypedElementCounter& InOutCounter) override;
+	virtual void DecrementCountersForElement(const FTypedElementHandle& InElementHandle, FTypedElementCounter& InOutCounter) override;
+};
diff --git a/Engine/Source/Runtime/Engine/Public/Elements/Component/ComponentElementCounterInterface.h b/Engine/Source/Runtime/Engine/Public/Elements/Component/ComponentElementCounterInterface.h
new file mode 100644
--- /dev/null
+++ b/Engine/Source/Runtime/Engine/Public/Elements/Component/ComponentElementCounterInterface.h
@@ -0,0 +1,16 @@
+// Copyright Epic Games, Inc. All Rights Reserved.
+
+#pragma once
+
+#include "Elements/Framework/TypedElementCounter.h"
+#include "ComponentElementCounterInterface.generated.h"
+
+UCLASS()
+class ENGINE_API UComponentElementCounterInterface : public UTypedElementCounterInterface
+{
+	GENERATED_BODY()
+
+public:
+	virtual void IncrementCountersForElement(const FTypedElementHandle& InElementHandle, FTypedElementCounter& InOutCounter) override;
+	virtual void DecrementCountersForElement(const FTypedElementHandle& InElementHandle, FTypedElementCounter& InOutCounter) override;
+};
diff --git a/Engine/Source/Runtime/Engine/Public/Elements/Object/ObjectElementCounterInterface.h b/Engine/Source/Runtime/Engine/Public/Elements/Object/ObjectElementCounterInterface.h
new file mode 100644
--- /dev/null
+++ b/Engine/Source/Runtime/Engine/Public/Elements/Object/ObjectElementCounterInterface.h
@@ -0,0 +1,19 @@
+// Copyright Epic Games, Inc. All Rights Reserved.
+
+#pragma once
+
+#include "Elements/Framework/TypedElementCounter.h"
+#include "ObjectElementCounterInterface.generated.h"
+
+UCLASS()
+class ENGINE_API UObjectElementCounterInterface : public UTypedElementCounterInterface
+{
+	GENERATED_BODY()
+
+public:
+	virtual void IncrementCountersForElement(const FTypedElementHandle& InElementHandle, FTypedElementCounter& InOutCounter) override;
+	virtual void DecrementCountersForElement(const FTypedElementHandle& InElementHandle, FTypedElementCounter& InOutCounter) override;
+
+	static void IncrementCounterForObjectClass(const UObject* InObject, FTypedElementCounter& InOutCounter);
+	static void DecrementCounterForObjectClass(const UObject* InObject, FTypedElementCounter& InOutCounter);
+};
diff --git a/Engine/Source/Runtime/TypedElementRuntime/Private/Elements/Framework/TypedElementListObjectUtil.cpp b/Engine/Source/Runtime/TypedElementRuntime/Private/Elements/Framework/TypedElementListObjectUtil.cpp
--- a/Engine/Source/Runtime/TypedElementRuntime/Private/Elements/Framework/TypedElementListObjectUtil.cpp
+++ b/Engine/Source/Runtime/TypedElementRuntime/Private/Elements/Framework/TypedElementListObjectUtil.cpp
@@ -98,4 +98,22 @@ UObject* GetBottomObject(const UTypedElementList* InElementList, const UClass* I
 	return nullptr;
 }
 
+bool HasObjectsOfExactClass(const UTypedElementList* InElementList, const UClass* InClass)
+{
+	return CountObjectsOfExactClass(InElementList, InClass) > 0;
+}
+
+int32 CountObjectsOfExactClass(const UTypedElementList* InElementList, const UClass* InClass)
+{
+	return InElementList->GetCounter().GetCounterValue(NAME_Class, InClass);
+}
+
+void ForEachObjectClass(const UTypedElementList* InElementList, TFunctionRef<bool(UClass*)> InCallback)
+{
+	InElementList->GetCounter().ForEachCounterValue<UClass*>(NAME_Class, [&InCallback](UClass* InClass, FTypedElementCounter::FCounterValue InCount)
+	{
+		return InCallback(InClass);
+	});
+}
+
 } // namespace TypedElementListObjectUtil
diff --git a/Engine/Source/Runtime/TypedElementRuntime/Public/Elements/Framework/TypedElementListObjectUtil.h b/Engine/Source/Runtime/TypedElementRuntime/Public/Elements/Framework/TypedElementListObjectUtil.h
--- a/Engine/Source/Runtime/TypedElementRuntime/Public/Elements/Framework/TypedElementListObjectUtil.h
+++ b/Engine/Source/Runtime/TypedElementRuntime/Public/Elements/Framework/TypedElementListObjectUtil.h
@@ -105,4 +105,38 @@ RequiredClassType* GetBottomObject(const UTypedElementList* InElementList)
 	return Cast<RequiredClassType>(GetBottomObject(InElementList, RequiredClassType::StaticClass()));
 }
 
+/**
+ * Test if there are any objects of the exact class in the given list of elements (a quick test using the class counter, skipping derived types).
+ */
+TYPEDELEMENTRUNTIME_API bool HasObjectsOfExactClass(const UTypedElementList* InElementList, const UClass* InClass);
+
+/**
+ * Test if there are any objects of the exact class in the given list of elements (a quick test using the class counter, skipping derived types).
+ */
+template <typename ClassType>
+bool HasObjectsOfExactClass(const UTypedElementList* InElementList)
+{
+	return HasObjectsOfExactClass(InElementList, ClassType::StaticClass());
+}
+
+/**
+ * Count the number of objects of the exact class in the given list of elements (a quick test using the class counter, skipping derived types).
+ */
+TYPEDELEMENTRUNTIME_API int32 CountObjectsOfExactClass(const UTypedElementList* InElementList, const UClass* InClass);
+
+/**
+ * Count the number of objects of the exact class in the given list of elements (a quick test using the class counter, skipping derived types).
+ */
+template <typename ClassType>
+bool CountObjectsOfExactClass(const UTypedElementList* InElementList)
+{
+	return CountObjectsOfExactClass(InElementList, ClassType::StaticClass());
+}
+
+/**
+ * Enumerate the classes of the objects in the given list of elements.
+ * @note Return true from the callback to continue enumeration.
+ */
+TYPEDELEMENTRUNTIME_API void ForEachObjectClass(const UTypedElementList* InElementList, TFunctionRef<bool(UClass*)> InCallback);
+
 } // namespace TypedElementListObjectUtil
