Bugfix for runtime generated static meshes with vertex colors.

diff --git a/Engine/Source/Runtime/Engine/Private/StaticMesh.cpp b/Engine/Source/Runtime/Engine/Private/StaticMesh.cpp
--- a/Engine/Source/Runtime/Engine/Private/StaticMesh.cpp
+++ b/Engine/Source/Runtime/Engine/Private/StaticMesh.cpp
@@ -5626,6 +5626,7 @@ void UStaticMesh::BuildFromMeshDescription(const FMeshDescription& MeshDescripti
 	LODResources.VertexBuffers.PositionVertexBuffer.Init(StaticMeshBuildVertices);
 	LODResources.VertexBuffers.StaticMeshVertexBuffer.Init(StaticMeshBuildVertices, VertexInstanceUVs.GetNumChannels());
 
+	LODResources.bHasColorVertexData = bHasVertexColors;
 	FColorVertexBuffer& ColorVertexBuffer = LODResources.VertexBuffers.ColorVertexBuffer;
 	if (bHasVertexColors)
 	{
