Avoid cooking a procedural sound wave asset type

diff --git a/Engine/Source/Runtime/Engine/Private/AudioDerivedData.cpp b/Engine/Source/Runtime/Engine/Private/AudioDerivedData.cpp
--- a/Engine/Source/Runtime/Engine/Private/AudioDerivedData.cpp
+++ b/Engine/Source/Runtime/Engine/Private/AudioDerivedData.cpp
@@ -18,6 +18,7 @@
 #include "AudioCompressionSettingsUtils.h"
 #include "Sound/SoundSourceBus.h"
 #include "Sound/SoundWave.h"
+#include "Sound/SoundWaveProcedural.h"
 
 DEFINE_LOG_CATEGORY_STATIC(LogAudioDerivedData, Log, All);
 
@@ -292,6 +293,11 @@ class FStreamedAudioCacheDerivedDataWorker : public FNonAbandonableTask
 	/** Build the streamed audio. This function is safe to call from any thread. */
 	void BuildStreamedAudio()
 	{
+		if (SoundWave.IsA<USoundWaveProcedural>())
+		{
+			return;
+		}
+
 		GetStreamedAudioDerivedDataKeySuffix(SoundWave, AudioFormatName, CompressionOverrides, KeySuffix);
 
 		DerivedData->Chunks.Empty();
@@ -1046,6 +1052,11 @@ static void CookSimpleWave(USoundWave* SoundWave, FName FormatName, const IAudio
 		return;
 	}
 
+	if (SoundWave->IsA<USoundWaveProcedural>())
+	{
+		return;
+	}
+
 #if WITH_EDITORONLY_DATA
 	FScopeLock ScopeLock(&SoundWave->RawDataCriticalSection);
 #endif
@@ -1173,6 +1184,12 @@ static void CookSurroundWave( USoundWave* SoundWave, FName FormatName, const IAu
 	TArray<TArray<uint8> >	SourceBuffers;
 	TArray<int32>			RequiredChannels;
 
+
+	if (SoundWave->IsA<USoundWaveProcedural>())
+	{
+		return;
+	}
+
 #if WITH_EDITORONLY_DATA
 	FScopeLock ScopeLock(&SoundWave->RawDataCriticalSection);
 #endif
