Fix single-stepping etc in BlendSpace
Disable context menu for animation editing in the scrub widget used by BlendSpace

diff --git a/Engine/Source/Editor/KismetWidgets/Private/SScrubControlPanel.cpp b/Engine/Source/Editor/KismetWidgets/Private/SScrubControlPanel.cpp
--- a/Engine/Source/Editor/KismetWidgets/Private/SScrubControlPanel.cpp
+++ b/Engine/Source/Editor/KismetWidgets/Private/SScrubControlPanel.cpp
@@ -57,6 +57,7 @@ void SScrubControlPanel::Construct( const SScrubControlPanel::FArguments& InArgs
 				.Value(InArgs._Value)
 				.NumOfKeys(InArgs._NumOfKeys)
 				.DisplayDrag(InArgs._DisplayDrag)
+				.bDisplayAnimScrubBarEditing(InArgs._bDisplayAnimScrubBarEditing)
 				.SequenceLength(InArgs._SequenceLength)
 				.OnValueChanged(InArgs._OnValueChanged)
 				.OnBeginSliderMovement(InArgs._OnBeginSliderMovement)
diff --git a/Engine/Source/Editor/KismetWidgets/Private/SScrubWidget.cpp b/Engine/Source/Editor/KismetWidgets/Private/SScrubWidget.cpp
--- a/Engine/Source/Editor/KismetWidgets/Private/SScrubWidget.cpp
+++ b/Engine/Source/Editor/KismetWidgets/Private/SScrubWidget.cpp
@@ -60,6 +60,7 @@ void SScrubWidget::Construct( const SScrubWidget::FArguments& InArgs )
 	DraggableBars = InArgs._DraggableBars;
 	OnBarDrag = InArgs._OnBarDrag;
 	bDisplayDrag = InArgs._DisplayDrag;
+	bDisplayAnimScrubBarEditing = InArgs._bDisplayAnimScrubBarEditing;
 	bMouseMovedDuringPanning = false;
 	bDragging = false;
 	bPanning = false;
@@ -425,7 +426,9 @@ FCursorReply SScrubWidget::OnCursorQuery( const FGeometry& MyGeometry, const FPo
 
 void SScrubWidget::CreateContextMenu(float CurrentFrameTime, const FPointerEvent& MouseEvent)
 {
-	if ((OnCropAnimSequence.IsBound() || OnReZeroAnimSequence.IsBound() || OnAddAnimSequence.IsBound()) && (SequenceLength.Get() >= MINIMUM_ANIMATION_LENGTH))
+	if (bDisplayAnimScrubBarEditing && 
+		(OnCropAnimSequence.IsBound() || OnReZeroAnimSequence.IsBound() || OnAddAnimSequence.IsBound()) && 
+		(SequenceLength.Get() >= MINIMUM_ANIMATION_LENGTH))
 	{
 		const bool CloseAfterSelection = true;
 		FMenuBuilder MenuBuilder( CloseAfterSelection, NULL );
diff --git a/Engine/Source/Editor/KismetWidgets/Public/SScrubControlPanel.h b/Engine/Source/Editor/KismetWidgets/Public/SScrubControlPanel.h
--- a/Engine/Source/Editor/KismetWidgets/Public/SScrubControlPanel.h
+++ b/Engine/Source/Editor/KismetWidgets/Public/SScrubControlPanel.h
@@ -42,6 +42,7 @@ public:
 		SLATE_ATTRIBUTE( uint32, NumOfKeys )
 		SLATE_ATTRIBUTE( float, SequenceLength )
 		SLATE_ARGUMENT( bool, bAllowZoom )
+		SLATE_ARGUMENT( bool, bDisplayAnimScrubBarEditing )
 		SLATE_ATTRIBUTE(bool, DisplayDrag)
 		/** Called when the value is changed by slider or typing */
 		SLATE_EVENT( FOnFloatValueChanged, OnValueChanged )
diff --git a/Engine/Source/Editor/KismetWidgets/Public/SScrubWidget.h b/Engine/Source/Editor/KismetWidgets/Public/SScrubWidget.h
--- a/Engine/Source/Editor/KismetWidgets/Public/SScrubWidget.h
+++ b/Engine/Source/Editor/KismetWidgets/Public/SScrubWidget.h
@@ -34,6 +34,7 @@ public:
 		, _NumOfKeys(10)
 		, _SequenceLength()
 		, _bAllowZoom(false)
+		, _bDisplayAnimScrubBarEditing(true)
 		, _DisplayDrag(true)
 		, _OnValueChanged()
 		, _OnBeginSliderMovement()
@@ -49,6 +50,7 @@ public:
 		SLATE_ATTRIBUTE( uint32, NumOfKeys )
 		SLATE_ATTRIBUTE( float, SequenceLength )
 		SLATE_ARGUMENT( bool, bAllowZoom )
+		SLATE_ARGUMENT( bool, bDisplayAnimScrubBarEditing )
 		SLATE_ATTRIBUTE( bool, DisplayDrag )
 		/** Called when the value is changed by slider or typing */
 		SLATE_EVENT( FOnFloatValueChanged, OnValueChanged )
@@ -146,7 +148,8 @@ private:
 	TAttribute<float> SequenceLength;
 	/** The number of decimal places to display */
 	bool bDragging;
-	bool bAllowZoom;	
+	bool bAllowZoom;
+	bool bDisplayAnimScrubBarEditing;
 	TAttribute<bool> bDisplayDrag;
 	/** If we are currently panning the panel*/
 	bool bPanning;
diff --git a/Engine/Source/Editor/Persona/Private/SAnimEditorBase.cpp b/Engine/Source/Editor/Persona/Private/SAnimEditorBase.cpp
--- a/Engine/Source/Editor/Persona/Private/SAnimEditorBase.cpp
+++ b/Engine/Source/Editor/Persona/Private/SAnimEditorBase.cpp
@@ -74,13 +74,13 @@ void SAnimEditorBase::Construct(const FArguments& InArgs, const TSharedPtr<class
 			+SHorizontalBox::Slot()
 			.FillWidth(1)
 			[
-				ConstructAnimScrubPanel()
+				ConstructAnimScrubPanel(InArgs._DisplayAnimScrubBarEditing)
 			]
 		];
 	}
 }
 
-TSharedRef<SWidget> SAnimEditorBase::ConstructAnimScrubPanel()
+TSharedRef<SWidget> SAnimEditorBase::ConstructAnimScrubPanel(bool bDisplayAnimScrubBarEditing)
 {
 	if(PreviewScenePtr.IsValid())
 	{
@@ -88,6 +88,7 @@ TSharedRef<SWidget> SAnimEditorBase::ConstructAnimScrubPanel()
 			.LockedSequence(Cast<UAnimSequenceBase>(GetEditorObject()))
 			.ViewInputMin(this, &SAnimEditorBase::GetViewMinInput)
 			.ViewInputMax(this, &SAnimEditorBase::GetViewMaxInput)
+			.bDisplayAnimScrubBarEditing(bDisplayAnimScrubBarEditing)
 			.OnSetInputViewRange(this, &SAnimEditorBase::SetInputViewRange)
 			.bAllowZoom(true);
 	}
diff --git a/Engine/Source/Editor/Persona/Private/SAnimEditorBase.h b/Engine/Source/Editor/Persona/Private/SAnimEditorBase.h
--- a/Engine/Source/Editor/Persona/Private/SAnimEditorBase.h
+++ b/Engine/Source/Editor/Persona/Private/SAnimEditorBase.h
@@ -45,6 +45,8 @@ public:
 
 	SLATE_ARGUMENT( bool, DisplayAnimScrubBar )
 
+	SLATE_ARGUMENT( bool, DisplayAnimScrubBarEditing )
+
 	SLATE_ARGUMENT( TSharedPtr<FAnimModel>, AnimModel )
 
 	SLATE_END_ARGS()
@@ -90,7 +92,7 @@ protected:
 	TSharedRef<IPersonaPreviewScene> GetPreviewScene() const { return PreviewScenePtr.Pin().ToSharedRef(); }
 
 	/** Allows derived classes to create different animation scrub panel */
-	virtual TSharedRef<SWidget> ConstructAnimScrubPanel();
+	virtual TSharedRef<SWidget> ConstructAnimScrubPanel(bool bDisplayAnimScrubBarEditing);
 
 	/** Allows derived classes to init newly created editor objects */
 	virtual void InitDetailsViewEditorObject(class UEditorAnimBaseObj* EdObj) {};
diff --git a/Engine/Source/Editor/Persona/Private/SAnimationBlendSpace.cpp b/Engine/Source/Editor/Persona/Private/SAnimationBlendSpace.cpp
--- a/Engine/Source/Editor/Persona/Private/SAnimationBlendSpace.cpp
+++ b/Engine/Source/Editor/Persona/Private/SAnimationBlendSpace.cpp
@@ -40,7 +40,8 @@ void SBlendSpaceEditor::Construct(const FArguments& InArgs)
 
 	SAnimEditorBase::Construct(SAnimEditorBase::FArguments()
 		.DisplayAnimTimeline(false)
-		.DisplayAnimScrubBar(InArgs._DisplayScrubBar),
+		.DisplayAnimScrubBar(InArgs._DisplayScrubBar)
+		.DisplayAnimScrubBarEditing(false),
 		PreviewScenePtr.Pin());
 
 	NonScrollEditorPanels->AddSlot()
diff --git a/Engine/Source/Editor/Persona/Private/SAnimationScrubPanel.cpp b/Engine/Source/Editor/Persona/Private/SAnimationScrubPanel.cpp
--- a/Engine/Source/Editor/Persona/Private/SAnimationScrubPanel.cpp
+++ b/Engine/Source/Editor/Persona/Private/SAnimationScrubPanel.cpp
@@ -55,6 +55,7 @@ void SAnimationScrubPanel::Construct( const SAnimationScrubPanel::FArguments& In
 			.OnGetRecording(this, &SAnimationScrubPanel::IsRecording)
 			.ViewInputMin(InArgs._ViewInputMin)
 			.ViewInputMax(InArgs._ViewInputMax)
+			.bDisplayAnimScrubBarEditing(InArgs._bDisplayAnimScrubBarEditing)
 			.OnSetInputViewRange(InArgs._OnSetInputViewRange)
 			.OnCropAnimSequence( this, &SAnimationScrubPanel::OnCropAnimSequence )
 			.OnAddAnimSequence( this, &SAnimationScrubPanel::OnInsertAnimSequence )
@@ -84,15 +85,13 @@ FReply SAnimationScrubPanel::OnClick_Forward_Step()
 	}
 	else if (SMC)
 	{
-		//@TODO: Should we hardcode 30 Hz here?
-		{
-			const float TargetFramerate = 30.0f;
+		// BlendSpace combines animations so there's no such thing as a frame. However, 1/30 is a sensible/common rate.
+		const float FixedFrameRate = 30.0f;
 
-			// Advance a single frame, leaving it paused afterwards
-			SMC->GlobalAnimRateScale = 1.0f;
-			SMC->TickAnimation(1.0f / TargetFramerate, false);
-			SMC->GlobalAnimRateScale = 0.0f;
-		}
+		// Advance a single frame, leaving it paused afterwards
+		SMC->GlobalAnimRateScale = 1.0f;
+		SMC->TickAnimation(1.0f / FixedFrameRate, false);
+		SMC->GlobalAnimRateScale = 0.0f;
 	}
 
 	return FReply::Handled();
diff --git a/Engine/Source/Editor/Persona/Private/SAnimationScrubPanel.h b/Engine/Source/Editor/Persona/Private/SAnimationScrubPanel.h
--- a/Engine/Source/Editor/Persona/Private/SAnimationScrubPanel.h
+++ b/Engine/Source/Editor/Persona/Private/SAnimationScrubPanel.h
@@ -28,6 +28,7 @@ public:
 		/** View Input range **/
 		SLATE_ATTRIBUTE( float, ViewInputMin )
 		SLATE_ATTRIBUTE( float, ViewInputMax )
+		SLATE_ARGUMENT( bool, bDisplayAnimScrubBarEditing)
 		SLATE_EVENT( FOnSetInputViewRange, OnSetInputViewRange )
 		/** Called when an anim sequence is cropped before/after a selected frame */
 		SLATE_EVENT( FOnCropAnimSequence, OnCropAnimSequence )
diff --git a/Engine/Source/Runtime/Engine/Private/Animation/AnimSingleNodeInstance.cpp b/Engine/Source/Runtime/Engine/Private/Animation/AnimSingleNodeInstance.cpp
--- a/Engine/Source/Runtime/Engine/Private/Animation/AnimSingleNodeInstance.cpp
+++ b/Engine/Source/Runtime/Engine/Private/Animation/AnimSingleNodeInstance.cpp
@@ -392,14 +392,14 @@ float UAnimSingleNodeInstance::GetLength()
 	{
 		if (UBlendSpace* BlendSpace = Cast<UBlendSpace>(CurrentAsset))
 		{
-			return BlendSpace->AnimLength;
+			// Blend space length is normalized to 1 when getting and setting
+			return 1.0f;
 		}
 		else if (UAnimSequenceBase* SequenceBase = Cast<UAnimSequenceBase>(CurrentAsset))
 		{
 			return SequenceBase->GetPlayLength();
 		}
 	}
-
 	return 0.f;
 }
 
@@ -423,7 +423,22 @@ void UAnimSingleNodeInstance::StepForward()
 		StepToFrame.Value = FMath::Clamp<int32>(StepToFrame.Value, 0, LastSequenceFrameNumber.Value);
 
 		SetPosition(FrameRate.AsSeconds(StepToFrame));
-	}	
+	}
+	else if (UBlendSpace* BlendSpace = Cast<UBlendSpace>(CurrentAsset))
+	{
+		// BlendSpace combines animations so there's no such thing as a frame. However, 1/30 is a sensible/common rate.
+		FAnimSingleNodeInstanceProxy& Proxy = GetProxyOnGameThread<FAnimSingleNodeInstanceProxy>();
+		float Length = Proxy.GetBlendSpaceLength();
+		if (Length > 0.0f)
+		{
+			const float FixedFrameRate = 30.0f;
+			float NormalizedDt = 1.0f / (FixedFrameRate * Length);
+			float NormalizedTime = Proxy.GetCurrentTime() + NormalizedDt;
+			NormalizedTime = IsLooping() ? 
+				FMath::Wrap(NormalizedTime, 0.0f, 1.0f) : FMath::Clamp(NormalizedTime, 0.0f, 1.0f);
+			SetPosition(NormalizedTime);
+		}
+	}
 }
 
 void UAnimSingleNodeInstance::StepBackward()
@@ -447,6 +462,21 @@ void UAnimSingleNodeInstance::StepBackward()
 
 		SetPosition(FrameRate.AsSeconds(StepToFrame));
 	}
+	else if (UBlendSpace* BlendSpace = Cast<UBlendSpace>(CurrentAsset))
+	{
+		// BlendSpace combines animations so there's no such thing as a frame. However, 1/30 is a sensible/common rate.
+		FAnimSingleNodeInstanceProxy& Proxy = GetProxyOnGameThread<FAnimSingleNodeInstanceProxy>();
+		float Length = Proxy.GetBlendSpaceLength();
+		if (Length > 0.0f)
+		{
+			const float FixedFrameRate = 30.0f;
+			float NormalizedDt = 1.0f / (FixedFrameRate * Length);
+			float NormalizedTime = Proxy.GetCurrentTime() - NormalizedDt;
+			NormalizedTime = IsLooping() ? 
+				FMath::Wrap(NormalizedTime, 0.0f, 1.0f) : FMath::Clamp(NormalizedTime, 0.0f, 1.0f);
+			SetPosition(NormalizedTime);
+		}
+	}
 }
 
 FAnimInstanceProxy* UAnimSingleNodeInstance::CreateAnimInstanceProxy()
diff --git a/Engine/Source/Runtime/Engine/Private/Animation/AnimSingleNodeInstanceProxy.cpp b/Engine/Source/Runtime/Engine/Private/Animation/AnimSingleNodeInstanceProxy.cpp
--- a/Engine/Source/Runtime/Engine/Private/Animation/AnimSingleNodeInstanceProxy.cpp
+++ b/Engine/Source/Runtime/Engine/Private/Animation/AnimSingleNodeInstanceProxy.cpp
@@ -264,6 +264,25 @@ void FAnimSingleNodeInstanceProxy::GetBlendSpaceState(FVector& OutPosition, FVec
 	OutPosition = BlendSpacePosition;
 }
 
+float FAnimSingleNodeInstanceProxy::GetBlendSpaceLength() const
+{
+	float TotalLength = 0.0f;
+	float TotalWeight = 0.0f;
+	for (int32 Index = 0 ; Index != BlendSampleData.Num() ; ++Index)
+	{
+		const FBlendSampleData& Data = BlendSampleData[Index];
+		float AnimLength = Data.Animation->GetPlayLength();
+		float Weight = Data.GetWeight();
+		TotalLength += AnimLength * Weight;
+		TotalWeight += Weight;
+	}
+	if (TotalWeight > 0.0f)
+	{
+		return TotalLength / TotalWeight;
+	}
+	return 0.0f;
+}
+
 void FAnimNode_SingleNode::Evaluate_AnyThread(FPoseContext& Output)
 {
 	const bool bCanProcessAdditiveAnimationsLocal
diff --git a/Engine/Source/Runtime/Engine/Public/Animation/AnimSingleNodeInstanceProxy.h b/Engine/Source/Runtime/Engine/Public/Animation/AnimSingleNodeInstanceProxy.h
--- a/Engine/Source/Runtime/Engine/Public/Animation/AnimSingleNodeInstanceProxy.h
+++ b/Engine/Source/Runtime/Engine/Public/Animation/AnimSingleNodeInstanceProxy.h
@@ -146,10 +146,20 @@ public:
 
 	void SetReverse(bool bInReverse);
 
+	/** Sets the target blend space position */
 	void SetBlendSpacePosition(const FVector& InPosition);
 
+	/**
+	 * Returns the current target/requested blend space position, and the filtered (smoothed) position.
+	 */
 	void GetBlendSpaceState(FVector& OutPosition, FVector& OutFilteredPosition) const;
 
+	/** 
+	* Returns the length (seconds), not including any rate multipliers, calculated by 
+	* weighting the currently active samples 
+	*/
+	float GetBlendSpaceLength() const;
+
 #if WITH_EDITOR
 	bool CanProcessAdditiveAnimations() const
 	{
