Use LOD index stored in skin cache rather than the one from mesh object, it's safer this way in case mesh object's LOD changed before skin cache entry's updated

diff --git a/Engine/Source/Runtime/Engine/Private/GPUSkinCache.cpp b/Engine/Source/Runtime/Engine/Private/GPUSkinCache.cpp
--- a/Engine/Source/Runtime/Engine/Private/GPUSkinCache.cpp
+++ b/Engine/Source/Runtime/Engine/Private/GPUSkinCache.cpp
@@ -1156,8 +1156,8 @@ void FGPUSkinCache::DispatchUpdateSkinTangents(FRHICommandListImmediate& RHICmdL
 	FGPUSkinCacheEntry::FSectionDispatchData& DispatchData = Entry->DispatchData[SectionIndex];
 
 	FSkeletalMeshRenderData& SkelMeshRenderData = Entry->GPUSkin->GetSkeletalMeshRenderData();
-	int32 LODIndex = Entry->LOD;
-	FSkeletalMeshLODRenderData& LodData = SkelMeshRenderData.LODRenderData[Entry->LOD];
+	const int32 LODIndex = Entry->LOD;
+	FSkeletalMeshLODRenderData& LodData = SkelMeshRenderData.LODRenderData[LODIndex];
 
 	if (bTrianglePass)
 	{
@@ -1981,7 +1981,7 @@ void FGPUSkinCache::DispatchUpdateSkinning(FRHICommandListImmediate& RHICmdList,
 
 	SCOPED_DRAW_EVENTF(RHICmdList, SkinCacheDispatch,
 		TEXT("Skinning%d%d%d Mesh=%s LOD=%d Chunk=%d InStreamStart=%d OutStart=%d Vert=%d Morph=%d/%d"),
-		(int32)Entry->bUse16BitBoneIndex, (int32)Entry->BoneInfluenceType, DispatchData.SkinType, *GetSkeletalMeshObjectName(Entry->GPUSkin), Entry->GPUSkin->GetLOD(),
+		(int32)Entry->bUse16BitBoneIndex, (int32)Entry->BoneInfluenceType, DispatchData.SkinType, *GetSkeletalMeshObjectName(Entry->GPUSkin), Entry->LOD,
 		DispatchData.SectionIndex, DispatchData.InputStreamStart, DispatchData.OutputStreamStart, DispatchData.NumVertices, Entry->MorphBuffer != 0, DispatchData.MorphBufferOffset);
 	auto* GlobalShaderMap = GetGlobalShaderMap(GetFeatureLevel());
 	TShaderMapRef<TGPUSkinCacheCS<0>> SkinCacheCS000(GlobalShaderMap);		// 16bit_0, BoneInfluenceType_0, SkinType_0
@@ -2230,7 +2230,7 @@ FCachedGeometry FGPUSkinCache::GetCachedGeometry(uint32 ComponentId) const
 	{
 		if (Entry && Entry->GPUSkin && Entry->GPUSkin->GetComponentId() == ComponentId)
 		{
-			const uint32 LODIndex = Entry->GPUSkin->GetLOD();
+			const uint32 LODIndex = Entry->LOD;
 			const FSkeletalMeshRenderData& RenderData = Entry->GPUSkin->GetSkeletalMeshRenderData();
 			const FSkeletalMeshLODRenderData& LODData = RenderData.LODRenderData[LODIndex];
 			const uint32 SectionCount = LODData.RenderSections.Num();
