Remove unused recompute tangent blend using vertex color buffer modes

diff --git a/Engine/Shaders/Private/RecomputeTangentsPerVertexPass.usf b/Engine/Shaders/Private/RecomputeTangentsPerVertexPass.usf
--- a/Engine/Shaders/Private/RecomputeTangentsPerVertexPass.usf
+++ b/Engine/Shaders/Private/RecomputeTangentsPerVertexPass.usf
@@ -152,30 +152,16 @@ void MainCS(
 		uint VertexIndex = DispatchThreadId;
 		uint Offset = VertexIndex + SkinCacheStart;
 
-#if BLEND_USING_VERTEX_COLOR > 0
+#if BLEND_USING_VERTEX_COLOR
 		float4 InputTangentX = TangentBias_SkinCache(TangentInputBuffer[2 * Offset + GPUSKIN_RWBUFFER_OFFSET_TANGENT_X]);
 		float4 InputTangentZ = TangentBias_SkinCache(TangentInputBuffer[2 * Offset + GPUSKIN_RWBUFFER_OFFSET_TANGENT_Z]);
 		float4 BlendFactors = ColorInputBuffer[DispatchThreadId + InputStreamStart];
 		float BlendFactor = (VertexColorChannel == 0) ? BlendFactors.r : ((VertexColorChannel == 1) ? BlendFactors.g : BlendFactors.b);
-#if BLEND_USING_VERTEX_COLOR == 1
-		// No blending, choose between input tangents and recompute tangents
-		TangentX = BlendFactor < 0.5f ? InputTangentX.xyz : TangentX;
-		TangentZ = BlendFactor < 0.5f ? InputTangentZ.xyz : TangentZ;
-#elif BLEND_USING_VERTEX_COLOR == 2
+
 		// Linear interpolation
 		TangentX = lerp(InputTangentX.xyz, TangentX, BlendFactor);
 		TangentZ = lerp(InputTangentZ.xyz, TangentZ, BlendFactor);
-#elif BLEND_USING_VERTEX_COLOR == 3
-		// Vector slerp
-		TangentX = Slerp(InputTangentX.xyz, TangentX, BlendFactor);
-		TangentZ = Slerp(InputTangentZ.xyz, TangentZ, BlendFactor);
-#elif BLEND_USING_VERTEX_COLOR == 4
-		// Convert tangents into quaternion, apply slerp, then convert from quaternion back to tangents
-		float4 InputQuat = TangentToQuaternion(InputTangentX.xyz, InputTangentZ.xyz);
-		float4 Quat = TangentToQuaternion(TangentX, TangentZ);
-		float4 BlendedQuat = Slerp(InputQuat, Quat, BlendFactor);
-		QuaternionToTangent(BlendedQuat, TangentX, TangentZ);
-#endif
+
 		// No need to realign TangentX and TangentZ to be perfect perpendicular after blending as vertex factory will do so later on
 		TangentX = normalize(TangentX);
 		TangentZ = normalize(TangentZ);
diff --git a/Engine/Source/Runtime/Engine/Private/GPUSkinCache.cpp b/Engine/Source/Runtime/Engine/Private/GPUSkinCache.cpp
--- a/Engine/Source/Runtime/Engine/Private/GPUSkinCache.cpp
+++ b/Engine/Source/Runtime/Engine/Private/GPUSkinCache.cpp
@@ -122,10 +122,7 @@ FAutoConsoleVariableRef CVarGPUSkinCacheBlendUsingVertexColorForRecomputeTangent
 	TEXT("r.SkinCache.BlendUsingVertexColorForRecomputeTangents"),
 	GBlendUsingVertexColorForRecomputeTangents,
 	TEXT("0: off (default)\n")
-	TEXT("1: No blending, choose between source and recompute tangents.\n")
-	TEXT("2: Linear interpolation between source and recompute tangents.\n")
-	TEXT("3: Vector slerp between source and recompute tangents.\n")
-	TEXT("4: Convert tangents into quaternion, apply slerp, then convert from quaternion back to tangents (most expensive).\n"),
+	TEXT("1: Linear interpolation between source and recompute tangents.\n"),
 	ECVF_RenderThreadSafe
 );
 
@@ -1147,9 +1144,6 @@ class FRecomputeTangentsPerVertexPassCS : public FBaseRecomputeTangentsPerVertex
 
 IMPLEMENT_SHADER_TYPE(template<>, FRecomputeTangentsPerVertexPassCS<0>, TEXT("/Engine/Private/RecomputeTangentsPerVertexPass.usf"), TEXT("MainCS"), SF_Compute);
 IMPLEMENT_SHADER_TYPE(template<>, FRecomputeTangentsPerVertexPassCS<1>, TEXT("/Engine/Private/RecomputeTangentsPerVertexPass.usf"), TEXT("MainCS"), SF_Compute);
-IMPLEMENT_SHADER_TYPE(template<>, FRecomputeTangentsPerVertexPassCS<2>, TEXT("/Engine/Private/RecomputeTangentsPerVertexPass.usf"), TEXT("MainCS"), SF_Compute);
-IMPLEMENT_SHADER_TYPE(template<>, FRecomputeTangentsPerVertexPassCS<3>, TEXT("/Engine/Private/RecomputeTangentsPerVertexPass.usf"), TEXT("MainCS"), SF_Compute);
-IMPLEMENT_SHADER_TYPE(template<>, FRecomputeTangentsPerVertexPassCS<4>, TEXT("/Engine/Private/RecomputeTangentsPerVertexPass.usf"), TEXT("MainCS"), SF_Compute);
 
 void FGPUSkinCache::DispatchUpdateSkinTangents(FRHICommandListImmediate& RHICmdList, FGPUSkinCacheEntry* Entry, int32 SectionIndex, FRWBuffer*& StagingBuffer, bool bTrianglePass)
 {
@@ -1252,18 +1246,9 @@ void FGPUSkinCache::DispatchUpdateSkinTangents(FRHICommandListImmediate& RHICmdL
 		auto* GlobalShaderMap = GetGlobalShaderMap(GetFeatureLevel());
 		TShaderMapRef<FRecomputeTangentsPerVertexPassCS<0>> ComputeShader0(GlobalShaderMap);
 		TShaderMapRef<FRecomputeTangentsPerVertexPassCS<1>> ComputeShader1(GlobalShaderMap);
-		TShaderMapRef<FRecomputeTangentsPerVertexPassCS<2>> ComputeShader2(GlobalShaderMap);
-		TShaderMapRef<FRecomputeTangentsPerVertexPassCS<3>> ComputeShader3(GlobalShaderMap);
-		TShaderMapRef<FRecomputeTangentsPerVertexPassCS<4>> ComputeShader4(GlobalShaderMap);
 		TShaderRef<FBaseRecomputeTangentsPerVertexShader> ComputeShader;
-		if (GBlendUsingVertexColorForRecomputeTangents == 1)
+		if (GBlendUsingVertexColorForRecomputeTangents > 0)
 			ComputeShader = ComputeShader1;
-		else if (GBlendUsingVertexColorForRecomputeTangents == 2)
-			ComputeShader = ComputeShader2;
-		else if (GBlendUsingVertexColorForRecomputeTangents == 3)
-			ComputeShader = ComputeShader3;
-		else if (GBlendUsingVertexColorForRecomputeTangents == 4)
-			ComputeShader = ComputeShader4;
 		else
 			ComputeShader = ComputeShader0;
 		RHICmdList.SetComputeShader(ComputeShader.GetComputeShader());
