Fix bugs caused by assigning the wrong active state when creating the ActiveStateMachineScope

diff --git a/Engine/Source/Runtime/Engine/Private/Animation/AnimNode_StateMachine.cpp b/Engine/Source/Runtime/Engine/Private/Animation/AnimNode_StateMachine.cpp
--- a/Engine/Source/Runtime/Engine/Private/Animation/AnimNode_StateMachine.cpp
+++ b/Engine/Source/Runtime/Engine/Private/Animation/AnimNode_StateMachine.cpp
@@ -482,8 +482,6 @@ void FAnimNode_StateMachine::Update_AnyThread(const FAnimationUpdateContext& Con
 		}
 	}
 	while (bFoundValidTransition && (TransitionCountThisFrame < MaxTransitionsPerFrame));
-
-	UE::Anim::TOptionalScopedGraphMessage<UE::Anim::FActiveStateMachineScope> Message(bCreateNotifyMetaData, Context, Context, this, CurrentState);
 	
 	if (bFirstUpdate)
 	{
@@ -541,6 +539,7 @@ void FAnimNode_StateMachine::Update_AnyThread(const FAnimationUpdateContext& Con
 	// Update the only active state if there are no transitions still in flight
 	if (ActiveTransitionArray.Num() == 0 && !IsAConduitState(CurrentState) && !StatesUpdated.Contains(CurrentState))
 	{
+		UE::Anim::TOptionalScopedGraphMessage<UE::Anim::FActiveStateMachineScope> Message(bCreateNotifyMetaData, Context, Context, this, CurrentState);
 		StatePoseLinks[CurrentState].Update(Context);
 	}
 
@@ -1107,6 +1106,7 @@ void FAnimNode_StateMachine::UpdateState(int32 StateIndex, const FAnimationUpdat
 {
 	if ((StateIndex != INDEX_NONE) && !StatesUpdated.Contains(StateIndex) && !IsAConduitState(StateIndex))
 	{
+		UE::Anim::TOptionalScopedGraphMessage<UE::Anim::FActiveStateMachineScope> Message(bCreateNotifyMetaData, Context, Context, this, StateIndex);
 		StatesUpdated.Add(StateIndex);
 		StatePoseLinks[StateIndex].Update(Context);
 	}
diff --git a/Engine/Source/Runtime/Engine/Private/Animation/AnimNotifyStateMachineInspectionLibrary.cpp b/Engine/Source/Runtime/Engine/Private/Animation/AnimNotifyStateMachineInspectionLibrary.cpp
--- a/Engine/Source/Runtime/Engine/Private/Animation/AnimNotifyStateMachineInspectionLibrary.cpp
+++ b/Engine/Source/Runtime/Engine/Private/Animation/AnimNotifyStateMachineInspectionLibrary.cpp
@@ -55,11 +55,11 @@ bool UAnimNotifyStateMachineInspectionLibrary::IsTriggeredByStateInStateMachine(
 	{
 		int32 MachineIndex;
 		const FBakedAnimationStateMachine* BakedAnimationStateMachine = nullptr;
-		AnimInstance->GetStateMachineIndexAndDescription(StateMachineName,MachineIndex, &BakedAnimationStateMachine);
+		AnimInstance->GetStateMachineIndexAndDescription(StateMachineName, MachineIndex, &BakedAnimationStateMachine);
 		if(MachineIndex != INDEX_NONE && BakedAnimationStateMachine)
 		{
 			int32 StateIndex = BakedAnimationStateMachine->FindStateIndex(StateName);
-			return IsStateInStateMachineInEventContext(EventReference, StateIndex, MachineIndex);
+			return IsStateInStateMachineInEventContext(EventReference, MachineIndex, StateIndex);
 		}
 	}
 	return false; 
