Fix singles stepping etc in BlendSpace when using Marker sync

diff --git a/Engine/Source/Runtime/Engine/Public/Animation/AnimSingleNodeInstanceProxy.h b/Engine/Source/Runtime/Engine/Public/Animation/AnimSingleNodeInstanceProxy.h
--- a/Engine/Source/Runtime/Engine/Public/Animation/AnimSingleNodeInstanceProxy.h
+++ b/Engine/Source/Runtime/Engine/Public/Animation/AnimSingleNodeInstanceProxy.h
@@ -121,6 +121,12 @@ public:
 
 	void SetCurrentTime(float InCurrentTime)
 	{
+		if (InCurrentTime != CurrentTime)
+		{
+			// If the current time is changed externally then our record of where we are in relation to markers will be
+			// out of sync, so reset it and it will be updated when necessary.
+			MarkerTickRecord.Reset();
+		}
 		CurrentTime = InCurrentTime;
 	}
 
