Convert decal conversion map check warning to asset check info.
Everything works without resave, but resave is encouraged to avoid some small fixup work on every load.

diff --git a/Engine/Source/Runtime/Engine/Private/Materials/Material.cpp b/Engine/Source/Runtime/Engine/Private/Materials/Material.cpp
--- a/Engine/Source/Runtime/Engine/Private/Materials/Material.cpp
+++ b/Engine/Source/Runtime/Engine/Private/Materials/Material.cpp
@@ -3677,14 +3677,14 @@ void UMaterial::BackwardsCompatibilityDecalConversion()
 		FlushResourceShaderMaps();
 
 		// Note we can't mark the package dirty during post load so we will recompile on load until this material is manually resaved. 
-		// So add a map load error here to encourage a save.
+		// Add an asset check here to encourage a save.
 		FFormatNamedArguments Arguments;
 		Arguments.Add(TEXT("Material"), FText::FromString(*GetPathName()));
-		FMessageLog("MapCheck").Warning()
+		FMessageLog("AssetCheck").Info()
 			->AddToken(FUObjectToken::Create(this))
-			->AddToken(FTextToken::Create(FText::Format(LOCTEXT("MapCheck_Decal", "Material {Material} has been updated to fix up decal settings. If the material asset is not re-saved it may not render correctly when run outside of the editor."), Arguments)))
-			->AddToken(FActionToken::Create(LOCTEXT("MapCheck_Decal_Fix", "Fix"), LOCTEXT("MapCheck_Decal_Action", "Click to mark the material as needing to be saved."), FOnActionTokenExecuted::CreateUObject(this, &UMaterial::FixupMaterialUsageAfterLoad), true));
-		FMessageLog("MapCheck").Open(EMessageSeverity::Warning);
+			->AddToken(FTextToken::Create(FText::Format(LOCTEXT("AssetCheck_Decal", "Material {Material} has been updated to fix up decal settings. Resave material to avoid it recompiling on each load."), Arguments)))
+			->AddToken(FActionToken::Create(LOCTEXT("AssetCheck_Decal_Fix", "Fix"), LOCTEXT("AssetCheck_Decal_Action", "Click to mark the material as needing to be saved."), FOnActionTokenExecuted::CreateUObject(this, &UMaterial::FixupMaterialUsageAfterLoad), true));
+		FMessageLog("AssetCheck").Open(EMessageSeverity::Info);
 	}
 #endif // WITH_EDITOR
 }
