Improved validation and diagnostic message for unsupported vertex factories in FSkeletalMeshSceneProxy::GetDynamicRayTracingInstances()

diff --git a/Engine/Source/Runtime/Engine/Private/SkeletalMesh.cpp b/Engine/Source/Runtime/Engine/Private/SkeletalMesh.cpp
--- a/Engine/Source/Runtime/Engine/Private/SkeletalMesh.cpp
+++ b/Engine/Source/Runtime/Engine/Private/SkeletalMesh.cpp
@@ -6252,7 +6252,12 @@ void FSkeletalMeshSceneProxy::GetDynamicRayTracingInstances(FRayTracingMaterialG
 				RayTracingInstance.InstanceTransforms.Add(GetLocalToWorld());
 			}
 
-			if (bAnySegmentUsesWorldPositionOffset && MeshObject->GetSkinVertexFactory(Context.ReferenceView, LODIndex, 0)->GetType()->SupportsRayTracingDynamicGeometry())
+			const FVertexFactory* VertexFactory = MeshObject->GetSkinVertexFactory(Context.ReferenceView, LODIndex, 0);
+			const FVertexFactoryType* VertexFactoryType = VertexFactory->GetType();
+			if (bAnySegmentUsesWorldPositionOffset 
+				&& ensureMsgf(VertexFactoryType->SupportsRayTracingDynamicGeometry(),
+					TEXT("Mesh uses world position offset, but the vertex factory does not support ray tracing dynamic geometry. ")
+					TEXT("MeshObject: %s, VertexFactory: %s."), *MeshObject->DebugName.ToString(), VertexFactoryType->GetName()))
 			{
 				TArray<FRayTracingGeometrySegment> GeometrySections;
 				GeometrySections.Reserve(LODData.RenderSections.Num());
