Add mechanism for parallel processing FObjectCacheEventSink events without incurring the cost of thread sychronization.

diff --git a/Engine/Source/Runtime/Engine/Private/LevelTick.cpp b/Engine/Source/Runtime/Engine/Private/LevelTick.cpp
--- a/Engine/Source/Runtime/Engine/Private/LevelTick.cpp
+++ b/Engine/Source/Runtime/Engine/Private/LevelTick.cpp
@@ -68,6 +68,7 @@
 #if WITH_EDITOR
 #include "Editor.h"
 #include "LevelInstance/LevelInstanceSubsystem.h"
+#include "ObjectCacheEventSink.h"
 #endif
 
 CSV_DECLARE_CATEGORY_MODULE_EXTERN(CORE_API, Basic);
@@ -1039,6 +1040,7 @@ void EndSendEndOfFrameUpdatesDrawEvent(FSendAllEndOfFrameUpdates* SendAllEndOfFr
 	*/
 void UWorld::SendAllEndOfFrameUpdates()
 {
+	SCOPED_NAMED_EVENT(UWorld_SendAllEndOfFrameUpdates, FColor::Yellow);
 	SCOPE_CYCLE_COUNTER(STAT_PostTickComponentUpdate);
 	CSV_SCOPED_TIMING_STAT_EXCLUSIVE(EndOfFrameUpdates);
 	CSV_SCOPED_SET_WAIT_STAT(EndOfFrameUpdates);
@@ -1106,8 +1108,16 @@ void UWorld::SendAllEndOfFrameUpdates()
 		LocalComponentsThatNeedEndOfFrameUpdate.Append(ComponentsThatNeedEndOfFrameUpdate);
 	}
 
+#if WITH_EDITOR
+	std::atomic<uint32> ParallelTasksRemaining(0);
+#endif
+
 	auto ParallelWork = 
-		[](int32 Index) 
+		[
+#if WITH_EDITOR
+		&ParallelTasksRemaining
+#endif
+		](int32 Index)
 		{
 			FOptionalTaskTagScope Scope(ETaskTag::EParallelGameThread);
 			UActorComponent* NextComponent = LocalComponentsThatNeedEndOfFrameUpdate[Index];
@@ -1120,10 +1130,17 @@ void UWorld::SendAllEndOfFrameUpdates()
 				check(NextComponent->IsPendingKill() || NextComponent->GetMarkedForEndOfFrameUpdateState() == EComponentMarkedForEndOfFrameUpdateState::Marked);
 				FMarkComponentEndOfFrameUpdateState::Set(NextComponent, INDEX_NONE, EComponentMarkedForEndOfFrameUpdateState::Unmarked);
 			}
+#if WITH_EDITOR
+			ParallelTasksRemaining--;
+#endif
 		};
 
 	auto GTWork = 
-		[this]()
+		[this
+#if WITH_EDITOR
+		, &ParallelTasksRemaining
+#endif
+		]()
 		{
 			QUICK_SCOPE_CYCLE_COUNTER(STAT_PostTickComponentUpdate_ForcedGameThread);
 
@@ -1153,11 +1170,29 @@ void UWorld::SendAllEndOfFrameUpdates()
 			{
 				Component->DoDeferredRenderUpdates_Concurrent();
 			}
+
+#if WITH_EDITOR
+			while (ParallelTasksRemaining > 0)
+			{
+				FObjectCacheEventSink::ProcessQueuedNotifyEvents();
+			}
+#endif
 	};
 
-	if (CVarAllowAsyncRenderThreadUpdatesDuringGamethreadUpdates.GetValueOnGameThread() > 0)
+	if (CVarAllowAsyncRenderThreadUpdatesDuringGamethreadUpdates.GetValueOnGameThread() > 0 && 
+		LocalComponentsThatNeedEndOfFrameUpdate.Num() > FTaskGraphInterface::Get().GetNumWorkerThreads() &&
+		FTaskGraphInterface::Get().GetNumWorkerThreads() > 2)
 	{
+#if WITH_EDITOR
+		ParallelTasksRemaining = LocalComponentsThatNeedEndOfFrameUpdate.Num();
+		FObjectCacheEventSink::BeginQueueNotifyEvents();
+#endif
 		ParallelForWithPreWork(LocalComponentsThatNeedEndOfFrameUpdate.Num(), ParallelWork, GTWork);
+
+#if WITH_EDITOR
+		// Any remaining events will be flushed with this call
+		FObjectCacheEventSink::EndQueueNotifyEvents();
+#endif
 	}
 	else
 	{
diff --git a/Engine/Source/Runtime/Engine/Private/ObjectCacheContext.cpp b/Engine/Source/Runtime/Engine/Private/ObjectCacheContext.cpp
--- a/Engine/Source/Runtime/Engine/Private/ObjectCacheContext.cpp
+++ b/Engine/Source/Runtime/Engine/Private/ObjectCacheContext.cpp
@@ -20,6 +20,8 @@ DEFINE_LOG_CATEGORY(LogObjectCache);
 
 #if WITH_EDITOR
 #include "ObjectCacheEventSink.h"
+#include "Containers/LockFreeList.h"
+#include "Containers/LockFreeFixedSizeAllocator.h"
 
 static TAutoConsoleVariable<int32> CVarObjectReverseLookupMode(
 	TEXT("Editor.ObjectReverseLookupMode"),
@@ -681,18 +683,72 @@ FObjectCacheContext& FObjectCacheContextScope::GetContext()
 
 #if WITH_EDITOR
 
+struct FObjectCacheEventSinkPrivate
+{
+	static void BeginQueueNotifyEvents();
+	static void ProcessQueuedNotifyEvents();
+	static void EndQueueNotifyEvents();
+
+	static void NotifyUsedMaterialsChanged_Concurrent(const UPrimitiveComponent* PrimitiveComponent, const TArray<UMaterialInterface*>& UsedMaterials);
+	static void NotifyRenderStateChanged_Concurrent(const UPrimitiveComponent*);
+	static void NotifyReferencedTextureChanged_Concurrent(UMaterialInterface*);
+	static void NotifyStaticMeshChanged_Concurrent(UStaticMeshComponent*);
+	static void NotifyMaterialDestroyed_Concurrent(UMaterialInterface*);
+
+	enum ECacheEventType
+	{
+		EMaterialDestroyed,
+		EMaterialsChanged,
+		ERenderStateChanged,
+		EReferencedTextureChanged,
+		EStaticMeshChanged
+	};
+
+	struct FNotifyEvent
+	{
+		ECacheEventType EventType;
+		UMaterialInterface* MaterialInterface;
+		const UPrimitiveComponent* PrimitiveComponent;
+		UStaticMeshComponent* StaticMeshComponent;
+		TArray<UMaterialInterface*> UsedMaterials;
+	};
+
+	static std::atomic<bool> bShouldQueueSinkEvents;
+
+	using TEventAllocator = TLockFreeClassAllocator<FNotifyEvent, PLATFORM_CACHE_LINE_SIZE>;
+	using TEventList = TLockFreePointerListUnordered<FNotifyEvent, PLATFORM_CACHE_LINE_SIZE>;
+	static TEventAllocator& GetAllocator();
+	static TEventList& GetNotifyEvents();
+	static void AddNotifyEvent(FObjectCacheEventSinkPrivate::ECacheEventType EventType, UMaterialInterface* MaterialInterface, const UPrimitiveComponent* PrimitiveComponent, UStaticMeshComponent* StaticMeshComponent, const TArray<UMaterialInterface*>* UsedMaterials);
+};
+
+void FObjectCacheEventSink::BeginQueueNotifyEvents()
+{
+	FObjectCacheEventSinkPrivate::BeginQueueNotifyEvents();
+}
+
+void FObjectCacheEventSink::ProcessQueuedNotifyEvents()
+{
+	FObjectCacheEventSinkPrivate::ProcessQueuedNotifyEvents();
+}
+
+void FObjectCacheEventSink::EndQueueNotifyEvents()
+{
+	FObjectCacheEventSinkPrivate::EndQueueNotifyEvents();
+}
+
 void FObjectCacheEventSink::NotifyMaterialDestroyed_Concurrent(UMaterialInterface* MaterialInterface)
 {
 	if (GetObjectReverseLookupMode() != EObjectReverseLookupMode::Temporary)
 	{
-		// If a material is destroyed, remove it from the cache so we don't return it anymore
-		// This can happen if a component doesn't update its render state after modifying its used materials
-		// (i.e. LandscapeComponent won't update it's render state when modifying MobileMaterialInterfaces during cook since they aren't meaningful for rendering)
-		// The net result is that the SceneProxy->UsedMaterialsForVerification will continue to hold an invalid material interface object but if it's not used for rendering, no harm will be done...
-		GMaterialToPrimitiveLookupCache.RemoveFrom(MaterialInterface);
-
-		// Remove any Material to Texture that might be present in the cache for this Material
-		GTextureToMaterialLookupCache.Update(MaterialInterface, {});
+		if (FObjectCacheEventSinkPrivate::bShouldQueueSinkEvents)
+		{
+			FObjectCacheEventSinkPrivate::AddNotifyEvent(FObjectCacheEventSinkPrivate::EMaterialDestroyed, MaterialInterface, nullptr, nullptr, nullptr);
+		}
+		else
+		{
+			FObjectCacheEventSinkPrivate::NotifyMaterialDestroyed_Concurrent(MaterialInterface);
+		}
 	}
 }
 
@@ -700,7 +756,14 @@ void FObjectCacheEventSink::NotifyUsedMaterialsChanged_Concurrent(const UPrimiti
 {
 	if (GetObjectReverseLookupMode() != EObjectReverseLookupMode::Temporary)
 	{
-		GMaterialToPrimitiveLookupCache.Update(const_cast<UPrimitiveComponent*>(PrimitiveComponent), UsedMaterials);
+		if (FObjectCacheEventSinkPrivate::bShouldQueueSinkEvents)
+		{
+			FObjectCacheEventSinkPrivate::AddNotifyEvent(FObjectCacheEventSinkPrivate::EMaterialsChanged, nullptr, PrimitiveComponent, nullptr, &UsedMaterials);
+		}
+		else
+		{
+			FObjectCacheEventSinkPrivate::NotifyUsedMaterialsChanged_Concurrent(PrimitiveComponent, UsedMaterials);
+		}
 	}
 }
 
@@ -710,8 +773,14 @@ void FObjectCacheEventSink::NotifyRenderStateChanged_Concurrent(UActorComponent*
 	{
 		if (UPrimitiveComponent* Primitive = Cast<UPrimitiveComponent>(ActorComponent))
 		{
-			// Clear mappings whenever the render state changes until NotifyUsedMaterialsChanged is called during SceneProxy creation
-			GMaterialToPrimitiveLookupCache.Update(Primitive, {});
+			if (FObjectCacheEventSinkPrivate::bShouldQueueSinkEvents)
+			{
+				FObjectCacheEventSinkPrivate::AddNotifyEvent(FObjectCacheEventSinkPrivate::ERenderStateChanged, nullptr, Primitive, nullptr, nullptr);
+			}
+			else
+			{
+				FObjectCacheEventSinkPrivate::NotifyRenderStateChanged_Concurrent(Primitive);
+			}
 		}
 	}
 }
@@ -720,15 +789,13 @@ void FObjectCacheEventSink::NotifyReferencedTextureChanged_Concurrent(UMaterialI
 {
 	if (GetObjectReverseLookupMode() != EObjectReverseLookupMode::Temporary)
 	{
-		if (MaterialInterface->HasAnyFlags(RF_BeginDestroyed))
+		if (FObjectCacheEventSinkPrivate::bShouldQueueSinkEvents)
 		{
-			GTextureToMaterialLookupCache.Update(MaterialInterface, {});
+			FObjectCacheEventSinkPrivate::AddNotifyEvent(FObjectCacheEventSinkPrivate::EReferencedTextureChanged, MaterialInterface, nullptr, nullptr, nullptr);
 		}
 		else
 		{
-			TSet<UTexture*> Textures;
-			ObjectCacheContextImpl::GetReferencedTextures(MaterialInterface, Textures);
-			GTextureToMaterialLookupCache.Update(MaterialInterface, Textures.Array());
+			FObjectCacheEventSinkPrivate::NotifyReferencedTextureChanged_Concurrent(MaterialInterface);
 		}
 	}
 }
@@ -737,18 +804,144 @@ void FObjectCacheEventSink::NotifyStaticMeshChanged_Concurrent(UStaticMeshCompon
 {
 	if (GetObjectReverseLookupMode() != EObjectReverseLookupMode::Temporary)
 	{
-		// Stop tracking components that are being destroyed
-		if (StaticMeshComponent->HasAnyFlags(RF_BeginDestroyed) || StaticMeshComponent->GetStaticMesh() == nullptr)
+		if (FObjectCacheEventSinkPrivate::bShouldQueueSinkEvents)
 		{
-			GStaticMeshToComponentLookupCache.Update(StaticMeshComponent, {});
+			FObjectCacheEventSinkPrivate::AddNotifyEvent(FObjectCacheEventSinkPrivate::EStaticMeshChanged, nullptr, nullptr, StaticMeshComponent, nullptr);
 		}
-		else 
+		else
 		{
-			GStaticMeshToComponentLookupCache.Update(StaticMeshComponent, { StaticMeshComponent->GetStaticMesh() });
+			FObjectCacheEventSinkPrivate::NotifyStaticMeshChanged_Concurrent(StaticMeshComponent);
 		}
 	}
 }
 
+std::atomic<bool> FObjectCacheEventSinkPrivate::bShouldQueueSinkEvents(false);
+
+FObjectCacheEventSinkPrivate::TEventAllocator& FObjectCacheEventSinkPrivate::GetAllocator()
+{
+	static FObjectCacheEventSinkPrivate::TEventAllocator TheAllocator;
+	return TheAllocator;
+}
+
+FObjectCacheEventSinkPrivate::TEventList& FObjectCacheEventSinkPrivate::GetNotifyEvents()
+{
+	static TEventList NotifyEvents;
+	return NotifyEvents;
+}
+
+void FObjectCacheEventSinkPrivate::AddNotifyEvent(FObjectCacheEventSinkPrivate::ECacheEventType EventType, UMaterialInterface* MaterialInterface, const UPrimitiveComponent* PrimitiveComponent, UStaticMeshComponent* StaticMeshComponent, const TArray<UMaterialInterface*>* UsedMaterials)
+{
+
+	FNotifyEvent* Event = GetAllocator().New();
+	Event->EventType = EventType;
+	Event->MaterialInterface = MaterialInterface;
+	Event->PrimitiveComponent = PrimitiveComponent;
+	Event->StaticMeshComponent = StaticMeshComponent;
+
+	if (EventType == EMaterialsChanged)
+	{
+		check(UsedMaterials);
+		Event->UsedMaterials = *UsedMaterials;
+	}
+
+	GetNotifyEvents().Push(Event);
+}
+
+void FObjectCacheEventSinkPrivate::BeginQueueNotifyEvents()
+{
+	check(!bShouldQueueSinkEvents);
+	bShouldQueueSinkEvents = true;
+}
+
+void FObjectCacheEventSinkPrivate::ProcessQueuedNotifyEvents()
+{
+	check(bShouldQueueSinkEvents);
+
+	TEventList& Events = GetNotifyEvents();
+	TEventAllocator& Allocator = GetAllocator();
+
+	while (FNotifyEvent* Event = Events.Pop())
+	{
+		switch (Event->EventType)
+		{
+		case EMaterialDestroyed:
+			NotifyMaterialDestroyed_Concurrent(Event->MaterialInterface);
+			break;
+		case EMaterialsChanged:
+			NotifyUsedMaterialsChanged_Concurrent(Event->PrimitiveComponent, Event->UsedMaterials);
+			break;
+		case ERenderStateChanged:
+			NotifyRenderStateChanged_Concurrent(Event->PrimitiveComponent);
+			break;
+		case EReferencedTextureChanged:
+			NotifyReferencedTextureChanged_Concurrent(Event->MaterialInterface);
+			break;
+		case EStaticMeshChanged:
+			NotifyStaticMeshChanged_Concurrent(Event->StaticMeshComponent);
+			break;
+		}
+
+		Allocator.Free(Event);
+	}
+}
+
+void FObjectCacheEventSinkPrivate::EndQueueNotifyEvents()
+{
+	check(bShouldQueueSinkEvents);
+	ProcessQueuedNotifyEvents();
+	bShouldQueueSinkEvents = false;
+}
+
+void FObjectCacheEventSinkPrivate::NotifyMaterialDestroyed_Concurrent(UMaterialInterface* MaterialInterface)
+{
+	// If a material is destroyed, remove it from the cache so we don't return it anymore
+	// This can happen if a component doesn't update its render state after modifying its used materials
+	// (i.e. LandscapeComponent won't update it's render state when modifying MobileMaterialInterfaces during cook since they aren't meaningful for rendering)
+	// The net result is that the SceneProxy->UsedMaterialsForVerification will continue to hold an invalid material interface object but if it's not used for rendering, no harm will be done...
+	GMaterialToPrimitiveLookupCache.RemoveFrom(MaterialInterface);
+
+	// Remove any Material to Texture that might be present in the cache for this Material
+	GTextureToMaterialLookupCache.Update(MaterialInterface, {});
+}
+
+void FObjectCacheEventSinkPrivate::NotifyUsedMaterialsChanged_Concurrent(const UPrimitiveComponent* PrimitiveComponent, const TArray<UMaterialInterface*>& UsedMaterials)
+{
+	GMaterialToPrimitiveLookupCache.Update(const_cast<UPrimitiveComponent*>(PrimitiveComponent), UsedMaterials);
+}
+
+void FObjectCacheEventSinkPrivate::NotifyRenderStateChanged_Concurrent(const UPrimitiveComponent* PrimitiveComponent)
+{
+	// Clear mappings whenever the render state changes until NotifyUsedMaterialsChanged is called during SceneProxy creation
+	GMaterialToPrimitiveLookupCache.Update(const_cast<UPrimitiveComponent*>(PrimitiveComponent), {});
+}
+
+void FObjectCacheEventSinkPrivate::NotifyReferencedTextureChanged_Concurrent(UMaterialInterface* MaterialInterface)
+{
+	if (MaterialInterface->HasAnyFlags(RF_BeginDestroyed))
+	{
+		GTextureToMaterialLookupCache.Update(MaterialInterface, {});
+	}
+	else
+	{
+		TSet<UTexture*> Textures;
+		ObjectCacheContextImpl::GetReferencedTextures(MaterialInterface, Textures);
+		GTextureToMaterialLookupCache.Update(MaterialInterface, Textures.Array());
+	}
+}
+
+void FObjectCacheEventSinkPrivate::NotifyStaticMeshChanged_Concurrent(UStaticMeshComponent* StaticMeshComponent)
+{
+	// Stop tracking components that are being destroyed
+	if (StaticMeshComponent->HasAnyFlags(RF_BeginDestroyed) || StaticMeshComponent->GetStaticMesh() == nullptr)
+	{
+		GStaticMeshToComponentLookupCache.Update(StaticMeshComponent, {});
+	}
+	else
+	{
+		GStaticMeshToComponentLookupCache.Update(StaticMeshComponent, { StaticMeshComponent->GetStaticMesh() });
+	}
+}
+
 #endif
 
 #undef LOCTEXT_NAMESPACE
\ No newline at end of file
diff --git a/Engine/Source/Runtime/Engine/Public/ObjectCacheEventSink.h b/Engine/Source/Runtime/Engine/Public/ObjectCacheEventSink.h
--- a/Engine/Source/Runtime/Engine/Public/ObjectCacheEventSink.h
+++ b/Engine/Source/Runtime/Engine/Public/ObjectCacheEventSink.h
@@ -7,6 +7,13 @@
 
 struct FObjectCacheEventSink
 {
+	// For parallel processing.  Calling BeginQueueNotifyEvents will prevent notify events from being processed, an instead placed in a list for later processing.
+	// While other threads are actively updating the events, the main thread should process those events using this ProcessQueuedNotifyEvents
+	// When parallel processing is complete, calling EndQueueNotifyEvents will flush all events, and return to normal in place processing.
+	static void BeginQueueNotifyEvents();
+	static void ProcessQueuedNotifyEvents();
+	static void EndQueueNotifyEvents();
+
 	static void NotifyUsedMaterialsChanged_Concurrent(const UPrimitiveComponent* PrimitiveComponent, const TArray<UMaterialInterface*>& UsedMaterials);
 	static void NotifyRenderStateChanged_Concurrent(UActorComponent*);
 	static void NotifyReferencedTextureChanged_Concurrent(UMaterialInterface*);
