Allow Static mesh components to recreate their render state in parallel threads instead of requiring a serial recreate on the game thread.

diff --git a/Engine/Source/Runtime/Engine/Classes/Components/StaticMeshComponent.h b/Engine/Source/Runtime/Engine/Classes/Components/StaticMeshComponent.h
--- a/Engine/Source/Runtime/Engine/Classes/Components/StaticMeshComponent.h
+++ b/Engine/Source/Runtime/Engine/Classes/Components/StaticMeshComponent.h
@@ -477,6 +477,7 @@ public:
 protected: 
 	virtual void OnRegister() override;
 	virtual void OnUnregister() override;
+	virtual bool RequiresGameThreadEndOfFrameRecreate() const override;
 	virtual void CreateRenderState_Concurrent(FRegisterComponentContext* Context) override;
 	virtual void OnCreatePhysicsState() override;
 	virtual void OnDestroyPhysicsState() override;
diff --git a/Engine/Source/Runtime/Engine/Private/Components/StaticMeshComponent.cpp b/Engine/Source/Runtime/Engine/Private/Components/StaticMeshComponent.cpp
--- a/Engine/Source/Runtime/Engine/Private/Components/StaticMeshComponent.cpp
+++ b/Engine/Source/Runtime/Engine/Private/Components/StaticMeshComponent.cpp
@@ -689,6 +689,11 @@ void UStaticMeshComponent::OnUnregister()
 	Super::OnUnregister();
 }
 
+bool UStaticMeshComponent::RequiresGameThreadEndOfFrameRecreate() const
+{
+	return false;
+}
+
 void UStaticMeshComponent::CreateRenderState_Concurrent(FRegisterComponentContext* Context)
 {
 	LLM_SCOPE(ELLMTag::StaticMesh);
@@ -2347,7 +2352,7 @@ UMaterialInterface* UStaticMeshComponent::GetMaterial(int32 MaterialIndex) const
 
 void UStaticMeshComponent::GetUsedMaterials(TArray<UMaterialInterface*>& OutMaterials, bool bGetDebugMaterials) const
 {
-	TRACE_CPUPROFILER_EVENT_SCOPE(UStaticMeshComponent::GetUsedMaterials);
+	//TRACE_CPUPROFILER_EVENT_SCOPE(UStaticMeshComponent::GetUsedMaterials);
 
 	if (GetStaticMesh() && GetStaticMesh()->GetRenderData())
 	{
