Fixed compile error in shipping config due to FSkeletalMeshObject::DebugName not being available.

diff --git a/Engine/Source/Runtime/Engine/Private/SkeletalMesh.cpp b/Engine/Source/Runtime/Engine/Private/SkeletalMesh.cpp
--- a/Engine/Source/Runtime/Engine/Private/SkeletalMesh.cpp
+++ b/Engine/Source/Runtime/Engine/Private/SkeletalMesh.cpp
@@ -6257,7 +6257,7 @@ void FSkeletalMeshSceneProxy::GetDynamicRayTracingInstances(FRayTracingMaterialG
 			if (bAnySegmentUsesWorldPositionOffset 
 				&& ensureMsgf(VertexFactoryType->SupportsRayTracingDynamicGeometry(),
 					TEXT("Mesh uses world position offset, but the vertex factory does not support ray tracing dynamic geometry. ")
-					TEXT("MeshObject: %s, VertexFactory: %s."), *MeshObject->DebugName.ToString(), VertexFactoryType->GetName()))
+					TEXT("MeshObject: %s, VertexFactory: %s."), *MeshObject->GetDebugName().ToString(), VertexFactoryType->GetName()))
 			{
 				TArray<FRayTracingGeometrySegment> GeometrySections;
 				GeometrySections.Reserve(LODData.RenderSections.Num());
diff --git a/Engine/Source/Runtime/Engine/Public/SkeletalRenderPublic.h b/Engine/Source/Runtime/Engine/Public/SkeletalRenderPublic.h
--- a/Engine/Source/Runtime/Engine/Public/SkeletalRenderPublic.h
+++ b/Engine/Source/Runtime/Engine/Public/SkeletalRenderPublic.h
@@ -246,7 +246,10 @@ public:
 	bool bSupportRayTracing;
 #endif
 
-#if !UE_BUILD_SHIPPING
+#if UE_BUILD_SHIPPING
+	FName GetDebugName() const { return FName(); }
+#else
+	FName GetDebugName() const { return DebugName; }
 	FName DebugName;
 #endif // !UE_BUILD_SHIPPING
 
