WorldPartition: Fix crash TFunctionRef that was outliving its callable. (was crashing pie/cook of WorldPartition only on Mac/Linux)

diff --git a/Engine/Source/Runtime/Engine/Private/WorldPartition/WorldPartitionLevelHelper.cpp b/Engine/Source/Runtime/Engine/Private/WorldPartition/WorldPartitionLevelHelper.cpp
--- a/Engine/Source/Runtime/Engine/Private/WorldPartition/WorldPartitionLevelHelper.cpp
+++ b/Engine/Source/Runtime/Engine/Private/WorldPartition/WorldPartitionLevelHelper.cpp
@@ -108,8 +108,8 @@ void FWorldPartitionLevelHelper::RemapLevelSoftObjectPaths(ULevel* InLevel, UWor
 
 	struct FSoftPathFixupSerializer : public FArchiveUObject
 	{
-		FSoftPathFixupSerializer(TFunctionRef<void(FSoftObjectPath&)> InCustomFixupFunction)
-			: CustomFixupFunction(InCustomFixupFunction)
+		FSoftPathFixupSerializer(UWorldPartition* InWorldPartition)
+			: WorldPartition(InWorldPartition)
 		{
 			this->SetIsSaving(true);
 			this->ArShouldSkipBulkData = true;
@@ -119,15 +119,15 @@ void FWorldPartitionLevelHelper::RemapLevelSoftObjectPaths(ULevel* InLevel, UWor
 		{
 			if (!Value.IsNull())
 			{
-				CustomFixupFunction(Value);
+				WorldPartition->RemapSoftObjectPath(Value);
 			}
 			return *this;
 		}
 
-		TFunctionRef<void(FSoftObjectPath&)> CustomFixupFunction;
+		UWorldPartition* WorldPartition;
 	};
 
-	FSoftPathFixupSerializer FixupSerializer([InWorldPartition](FSoftObjectPath& ObjectPath) { InWorldPartition->RemapSoftObjectPath(ObjectPath); });
+	FSoftPathFixupSerializer FixupSerializer(InWorldPartition);
 	TArray<UObject*> SubObjects;
 	GetObjectsWithOuter(InLevel, SubObjects);
 	for (UObject* Object : SubObjects)
