UE-112780 - Crash loading project after deleting in-use structs

This crash is unique to UK2Node_SetFieldsInStruct since it uses the ReferencePassThroughConnection pin. Whenever we serialize a node, we need to reconcile the serialized pins with the default allocated ones. However, we don't actually allocate new pins for this node because we're missing the struct. RewireOldPinsToNewPins will effectively orphan the serialized nodes.

However, we might keep one pin around in certain cases. If the pin is an input and has a link, it'll hang around. If this same pin also happens to use ReferencePassThroughConnection, we run into our assert due to a bWasTrashed check. This is because ReferencePassThroughConnection points to a pin that was orphaned and subsequently trashed.

ReferencePassThroughConnection pins are bi-directional: if one pin is trashed, then we must unlink the other pin. In this case, we modify DestroyImpl to unlink any ReferencePassThroughConnection references. Additionally, an ensure was added to validate that the partner pin points to the one that is about to get trashed.

diff --git a/Engine/Source/Runtime/Engine/Private/EdGraph/EdGraphPin.cpp b/Engine/Source/Runtime/Engine/Private/EdGraph/EdGraphPin.cpp
--- a/Engine/Source/Runtime/Engine/Private/EdGraph/EdGraphPin.cpp
+++ b/Engine/Source/Runtime/Engine/Private/EdGraph/EdGraphPin.cpp
@@ -1579,7 +1579,16 @@ void UEdGraphPin::DestroyImpl(bool bClearLinks)
 	}
 	SubPins.Reset();
 	ParentPin = nullptr;
-	ReferencePassThroughConnection = nullptr;
+
+	// ReferencePassThroughConnection should be symmetrical: the source and target pins should refer to each other.
+	// If one pin is destroyed, then we must disconnect the other so that it doesn't inadvertently reference a trashed pin.
+	if (ReferencePassThroughConnection != nullptr)
+	{
+		ensure(ReferencePassThroughConnection->ReferencePassThroughConnection == this);
+		ReferencePassThroughConnection->ReferencePassThroughConnection = nullptr;
+		ReferencePassThroughConnection = nullptr;
+	}
+
 	bWasTrashed = true;
 }
 
