Fix stack overflow when dealing with circular references in the WorldPartitionResaveActorsBuilder

diff --git a/Engine/Source/Runtime/Engine/Private/WorldPartition/WorldPartitionHelpers.cpp b/Engine/Source/Runtime/Engine/Private/WorldPartition/WorldPartitionHelpers.cpp
--- a/Engine/Source/Runtime/Engine/Private/WorldPartition/WorldPartitionHelpers.cpp
+++ b/Engine/Source/Runtime/Engine/Private/WorldPartition/WorldPartitionHelpers.cpp
@@ -41,12 +41,14 @@ void FWorldPartitionHelpers::ForEachActorWithLoading(UWorldPartition* WorldParti
 
 		if (const FWorldPartitionActorDesc* ActorDesc = WorldPartition->GetActorDesc(ActorGuid))
 		{
+			FWorldPartitionReference& Reference = InOutActorReferences.Emplace(ActorGuid);
+			
 			for (FGuid ReferenceGuid : ActorDesc->GetReferences())
 			{
 				LoadReferences(ReferenceGuid, InOutActorReferences);
 			}
 
-			InOutActorReferences.Add(ActorGuid, FWorldPartitionReference(WorldPartition, ActorGuid));
+			Reference = FWorldPartitionReference(WorldPartition, ActorGuid);
 		}
 	};
 
