Improved WorldPartition debugging tools:
- Added wp.Runtime.ToggleDrawRuntimeCellsDetails : Lists all world partition runtime streaming cells with their name and grouped by streaming state
- Added world partition debug filters (by datalayer(s), streaming status, grid name(s), cell name(s)) that affects debug display (runtime hash 2d/3d, cell details)
- Added Data Layer streaming status in 2d/3d runtime hash debug draw
- Added Legends with colors and minimal stats for DataLayers and Runtime Streaming Status

diff --git a/Engine/Source/Runtime/Core/Public/ProfilingDebugging/ProfilingHelpers.h b/Engine/Source/Runtime/Core/Public/ProfilingDebugging/ProfilingHelpers.h
--- a/Engine/Source/Runtime/Core/Public/ProfilingDebugging/ProfilingHelpers.h
+++ b/Engine/Source/Runtime/Core/Public/ProfilingDebugging/ProfilingHelpers.h
@@ -22,7 +22,8 @@ enum EStreamingStatus
 	LEVEL_Visible,
 	LEVEL_Preloading,
 	LEVEL_FailedToLoad,
-	LEVEL_MakingInvisible
+	LEVEL_MakingInvisible,
+	LEVEL_StreamingStatusCount
 };
 
 
diff --git a/Engine/Source/Runtime/Engine/Classes/Engine/LevelStreaming.h b/Engine/Source/Runtime/Engine/Classes/Engine/LevelStreaming.h
--- a/Engine/Source/Runtime/Engine/Classes/Engine/LevelStreaming.h
+++ b/Engine/Source/Runtime/Engine/Classes/Engine/LevelStreaming.h
@@ -458,6 +458,9 @@ public:
 	/** Utility that returns a string for a streaming level status */
 	static const TCHAR* GetLevelStreamingStatusDisplayName(EStreamingStatus Status);
 
+	/** Utility that draws a legend of level streaming status */
+	static void DebugDrawLegend(const UWorld* World, class UCanvas* Canvas, const FVector2D& Offset);
+
 #if WITH_EDITOR
 	/** Get the folder path for this level for use in the world browser. Only available in editor builds */
 	const FName& GetFolderPath() const;
diff --git a/Engine/Source/Runtime/Engine/Private/WorldPartition/DataLayer/DataLayer.cpp b/Engine/Source/Runtime/Engine/Private/WorldPartition/DataLayer/DataLayer.cpp
--- a/Engine/Source/Runtime/Engine/Private/WorldPartition/DataLayer/DataLayer.cpp
+++ b/Engine/Source/Runtime/Engine/Private/WorldPartition/DataLayer/DataLayer.cpp
@@ -22,6 +22,7 @@ UDataLayer::UDataLayer(const FObjectInitializer& ObjectInitializer)
 , DataLayerLabel(GetFName())
 , InitialState(EDataLayerState::Unloaded)
 , bIsDynamicallyLoaded(false)
+, DebugColor(FColor::Black)
 {
 }
 
@@ -40,6 +41,15 @@ void UDataLayer::PostLoad()
 
 	// Sanitize Label
 	DataLayerLabel = UDataLayer::GetSanitizedDataLayerLabel(DataLayerLabel);
+
+	if (DebugColor == FColor::Black)
+	{
+		FRandomStream RandomStream(GetFName());
+		const uint8 R = (uint8)(RandomStream.GetFraction() * 255.f);
+		const uint8 G = (uint8)(RandomStream.GetFraction() * 255.f);
+		const uint8 B = (uint8)(RandomStream.GetFraction() * 255.f);
+		DebugColor = FColor(R, G, B);
+	}
 #endif
 }
 
diff --git a/Engine/Source/Runtime/Engine/Private/WorldPartition/DataLayer/DataLayerSubsystem.cpp b/Engine/Source/Runtime/Engine/Private/WorldPartition/DataLayer/DataLayerSubsystem.cpp
--- a/Engine/Source/Runtime/Engine/Private/WorldPartition/DataLayer/DataLayerSubsystem.cpp
+++ b/Engine/Source/Runtime/Engine/Private/WorldPartition/DataLayer/DataLayerSubsystem.cpp
@@ -3,10 +3,10 @@
 #include "WorldPartition/DataLayer/DataLayerSubsystem.h"
 #include "WorldPartition/DataLayer/WorldDataLayers.h"
 #include "WorldPartition/DataLayer/DataLayer.h"
+#include "WorldPartition/WorldPartitionDebugHelper.h"
 #include "Engine/Engine.h"
 #include "Engine/World.h"
 #include "Engine/Canvas.h"
-#include "Debug/DebugDrawService.h"
 #if WITH_EDITOR
 #include "Editor.h"
 #include "Modules/ModuleManager.h"
@@ -68,22 +68,6 @@ void UDataLayerSubsystem::PostInitialize()
 			});
 		}
 	}
-
-	if (GetWorld()->IsGameWorld() && (GetWorld()->GetNetMode() != NM_DedicatedServer))
-	{
-		DrawHandle = UDebugDrawService::Register(TEXT("Game"), FDebugDrawDelegate::CreateUObject(this, &UDataLayerSubsystem::Draw));
-	}
-}
-
-void UDataLayerSubsystem::Deinitialize()
-{
-	if (DrawHandle.IsValid())
-	{
-		UDebugDrawService::Unregister(DrawHandle);
-		DrawHandle.Reset();
-	}
-
-	Super::Deinitialize();
 }
 
 UDataLayer* UDataLayerSubsystem::GetDataLayer(const FActorDataLayer& InDataLayer) const
@@ -226,44 +210,37 @@ bool UDataLayerSubsystem::IsAnyDataLayerInState(const TArray<FName>& InDataLayer
 	return false;
 }
 
-void UDataLayerSubsystem::Draw(UCanvas* Canvas, class APlayerController* PC)
+void UDataLayerSubsystem::DrawDataLayersStatus(UCanvas* Canvas, FVector2D& Offset) const
 {
 	if (!GDrawDataLayers || !Canvas || !Canvas->SceneView)
 	{
 		return;
 	}
+	
+	FVector2D Pos = Offset;
+	float MaxTextWidth = 0.f;
 
-	auto DrawText = [](UCanvas* Canvas, const FString& Text, UFont* Font, const FColor& Color, FVector2D& Pos)
-	{
-		float TextWidth, TextHeight;
-		Canvas->StrLen(Font, Text, TextWidth, TextHeight);
-		Canvas->SetDrawColor(Color);
-		Canvas->DrawText(Font, Text, Pos.X, Pos.Y);
-		Pos.Y += TextHeight + 1;
-	};
-
-	const FVector2D CanvasTopLeftPadding(10.f, 10.f);
-	FVector2D Pos = CanvasTopLeftPadding;
-
-	auto DrawLayerNames = [this, &Pos, DrawText, Canvas](const FString& Title, FColor Color, const TSet<FName>& LayerNames)
+	auto DrawLayerNames = [this, Canvas, &Pos, &MaxTextWidth](const FString& Title, FColor Color, const TSet<FName>& LayerNames)
 	{
 		if (LayerNames.Num() > 0)
 		{
-			DrawText(Canvas, Title, GEngine->GetSmallFont(), FColor::Yellow, Pos);
+			FWorldPartitionDebugHelper::DrawText(Canvas, Title, GEngine->GetSmallFont(), FColor::Yellow, Pos, &MaxTextWidth);
 			UFont* DataLayerFont = GEngine->GetTinyFont();
 			for (const FName& DataLayerName : LayerNames)
 			{
 				if (UDataLayer* DataLayer = GetDataLayerFromName(DataLayerName))
 				{
 					const FString Text = DataLayer->GetDataLayerLabel().ToString();
-					DrawText(Canvas, Text, DataLayerFont, FColor::White, Pos);
+					FWorldPartitionDebugHelper::DrawText(Canvas, Text, DataLayerFont, Color, Pos, &MaxTextWidth);
 				}
 			}
 		}
 	};
 
-	DrawLayerNames(TEXT("Loaded Data Layers"), FColor::Green, LoadedDataLayerNames);
-	DrawLayerNames(TEXT("Active Data Layers"), FColor::White, ActiveDataLayerNames);
+	DrawLayerNames(TEXT("Loaded Data Layers"), FColor::Cyan, LoadedDataLayerNames);
+	DrawLayerNames(TEXT("Active Data Layers"), FColor::Green, ActiveDataLayerNames);
+
+	Offset.X += MaxTextWidth + 10;
 }
 
 TArray<UDataLayer*> UDataLayerSubsystem::ConvertArgsToDataLayers(UWorld* World, const TArray<FString>& InArgs)
@@ -362,3 +339,49 @@ FAutoConsoleCommand UDataLayerSubsystem::ToggleDataLayerActivation(
 		}
 	})
 );
+
+void UDataLayerSubsystem::GetDataLayerDebugColors(TMap<FName, FColor>& OutMapping) const
+{
+	OutMapping.Reset();
+
+	const AWorldDataLayers* WorldDataLayers = AWorldDataLayers::Get(GetWorld());
+	if (!WorldDataLayers)
+	{
+		return;
+	}
+	
+	WorldDataLayers->ForEachDataLayer([&OutMapping](UDataLayer* DataLayer)
+	{
+		OutMapping.Add(DataLayer->GetFName(), DataLayer->GetDebugColor());
+		return true;
+	});
+}
+
+void UDataLayerSubsystem::DrawDataLayersLegend(UCanvas* Canvas, FVector2D& Offset) const
+{
+	check(Canvas);
+
+	const AWorldDataLayers* WorldDataLayers = AWorldDataLayers::Get(GetWorld());
+	if (!WorldDataLayers)
+	{
+		return;
+	}
+
+	TMap<FName, FColor> ColorMapping;
+	GetDataLayerDebugColors(ColorMapping);
+
+	float MaxItemWidth = 0;
+	FVector2D Pos = Offset;
+	FWorldPartitionDebugHelper::DrawText(Canvas, TEXT("DataLayers Legend"), GEngine->GetSmallFont(), FColor::Yellow, Pos, &MaxItemWidth);
+	
+	for (auto& It : ColorMapping)
+	{
+		if (const UDataLayer* DataLayer = WorldDataLayers->GetDataLayerFromName(It.Key))
+		{
+			FString StateString = StaticEnum<EDataLayerState>()->GetDisplayNameTextByValue((int64)GetDataLayerState(DataLayer)).ToString();
+			FWorldPartitionDebugHelper::DrawLegendItem(Canvas, *FString::Printf(TEXT("%s (%s)"), *DataLayer->GetDataLayerLabel().ToString(), *StateString), GEngine->GetTinyFont(), It.Value, Pos, &MaxItemWidth);
+		}
+	}
+
+	Offset.X += MaxItemWidth + 10;
+}
\ No newline at end of file
diff --git a/Engine/Source/Runtime/Engine/Private/WorldPartition/WorldPartition.cpp b/Engine/Source/Runtime/Engine/Private/WorldPartition/WorldPartition.cpp
--- a/Engine/Source/Runtime/Engine/Private/WorldPartition/WorldPartition.cpp
+++ b/Engine/Source/Runtime/Engine/Private/WorldPartition/WorldPartition.cpp
@@ -1114,10 +1114,10 @@ FVector2D UWorldPartition::GetDrawRuntimeHash2DDesiredFootprint(const FVector2D&
 	return StreamingPolicy->GetDrawRuntimeHash2DDesiredFootprint(CanvasSize);
 }
 
-void UWorldPartition::DrawRuntimeHash2D(class UCanvas* Canvas, const FVector2D& PartitionCanvasOffset, const FVector2D& PartitionCanvasSize)
+void UWorldPartition::DrawRuntimeHash2D(class UCanvas* Canvas, const FVector2D& PartitionCanvasSize, FVector2D& Offset)
 {
 	check(CanDrawRuntimeHash());
-	StreamingPolicy->DrawRuntimeHash2D(Canvas, PartitionCanvasOffset, PartitionCanvasSize);
+	StreamingPolicy->DrawRuntimeHash2D(Canvas, PartitionCanvasSize, Offset);
 }
 
 void UWorldPartition::DrawRuntimeHash3D()
@@ -1126,6 +1126,16 @@ void UWorldPartition::DrawRuntimeHash3D()
 	StreamingPolicy->DrawRuntimeHash3D();
 }
 
+void UWorldPartition::DrawRuntimeCellsDetails(class UCanvas* Canvas, FVector2D& Offset)
+{
+	StreamingPolicy->DrawRuntimeCellsDetails(Canvas, Offset);
+}
+
+void UWorldPartition::DrawStreamingStatusLegend(class UCanvas* Canvas, FVector2D& Offset)
+{
+	StreamingPolicy->DrawStreamingStatusLegend(Canvas, Offset);
+}
+
 #if WITH_EDITOR
 void UWorldPartition::DrawRuntimeHashPreview()
 {
diff --git a/Engine/Source/Runtime/Engine/Private/WorldPartition/WorldPartitionDebugHelper.cpp b/Engine/Source/Runtime/Engine/Private/WorldPartition/WorldPartitionDebugHelper.cpp
new file mode 100644
--- /dev/null
+++ b/Engine/Source/Runtime/Engine/Private/WorldPartition/WorldPartitionDebugHelper.cpp
@@ -0,0 +1,173 @@
+// Copyright Epic Games, Inc. All Rights Reserved.
+
+#include "WorldPartition/WorldPartitionDebugHelper.h"
+#include "WorldPartition/WorldPartitionRuntimeHash.h"
+#include "WorldPartition/WorldPartition.h"
+#include "WorldPartition/DataLayer/DataLayer.h"
+#include "WorldPartition/DataLayer/DataLayerSubsystem.h"
+#include "ProfilingDebugging/ProfilingHelpers.h"
+#include "Engine/Canvas.h"
+#include "Engine/Font.h"
+#include "HAL/IConsoleManager.h"
+#include "Algo/AnyOf.h"
+#include "Algo/AllOf.h"
+
+TSet<FName> FWorldPartitionDebugHelper::DebugRuntimeHashFilter;
+FAutoConsoleCommand FWorldPartitionDebugHelper::DebugFilterByRuntimeHashGridNameCommand(
+	TEXT("wp.Runtime.DebugFilterByRuntimeHashGridName"),
+	TEXT("Filter debug diplay of world partition streaming by grid name. Args [grid names]"),
+	FConsoleCommandWithArgsDelegate::CreateLambda([](const TArray<FString>& InArgs)
+	{
+		DebugRuntimeHashFilter.Reset();
+		for (const FWorldContext& Context : GEngine->GetWorldContexts())
+		{
+			UWorld* World = Context.World();
+			if (World && World->IsGameWorld())
+			{
+				if (UWorldPartition* WorldPartition = World->GetWorldPartition())
+				{
+					for (const FString& Arg : InArgs)
+					{
+						if (WorldPartition->RuntimeHash && WorldPartition->RuntimeHash->ContainsRuntimeHash(Arg))
+						{
+							DebugRuntimeHashFilter.Add(FName(Arg));
+						}
+					}
+				}
+			}
+		}
+	})
+);
+
+bool FWorldPartitionDebugHelper::IsDebugRuntimeHashGridShown(FName Name)
+{
+	return !DebugRuntimeHashFilter.Num() || DebugRuntimeHashFilter.Contains(Name);
+}
+
+TSet<FName> FWorldPartitionDebugHelper::DebugDataLayerFilter;
+FAutoConsoleCommand FWorldPartitionDebugHelper::DebugFilterByDataLayerCommand(
+	TEXT("wp.Runtime.DebugFilterByDataLayer"),
+	TEXT("Filter debug diplay of world partition streaming by data layer. Args [datalayer labels]"),
+	FConsoleCommandWithArgsDelegate::CreateLambda([](const TArray<FString>& InArgs)
+	{
+		DebugDataLayerFilter.Reset();
+
+		for (const FWorldContext& Context : GEngine->GetWorldContexts())
+		{
+			UWorld* World = Context.World();
+			if (World && World->IsGameWorld())
+			{
+				TArray<UDataLayer*> DataLayers = UDataLayerSubsystem::ConvertArgsToDataLayers(World, InArgs);
+				for (UDataLayer* DataLayer : DataLayers)
+				{
+					DebugDataLayerFilter.Add(DataLayer->GetFName());
+				}
+			}
+		}
+
+		if (Algo::AnyOf(InArgs, [](const FString& Arg) { return FName(Arg) == NAME_None; }))
+		{
+			DebugDataLayerFilter.Add(NAME_None);
+		}
+	})
+);
+
+bool FWorldPartitionDebugHelper::IsDebugDataLayerShown(FName DataLayerName)
+{
+	return !DebugDataLayerFilter.Num() || DebugDataLayerFilter.Contains(DataLayerName);
+}
+
+bool FWorldPartitionDebugHelper::AreDebugDataLayersShown(const TArray<FName>& DataLayerNames)
+{
+	if (DebugDataLayerFilter.Num())
+	{
+		bool bFilter = !DataLayerNames.IsEmpty() || !DebugDataLayerFilter.Contains(NAME_None);
+		for (const FName& DataLayerName : DataLayerNames)
+		{
+			if (DebugDataLayerFilter.Contains(DataLayerName))
+			{
+				bFilter = false;
+				break;
+			}
+		}
+		if (bFilter)
+		{
+			return false;
+		}
+	}
+	return true;
+}
+
+TSet<EStreamingStatus> FWorldPartitionDebugHelper::DebugStreamingStatusFilter;
+FAutoConsoleCommand FWorldPartitionDebugHelper::DebugFilterByStreamingStatusCommand(
+	TEXT("wp.Runtime.DebugFilterByStreamingStatus"),
+	TEXT("Filter debug diplay of world partition streaming by streaming status. Args [streaming status]"),
+	FConsoleCommandWithArgsDelegate::CreateLambda([](const TArray<FString>& InArgs)
+	{
+		DebugStreamingStatusFilter.Reset();
+		for (const FString& Arg : InArgs)
+		{
+			int32 Result = 0;
+			TTypeFromString<int32>::FromString(Result, *Arg);
+			if (Result >= 0 && Result < (int32)LEVEL_StreamingStatusCount)
+			{
+				DebugStreamingStatusFilter.Add((EStreamingStatus)Result);
+			}
+		}
+	})
+);
+
+bool FWorldPartitionDebugHelper::IsDebugStreamingStatusShown(EStreamingStatus StreamingStatus)
+{
+	return !DebugStreamingStatusFilter.Num() || DebugStreamingStatusFilter.Contains(StreamingStatus);
+}
+
+TArray<FString> FWorldPartitionDebugHelper::DebugCellNameFilter;
+FAutoConsoleCommand FWorldPartitionDebugHelper::DebugFilterByCellNameCommand(
+	TEXT("wp.Runtime.DebugFilterByCellName"),
+	TEXT("Filter debug diplay of world partition streaming by full or partial cell name. Args [cell name]"),
+	FConsoleCommandWithArgsDelegate::CreateLambda([](const TArray<FString>& InArgs)
+	{
+		TSet<FString> Filter;
+		Filter.Append(InArgs);
+		DebugCellNameFilter = Filter.Array();
+	})
+);
+
+bool FWorldPartitionDebugHelper::IsDebugCellNameShow(const FString& CellName)
+{
+	return !DebugCellNameFilter.Num() || Algo::AllOf(DebugCellNameFilter, [&CellName](const FString& NameFilter) { return CellName.Contains(NameFilter); });
+}
+
+void FWorldPartitionDebugHelper::DrawText(UCanvas* Canvas, const FString& Text, const UFont* Font, const FColor& Color, FVector2D& Pos, float* MaxTextWidth)
+{
+	float TextWidth, TextHeight;
+	Canvas->StrLen(Font, Text, TextWidth, TextHeight);
+	Canvas->SetDrawColor(Color);
+	Canvas->DrawText(Font, Text, Pos.X, Pos.Y);
+	Pos.Y += TextHeight + 1;
+	if (MaxTextWidth)
+	{
+		*MaxTextWidth = FMath::Max(*MaxTextWidth, TextWidth);
+	}
+}
+
+void FWorldPartitionDebugHelper::DrawLegendItem(UCanvas* Canvas, const FString& Text, const UFont* Font, const FColor& Color, FVector2D& Pos, float* MaxItemWidth)
+{
+	static const FVector2D ItemSize(12, 12);
+	
+	float MaxTextWidth = 0.f;
+	
+	FCanvasTileItem Item(Pos, GWhiteTexture, ItemSize, Color);
+	Canvas->DrawItem(Item);
+
+	FVector2D TextPos(Pos.X + ItemSize.X + 10, Pos.Y);
+	float TextWidth = 0.f;
+	FWorldPartitionDebugHelper::DrawText(Canvas, Text, Font, Color, TextPos, &TextWidth);
+	if (MaxItemWidth)
+	{
+		*MaxItemWidth = FMath::Max(*MaxItemWidth, TextWidth + ItemSize.X + 10);
+	}
+
+	Pos.Y = TextPos.Y;
+}
\ No newline at end of file
diff --git a/Engine/Source/Runtime/Engine/Private/WorldPartition/WorldPartitionLevelStreamingDynamic.cpp b/Engine/Source/Runtime/Engine/Private/WorldPartition/WorldPartitionLevelStreamingDynamic.cpp
--- a/Engine/Source/Runtime/Engine/Private/WorldPartition/WorldPartitionLevelStreamingDynamic.cpp
+++ b/Engine/Source/Runtime/Engine/Private/WorldPartition/WorldPartitionLevelStreamingDynamic.cpp
@@ -45,6 +45,7 @@ UWorldPartitionLevelStreamingDynamic::UWorldPartitionLevelStreamingDynamic(const
  */
 void UWorldPartitionLevelStreamingDynamic::Initialize(const UWorldPartitionRuntimeLevelStreamingCell& InCell)
 {
+	StreamingCell = &InCell;
 	UWorld* World = GetWorld();
 	check(!ShouldBeLoaded());
 	check((World->IsGameWorld() && !ShouldBeVisible()) || (!World->IsGameWorld() && !GetShouldBeVisibleFlag()));
diff --git a/Engine/Source/Runtime/Engine/Private/WorldPartition/WorldPartitionLevelStreamingPolicy.cpp b/Engine/Source/Runtime/Engine/Private/WorldPartition/WorldPartitionLevelStreamingPolicy.cpp
--- a/Engine/Source/Runtime/Engine/Private/WorldPartition/WorldPartitionLevelStreamingPolicy.cpp
+++ b/Engine/Source/Runtime/Engine/Private/WorldPartition/WorldPartitionLevelStreamingPolicy.cpp
@@ -11,8 +11,10 @@
 #include "WorldPartition/WorldPartitionLevelStreamingDynamic.h"
 #include "WorldPartition/WorldPartitionRuntimeHash.h"
 #include "WorldPartition/WorldPartitionLevelHelper.h"
+#include "WorldPartition/WorldPartitionDebugHelper.h"
 #include "Engine/Level.h"
 #include "Engine/World.h"
+#include "Engine/Canvas.h"
 
 #if WITH_EDITOR
 #include "Misc/PackageName.h"
@@ -27,22 +29,30 @@ static FAutoConsoleVariableRef CMaxLoadingLevelStreamingCells(
 int32 UWorldPartitionLevelStreamingPolicy::GetCellLoadingCount() const
 {
 	int32 CellLoadingCount = 0;
-	for (const UWorldPartitionRuntimeCell* ActivatedCell : ActivatedCells)
+
+	ForEachActiveRuntimeCell([&CellLoadingCount](const UWorldPartitionRuntimeCell* Cell)
 	{
-		const UWorldPartitionRuntimeLevelStreamingCell* Cell = Cast<const UWorldPartitionRuntimeLevelStreamingCell>(ActivatedCell);
-		if (!Cell->IsAlwaysLoaded())
+		if (Cell->IsLoading())
 		{
-			if (UWorldPartitionLevelStreamingDynamic* LevelStreaming = Cell->GetLevelStreaming())
+			++CellLoadingCount;
+		}
+	});
+	return CellLoadingCount;
+}
+
+void UWorldPartitionLevelStreamingPolicy::ForEachActiveRuntimeCell(TFunctionRef<void(const UWorldPartitionRuntimeCell*)> Func) const
+{
+	UWorld* World = WorldPartition->GetWorld();
+	for (ULevelStreaming* LevelStreaming : World->GetStreamingLevels())
+	{
+		if (UWorldPartitionLevelStreamingDynamic* WorldPartitionLevelStreaming = Cast<UWorldPartitionLevelStreamingDynamic>(LevelStreaming))
+		{
+			if (const UWorldPartitionRuntimeCell* Cell = WorldPartitionLevelStreaming->GetWorldPartitionRuntimeCell())
 			{
-				ULevelStreaming::ECurrentState CurrentState = LevelStreaming->GetCurrentState();
-				if (CurrentState == ULevelStreaming::ECurrentState::Removed || CurrentState == ULevelStreaming::ECurrentState::Unloaded || CurrentState == ULevelStreaming::ECurrentState::Loading)
-				{
-					++CellLoadingCount;
-				}
+				Func(Cell);
 			}
 		}
 	}
-	return CellLoadingCount;
 }
 
 int32 UWorldPartitionLevelStreamingPolicy::GetMaxCellsToLoad() const
@@ -308,4 +318,87 @@ UObject* UWorldPartitionLevelStreamingPolicy::GetSubObject(const TCHAR* SubObjec
 	}
 
 	return nullptr;
+}
+
+void UWorldPartitionLevelStreamingPolicy::DrawRuntimeCellsDetails(UCanvas* Canvas, FVector2D& Offset)
+{
+	UWorld* World = WorldPartition->GetWorld();
+	struct FCellsPerStreamingStatus
+	{
+		TArray<const UWorldPartitionRuntimeCell*> Cells;
+	};
+	FCellsPerStreamingStatus CellsPerStreamingStatus[(int32)LEVEL_StreamingStatusCount];
+	ForEachActiveRuntimeCell([&CellsPerStreamingStatus](const UWorldPartitionRuntimeCell* Cell)
+	{
+		if (Cell->IsDebugShown())
+		{
+			CellsPerStreamingStatus[(int32)Cell->GetStreamingStatus()].Cells.Add(Cell);
+		}
+	});
+
+	FVector2D Pos = Offset;
+	const float BaseY = Offset.Y;
+
+	float CurrentColumnWidth = 0.f;
+	float MaxPosY = Pos.Y;
+
+	auto DrawCellDetails = [&](const FString& Text, const UFont* Font, const FColor& Color)
+	{
+		FWorldPartitionDebugHelper::DrawText(Canvas, Text, Font, Color, Pos, &CurrentColumnWidth);
+		MaxPosY = FMath::Max(MaxPosY, Pos.Y);
+		if ((Pos.Y + 30) > Canvas->ClipY)
+		{
+			Pos.Y = BaseY;
+			Pos.X += CurrentColumnWidth + 5;
+			CurrentColumnWidth = 0.f;
+		}
+	};
+
+	for (int32 i = 0; i < (int32)LEVEL_StreamingStatusCount; ++i)
+	{
+		const EStreamingStatus StreamingStatus = (EStreamingStatus)i;
+		const TArray<const UWorldPartitionRuntimeCell*>& Cells = CellsPerStreamingStatus[i].Cells;
+		if (Cells.Num() > 0)
+		{
+			const FString StatusDisplayName = *FString::Printf(TEXT("%s (%d)"), ULevelStreaming::GetLevelStreamingStatusDisplayName(StreamingStatus), Cells.Num());
+			DrawCellDetails(StatusDisplayName, GEngine->GetSmallFont(), FColor::Yellow);
+
+			const FColor Color = ULevelStreaming::GetLevelStreamingStatusColor(StreamingStatus);
+			for (const UWorldPartitionRuntimeCell* Cell : Cells)
+			{
+				DrawCellDetails(Cell->GetDebugName(), GEngine->GetTinyFont(), Color);
+			}
+		}
+	}
+
+	Offset.Y = MaxPosY;
+}
+
+/**
+ * Debug Draw Streaming Status Legend
+ */
+void UWorldPartitionLevelStreamingPolicy::DrawStreamingStatusLegend(UCanvas* Canvas, FVector2D& Offset)
+{
+	check(Canvas);
+
+	// Cumulate counter stats
+	int32 StatusCount[(int32)LEVEL_StreamingStatusCount] = { 0 };
+	ForEachActiveRuntimeCell([&StatusCount](const UWorldPartitionRuntimeCell* Cell)
+	{
+		StatusCount[(int32)Cell->GetStreamingStatus()]++;
+	});
+
+	// Draw legend
+	FVector2D Pos = Offset;
+	float MaxTextWidth = 0.f;
+	FWorldPartitionDebugHelper::DrawText(Canvas, TEXT("Streaming Status Legend"), GEngine->GetSmallFont(), FColor::Yellow, Pos, &MaxTextWidth);
+	
+	for (int32 i = 0; i < (int32)LEVEL_StreamingStatusCount; ++i)
+	{
+		EStreamingStatus Status = (EStreamingStatus)i;
+		const FColor& StatusColor = ULevelStreaming::GetLevelStreamingStatusColor(Status);
+		FWorldPartitionDebugHelper::DrawLegendItem(Canvas, *FString::Printf(TEXT("%d) %s (%d)"), i, ULevelStreaming::GetLevelStreamingStatusDisplayName(Status), StatusCount[(int32)Status]), GEngine->GetTinyFont(), StatusColor, Pos, &MaxTextWidth);
+	}
+
+	Offset.X += MaxTextWidth + 10;
 }
\ No newline at end of file
diff --git a/Engine/Source/Runtime/Engine/Private/WorldPartition/WorldPartitionRuntimeCell.cpp b/Engine/Source/Runtime/Engine/Private/WorldPartition/WorldPartitionRuntimeCell.cpp
--- a/Engine/Source/Runtime/Engine/Private/WorldPartition/WorldPartitionRuntimeCell.cpp
+++ b/Engine/Source/Runtime/Engine/Private/WorldPartition/WorldPartitionRuntimeCell.cpp
@@ -1,6 +1,13 @@
 // Copyright Epic Games, Inc. All Rights Reserved.
 
 #include "WorldPartition/WorldPartitionRuntimeCell.h"
+#include "WorldPartition/WorldPartition.h"
+#include "WorldPartition/WorldPartitionDebugHelper.h"
+#include "WorldPartition/WorldPartitionRuntimeSpatialHash.h"
+#include "WorldPartition/WorldPartitionLevelStreamingPolicy.h"
+#include "WorldPartition/DataLayer/WorldDataLayers.h"
+#include "WorldPartition/DataLayer/DataLayersID.h"
+#include "Engine/World.h"
 
 UWorldPartitionRuntimeCell::UWorldPartitionRuntimeCell(const FObjectInitializer& ObjectInitializer)
 : Super(ObjectInitializer)
@@ -14,7 +21,7 @@ void UWorldPartitionRuntimeCell::AddCellData(const UWorldPartitionRuntimeCellDat
 	CellDataMap.Add(InCellData->GetClass(), InCellData);
 }
 
-void UWorldPartitionRuntimeCell::SetDataLayers(const TArray<const UDataLayer*> InDataLayers)
+void UWorldPartitionRuntimeCell::SetDataLayers(const TArray<const UDataLayer*>& InDataLayers)
 {
 	check(DataLayers.IsEmpty());
 	DataLayers.Reserve(InDataLayers.Num());
@@ -23,10 +30,52 @@ void UWorldPartitionRuntimeCell::SetDataLayers(const TArray<const UDataLayer*> I
 		check(DataLayer->IsDynamicallyLoaded());
 		DataLayers.Add(DataLayer->GetFName());
 	}
+	DataLayers.Sort([](const FName& A, const FName& B) { return A.ToString() < B.ToString(); });
+	UpdateDebugName();
+}
+
+void UWorldPartitionRuntimeCell::SetDebugInfo(FIntVector InCoords, FName InGridName)
+{
+	Coords = InCoords;
+	GridName = InGridName;
+	UpdateDebugName();
+}
+
+void UWorldPartitionRuntimeCell::UpdateDebugName()
+{
+	TStringBuilder<512> Builder;
+	Builder += GridName.ToString();
+	Builder += TEXT("_");
+	Builder += FString::Printf(TEXT("L%d_X%d_Y%d"), Coords.Z, Coords.X, Coords.Y);
+	int32 DataLayerCount = DataLayers.Num();
+
+	const AWorldDataLayers* WorldDataLayers = AWorldDataLayers::Get(GetOuterUWorldPartition()->GetWorld());
+	TArray<const UDataLayer*> DataLayerObjects;
+	if (WorldDataLayers && (DataLayerCount > 0))
+	{
+		Builder += TEXT(" DL[");
+		for (int i = 0; i < DataLayerCount; ++i)
+		{
+			const UDataLayer* DataLayer = WorldDataLayers->GetDataLayerFromName(DataLayers[i]);
+			DataLayerObjects.Add(DataLayer);
+			Builder += DataLayer->GetDataLayerLabel().ToString();
+			Builder += TEXT(",");
+		}
+		Builder += FString::Printf(TEXT("ID:%X]"), FDataLayersID(DataLayerObjects).GetHash());
+	}
+	DebugName = Builder.ToString();
 }
 
 #endif
 
+bool UWorldPartitionRuntimeCell::IsDebugShown() const
+{
+	return FWorldPartitionDebugHelper::IsDebugRuntimeHashGridShown(GridName) &&
+		   FWorldPartitionDebugHelper::IsDebugStreamingStatusShown(GetStreamingStatus()) &&
+		   FWorldPartitionDebugHelper::AreDebugDataLayersShown(DataLayers) &&
+		   FWorldPartitionDebugHelper::IsDebugCellNameShow(DebugName);
+}
+
 const UWorldPartitionRuntimeCellData* UWorldPartitionRuntimeCell::GetCellData(const TSubclassOf<UWorldPartitionRuntimeCellData> InCellDataClass) const
 {
 	return CellDataMap.FindRef(InCellDataClass);
@@ -34,4 +83,4 @@ const UWorldPartitionRuntimeCellData* UWorldPartitionRuntimeCell::GetCellData(co
 
 UWorldPartitionRuntimeCellData::UWorldPartitionRuntimeCellData(const FObjectInitializer& ObjectInitializer)
 : Super(ObjectInitializer)
-{}
+{}
\ No newline at end of file
diff --git a/Engine/Source/Runtime/Engine/Private/WorldPartition/WorldPartitionRuntimeLevelStreamingCell.cpp b/Engine/Source/Runtime/Engine/Private/WorldPartition/WorldPartitionRuntimeLevelStreamingCell.cpp
--- a/Engine/Source/Runtime/Engine/Private/WorldPartition/WorldPartitionRuntimeLevelStreamingCell.cpp
+++ b/Engine/Source/Runtime/Engine/Private/WorldPartition/WorldPartitionRuntimeLevelStreamingCell.cpp
@@ -45,15 +45,28 @@ UWorldPartitionLevelStreamingDynamic* UWorldPartitionRuntimeLevelStreamingCell::
 	return LevelStreaming;
 }
 
-EStreamingStatus UWorldPartitionRuntimeLevelStreamingCell::GetLevelStreamingStatus() const
+EStreamingStatus UWorldPartitionRuntimeLevelStreamingCell::GetStreamingStatus() const
 {
-	check(LevelStreaming)
-	return LevelStreaming->GetLevelStreamingStatus();
+	if (LevelStreaming)
+	{
+		return LevelStreaming->GetLevelStreamingStatus();
+	}
+	return Super::GetStreamingStatus();
+}
+
+bool UWorldPartitionRuntimeLevelStreamingCell::IsLoading() const
+{
+	if (LevelStreaming)
+	{
+		ULevelStreaming::ECurrentState CurrentState = LevelStreaming->GetCurrentState();
+		return (CurrentState == ULevelStreaming::ECurrentState::Removed || CurrentState == ULevelStreaming::ECurrentState::Unloaded || CurrentState == ULevelStreaming::ECurrentState::Loading);
+	}
+	return Super::IsLoading();
 }
 
 FLinearColor UWorldPartitionRuntimeLevelStreamingCell::GetDebugColor() const
 {
-	FLinearColor Color = LevelStreaming ? ULevelStreaming::GetLevelStreamingStatusColor(GetLevelStreamingStatus()) : FLinearColor::Black;
+	FLinearColor Color = LevelStreaming ? ULevelStreaming::GetLevelStreamingStatusColor(GetStreamingStatus()) : FLinearColor::Black;
 	Color.A = 0.25f / (Level + 1);
 	return Color;
 }
@@ -293,4 +306,4 @@ void UWorldPartitionRuntimeLevelStreamingCell::OnLevelHidden()
 {
 	check(LevelStreaming);
 	LevelStreaming->GetWorld()->GetSubsystem<UHLODSubsystem>()->OnCellHidden(this);
-}
+}
\ No newline at end of file
diff --git a/Engine/Source/Runtime/Engine/Private/WorldPartition/WorldPartitionRuntimeSpatialHash.cpp b/Engine/Source/Runtime/Engine/Private/WorldPartition/WorldPartitionRuntimeSpatialHash.cpp
--- a/Engine/Source/Runtime/Engine/Private/WorldPartition/WorldPartitionRuntimeSpatialHash.cpp
+++ b/Engine/Source/Runtime/Engine/Private/WorldPartition/WorldPartitionRuntimeSpatialHash.cpp
@@ -10,6 +10,7 @@
 #include "WorldPartition/WorldPartitionActorDescView.h"
 #include "WorldPartition/WorldPartitionStreamingPolicy.h"
 #include "WorldPartition/WorldPartitionHandle.h"
+#include "WorldPartition/WorldPartitionDebugHelper.h"
 
 #include "WorldPartition/RuntimeSpatialHash/RuntimeSpatialHashGridHelper.h"
 
@@ -30,6 +31,7 @@
 #include "DisplayDebugHelpers.h"
 #include "RenderUtils.h"
 #include "Algo/Transform.h"
+#include "Algo/Copy.h"
 #include "UObject/ObjectSaveContext.h"
 
 #if WITH_EDITOR
@@ -63,12 +65,6 @@ static FAutoConsoleVariableRef CVarShowRuntimeSpatialHashGridLevelCount(
 	GShowRuntimeSpatialHashGridLevelCount,
 	TEXT("Used to choose how many grid levels to display when showing world partition runtime hash."));
 
-static int32 GShowRuntimeSpatialHashGridIndex = 0;
-static FAutoConsoleVariableRef CVarShowRuntimeSpatialHashGridIndex(
-	TEXT("wp.Runtime.ShowRuntimeSpatialHashGridIndex"),
-	GShowRuntimeSpatialHashGridIndex,
-	TEXT("Used to show only one particular grid when showing world partition runtime hash (invalid index will show all)."));
-
 static float GRuntimeSpatialHashCellToSourceAngleContributionToCellImportance = 0.4f; // Value between [0, 1]
 static FAutoConsoleVariableRef CVarRuntimeSpatialHashCellToSourceAngleContributionToCellImportance(
 	TEXT("wp.Runtime.RuntimeSpatialHashCellToSourceAngleContributionToCellImportance"),
@@ -225,7 +221,21 @@ void FSpatialHashStreamingGrid::GetAlwaysLoadedCells(const UDataLayerSubsystem*
 	}
 }
 
-void FSpatialHashStreamingGrid::Draw3D(UWorld* World, const TArray<FWorldPartitionStreamingSource>& Sources, const FTransform& Transform) const
+void FSpatialHashStreamingGrid::GetFilteredCellsForDebugDraw(const FSpatialHashStreamingGridLayerCell* LayerCell, TArray<const UWorldPartitionRuntimeCell*>& FilteredCells) const
+{
+	FilteredCells.Reset();
+	if (LayerCell)
+	{
+		Algo::CopyIf(LayerCell->GridCells, FilteredCells, [](const UWorldPartitionRuntimeCell* GridCell) { return GridCell->IsDebugShown(); });
+	}
+	if (FilteredCells.IsEmpty())
+	{
+		const UWorldPartitionRuntimeCell* DefaultEmptyCell = UWorldPartitionRuntimeCell::StaticClass()->GetDefaultObject<UWorldPartitionRuntimeCell>();
+		FilteredCells.Add(DefaultEmptyCell);
+	}
+}
+
+void FSpatialHashStreamingGrid::Draw3D(UWorld* World, const TArray<FWorldPartitionStreamingSource>& Sources, const FTransform& Transform, const TMap<FName, FColor>& DataLayerDebugColors) const
 {
 	const FSquare2DGridHelper& Helper = GetGridHelper();
 	int32 MinGridLevel = FMath::Clamp<int32>(GShowRuntimeSpatialHashGridLevel, 0, GridLevels.Num() - 1);
@@ -234,6 +244,7 @@ void FSpatialHashStreamingGrid::Draw3D(UWorld* World, const TArray<FWorldPartiti
 	const float GridViewLoadingRangeExtentRatio = 1.5f;
 	const float Radius = GetLoadingRange();
 	const float GridSideDistance = FMath::Max((2.f * Radius * GridViewLoadingRangeExtentRatio), CellSize * GridViewMinimumSizeInCellCount);
+	TArray<const UWorldPartitionRuntimeCell*> FilteredCells;
 
 	for (const FWorldPartitionStreamingSource& Source : Sources)
 	{
@@ -252,32 +263,44 @@ void FSpatialHashStreamingGrid::Draw3D(UWorld* World, const TArray<FWorldPartiti
 		{
 			Helper.Levels[GridLevel].ForEachIntersectingCells(Region, [&](const FIntVector2& Coords)
 			{
+				const int32* LayerCellIndexPtr = GridLevels[GridLevel].LayerCellsMapping.Find(Coords.Y * Helper.Levels[GridLevel].GridSize + Coords.X);
+				const FSpatialHashStreamingGridLayerCell* LayerCell = LayerCellIndexPtr ? &GridLevels[GridLevel].LayerCells[*LayerCellIndexPtr] : nullptr;
+				GetFilteredCellsForDebugDraw(LayerCell, FilteredCells);
+				check(FilteredCells.Num());
+
 				FBox2D CellWorldBounds;
 				Helper.Levels[GridLevel].GetCellBounds(FIntVector2(Coords.X, Coords.Y), CellWorldBounds);
-				
+				float CellSizeY = CellWorldBounds.GetSize().Y / FilteredCells.Num();
+				CellWorldBounds.Max.Y = CellWorldBounds.Min.Y + CellSizeY;
 				FVector BoundsExtent(CellWorldBounds.GetExtent(), 100.f);
+				FVector BoundsOrigin(CellWorldBounds.GetCenter(), Z);
+				FBox CellBox = FBox::BuildAABB(BoundsOrigin, BoundsExtent);
+				FTranslationMatrix CellOffsetMatrix(FVector(0.f, CellSizeY, 0.f));
 
-				const UWorldPartitionRuntimeSpatialHashCell* Cell = nullptr;
-				if (const int32* LayerCellIndexPtr = GridLevels[GridLevel].LayerCellsMapping.Find(Coords.Y * Helper.Levels[GridLevel].GridSize + Coords.X))
+				for (const UWorldPartitionRuntimeCell* Cell : FilteredCells)
 				{
-					const FSpatialHashStreamingGridLayerCell& LayerCell = GridLevels[GridLevel].LayerCells[*LayerCellIndexPtr];
-					
-					for (const UWorldPartitionRuntimeSpatialHashCell* GridCell : LayerCell.GridCells)
+					// Draw Cell using its debug color
+					FColor CellColor = Cell->GetDebugColor().ToFColor(false).WithAlpha(16);
+					DrawDebugSolidBox(World, CellBox, CellColor, Transform, false, -1.f, 255);
+					FVector CellPos = Transform.TransformPosition(CellBox.GetCenter());
+					DrawDebugBox(World, CellPos, BoundsExtent, Transform.GetRotation(), CellColor.WithAlpha(255), false, -1.f, 255, 10.f);
+
+					// Draw Cell's DataLayers colored boxes
+					if (DataLayerDebugColors.Num() && Cell->GetDataLayers().Num() > 0)
 					{
-						Cell = GridCell;
-						if (!GridCell->HasDataLayers())
+						FBox DataLayerColoredBox = CellBox;
+						float DataLayerSizeX = DataLayerColoredBox.GetSize().X / (5 * Cell->GetDataLayers().Num()); // Use 20% of cell's width
+						DataLayerColoredBox.Max.X = DataLayerColoredBox.Min.X + DataLayerSizeX; 
+						FTranslationMatrix DataLayerOffsetMatrix(FVector(DataLayerSizeX, 0, 0));
+						for (const FName& DataLayer : Cell->GetDataLayers())
 						{
-							break;
+							const FColor& DataLayerColor = DataLayerDebugColors[DataLayer];
+							DrawDebugSolidBox(World, DataLayerColoredBox, DataLayerColor, Transform, false, -1.f, 255);
+							DataLayerColoredBox = DataLayerColoredBox.TransformBy(DataLayerOffsetMatrix);
 						}
 					}
+					CellBox = CellBox.TransformBy(CellOffsetMatrix);
 				}
-
-				FColor CellColor = Cell ? Cell->GetDebugColor().ToFColor(false).WithAlpha(16) : FColor(0, 0, 0, 16);
-				FVector BoundsOrigin(CellWorldBounds.GetCenter(), Z);
-				FBox Box = FBox::BuildAABB(BoundsOrigin, BoundsExtent);
-				DrawDebugSolidBox(World, Box, CellColor, Transform, false, -1.f, 255);
-				FVector Pos = Transform.TransformPosition(BoundsOrigin);
-				DrawDebugBox(World, Pos, BoundsExtent, Transform.GetRotation(), CellColor.WithAlpha(255), false, -1.f, 255, 10.f);
 			});
 		}
 
@@ -288,7 +311,7 @@ void FSpatialHashStreamingGrid::Draw3D(UWorld* World, const TArray<FWorldPartiti
 	}
 }
 
-void FSpatialHashStreamingGrid::Draw2D(UCanvas* Canvas, const TArray<FWorldPartitionStreamingSource>& Sources, const FBox& Region, const FBox2D& GridScreenBounds, TFunctionRef<FVector2D(const FVector2D&)> WorldToScreen) const
+void FSpatialHashStreamingGrid::Draw2D(UCanvas* Canvas, const TArray<FWorldPartitionStreamingSource>& Sources, const FBox& Region, const FBox2D& GridScreenBounds, TFunctionRef<FVector2D(const FVector2D&)> WorldToScreen, const TMap<FName, FColor>& DataLayerDebugColors) const
 {
 	FCanvas* CanvasObject = Canvas->Canvas;
 	int32 MinGridLevel = FMath::Clamp<int32>(GShowRuntimeSpatialHashGridLevel, 0, GridLevels.Num() - 1);
@@ -300,6 +323,8 @@ void FSpatialHashStreamingGrid::Draw2D(UCanvas* Canvas, const TArray<FWorldParti
 	float MaxCellCoordTextWidth, MaxCellCoordTextHeight;
 	Canvas->StrLen(GEngine->GetTinyFont(), CellCoordString, MaxCellCoordTextWidth, MaxCellCoordTextHeight);
 
+	const UWorldPartitionRuntimeCell* DefaultEmptyCell = UWorldPartitionRuntimeCell::StaticClass()->GetDefaultObject<UWorldPartitionRuntimeCell>();
+	TArray<const UWorldPartitionRuntimeCell*> FilteredCells;
 	for (int32 GridLevel = MinGridLevel; GridLevel <= MaxGridLevel; ++GridLevel)
 	{
 		// Draw X/Y Axis
@@ -359,29 +384,45 @@ void FSpatialHashStreamingGrid::Draw2D(UCanvas* Canvas, const TArray<FWorldParti
 				}
 			}
 
-			const UWorldPartitionRuntimeSpatialHashCell* Cell = nullptr;
-			if (const int32* LayerCellIndexPtr = GridLevels[GridLevel].LayerCellsMapping.Find(Coords.Y * Helper.Levels[GridLevel].GridSize + Coords.X))
+			// Draw cell bounds
+			static FLinearColor CellDefaultColor = DefaultEmptyCell->GetDebugColor().CopyWithNewOpacity(0.25f);
+			FCanvasBoxItem Box(CellScreenBounds.Min, CellScreenBounds.GetSize());
+			Box.SetColor(CellDefaultColor);
+			Box.BlendMode = SE_BLEND_Translucent;
+			Canvas->DrawItem(Box);
+
+			const int32* LayerCellIndexPtr = GridLevels[GridLevel].LayerCellsMapping.Find(Coords.Y * Helper.Levels[GridLevel].GridSize + Coords.X);
+			const FSpatialHashStreamingGridLayerCell* LayerCell = LayerCellIndexPtr ? &GridLevels[GridLevel].LayerCells[*LayerCellIndexPtr] : nullptr;
+			GetFilteredCellsForDebugDraw(LayerCell, FilteredCells);
+			check(FilteredCells.Num());
+
+			FVector2D CellBoundsSize = CellScreenBounds.GetSize();
+			CellBoundsSize.Y /= FilteredCells.Num();
+			FVector2D CellOffset(0.f, 0.f);
+			for (const UWorldPartitionRuntimeCell* Cell : FilteredCells)
 			{
-				const FSpatialHashStreamingGridLayerCell& LayerCell = GridLevels[GridLevel].LayerCells[*LayerCellIndexPtr];
-				for (const UWorldPartitionRuntimeSpatialHashCell* GridCell : LayerCell.GridCells)
+				// Draw Cell using its debug color
+				FVector2D StartPos = CellScreenBounds.Min + CellOffset;
+				FCanvasTileItem Item(StartPos, GWhiteTexture, CellBoundsSize, Cell->GetDebugColor().CopyWithNewOpacity(0.25f));
+				Item.BlendMode = SE_BLEND_Translucent;
+				Canvas->DrawItem(Item);
+				CellOffset.Y += CellBoundsSize.Y;
+
+				// Draw Cell's DataLayers colored boxes
+				if (DataLayerDebugColors.Num() && Cell->GetDataLayers().Num() > 0)
 				{
-					Cell = GridCell;
-					if (!GridCell->HasDataLayers())
+					FVector2D DataLayerOffset(0, 0);
+					FVector2D DataLayerColoredBoxSize = CellBoundsSize;
+					DataLayerColoredBoxSize.X /= (5 * Cell->GetDataLayers().Num()); // Use 20% of cell's width
+					for (const FName& DataLayer : Cell->GetDataLayers())
 					{
-						break;
+						const FColor& DataLayerColor = DataLayerDebugColors[DataLayer];
+						FCanvasTileItem DataLayerItem(StartPos + DataLayerOffset, GWhiteTexture, DataLayerColoredBoxSize, DataLayerColor);
+						Canvas->DrawItem(DataLayerItem);
+						DataLayerOffset.X += DataLayerColoredBoxSize.X;
 					}
 				}
 			}
-
-			FLinearColor CellColor = Cell ? Cell->GetDebugColor() : FLinearColor(0.f, 0.f, 0.f, 0.25f);
-			FCanvasTileItem Item(CellScreenBounds.Min, GWhiteTexture, CellScreenBounds.GetSize(), CellColor);
-			Item.BlendMode = SE_BLEND_Translucent;
-			Canvas->DrawItem(Item);
-
-			FCanvasBoxItem Box(CellScreenBounds.Min, CellScreenBounds.GetSize());
-			Box.SetColor(CellColor);
-			Box.BlendMode = SE_BLEND_Translucent;
-			Canvas->DrawItem(Box);
 		});
 
 		// Draw Loading Ranges
@@ -794,6 +835,7 @@ bool UWorldPartitionRuntimeSpatialHash::CreateStreamingGrid(const FSpatialHashRu
 				FBox2D Bounds;
 				verify(TempLevel.GetCellBounds(FIntVector2(CellCoordX, CellCoordY), Bounds));
 				StreamingCell->Position = FVector(Bounds.GetCenter(), 0.f);
+				StreamingCell->SetDebugInfo(CellGlobalCoords, CurrentStreamingGrid.GridName);
 
 				UE_LOG(LogWorldPartition, Verbose, TEXT("Cell%s %s Actors = %d Bounds (%s)"), bIsCellAlwaysLoaded ? TEXT(" (AlwaysLoaded)") : TEXT(""), *StreamingCell->GetName(), FilteredActors.Num(), *Bounds.ToString());
 
@@ -1057,35 +1099,42 @@ void UWorldPartitionRuntimeSpatialHash::SortStreamingCellsByImportance(const TSe
 
 FVector2D UWorldPartitionRuntimeSpatialHash::GetDraw2DDesiredFootprint(const FVector2D& CanvasSize) const
 {
-	return FVector2D(CanvasSize.X * StreamingGrids.Num(), CanvasSize.Y);
+	return FVector2D(CanvasSize.X * GetFilteredStreamingGrids().Num(), CanvasSize.Y);
 }
 
-void UWorldPartitionRuntimeSpatialHash::Draw2D(UCanvas* Canvas, const TArray<FWorldPartitionStreamingSource>& Sources, const FVector2D& PartitionCanvasOffset, const FVector2D& PartitionCanvasSize) const
+void UWorldPartitionRuntimeSpatialHash::Draw2D(UCanvas* Canvas, const TArray<FWorldPartitionStreamingSource>& Sources, const FVector2D& PartitionCanvasSize, FVector2D& Offset) const
 {
-	if (StreamingGrids.Num() == 0 || Sources.Num() == 0)
+	TArray<const FSpatialHashStreamingGrid*> FilteredStreamingGrids = GetFilteredStreamingGrids();
+	if (FilteredStreamingGrids.Num() == 0 || Sources.Num() == 0)
 	{
 		return;
 	}
 
 	UWorldPartition* WorldPartition = GetOuterUWorldPartition();
+	UWorld* World = WorldPartition->GetWorld();
+	UDataLayerSubsystem* DataLayerSubsystem = World->GetSubsystem<UDataLayerSubsystem>();
+
+	TMap<FName, FColor> DataLayerDebugColors;
+	if (DataLayerSubsystem)
+	{
+		DataLayerSubsystem->GetDataLayerDebugColors(DataLayerDebugColors);
+	}
 
 	const float CanvasMaxScreenSize = PartitionCanvasSize.X;
-	const float GridMaxScreenSize = CanvasMaxScreenSize / StreamingGrids.Num();
+	const float GridMaxScreenSize = CanvasMaxScreenSize / FilteredStreamingGrids.Num();
 	const float GridEffectiveScreenRatio = 1.f;
 	const float GridEffectiveScreenSize = FMath::Min(GridMaxScreenSize, PartitionCanvasSize.Y) - 10.f;
 	const float GridViewLoadingRangeExtentRatio = 1.5f;
 	const float GridViewMinimumSizeInCellCount = 5.f;
 	const FVector2D GridScreenExtent = FVector2D(GridEffectiveScreenSize, GridEffectiveScreenSize);
 	const FVector2D GridScreenHalfExtent = 0.5f * GridScreenExtent;
-	const FVector2D GridScreenInitialOffset = PartitionCanvasOffset;
+	const FVector2D GridScreenInitialOffset = Offset;
 
 	// Sort streaming grids to render them sorted by loading range
-	TArray<const FSpatialHashStreamingGrid*> SortedStreamingGrids;
-	Algo::Transform(StreamingGrids, SortedStreamingGrids, [](const FSpatialHashStreamingGrid& StreamingGrid) { return &StreamingGrid; });
-	SortedStreamingGrids.Sort([](const FSpatialHashStreamingGrid& A, const FSpatialHashStreamingGrid& B) { return A.LoadingRange < B.LoadingRange; });
+	FilteredStreamingGrids.Sort([](const FSpatialHashStreamingGrid& A, const FSpatialHashStreamingGrid& B) { return A.LoadingRange < B.LoadingRange; });
 
 	int32 GridIndex = 0;
-	for (const FSpatialHashStreamingGrid* StreamingGrid : SortedStreamingGrids)
+	for (const FSpatialHashStreamingGrid* StreamingGrid : FilteredStreamingGrids)
 	{
 		// Display view sides based on extended grid loading range (minimum of N cells)
 		// Take into consideration GShowRuntimeSpatialHashGridLevel when using CellSize
@@ -1104,7 +1153,8 @@ void UWorldPartitionRuntimeSpatialHash::Draw2D(UCanvas* Canvas, const TArray<FWo
 		const float WorldToScreenScale = (0.5f * GridEffectiveScreenSize) / AverageSphere.W;
 		auto WorldToScreen = [&](const FVector2D& WorldPos) { return (WorldToScreenScale * (WorldPos - GridReferenceWorldPos)) + GridScreenOffset; };
 
-		StreamingGrid->Draw2D(Canvas, Sources, Region, GridScreenBounds, WorldToScreen);
+		StreamingGrid->Draw2D(Canvas, Sources, Region, GridScreenBounds, WorldToScreen, DataLayerDebugColors);
+		Offset = GridScreenBounds.Max;
 
 		// Draw WorldPartition name
 		FVector2D GridInfoPos = GridScreenOffset - GridScreenHalfExtent;
@@ -1131,21 +1181,43 @@ void UWorldPartitionRuntimeSpatialHash::Draw2D(UCanvas* Canvas, const TArray<FWo
 void UWorldPartitionRuntimeSpatialHash::Draw3D(const TArray<FWorldPartitionStreamingSource>& Sources) const
 {
 	UWorld* World = GetWorld();
-
 	UWorldPartition* WorldPartition = GetOuterUWorldPartition();
-	FTransform Transform = WorldPartition->GetInstanceTransform();
+	UDataLayerSubsystem* DataLayerSubsystem = World->GetSubsystem<UDataLayerSubsystem>();
 
-	if (StreamingGrids.IsValidIndex(GShowRuntimeSpatialHashGridIndex))
+	TMap<FName, FColor> DataLayerDebugColors;
+	if (DataLayerSubsystem)
 	{
-		StreamingGrids[GShowRuntimeSpatialHashGridIndex].Draw3D(World, Sources, Transform);
+		DataLayerSubsystem->GetDataLayerDebugColors(DataLayerDebugColors);
 	}
-	else
+
+	FTransform Transform = WorldPartition->GetInstanceTransform();
+	TArray<const FSpatialHashStreamingGrid*> FilteredStreamingGrids = GetFilteredStreamingGrids();
+	for (const FSpatialHashStreamingGrid* StreamingGrid : FilteredStreamingGrids)
 	{
-		for (const FSpatialHashStreamingGrid& StreamingGrid : StreamingGrids)
+		StreamingGrid->Draw3D(World, Sources, Transform, DataLayerDebugColors);
+	}
+}
+
+bool UWorldPartitionRuntimeSpatialHash::ContainsRuntimeHash(const FString& Name) const
+{
+	for (const FSpatialHashStreamingGrid& StreamingGrid : StreamingGrids)
+	{
+		if (StreamingGrid.GridName.ToString().Equals(Name, ESearchCase::IgnoreCase))
 		{
-			StreamingGrid.Draw3D(World, Sources, Transform);
+			return true;
 		}
 	}
+	return false;
+}
+
+TArray<const FSpatialHashStreamingGrid*> UWorldPartitionRuntimeSpatialHash::GetFilteredStreamingGrids() const
+{
+	TArray<const FSpatialHashStreamingGrid*> FilteredStreamingGrids;
+	FilteredStreamingGrids.Reserve(StreamingGrids.Num());
+	Algo::TransformIf(StreamingGrids, FilteredStreamingGrids,
+		[](const FSpatialHashStreamingGrid& StreamingGrid) { return FWorldPartitionDebugHelper::IsDebugRuntimeHashGridShown(StreamingGrid.GridName); },
+		[](const FSpatialHashStreamingGrid& StreamingGrid) { return &StreamingGrid; });
+	return FilteredStreamingGrids;
 }
 
 #undef LOCTEXT_NAMESPACE
\ No newline at end of file
diff --git a/Engine/Source/Runtime/Engine/Private/WorldPartition/WorldPartitionStreamingPolicy.cpp b/Engine/Source/Runtime/Engine/Private/WorldPartition/WorldPartitionStreamingPolicy.cpp
--- a/Engine/Source/Runtime/Engine/Private/WorldPartition/WorldPartitionStreamingPolicy.cpp
+++ b/Engine/Source/Runtime/Engine/Private/WorldPartition/WorldPartitionStreamingPolicy.cpp
@@ -259,11 +259,11 @@ FVector2D UWorldPartitionStreamingPolicy::GetDrawRuntimeHash2DDesiredFootprint(c
 	return WorldPartition->RuntimeHash->GetDraw2DDesiredFootprint(CanvasSize);
 }
 
-void UWorldPartitionStreamingPolicy::DrawRuntimeHash2D(class UCanvas* Canvas, const FVector2D& PartitionCanvasOffset, const FVector2D& PartitionCanvasSize)
+void UWorldPartitionStreamingPolicy::DrawRuntimeHash2D(class UCanvas* Canvas, const FVector2D& PartitionCanvasSize, FVector2D& Offset)
 {
 	if (StreamingSources.Num() > 0)
 	{
-		WorldPartition->RuntimeHash->Draw2D(Canvas, StreamingSources, PartitionCanvasOffset, PartitionCanvasSize);
+		WorldPartition->RuntimeHash->Draw2D(Canvas, StreamingSources, PartitionCanvasSize, Offset);
 	}
 }
 
diff --git a/Engine/Source/Runtime/Engine/Private/WorldPartition/WorldPartitionSubsystem.cpp b/Engine/Source/Runtime/Engine/Private/WorldPartition/WorldPartitionSubsystem.cpp
--- a/Engine/Source/Runtime/Engine/Private/WorldPartition/WorldPartitionSubsystem.cpp
+++ b/Engine/Source/Runtime/Engine/Private/WorldPartition/WorldPartitionSubsystem.cpp
@@ -3,6 +3,8 @@
 #include "WorldPartition/WorldPartitionSubsystem.h"
 #include "WorldPartition/WorldPartition.h"
 #include "WorldPartition/WorldPartitionStreamingPolicy.h"
+#include "WorldPartition/WorldPartitionDebugHelper.h"
+#include "WorldPartition/DataLayer/DataLayerSubsystem.h"
 #include "HAL/IConsoleManager.h"
 #include "Engine/Canvas.h"
 #include "Engine/Console.h"
@@ -29,6 +31,18 @@ static FAutoConsoleCommand CVarDrawStreamingSources(
 	TEXT("Toggles debug display of world partition streaming sources."),
 	FConsoleCommandDelegate::CreateLambda([] { GDrawStreamingSources = !GDrawStreamingSources; }));
 
+static int32 GDrawLegends = 0;
+static FAutoConsoleCommand CVarGDrawLegends(
+	TEXT("wp.Runtime.ToggleDrawLegends"),
+	TEXT("Toggles debug display of world partition legends."),
+	FConsoleCommandDelegate::CreateLambda([] { GDrawLegends = !GDrawLegends; }));
+
+static int32 GDrawRuntimeCellsDetails = 0;
+static FAutoConsoleCommand CVarDrawRuntimeCellsDetails(
+	TEXT("wp.Runtime.ToggleDrawRuntimeCellsDetails"),
+	TEXT("Toggles debug display of world partition runtime streaming cells."),
+	FConsoleCommandDelegate::CreateLambda([] { GDrawRuntimeCellsDetails = !GDrawRuntimeCellsDetails; }));
+
 UWorldPartitionSubsystem::UWorldPartitionSubsystem()
 {}
 
@@ -209,23 +223,15 @@ void UWorldPartitionSubsystem::Draw(UCanvas* Canvas, class APlayerController* PC
 	}
 
 	const FVector2D CanvasTopLeftPadding(10.f, 10.f);
-	FVector2D Pos = CanvasTopLeftPadding;
 
-	auto DrawText = [](UCanvas* Canvas, const FString& Text, UFont* Font, const FColor& Color, FVector2D& Pos)
-	{
-		float TextWidth, TextHeight;
-		Canvas->StrLen(Font, Text, TextWidth, TextHeight);
-		Canvas->SetDrawColor(Color);
-		Canvas->DrawText(Font, Text, Pos.X, Pos.Y);
-		Pos.Y += TextHeight + 1;
-	};
+	FVector2D CurrentOffset(CanvasTopLeftPadding);
 
 	if (GDrawRuntimeHash2D)
 	{
 		const float MaxScreenRatio = 0.75f;
 		const FVector2D CanvasBottomRightPadding(10.f, 10.f);
 		const FVector2D CanvasMinimumSize(100.f, 100.f);
-		const FVector2D CanvasMaxScreenSize = FVector2D::Max(MaxScreenRatio*FVector2D(Canvas->ClipX, Canvas->ClipY) - CanvasBottomRightPadding - CanvasTopLeftPadding, CanvasMinimumSize);
+		const FVector2D CanvasMaxScreenSize = FVector2D::Max(MaxScreenRatio*FVector2D(Canvas->ClipX, Canvas->ClipY) - CanvasBottomRightPadding - CurrentOffset, CanvasMinimumSize);
 
 		FVector2D TotalFootprint(ForceInitToZero);
 		for (UWorldPartition* Partition : RegisteredWorldPartitions)
@@ -236,33 +242,70 @@ void UWorldPartitionSubsystem::Draw(UCanvas* Canvas, class APlayerController* PC
 
 		if (TotalFootprint.X > 0.f)
 		{
-			FVector2D PartitionCanvasOffset(CanvasTopLeftPadding);
 			for (UWorldPartition* Partition : RegisteredWorldPartitions)
 			{
 				FVector2D Footprint = Partition->GetDrawRuntimeHash2DDesiredFootprint(CanvasMaxScreenSize);
 				float FootprintRatio = Footprint.X / TotalFootprint.X;
 				FVector2D PartitionCanvasSize = FVector2D(CanvasMaxScreenSize.X * FootprintRatio, CanvasMaxScreenSize.Y);
-				Partition->DrawRuntimeHash2D(Canvas, PartitionCanvasOffset, PartitionCanvasSize);
-				PartitionCanvasOffset.X += PartitionCanvasSize.X;
+				Partition->DrawRuntimeHash2D(Canvas, PartitionCanvasSize, CurrentOffset);
 			}
+			CurrentOffset.X = CanvasBottomRightPadding.X;
 		}
 	}
 
 	if (GDrawStreamingSources)
 	{
-		if (const UWorldPartition* WorldPartition = GetMainWorldPartition())
+		const UWorldPartition* WorldPartition = GetMainWorldPartition();
+		const TArray<FWorldPartitionStreamingSource>* StreamingSources = WorldPartition ? &WorldPartition->GetStreamingSources() : nullptr;
+		if (StreamingSources && (StreamingSources->Num() > 0))
 		{
-			const TArray<FWorldPartitionStreamingSource>& StreamingSources = WorldPartition->GetStreamingSources();
-
 			FString Title(TEXT("Streaming Sources"));
-			DrawText(Canvas, Title, GEngine->GetSmallFont(), FColor::Yellow, Pos);
-			UFont* Font = GEngine->GetTinyFont();
-			for (const FWorldPartitionStreamingSource& StreamingSource : StreamingSources)
+			FWorldPartitionDebugHelper::DrawText(Canvas, Title, GEngine->GetSmallFont(), FColor::Yellow, CurrentOffset);
+
+			FVector2D Pos = CurrentOffset;
+			float MaxTextWidth = 0;
+			for (const FWorldPartitionStreamingSource& StreamingSource : *StreamingSources)
+			{
+				FWorldPartitionDebugHelper::DrawText(Canvas, StreamingSource.Name.ToString(), GEngine->GetTinyFont(), FColor::White, Pos, &MaxTextWidth);
+			}
+			Pos = CurrentOffset + FVector2D(MaxTextWidth + 10, 0.f);
+			for (const FWorldPartitionStreamingSource& StreamingSource : *StreamingSources)
+			{
+				const FString Text = FString::Printf(TEXT("Pos: %s | Rot: %s"), *StreamingSource.Location.ToString(), *StreamingSource.Rotation.ToString());
+				FWorldPartitionDebugHelper::DrawText(Canvas, Text, GEngine->GetTinyFont(), FColor::White, Pos);
+			}
+			CurrentOffset.Y = Pos.Y;
+		}
+	}
+
+	if (UWorldPartition* WorldPartition = GetMainWorldPartition())
+	{
+		UDataLayerSubsystem* DataLayerSubsystem = WorldPartition->GetWorld()->GetSubsystem<UDataLayerSubsystem>();
+
+		if (GDrawLegends)
+		{
+			// Streaming Status Legend
+			WorldPartition->DrawStreamingStatusLegend(Canvas, CurrentOffset);
+
+			// DataLayers Legend
+			if (DataLayerSubsystem)
 			{
-				const FString Text = FString::Printf(TEXT("%s %s %s"), *StreamingSource.Name.ToString(), *StreamingSource.Location.ToString(), *StreamingSource.Rotation.ToString());
-				DrawText(Canvas, Text, Font, FColor::White, Pos);
+				DataLayerSubsystem->DrawDataLayersLegend(Canvas, CurrentOffset);
 			}
 		}
+
+		if (DataLayerSubsystem)
+		{
+			DataLayerSubsystem->DrawDataLayersStatus(Canvas, CurrentOffset);
+		}
+	}
+
+	if (GDrawRuntimeCellsDetails)
+	{
+		if (UWorldPartition* Partition = GetMainWorldPartition())
+		{
+			Partition->DrawRuntimeCellsDetails(Canvas, CurrentOffset);
+		}
 	}
 }
 
diff --git a/Engine/Source/Runtime/Engine/Public/WorldPartition/DataLayer/DataLayer.h b/Engine/Source/Runtime/Engine/Public/WorldPartition/DataLayer/DataLayer.h
--- a/Engine/Source/Runtime/Engine/Public/WorldPartition/DataLayer/DataLayer.h
+++ b/Engine/Source/Runtime/Engine/Public/WorldPartition/DataLayer/DataLayer.h
@@ -8,6 +8,7 @@
 #include "UObject/Object.h"
 #include "WorldDataLayers.h"
 #include "ActorDataLayer.h"
+#include "Math/Color.h"
 
 #include "DataLayer.generated.h"
 
@@ -66,6 +67,8 @@ public:
 	UFUNCTION(Category = "Data Layer - Runtime", BlueprintCallable)
 	EDataLayerState GetInitialState() const { return IsDynamicallyLoaded() ? InitialState : EDataLayerState::Unloaded; }
 
+	FColor GetDebugColor() const { return DebugColor; }
+
 	/** Returns a sanitized version of the provided Data Layer Label */
 	static FName GetSanitizedDataLayerLabel(FName InDataLayerLabel);
 
@@ -109,4 +112,7 @@ private:
 	/** Whether the Data Layer affects actor runtime loading */
 	UPROPERTY(Category = "Data Layer - Runtime", EditAnywhere)
 	uint32 bIsDynamicallyLoaded : 1;
+
+	UPROPERTY(Category = "Data Layer - Runtime", EditAnywhere)
+	FColor DebugColor;
 };
\ No newline at end of file
diff --git a/Engine/Source/Runtime/Engine/Public/WorldPartition/DataLayer/DataLayerSubsystem.h b/Engine/Source/Runtime/Engine/Public/WorldPartition/DataLayer/DataLayerSubsystem.h
--- a/Engine/Source/Runtime/Engine/Public/WorldPartition/DataLayer/DataLayerSubsystem.h
+++ b/Engine/Source/Runtime/Engine/Public/WorldPartition/DataLayer/DataLayerSubsystem.h
@@ -12,6 +12,8 @@
  * UDataLayerSubsystem
  */
 
+class UCanvas;
+
 // Deprecate
 DECLARE_DYNAMIC_MULTICAST_DELEGATE_TwoParams(FOnDataLayerActivationStateChanged, const UDataLayer*, DataLayer, bool, bIsActive);
 
@@ -28,7 +30,6 @@ public:
 	//~ Begin USubsystem Interface.
 	virtual bool ShouldCreateSubsystem(UObject* Outer) const override;
 	virtual void Initialize(FSubsystemCollectionBase& Collection) override;
-	virtual void Deinitialize() override;
 	//~ End USubsystem Interface.
 
 	//~ Begin UWorldSubsystem Interface.
@@ -83,7 +84,10 @@ public:
 	
 	bool IsAnyDataLayerInState(const TArray<FName>& InDataLayerNames, EDataLayerState InState) const;
 	
-	void Draw(class UCanvas* Canvas, class APlayerController* PC);
+	void GetDataLayerDebugColors(TMap<FName, FColor>& OutMapping) const;
+	void DrawDataLayersLegend(UCanvas* Canvas, FVector2D& Offset) const;
+	void DrawDataLayersStatus(UCanvas* Canvas, FVector2D& Offset) const;
+	static TArray<UDataLayer*> ConvertArgsToDataLayers(UWorld* World, const TArray<FString>& InArgs);
 
 #if WITH_EDITOR
 public:
@@ -96,13 +100,10 @@ private:
 #endif
 
 private:
-	static TArray<UDataLayer*> ConvertArgsToDataLayers(UWorld* World, const TArray<FString>& InArgs);
 
 	/** Console command used to toggle activation of a DataLayer */
 	static class FAutoConsoleCommand ToggleDataLayerActivation;
 
 	TSet<FName> ActiveDataLayerNames;
 	TSet<FName> LoadedDataLayerNames;
-
-	FDelegateHandle	DrawHandle;
 };
\ No newline at end of file
diff --git a/Engine/Source/Runtime/Engine/Public/WorldPartition/WorldPartition.h b/Engine/Source/Runtime/Engine/Public/WorldPartition/WorldPartition.h
--- a/Engine/Source/Runtime/Engine/Public/WorldPartition/WorldPartition.h
+++ b/Engine/Source/Runtime/Engine/Public/WorldPartition/WorldPartition.h
@@ -24,6 +24,7 @@ class UWorldPartitionRuntimeCell;
 class UWorldPartitionRuntimeHash;
 class UWorldPartitionStreamingPolicy;
 class FHLODActorDesc;
+class UCanvas;
 
 struct IWorldPartitionStreamingSourceProvider;
 
@@ -175,14 +176,16 @@ public:
 	class ULevel* GetPreferredLoadedLevelToAddToWorld() const;
 	bool IsStreamingCompleted(EWorldPartitionRuntimeCellState QueryState, const TArray<FWorldPartitionStreamingQuerySource>& QuerySources, bool bExactState) const;
 
-	bool CanDrawRuntimeHash() const;
-	FVector2D GetDrawRuntimeHash2DDesiredFootprint(const FVector2D& CanvasSize);
-	void DrawRuntimeHash2D(class UCanvas* Canvas, const FVector2D& PartitionCanvasOffset, const FVector2D& PartitionCanvasSize);
-	void DrawRuntimeHash3D();
-
 	void RegisterStreamingSourceProvider(IWorldPartitionStreamingSourceProvider* StreamingSource);
 	bool UnregisterStreamingSourceProvider(IWorldPartitionStreamingSourceProvider* StreamingSource);
 
+	// Debugging Methods
+	bool CanDrawRuntimeHash() const;
+	FVector2D GetDrawRuntimeHash2DDesiredFootprint(const FVector2D& CanvasSize);
+	void DrawRuntimeHash2D(UCanvas* Canvas, const FVector2D& PartitionCanvasSize, FVector2D& Offset);
+	void DrawRuntimeHash3D();
+	void DrawRuntimeCellsDetails(UCanvas* Canvas, FVector2D& Offset);
+	void DrawStreamingStatusLegend(UCanvas* Canvas, FVector2D& Offset);
 
 #if WITH_EDITORONLY_DATA
 	UPROPERTY(DuplicateTransient)
diff --git a/Engine/Source/Runtime/Engine/Public/WorldPartition/WorldPartitionDebugHelper.h b/Engine/Source/Runtime/Engine/Public/WorldPartition/WorldPartitionDebugHelper.h
new file mode 100644
--- /dev/null
+++ b/Engine/Source/Runtime/Engine/Public/WorldPartition/WorldPartitionDebugHelper.h
@@ -0,0 +1,36 @@
+// Copyright Epic Games, Inc. All Rights Reserved.
+
+#pragma once
+
+#include "CoreTypes.h"
+#include "UObject/NameTypes.h"
+#include "ProfilingDebugging/ProfilingHelpers.h"
+
+class FAutoConsoleCommand;
+class UCanvas;
+class UFont;
+
+class FWorldPartitionDebugHelper
+{
+public:
+	static bool IsDebugDataLayerShown(FName DataLayerName);
+	static bool AreDebugDataLayersShown(const TArray<FName>& DataLayerNames);
+	static bool IsDebugStreamingStatusShown(EStreamingStatus StreamingStatus);
+	static bool IsDebugRuntimeHashGridShown(FName Name);
+	static bool IsDebugCellNameShow(const FString& Name);
+	static void DrawText(UCanvas* Canvas, const FString& Text, const UFont* Font, const FColor& Color, FVector2D& Pos, float* MaxTextWidth = nullptr);
+	static void DrawLegendItem(UCanvas* Canvas, const FString& Text, const UFont* Font, const FColor& Color, FVector2D& Pos, float* MaxItemWidth = nullptr);
+
+private:
+	static FAutoConsoleCommand DebugFilterByStreamingStatusCommand;
+	static TSet<EStreamingStatus> DebugStreamingStatusFilter;
+
+	static FAutoConsoleCommand DebugFilterByDataLayerCommand;
+	static TSet<FName> DebugDataLayerFilter;
+
+	static FAutoConsoleCommand DebugFilterByRuntimeHashGridNameCommand;
+	static TSet<FName> DebugRuntimeHashFilter;
+
+	static FAutoConsoleCommand DebugFilterByCellNameCommand;
+	static TArray<FString> DebugCellNameFilter;
+};
\ No newline at end of file
diff --git a/Engine/Source/Runtime/Engine/Public/WorldPartition/WorldPartitionLevelStreamingDynamic.h b/Engine/Source/Runtime/Engine/Public/WorldPartition/WorldPartitionLevelStreamingDynamic.h
--- a/Engine/Source/Runtime/Engine/Public/WorldPartition/WorldPartitionLevelStreamingDynamic.h
+++ b/Engine/Source/Runtime/Engine/Public/WorldPartition/WorldPartitionLevelStreamingDynamic.h
@@ -14,8 +14,8 @@
 #include "Engine/LevelStreaming.h"
 #include "Engine/LevelStreamingDynamic.h"
 #include "WorldPartition/WorldPartition.h"
-#if WITH_EDITOR
 #include "WorldPartition/WorldPartitionRuntimeLevelStreamingCell.h"
+#if WITH_EDITOR
 #include "WorldPartition/WorldPartitionPackageCache.h"
 #endif
 #include "WorldPartitionLevelStreamingDynamic.generated.h"
@@ -29,10 +29,12 @@ class ENGINE_API UWorldPartitionLevelStreamingDynamic : public ULevelStreamingDy
 	void Unload();
 	void Activate();
 	void Deactivate();
-	virtual bool ShouldBeAlwaysLoaded() const { return bShouldBeAlwaysLoaded; }
-	virtual bool ShouldRequireFullVisibilityToRender() const override { return true; }
-	void SetShouldBeAlwaysLoaded(bool bInShouldBeAlwaysLoaded) { bShouldBeAlwaysLoaded = bInShouldBeAlwaysLoaded; }
 	UWorld* GetOuterWorld() const;
+	void SetShouldBeAlwaysLoaded(bool bInShouldBeAlwaysLoaded) { bShouldBeAlwaysLoaded = bInShouldBeAlwaysLoaded; }
+	const UWorldPartitionRuntimeCell* GetWorldPartitionRuntimeCell() const { return StreamingCell.Get(); }
+
+	virtual bool ShouldBeAlwaysLoaded() const override { return bShouldBeAlwaysLoaded; }
+	virtual bool ShouldRequireFullVisibilityToRender() const override { return true; }
 
 #if WITH_EDITOR
 	void Initialize(const UWorldPartitionRuntimeLevelStreamingCell& InCell);
@@ -69,6 +71,9 @@ private:
 	UActorContainer* ActorContainer;
 #endif
 
+	UPROPERTY()
+	TWeakObjectPtr<const UWorldPartitionRuntimeLevelStreamingCell> StreamingCell;
+
 	UPROPERTY()
 	TWeakObjectPtr<UWorldPartition> OuterWorldPartition;
 };
\ No newline at end of file
diff --git a/Engine/Source/Runtime/Engine/Public/WorldPartition/WorldPartitionLevelStreamingPolicy.h b/Engine/Source/Runtime/Engine/Public/WorldPartition/WorldPartitionLevelStreamingPolicy.h
--- a/Engine/Source/Runtime/Engine/Public/WorldPartition/WorldPartitionLevelStreamingPolicy.h
+++ b/Engine/Source/Runtime/Engine/Public/WorldPartition/WorldPartitionLevelStreamingPolicy.h
@@ -25,6 +25,8 @@ public:
 	virtual void SetTargetStateForCells(EWorldPartitionRuntimeCellState TargetState, const TSet<const UWorldPartitionRuntimeCell*>& Cells) override;
 	virtual EWorldPartitionRuntimeCellState GetCurrentStateForCell(const UWorldPartitionRuntimeCell* Cell) const override;
 	virtual class ULevel* GetPreferredLoadedLevelToAddToWorld() const override;
+	virtual void DrawRuntimeCellsDetails(class UCanvas* Canvas, FVector2D& Offset) override;
+	virtual void DrawStreamingStatusLegend(UCanvas* Canvas, FVector2D& Offset) override;
 
 #if WITH_EDITOR
 	virtual TSubclassOf<class UWorldPartitionRuntimeCell> GetRuntimeCellClass() const override;
@@ -36,6 +38,9 @@ public:
 
 	virtual UObject* GetSubObject(const TCHAR* SubObjectPath) override;
 
+protected:
+	void ForEachActiveRuntimeCell(TFunctionRef<void(const UWorldPartitionRuntimeCell*)> Func) const;
+
 private:
 	void SetCellsStateToLoaded(const TSet<const UWorldPartitionRuntimeCell*>& ToLoadCells);
 	void SetCellsStateToActivated(const TSet<const UWorldPartitionRuntimeCell*>& ToActivateCells);
diff --git a/Engine/Source/Runtime/Engine/Public/WorldPartition/WorldPartitionRuntimeCell.h b/Engine/Source/Runtime/Engine/Public/WorldPartition/WorldPartitionRuntimeCell.h
--- a/Engine/Source/Runtime/Engine/Public/WorldPartition/WorldPartitionRuntimeCell.h
+++ b/Engine/Source/Runtime/Engine/Public/WorldPartition/WorldPartitionRuntimeCell.h
@@ -5,6 +5,7 @@
 #include "Templates/SubclassOf.h"
 #include "WorldPartition/DataLayer/DataLayer.h"
 #include "WorldPartition/WorldPartitionActorDescView.h"
+#include "ProfilingDebugging/ProfilingHelpers.h"
 #include "WorldPartitionRuntimeCell.generated.h"
 
 USTRUCT()
@@ -105,19 +106,24 @@ static_assert(EWorldPartitionRuntimeCellState::Unloaded < EWorldPartitionRuntime
 /**
  * Represents a PIE/Game streaming cell which points to external actor/data chunk packages
  */
-UCLASS(Abstract)
+UCLASS(Abstract, Within = WorldPartition)
 class UWorldPartitionRuntimeCell : public UObject
 {
 	GENERATED_UCLASS_BODY()
 
-	virtual FLinearColor GetDebugColor() const PURE_VIRTUAL(UWorldPartitionRuntimeCell::GetDebugColor, return FLinearColor::Black;);
+	virtual FLinearColor GetDebugColor() const { return FLinearColor::Black; }
 	virtual bool IsAlwaysLoaded() const { return bIsAlwaysLoaded; }
 	virtual void SetIsAlwaysLoaded(bool bInIsAlwaysLoaded) { bIsAlwaysLoaded = bInIsAlwaysLoaded; }
 	bool HasDataLayers() const { return !DataLayers.IsEmpty(); }
 	const TArray<FName>& GetDataLayers() const { return DataLayers; }
+	virtual EStreamingStatus GetStreamingStatus() const { return LEVEL_Unloaded; }
+	virtual bool IsLoading() const { return false; }
+	virtual const FString& GetDebugName() const { return DebugName; }
+	virtual bool IsDebugShown() const;
 
 #if WITH_EDITOR
-	void SetDataLayers(const TArray<const UDataLayer*> InDataLayers);
+	void SetDataLayers(const TArray<const UDataLayer*>& InDataLayers);
+	void SetDebugInfo(FIntVector InCoords, FName InGridName);
 	void AddCellData(const UWorldPartitionRuntimeCellData* InCellData);
 	virtual void AddActorToCell(const FWorldPartitionActorDescView& ActorDescView, uint64 InContainerID, const FTransform& InContainerTransform, const UActorDescContainer* InContainer) PURE_VIRTUAL(UWorldPartitionRuntimeCell::AddActorToCell,);
 	virtual int32 GetActorCount() const PURE_VIRTUAL(UWorldPartitionRuntimeCell::GetActorCount, return 0;);
@@ -132,8 +138,12 @@ class UWorldPartitionRuntimeCell : public UObject
 	const UWorldPartitionRuntimeCellData* GetCellData(const TSubclassOf<UWorldPartitionRuntimeCellData> InCellDataClass) const;
 	template <class T> inline const T* GetCellData() const { return Cast<const T>(GetCellData(T::StaticClass())); }
 	template <class T> inline bool HasCellData() const { return GetCellData<T>() != nullptr; }
-
+	
 protected:
+#if WITH_EDITOR
+	void UpdateDebugName();
+#endif
+
 	UPROPERTY()
 	bool bIsAlwaysLoaded;
 
@@ -143,4 +153,14 @@ private:
 
 	UPROPERTY()
 	TArray<FName> DataLayers;
+
+	// Debug Info
+	UPROPERTY()
+	FIntVector Coords;
+
+	UPROPERTY()
+	FName GridName;
+
+	UPROPERTY()
+	FString DebugName;
 };
diff --git a/Engine/Source/Runtime/Engine/Public/WorldPartition/WorldPartitionRuntimeHash.h b/Engine/Source/Runtime/Engine/Public/WorldPartition/WorldPartitionRuntimeHash.h
--- a/Engine/Source/Runtime/Engine/Public/WorldPartition/WorldPartitionRuntimeHash.h
+++ b/Engine/Source/Runtime/Engine/Public/WorldPartition/WorldPartitionRuntimeHash.h
@@ -54,8 +54,9 @@ class ENGINE_API UWorldPartitionRuntimeHash : public UObject
 	 */
 	virtual FVector2D GetDraw2DDesiredFootprint(const FVector2D& CanvasSize) const { return FVector2D::ZeroVector; }
 
-	virtual void Draw2D(class UCanvas* Canvas, const TArray<FWorldPartitionStreamingSource>& Sources, const FVector2D& PartitionCanvasOffset, const FVector2D& PartitionCanvasSize) const {}
+	virtual void Draw2D(class UCanvas* Canvas, const TArray<FWorldPartitionStreamingSource>& Sources, const FVector2D& PartitionCanvasSize, FVector2D& Offset) const {}
 	virtual void Draw3D(const TArray<FWorldPartitionStreamingSource>& Sources) const {}
+	virtual bool ContainsRuntimeHash(const FString& Name) const { return false; }
 
 protected:
 #if WITH_EDITOR
diff --git a/Engine/Source/Runtime/Engine/Public/WorldPartition/WorldPartitionRuntimeLevelStreamingCell.h b/Engine/Source/Runtime/Engine/Public/WorldPartition/WorldPartitionRuntimeLevelStreamingCell.h
--- a/Engine/Source/Runtime/Engine/Public/WorldPartition/WorldPartitionRuntimeLevelStreamingCell.h
+++ b/Engine/Source/Runtime/Engine/Public/WorldPartition/WorldPartitionRuntimeLevelStreamingCell.h
@@ -17,6 +17,8 @@ class UWorldPartitionRuntimeLevelStreamingCell : public UWorldPartitionRuntimeSp
 	int32 GetStreamingPriority() const { return Level; }
 	class UWorldPartitionLevelStreamingDynamic* GetLevelStreaming() const;
 	virtual FLinearColor GetDebugColor() const override;
+	virtual EStreamingStatus GetStreamingStatus() const override;
+	virtual bool IsLoading() const override;
 	EWorldPartitionRuntimeCellState GetCurrentState() const;
 
 #if WITH_EDITOR
@@ -38,8 +40,6 @@ private:
 	UFUNCTION()
 	void OnLevelHidden();
 
-	EStreamingStatus GetLevelStreamingStatus() const;
-
 	class UWorldPartitionLevelStreamingDynamic* GetOrCreateLevelStreaming() const;
 
 #if WITH_EDITOR
diff --git a/Engine/Source/Runtime/Engine/Public/WorldPartition/WorldPartitionRuntimeSpatialHash.h b/Engine/Source/Runtime/Engine/Public/WorldPartition/WorldPartitionRuntimeSpatialHash.h
--- a/Engine/Source/Runtime/Engine/Public/WorldPartition/WorldPartitionRuntimeSpatialHash.h
+++ b/Engine/Source/Runtime/Engine/Public/WorldPartition/WorldPartitionRuntimeSpatialHash.h
@@ -127,10 +127,11 @@ struct FSpatialHashStreamingGrid
 	void GetCells(const FWorldPartitionStreamingQuerySource& QuerySource, TSet<const UWorldPartitionRuntimeCell*>& OutCells) const;
 	void GetCells(const TArray<FWorldPartitionStreamingSource>& Sources, const class UDataLayerSubsystem* DataLayerSubsystem, TSet<const UWorldPartitionRuntimeCell*>& OutActivateCells, TSet<const UWorldPartitionRuntimeCell*>& OutLoadCells) const;
 	void GetAlwaysLoadedCells(const UDataLayerSubsystem* DataLayerSubsystem, TSet<const UWorldPartitionRuntimeCell*>& OutActivateCells, TSet<const UWorldPartitionRuntimeCell*>& OutLoadCells) const;
-	void Draw2D(UCanvas* Canvas, const TArray<FWorldPartitionStreamingSource>& Sources, const FBox& Region, const FBox2D& GridScreenBounds, TFunctionRef<FVector2D(const FVector2D&)> WorldToScreen) const;
-	void Draw3D(UWorld* World, const TArray<FWorldPartitionStreamingSource>& Sources, const FTransform& Transform) const;
+	void Draw2D(UCanvas* Canvas, const TArray<FWorldPartitionStreamingSource>& Sources, const FBox& Region, const FBox2D& GridScreenBounds, TFunctionRef<FVector2D(const FVector2D&)> WorldToScreen, const TMap<FName, FColor>& DataLayerDebugColors) const;
+	void Draw3D(UWorld* World, const TArray<FWorldPartitionStreamingSource>& Sources, const FTransform& Transform, const TMap<FName, FColor>& DataLayerDebugColors) const;
 
 private:
+	void GetFilteredCellsForDebugDraw(const FSpatialHashStreamingGridLayerCell* LayerCell, TArray<const UWorldPartitionRuntimeCell*>& FilteredCells) const;
 	const FSquare2DGridHelper& GetGridHelper() const;
 	mutable FSquare2DGridHelper* GridHelper;
 };
@@ -256,8 +257,9 @@ private:
 	TArray<FSpatialHashStreamingGrid> StreamingGrids;
 
 	virtual FVector2D GetDraw2DDesiredFootprint(const FVector2D& CanvasSize) const override;
-	virtual void Draw2D(class UCanvas* Canvas, const TArray<FWorldPartitionStreamingSource>& Sources, const FVector2D& PartitionCanvasOffset, const FVector2D& PartitionCanvasSize) const override;
+	virtual void Draw2D(class UCanvas* Canvas, const TArray<FWorldPartitionStreamingSource>& Sources, const FVector2D& PartitionCanvasSize, FVector2D& Offset) const override;
 	virtual void Draw3D(const TArray<FWorldPartitionStreamingSource>& Sources) const override;
+	virtual bool ContainsRuntimeHash(const FString& Name) const override;
 
 private:
 	void GetAlwaysLoadedStreamingCells(const FSpatialHashStreamingGrid& StreamingGrid, TSet<const UWorldPartitionRuntimeCell*>& Cells) const;
@@ -265,6 +267,7 @@ private:
 #if WITH_EDITOR
 	bool CreateStreamingGrid(const FSpatialHashRuntimeGrid& RuntimeGrid, const FSquare2DGridHelper& PartionedActors, EWorldPartitionStreamingMode Mode, UWorldPartitionStreamingPolicy* StreamingPolicy, TArray<FString>* OutPackagesToGenerate = nullptr);
 #endif
+	TArray<const FSpatialHashStreamingGrid*> GetFilteredStreamingGrids() const;
 
 	friend class UWorldPartitionSubsystem;
 };
\ No newline at end of file
diff --git a/Engine/Source/Runtime/Engine/Public/WorldPartition/WorldPartitionStreamingPolicy.h b/Engine/Source/Runtime/Engine/Public/WorldPartition/WorldPartitionStreamingPolicy.h
--- a/Engine/Source/Runtime/Engine/Public/WorldPartition/WorldPartitionStreamingPolicy.h
+++ b/Engine/Source/Runtime/Engine/Public/WorldPartition/WorldPartitionStreamingPolicy.h
@@ -27,8 +27,10 @@ public:
 	virtual EWorldPartitionRuntimeCellState GetCurrentStateForCell(const UWorldPartitionRuntimeCell* Cell) const PURE_VIRTUAL(UWorldPartitionStreamingPolicy::GetCurrentStateForCell, return EWorldPartitionRuntimeCellState::Unloaded; );
 	virtual class ULevel* GetPreferredLoadedLevelToAddToWorld() const { return nullptr; }
 	virtual FVector2D GetDrawRuntimeHash2DDesiredFootprint(const FVector2D& CanvasSize);
-	virtual void DrawRuntimeHash2D(class UCanvas* Canvas, const FVector2D& PartitionCanvasOffset, const FVector2D& PartitionCanvasSize);
+	virtual void DrawRuntimeHash2D(class UCanvas* Canvas, const FVector2D& PartitionCanvasSize, FVector2D& Offset);
 	virtual void DrawRuntimeHash3D();
+	virtual void DrawRuntimeCellsDetails(class UCanvas* Canvas, FVector2D& Offset) {}
+	virtual void DrawStreamingStatusLegend(class UCanvas* Canvas, FVector2D& Offset) {}
 
 	virtual bool IsStreamingCompleted(EWorldPartitionRuntimeCellState QueryState, const TArray<FWorldPartitionStreamingQuerySource>& QuerySources, bool bExactState = true) const;
 
