Removed Scene Software Occlusion

diff --git a/Engine/Source/Developer/Android/AndroidTargetPlatform/Private/AndroidTargetPlatform.cpp b/Engine/Source/Developer/Android/AndroidTargetPlatform/Private/AndroidTargetPlatform.cpp
--- a/Engine/Source/Developer/Android/AndroidTargetPlatform/Private/AndroidTargetPlatform.cpp
+++ b/Engine/Source/Developer/Android/AndroidTargetPlatform/Private/AndroidTargetPlatform.cpp
@@ -253,12 +253,6 @@ bool FAndroidTargetPlatform::SupportsVulkanSM5() const
 	return bSupportsMobileVulkanSM5;
 }
 
-bool FAndroidTargetPlatform::SupportsSoftwareOcclusion() const
-{
-	static auto* CVarMobileAllowSoftwareOcclusion = IConsoleManager::Get().FindTConsoleVariableDataInt(TEXT("r.Mobile.AllowSoftwareOcclusion"));
-	return CVarMobileAllowSoftwareOcclusion->GetValueOnAnyThread() != 0;
-}
-
 bool FAndroidTargetPlatform::SupportsLandscapeMeshLODStreaming() const
 {
 	bool bStreamLandscapeMeshLODs = false;
@@ -356,9 +350,6 @@ bool FAndroidTargetPlatform::SupportsFeature( ETargetPlatformFeatures Feature )
 		case ETargetPlatformFeatures::DeferredRendering:
 			return SupportsVulkanSM5();
 
-		case ETargetPlatformFeatures::SoftwareOcclusion:
-			return SupportsSoftwareOcclusion();
-
 		case ETargetPlatformFeatures::VirtualTextureStreaming:
 			return UsesVirtualTextures();
 
diff --git a/Engine/Source/Developer/Android/AndroidTargetPlatform/Private/AndroidTargetPlatform.h b/Engine/Source/Developer/Android/AndroidTargetPlatform/Private/AndroidTargetPlatform.h
--- a/Engine/Source/Developer/Android/AndroidTargetPlatform/Private/AndroidTargetPlatform.h
+++ b/Engine/Source/Developer/Android/AndroidTargetPlatform/Private/AndroidTargetPlatform.h
@@ -229,7 +229,6 @@ protected:
 	// query for rene3ring mode support
 	bool SupportsES31() const;
 	bool SupportsVulkan() const;
-	bool SupportsSoftwareOcclusion() const;
 	bool SupportsLandscapeMeshLODStreaming() const;
 	bool SupportsVulkanSM5() const;
 
diff --git a/Engine/Source/Developer/IOS/IOSTargetPlatform/Private/IOSTargetPlatform.cpp b/Engine/Source/Developer/IOS/IOSTargetPlatform/Private/IOSTargetPlatform.cpp
--- a/Engine/Source/Developer/IOS/IOSTargetPlatform/Private/IOSTargetPlatform.cpp
+++ b/Engine/Source/Developer/IOS/IOSTargetPlatform/Private/IOSTargetPlatform.cpp
@@ -439,12 +439,6 @@ static bool SupportsMetalMRT()
 	return bSupportsMetalMRT;
 }
 
-static bool SupportsSoftwareOcclusion()
-{
-	static auto* CVarMobileAllowSoftwareOcclusion = IConsoleManager::Get().FindTConsoleVariableDataInt(TEXT("r.Mobile.AllowSoftwareOcclusion"));
-	return CVarMobileAllowSoftwareOcclusion->GetValueOnAnyThread() != 0;
-}
-
 static bool SupportsLandscapeMeshLODStreaming()
 {
 	bool bStreamLandscapeMeshLODs = false;
@@ -476,9 +470,6 @@ bool FIOSTargetPlatform::SupportsFeature( ETargetPlatformFeatures Feature ) cons
 		case ETargetPlatformFeatures::HighQualityLightmaps:
 			return SupportsMetalMRT();
 
-		case ETargetPlatformFeatures::SoftwareOcclusion:
-			return SupportsSoftwareOcclusion();
-
 		case ETargetPlatformFeatures::VirtualTextureStreaming:
 			return UsesVirtualTextures();
 
diff --git a/Engine/Source/Developer/Lumin/LuminTargetPlatform/Private/LuminTargetPlatform.cpp b/Engine/Source/Developer/Lumin/LuminTargetPlatform/Private/LuminTargetPlatform.cpp
--- a/Engine/Source/Developer/Lumin/LuminTargetPlatform/Private/LuminTargetPlatform.cpp
+++ b/Engine/Source/Developer/Lumin/LuminTargetPlatform/Private/LuminTargetPlatform.cpp
@@ -125,9 +125,6 @@ bool FLuminTargetPlatform::SupportsFeature(ETargetPlatformFeatures Feature) cons
 	case ETargetPlatformFeatures::DeferredRendering:
 		return SupportsDesktopRendering();
 
-	case ETargetPlatformFeatures::SoftwareOcclusion:
-		return SupportsSoftwareOcclusion();
-
 	default:
 		break;
 	}
diff --git a/Engine/Source/Developer/TargetPlatform/Public/Interfaces/ITargetPlatform.h b/Engine/Source/Developer/TargetPlatform/Public/Interfaces/ITargetPlatform.h
--- a/Engine/Source/Developer/TargetPlatform/Public/Interfaces/ITargetPlatform.h
+++ b/Engine/Source/Developer/TargetPlatform/Public/Interfaces/ITargetPlatform.h
@@ -68,9 +68,6 @@ enum class ETargetPlatformFeatures
 	/* Should split paks into smaller sized paks */
 	ShouldSplitPaksIntoSmallerSizes,
 
-	/* The platform uses software rasterization of the scene for primitive occlusion */
-	SoftwareOcclusion,
-
 	/* The platform supports half float vertex format */
 	HalfFloatVertexFormat,
 
diff --git a/Engine/Source/Editor/LandscapeEditor/Private/LandscapeEdMode.cpp b/Engine/Source/Editor/LandscapeEditor/Private/LandscapeEdMode.cpp
--- a/Engine/Source/Editor/LandscapeEditor/Private/LandscapeEdMode.cpp
+++ b/Engine/Source/Editor/LandscapeEditor/Private/LandscapeEdMode.cpp
@@ -3402,7 +3402,6 @@ ALandscape* FEdModeLandscape::ChangeComponentSetting(int32 NumComponentsX, int32
 			NewLandscape->ComponentScreenSizeToUseSubSections = OldLandscape->ComponentScreenSizeToUseSubSections;
 			NewLandscape->LODDistributionSetting = OldLandscape->LODDistributionSetting;
 			NewLandscape->LOD0DistributionSetting = OldLandscape->LOD0DistributionSetting;
-			NewLandscape->OccluderGeometryLOD = OldLandscape->OccluderGeometryLOD;
 			NewLandscape->ExportLOD = OldLandscape->ExportLOD;
 			NewLandscape->StaticLightingLOD = OldLandscape->StaticLightingLOD;
 			NewLandscape->NegativeZBoundsExtension = OldLandscape->NegativeZBoundsExtension;
diff --git a/Engine/Source/Runtime/Engine/Classes/Engine/InstancedStaticMesh.h b/Engine/Source/Runtime/Engine/Classes/Engine/InstancedStaticMesh.h
--- a/Engine/Source/Runtime/Engine/Classes/Engine/InstancedStaticMesh.h
+++ b/Engine/Source/Runtime/Engine/Classes/Engine/InstancedStaticMesh.h
@@ -561,8 +561,6 @@ public:
 	virtual void GetDistancefieldAtlasData(const FDistanceFieldVolumeData*& OutDistanceFieldData, float& SelfShadowBias) const override;
 	virtual void GetDistancefieldInstanceData(TArray<FMatrix>& ObjectLocalToWorldTransforms) const override;
 
-	virtual int32 CollectOccluderElements(FOccluderElementsCollector& Collector) const override;
-
 	/**
 	 * Creates the hit proxies are used when DrawDynamicElements is called.
 	 * Called in the game thread.
diff --git a/Engine/Source/Runtime/Engine/Classes/Engine/RendererSettings.h b/Engine/Source/Runtime/Engine/Classes/Engine/RendererSettings.h
--- a/Engine/Source/Runtime/Engine/Classes/Engine/RendererSettings.h
+++ b/Engine/Source/Runtime/Engine/Classes/Engine/RendererSettings.h
@@ -201,12 +201,6 @@ class ENGINE_API URendererSettings : public UDeveloperSettings
 		ConfigRestartRequired=true))
 	uint32 bMobileAllowDitheredLODTransition:1;
 
-	UPROPERTY(config, EditAnywhere, Category = Mobile, meta=(
-		ConsoleVariable="r.Mobile.AllowSoftwareOcclusion", DisplayName="Support Software Occlusion Culling",
-		ToolTip="Whether to support 'Software Occlusion Culling' on mobile platforms. This will package occluder information and enable Software Occlusion Culling.",
-		ConfigRestartRequired=false))
-	uint32 bMobileAllowSoftwareOcclusionCulling:1;
-
 	UPROPERTY(config, EditAnywhere, Category = Mobile, meta = (
 		EditCondition = "bVirtualTextures",
 		ConsoleVariable = "r.Mobile.VirtualTextures", DisplayName = "Enable virtual texture support on Mobile",
diff --git a/Engine/Source/Runtime/Engine/Classes/Engine/StaticMesh.h b/Engine/Source/Runtime/Engine/Classes/Engine/StaticMesh.h
--- a/Engine/Source/Runtime/Engine/Classes/Engine/StaticMesh.h
+++ b/Engine/Source/Runtime/Engine/Classes/Engine/StaticMesh.h
@@ -46,7 +46,7 @@ enum class EStaticMeshAsyncProperties : uint32
 {
 	None                    = 0,
 	RenderData              = 1 << 0,
-	OccluderData            = 1 << 1,
+	//OccluderData            = 1 << 1,
 	SourceModels            = 1 << 2,
 	SectionInfoMap          = 1 << 3,
 	OriginalSectionInfoMap  = 1 << 4,
@@ -76,8 +76,6 @@ inline const TCHAR* ToString(EStaticMeshAsyncProperties Value)
 			return TEXT("None");
 		case EStaticMeshAsyncProperties::RenderData: 
 			return TEXT("RenderData");
-		case EStaticMeshAsyncProperties::OccluderData: 
-			return TEXT("OccluderData");
 		case EStaticMeshAsyncProperties::SourceModels: 
 			return TEXT("SourceModels");
 		case EStaticMeshAsyncProperties::SectionInfoMap: 
@@ -608,10 +606,6 @@ private:
 	UE_DEPRECATED(5.00, "This must be protected for async build, always use the accessors even internally.")
 	TUniquePtr<class FStaticMeshRenderData> RenderData;
 
-	/** Pointer to the occluder data used to rasterize this static mesh for software occlusion. */
-	UE_DEPRECATED(5.00, "This must be protected for async build, always use the accessors even internally.")
-	TUniquePtr<class FStaticMeshOccluderData> OccluderData;
-
 public:
 #if WITH_EDITOR
 	ENGINE_API bool IsCompiling() const override { return AsyncTask != nullptr || LockedProperties.load(std::memory_order_relaxed) != 0; }
@@ -623,10 +617,6 @@ public:
 	ENGINE_API const FStaticMeshRenderData* GetRenderData() const;
 	ENGINE_API void SetRenderData(TUniquePtr<class FStaticMeshRenderData>&& InRenderData);
 
-	ENGINE_API FStaticMeshOccluderData* GetOccluderData();
-	ENGINE_API const FStaticMeshOccluderData* GetOccluderData() const;
-	ENGINE_API void SetOccluderData(TUniquePtr<class FStaticMeshOccluderData>&& InOccluderData);
-
 #if WITH_EDITORONLY_DATA
 	static const float MinimumAutoLODPixelError;
 
@@ -1035,14 +1025,6 @@ public:
 	/** If the user has modified collision in any way or has custom collision imported. Used for determining if to auto generate collision on import */
 	UPROPERTY(EditAnywhere, Category = Collision)
 	bool bCustomizedCollision;
-
-	/** 
-	 *	Specifies which mesh LOD to use as occluder geometry for software occlusion
-	 *  Set to -1 to not use this mesh as occluder 
-	 */
-	UPROPERTY(EditAnywhere, Category=StaticMesh, AdvancedDisplay, meta=(DisplayName="LOD For Occluder Mesh"))
-	int32 LODForOccluderMesh;
-
 #endif // WITH_EDITORONLY_DATA
 
 private:
diff --git a/Engine/Source/Runtime/Engine/Private/InstancedStaticMesh.cpp b/Engine/Source/Runtime/Engine/Private/InstancedStaticMesh.cpp
--- a/Engine/Source/Runtime/Engine/Private/InstancedStaticMesh.cpp
+++ b/Engine/Source/Runtime/Engine/Private/InstancedStaticMesh.cpp
@@ -1161,25 +1161,6 @@ void FInstancedStaticMeshSceneProxy::GetDynamicMeshElements(const TArray<const F
 	}
 }
 
-int32 FInstancedStaticMeshSceneProxy::CollectOccluderElements(FOccluderElementsCollector& Collector) const
-{
-	if (OccluderData)
-	{	
-		FStaticMeshInstanceBuffer& InstanceBuffer = InstancedRenderData.PerInstanceRenderData->InstanceBuffer;
-		const int32 NumInstances = InstanceBuffer.GetNumInstances();
-		
-		const TArray<FMatrix>& PerInstanceTransforms = InstancedRenderData.PerInstanceRenderData->GetPerInstanceTransforms();
-		for (const FMatrix& InstanceToLocal: PerInstanceTransforms)
-		{
-			Collector.AddElements(OccluderData->VerticesSP, OccluderData->IndicesSP, InstanceToLocal * GetLocalToWorld());
-		}
-		
-		return NumInstances;
-	}
-	
-	return 0;
-}
-
 void FInstancedStaticMeshSceneProxy::SetupProxy(UInstancedStaticMeshComponent* InComponent)
 {
 #if WITH_EDITOR
diff --git a/Engine/Source/Runtime/Engine/Private/StaticMesh.cpp b/Engine/Source/Runtime/Engine/Private/StaticMesh.cpp
--- a/Engine/Source/Runtime/Engine/Private/StaticMesh.cpp
+++ b/Engine/Source/Runtime/Engine/Private/StaticMesh.cpp
@@ -92,7 +92,6 @@ DECLARE_MEMORY_STAT( TEXT( "StaticMesh Total Memory" ), STAT_StaticMeshTotalMemo
 DECLARE_MEMORY_STAT( TEXT( "StaticMesh Vertex Memory" ), STAT_StaticMeshVertexMemory, STATGROUP_MemoryStaticMesh );
 DECLARE_MEMORY_STAT( TEXT( "StaticMesh VxColor Resource Mem" ), STAT_ResourceVertexColorMemory, STATGROUP_MemoryStaticMesh );
 DECLARE_MEMORY_STAT( TEXT( "StaticMesh Index Memory" ), STAT_StaticMeshIndexMemory, STATGROUP_MemoryStaticMesh );
-DECLARE_MEMORY_STAT( TEXT( "StaticMesh Occluder Memory" ), STAT_StaticMeshOccluderMemory, STATGROUP_MemoryStaticMesh );
 
 DECLARE_MEMORY_STAT( TEXT( "StaticMesh Total Memory" ), STAT_StaticMeshTotalMemory, STATGROUP_Memory );
 
@@ -1805,82 +1804,6 @@ int32 FStaticMeshRenderData::GetFirstValidLODIdx(int32 MinIdx) const
 	return LODIndex;
 }
 
-FStaticMeshOccluderData::FStaticMeshOccluderData()
-{
-	VerticesSP = MakeShared<FOccluderVertexArray, ESPMode::ThreadSafe>();
-	IndicesSP = MakeShared<FOccluderIndexArray, ESPMode::ThreadSafe>();
-}
-
-SIZE_T FStaticMeshOccluderData::GetResourceSizeBytes() const
-{
-	return VerticesSP->GetAllocatedSize() + IndicesSP->GetAllocatedSize();
-}
-
-TUniquePtr<FStaticMeshOccluderData> FStaticMeshOccluderData::Build(UStaticMesh* Owner)
-{
-	TUniquePtr<FStaticMeshOccluderData> Result;
-#if WITH_EDITOR		
-	if (Owner->LODForOccluderMesh >= 0)
-	{
-		// TODO: Custom geometry for occluder mesh?
-		int32 LODIndex = FMath::Min(Owner->LODForOccluderMesh, Owner->GetRenderData()->LODResources.Num()-1);
-		const FStaticMeshLODResources& LODModel = Owner->GetRenderData()->LODResources[LODIndex];
-			
-		const FRawStaticIndexBuffer& IndexBuffer = LODModel.DepthOnlyIndexBuffer.GetNumIndices() > 0 ? LODModel.DepthOnlyIndexBuffer : LODModel.IndexBuffer;
-		int32 NumVtx = LODModel.VertexBuffers.PositionVertexBuffer.GetNumVertices();
-		int32 NumIndices = IndexBuffer.GetNumIndices();
-		
-		if (NumVtx > 0 && NumIndices > 0 && !IndexBuffer.Is32Bit())
-		{
-			Result = MakeUnique<FStaticMeshOccluderData>();
-		
-			Result->VerticesSP->SetNumUninitialized(NumVtx);
-			Result->IndicesSP->SetNumUninitialized(NumIndices);
-
-			const FVector* V0 = &LODModel.VertexBuffers.PositionVertexBuffer.VertexPosition(0);
-			const uint16* Indices = IndexBuffer.AccessStream16();
-
-			FMemory::Memcpy(Result->VerticesSP->GetData(), V0, NumVtx*sizeof(FVector));
-			FMemory::Memcpy(Result->IndicesSP->GetData(), Indices, NumIndices*sizeof(uint16));
-		}
-	}
-#endif // WITH_EDITOR
-	return Result;
-}
-
-void FStaticMeshOccluderData::SerializeCooked(FArchive& Ar, UStaticMesh* Owner)
-{
-#if WITH_EDITOR	
-	if (Ar.IsSaving())
-	{
-		bool bHasOccluderData = false;
-		if (Ar.CookingTarget()->SupportsFeature(ETargetPlatformFeatures::SoftwareOcclusion) && Owner->GetOccluderData())
-		{
-			bHasOccluderData = true;
-		}
-		
-		Ar << bHasOccluderData;
-		
-		if (bHasOccluderData)
-		{
-			Owner->GetOccluderData()->VerticesSP->BulkSerialize(Ar);
-			Owner->GetOccluderData()->IndicesSP->BulkSerialize(Ar);
-		}
-	}
-	else
-#endif // WITH_EDITOR
-	{
-		bool bHasOccluderData;
-		Ar << bHasOccluderData;
-		if (bHasOccluderData)
-		{
-			Owner->SetOccluderData(MakeUnique<FStaticMeshOccluderData>());
-			Owner->GetOccluderData()->VerticesSP->BulkSerialize(Ar);
-			Owner->GetOccluderData()->IndicesSP->BulkSerialize(Ar);
-		}
-	}
-}
-
 
 #if WITH_EDITOR
 /**
@@ -2896,7 +2819,6 @@ UStaticMesh::UStaticMesh(const FObjectInitializer& ObjectInitializer)
 #if WITH_EDITORONLY_DATA
 	bAutoComputeLODScreenSize=true;
 	ImportVersion = EImportStaticMeshVersion::BeforeImportStaticMeshVersionWasAdded;
-	LODForOccluderMesh = -1;
 	NumStreamedLODs.Default = -1;
 	GetHiResSourceModel().StaticMeshDescriptionBulkData = CreateDefaultSubobject<UStaticMeshDescriptionBulkData>(TEXT("HiResMeshDescription"));
 	GetHiResSourceModel().StaticMeshDescriptionBulkData->SetFlags(RF_Transactional);
@@ -2994,33 +2916,6 @@ void UStaticMesh::SetRenderData(TUniquePtr<class FStaticMeshRenderData>&& InRend
 	PRAGMA_ENABLE_DEPRECATION_WARNINGS
 }
 
-FStaticMeshOccluderData* UStaticMesh::GetOccluderData()
-{
-	WaitUntilAsyncPropertyReleased(EStaticMeshAsyncProperties::OccluderData);
-
-	PRAGMA_DISABLE_DEPRECATION_WARNINGS
-	return OccluderData.Get();
-	PRAGMA_ENABLE_DEPRECATION_WARNINGS
-}
-
-const FStaticMeshOccluderData* UStaticMesh::GetOccluderData() const
-{
-	WaitUntilAsyncPropertyReleased(EStaticMeshAsyncProperties::OccluderData);
-
-	PRAGMA_DISABLE_DEPRECATION_WARNINGS
-	return OccluderData.Get();
-	PRAGMA_ENABLE_DEPRECATION_WARNINGS
-}
-
-void UStaticMesh::SetOccluderData(TUniquePtr<class FStaticMeshOccluderData>&& InOccluderData)
-{
-	WaitUntilAsyncPropertyReleased(EStaticMeshAsyncProperties::OccluderData);
-
-	PRAGMA_DISABLE_DEPRECATION_WARNINGS
-	OccluderData = MoveTemp(InOccluderData);
-	PRAGMA_ENABLE_DEPRECATION_WARNINGS
-}
-
 void UStaticMesh::PostInitProperties()
 {
 #if WITH_EDITORONLY_DATA
@@ -3073,11 +2968,6 @@ void UStaticMesh::InitResources()
 		CachedSRRState.bHasPendingInitHint = true;
 	}
 
-	if (GetOccluderData())
-	{
-		INC_DWORD_STAT_BY( STAT_StaticMeshOccluderMemory, GetOccluderData()->GetResourceSizeBytes() );
-	}
-
 #if (WITH_EDITOR && DO_CHECK)
 	if (GetRenderData() && CachedSRRState.bSupportsStreaming && !GetOutermost()->bIsCookedForEditor)
 	{
@@ -3114,11 +3004,6 @@ void UStaticMesh::GetResourceSizeEx(FResourceSizeEx& CumulativeResourceSize)
 	{
 		GetRenderData()->GetResourceSizeEx(CumulativeResourceSize);
 	}
-
-	if (GetOccluderData())
-	{
-		CumulativeResourceSize.AddDedicatedSystemMemoryBytes(GetOccluderData()->GetResourceSizeBytes());
-	}
 }
 
 void FStaticMeshRenderData::GetResourceSizeEx(FResourceSizeEx& CumulativeResourceSize) const
@@ -3415,11 +3300,6 @@ void UStaticMesh::ReleaseResources()
 		ReleaseResourcesFence.BeginFence();
 	}
 
-	if (GetOccluderData())
-	{
-		DEC_DWORD_STAT_BY( STAT_StaticMeshOccluderMemory, GetOccluderData()->GetResourceSizeBytes() );
-	}
-
 	bRenderingResourcesInitialized = false;
 }
 
@@ -4665,10 +4545,6 @@ void UStaticMesh::CacheDerivedData()
 	SetRenderData(MakeUnique<FStaticMeshRenderData>());
 	GetRenderData()->Cache(RunningPlatform, this, LODSettings);
 
-	// Conditionally create occluder data
-	SetOccluderData(FStaticMeshOccluderData::Build(this));
-	ReleaseAsyncProperty(EStaticMeshAsyncProperties::OccluderData);
-
 	// Additionally cache derived data for any other platforms we care about.
 	const TArray<ITargetPlatform*>& TargetPlatforms = TargetPlatformManager.GetActiveTargetPlatforms();
 	for (int32 PlatformIndex = 0; PlatformIndex < TargetPlatforms.Num(); ++PlatformIndex)
@@ -4896,8 +4772,6 @@ void UStaticMesh::Serialize(FArchive& Ar)
 			SCOPE_MS_ACCUMULATOR(STAT_StaticMesh_RenderData);
 			SetRenderData(MakeUnique<FStaticMeshRenderData>());
 			GetRenderData()->Serialize(Ar, this, bCooked);
-
-			FStaticMeshOccluderData::SerializeCooked(Ar, this);
 		}
 #if WITH_EDITOR
 		else if (Ar.IsSaving())
@@ -4909,10 +4783,7 @@ void UStaticMesh::Serialize(FArchive& Ar)
 			}
 
 			FStaticMeshRenderData& PlatformRenderData = GetPlatformStaticMeshRenderData(this, Ar.CookingTarget());
-
 			PlatformRenderData.Serialize(Ar, this, bCooked);
-
-			FStaticMeshOccluderData::SerializeCooked(Ar, this);
 		}
 #endif
 	}
diff --git a/Engine/Source/Runtime/Engine/Private/StaticMeshRender.cpp b/Engine/Source/Runtime/Engine/Private/StaticMeshRender.cpp
--- a/Engine/Source/Runtime/Engine/Private/StaticMeshRender.cpp
+++ b/Engine/Source/Runtime/Engine/Private/StaticMeshRender.cpp
@@ -158,7 +158,6 @@ static TAutoConsoleVariable<float> CVarRayTracingStaticMeshesWPOCullingRadius(
 FStaticMeshSceneProxy::FStaticMeshSceneProxy(UStaticMeshComponent* InComponent, bool bForceLODsShareStaticLighting)
 	: FPrimitiveSceneProxy(InComponent, InComponent->GetStaticMesh()->GetFName())
 	, RenderData(InComponent->GetStaticMesh()->GetRenderData())
-	, OccluderData(InComponent->GetStaticMesh()->GetOccluderData())
 	, ForcedLodModel(InComponent->ForcedLodModel)
 	, bCastShadow(InComponent->CastShadow)
 	, bReverseCulling(InComponent->bReverseCulling)
@@ -666,17 +665,6 @@ bool FStaticMeshSceneProxy::GetMeshElement(
 	}
 }
 
-int32 FStaticMeshSceneProxy::CollectOccluderElements(FOccluderElementsCollector& Collector) const
-{
-	if (OccluderData)
-	{	
-		Collector.AddElements(OccluderData->VerticesSP, OccluderData->IndicesSP, GetLocalToWorld());
-		return 1;
-	}
-	
-	return 0;
-}
-
 void FStaticMeshSceneProxy::CreateRenderThreadResources()
 {
 #if RHI_RAYTRACING
diff --git a/Engine/Source/Runtime/Engine/Public/PrimitiveSceneProxy.h b/Engine/Source/Runtime/Engine/Public/PrimitiveSceneProxy.h
--- a/Engine/Source/Runtime/Engine/Public/PrimitiveSceneProxy.h
+++ b/Engine/Source/Runtime/Engine/Public/PrimitiveSceneProxy.h
@@ -264,9 +264,6 @@ public:
 	ENGINE_API virtual ERayTracingPrimitiveFlags GetCachedRayTracingInstance(FRayTracingInstance& OutRayTracingInstance);
 #endif // RHI_RAYTRACING
 
-	/** Collects occluder geometry for software occlusion culling */
-	virtual int32 CollectOccluderElements(class FOccluderElementsCollector& Collector) const { return 0; }
-
 	/** 
 	 * Gathers the primitive's dynamic mesh elements.  This will only be called if GetViewRelevance declares dynamic relevance.
 	 * This is called from the rendering thread for each set of views that might be rendered.  
diff --git a/Engine/Source/Runtime/Engine/Public/SceneManagement.h b/Engine/Source/Runtime/Engine/Public/SceneManagement.h
--- a/Engine/Source/Runtime/Engine/Public/SceneManagement.h
+++ b/Engine/Source/Runtime/Engine/Public/SceneManagement.h
@@ -2069,26 +2069,6 @@ public:
 };
 
 
-
-/** 
- * Convenience typedefs for a software occlusion mesh elements
- */
-typedef TArray<FVector> FOccluderVertexArray;
-typedef TArray<uint16> FOccluderIndexArray;
-typedef TSharedPtr<FOccluderVertexArray, ESPMode::ThreadSafe> FOccluderVertexArraySP;
-typedef TSharedPtr<FOccluderIndexArray, ESPMode::ThreadSafe> FOccluderIndexArraySP;
-
-/**
- * An interface used to collect primitive occluder geometry.
- */
-class FOccluderElementsCollector
-{
-public:
-	virtual ~FOccluderElementsCollector() {};
-	virtual void AddElements(const FOccluderVertexArraySP& Vertices, const FOccluderIndexArraySP& Indices, const FMatrix& LocalToWorld)
-	{}
-};
-
 /** Primitive draw interface implementation used to store primitives requested to be drawn when gathering dynamic mesh elements. */
 class ENGINE_API FSimpleElementCollector : public FPrimitiveDrawInterface
 {
diff --git a/Engine/Source/Runtime/Engine/Public/StaticMeshResources.h b/Engine/Source/Runtime/Engine/Public/StaticMeshResources.h
--- a/Engine/Source/Runtime/Engine/Public/StaticMeshResources.h
+++ b/Engine/Source/Runtime/Engine/Public/StaticMeshResources.h
@@ -742,25 +742,6 @@ private:
 	bool bIsInitialized = false;
 };
 
-/**
- * This geometry is used to rasterize mesh for software occlusion
- * Generated only for platforms that support ETargetPlatformFeatures::SoftwareOcclusion
- */
-class FStaticMeshOccluderData
-{
-public:
-	FStaticMeshOccluderData();
-
-	FOccluderVertexArraySP VerticesSP;
-	FOccluderIndexArraySP IndicesSP;
-
-	SIZE_T GetResourceSizeBytes() const;
-
-	static TUniquePtr<FStaticMeshOccluderData> Build(UStaticMesh* Owner);
-	/** Serialization. */
-	static void SerializeCooked(FArchive& Ar, UStaticMesh* Owner);
-};
-
 /**
  * FStaticMeshComponentRecreateRenderStateContext - Destroys render state for all StaticMeshComponents using a given StaticMesh and 
  * recreates them when it goes out of scope. Used to ensure stale rendering data isn't kept around in the components when importing
@@ -910,8 +891,6 @@ public:
 		bool bAllowPreCulledIndices,
 		FMeshBatch& OutMeshBatch) const;
 
-	virtual int32 CollectOccluderElements(class FOccluderElementsCollector& Collector) const override;
-
 	virtual void CreateRenderThreadResources() override;
 
 	virtual void DestroyRenderThreadResources() override;
@@ -1076,8 +1055,6 @@ protected:
 
 	FStaticMeshRenderData* RenderData;
 
-	FStaticMeshOccluderData* OccluderData;
-
 	TArray<FLODInfo> LODs;
 
 	const FDistanceFieldVolumeData* DistanceFieldData;	
diff --git a/Engine/Source/Runtime/Landscape/Classes/LandscapeProxy.h b/Engine/Source/Runtime/Landscape/Classes/LandscapeProxy.h
--- a/Engine/Source/Runtime/Landscape/Classes/LandscapeProxy.h
+++ b/Engine/Source/Runtime/Landscape/Classes/LandscapeProxy.h
@@ -408,10 +408,6 @@ public:
 	UPROPERTY(EditAnywhere, Category = "LOD Distribution", meta = (DisplayName = "Other LODs", ClampMin = "1.0", ClampMax = "10.0", UIMin = "1.0", UIMax = "10.0"))
 	float LODDistributionSetting;
 
-	/** Landscape LOD to use as an occluder geometry for software occlusion */
-	UPROPERTY(EditAnywhere, AdvancedDisplay, Category=LOD)
-	int32 OccluderGeometryLOD;
-
 #if WITH_EDITORONLY_DATA
 	/** LOD level to use when exporting the landscape to obj or FBX */
 	UPROPERTY(EditAnywhere, Category=LOD, AdvancedDisplay)
diff --git a/Engine/Source/Runtime/Landscape/Private/Landscape.cpp b/Engine/Source/Runtime/Landscape/Private/Landscape.cpp
--- a/Engine/Source/Runtime/Landscape/Private/Landscape.cpp
+++ b/Engine/Source/Runtime/Landscape/Private/Landscape.cpp
@@ -111,7 +111,6 @@ DEFINE_STAT(STAT_LandscapeLayersResolveWeightmaps);
 DEFINE_STAT(STAT_LandscapeLayersRegenerateWeightmaps);
 
 DEFINE_STAT(STAT_LandscapeVertexMem);
-DEFINE_STAT(STAT_LandscapeOccluderMem);
 DEFINE_STAT(STAT_LandscapeHoleMem);
 DEFINE_STAT(STAT_LandscapeComponentMem);
 
@@ -133,7 +132,7 @@ namespace LandscapeCookStats
 // differences, etc.) replace the version GUID below with a new one.
 // In case of merge conflicts with DDC versions, you MUST generate a new GUID
 // and set this new GUID as the version.                                       
-#define LANDSCAPE_MOBILE_COOK_VERSION TEXT("C2781BC371584D8BB9A2AB8BFDFDF7A7")
+#define LANDSCAPE_MOBILE_COOK_VERSION TEXT("43D02EF867C74B71A0D4E0FA41392707")
 
 #define LOCTEXT_NAMESPACE "Landscape"
 
@@ -1032,7 +1031,6 @@ ALandscapeProxy::ALandscapeProxy(const FObjectInitializer& ObjectInitializer)
 	StreamingDistanceMultiplier = 1.0f;
 	MaxLODLevel = -1;
 	bUseDynamicMaterialInstance = false;
-	OccluderGeometryLOD = 1; // 1 - usually is a good default
 #if WITH_EDITORONLY_DATA
 	bLockLocation = true;
 	bIsMovingToLevel = false;
@@ -2542,7 +2540,6 @@ void ALandscapeProxy::GetSharedProperties(ALandscapeProxy* Landscape)
 		LODDistributionSetting = Landscape->LODDistributionSetting;
 		LOD0DistributionSetting = Landscape->LOD0DistributionSetting;
 		LOD0ScreenSize = Landscape->LOD0ScreenSize;
-		OccluderGeometryLOD = Landscape->OccluderGeometryLOD;
 		NegativeZBoundsExtension = Landscape->NegativeZBoundsExtension;
 		PositiveZBoundsExtension = Landscape->PositiveZBoundsExtension;
 		CollisionMipLevel = Landscape->CollisionMipLevel;
@@ -2612,12 +2609,6 @@ void ALandscapeProxy::FixupSharedData(ALandscape* Landscape)
 		bUpdated = true;
 	}
 
-	if (OccluderGeometryLOD != Landscape->OccluderGeometryLOD)
-	{
-		OccluderGeometryLOD = Landscape->OccluderGeometryLOD;
-		bUpdated = true;
-	}
-
 	if (TargetDisplayOrder != Landscape->TargetDisplayOrder)
 	{
 		TargetDisplayOrder = Landscape->TargetDisplayOrder;
@@ -3685,9 +3676,6 @@ void ULandscapeComponent::SerializeStateHashes(FArchive& Ar)
 		Ar << WeightmapGuid;
 	}
 
-	int32 OccluderGeometryLOD = GetLandscapeProxy()->OccluderGeometryLOD;
-	Ar << OccluderGeometryLOD;
-
 	bool bMeshHoles = GetLandscapeProxy()->bMeshHoles;
 	uint8 MeshHolesMaxLod = GetLandscapeProxy()->MeshHolesMaxLod;
 	Ar << bMeshHoles << MeshHolesMaxLod;
diff --git a/Engine/Source/Runtime/Landscape/Private/LandscapeEdit.cpp b/Engine/Source/Runtime/Landscape/Private/LandscapeEdit.cpp
--- a/Engine/Source/Runtime/Landscape/Private/LandscapeEdit.cpp
+++ b/Engine/Source/Runtime/Landscape/Private/LandscapeEdit.cpp
@@ -4690,7 +4690,6 @@ bool ALandscapeProxy::CanEditChange(const FProperty* InProperty) const
 			PropertyName == GET_MEMBER_NAME_CHECKED(ALandscapeProxy, LODDistributionSetting) ||
 			PropertyName == GET_MEMBER_NAME_CHECKED(ALandscapeProxy, LOD0DistributionSetting) ||
 			PropertyName == GET_MEMBER_NAME_CHECKED(ALandscapeProxy, LOD0ScreenSize) ||
-			PropertyName == GET_MEMBER_NAME_CHECKED(ALandscapeProxy, OccluderGeometryLOD) ||
 			PropertyName == GET_MEMBER_NAME_CHECKED(ALandscapeProxy, TargetDisplayOrder) ||
 			PropertyName == GET_MEMBER_NAME_CHECKED(ALandscapeProxy, TargetDisplayOrderList))
 		{
@@ -4840,7 +4839,7 @@ void ALandscapeProxy::PostEditChangeProperty(FPropertyChangedEvent& PropertyChan
 		}
 	}
 	else if (GIsEditor && 
-		(PropertyName == GET_MEMBER_NAME_CHECKED(ALandscapeProxy, OccluderGeometryLOD) || PropertyName == GET_MEMBER_NAME_CHECKED(ALandscapeProxy, bMeshHoles) || PropertyName == GET_MEMBER_NAME_CHECKED(ALandscapeProxy, MeshHolesMaxLod)))
+		(PropertyName == GET_MEMBER_NAME_CHECKED(ALandscapeProxy, bMeshHoles) || PropertyName == GET_MEMBER_NAME_CHECKED(ALandscapeProxy, MeshHolesMaxLod)))
 	{
 		CheckGenerateLandscapePlatformData(false, nullptr);
 		MarkComponentsRenderStateDirty();
@@ -5163,10 +5162,6 @@ void ALandscape::PostEditChangeProperty(FPropertyChangedEvent& PropertyChangedEv
 	{
 		bPropagateToProxies = true;
 	}
-	else if (PropertyName == GET_MEMBER_NAME_CHECKED(ALandscapeProxy, OccluderGeometryLOD))
-	{
-		bPropagateToProxies = true;
-	}
 	else if (GIsEditor && PropertyName == FName(TEXT("StaticLightingResolution")))
 	{
 		StaticLightingResolution = ::AdjustStaticLightingResolution(StaticLightingResolution, NumSubsections, SubsectionSizeQuads, ComponentSizeQuads);
@@ -7098,48 +7093,6 @@ void ULandscapeComponent::GeneratePlatformVertexData(const ITargetPlatform* Targ
 	PlatformAr << NumInlineMobileVertices;
 	PlatformAr.Serialize(InlineMobileVertices.GetData(), NumInlineMobileVertices*sizeof(FLandscapeMobileVertex));
 
-	// Generate occlusion mesh
-	TArray<FVector> OccluderVertices;
-	const int32 OcclusionMeshMip = FMath::Clamp<int32>(GetLandscapeProxy()->OccluderGeometryLOD, -1, HeightmapMipData.Num() - 1);
-
-	if (OcclusionMeshMip >= 0 && (!TargetPlatform || TargetPlatform->SupportsFeature(ETargetPlatformFeatures::SoftwareOcclusion)))
-	{
-		int32 LodSubsectionSizeQuads = (SubsectionSizeVerts >> OcclusionMeshMip) - 1;
-		float MipRatio = (float)SubsectionSizeQuads / (float)LodSubsectionSizeQuads;
-		
-		for (int32 SubY = 0; SubY < NumSubsections; SubY++)
-		{
-			for (int32 SubX = 0; SubX < NumSubsections; SubX++)
-			{
-				float HeightmapScaleBiasZ = HeightmapScaleBias.Z + HeightmapSubsectionOffsetU * (float)SubX;
-				float HeightmapScaleBiasW = HeightmapScaleBias.W + HeightmapSubsectionOffsetV * (float)SubY;
-				int32 BaseMipOfsX = FMath::RoundToInt(HeightmapScaleBiasZ * (float)GetHeightmap()->Source.GetSizeX());
-				int32 BaseMipOfsY = FMath::RoundToInt(HeightmapScaleBiasW * (float)GetHeightmap()->Source.GetSizeY());
-
-				for (int32 y = 0; y <= LodSubsectionSizeQuads; y++)
-				{
-					for (int32 x = 0; x <= LodSubsectionSizeQuads; x++)
-					{
-						int32 MipSizeX = GetHeightmap()->Source.GetSizeX() >> OcclusionMeshMip;
-
-						int32 CurrentMipOfsX = BaseMipOfsX >> OcclusionMeshMip;
-						int32 CurrentMipOfsY = BaseMipOfsY >> OcclusionMeshMip;
-												
-						FColor* CurrentMipSrcRow = HeightmapMipData[OcclusionMeshMip] + (CurrentMipOfsY + y) * MipSizeX + CurrentMipOfsX;
-						uint16 Height = CurrentMipSrcRow[x].R << 8 | CurrentMipSrcRow[x].G;
-
-						FVector VtxPos = FVector(x*MipRatio + SubX * SubsectionSizeQuads, y*MipRatio + SubY * SubsectionSizeQuads, ((float)Height - 32768.f) * LANDSCAPE_ZSCALE);
-						OccluderVertices.Add(VtxPos);
-					}
-				}
-			}
-		}
-	}
-
-	int32 NumOccluderVerices = OccluderVertices.Num();
-	PlatformAr << NumOccluderVerices;
-	PlatformAr.Serialize(OccluderVertices.GetData(), NumOccluderVerices*sizeof(FVector));
-	
 	// Copy to PlatformData as Compressed
 	PlatformData.InitializeFromUncompressedData(NewPlatformData, StreamingLODData);
 }
diff --git a/Engine/Source/Runtime/Landscape/Private/LandscapePrivate.h b/Engine/Source/Runtime/Landscape/Private/LandscapePrivate.h
--- a/Engine/Source/Runtime/Landscape/Private/LandscapePrivate.h
+++ b/Engine/Source/Runtime/Landscape/Private/LandscapePrivate.h
@@ -42,6 +42,5 @@ DECLARE_DWORD_COUNTER_STAT_EXTERN(TEXT("Render Passes"), STAT_LandscapeComponent
 DECLARE_DWORD_COUNTER_STAT_EXTERN(TEXT("DrawCalls"), STAT_LandscapeDrawCalls, STATGROUP_Landscape, );
 
 DECLARE_MEMORY_STAT_EXTERN(TEXT("Vertex Mem"), STAT_LandscapeVertexMem, STATGROUP_Landscape, );
-DECLARE_MEMORY_STAT_EXTERN(TEXT("Occluder Mem"), STAT_LandscapeOccluderMem, STATGROUP_Landscape, );
 DECLARE_MEMORY_STAT_EXTERN(TEXT("Hole Mem"), STAT_LandscapeHoleMem, STATGROUP_Landscape, );
 DECLARE_DWORD_COUNTER_STAT_EXTERN(TEXT("Component Mem"), STAT_LandscapeComponentMem, STATGROUP_Landscape, );
diff --git a/Engine/Source/Runtime/Landscape/Private/LandscapeRender.cpp b/Engine/Source/Runtime/Landscape/Private/LandscapeRender.cpp
--- a/Engine/Source/Runtime/Landscape/Private/LandscapeRender.cpp
+++ b/Engine/Source/Runtime/Landscape/Private/LandscapeRender.cpp
@@ -1366,7 +1366,7 @@ void FLandscapeComponentSceneProxy::CreateRenderThreadResources()
 	{
 		SharedBuffers = new FLandscapeSharedBuffers(
 			SharedBuffersKey, SubsectionSizeQuads, NumSubsections,
-			FeatureLevel, /*NumOcclusionVertices*/ 0);
+			FeatureLevel);
 
 		FLandscapeComponentSceneProxy::SharedBuffersMap.Add(SharedBuffersKey, SharedBuffers);
 
@@ -2630,12 +2630,6 @@ void FLandscapeComponentSceneProxy::GetDynamicRayTracingInstances(FRayTracingMat
 }
 #endif
 
-int32 FLandscapeComponentSceneProxy::CollectOccluderElements(FOccluderElementsCollector& Collector) const
-{
-	// TODO: implement
-	return 0;
-}
-
 //
 // FLandscapeVertexBuffer
 //
@@ -2874,45 +2868,6 @@ void FLandscapeSharedBuffers::CreateIndexBuffers(ERHIFeatureLevel::Type InFeatur
 	}
 }
 
-void FLandscapeSharedBuffers::CreateOccluderIndexBuffer(int32 NumOccluderVertices)
-{
-	if (NumOccluderVertices <= 0 || NumOccluderVertices > MAX_uint16)
-	{
-		return;
-	}
-
-	uint16 NumLineQuads = ((uint16)FMath::Sqrt(static_cast<float>(NumOccluderVertices)) - 1);
-	uint16 NumLineVtx = NumLineQuads + 1;
-	check(NumLineVtx*NumLineVtx == NumOccluderVertices);
-
-	int32 NumTris = NumLineQuads*NumLineQuads * 2;
-	int32 NumIndices = NumTris * 3;
-	OccluderIndicesSP = MakeShared<FOccluderIndexArray, ESPMode::ThreadSafe>();
-	OccluderIndicesSP->SetNumUninitialized(NumIndices, false);
-
-	uint16* OcclusionIndices = OccluderIndicesSP->GetData();
-	const uint16 NumLineVtxPlusOne = NumLineVtx + 1;
-	const uint16 QuadIndices[2][3] = { {0, NumLineVtx, NumLineVtxPlusOne}, {0, NumLineVtxPlusOne, 1} };
-	uint16 QuadOffset = 0;
-	int32 Index = 0;
-	for (int32 y = 0; y < NumLineQuads; y++)
-	{
-		for (int32 x = 0; x < NumLineQuads; x++)
-		{
-			for (int32 i = 0; i < 2; i++)
-			{
-				OcclusionIndices[Index++] = QuadIndices[i][0] + QuadOffset;
-				OcclusionIndices[Index++] = QuadIndices[i][1] + QuadOffset;
-				OcclusionIndices[Index++] = QuadIndices[i][2] + QuadOffset;
-			}
-			QuadOffset++;
-		}
-		QuadOffset++;
-	}
-
-	INC_DWORD_STAT_BY(STAT_LandscapeOccluderMem, OccluderIndicesSP->GetAllocatedSize());
-}
-
 #if WITH_EDITOR
 template <typename INDEX_TYPE>
 void FLandscapeSharedBuffers::CreateGrassIndexBuffer()
@@ -2960,7 +2915,7 @@ void FLandscapeSharedBuffers::CreateGrassIndexBuffer()
 }
 #endif
 
-FLandscapeSharedBuffers::FLandscapeSharedBuffers(const int32 InSharedBuffersKey, const int32 InSubsectionSizeQuads, const int32 InNumSubsections, const ERHIFeatureLevel::Type InFeatureLevel, int32 NumOccluderVertices)
+FLandscapeSharedBuffers::FLandscapeSharedBuffers(const int32 InSharedBuffersKey, const int32 InSubsectionSizeQuads, const int32 InNumSubsections, const ERHIFeatureLevel::Type InFeatureLevel)
 	: SharedBuffersKey(InSharedBuffersKey)
 	, NumIndexBuffers(FMath::CeilLogTwo(InSubsectionSizeQuads + 1))
 	, SubsectionSizeVerts(InSubsectionSizeQuads + 1)
@@ -3012,8 +2967,6 @@ FLandscapeSharedBuffers::FLandscapeSharedBuffers(const int32 InSharedBuffersKey,
 		}
 #endif
 	}
-
-	CreateOccluderIndexBuffer(NumOccluderVertices);
 }
 
 FLandscapeSharedBuffers::~FLandscapeSharedBuffers()
@@ -3049,11 +3002,6 @@ FLandscapeSharedBuffers::~FLandscapeSharedBuffers()
 #endif
 
 	delete VertexFactory;
-
-	if (OccluderIndicesSP.IsValid())
-	{
-		DEC_DWORD_STAT_BY(STAT_LandscapeOccluderMem, OccluderIndicesSP->GetAllocatedSize());
-	}
 }
 
 //
diff --git a/Engine/Source/Runtime/Landscape/Private/LandscapeRenderMobile.cpp b/Engine/Source/Runtime/Landscape/Private/LandscapeRenderMobile.cpp
--- a/Engine/Source/Runtime/Landscape/Private/LandscapeRenderMobile.cpp
+++ b/Engine/Source/Runtime/Landscape/Private/LandscapeRenderMobile.cpp
@@ -304,19 +304,6 @@ FLandscapeMobileRenderData::FLandscapeMobileRenderData(const TArray<uint8>& InPl
 		MemAr.Serialize(VertexData.GetData(), VertexData.Num());
 		VertexBuffer = new FLandscapeVertexBufferMobile(MoveTemp(VertexData));
 	}
-
-	{
-		int32 NumOccluderVertices = 0;
-		MemAr << NumOccluderVertices;
-		if (NumOccluderVertices > 0)
-		{
-			OccluderVerticesSP = MakeShared<FOccluderVertexArray, ESPMode::ThreadSafe>();
-			OccluderVerticesSP->SetNumUninitialized(NumOccluderVertices);
-			MemAr.Serialize(OccluderVerticesSP->GetData(), NumOccluderVertices * sizeof(FVector));
-
-			INC_DWORD_STAT_BY(STAT_LandscapeOccluderMem, OccluderVerticesSP->GetAllocatedSize());
-		}
-	}
 }
 
 FLandscapeMobileRenderData::~FLandscapeMobileRenderData()
@@ -341,11 +328,6 @@ FLandscapeMobileRenderData::~FLandscapeMobileRenderData()
 			});
 		}
 	}
-
-	if (OccluderVerticesSP.IsValid())
-	{
-		DEC_DWORD_STAT_BY(STAT_LandscapeOccluderMem, OccluderVerticesSP->GetAllocatedSize());
-	}
 }
 
 FLandscapeComponentSceneProxyMobile::FLandscapeComponentSceneProxyMobile(ULandscapeComponent* InComponent)
@@ -373,17 +355,6 @@ FLandscapeComponentSceneProxyMobile::FLandscapeComponentSceneProxyMobile(ULandsc
 #endif
 }
 
-int32 FLandscapeComponentSceneProxyMobile::CollectOccluderElements(FOccluderElementsCollector& Collector) const
-{
-	if (MobileRenderData->OccluderVerticesSP.IsValid() && SharedBuffers->OccluderIndicesSP.IsValid())
-	{
-		Collector.AddElements(MobileRenderData->OccluderVerticesSP, SharedBuffers->OccluderIndicesSP, GetLocalToWorld());
-		return 1;
-	}
-
-	return 0;
-}
-
 FLandscapeComponentSceneProxyMobile::~FLandscapeComponentSceneProxyMobile()
 {
 	if (VertexFactory)
@@ -414,11 +385,9 @@ void FLandscapeComponentSceneProxyMobile::CreateRenderThreadResources()
 	SharedBuffers = FLandscapeComponentSceneProxy::SharedBuffersMap.FindRef(SharedBuffersKey);
 	if (SharedBuffers == nullptr)
 	{
-		int32 NumOcclusionVertices = MobileRenderData->OccluderVerticesSP.IsValid() ? MobileRenderData->OccluderVerticesSP->Num() : 0;
-				
 		SharedBuffers = new FLandscapeSharedBuffers(
 			SharedBuffersKey, SubsectionSizeQuads, NumSubsections,
-			GetScene().GetFeatureLevel(), NumOcclusionVertices);
+			GetScene().GetFeatureLevel());
 
 		FLandscapeComponentSceneProxy::SharedBuffersMap.Add(SharedBuffersKey, SharedBuffers);
 	}
diff --git a/Engine/Source/Runtime/Landscape/Public/LandscapeRender.h b/Engine/Source/Runtime/Landscape/Public/LandscapeRender.h
--- a/Engine/Source/Runtime/Landscape/Public/LandscapeRender.h
+++ b/Engine/Source/Runtime/Landscape/Public/LandscapeRender.h
@@ -317,7 +317,6 @@ public:
 	FLandscapeVertexBuffer* VertexBuffer;
 	FIndexBuffer** IndexBuffers;
 	FLandscapeIndexRanges* IndexRanges;
-	FOccluderIndexArraySP OccluderIndicesSP;
 	bool bUse32BitIndices;
 #if WITH_EDITOR
 	FIndexBuffer* GrassIndexBuffer;
@@ -328,12 +327,10 @@ public:
 	TArray<FIndexBuffer*> ZeroOffsetIndexBuffers;
 #endif
 
-	FLandscapeSharedBuffers(int32 SharedBuffersKey, int32 SubsectionSizeQuads, int32 NumSubsections, ERHIFeatureLevel::Type FeatureLevel, int32 NumOcclusionVertices);
+	FLandscapeSharedBuffers(int32 SharedBuffersKey, int32 SubsectionSizeQuads, int32 NumSubsections, ERHIFeatureLevel::Type FeatureLevel);
 
 	template <typename INDEX_TYPE>
 	void CreateIndexBuffers(ERHIFeatureLevel::Type InFeatureLevel);
-
-	void CreateOccluderIndexBuffer(int32 NumOcclderVertices);
 	
 #if WITH_EDITOR
 	template <typename INDEX_TYPE>
@@ -910,7 +907,6 @@ public:
 	virtual void ApplyWorldOffset(FVector InOffset) override;
 	virtual void DrawStaticElements(FStaticPrimitiveDrawInterface* PDI) override;
 	virtual void GetDynamicMeshElements(const TArray<const FSceneView*>& Views, const FSceneViewFamily& ViewFamily, uint32 VisibilityMap, FMeshElementCollector& Collector) const override;
-	virtual int32 CollectOccluderElements(FOccluderElementsCollector& Collector) const override;
 	virtual uint32 GetMemoryFootprint() const override { return(sizeof(*this) + GetAllocatedSize()); }
 	virtual FPrimitiveViewRelevance GetViewRelevance(const FSceneView* View) const override;
 	virtual bool CanBeOccluded() const override;
diff --git a/Engine/Source/Runtime/Landscape/Public/LandscapeRenderMobile.h b/Engine/Source/Runtime/Landscape/Public/LandscapeRenderMobile.h
--- a/Engine/Source/Runtime/Landscape/Public/LandscapeRenderMobile.h
+++ b/Engine/Source/Runtime/Landscape/Public/LandscapeRenderMobile.h
@@ -116,7 +116,6 @@ struct FLandscapeMobileRenderData
 {
 	FLandscapeVertexBufferMobile* VertexBuffer = nullptr;
 	struct FLandscapeMobileHoleData* HoleData = nullptr;
-	FOccluderVertexArraySP OccluderVerticesSP;
 	uint8 CurrentFirstLODIdx;
 	bool bReadyForStreaming = false;
 
@@ -139,7 +138,6 @@ public:
 	FLandscapeComponentSceneProxyMobile(ULandscapeComponent* InComponent);
 
 	virtual void CreateRenderThreadResources() override;
-	virtual int32 CollectOccluderElements(FOccluderElementsCollector& Collector) const override;
 
 	friend class FLandscapeVertexBufferMobile;
 
diff --git a/Engine/Source/Runtime/RenderCore/Private/RenderCore.cpp b/Engine/Source/Runtime/RenderCore/Private/RenderCore.cpp
--- a/Engine/Source/Runtime/RenderCore/Private/RenderCore.cpp
+++ b/Engine/Source/Runtime/RenderCore/Private/RenderCore.cpp
@@ -111,7 +111,6 @@ DEFINE_STAT(STAT_StaticRelevance);
 DEFINE_STAT(STAT_ViewRelevance);
 DEFINE_STAT(STAT_ComputeViewRelevance);
 DEFINE_STAT(STAT_OcclusionCull);
-DEFINE_STAT(STAT_SoftwareOcclusionCull);
 DEFINE_STAT(STAT_UpdatePrimitiveFading);
 DEFINE_STAT(STAT_FrustumCull);
 DEFINE_STAT(STAT_DecompressPrecomputedOcclusion);
diff --git a/Engine/Source/Runtime/RenderCore/Public/RenderCore.h b/Engine/Source/Runtime/RenderCore/Public/RenderCore.h
--- a/Engine/Source/Runtime/RenderCore/Public/RenderCore.h
+++ b/Engine/Source/Runtime/RenderCore/Public/RenderCore.h
@@ -51,7 +51,6 @@ DECLARE_CYCLE_STAT_EXTERN(TEXT("Decompress Occlusion"),STAT_DecompressPrecompute
 DECLARE_CYCLE_STAT_EXTERN(TEXT("Frustum Cull"),STAT_FrustumCull,STATGROUP_InitViews, RENDERCORE_API);
 DECLARE_CYCLE_STAT_EXTERN(TEXT("Update Fading"),STAT_UpdatePrimitiveFading,STATGROUP_InitViews, RENDERCORE_API);
 DECLARE_CYCLE_STAT_EXTERN(TEXT("Occlusion Cull"),STAT_OcclusionCull,STATGROUP_InitViews, RENDERCORE_API);
-DECLARE_CYCLE_STAT_EXTERN(TEXT("Software Occlusion Cull"),STAT_SoftwareOcclusionCull,STATGROUP_InitViews, RENDERCORE_API);
 DECLARE_CYCLE_STAT_EXTERN(TEXT("View Relevance"),STAT_ViewRelevance,STATGROUP_InitViews, RENDERCORE_API);
 DECLARE_CYCLE_STAT_EXTERN(TEXT("Compute View Relevance"),STAT_ComputeViewRelevance,STATGROUP_InitViews, RENDERCORE_API);
 DECLARE_CYCLE_STAT_EXTERN(TEXT("Static Mesh Relevance"),STAT_StaticRelevance,STATGROUP_InitViews, RENDERCORE_API);
diff --git a/Engine/Source/Runtime/Renderer/Private/SceneOcclusion.cpp b/Engine/Source/Runtime/Renderer/Private/SceneOcclusion.cpp
--- a/Engine/Source/Runtime/Renderer/Private/SceneOcclusion.cpp
+++ b/Engine/Source/Runtime/Renderer/Private/SceneOcclusion.cpp
@@ -56,13 +56,6 @@ FAutoConsoleVariableRef CVarOcclusionCullCascadedShadowMaps(
 	ECVF_Scalability | ECVF_RenderThreadSafe
 	);
 
-static TAutoConsoleVariable<int32> CVarMobileAllowSoftwareOcclusion(
-	TEXT("r.Mobile.AllowSoftwareOcclusion"),
-	0,
-	TEXT("Whether to allow rasterizing scene on CPU for primitive occlusion.\n"),
-	ECVF_RenderThreadSafe
-	);
-
 static TAutoConsoleVariable<bool> CVarMobileEnableOcclusionExtraFrame(
 	TEXT("r.Mobile.EnableOcclusionExtraFrame"),
 	true,
@@ -294,21 +287,6 @@ bool FSceneViewState::IsShadowOccluded(FRHICommandListImmediate& RHICmdList, FSc
 	}
 }
 
-void FSceneViewState::ConditionallyAllocateSceneSoftwareOcclusion(ERHIFeatureLevel::Type InFeatureLevel)
-{
-	bool bMobileAllowSoftwareOcclusion = CVarMobileAllowSoftwareOcclusion.GetValueOnAnyThread() != 0;
-	bool bShouldBeEnabled = InFeatureLevel <= ERHIFeatureLevel::ES3_1 && bMobileAllowSoftwareOcclusion;
-
-	if (bShouldBeEnabled && !SceneSoftwareOcclusion)
-	{
-		SceneSoftwareOcclusion = MakeUnique<FSceneSoftwareOcclusion>();
-	}
-	else if (!bShouldBeEnabled && SceneSoftwareOcclusion)
-	{
-		SceneSoftwareOcclusion.Reset();
-	}
-}
-
 void FSceneViewState::Destroy()
 {
 	FSceneViewState* self = this;
diff --git a/Engine/Source/Runtime/Renderer/Private/ScenePrivate.h b/Engine/Source/Runtime/Renderer/Private/ScenePrivate.h
--- a/Engine/Source/Runtime/Renderer/Private/ScenePrivate.h
+++ b/Engine/Source/Runtime/Renderer/Private/ScenePrivate.h
@@ -40,7 +40,6 @@
 #include "BasePassRendering.h"
 #include "MobileBasePassRendering.h"
 #include "VolumeRendering.h"
-#include "SceneSoftwareOcclusion.h"
 #include "CommonRenderResources.h"
 #include "VisualizeTexture.h"
 #include "UnifiedBuffer.h"
@@ -798,14 +797,9 @@ public:
 	FHLODVisibilityState HLODVisibilityState;
 	TMap<FPrimitiveComponentId, FHLODSceneNodeVisibilityState> HLODSceneNodeVisibilityStates;
 
-	// Software occlusion data
-	TUniquePtr<FSceneSoftwareOcclusion> SceneSoftwareOcclusion;
-
 	void UpdatePreExposure(FViewInfo& View);
 
 private:
-	void ConditionallyAllocateSceneSoftwareOcclusion(ERHIFeatureLevel::Type InFeatureLevel);
-
 	/** The current frame PreExposure */
 	float PreExposure;
 
@@ -1369,13 +1363,6 @@ public:
 		}
 	}
 
-	/** called in InitViews() */
-	void OnStartRender(FViewInfo& View, FSceneViewFamily& ViewFamily)
-	{
-		check(IsInRenderingThread());
-		ConditionallyAllocateSceneSoftwareOcclusion(View.GetFeatureLevel());
-	}
-
 	// needed for GetReusableMID()
 	virtual void OnStartPostProcessing(FSceneView& CurrentView) override
 	{
diff --git a/Engine/Source/Runtime/Renderer/Private/SceneRendering.cpp b/Engine/Source/Runtime/Renderer/Private/SceneRendering.cpp
--- a/Engine/Source/Runtime/Renderer/Private/SceneRendering.cpp
+++ b/Engine/Source/Runtime/Renderer/Private/SceneRendering.cpp
@@ -44,7 +44,6 @@
 #include "DeviceProfiles/DeviceProfileManager.h"
 #include "DeviceProfiles/DeviceProfile.h"
 #include "PostProcess/PostProcessing.h"
-#include "SceneSoftwareOcclusion.h"
 #include "VirtualTexturing.h"
 #include "VisualizeTexturePresent.h"
 #include "GPUScene.h"
@@ -3230,12 +3229,6 @@ void FSceneRenderer::RenderFinish(FRDGBuilder& GraphBuilder, FRDGTextureRef View
 						FXInterface->DrawDebug_RenderThread(GraphBuilder, View, Output);
 					}
 				}
-				
-				// Software occlusion debug draw
-				if (ViewState && ViewState->SceneSoftwareOcclusion)
-				{
-					ViewState->SceneSoftwareOcclusion->DebugDraw(GraphBuilder, View, Output, 20, 20);
-				}
 			}
 		}
 	}
@@ -3425,14 +3418,6 @@ FSceneRenderer* FSceneRenderer::CreateSceneRenderer(const FSceneViewFamily* InVi
 void FSceneRenderer::OnStartRender(FRHICommandListImmediate& RHICmdList)
 {
 	FVisualizeTexturePresent::OnStartRender(Views[0]);
-
-	for (FViewInfo& View : Views)
-	{
-		if (View.ViewState)
-		{
-			View.ViewState->OnStartRender(View, ViewFamily);
-		}
-	}
 }
 
 bool FSceneRenderer::ShouldCompositeEditorPrimitives(const FViewInfo& View)
diff --git a/Engine/Source/Runtime/Renderer/Private/SceneSoftwareOcclusion.cpp b/Engine/Source/Runtime/Renderer/Private/SceneSoftwareOcclusion.cpp
deleted file mode 100644
diff --git a/Engine/Source/Runtime/Renderer/Private/SceneSoftwareOcclusion.h b/Engine/Source/Runtime/Renderer/Private/SceneSoftwareOcclusion.h
deleted file mode 100644
diff --git a/Engine/Source/Runtime/Renderer/Private/SceneVisibility.cpp b/Engine/Source/Runtime/Renderer/Private/SceneVisibility.cpp
--- a/Engine/Source/Runtime/Renderer/Private/SceneVisibility.cpp
+++ b/Engine/Source/Runtime/Renderer/Private/SceneVisibility.cpp
@@ -33,7 +33,6 @@
 #include "FXSystem.h"
 #include "PostProcess/PostProcessing.h"
 #include "SceneView.h"
-#include "SceneSoftwareOcclusion.h"
 #include "Engine/LODActor.h"
 #include "GPUScene.h"
 #include "TranslucentRendering.h"
@@ -1692,64 +1691,47 @@ static int32 OcclusionCull(FRHICommandListImmediate& RHICmdList, const FScene* S
 	float CurrentRealTime = View.Family->CurrentRealTime;
 	if (ViewState)
 	{
-		if (ViewState->SceneSoftwareOcclusion)
-		{
-			SCOPE_CYCLE_COUNTER(STAT_SoftwareOcclusionCull)
-			// TODO: Respect PrimitivesAlwaysVisible
-			NumOccludedPrimitives += ViewState->SceneSoftwareOcclusion->Process(RHICmdList, Scene, View);
-		}
-		else if (Scene->GetFeatureLevel() >= ERHIFeatureLevel::ES3_1)
-		{
-			bool bSubmitQueries = !View.bDisableQuerySubmissions;
+		bool bSubmitQueries = !View.bDisableQuerySubmissions;
 #if !(UE_BUILD_SHIPPING || UE_BUILD_TEST)
-			bSubmitQueries = bSubmitQueries && !ViewState->HasViewParent() && !ViewState->bIsFrozen;
+		bSubmitQueries = bSubmitQueries && !ViewState->HasViewParent() && !ViewState->bIsFrozen;
 #endif
 
-			if( bHZBOcclusion )
-			{
-				QUICK_SCOPE_CYCLE_COUNTER(STAT_MapHZBResults);
-				check(!ViewState->HZBOcclusionTests.IsValidFrame(ViewState->OcclusionFrameCounter));
-				ViewState->HZBOcclusionTests.MapResults(RHICmdList);
-			}
+		if( bHZBOcclusion )
+		{
+			QUICK_SCOPE_CYCLE_COUNTER(STAT_MapHZBResults);
+			check(!ViewState->HZBOcclusionTests.IsValidFrame(ViewState->OcclusionFrameCounter));
+			ViewState->HZBOcclusionTests.MapResults(RHICmdList);
+		}
  
-			// Perform round-robin occlusion queries
-			if (View.ViewState->IsRoundRobinEnabled() &&
-				!View.bIsSceneCapture && // We only round-robin on the main renderer (not scene captures)
-				!View.bIgnoreExistingQueries && // We do not alternate occlusion queries when we want to refresh the occlusion history
-				(IStereoRendering::IsStereoEyeView(View))) // Only relevant to stereo views
-			{
-				// For even frames, prevent left eye from occlusion querying
-				// For odd frames, prevent right eye from occlusion querying
-				const bool FrameParity = ((View.ViewState->PrevFrameNumber & 0x01) == 1);
-				bSubmitQueries &= (FrameParity  && IStereoRendering::IsAPrimaryView(View)) ||
-								  (!FrameParity && IStereoRendering::IsASecondaryView(View));
-			}
+		// Perform round-robin occlusion queries
+		if (View.ViewState->IsRoundRobinEnabled() &&
+			!View.bIsSceneCapture && // We only round-robin on the main renderer (not scene captures)
+			!View.bIgnoreExistingQueries && // We do not alternate occlusion queries when we want to refresh the occlusion history
+			(IStereoRendering::IsStereoEyeView(View))) // Only relevant to stereo views
+		{
+			// For even frames, prevent left eye from occlusion querying
+			// For odd frames, prevent right eye from occlusion querying
+			const bool FrameParity = ((View.ViewState->PrevFrameNumber & 0x01) == 1);
+			bSubmitQueries &= (FrameParity  && IStereoRendering::IsAPrimaryView(View)) ||
+								(!FrameParity && IStereoRendering::IsASecondaryView(View));
+		}
 
-			View.ViewState->PrimitiveOcclusionQueryPool.AdvanceFrame(
-				ViewState->OcclusionFrameCounter,
-				FOcclusionQueryHelpers::GetNumBufferedFrames(Scene->GetFeatureLevel()),
-				View.ViewState->IsRoundRobinEnabled() && !View.bIsSceneCapture && IStereoRendering::IsStereoEyeView(View));
+		View.ViewState->PrimitiveOcclusionQueryPool.AdvanceFrame(
+			ViewState->OcclusionFrameCounter,
+			FOcclusionQueryHelpers::GetNumBufferedFrames(Scene->GetFeatureLevel()),
+			View.ViewState->IsRoundRobinEnabled() && !View.bIsSceneCapture && IStereoRendering::IsStereoEyeView(View));
 
-			NumOccludedPrimitives += FetchVisibilityForPrimitives(Scene, View, bSubmitQueries, bHZBOcclusion, DynamicVertexBuffer);
+		NumOccludedPrimitives += FetchVisibilityForPrimitives(Scene, View, bSubmitQueries, bHZBOcclusion, DynamicVertexBuffer);
 
-			if( bHZBOcclusion )
-			{
-				QUICK_SCOPE_CYCLE_COUNTER(STAT_HZBUnmapResults);
+		if( bHZBOcclusion )
+		{
+			QUICK_SCOPE_CYCLE_COUNTER(STAT_HZBUnmapResults);
 
-				ViewState->HZBOcclusionTests.UnmapResults(RHICmdList);
+			ViewState->HZBOcclusionTests.UnmapResults(RHICmdList);
 
-				if( bSubmitQueries )
-				{
-					ViewState->HZBOcclusionTests.SetValidFrameNumber(ViewState->OcclusionFrameCounter);
-				}
-			}
-		}
-		else
-		{
-			// No occlusion queries, so mark primitives as not occluded
-			for (FSceneSetBitIterator BitIt(View.PrimitiveVisibilityMap); BitIt; ++BitIt)
+			if( bSubmitQueries )
 			{
-				View.PrimitiveDefinitelyUnoccludedMap.AccessCorrespondingBit(BitIt) = true;
+				ViewState->HZBOcclusionTests.SetValidFrameNumber(ViewState->OcclusionFrameCounter);
 			}
 		}
 	}
