WorldPartition: Allow for manually pinning and unpinning world partition actors in the editor.
Pinned actors remain loaded regardless of the state of their cells or data layers.

diff --git a/Engine/Source/Runtime/Engine/Private/WorldPartition/ActorDescContainer.cpp b/Engine/Source/Runtime/Engine/Private/WorldPartition/ActorDescContainer.cpp
--- a/Engine/Source/Runtime/Engine/Private/WorldPartition/ActorDescContainer.cpp
+++ b/Engine/Source/Runtime/Engine/Private/WorldPartition/ActorDescContainer.cpp
@@ -105,6 +105,9 @@ void UActorDescContainer::Initialize(UWorld* InWorld, FName InPackageName)
 void UActorDescContainer::Uninitialize()
 {
 #if WITH_EDITOR
+	PinnedActors.Empty();
+	PinnedActorRefs.Empty();
+
 	if (bContainerInitialized)
 	{
 		UnregisterEditorDelegates();
@@ -234,6 +237,36 @@ void UActorDescContainer::RemoveActor(const FGuid& ActorGuid)
 	}
 }
 
+void UActorDescContainer::PinActor(const FGuid& ActorGuid)
+{
+	if (PinnedActors.Contains(ActorGuid))
+	{
+		return;
+	}
+	
+	if (TUniquePtr<FWorldPartitionActorDesc>* const ActorDescPtr = GetActorDescriptor(ActorGuid))
+	{
+		if (const FWorldPartitionActorDesc* const ActorDesc = ActorDescPtr->Get())
+		{
+			PinnedActors.Emplace(ActorGuid, ActorDescPtr);
+
+			// If the pinned actor has references, we must also create references to those to ensure they are loaded
+			TArray<FWorldPartitionReference>& References = PinnedActorRefs.Emplace(ActorGuid);
+
+			for (const FGuid& ReferencedGuid : ActorDesc->GetReferences())
+			{
+				References.Emplace(GetActorDescriptor(ReferencedGuid));
+			}
+		}
+	}
+}
+
+void UActorDescContainer::UnpinActor(const FGuid& ActorGuid)
+{
+	PinnedActors.Remove(ActorGuid);
+	PinnedActorRefs.Remove(ActorGuid);
+}
+
 void UActorDescContainer::RegisterEditorDelegates()
 {
 	if (GEditor && !IsTemplate() && !World->IsGameWorld())
diff --git a/Engine/Source/Runtime/Engine/Public/WorldPartition/ActorDescContainer.h b/Engine/Source/Runtime/Engine/Public/WorldPartition/ActorDescContainer.h
--- a/Engine/Source/Runtime/Engine/Public/WorldPartition/ActorDescContainer.h
+++ b/Engine/Source/Runtime/Engine/Public/WorldPartition/ActorDescContainer.h
@@ -6,6 +6,7 @@
 #include "UObject/Object.h"
 #include "GameFramework/Actor.h"
 #include "WorldPartition/ActorDescList.h"
+#include "WorldPartition/WorldPartitionHandle.h"
 #include "ActorDescContainer.generated.h"
 
 UCLASS()
@@ -33,6 +34,9 @@ public:
 	/** Removes an actor desc without the need to load a package */
 	virtual void RemoveActor(const FGuid& ActorGuid);
 
+	void PinActor(const FGuid& ActorGuid);
+	void UnpinActor(const FGuid& ActorGuid);
+	bool IsActorPinned(const FGuid& ActorGuid) const { return PinnedActors.Contains(ActorGuid); }
 public:
 	DECLARE_EVENT_OneParam(UWorldPartition, FActorDescAddedEvent, FWorldPartitionActorDesc*);
 	FActorDescAddedEvent OnActorDescAddedEvent;
@@ -60,6 +64,9 @@ protected:
 
 	bool ShouldHandleActorEvent(const AActor* Actor);
 
+	TMap<FGuid, FWorldPartitionReference> PinnedActors;
+	TMap<FGuid, TArray<FWorldPartitionReference>> PinnedActorRefs;
+	
 	bool bContainerInitialized;
 
 	FName ContainerPackageName;
