Better fix for circular references in the WorldPartitionResaveActorsBuilder

diff --git a/Engine/Source/Runtime/Engine/Private/WorldPartition/WorldPartitionHelpers.cpp b/Engine/Source/Runtime/Engine/Private/WorldPartition/WorldPartitionHelpers.cpp
--- a/Engine/Source/Runtime/Engine/Private/WorldPartition/WorldPartitionHelpers.cpp
+++ b/Engine/Source/Runtime/Engine/Private/WorldPartition/WorldPartitionHelpers.cpp
@@ -41,14 +41,14 @@ void FWorldPartitionHelpers::ForEachActorWithLoading(UWorldPartition* WorldParti
 
 		if (const FWorldPartitionActorDesc* ActorDesc = WorldPartition->GetActorDesc(ActorGuid))
 		{
-			FWorldPartitionReference& Reference = InOutActorReferences.Emplace(ActorGuid);
+			InOutActorReferences.Emplace(ActorGuid);
 			
 			for (FGuid ReferenceGuid : ActorDesc->GetReferences())
 			{
 				LoadReferences(ReferenceGuid, InOutActorReferences);
 			}
 
-			Reference = FWorldPartitionReference(WorldPartition, ActorGuid);
+			InOutActorReferences[ActorGuid] = FWorldPartitionReference(WorldPartition, ActorGuid);
 		}
 	};
 
