Cherrypicking Mirage code base (disabled) from Dev-Cooker

diff --git a/Engine/Source/Runtime/Engine/Private/Texture.cpp b/Engine/Source/Runtime/Engine/Private/Texture.cpp
--- a/Engine/Source/Runtime/Engine/Private/Texture.cpp
+++ b/Engine/Source/Runtime/Engine/Private/Texture.cpp
@@ -28,6 +28,7 @@
 #include "TextureDerivedDataTask.h"
 #include "Interfaces/ITextureFormat.h"
 #include "Interfaces/ITextureFormatModule.h"
+#include "UObject/UE5CookerObjectVersion.h"
 
 #if WITH_EDITOR
 #include "TextureCompiler.h"
@@ -441,6 +442,10 @@ void UTexture::PostEditChangeProperty(FPropertyChangedEvent& PropertyChangedEven
 
 void UTexture::Serialize(FArchive& Ar)
 {
+#if UE_USE_VIRTUALBULKDATA
+	Ar.UsingCustomVersion(FUE5CookerObjectVersion::GUID);
+#endif //UE_USE_VIRTUALBULKDATA
+
 	Super::Serialize(Ar);
 
 	FStripDataFlags StripFlags(Ar);
@@ -452,7 +457,20 @@ void UTexture::Serialize(FArchive& Ar)
 #if WITH_EDITOR
 		FWriteScopeLock BulkDataExclusiveScope(Source.BulkDataLock.Get());
 #endif
-		Source.BulkData.Serialize(Ar, this);
+
+#if UE_USE_VIRTUALBULKDATA
+		if (Ar.IsLoading() && Ar.CustomVer(FUE5CookerObjectVersion::GUID) < FUE5CookerObjectVersion::TextureSourceVirtualization)
+		{
+			FByteBulkData TempBulkData;
+			TempBulkData.Serialize(Ar, this);
+
+			Source.BulkData.CreateFromBulkData(TempBulkData, Source.GetId());
+		}
+		else
+#endif //UE_USE_VIRTUALBULKDATA
+		{
+			Source.BulkData.Serialize(Ar, this);
+		}
 	}
 
 	if ( GetLinkerUEVersion() < VER_UE4_TEXTURE_LEGACY_GAMMA )
@@ -633,7 +651,7 @@ void UTexture::GetAssetRegistryTags(TArray<FAssetRegistryTag>& OutTags) const
 	}
 
 	OutTags.Add(FAssetRegistryTag("SourceCompression", Source.GetSourceCompressionAsString(), FAssetRegistryTag::TT_Alphabetical));
-	
+
 	Super::GetAssetRegistryTags(OutTags);
 }
 #endif
@@ -952,8 +970,7 @@ FStreamableRenderResourceState UTexture::GetResourcePostInitState(FTexturePlatfo
 ------------------------------------------------------------------------------*/
 
 FTextureSource::FTextureSource()
-	: LockedMipData(NULL)
-	, NumLockedMips(0u)
+	: NumLockedMips(0u)
 	, LockState(ELockState::None)
 #if WITH_EDITOR
 	, bHasHadBulkDataCleared(false)
@@ -1042,8 +1059,15 @@ void FTextureSource::InitBlocked(const ETextureSourceFormat* InLayerFormats,
 	}
 
 	checkf(LockState == ELockState::None, TEXT("InitBlocked shouldn't be called in-between LockMip/UnlockMip"));
+
+#if UE_USE_VIRTUALBULKDATA
+	FUniqueBuffer Buffer = FUniqueBuffer::Alloc(TotalBytes);
+	uint8* DataPtr = (uint8*)Buffer.GetData();
+#else
 	BulkData.Lock(LOCK_READ_WRITE);
-	uint8* DestData = (uint8*)BulkData.Realloc(TotalBytes);
+	uint8* DataPtr = (uint8*)BulkData.Realloc(TotalBytes);
+#endif //UE_USE_VIRTUALBULKDATA
+
 	if (InDataPerBlock)
 	{
 		for (int i = 0; i < InNumBlocks; ++i)
@@ -1051,12 +1075,17 @@ void FTextureSource::InitBlocked(const ETextureSourceFormat* InLayerFormats,
 			const int64 BlockSize = CalcBlockSize(i);
 			if (InDataPerBlock[i])
 			{
-				FMemory::Memcpy(DestData, InDataPerBlock[i], BlockSize);
+				FMemory::Memcpy(DataPtr, InDataPerBlock[i], BlockSize);
 			}
-			DestData += BlockSize;
+			DataPtr += BlockSize;
 		}
 	}
+
+#if UE_USE_VIRTUALBULKDATA
+	BulkData.UpdatePayload(Buffer.MoveToShared());
+#else
 	BulkData.Unlock();
+#endif //UE_USE_VIRTUALBULKDATA
 }
 
 void FTextureSource::InitLayered(
@@ -1088,13 +1117,28 @@ void FTextureSource::InitLayered(
 	}
 
 	checkf(LockState == ELockState::None, TEXT("InitLayered shouldn't be called in-between LockMip/UnlockMip"));
+
+#if UE_USE_VIRTUALBULKDATA
+	// TODO: Might be able to just use FUniqueBuffer::Clone, if we can be sure that
+	// TotalBytes and NewData always match (the null check onNewData before memcpy is
+	// suspicious)
+	FUniqueBuffer Buffer = FUniqueBuffer::Alloc(TotalBytes);
+	uint8* DestData = (uint8*)Buffer.GetData();
+#else
 	BulkData.Lock(LOCK_READ_WRITE);
 	uint8* DestData = (uint8*)BulkData.Realloc(TotalBytes);
+#endif //UE_USE_VIRTUALBULKDATA
+
 	if (NewData)
 	{
 		FMemory::Memcpy(DestData, NewData, TotalBytes);
 	}
+
+#if UE_USE_VIRTUALBULKDATA
+	BulkData.UpdatePayload(Buffer.MoveToShared());
+#else
 	BulkData.Unlock();
+#endif //UE_USE_VIRTUALBULKDATA
 }
 
 void FTextureSource::Init(
@@ -1160,10 +1204,16 @@ void FTextureSource::InitWithCompressedSourceData(
 	Format = NewFormat;
 	LayerFormat.SetNum(1, true);
 	LayerFormat[0] = NewFormat;
-	
+
 	CompressionFormat = NewSourceFormat;
 
 	checkf(LockState == ELockState::None, TEXT("InitWithCompressedSourceData shouldn't be called in-between LockMip/UnlockMip"));
+	
+#if UE_USE_VIRTUALBULKDATA
+	// Disable the internal bulkdata compression if the source data is already compressed
+	const FName CompressionName = CompressionFormat == TSCF_None ? NAME_Default : NAME_None;
+	BulkData.UpdatePayload(FSharedBuffer::Clone(NewData.GetData(), NewData.Num()), CompressionName);
+#else
 	BulkData.Lock(LOCK_READ_WRITE);
 	uint8* DestData = (uint8*)BulkData.Realloc(NewData.Num());
 	if (NewData.GetData() != nullptr)
@@ -1171,6 +1221,7 @@ void FTextureSource::InitWithCompressedSourceData(
 		FMemory::Memcpy(DestData, NewData.GetData(), NewData.Num());
 	}
 	BulkData.Unlock();
+#endif //UE_USE_VIRTUALBULKDATA
 }
 
 void FTextureSource::Compress()
@@ -1183,27 +1234,72 @@ void FTextureSource::Compress()
 
 	if (CanPNGCompress()) // Note that this will return false if the data is already a compressed PNG
 	{
+#if UE_USE_VIRTUALBULKDATA
+		FSharedBuffer Payload = BulkData.GetPayload().Get();
+		const void* BulkDataPtr = Payload.GetData();
+		const int64 BulkDataSize = Payload.GetSize();
+#else
 		uint8* BulkDataPtr = (uint8*)BulkData.Lock(LOCK_READ_WRITE);
+		const int64 BulkDataSize = BulkData.GetBulkDataSize();
+#endif //UE_USE_VIRTUALBULKDATA
+
 		IImageWrapperModule& ImageWrapperModule = FModuleManager::LoadModuleChecked<IImageWrapperModule>( FName("ImageWrapper") );
 		TSharedPtr<IImageWrapper> ImageWrapper = ImageWrapperModule.CreateImageWrapper( EImageFormat::PNG );
 		// TODO: TSF_BGRA8 is stored as RGBA, so the R and B channels are swapped in the internal png. Should we fix this?
 		ERGBFormat RawFormat = (Format == TSF_G8 || Format == TSF_G16) ? ERGBFormat::Gray : ERGBFormat::RGBA;
-		if ( ImageWrapper.IsValid() && ImageWrapper->SetRaw( BulkDataPtr, BulkData.GetBulkDataSize(), SizeX, SizeY, RawFormat, (Format == TSF_G16 || Format == TSF_RGBA16) ? 16 : 8 ) )
+		if ( ImageWrapper.IsValid() && ImageWrapper->SetRaw( BulkDataPtr, BulkDataSize, SizeX, SizeY, RawFormat, (Format == TSF_G16 || Format == TSF_RGBA16) ? 16 : 8 ) )
 		{
 			const TArray64<uint8> CompressedData = ImageWrapper->GetCompressed();
 			if ( CompressedData.Num() > 0 )
 			{
+#if UE_USE_VIRTUALBULKDATA	
+				BulkData.UpdatePayload(FSharedBuffer::Clone(CompressedData.GetData(), CompressedData.Num()), NAME_None);
+#else
 				BulkDataPtr = (uint8*)BulkData.Realloc(CompressedData.Num());
 				FMemory::Memcpy(BulkDataPtr, CompressedData.GetData(), CompressedData.Num());
 				BulkData.Unlock();
+#endif //UE_USE_VIRTUALBULKDATA
 				bPNGCompressed = true;
 				CompressionFormat = TSCF_PNG;
 			}
 		}
 	}
-	
-	// Disable the internal bulkdata compression if the source data is already compressed
+
+	// Fix up for packages that were saved before CompressionFormat was introduced. Can removed this when we deprecate bPNGCompressed!
+	if (bPNGCompressed)
+	{
+		CompressionFormat = TSCF_PNG;
+	}
+
+	// TODO: Cannot fix incorrectly assigned compression types if UE_USE_VIRTUALBULKDATA is enabled. Need a fix for that.
+#if !UE_USE_VIRTUALBULKDATA
+	// Disable the internal bulkdata compression if the source data is already compressed	
 	BulkData.StoreCompressedOnDisk(CompressionFormat == TSCF_None ? NAME_Zlib : NAME_None);
+#endif
+}
+
+FTextureSource::FMipAllocation FTextureSource::Decompress(IImageWrapperModule* ImageWrapperModule) UE_VBD_CONST
+{
+	TRACE_CPUPROFILER_EVENT_SCOPE(FTextureSource::Decompress);
+
+	if (CompressionFormat == TSCF_JPEG)
+	{
+		return TryDecompressJpegData(ImageWrapperModule);
+	}
+	else if (bPNGCompressed) // Change to CompressionFormat == TSCF_PNG once bPNGCompressed is deprecated
+	{
+		return TryDecompressPngData(ImageWrapperModule);
+	}
+	else
+	{
+#if UE_USE_VIRTUALBULKDATA
+		FMipAllocation Payload(BulkData.GetPayload().Get());
+#else
+		FMipAllocation Payload(BulkData);
+#endif //UE_USE_VIRTUALBULKDATA
+
+		return Payload;
+	}
 }
 
 const uint8* FTextureSource::LockMipReadOnly(int32 BlockIndex, int32 LayerIndex, int32 MipIndex)
@@ -1216,65 +1312,24 @@ uint8* FTextureSource::LockMip(int32 BlockIndex, int32 LayerIndex, int32 MipInde
 	return LockMipInternal(BlockIndex, LayerIndex, MipIndex, ELockState::ReadWrite);
 }
 
-uint8* FTextureSource::LockMipInternal(int32 BlockIndex, int32 LayerIndex, int32 MipIndex, ELockState RequestedLockState )
+uint8* FTextureSource::LockMipInternal(int32 BlockIndex, int32 LayerIndex, int32 MipIndex, ELockState RequestedLockState)
 {
+	TRACE_CPUPROFILER_EVENT_SCOPE(FTextureSource::LockMip);
+
 	checkf(RequestedLockState != ELockState::None, TEXT("Cannot call FTextureSource::LockMipInternal with a RequestedLockState of type ELockState::None"));
 
-	uint8* MipData = NULL;
+	uint8* MipData = nullptr;
+
 	if (BlockIndex < GetNumBlocks() && LayerIndex < NumLayers && MipIndex < NumMips)
 	{
-		if (LockedMipData == NULL)
+		if (LockedMipData.IsNull())
 		{
-			LockedMipData = (uint8*)BulkData.Lock(LOCK_READ_WRITE);
-
-			if (CompressionFormat == TSCF_JPEG)
-			{
-				TArray64<uint8> RawData;
-				// Note that the method itself will give warnings on all failure conditions so we 
-				// don't need to handle that here
-				if (TryDecompressJpegData(nullptr, LockedMipData, MipIndex, RawData))
-				{
-					LockedMipData = (uint8*)FMemory::Malloc(RawData.Num());
-					FMemory::Memcpy(LockedMipData, RawData.GetData(), RawData.Num());
-				}
-			}	
-			else if (bPNGCompressed)
-			{
-				bool bCanPngCompressFormat = (Format == TSF_G8 || Format == TSF_G16 || Format == TSF_RGBA8 || Format == TSF_BGRA8 || Format == TSF_RGBA16);
-				check(Blocks.Num() == 0 && NumLayers == 1 && NumSlices == 1 && bCanPngCompressFormat);
-				if (MipIndex != 0)
-				{
-					return NULL;
-				}
-
-				IImageWrapperModule& ImageWrapperModule = FModuleManager::LoadModuleChecked<IImageWrapperModule>( FName("ImageWrapper") );
-				TSharedPtr<IImageWrapper> ImageWrapper = ImageWrapperModule.CreateImageWrapper( EImageFormat::PNG );
-				if ( ImageWrapper.IsValid() && ImageWrapper->SetCompressed( LockedMipData, BulkData.GetBulkDataSize() ) )
-				{
-					check( ImageWrapper->GetWidth() == SizeX );
-					check( ImageWrapper->GetHeight() == SizeY );
-					TArray64<uint8> RawData;
-					// TODO: TSF_BGRA8 is stored as RGBA, so the R and B channels are swapped in the internal png. Should we fix this?
-					ERGBFormat RawFormat = (Format == TSF_G8 || Format == TSF_G16) ? ERGBFormat::Gray : ERGBFormat::RGBA;
-					if (ImageWrapper->GetRaw( RawFormat, (Format == TSF_G16 || Format == TSF_RGBA16) ? 16 : 8, RawData ) && RawData.Num() > 0)
-					{
-						LockedMipData = (uint8*)FMemory::Malloc(RawData.Num());
-						FMemory::Memcpy(LockedMipData, RawData.GetData(), RawData.Num());
-					}
-					else
-					{
-						UE_LOG(LogTexture, Warning, TEXT("PNG decompression of source art failed"));
-					}
-				}
-				else
-				{
-					UE_LOG(LogTexture, Log, TEXT("Only pngs are supported"));
-				}
-			}
+			checkf(NumLockedMips == 0, TEXT("Texture mips are locked but the LockedMipData is missing"));
+			LockedMipData = Decompress(nullptr);
 		}
 
-		MipData = LockedMipData + CalcMipOffset(BlockIndex, LayerIndex, MipIndex);
-
+		MipData = LockedMipData.GetDataReadWrite() + CalcMipOffset(BlockIndex, LayerIndex, MipIndex);
+		
 		if (NumLockedMips == 0)
 		{
 			LockState = RequestedLockState;
@@ -1284,23 +1339,39 @@ uint8* FTextureSource::LockMipInternal(int32 BlockIndex, int32 LayerIndex, int32
 			checkf(LockState == RequestedLockState, TEXT("Cannot change the lock type until UnlockMip is called"));
 		}
 
-		++NumLockedMips;		
+		++NumLockedMips;
 	}
+
 	return MipData;
 }
 
 void FTextureSource::UnlockMip(int32 BlockIndex, int32 LayerIndex, int32 MipIndex)
 {
+	TRACE_CPUPROFILER_EVENT_SCOPE(FTextureSource::UnlockMip);
+
 	check(BlockIndex < GetNumBlocks());
 	check(LayerIndex < NumLayers);
 	check(MipIndex < MAX_TEXTURE_MIP_COUNT);
-
 	check(NumLockedMips > 0u);
 	check(LockState != ELockState::None);
 
 	--NumLockedMips;
 	if (NumLockedMips == 0u)
 	{
+#if UE_USE_VIRTUALBULKDATA
+		// If the lock was for Read/Write then we need to assume that the decompressed copy
+		// we returned (LockedMipData) was updated and should update the payload accordingly.
+		// This will wipe the compression format that we used to have.
+		if (LockState == ELockState::ReadWrite)
+		{
+			UE_CLOG(CompressionFormat == TSCF_JPEG, LogTexture, Warning, TEXT("Call to FTextureSource::UnlockMip will cause texture source to lose it's jpeg storage format"));
+
+			BulkData.UpdatePayload(LockedMipData.Release());
+
+			bPNGCompressed = false;
+			CompressionFormat = TSCF_None;
+		}
+#else
 		// If the source data is stored in a compressed format then we returned a copy when the mip was
 		// locked. We now need to deal with that copy.
 		if (CompressionFormat == TSCF_JPEG || bPNGCompressed)
@@ -1318,21 +1389,20 @@ void FTextureSource::UnlockMip(int32 BlockIndex, int32 LayerIndex, int32 MipInde
 				check(LayerIndex == 0);
 				check(MipIndex == 0);
 				int32 MipSize = CalcMipSize(0, 0, 0);
+
 				uint8* UncompressedData = (uint8*)BulkData.Realloc(MipSize);
-				FMemory::Memcpy(UncompressedData, LockedMipData, MipSize);
-				FMemory::Free(LockedMipData);
+				FMemory::Memcpy(UncompressedData, LockedMipData.GetDataReadOnly().GetData(), MipSize);
+
 				bPNGCompressed = false;
 				CompressionFormat = TSCF_None;
 			}
-			else
-			{
-				FMemory::Free(LockedMipData);
-			}
 		}
 
+		BulkData.Unlock();	
+#endif //UE_USE_VIRTUALBULKDATA
+
 		LockState = ELockState::None;
-		LockedMipData = NULL;
-		BulkData.Unlock();
+		LockedMipData.Reset();
 
 		// TODO: Only call this if the LockState was ELockState::ReadWrite
 		ForceGenerateGuid();
@@ -1341,86 +1411,74 @@ void FTextureSource::UnlockMip(int32 BlockIndex, int32 LayerIndex, int32 MipInde
 
 bool FTextureSource::GetMipData(TArray64<uint8>& OutMipData, int32 BlockIndex, int32 LayerIndex, int32 MipIndex, IImageWrapperModule* ImageWrapperModule)
 {
-	checkf(LockState == ELockState::None, TEXT("GetMipData shouldn't be called in-between LockMip/UnlockMip"));
+	TRACE_CPUPROFILER_EVENT_SCOPE(FTextureSource::GetMipData (TArray64));
+
+	checkf(LockState == ELockState::None, TEXT("GetMipData (TArray64) shouldn't be called in-between LockMip/UnlockMip"));
 
 	bool bSuccess = false;
-	if (BlockIndex < GetNumBlocks() && LayerIndex < NumLayers && MipIndex < NumMips && BulkData.GetBulkDataSize() > 0)
+
+	if (BlockIndex < GetNumBlocks() && LayerIndex < NumLayers && MipIndex < NumMips && GetSizeOnDisk() > 0)
 	{
 #if WITH_EDITOR
 		FWriteScopeLock BulkDataExclusiveScope(BulkDataLock.Get());
 #endif
-		void* RawSourceData = BulkData.Lock(LOCK_READ_ONLY);
-		if (CompressionFormat == TSCF_JPEG)
-		{
-			bSuccess = TryDecompressJpegData(ImageWrapperModule, RawSourceData, MipIndex, OutMipData);
-		}
-		else if (bPNGCompressed)
-		{
-			bool bCanPngCompressFormat = (Format == TSF_G8 || Format == TSF_G16 || Format == TSF_RGBA8 || Format == TSF_BGRA8 || Format == TSF_RGBA16);
-			if (MipIndex == 0 && NumLayers == 1 && NumSlices == 1 && Blocks.Num() == 0 && bCanPngCompressFormat)
-			{
-				if (!ImageWrapperModule) // Optional if called from the gamethread, see FModuleManager::WarnIfItWasntSafeToLoadHere()
-				{
-					ImageWrapperModule = &FModuleManager::LoadModuleChecked<IImageWrapperModule>( FName("ImageWrapper") );
-				}
 
-				TSharedPtr<IImageWrapper> ImageWrapper = ImageWrapperModule->CreateImageWrapper( EImageFormat::PNG );
+		checkf(NumLockedMips == 0, TEXT("Attempting to access a locked FTextureSource"));
+		// LockedMipData should only be allocated if NumLockedMips > 0 so the following assert should have been caught
+		// by the one above. If it fires then it indicates that there is a lock/unlock mismatch as well as invalid access!
+		checkf(LockedMipData.IsNull(), TEXT("Attempting to access mip data while locked mip data is still allocated")); 
 
-				if ( ImageWrapper.IsValid() && ImageWrapper->SetCompressed( RawSourceData, BulkData.GetBulkDataSize() ) )
-				{
-					if (ImageWrapper->GetWidth() == SizeX
-						&& ImageWrapper->GetHeight() == SizeY)
-					{
-						// TODO: TSF_BGRA8 is stored as RGBA, so the R and B channels are swapped in the internal png. Should we fix this?
-						ERGBFormat RawFormat = (Format == TSF_G8 || Format == TSF_G16) ? ERGBFormat::Gray : ERGBFormat::RGBA;
-						if (ImageWrapper->GetRaw( RawFormat, (Format == TSF_G16 || Format == TSF_RGBA16) ? 16 : 8, OutMipData ))
-						{
-							bSuccess = true;
-						}
-						else
-						{
-							UE_LOG(LogTexture, Warning, TEXT("PNG decompression of source art failed"));
-							OutMipData.Empty();
-						}
-					}
-					else
-					{
-						UE_LOG(LogTexture, Warning,
-							TEXT("PNG decompression of source art failed. ")
-							TEXT("Source image should be %dx%d but is %dx%d"),
-							SizeX, SizeY,
-							ImageWrapper->GetWidth(), ImageWrapper->GetHeight()
-							);
-					}
-				}
-				else
-				{
-					UE_LOG(LogTexture, Log, TEXT("Only pngs are supported"));
-				}
-			}
-		}
-		else
+		FMipAllocation DecompressedData = Decompress(ImageWrapperModule);
+
+		if (!DecompressedData.IsNull())
 		{
-			int64 MipOffset = CalcMipOffset(BlockIndex, LayerIndex, MipIndex);
-			int64 MipSize = CalcMipSize(BlockIndex, LayerIndex, MipIndex);
-			if (BulkData.GetBulkDataSize() >= MipOffset + MipSize)
+			const int64 MipOffset = CalcMipOffset(BlockIndex, LayerIndex, MipIndex);
+			const int64 MipSize = CalcMipSize(BlockIndex, LayerIndex, MipIndex);
+
+			if ((int64)DecompressedData.GetSize() >= MipOffset + MipSize)
 			{
 				OutMipData.Empty(MipSize);
 				OutMipData.AddUninitialized(MipSize);
 				FMemory::Memcpy(
 					OutMipData.GetData(),
-					(uint8*)RawSourceData + MipOffset,
+					(const uint8*)DecompressedData.GetDataReadOnly().GetData() + MipOffset,
 					MipSize
-					);
+				);
+
+				bSuccess = true;
 			}
-			bSuccess = true;
 		}
+
+#if !UE_USE_VIRTUALBULKDATA
 		BulkData.Unlock();
+#endif //!UE_USE_VIRTUALBULKDATA		
 	}
-
+	
 	return bSuccess;
 }
 
+FTextureSource::FMipData FTextureSource::GetMipData(IImageWrapperModule* ImageWrapperModule)
+{
+	TRACE_CPUPROFILER_EVENT_SCOPE(FTextureSource::GetMipData (FMipData));
+
+	checkf(LockState == ELockState::None, TEXT("GetMipData (FMipData) shouldn't be called in-between LockMip/UnlockMip"));
+
+	check(LockedMipData.IsNull());
+	check(NumLockedMips == 0);
+
+#if WITH_EDITOR
+	// When UE_USE_VIRTUALBULKDATA is removed, replace this with a FReadScopeLock,
+	// we only need the expanded scope for old style bulkdata
+	BulkDataLock.Get().ReadLock();
+
+	FMipAllocation DecompressedData = Decompress(ImageWrapperModule);
+	return FMipData(*this, MoveTemp(DecompressedData), &BulkDataLock.Get());
+#else
+	FMipAllocation DecompressedData = Decompress(ImageWrapperModule);
+	return FMipData(*this, MoveTemp(DecompressedData));
+#endif
+}
+
 int64 FTextureSource::CalcMipSize(int32 BlockIndex, int32 LayerIndex, int32 MipIndex) const
 {
 	FTextureSourceBlock Block;
@@ -1448,7 +1506,7 @@ bool FTextureSource::IsPowerOfTwo(int32 BlockIndex) const
 bool FTextureSource::IsValid() const
 {
 	return SizeX > 0 && SizeY > 0 && NumSlices > 0 && NumLayers > 0 && NumMips > 0 &&
-		Format != TSF_Invalid && BulkData.GetBulkDataSize() > 0;
+		Format != TSF_Invalid && GetSizeOnDisk() > 0;
 }
 
 void FTextureSource::GetBlock(int32 Index, FTextureSourceBlock& OutBlock) const
@@ -1525,42 +1583,91 @@ FString FTextureSource::GetSourceCompressionAsString() const
 	return StaticEnum<ETextureSourceCompressionFormat>()->GetDisplayNameTextByValue(CompressionFormat).ToString();
 }
 
-bool FTextureSource::TryDecompressJpegData(IImageWrapperModule* ImageWrapperModule, const void* RawSourceData, int32 MipIndex, TArray64<uint8>& OutRawData)
+FTextureSource::FMipAllocation FTextureSource::TryDecompressPngData(IImageWrapperModule* ImageWrapperModule) UE_VBD_CONST
 {
-	if (MipIndex == 0 && NumLayers == 1 && NumSlices == 1 && Blocks.Num() == 0)
+	bool bCanPngCompressFormat = (Format == TSF_G8 || Format == TSF_G16 || Format == TSF_RGBA8 || Format == TSF_BGRA8 || Format == TSF_RGBA16);
+	check(Blocks.Num() == 0 && NumLayers == 1 && NumSlices == 1 && bCanPngCompressFormat);
+
+#if UE_USE_VIRTUALBULKDATA
+	FMipAllocation Payload(BulkData.GetPayload().Get());
+#else
+	FMipAllocation Payload(BulkData);
+#endif //UE_USE_VIRTUALBULKDATA
+
+	if (ImageWrapperModule == nullptr) // Optional if called from the gamethread, see FModuleManager::WarnIfItWasntSafeToLoadHere()
+	{
+		ImageWrapperModule = &FModuleManager::LoadModuleChecked<IImageWrapperModule>(FName("ImageWrapper"));
+	}
+
+	TSharedPtr<IImageWrapper> ImageWrapper = ImageWrapperModule->CreateImageWrapper(EImageFormat::PNG);
+	if (ImageWrapper.IsValid() && ImageWrapper->SetCompressed(Payload.GetDataReadOnly().GetData(), Payload.GetSize()))
+	{
+		check(ImageWrapper->GetWidth() == SizeX);
+		check(ImageWrapper->GetHeight() == SizeY);
+
+		TArray64<uint8> RawData;
+		// TODO: TSF_BGRA8 is stored as RGBA, so the R and B channels are swapped in the internal png. Should we fix this?
+		ERGBFormat RawFormat = (Format == TSF_G8 || Format == TSF_G16) ? ERGBFormat::Gray : ERGBFormat::RGBA;
+		if (ImageWrapper->GetRaw(RawFormat, (Format == TSF_G16 || Format == TSF_RGBA16) ? 16 : 8, RawData) && RawData.Num() > 0)
+		{
+			return FMipAllocation(RawData.GetData(), RawData.Num());
+		}
+		else
+		{
+			UE_LOG(LogTexture, Warning, TEXT("PNG decompression of source art failed"));
+			return FMipAllocation();
+		}
+	}
+	else
+	{
+		UE_LOG(LogTexture, Log, TEXT("Only pngs are supported"));
+		return FMipAllocation();
+	}
+}
+
+FTextureSource::FMipAllocation FTextureSource::TryDecompressJpegData(IImageWrapperModule* ImageWrapperModule) UE_VBD_CONST
+{
+	if (NumLayers == 1 && NumSlices == 1 && Blocks.Num() == 0)
 	{
 		if (!ImageWrapperModule)
 		{
 			ImageWrapperModule = &FModuleManager::LoadModuleChecked<IImageWrapperModule>(FName("ImageWrapper"));
 		}
 
-		TSharedPtr<IImageWrapper> JpegImageWrapper = ImageWrapperModule->CreateImageWrapper(EImageFormat::JPEG);
-		if (JpegImageWrapper.IsValid() && JpegImageWrapper->SetCompressed(RawSourceData, BulkData.GetBulkDataSize()))
+#if UE_USE_VIRTUALBULKDATA
+		FMipAllocation Payload(BulkData.GetPayload().Get());
+#else
+		FMipAllocation Payload(BulkData);
+#endif //UE_USE_VIRTUALBULKDATA
+
+		TSharedPtr<IImageWrapper> ImageWrapper = ImageWrapperModule->CreateImageWrapper(EImageFormat::JPEG);
+		if (ImageWrapper.IsValid() && ImageWrapper->SetCompressed(Payload.GetDataReadOnly().GetData(), Payload.GetSize()))
 		{
+			TArray64<uint8> RawData;
 			// The two formats we support for JPEG imports, see UTextureFactory::ImportImage
 			const ERGBFormat JpegFormat = Format == TSF_G8 ? ERGBFormat::Gray : ERGBFormat::BGRA;
-			if (JpegImageWrapper->GetRaw(JpegFormat, 8, OutRawData))
+			if (ImageWrapper->GetRaw(JpegFormat, 8, RawData))
 			{
-				return true;
+				return FMipAllocation(RawData.GetData(), RawData.Num());
 			}
 			else
 			{
 				UE_LOG(LogTexture, Warning, TEXT("JPEG decompression of source art failed to return uncompressed data"));
-				return false;
+				return FMipAllocation();
 			}
 		}
 		else
 		{
 			UE_LOG(LogTexture, Warning, TEXT("JPEG decompression of source art failed initialization"));
-			return false;
+			return FMipAllocation();
 		}
 	}
 	else
 	{
-		UE_LOG(LogTexture, Warning, TEXT("JPEG compressed source art is in an invalid format MipIndex:(%d) NumLayers:(%d) NumSlices:(%d) NumBlocks:(%d)"),
-			MipIndex, NumLayers, NumSlices, Blocks.Num());
-		return false;
-	}
+		UE_LOG(LogTexture, Warning, TEXT("JPEG compressed source art is in an invalid format NumLayers:(%d) NumSlices:(%d) NumBlocks:(%d)"),
+			NumLayers, NumSlices, Blocks.Num());
+		return FMipAllocation();
+	}	
 }
 
 bool FTextureSource::CanPNGCompress() const
@@ -1574,7 +1681,7 @@ bool FTextureSource::CanPNGCompress() const
 		Blocks.Num() == 0 &&
 		SizeX > 4 &&
 		SizeY > 4 &&
-		BulkData.GetBulkDataSize() > 0 &&
+		GetSizeOnDisk() > 0 &&
 		bCanPngCompressFormat &&
 		CompressionFormat == TSCF_None)
 	{
@@ -1592,11 +1699,16 @@ void FTextureSource::ForceGenerateGuid()
 void FTextureSource::ReleaseSourceMemory()
 {
 	bHasHadBulkDataCleared = true;
+
+#if UE_USE_VIRTUALBULKDATA
+	BulkData.UnloadData();
+#else
 	if (BulkData.IsLocked())
 	{
 		BulkData.Unlock();
 	}
 	BulkData.RemoveBulkData();
+#endif //UE_USE_VIRTUALBULKDATA
 }
 
 void FTextureSource::RemoveSourceData()
@@ -1610,14 +1722,21 @@ void FTextureSource::RemoveSourceData()
 	LayerFormat.Empty();
 	Blocks.Empty();
 	bPNGCompressed = false;
-	LockedMipData = NULL;
+	CompressionFormat = TSCF_None;
+	LockedMipData.Reset();
 	NumLockedMips = 0u;
 	LockState = ELockState::None;
+
+#if UE_USE_VIRTUALBULKDATA
+	BulkData.UnloadData();
+#else
 	if (BulkData.IsLocked())
 	{
 		BulkData.Unlock();
 	}
 	BulkData.RemoveBulkData();
+#endif //UE_USE_VIRTUALBULKDATA
+
 	ForceGenerateGuid();
 }
 
@@ -1686,15 +1805,25 @@ int64 FTextureSource::CalcMipOffset(int32 BlockIndex, int32 LayerIndex, int32 Mi
 
 void FTextureSource::UseHashAsGuid()
 {
-	uint32 Hash[5];
-
-	if (BulkData.GetBulkDataSize() > 0)
+	if (GetSizeOnDisk() > 0)
 	{
 		bGuidIsHash = true;
-		void* Buffer = BulkData.Lock(LOCK_READ_ONLY);
+
+#if UE_USE_VIRTUALBULKDATA
+		Id = BulkData.GetPayloadId().ToGuid();
+#else
+		uint32 Hash[5] = {};
+
+		const void* Buffer = BulkData.LockReadOnly();
 		FSHA1::HashBuffer(Buffer, BulkData.GetBulkDataSize(), (uint8*)Hash);
 		BulkData.Unlock();
+
 		Id = FGuid(Hash[0] ^ Hash[4], Hash[1], Hash[2], Hash[3]);
+#endif //UE_USE_VIRTUALBULKDATA
+	}
+	else
+	{
+		Id.Invalidate();
 	}
 }
 
@@ -2100,3 +2229,135 @@ void UTexture::NotifyMaterials(const ENotifyMaterialsEffectOnShaders EffectOnSha
 }
 
 #endif //WITH_EDITOR
+
+#if WITH_EDITOR
+
+FTextureSource::FMipData::FMipData(const FTextureSource& InSource, FMipAllocation&& InMipData, FRWLock* InReadLock)
+	: TextureSource(InSource)
+	, MipData(MoveTemp(InMipData))
+	, ReadLock(InReadLock)
+{
+}
+
+FTextureSource::FMipData::~FMipData()
+{
+#if !UE_USE_VIRTUALBULKDATA
+	TextureSource.BulkData.Unlock();
+#endif //!UE_USE_VIRTUALBULKDATA
+
+	if (ReadLock != nullptr)
+	{
+		ReadLock->ReadUnlock();
+	}
+}
+
+bool FTextureSource::FMipData::GetMipData(TArray64<uint8>& OutMipData, int32 BlockIndex, int32 LayerIndex, int32 MipIndex) const
+{
+	if (BlockIndex < TextureSource.GetNumBlocks() && LayerIndex < TextureSource.GetNumLayers() && MipIndex < TextureSource.GetNumMips() && !MipData.IsNull())
+	{
+		const int64 MipOffset = TextureSource.CalcMipOffset(BlockIndex, LayerIndex, MipIndex);
+		const int64 MipSize = TextureSource.CalcMipSize(BlockIndex, LayerIndex, MipIndex);
+
+		if ((int64)MipData.GetSize() >= MipOffset + MipSize)
+		{
+			OutMipData.Empty(MipSize);
+			OutMipData.AddUninitialized(MipSize);
+			FMemory::Memcpy(
+				OutMipData.GetData(),
+				(const uint8*)MipData.GetDataReadOnly().GetData() + MipOffset,
+				MipSize
+			);
+		}
+
+		return true;
+	}
+	else
+	{
+		return false;
+	}
+}
+
+#endif //WITH_EDITOR
+
+#if WITH_EDITOR
+
+FTextureSource::FMipAllocation::FMipAllocation(FSharedBuffer SrcData)
+	: ReadOnlyReference(SrcData)
+{
+}
+
+FTextureSource::FMipAllocation::FMipAllocation(const void* SrcData, int64 DataLength)
+{
+	CreateReadWriteBuffer(SrcData, DataLength);
+}
+
+FTextureSource::FMipAllocation::FMipAllocation(FByteBulkData& BulkData)
+{
+	// Note that we DO NOT unlock the bulkdata via FMipAllocation as some areas of code keep the lock open for longer
+	// than the memory buffer is actually used.
+	BulkDataPtr = (uint8*)BulkData.Lock(LOCK_READ_WRITE);
+	ReadOnlyReference = FSharedBuffer::MakeView(BulkDataPtr, BulkData.GetBulkDataSize());
+}
+
+FTextureSource::FMipAllocation::FMipAllocation(FTextureSource::FMipAllocation&& Other)
+{
+	*this = MoveTemp(Other);
+}
+
+FTextureSource::FMipAllocation& FTextureSource::FMipAllocation::operator =(FTextureSource::FMipAllocation&& Other)
+{
+	ReadOnlyReference = MoveTemp(Other.ReadOnlyReference);
+	ReadWriteBuffer = MoveTemp(Other.ReadWriteBuffer);
+
+	BulkDataPtr = Other.BulkDataPtr;
+	Other.BulkDataPtr = nullptr;
+
+	return *this;
+}
+
+void FTextureSource::FMipAllocation::Reset()
+{
+	ReadOnlyReference.Reset();
+	ReadWriteBuffer = nullptr;
+	BulkDataPtr = nullptr;
+}
+
+uint8* FTextureSource::FMipAllocation::GetDataReadWrite()
+{
+	if (BulkDataPtr != nullptr)
+	{
+		return BulkDataPtr;
+	}
+
+	if (ReadWriteBuffer == nullptr)
+	{
+		CreateReadWriteBuffer(ReadOnlyReference.GetData(), ReadOnlyReference.GetSize());
+	}
+
+	return ReadWriteBuffer.Get();
+}
+
+FSharedBuffer FTextureSource::FMipAllocation::Release()
+{
+	check(BulkDataPtr == nullptr); // Should not be called if we are using old bulkdata
+	if (ReadWriteBuffer != nullptr)
+	{
+		const int64 DataSize = ReadOnlyReference.GetSize();
+		ReadOnlyReference.Reset();
+		return FSharedBuffer::TakeOwnership(ReadWriteBuffer.Release(), DataSize, FMemory::Free);
+	}
+	else
+	{
+		return MoveTemp(ReadOnlyReference);
+	}
+}
+
+void FTextureSource::FMipAllocation::CreateReadWriteBuffer(const void* SrcData, int64 DataLength)
+{
+	ReadWriteBuffer = MakeUnique<uint8[]>(DataLength);
+	FMemory::Memcpy(ReadWriteBuffer.Get(), SrcData, DataLength);
+
+	ReadOnlyReference = FSharedBuffer::MakeView(ReadWriteBuffer.Get(), DataLength);
+}
+
+#endif //WITH_EDITOR
diff --git a/Engine/Source/Runtime/Engine/Private/TextureDerivedData.cpp b/Engine/Source/Runtime/Engine/Private/TextureDerivedData.cpp
--- a/Engine/Source/Runtime/Engine/Private/TextureDerivedData.cpp
+++ b/Engine/Source/Runtime/Engine/Private/TextureDerivedData.cpp
@@ -1834,7 +1834,7 @@ void UTexture::BeginCacheForCookedPlatformData( const ITargetPlatform *TargetPla
 
 		// If source data is resident in memory then allow the texture to be built
 		// in a background thread.
-		bool bAllowAsyncBuild = Source.BulkData.IsBulkDataLoaded();
+		bool bAllowAsyncBuild = Source.IsBulkDataLoaded();
 
 		if (bAllowAsyncBuild)
 		{
diff --git a/Engine/Source/Runtime/Engine/Private/TextureDerivedDataTask.cpp b/Engine/Source/Runtime/Engine/Private/TextureDerivedDataTask.cpp
--- a/Engine/Source/Runtime/Engine/Private/TextureDerivedDataTask.cpp
+++ b/Engine/Source/Runtime/Engine/Private/TextureDerivedDataTask.cpp
@@ -195,6 +195,8 @@ void FTextureSourceData::GetSourceMips(FTextureSource& Source, IImageWrapperModu
 				return;
 		}
 
+		const FTextureSource::FMipData ScopedMipData = Source.GetMipData(InImageWrapper);
+
 		for (int32 BlockIndex = 0; BlockIndex < Blocks.Num(); ++BlockIndex)
 		{
 			FTextureSourceBlock SourceBlock;
@@ -217,7 +219,7 @@ void FTextureSourceData::GetSourceMips(FTextureSource& Source, IImageWrapperModu
 							LayerData.GammaSpace
 						);
 
-						if (!Source.GetMipData(SourceMip->RawData, BlockIndex, LayerIndex, MipIndex, InImageWrapper))
+						if (!ScopedMipData.GetMipData(SourceMip->RawData, BlockIndex, LayerIndex, MipIndex))
 						{
 							UE_LOG(LogTexture, Warning, TEXT("Cannot retrieve source data for mip %d of texture %s"), MipIndex, *TextureName.ToString());
 							ReleaseMemory();
diff --git a/Engine/Source/Runtime/Launch/Private/LaunchEngineLoop.cpp b/Engine/Source/Runtime/Launch/Private/LaunchEngineLoop.cpp
--- a/Engine/Source/Runtime/Launch/Private/LaunchEngineLoop.cpp
+++ b/Engine/Source/Runtime/Launch/Private/LaunchEngineLoop.cpp
@@ -3875,6 +3875,10 @@ void FEngineLoop::LoadPreInitModules()
 	FModuleManager::Get().LoadModule(TEXT("TextureCompressor"));
 #endif
 
+	if (!FPlatformProperties::RequiresCookedData())
+	{
+		FModuleManager::Get().LoadModule(TEXT("Virtualization"));
+	}
 #endif // WITH_ENGINE
 
 #if (WITH_EDITOR && !(UE_BUILD_SHIPPING || UE_BUILD_TEST))
diff --git a/Engine/Source/Runtime/MeshDescription/Public/MeshDescription.h b/Engine/Source/Runtime/MeshDescription/Public/MeshDescription.h
--- a/Engine/Source/Runtime/MeshDescription/Public/MeshDescription.h
+++ b/Engine/Source/Runtime/MeshDescription/Public/MeshDescription.h
@@ -1290,11 +1290,9 @@ public:
 		: bBulkDataUpdated(false)
 		, bGuidIsHash(false)
 	{
-#if UE_USE_VIRTUALBULKDATA
-		BulkData.EnableDataCompression(true);
-#else
+#if !UE_USE_VIRTUALBULKDATA
 		BulkData.SetBulkDataFlags(BULKDATA_SerializeCompressed | BULKDATA_SerializeCompressedBitWindow);
-#endif //UE_USE_VIRTUALBULKDATA
+#endif //!UE_USE_VIRTUALBULKDATA
 	}
 
 	/** Serialization */
@@ -1327,7 +1325,7 @@ private:
 	FRWLock       BulkDataLock;
 #endif
 	/** Internally store bulk data as bytes */
-	TChooseClass<UE_USE_VIRTUALBULKDATA, FByteVirtualizedBulkData, FByteBulkData>::Result BulkData;
+	TChooseClass<UE_USE_VIRTUALBULKDATA, UE::Virtualization::FByteVirtualizedBulkData, FByteBulkData>::Result BulkData;
 
 	/** GUID associated with the data stored herein. */
 	FGuid Guid;
