Data Layers changes
- Keep a reference to the data layers actor in the level to ensure initial loading.
- Always create a data layers actor when creating a new partitioned world.

diff --git a/Engine/Source/Editor/DataLayerEditor/Private/DataLayer/DataLayerEditorSubsystem.cpp b/Engine/Source/Editor/DataLayerEditor/Private/DataLayer/DataLayerEditorSubsystem.cpp
--- a/Engine/Source/Editor/DataLayerEditor/Private/DataLayer/DataLayerEditorSubsystem.cpp
+++ b/Engine/Source/Editor/DataLayerEditor/Private/DataLayer/DataLayerEditorSubsystem.cpp
@@ -1093,12 +1093,19 @@ bool UDataLayerEditorSubsystem::TryGetDataLayerFromLabel(const FName& DataLayerL
 
 const AWorldDataLayers* UDataLayerEditorSubsystem::GetWorldDataLayers() const
 {
-	return AWorldDataLayers::Get(GetWorld());
+	return GetWorld()->GetWorldDataLayers();
 }
 
 AWorldDataLayers* UDataLayerEditorSubsystem::GetWorldDataLayers(bool bCreateIfNotFound)
 {
-	return AWorldDataLayers::Get(GetWorld(), bCreateIfNotFound);
+	AWorldDataLayers* WorldDataLayers = GetWorld()->GetWorldDataLayers();
+
+	if (!WorldDataLayers && bCreateIfNotFound)
+	{
+		WorldDataLayers = AWorldDataLayers::Create(GetWorld());
+	}
+
+	return WorldDataLayers;
 }
 
 void UDataLayerEditorSubsystem::AddAllDataLayersTo(TArray<TWeakObjectPtr<UDataLayer>>& OutDataLayers) const
diff --git a/Engine/Source/Editor/DataLayerEditor/Private/DataLayer/DataLayerHierarchy.cpp b/Engine/Source/Editor/DataLayerEditor/Private/DataLayer/DataLayerHierarchy.cpp
--- a/Engine/Source/Editor/DataLayerEditor/Private/DataLayer/DataLayerHierarchy.cpp
+++ b/Engine/Source/Editor/DataLayerEditor/Private/DataLayer/DataLayerHierarchy.cpp
@@ -59,7 +59,7 @@ FDataLayerHierarchy::~FDataLayerHierarchy()
 
 void FDataLayerHierarchy::CreateItems(TArray<FSceneOutlinerTreeItemPtr>& OutItems) const
 {
-	const AWorldDataLayers* WorldDataLayers = AWorldDataLayers::Get(RepresentingWorld.Get());
+	const AWorldDataLayers* WorldDataLayers = RepresentingWorld.Get()->GetWorldDataLayers();
 	if (!WorldDataLayers)
 	{
 		return;
diff --git a/Engine/Source/Editor/DataLayerEditor/Private/DataLayer/DataLayerMode.cpp b/Engine/Source/Editor/DataLayerEditor/Private/DataLayer/DataLayerMode.cpp
--- a/Engine/Source/Editor/DataLayerEditor/Private/DataLayer/DataLayerMode.cpp
+++ b/Engine/Source/Editor/DataLayerEditor/Private/DataLayer/DataLayerMode.cpp
@@ -506,7 +506,7 @@ void FDataLayerMode::RegisterContextMenu()
 			TArray<UDataLayer*> SelectedDataLayers = GetSelectedDataLayers(SceneOutliner);
 
 			TArray<const UDataLayer*> AllDataLayers;
-			if (const AWorldDataLayers* WorldDataLayers = AWorldDataLayers::Get(RepresentingWorld.Get()))
+			if (const AWorldDataLayers* WorldDataLayers = RepresentingWorld.Get()->GetWorldDataLayers())
 			{
 				WorldDataLayers->ForEachDataLayer([&AllDataLayers](UDataLayer* DataLayer)
 				{
diff --git a/Engine/Source/Editor/LevelEditor/Private/SLevelViewport.cpp b/Engine/Source/Editor/LevelEditor/Private/SLevelViewport.cpp
--- a/Engine/Source/Editor/LevelEditor/Private/SLevelViewport.cpp
+++ b/Engine/Source/Editor/LevelEditor/Private/SLevelViewport.cpp
@@ -2112,7 +2112,7 @@ void SLevelViewport::OnToggleAllDataLayers(bool bVisible)
 	else
 	{
 		TArray<FName> AllDataLayerNames;
-		if (const AWorldDataLayers* WorldDataLayers = AWorldDataLayers::Get(GetWorld()))
+		if (const AWorldDataLayers* WorldDataLayers = GetWorld()->GetWorldDataLayers())
 		{
 			WorldDataLayers->ForEachDataLayer([&AllDataLayerNames](UDataLayer* DataLayer)
 			{
diff --git a/Engine/Source/Editor/LevelEditor/Private/SLevelViewportToolBar.cpp b/Engine/Source/Editor/LevelEditor/Private/SLevelViewportToolBar.cpp
--- a/Engine/Source/Editor/LevelEditor/Private/SLevelViewportToolBar.cpp
+++ b/Engine/Source/Editor/LevelEditor/Private/SLevelViewportToolBar.cpp
@@ -1462,7 +1462,7 @@ void SLevelViewportToolBar::FillShowDataLayersMenu(UToolMenu* Menu, TWeakPtr<cla
 		TSharedRef<SLevelViewport> ViewportRef = Viewport.Pin().ToSharedRef();
 		FToolMenuSection& Section = Menu->AddSection("LevelViewportDataLayers2");
 		// Get all the DataLayers and create an entry for each of them
-		if (const AWorldDataLayers* WorldDataLayers = AWorldDataLayers::Get(ViewportRef->GetWorld()))
+		if (const AWorldDataLayers* WorldDataLayers = ViewportRef->GetWorld()->GetWorldDataLayers())
 		{
 			WorldDataLayers->ForEachDataLayer([&Section, &ViewportRef, &DataLayerIcon](UDataLayer* DataLayer)
 			{
diff --git a/Engine/Source/Editor/UnrealEd/Private/Commandlets/WorldPartitionConvertCommandlet.cpp b/Engine/Source/Editor/UnrealEd/Private/Commandlets/WorldPartitionConvertCommandlet.cpp
--- a/Engine/Source/Editor/UnrealEd/Private/Commandlets/WorldPartitionConvertCommandlet.cpp
+++ b/Engine/Source/Editor/UnrealEd/Private/Commandlets/WorldPartitionConvertCommandlet.cpp
@@ -853,7 +853,11 @@ int32 UWorldPartitionConvertCommandlet::Main(const FString& Params)
 	}
 
 	UPackage* MainPackage = MainLevel->GetPackage();
-	AWorldDataLayers* MainWorldDataLayers = AWorldDataLayers::Get(MainWorld, /*bCreateIfNotFound*/true);
+	AWorldDataLayers* MainWorldDataLayers = MainWorld->GetWorldDataLayers();
+	if (!MainWorldDataLayers)
+	{
+		MainWorldDataLayers = AWorldDataLayers::Create(MainWorld);
+	}
 
 	OnWorldLoaded(MainWorld);
 
diff --git a/Engine/Source/Editor/UnrealEd/Private/WorldPartition/WorldPartitionBuilder.cpp b/Engine/Source/Editor/UnrealEd/Private/WorldPartition/WorldPartitionBuilder.cpp
--- a/Engine/Source/Editor/UnrealEd/Private/WorldPartition/WorldPartitionBuilder.cpp
+++ b/Engine/Source/Editor/UnrealEd/Private/WorldPartition/WorldPartitionBuilder.cpp
@@ -162,7 +162,7 @@ bool UWorldPartitionBuilder::Run(UWorld* World, FPackageSourceControlHelper& Pac
 	UWorldPartition* WorldPartition = World->GetWorldPartition();
 
 	// Properly Setup DataLayers for Builder
-	if (const AWorldDataLayers* WorldDataLayers = AWorldDataLayers::Get(World))
+	if (const AWorldDataLayers* WorldDataLayers = World->GetWorldDataLayers())
 	{
 		// Load Data Layers
 		bool bUpdateEditorCells = false;
diff --git a/Engine/Source/Runtime/Engine/Classes/Engine/Level.h b/Engine/Source/Runtime/Engine/Classes/Engine/Level.h
--- a/Engine/Source/Runtime/Engine/Classes/Engine/Level.h
+++ b/Engine/Source/Runtime/Engine/Classes/Engine/Level.h
@@ -24,6 +24,7 @@ class AInstancedFoliageActor;
 class ALevelBounds;
 class APlayerController;
 class AWorldSettings;
+class AWorldDataLayers;
 class FSceneInterface;
 class ITargetPlatform;
 class UAssetUserData;
@@ -745,6 +746,9 @@ private:
 	UPROPERTY()
 	TObjectPtr<AWorldSettings> WorldSettings;
 
+	UPROPERTY()
+	TObjectPtr<AWorldDataLayers> WorldDataLayers;
+
 	/** Cached level collection that this level is contained in, for faster access than looping through the collections in the world. */
 	FLevelCollection* CachedLevelCollection;
 
@@ -962,6 +966,15 @@ public:
 
 	ENGINE_API void SetWorldSettings(AWorldSettings* NewWorldSettings);
 
+	/**
+	 * Returns the world data layers info for this level.
+	 *
+	 * @return		The AWorldDataLayers for this level.
+	 */
+	ENGINE_API AWorldDataLayers* GetWorldDataLayers() const;
+
+	ENGINE_API void SetWorldDataLayers(AWorldDataLayers* NewWorldDataLayers);
+
 	/**
 	 * Returns the UWorldPartition for this level.
 	 *
diff --git a/Engine/Source/Runtime/Engine/Classes/Engine/World.h b/Engine/Source/Runtime/Engine/Classes/Engine/World.h
--- a/Engine/Source/Runtime/Engine/Classes/Engine/World.h
+++ b/Engine/Source/Runtime/Engine/Classes/Engine/World.h
@@ -1464,9 +1464,6 @@ private:
 	/** Utility function to cleanup streaming levels that point to invalid level packages */
 	void RepairStreamingLevels();
 
-	/** Utility function that is used to ensure that a World has the correct WorldDataLayers */
-	void RepairWorldDataLayers();
-	
 #if INCLUDE_CHAOS
 	/** Utility function that is used to ensure that a World has the correct ChaosActor */
 	void RepairChaosActors();
@@ -2494,6 +2491,14 @@ public:
 	AWorldSettings* K2_GetWorldSettings();
 	AWorldSettings* GetWorldSettings( bool bCheckStreamingPersistent = false, bool bChecked = true ) const;
 
+	/**
+	 * Returns the AWorldDataLayers actor associated with this world.
+	 *
+	 * @return AWorldDataLayers actor associated with this world
+	 */
+	AWorldDataLayers* GetWorldDataLayers() const;
+	void SetWorldDataLayers(AWorldDataLayers* NewWorldDataLayers);
+
 	/** Returns a human friendly display string for the current world (showing the kind of world when in multiplayer PIE) */
 	FString GetDebugDisplayName() const;
 	/**
diff --git a/Engine/Source/Runtime/Engine/Private/ActorEditor.cpp b/Engine/Source/Runtime/Engine/Private/ActorEditor.cpp
--- a/Engine/Source/Runtime/Engine/Private/ActorEditor.cpp
+++ b/Engine/Source/Runtime/Engine/Private/ActorEditor.cpp
@@ -1120,7 +1120,7 @@ bool AActor::HasDataLayers() const
 
 bool AActor::HasValidDataLayers() const
 {
-	if (const AWorldDataLayers* WorldDataLayers = AWorldDataLayers::Get(GetWorld()))
+	if (const AWorldDataLayers* WorldDataLayers = GetWorld()->GetWorldDataLayers())
 	{
 		for (const FActorDataLayer& DataLayer : DataLayers)
 		{
@@ -1152,13 +1152,13 @@ bool AActor::HasAllDataLayers(const TArray<const UDataLayer*>& InDataLayers) con
 
 TArray<FName> AActor::GetDataLayerNames() const
 {
-	const AWorldDataLayers* WorldDataLayers = AWorldDataLayers::Get(GetWorld());
+	const AWorldDataLayers* WorldDataLayers = GetWorld()->GetWorldDataLayers();
 	return WorldDataLayers ? WorldDataLayers->GetDataLayerNames(DataLayers) : TArray<FName>();
 }
 
 TArray<const UDataLayer*> AActor::GetDataLayerObjects() const
 {
-	const AWorldDataLayers* WorldDataLayers = AWorldDataLayers::Get(GetWorld());
+	const AWorldDataLayers* WorldDataLayers = GetWorld()->GetWorldDataLayers();
 	return WorldDataLayers ? WorldDataLayers->GetDataLayerObjects(DataLayers) : TArray<const UDataLayer*>();
 }
 
@@ -1176,25 +1176,31 @@ bool AActor::HasAnyOfDataLayers(const TArray<FName>& DataLayerNames) const
 
 void AActor::FixupDataLayers()
 {
-	if (!SupportsDataLayer())
+	if (!GetPackage()->HasAnyPackageFlags(PKG_PlayInEditor))
 	{
-		DataLayers.Empty();
-	}
+		if (!SupportsDataLayer())
+		{
+			DataLayers.Empty();
+		}
 
-	if (const AWorldDataLayers* WorldDataLayers = AWorldDataLayers::Get(GetWorld()))
-	{
-		TSet<FName> ExistingDataLayers;
-		for (int32 Index = 0; Index < DataLayers.Num();)
+		if (GetWorld())
 		{
-			const FName& DataLayer = DataLayers[Index].Name;
-			if (!WorldDataLayers->GetDataLayerFromName(DataLayer) || ExistingDataLayers.Contains(DataLayer))
-			{
-				DataLayers.RemoveAtSwap(Index);
-			}
-			else
+			if (const AWorldDataLayers* WorldDataLayers = GetWorld()->GetWorldDataLayers())
 			{
-				ExistingDataLayers.Add(DataLayer);
-				++Index;
+				TSet<FName> ExistingDataLayers;
+				for (int32 Index = 0; Index < DataLayers.Num();)
+				{
+					const FName& DataLayer = DataLayers[Index].Name;
+					if (!WorldDataLayers->GetDataLayerFromName(DataLayer) || ExistingDataLayers.Contains(DataLayer))
+					{
+						DataLayers.RemoveAtSwap(Index);
+					}
+					else
+					{
+						ExistingDataLayers.Add(DataLayer);
+						++Index;
+					}
+				}
 			}
 		}
 	}
diff --git a/Engine/Source/Runtime/Engine/Private/ActorPartition/ActorPartitionSubsystem.cpp b/Engine/Source/Runtime/Engine/Private/ActorPartition/ActorPartitionSubsystem.cpp
--- a/Engine/Source/Runtime/Engine/Private/ActorPartition/ActorPartitionSubsystem.cpp
+++ b/Engine/Source/Runtime/Engine/Private/ActorPartition/ActorPartitionSubsystem.cpp
@@ -383,7 +383,7 @@ APartitionActor* UActorPartitionSubsystem::GetActor(const TSubclassOf<APartition
 	{
 		if (UDataLayerSubsystem* DataLayerSubsystem = World->GetSubsystem<UDataLayerSubsystem>())
 		{
-			if (const AWorldDataLayers* WorldDataLayers = AWorldDataLayers::Get(World))
+			if (const AWorldDataLayers* WorldDataLayers = World->GetWorldDataLayers())
 			{
 				for (const FName& DataLayer : DataLayerSubsystem->GetDataLayerEditorContext().GetDataLayers())
 				{
diff --git a/Engine/Source/Runtime/Engine/Private/Level.cpp b/Engine/Source/Runtime/Engine/Private/Level.cpp
--- a/Engine/Source/Runtime/Engine/Private/Level.cpp
+++ b/Engine/Source/Runtime/Engine/Private/Level.cpp
@@ -2113,6 +2113,18 @@ void ULevel::SetWorldSettings(AWorldSettings* NewWorldSettings)
 	}
 }
 
+AWorldDataLayers* ULevel::GetWorldDataLayers() const
+{
+	return WorldDataLayers;
+}
+
+void ULevel::SetWorldDataLayers(AWorldDataLayers* NewWorldDataLayers)
+{
+	check(!WorldDataLayers || WorldDataLayers == NewWorldDataLayers);
+	WorldDataLayers = NewWorldDataLayers;
+}
+
+
 UWorldPartition* ULevel::GetWorldPartition() const
 {
 	return bIsPartitioned ? GetWorldSettings()->GetWorldPartition() : nullptr;
diff --git a/Engine/Source/Runtime/Engine/Private/World.cpp b/Engine/Source/Runtime/Engine/Private/World.cpp
--- a/Engine/Source/Runtime/Engine/Private/World.cpp
+++ b/Engine/Source/Runtime/Engine/Private/World.cpp
@@ -1118,7 +1118,6 @@ void UWorld::PostLoad()
 #if WITH_EDITOR
 	RepairWorldSettings();
 	RepairStreamingLevels();
-	RepairWorldDataLayers();
 #endif
 #if INCLUDE_CHAOS
 	//RepairChaosActors();
@@ -1558,16 +1557,6 @@ void UWorld::RepairWorldSettings()
 	check(GetWorldSettings());
 }
 
-void UWorld::RepairWorldDataLayers()
-{
-#if 0
-	if (const UWorldPartition* WorldPartition = GetWorldPartition())
-	{
-		AWorldDataLayers::Get(this, /*bCreateIfNotFound*/true);
-	}
-#endif
-}
-
 void UWorld::InitWorld(const InitializationValues IVS)
 {
 	if (!ensure(!bIsWorldInitialized))
@@ -1652,7 +1641,6 @@ void UWorld::InitWorld(const InitializationValues IVS)
 #if WITH_EDITOR
 	RepairWorldSettings();
 	RepairStreamingLevels();
-	RepairWorldDataLayers();
 #endif
 #if INCLUDE_CHAOS
 	//RepairChaosActors();
@@ -5213,6 +5201,24 @@ AWorldSettings* UWorld::GetWorldSettings( const bool bCheckStreamingPersistent,
 	return WorldSettings;
 }
 
+AWorldDataLayers* UWorld::GetWorldDataLayers() const
+{
+	AWorldDataLayers* WorldDataLayers = nullptr;
+	if (PersistentLevel)
+	{
+		WorldDataLayers = PersistentLevel->GetWorldDataLayers();
+	}
+	return WorldDataLayers;
+}
+
+void UWorld::SetWorldDataLayers(AWorldDataLayers* NewWorldDataLayers)
+{
+	if (PersistentLevel)
+	{
+		PersistentLevel->SetWorldDataLayers(NewWorldDataLayers);
+	}
+}
+
 FString UWorld::GetDebugDisplayName() const
 {
 #if WITH_EDITOR
diff --git a/Engine/Source/Runtime/Engine/Private/WorldPartition/DataLayer/DataLayerEditorContext.cpp b/Engine/Source/Runtime/Engine/Private/WorldPartition/DataLayer/DataLayerEditorContext.cpp
--- a/Engine/Source/Runtime/Engine/Private/WorldPartition/DataLayer/DataLayerEditorContext.cpp
+++ b/Engine/Source/Runtime/Engine/Private/WorldPartition/DataLayer/DataLayerEditorContext.cpp
@@ -15,7 +15,7 @@
 FDataLayerEditorContext::FDataLayerEditorContext(UWorld* InWorld, const TArray<FName>& InDataLayers)
 	: Hash(FDataLayerEditorContext::EmptyHash)
 {
-	const AWorldDataLayers* WorldDataLayers = AWorldDataLayers::Get(InWorld);
+	const AWorldDataLayers* WorldDataLayers = InWorld->GetWorldDataLayers();
 	if (!WorldDataLayers)
 	{
 		return;
@@ -53,7 +53,7 @@ FScopeChangeDataLayerEditorContext::FScopeChangeDataLayerEditorContext(UWorld* I
 FScopeChangeDataLayerEditorContext::FScopeChangeDataLayerEditorContext(UWorld* InWorld, const FActorDataLayer& InContextDataLayer)
 	: World(InWorld)
 {
-	const AWorldDataLayers* WorldDataLayers = AWorldDataLayers::Get(InWorld);
+	const AWorldDataLayers* WorldDataLayers = InWorld->GetWorldDataLayers();
 	Initialize(FDataLayerEditorContext(InWorld, WorldDataLayers ? WorldDataLayers->GetDataLayerNames({ InContextDataLayer }) : TArray<FName>()));
 }
 
diff --git a/Engine/Source/Runtime/Engine/Private/WorldPartition/DataLayer/DataLayerSubsystem.cpp b/Engine/Source/Runtime/Engine/Private/WorldPartition/DataLayer/DataLayerSubsystem.cpp
--- a/Engine/Source/Runtime/Engine/Private/WorldPartition/DataLayer/DataLayerSubsystem.cpp
+++ b/Engine/Source/Runtime/Engine/Private/WorldPartition/DataLayer/DataLayerSubsystem.cpp
@@ -56,7 +56,7 @@ void UDataLayerSubsystem::PostInitialize()
 	if (GetWorld()->IsGameWorld())
 	{
 		// Initialize Dynamically loaded Data Layers state
-		if (const AWorldDataLayers* WorldDataLayers = AWorldDataLayers::Get(GetWorld()))
+		if (const AWorldDataLayers* WorldDataLayers = GetWorld()->GetWorldDataLayers())
 		{
 			WorldDataLayers->ForEachDataLayer([this](class UDataLayer* DataLayer)
 			{
@@ -77,13 +77,13 @@ UDataLayer* UDataLayerSubsystem::GetDataLayer(const FActorDataLayer& InDataLayer
 
 UDataLayer* UDataLayerSubsystem::GetDataLayerFromLabel(FName InDataLayerLabel) const
 {
-	const AWorldDataLayers* WorldDataLayers = AWorldDataLayers::Get(GetWorld());
+	const AWorldDataLayers* WorldDataLayers = GetWorld()->GetWorldDataLayers();
 	return WorldDataLayers ? const_cast<UDataLayer*>(WorldDataLayers->GetDataLayerFromLabel(InDataLayerLabel)) : nullptr;
 }
 
 UDataLayer* UDataLayerSubsystem::GetDataLayerFromName(FName InDataLayerName) const
 {
-	const AWorldDataLayers* WorldDataLayers = AWorldDataLayers::Get(GetWorld());
+	const AWorldDataLayers* WorldDataLayers = GetWorld()->GetWorldDataLayers();
 	return WorldDataLayers ? const_cast<UDataLayer*>(WorldDataLayers->GetDataLayerFromName(InDataLayerName)) : nullptr;
 }
 
@@ -290,7 +290,7 @@ TArray<UDataLayer*> UDataLayerSubsystem::ConvertArgsToDataLayers(UWorld* World,
 	for (const FString& Arg : Args)
 	{
 		FName DataLayerLabel = FName(Arg);
-		if (const AWorldDataLayers* WorldDataLayers = AWorldDataLayers::Get(World))
+		if (const AWorldDataLayers* WorldDataLayers = World->GetWorldDataLayers())
 		{
 			UDataLayer* DataLayer = const_cast<UDataLayer*>(WorldDataLayers->GetDataLayerFromLabel(DataLayerLabel));
 			if (!DataLayer)
@@ -344,7 +344,7 @@ void UDataLayerSubsystem::GetDataLayerDebugColors(TMap<FName, FColor>& OutMappin
 {
 	OutMapping.Reset();
 
-	const AWorldDataLayers* WorldDataLayers = AWorldDataLayers::Get(GetWorld());
+	const AWorldDataLayers* WorldDataLayers = GetWorld()->GetWorldDataLayers();
 	if (!WorldDataLayers)
 	{
 		return;
@@ -361,7 +361,7 @@ void UDataLayerSubsystem::DrawDataLayersLegend(UCanvas* Canvas, FVector2D& Offse
 {
 	check(Canvas);
 
-	const AWorldDataLayers* WorldDataLayers = AWorldDataLayers::Get(GetWorld());
+	const AWorldDataLayers* WorldDataLayers = GetWorld()->GetWorldDataLayers();
 	if (!WorldDataLayers)
 	{
 		return;
diff --git a/Engine/Source/Runtime/Engine/Private/WorldPartition/DataLayer/WorldDataLayers.cpp b/Engine/Source/Runtime/Engine/Private/WorldPartition/DataLayer/WorldDataLayers.cpp
--- a/Engine/Source/Runtime/Engine/Private/WorldPartition/DataLayer/WorldDataLayers.cpp
+++ b/Engine/Source/Runtime/Engine/Private/WorldPartition/DataLayer/WorldDataLayers.cpp
@@ -18,62 +18,42 @@ AWorldDataLayers::AWorldDataLayers(const FObjectInitializer& ObjectInitializer)
 {
 }
 
-const AWorldDataLayers* AWorldDataLayers::Get(UWorld* World)
-{
-	if (World)
-	{
-		// Prepare flags for actor iterator. Don't use default Flags because it uses EActorIteratorFlags::OnlyActiveLevels 
-		// which will make this code return no actor when cooking (because world is not initialized)
-		EActorIteratorFlags Flags = EActorIteratorFlags::SkipPendingKill;
-		if (!IsRunningCookCommandlet())
-		{
-			Flags |= EActorIteratorFlags::OnlyActiveLevels;
-		}
-		for (AWorldDataLayers* Actor : TActorRange<AWorldDataLayers>(World, AWorldDataLayers::StaticClass(), Flags))
-		{
-			if (Actor)
-			{
-				check(!Actor->IsPendingKill());
-				return Actor;
-			}
-		}
-	}
-	return nullptr;
-}
-
 #if WITH_EDITOR
-AWorldDataLayers* AWorldDataLayers::Get(UWorld* World, bool bCreateIfNotFound)
+AWorldDataLayers* AWorldDataLayers::Create(UWorld* World)
 {
-	if (bCreateIfNotFound)
-	{
-		AWorldDataLayers* WorldDataLayers = nullptr;
-		static FName WorldDataLayersName = AWorldDataLayers::StaticClass()->GetFName();
-		if (UObject* ExistingObject = StaticFindObject(nullptr, World->PersistentLevel, *WorldDataLayersName.ToString()))
-		{
-			WorldDataLayers = CastChecked<AWorldDataLayers>(ExistingObject);
-			if (WorldDataLayers->IsPendingKill())
-			{
-				// Handle the case where the actor already exists, but it's pending kill
-				WorldDataLayers->Rename(nullptr, nullptr, REN_DontCreateRedirectors | REN_DoNotDirty | REN_NonTransactional | REN_ForceNoResetLoaders);
-				WorldDataLayers = nullptr;
-			}
-		}
+	check(World);
+	check(!World->GetWorldDataLayers());
+
+	AWorldDataLayers* WorldDataLayers = nullptr;
 
-		if (!WorldDataLayers)
+	static FName WorldDataLayersName = AWorldDataLayers::StaticClass()->GetFName();
+	if (UObject* ExistingObject = StaticFindObject(nullptr, World->PersistentLevel, *WorldDataLayersName.ToString()))
+	{
+		WorldDataLayers = CastChecked<AWorldDataLayers>(ExistingObject);
+		if (WorldDataLayers->IsPendingKill())
 		{
-			FActorSpawnParameters SpawnParams;
-			SpawnParams.OverrideLevel = World->PersistentLevel;
-			SpawnParams.bHideFromSceneOutliner = true;
-			SpawnParams.Name = WorldDataLayersName;
-			SpawnParams.NameMode = FActorSpawnParameters::ESpawnActorNameMode::Required_Fatal;
-			WorldDataLayers = World->SpawnActor<AWorldDataLayers>(AWorldDataLayers::StaticClass(), SpawnParams);
+			// Handle the case where the actor already exists, but it's pending kill
+			WorldDataLayers->Rename(nullptr, nullptr, REN_DontCreateRedirectors | REN_DoNotDirty | REN_NonTransactional | REN_ForceNoResetLoaders);
+			WorldDataLayers = nullptr;
 		}
-		return WorldDataLayers;
 	}
-	else
+
+	if (!WorldDataLayers)
 	{
-		return const_cast<AWorldDataLayers*>(Get(World));
+		FActorSpawnParameters SpawnParams;
+		SpawnParams.OverrideLevel = World->PersistentLevel;
+		SpawnParams.bHideFromSceneOutliner = true;
+		SpawnParams.Name = WorldDataLayersName;
+		SpawnParams.NameMode = FActorSpawnParameters::ESpawnActorNameMode::Required_Fatal;
+		WorldDataLayers = World->SpawnActor<AWorldDataLayers>(AWorldDataLayers::StaticClass(), SpawnParams);
 	}
+
+	check(WorldDataLayers);
+
+	World->Modify();
+	World->SetWorldDataLayers(WorldDataLayers);
+
+	return WorldDataLayers;
 }
 
 FName AWorldDataLayers::GenerateUniqueDataLayerLabel(const FName& InDataLayerLabel) const
@@ -231,6 +211,11 @@ void AWorldDataLayers::PostLoad()
 {
 	Super::PostLoad();
 
+	if (GetWorld())
+	{
+		GetWorld()->SetWorldDataLayers(this);
+	}
+
 #if WITH_EDITOR
 	// Initialize DataLayer's IsDynamicallyLoadedInEditor based on DataLayerEditorPerProjectUserSettings
 	const TArray<FName>& SettingsDataLayersNotLoadedInEditor = GetMutableDefault<UWorldPartitionEditorPerProjectUserSettings>()->GetWorldDataLayersNotLoadedInEditor(GetWorld());
diff --git a/Engine/Source/Runtime/Engine/Private/WorldPartition/HLOD/HLODLayer.cpp b/Engine/Source/Runtime/Engine/Private/WorldPartition/HLOD/HLODLayer.cpp
--- a/Engine/Source/Runtime/Engine/Private/WorldPartition/HLOD/HLODLayer.cpp
+++ b/Engine/Source/Runtime/Engine/Private/WorldPartition/HLOD/HLODLayer.cpp
@@ -124,7 +124,7 @@ UHLODLayer* UHLODLayer::GetHLODLayer(const FWorldPartitionActorDesc& InActorDesc
 	{
 		if (UDataLayerSubsystem* DataLayerSubsystem = InWorldPartition->GetWorld()->GetSubsystem<UDataLayerSubsystem>())
 		{
-			if (const AWorldDataLayers* WorldDataLayers = AWorldDataLayers::Get(InWorldPartition->GetWorld()))
+			if (const AWorldDataLayers* WorldDataLayers = InWorldPartition->GetWorld()->GetWorldDataLayers())
 			{
 				for (const FName& DataLayerName : InActorDesc.GetDataLayers())
 				{
diff --git a/Engine/Source/Runtime/Engine/Private/WorldPartition/WorldPartition.cpp b/Engine/Source/Runtime/Engine/Private/WorldPartition/WorldPartition.cpp
--- a/Engine/Source/Runtime/Engine/Private/WorldPartition/WorldPartition.cpp
+++ b/Engine/Source/Runtime/Engine/Private/WorldPartition/WorldPartition.cpp
@@ -324,6 +324,11 @@ void UWorldPartition::Initialize(UWorld* InWorld, const FTransform& InTransform)
 		}
 
 		EditorHash->Initialize();
+
+		if (!GetWorld()->GetWorldDataLayers())
+		{
+			AWorldDataLayers::Create(GetWorld());
+		}
 	}
 
 	if (!RuntimeHash)
@@ -828,7 +833,7 @@ bool UWorldPartition::ShouldActorBeLoaded(const FWorldPartitionActorDesc* ActorD
 		return true;
 	}
 
-	if (const AWorldDataLayers* WorldDataLayers = AWorldDataLayers::Get(GetWorld()))
+	if (const AWorldDataLayers* WorldDataLayers = GetWorld()->GetWorldDataLayers())
 	{
 		uint32 NumValidLayers = 0;
 		for (const FName& DataLayerName : ActorDesc->GetDataLayers())
diff --git a/Engine/Source/Runtime/Engine/Private/WorldPartition/WorldPartitionActorCluster.cpp b/Engine/Source/Runtime/Engine/Private/WorldPartition/WorldPartitionActorCluster.cpp
--- a/Engine/Source/Runtime/Engine/Private/WorldPartition/WorldPartitionActorCluster.cpp
+++ b/Engine/Source/Runtime/Engine/Private/WorldPartition/WorldPartitionActorCluster.cpp
@@ -16,7 +16,7 @@ template<class LayerNameContainer>
 TSet<const UDataLayer*> GetDataLayers(UWorld* InWorld, const LayerNameContainer& DataLayerNames)
 {
 	TSet<const UDataLayer*> DataLayers;
-	if (const AWorldDataLayers* WorldDataLayers = AWorldDataLayers::Get(InWorld))
+	if (const AWorldDataLayers* WorldDataLayers = InWorld->GetWorldDataLayers())
 	{
 		for (const FName& DataLayerName : DataLayerNames)
 		{
diff --git a/Engine/Source/Runtime/Engine/Private/WorldPartition/WorldPartitionRuntimeCell.cpp b/Engine/Source/Runtime/Engine/Private/WorldPartition/WorldPartitionRuntimeCell.cpp
--- a/Engine/Source/Runtime/Engine/Private/WorldPartition/WorldPartitionRuntimeCell.cpp
+++ b/Engine/Source/Runtime/Engine/Private/WorldPartition/WorldPartitionRuntimeCell.cpp
@@ -49,7 +49,7 @@ void UWorldPartitionRuntimeCell::UpdateDebugName()
 	Builder += FString::Printf(TEXT("L%d_X%d_Y%d"), Coords.Z, Coords.X, Coords.Y);
 	int32 DataLayerCount = DataLayers.Num();
 
-	const AWorldDataLayers* WorldDataLayers = AWorldDataLayers::Get(GetOuterUWorldPartition()->GetWorld());
+	const AWorldDataLayers* WorldDataLayers = GetOuterUWorldPartition()->GetWorld()->GetWorldDataLayers();
 	TArray<const UDataLayer*> DataLayerObjects;
 	if (WorldDataLayers && (DataLayerCount > 0))
 	{
diff --git a/Engine/Source/Runtime/Engine/Public/WorldPartition/DataLayer/WorldDataLayers.h b/Engine/Source/Runtime/Engine/Public/WorldPartition/DataLayer/WorldDataLayers.h
--- a/Engine/Source/Runtime/Engine/Public/WorldPartition/DataLayer/WorldDataLayers.h
+++ b/Engine/Source/Runtime/Engine/Public/WorldPartition/DataLayer/WorldDataLayers.h
@@ -18,11 +18,10 @@ class ENGINE_API AWorldDataLayers : public AInfo
 	GENERATED_UCLASS_BODY()
 
 public:
-	static const AWorldDataLayers* Get(UWorld* World);
 	virtual void PostLoad() override;
 	
-#if WITH_EDITOR
-	static AWorldDataLayers* Get(UWorld* World, bool bCreateIfNotFound);
+#if WITH_EDITOR	
+	static AWorldDataLayers* Create(UWorld* World);
 	UDataLayer* CreateDataLayer(FName InName = TEXT("DataLayer"), EObjectFlags InObjectFlags = RF_NoFlags);
 	bool RemoveDataLayer(UDataLayer* InDataLayer);
 	bool RemoveDataLayers(const TArray<UDataLayer*>& InDataLayers);
