Disable VT Compression validation code

diff --git a/Engine/Source/Runtime/Engine/Private/TextureDerivedDataTask.cpp b/Engine/Source/Runtime/Engine/Private/TextureDerivedDataTask.cpp
--- a/Engine/Source/Runtime/Engine/Private/TextureDerivedDataTask.cpp
+++ b/Engine/Source/Runtime/Engine/Private/TextureDerivedDataTask.cpp
@@ -37,14 +37,14 @@
 
 static TAutoConsoleVariable<int32> CVarVTValidateCompressionOnLoad(
 	TEXT("r.VT.ValidateCompressionOnLoad"),
-	1,
+	0,
 	TEXT("Validates that VT data contains no compression errors when loading from DDC")
 	TEXT("This is slow, but allows debugging corrupt VT data (and allows recovering from bad DDC)")
 );
 
 static TAutoConsoleVariable<int32> CVarVTValidateCompressionOnSave(
 	TEXT("r.VT.ValidateCompressionOnSave"),
-	1,
+	0,
 	TEXT("Validates that VT data contains no compression errors before saving to DDC")
 	TEXT("This is slow, but allows debugging corrupt VT data")
 );
