Added a bunch of profile tags for partitioned worlds generate streaming.

diff --git a/Engine/Source/Runtime/Engine/Private/WorldPartition/RuntimeSpatialHash/RuntimeSpatialHashGridHelper.cpp b/Engine/Source/Runtime/Engine/Private/WorldPartition/RuntimeSpatialHash/RuntimeSpatialHashGridHelper.cpp
--- a/Engine/Source/Runtime/Engine/Private/WorldPartition/RuntimeSpatialHash/RuntimeSpatialHashGridHelper.cpp
+++ b/Engine/Source/Runtime/Engine/Private/WorldPartition/RuntimeSpatialHash/RuntimeSpatialHashGridHelper.cpp
@@ -130,6 +130,7 @@ FSquare2DGridHelper GetGridHelper(const FBox& WorldBounds, int32 GridCellSize)
 
 FSquare2DGridHelper GetPartitionedActors(const UWorldPartition* WorldPartition, const FBox& WorldBounds, const FSpatialHashRuntimeGrid& Grid, const TArray<const FActorClusterInstance*>& GridActors)
 {
+	TRACE_CPUPROFILER_EVENT_SCOPE(GetPartitionedActors);
 	UE_SCOPED_TIMER(TEXT("GetPartitionedActors"), LogWorldPartition);
 
 	//
diff --git a/Engine/Source/Runtime/Engine/Private/WorldPartition/WorldPartitionActorCluster.cpp b/Engine/Source/Runtime/Engine/Private/WorldPartition/WorldPartitionActorCluster.cpp
--- a/Engine/Source/Runtime/Engine/Private/WorldPartition/WorldPartitionActorCluster.cpp
+++ b/Engine/Source/Runtime/Engine/Private/WorldPartition/WorldPartitionActorCluster.cpp
@@ -397,6 +397,8 @@ const TArray<FActorCluster>& FActorClusterContext::CreateActorClustersImpl(const
 
 void FActorClusterContext::CreateActorClusters()
 {
+	TRACE_CPUPROFILER_EVENT_SCOPE(CreateActorClusters);
+
 	// First Instance is the main WorldPartition
 	FBox WorldBounds(ForceInit);
 	CreateContainerInstanceRecursive(0, FTransform::Identity, EContainerClusterMode::Partitioned, WorldPartition, TSet<FName>(), WorldBounds);
diff --git a/Engine/Source/Runtime/Engine/Private/WorldPartition/WorldPartitionRuntimeHash.cpp b/Engine/Source/Runtime/Engine/Private/WorldPartition/WorldPartitionRuntimeHash.cpp
--- a/Engine/Source/Runtime/Engine/Private/WorldPartition/WorldPartitionRuntimeHash.cpp
+++ b/Engine/Source/Runtime/Engine/Private/WorldPartition/WorldPartitionRuntimeHash.cpp
@@ -77,6 +77,8 @@ void UWorldPartitionRuntimeHash::ForceExternalActorLevelReference(bool bForceExt
 
 void UWorldPartitionRuntimeHash::CreateActorDescViewMap(const UActorDescContainer* Container, TMap<FGuid, FWorldPartitionActorDescView>& OutActorDescViewMap) const
 {
+	TRACE_CPUPROFILER_EVENT_SCOPE(CreateActorDescViewMap);
+
 	// Build the actor desc view map
 	OutActorDescViewMap.Empty();
 
diff --git a/Engine/Source/Runtime/Engine/Private/WorldPartition/WorldPartitionRuntimeSpatialHash.cpp b/Engine/Source/Runtime/Engine/Private/WorldPartition/WorldPartitionRuntimeSpatialHash.cpp
--- a/Engine/Source/Runtime/Engine/Private/WorldPartition/WorldPartitionRuntimeSpatialHash.cpp
+++ b/Engine/Source/Runtime/Engine/Private/WorldPartition/WorldPartitionRuntimeSpatialHash.cpp
@@ -675,6 +675,8 @@ FName UWorldPartitionRuntimeSpatialHash::GetCellName(FName InGridName, const FIn
 
 void UWorldPartitionRuntimeSpatialHash::UpdateActorDescViewMap(const FBox& WorldBounds, TMap<FGuid, FWorldPartitionActorDescView>& ActorDescViewMap) const
 {
+	TRACE_CPUPROFILER_EVENT_SCOPE(UpdateActorDescViewMap);
+
 	Super::UpdateActorDescViewMap(WorldBounds, ActorDescViewMap);
 
 	TMap<FName, int32> GridsMapping;
@@ -726,7 +728,7 @@ void UWorldPartitionRuntimeSpatialHash::UpdateActorDescViewMap(const FBox& World
 
 bool UWorldPartitionRuntimeSpatialHash::CreateStreamingGrid(const FSpatialHashRuntimeGrid& RuntimeGrid, const FSquare2DGridHelper& PartionedActors, EWorldPartitionStreamingMode Mode, UWorldPartitionStreamingPolicy* StreamingPolicy, TArray<FString>* OutPackagesToGenerate)
 {
-	TRACE_CPUPROFILER_EVENT_SCOPE(UWorldPartitionRuntimeSpatialHash::CreateStreamingGrid);
+	TRACE_CPUPROFILER_EVENT_SCOPE(CreateStreamingGrid);
 
 	UWorldPartition* WorldPartition = GetOuterUWorldPartition();
 	UWorld* World = WorldPartition->GetWorld();
