Backout 16168385

diff --git a/Engine/Source/Editor/DataLayerEditor/Private/DataLayer/DataLayerEditorSubsystem.cpp b/Engine/Source/Editor/DataLayerEditor/Private/DataLayer/DataLayerEditorSubsystem.cpp
--- a/Engine/Source/Editor/DataLayerEditor/Private/DataLayer/DataLayerEditorSubsystem.cpp
+++ b/Engine/Source/Editor/DataLayerEditor/Private/DataLayer/DataLayerEditorSubsystem.cpp
@@ -1093,19 +1093,12 @@ bool UDataLayerEditorSubsystem::TryGetDataLayerFromLabel(const FName& DataLayerL
 
 const AWorldDataLayers* UDataLayerEditorSubsystem::GetWorldDataLayers() const
 {
-	return GetWorld()->GetWorldDataLayers();
+	return AWorldDataLayers::Get(GetWorld());
 }
 
 AWorldDataLayers* UDataLayerEditorSubsystem::GetWorldDataLayers(bool bCreateIfNotFound)
 {
-	AWorldDataLayers* WorldDataLayers = GetWorld()->GetWorldDataLayers();
-
-	if (!WorldDataLayers && bCreateIfNotFound)
-	{
-		WorldDataLayers = AWorldDataLayers::Create(GetWorld());
-	}
-
-	return WorldDataLayers;
+	return AWorldDataLayers::Get(GetWorld(), bCreateIfNotFound);
 }
 
 void UDataLayerEditorSubsystem::AddAllDataLayersTo(TArray<TWeakObjectPtr<UDataLayer>>& OutDataLayers) const
diff --git a/Engine/Source/Editor/DataLayerEditor/Private/DataLayer/DataLayerHierarchy.cpp b/Engine/Source/Editor/DataLayerEditor/Private/DataLayer/DataLayerHierarchy.cpp
--- a/Engine/Source/Editor/DataLayerEditor/Private/DataLayer/DataLayerHierarchy.cpp
+++ b/Engine/Source/Editor/DataLayerEditor/Private/DataLayer/DataLayerHierarchy.cpp
@@ -59,7 +59,7 @@ FDataLayerHierarchy::~FDataLayerHierarchy()
 
 void FDataLayerHierarchy::CreateItems(TArray<FSceneOutlinerTreeItemPtr>& OutItems) const
 {
-	const AWorldDataLayers* WorldDataLayers = RepresentingWorld.Get()->GetWorldDataLayers();
+	const AWorldDataLayers* WorldDataLayers = AWorldDataLayers::Get(RepresentingWorld.Get());
 	if (!WorldDataLayers)
 	{
 		return;
diff --git a/Engine/Source/Editor/DataLayerEditor/Private/DataLayer/DataLayerMode.cpp b/Engine/Source/Editor/DataLayerEditor/Private/DataLayer/DataLayerMode.cpp
--- a/Engine/Source/Editor/DataLayerEditor/Private/DataLayer/DataLayerMode.cpp
+++ b/Engine/Source/Editor/DataLayerEditor/Private/DataLayer/DataLayerMode.cpp
@@ -506,7 +506,7 @@ void FDataLayerMode::RegisterContextMenu()
 			TArray<UDataLayer*> SelectedDataLayers = GetSelectedDataLayers(SceneOutliner);
 
 			TArray<const UDataLayer*> AllDataLayers;
-			if (const AWorldDataLayers* WorldDataLayers = RepresentingWorld.Get()->GetWorldDataLayers())
+			if (const AWorldDataLayers* WorldDataLayers = AWorldDataLayers::Get(RepresentingWorld.Get()))
 			{
 				WorldDataLayers->ForEachDataLayer([&AllDataLayers](UDataLayer* DataLayer)
 				{
diff --git a/Engine/Source/Editor/LevelEditor/Private/SLevelViewport.cpp b/Engine/Source/Editor/LevelEditor/Private/SLevelViewport.cpp
--- a/Engine/Source/Editor/LevelEditor/Private/SLevelViewport.cpp
+++ b/Engine/Source/Editor/LevelEditor/Private/SLevelViewport.cpp
@@ -2112,7 +2112,7 @@ void SLevelViewport::OnToggleAllDataLayers(bool bVisible)
 	else
 	{
 		TArray<FName> AllDataLayerNames;
-		if (const AWorldDataLayers* WorldDataLayers = GetWorld()->GetWorldDataLayers())
+		if (const AWorldDataLayers* WorldDataLayers = AWorldDataLayers::Get(GetWorld()))
 		{
 			WorldDataLayers->ForEachDataLayer([&AllDataLayerNames](UDataLayer* DataLayer)
 			{
diff --git a/Engine/Source/Editor/LevelEditor/Private/SLevelViewportToolBar.cpp b/Engine/Source/Editor/LevelEditor/Private/SLevelViewportToolBar.cpp
--- a/Engine/Source/Editor/LevelEditor/Private/SLevelViewportToolBar.cpp
+++ b/Engine/Source/Editor/LevelEditor/Private/SLevelViewportToolBar.cpp
@@ -1462,7 +1462,7 @@ void SLevelViewportToolBar::FillShowDataLayersMenu(UToolMenu* Menu, TWeakPtr<cla
 		TSharedRef<SLevelViewport> ViewportRef = Viewport.Pin().ToSharedRef();
 		FToolMenuSection& Section = Menu->AddSection("LevelViewportDataLayers2");
 		// Get all the DataLayers and create an entry for each of them
-		if (const AWorldDataLayers* WorldDataLayers = ViewportRef->GetWorld()->GetWorldDataLayers())
+		if (const AWorldDataLayers* WorldDataLayers = AWorldDataLayers::Get(ViewportRef->GetWorld()))
 		{
 			WorldDataLayers->ForEachDataLayer([&Section, &ViewportRef, &DataLayerIcon](UDataLayer* DataLayer)
 			{
diff --git a/Engine/Source/Editor/UnrealEd/Private/Commandlets/WorldPartitionConvertCommandlet.cpp b/Engine/Source/Editor/UnrealEd/Private/Commandlets/WorldPartitionConvertCommandlet.cpp
--- a/Engine/Source/Editor/UnrealEd/Private/Commandlets/WorldPartitionConvertCommandlet.cpp
+++ b/Engine/Source/Editor/UnrealEd/Private/Commandlets/WorldPartitionConvertCommandlet.cpp
@@ -853,11 +853,7 @@ int32 UWorldPartitionConvertCommandlet::Main(const FString& Params)
 	}
 
 	UPackage* MainPackage = MainLevel->GetPackage();
-	AWorldDataLayers* MainWorldDataLayers = MainWorld->GetWorldDataLayers();
-	if (!MainWorldDataLayers)
-	{
-		MainWorldDataLayers = AWorldDataLayers::Create(MainWorld);
-	}
+	AWorldDataLayers* MainWorldDataLayers = AWorldDataLayers::Get(MainWorld, /*bCreateIfNotFound*/true);
 
 	OnWorldLoaded(MainWorld);
 
diff --git a/Engine/Source/Editor/UnrealEd/Private/WorldPartition/WorldPartitionBuilder.cpp b/Engine/Source/Editor/UnrealEd/Private/WorldPartition/WorldPartitionBuilder.cpp
--- a/Engine/Source/Editor/UnrealEd/Private/WorldPartition/WorldPartitionBuilder.cpp
+++ b/Engine/Source/Editor/UnrealEd/Private/WorldPartition/WorldPartitionBuilder.cpp
@@ -162,7 +162,7 @@ bool UWorldPartitionBuilder::Run(UWorld* World, FPackageSourceControlHelper& Pac
 	UWorldPartition* WorldPartition = World->GetWorldPartition();
 
 	// Properly Setup DataLayers for Builder
-	if (const AWorldDataLayers* WorldDataLayers = World->GetWorldDataLayers())
+	if (const AWorldDataLayers* WorldDataLayers = AWorldDataLayers::Get(World))
 	{
 		// Load Data Layers
 		bool bUpdateEditorCells = false;
diff --git a/Engine/Source/Runtime/Engine/Classes/Engine/Level.h b/Engine/Source/Runtime/Engine/Classes/Engine/Level.h
--- a/Engine/Source/Runtime/Engine/Classes/Engine/Level.h
+++ b/Engine/Source/Runtime/Engine/Classes/Engine/Level.h
@@ -24,7 +24,6 @@ class AInstancedFoliageActor;
 class ALevelBounds;
 class APlayerController;
 class AWorldSettings;
-class AWorldDataLayers;
 class FSceneInterface;
 class ITargetPlatform;
 class UAssetUserData;
@@ -746,9 +745,6 @@ private:
 	UPROPERTY()
 	TObjectPtr<AWorldSettings> WorldSettings;
 
-	UPROPERTY()
-	TObjectPtr<AWorldDataLayers> WorldDataLayers;
-
 	/** Cached level collection that this level is contained in, for faster access than looping through the collections in the world. */
 	FLevelCollection* CachedLevelCollection;
 
@@ -966,15 +962,6 @@ public:
 
 	ENGINE_API void SetWorldSettings(AWorldSettings* NewWorldSettings);
 
-	/**
-	 * Returns the world data layers info for this level.
-	 *
-	 * @return		The AWorldDataLayers for this level.
-	 */
-	ENGINE_API AWorldDataLayers* GetWorldDataLayers() const;
-
-	ENGINE_API void SetWorldDataLayers(AWorldDataLayers* NewWorldDataLayers);
-
 	/**
 	 * Returns the UWorldPartition for this level.
 	 *
diff --git a/Engine/Source/Runtime/Engine/Classes/Engine/World.h b/Engine/Source/Runtime/Engine/Classes/Engine/World.h
--- a/Engine/Source/Runtime/Engine/Classes/Engine/World.h
+++ b/Engine/Source/Runtime/Engine/Classes/Engine/World.h
@@ -1464,6 +1464,9 @@ private:
 	/** Utility function to cleanup streaming levels that point to invalid level packages */
 	void RepairStreamingLevels();
 
+	/** Utility function that is used to ensure that a World has the correct WorldDataLayers */
+	void RepairWorldDataLayers();
+	
 #if INCLUDE_CHAOS
 	/** Utility function that is used to ensure that a World has the correct ChaosActor */
 	void RepairChaosActors();
@@ -2491,14 +2494,6 @@ public:
 	AWorldSettings* K2_GetWorldSettings();
 	AWorldSettings* GetWorldSettings( bool bCheckStreamingPersistent = false, bool bChecked = true ) const;
 
-	/**
-	 * Returns the AWorldDataLayers actor associated with this world.
-	 *
-	 * @return AWorldDataLayers actor associated with this world
-	 */
-	AWorldDataLayers* GetWorldDataLayers() const;
-	void SetWorldDataLayers(AWorldDataLayers* NewWorldDataLayers);
-
 	/** Returns a human friendly display string for the current world (showing the kind of world when in multiplayer PIE) */
 	FString GetDebugDisplayName() const;
 	/**
diff --git a/Engine/Source/Runtime/Engine/Private/ActorEditor.cpp b/Engine/Source/Runtime/Engine/Private/ActorEditor.cpp
--- a/Engine/Source/Runtime/Engine/Private/ActorEditor.cpp
+++ b/Engine/Source/Runtime/Engine/Private/ActorEditor.cpp
@@ -1120,7 +1120,7 @@ bool AActor::HasDataLayers() const
 
 bool AActor::HasValidDataLayers() const
 {
-	if (const AWorldDataLayers* WorldDataLayers = GetWorld()->GetWorldDataLayers())
+	if (const AWorldDataLayers* WorldDataLayers = AWorldDataLayers::Get(GetWorld()))
 	{
 		for (const FActorDataLayer& DataLayer : DataLayers)
 		{
@@ -1152,13 +1152,13 @@ bool AActor::HasAllDataLayers(const TArray<const UDataLayer*>& InDataLayers) con
 
 TArray<FName> AActor::GetDataLayerNames() const
 {
-	const AWorldDataLayers* WorldDataLayers = GetWorld()->GetWorldDataLayers();
+	const AWorldDataLayers* WorldDataLayers = AWorldDataLayers::Get(GetWorld());
 	return WorldDataLayers ? WorldDataLayers->GetDataLayerNames(DataLayers) : TArray<FName>();
 }
 
 TArray<const UDataLayer*> AActor::GetDataLayerObjects() const
 {
-	const AWorldDataLayers* WorldDataLayers = GetWorld()->GetWorldDataLayers();
+	const AWorldDataLayers* WorldDataLayers = AWorldDataLayers::Get(GetWorld());
 	return WorldDataLayers ? WorldDataLayers->GetDataLayerObjects(DataLayers) : TArray<const UDataLayer*>();
 }
 
@@ -1176,31 +1176,25 @@ bool AActor::HasAnyOfDataLayers(const TArray<FName>& DataLayerNames) const
 
 void AActor::FixupDataLayers()
 {
-	if (!GetPackage()->HasAnyPackageFlags(PKG_PlayInEditor))
+	if (!SupportsDataLayer())
 	{
-		if (!SupportsDataLayer())
-		{
-			DataLayers.Empty();
-		}
+		DataLayers.Empty();
+	}
 
-		if (GetWorld())
+	if (const AWorldDataLayers* WorldDataLayers = AWorldDataLayers::Get(GetWorld()))
+	{
+		TSet<FName> ExistingDataLayers;
+		for (int32 Index = 0; Index < DataLayers.Num();)
 		{
-			if (const AWorldDataLayers* WorldDataLayers = GetWorld()->GetWorldDataLayers())
+			const FName& DataLayer = DataLayers[Index].Name;
+			if (!WorldDataLayers->GetDataLayerFromName(DataLayer) || ExistingDataLayers.Contains(DataLayer))
 			{
-				TSet<FName> ExistingDataLayers;
-				for (int32 Index = 0; Index < DataLayers.Num();)
-				{
-					const FName& DataLayer = DataLayers[Index].Name;
-					if (!WorldDataLayers->GetDataLayerFromName(DataLayer) || ExistingDataLayers.Contains(DataLayer))
-					{
-						DataLayers.RemoveAtSwap(Index);
-					}
-					else
-					{
-						ExistingDataLayers.Add(DataLayer);
-						++Index;
-					}
-				}
+				DataLayers.RemoveAtSwap(Index);
+			}
+			else
+			{
+				ExistingDataLayers.Add(DataLayer);
+				++Index;
 			}
 		}
 	}
diff --git a/Engine/Source/Runtime/Engine/Private/ActorPartition/ActorPartitionSubsystem.cpp b/Engine/Source/Runtime/Engine/Private/ActorPartition/ActorPartitionSubsystem.cpp
--- a/Engine/Source/Runtime/Engine/Private/ActorPartition/ActorPartitionSubsystem.cpp
+++ b/Engine/Source/Runtime/Engine/Private/ActorPartition/ActorPartitionSubsystem.cpp
@@ -387,7 +387,7 @@ APartitionActor* UActorPartitionSubsystem::GetActor(const TSubclassOf<APartition
 	{
 		if (UDataLayerSubsystem* DataLayerSubsystem = World->GetSubsystem<UDataLayerSubsystem>())
 		{
-			if (const AWorldDataLayers* WorldDataLayers = World->GetWorldDataLayers())
+			if (const AWorldDataLayers* WorldDataLayers = AWorldDataLayers::Get(World))
 			{
 				for (const FName& DataLayer : DataLayerSubsystem->GetDataLayerEditorContext().GetDataLayers())
 				{
diff --git a/Engine/Source/Runtime/Engine/Private/Level.cpp b/Engine/Source/Runtime/Engine/Private/Level.cpp
--- a/Engine/Source/Runtime/Engine/Private/Level.cpp
+++ b/Engine/Source/Runtime/Engine/Private/Level.cpp
@@ -2113,18 +2113,6 @@ void ULevel::SetWorldSettings(AWorldSettings* NewWorldSettings)
 	}
 }
 
-AWorldDataLayers* ULevel::GetWorldDataLayers() const
-{
-	return WorldDataLayers;
-}
-
-void ULevel::SetWorldDataLayers(AWorldDataLayers* NewWorldDataLayers)
-{
-	check(!WorldDataLayers || WorldDataLayers == NewWorldDataLayers);
-	WorldDataLayers = NewWorldDataLayers;
-}
-
-
 UWorldPartition* ULevel::GetWorldPartition() const
 {
 	return bIsPartitioned ? GetWorldSettings()->GetWorldPartition() : nullptr;
diff --git a/Engine/Source/Runtime/Engine/Private/World.cpp b/Engine/Source/Runtime/Engine/Private/World.cpp
--- a/Engine/Source/Runtime/Engine/Private/World.cpp
+++ b/Engine/Source/Runtime/Engine/Private/World.cpp
@@ -1118,6 +1118,7 @@ void UWorld::PostLoad()
 #if WITH_EDITOR
 	RepairWorldSettings();
 	RepairStreamingLevels();
+	RepairWorldDataLayers();
 #endif
 #if INCLUDE_CHAOS
 	//RepairChaosActors();
@@ -1557,6 +1558,16 @@ void UWorld::RepairWorldSettings()
 	check(GetWorldSettings());
 }
 
+void UWorld::RepairWorldDataLayers()
+{
+#if 0
+	if (const UWorldPartition* WorldPartition = GetWorldPartition())
+	{
+		AWorldDataLayers::Get(this, /*bCreateIfNotFound*/true);
+	}
+#endif
+}
+
 void UWorld::InitWorld(const InitializationValues IVS)
 {
 	if (!ensure(!bIsWorldInitialized))
@@ -1641,6 +1652,7 @@ void UWorld::InitWorld(const InitializationValues IVS)
 #if WITH_EDITOR
 	RepairWorldSettings();
 	RepairStreamingLevels();
+	RepairWorldDataLayers();
 #endif
 #if INCLUDE_CHAOS
 	//RepairChaosActors();
@@ -5201,24 +5213,6 @@ AWorldSettings* UWorld::GetWorldSettings( const bool bCheckStreamingPersistent,
 	return WorldSettings;
 }
 
-AWorldDataLayers* UWorld::GetWorldDataLayers() const
-{
-	AWorldDataLayers* WorldDataLayers = nullptr;
-	if (PersistentLevel)
-	{
-		WorldDataLayers = PersistentLevel->GetWorldDataLayers();
-	}
-	return WorldDataLayers;
-}
-
-void UWorld::SetWorldDataLayers(AWorldDataLayers* NewWorldDataLayers)
-{
-	if (PersistentLevel)
-	{
-		PersistentLevel->SetWorldDataLayers(NewWorldDataLayers);
-	}
-}
-
 FString UWorld::GetDebugDisplayName() const
 {
 #if WITH_EDITOR
diff --git a/Engine/Source/Runtime/Engine/Private/WorldPartition/DataLayer/DataLayerEditorContext.cpp b/Engine/Source/Runtime/Engine/Private/WorldPartition/DataLayer/DataLayerEditorContext.cpp
--- a/Engine/Source/Runtime/Engine/Private/WorldPartition/DataLayer/DataLayerEditorContext.cpp
+++ b/Engine/Source/Runtime/Engine/Private/WorldPartition/DataLayer/DataLayerEditorContext.cpp
@@ -15,7 +15,7 @@
 FDataLayerEditorContext::FDataLayerEditorContext(UWorld* InWorld, const TArray<FName>& InDataLayers)
 	: Hash(FDataLayerEditorContext::EmptyHash)
 {
-	const AWorldDataLayers* WorldDataLayers = InWorld->GetWorldDataLayers();
+	const AWorldDataLayers* WorldDataLayers = AWorldDataLayers::Get(InWorld);
 	if (!WorldDataLayers)
 	{
 		return;
@@ -53,7 +53,7 @@ FScopeChangeDataLayerEditorContext::FScopeChangeDataLayerEditorContext(UWorld* I
 FScopeChangeDataLayerEditorContext::FScopeChangeDataLayerEditorContext(UWorld* InWorld, const FActorDataLayer& InContextDataLayer)
 	: World(InWorld)
 {
-	const AWorldDataLayers* WorldDataLayers = InWorld->GetWorldDataLayers();
+	const AWorldDataLayers* WorldDataLayers = AWorldDataLayers::Get(InWorld);
 	Initialize(FDataLayerEditorContext(InWorld, WorldDataLayers ? WorldDataLayers->GetDataLayerNames({ InContextDataLayer }) : TArray<FName>()));
 }
 
diff --git a/Engine/Source/Runtime/Engine/Private/WorldPartition/DataLayer/DataLayerSubsystem.cpp b/Engine/Source/Runtime/Engine/Private/WorldPartition/DataLayer/DataLayerSubsystem.cpp
--- a/Engine/Source/Runtime/Engine/Private/WorldPartition/DataLayer/DataLayerSubsystem.cpp
+++ b/Engine/Source/Runtime/Engine/Private/WorldPartition/DataLayer/DataLayerSubsystem.cpp
@@ -56,7 +56,7 @@ void UDataLayerSubsystem::PostInitialize()
 	if (GetWorld()->IsGameWorld())
 	{
 		// Initialize Dynamically loaded Data Layers state
-		if (const AWorldDataLayers* WorldDataLayers = GetWorld()->GetWorldDataLayers())
+		if (const AWorldDataLayers* WorldDataLayers = AWorldDataLayers::Get(GetWorld()))
 		{
 			WorldDataLayers->ForEachDataLayer([this](class UDataLayer* DataLayer)
 			{
@@ -77,13 +77,13 @@ UDataLayer* UDataLayerSubsystem::GetDataLayer(const FActorDataLayer& InDataLayer
 
 UDataLayer* UDataLayerSubsystem::GetDataLayerFromLabel(FName InDataLayerLabel) const
 {
-	const AWorldDataLayers* WorldDataLayers = GetWorld()->GetWorldDataLayers();
+	const AWorldDataLayers* WorldDataLayers = AWorldDataLayers::Get(GetWorld());
 	return WorldDataLayers ? const_cast<UDataLayer*>(WorldDataLayers->GetDataLayerFromLabel(InDataLayerLabel)) : nullptr;
 }
 
 UDataLayer* UDataLayerSubsystem::GetDataLayerFromName(FName InDataLayerName) const
 {
-	const AWorldDataLayers* WorldDataLayers = GetWorld()->GetWorldDataLayers();
+	const AWorldDataLayers* WorldDataLayers = AWorldDataLayers::Get(GetWorld());
 	return WorldDataLayers ? const_cast<UDataLayer*>(WorldDataLayers->GetDataLayerFromName(InDataLayerName)) : nullptr;
 }
 
@@ -290,7 +290,7 @@ TArray<UDataLayer*> UDataLayerSubsystem::ConvertArgsToDataLayers(UWorld* World,
 	for (const FString& Arg : Args)
 	{
 		FName DataLayerLabel = FName(Arg);
-		if (const AWorldDataLayers* WorldDataLayers = World->GetWorldDataLayers())
+		if (const AWorldDataLayers* WorldDataLayers = AWorldDataLayers::Get(World))
 		{
 			UDataLayer* DataLayer = const_cast<UDataLayer*>(WorldDataLayers->GetDataLayerFromLabel(DataLayerLabel));
 			if (!DataLayer)
@@ -344,7 +344,7 @@ void UDataLayerSubsystem::GetDataLayerDebugColors(TMap<FName, FColor>& OutMappin
 {
 	OutMapping.Reset();
 
-	const AWorldDataLayers* WorldDataLayers = GetWorld()->GetWorldDataLayers();
+	const AWorldDataLayers* WorldDataLayers = AWorldDataLayers::Get(GetWorld());
 	if (!WorldDataLayers)
 	{
 		return;
@@ -361,7 +361,7 @@ void UDataLayerSubsystem::DrawDataLayersLegend(UCanvas* Canvas, FVector2D& Offse
 {
 	check(Canvas);
 
-	const AWorldDataLayers* WorldDataLayers = GetWorld()->GetWorldDataLayers();
+	const AWorldDataLayers* WorldDataLayers = AWorldDataLayers::Get(GetWorld());
 	if (!WorldDataLayers)
 	{
 		return;
diff --git a/Engine/Source/Runtime/Engine/Private/WorldPartition/DataLayer/WorldDataLayers.cpp b/Engine/Source/Runtime/Engine/Private/WorldPartition/DataLayer/WorldDataLayers.cpp
--- a/Engine/Source/Runtime/Engine/Private/WorldPartition/DataLayer/WorldDataLayers.cpp
+++ b/Engine/Source/Runtime/Engine/Private/WorldPartition/DataLayer/WorldDataLayers.cpp
@@ -18,42 +18,62 @@ AWorldDataLayers::AWorldDataLayers(const FObjectInitializer& ObjectInitializer)
 {
 }
 
-#if WITH_EDITOR
-AWorldDataLayers* AWorldDataLayers::Create(UWorld* World)
+const AWorldDataLayers* AWorldDataLayers::Get(UWorld* World)
 {
-	check(World);
-	check(!World->GetWorldDataLayers());
-
-	AWorldDataLayers* WorldDataLayers = nullptr;
-
-	static FName WorldDataLayersName = AWorldDataLayers::StaticClass()->GetFName();
-	if (UObject* ExistingObject = StaticFindObject(nullptr, World->PersistentLevel, *WorldDataLayersName.ToString()))
+	if (World)
 	{
-		WorldDataLayers = CastChecked<AWorldDataLayers>(ExistingObject);
-		if (WorldDataLayers->IsPendingKill())
+		// Prepare flags for actor iterator. Don't use default Flags because it uses EActorIteratorFlags::OnlyActiveLevels 
+		// which will make this code return no actor when cooking (because world is not initialized)
+		EActorIteratorFlags Flags = EActorIteratorFlags::SkipPendingKill;
+		if (!IsRunningCookCommandlet())
 		{
-			// Handle the case where the actor already exists, but it's pending kill
-			WorldDataLayers->Rename(nullptr, nullptr, REN_DontCreateRedirectors | REN_DoNotDirty | REN_NonTransactional | REN_ForceNoResetLoaders);
-			WorldDataLayers = nullptr;
+			Flags |= EActorIteratorFlags::OnlyActiveLevels;
+		}
+		for (AWorldDataLayers* Actor : TActorRange<AWorldDataLayers>(World, AWorldDataLayers::StaticClass(), Flags))
+		{
+			if (Actor)
+			{
+				check(!Actor->IsPendingKill());
+				return Actor;
+			}
 		}
 	}
+	return nullptr;
+}
 
-	if (!WorldDataLayers)
+#if WITH_EDITOR
+AWorldDataLayers* AWorldDataLayers::Get(UWorld* World, bool bCreateIfNotFound)
+{
+	if (bCreateIfNotFound)
 	{
-		FActorSpawnParameters SpawnParams;
-		SpawnParams.OverrideLevel = World->PersistentLevel;
-		SpawnParams.bHideFromSceneOutliner = true;
-		SpawnParams.Name = WorldDataLayersName;
-		SpawnParams.NameMode = FActorSpawnParameters::ESpawnActorNameMode::Required_Fatal;
-		WorldDataLayers = World->SpawnActor<AWorldDataLayers>(AWorldDataLayers::StaticClass(), SpawnParams);
-	}
-
-	check(WorldDataLayers);
-
-	World->Modify();
-	World->SetWorldDataLayers(WorldDataLayers);
+		AWorldDataLayers* WorldDataLayers = nullptr;
+		static FName WorldDataLayersName = AWorldDataLayers::StaticClass()->GetFName();
+		if (UObject* ExistingObject = StaticFindObject(nullptr, World->PersistentLevel, *WorldDataLayersName.ToString()))
+		{
+			WorldDataLayers = CastChecked<AWorldDataLayers>(ExistingObject);
+			if (WorldDataLayers->IsPendingKill())
+			{
+				// Handle the case where the actor already exists, but it's pending kill
+				WorldDataLayers->Rename(nullptr, nullptr, REN_DontCreateRedirectors | REN_DoNotDirty | REN_NonTransactional | REN_ForceNoResetLoaders);
+				WorldDataLayers = nullptr;
+			}
+		}
 
-	return WorldDataLayers;
+		if (!WorldDataLayers)
+		{
+			FActorSpawnParameters SpawnParams;
+			SpawnParams.OverrideLevel = World->PersistentLevel;
+			SpawnParams.bHideFromSceneOutliner = true;
+			SpawnParams.Name = WorldDataLayersName;
+			SpawnParams.NameMode = FActorSpawnParameters::ESpawnActorNameMode::Required_Fatal;
+			WorldDataLayers = World->SpawnActor<AWorldDataLayers>(AWorldDataLayers::StaticClass(), SpawnParams);
+		}
+		return WorldDataLayers;
+	}
+	else
+	{
+		return const_cast<AWorldDataLayers*>(Get(World));
+	}
 }
 
 FName AWorldDataLayers::GenerateUniqueDataLayerLabel(const FName& InDataLayerLabel) const
@@ -211,11 +231,6 @@ void AWorldDataLayers::PostLoad()
 {
 	Super::PostLoad();
 
-	if (GetWorld())
-	{
-		GetWorld()->SetWorldDataLayers(this);
-	}
-
 #if WITH_EDITOR
 	// Initialize DataLayer's IsDynamicallyLoadedInEditor based on DataLayerEditorPerProjectUserSettings
 	const TArray<FName>& SettingsDataLayersNotLoadedInEditor = GetMutableDefault<UWorldPartitionEditorPerProjectUserSettings>()->GetWorldDataLayersNotLoadedInEditor(GetWorld());
diff --git a/Engine/Source/Runtime/Engine/Private/WorldPartition/HLOD/HLODLayer.cpp b/Engine/Source/Runtime/Engine/Private/WorldPartition/HLOD/HLODLayer.cpp
--- a/Engine/Source/Runtime/Engine/Private/WorldPartition/HLOD/HLODLayer.cpp
+++ b/Engine/Source/Runtime/Engine/Private/WorldPartition/HLOD/HLODLayer.cpp
@@ -124,7 +124,7 @@ UHLODLayer* UHLODLayer::GetHLODLayer(const FWorldPartitionActorDesc& InActorDesc
 	{
 		if (UDataLayerSubsystem* DataLayerSubsystem = InWorldPartition->GetWorld()->GetSubsystem<UDataLayerSubsystem>())
 		{
-			if (const AWorldDataLayers* WorldDataLayers = InWorldPartition->GetWorld()->GetWorldDataLayers())
+			if (const AWorldDataLayers* WorldDataLayers = AWorldDataLayers::Get(InWorldPartition->GetWorld()))
 			{
 				for (const FName& DataLayerName : InActorDesc.GetDataLayers())
 				{
diff --git a/Engine/Source/Runtime/Engine/Private/WorldPartition/WorldPartition.cpp b/Engine/Source/Runtime/Engine/Private/WorldPartition/WorldPartition.cpp
--- a/Engine/Source/Runtime/Engine/Private/WorldPartition/WorldPartition.cpp
+++ b/Engine/Source/Runtime/Engine/Private/WorldPartition/WorldPartition.cpp
@@ -324,11 +324,6 @@ void UWorldPartition::Initialize(UWorld* InWorld, const FTransform& InTransform)
 		}
 
 		EditorHash->Initialize();
-
-		if (!GetWorld()->GetWorldDataLayers())
-		{
-			AWorldDataLayers::Create(GetWorld());
-		}
 	}
 
 	if (!RuntimeHash)
@@ -833,7 +828,7 @@ bool UWorldPartition::ShouldActorBeLoaded(const FWorldPartitionActorDesc* ActorD
 		return true;
 	}
 
-	if (const AWorldDataLayers* WorldDataLayers = GetWorld()->GetWorldDataLayers())
+	if (const AWorldDataLayers* WorldDataLayers = AWorldDataLayers::Get(GetWorld()))
 	{
 		uint32 NumValidLayers = 0;
 		for (const FName& DataLayerName : ActorDesc->GetDataLayers())
diff --git a/Engine/Source/Runtime/Engine/Private/WorldPartition/WorldPartitionActorCluster.cpp b/Engine/Source/Runtime/Engine/Private/WorldPartition/WorldPartitionActorCluster.cpp
--- a/Engine/Source/Runtime/Engine/Private/WorldPartition/WorldPartitionActorCluster.cpp
+++ b/Engine/Source/Runtime/Engine/Private/WorldPartition/WorldPartitionActorCluster.cpp
@@ -16,7 +16,7 @@ template<class LayerNameContainer>
 TSet<const UDataLayer*> GetDataLayers(UWorld* InWorld, const LayerNameContainer& DataLayerNames)
 {
 	TSet<const UDataLayer*> DataLayers;
-	if (const AWorldDataLayers* WorldDataLayers = InWorld->GetWorldDataLayers())
+	if (const AWorldDataLayers* WorldDataLayers = AWorldDataLayers::Get(InWorld))
 	{
 		for (const FName& DataLayerName : DataLayerNames)
 		{
diff --git a/Engine/Source/Runtime/Engine/Private/WorldPartition/WorldPartitionRuntimeCell.cpp b/Engine/Source/Runtime/Engine/Private/WorldPartition/WorldPartitionRuntimeCell.cpp
--- a/Engine/Source/Runtime/Engine/Private/WorldPartition/WorldPartitionRuntimeCell.cpp
+++ b/Engine/Source/Runtime/Engine/Private/WorldPartition/WorldPartitionRuntimeCell.cpp
@@ -49,7 +49,7 @@ void UWorldPartitionRuntimeCell::UpdateDebugName()
 	Builder += FString::Printf(TEXT("L%d_X%d_Y%d"), Coords.Z, Coords.X, Coords.Y);
 	int32 DataLayerCount = DataLayers.Num();
 
-	const AWorldDataLayers* WorldDataLayers = GetOuterUWorldPartition()->GetWorld()->GetWorldDataLayers();
+	const AWorldDataLayers* WorldDataLayers = AWorldDataLayers::Get(GetOuterUWorldPartition()->GetWorld());
 	TArray<const UDataLayer*> DataLayerObjects;
 	if (WorldDataLayers && (DataLayerCount > 0))
 	{
diff --git a/Engine/Source/Runtime/Engine/Public/WorldPartition/DataLayer/WorldDataLayers.h b/Engine/Source/Runtime/Engine/Public/WorldPartition/DataLayer/WorldDataLayers.h
--- a/Engine/Source/Runtime/Engine/Public/WorldPartition/DataLayer/WorldDataLayers.h
+++ b/Engine/Source/Runtime/Engine/Public/WorldPartition/DataLayer/WorldDataLayers.h
@@ -18,10 +18,11 @@ class ENGINE_API AWorldDataLayers : public AInfo
 	GENERATED_UCLASS_BODY()
 
 public:
+	static const AWorldDataLayers* Get(UWorld* World);
 	virtual void PostLoad() override;
 	
-#if WITH_EDITOR	
-	static AWorldDataLayers* Create(UWorld* World);
+#if WITH_EDITOR
+	static AWorldDataLayers* Get(UWorld* World, bool bCreateIfNotFound);
 	UDataLayer* CreateDataLayer(FName InName = TEXT("DataLayer"), EObjectFlags InObjectFlags = RF_NoFlags);
 	bool RemoveDataLayer(UDataLayer* InDataLayer);
 	bool RemoveDataLayers(const TArray<UDataLayer*>& InDataLayers);
