Enable Packet Order Correction by default.

diff --git a/Engine/Source/Runtime/Engine/Private/NetConnection.cpp b/Engine/Source/Runtime/Engine/Private/NetConnection.cpp
--- a/Engine/Source/Runtime/Engine/Private/NetConnection.cpp
+++ b/Engine/Source/Runtime/Engine/Private/NetConnection.cpp
@@ -64,7 +64,7 @@ static TAutoConsoleVariable<int32> CVarForceNetFlush(TEXT("net.ForceNetFlush"),
 	TEXT("Immediately flush send buffer when written to (helps trace packet writes - WARNING: May be unstable)."));
 #endif
 
-static TAutoConsoleVariable<int32> CVarNetDoPacketOrderCorrection(TEXT("net.DoPacketOrderCorrection"), 0,
+static TAutoConsoleVariable<int32> CVarNetDoPacketOrderCorrection(TEXT("net.DoPacketOrderCorrection"), 1,
 	TEXT("Whether or not to try to fix 'out of order' packet sequences, by caching packets and waiting for the missing sequence."));
 
 static TAutoConsoleVariable<int32> CVarNetPacketOrderCorrectionEnableThreshold(TEXT("net.PacketOrderCorrectionEnableThreshold"), 1,
