Fix out-of-bound memory access when decoding RIFF audio files

diff --git a/Engine/Source/Runtime/Engine/Private/Audio.cpp b/Engine/Source/Runtime/Engine/Private/Audio.cpp
--- a/Engine/Source/Runtime/Engine/Private/Audio.cpp
+++ b/Engine/Source/Runtime/Engine/Private/Audio.cpp
@@ -1497,7 +1497,7 @@ bool FWaveModInfo::ReadWaveInfo( const uint8* WaveData, int32 WaveDataSize, FStr
 	}
 
 	// Cue chunks are optional
-	if (INTEL_ORDER32(RiffChunk->ChunkID) == UE_mmioFOURCC('c', 'u', 'e', ' '))
+	if (reinterpret_cast<uint8*>(RiffChunk) + sizeof(RiffChunk->ChunkID) <= WaveDataEnd && INTEL_ORDER32(RiffChunk->ChunkID) == UE_mmioFOURCC('c', 'u', 'e', ' '))
 	{
 		FRiffCueChunk* CueChunk = (FRiffCueChunk*)((uint8*)RiffChunk);
 
@@ -1524,7 +1524,7 @@ bool FWaveModInfo::ReadWaveInfo( const uint8* WaveData, int32 WaveDataSize, FStr
 		}
 
 		// Label chunks are also optional
-		if (INTEL_ORDER32(RiffChunk->ChunkID) == UE_mmioFOURCC('L', 'I', 'S', 'T'))
+		if (reinterpret_cast<uint8*>(RiffChunk) + sizeof(RiffChunk->ChunkID) <= WaveDataEnd && INTEL_ORDER32(RiffChunk->ChunkID) == UE_mmioFOURCC('L', 'I', 'S', 'T'))
 		{
 			FRiffListChunk* ListChunk = (FRiffListChunk*)((uint8*)RiffChunk);
 			
