Fix runtime grid debug display 2d world axis to match world partition editor.

diff --git a/Engine/Source/Runtime/Engine/Private/WorldPartition/WorldPartitionRuntimeSpatialHash.cpp b/Engine/Source/Runtime/Engine/Private/WorldPartition/WorldPartitionRuntimeSpatialHash.cpp
--- a/Engine/Source/Runtime/Engine/Private/WorldPartition/WorldPartitionRuntimeSpatialHash.cpp
+++ b/Engine/Source/Runtime/Engine/Private/WorldPartition/WorldPartitionRuntimeSpatialHash.cpp
@@ -332,7 +332,7 @@ void FSpatialHashStreamingGrid::Draw2D(UCanvas* Canvas, const TArray<FWorldParti
 			FCanvasLineItem Axis;
 			Axis.LineThickness = 3;
 			{
-				Axis.SetColor(FLinearColor::Green);
+				Axis.SetColor(FLinearColor::Red);
 				FVector2D LineStart = WorldToScreen(FVector2D(-1638400.f, 0.f));
 				FVector2D LineEnd = WorldToScreen(FVector2D(1638400.f, 0.f));
 				LineStart.X = FMath::Clamp(LineStart.X, GridScreenBounds.Min.X, GridScreenBounds.Max.X);
@@ -342,7 +342,7 @@ void FSpatialHashStreamingGrid::Draw2D(UCanvas* Canvas, const TArray<FWorldParti
 				Axis.Draw(CanvasObject, LineStart, LineEnd);
 			}
 			{
-				Axis.SetColor(FLinearColor::Red);
+				Axis.SetColor(FLinearColor::Green);
 				FVector2D LineStart = WorldToScreen(FVector2D(0.f, -1638400.f));
 				FVector2D LineEnd = WorldToScreen(FVector2D(0.f, 1638400.f));
 				LineStart.X = FMath::Clamp(LineStart.X, GridScreenBounds.Min.X, GridScreenBounds.Max.X);
