Add profile tag for UWorld::GetDuplicatedWorldForPIE.

diff --git a/Engine/Source/Runtime/Engine/Private/World.cpp b/Engine/Source/Runtime/Engine/Private/World.cpp
--- a/Engine/Source/Runtime/Engine/Private/World.cpp
+++ b/Engine/Source/Runtime/Engine/Private/World.cpp
@@ -3244,6 +3244,8 @@ bool UWorld::RemapCompiledScriptActor(FString& Str) const
 
 UWorld* UWorld::GetDuplicatedWorldForPIE(UWorld* InWorld, UPackage* InPIEackage, int32 PIEInstanceID)
 {
+	TRACE_CPUPROFILER_EVENT_SCOPE(UWorld::GetDuplicatedWorldForPIE);
+
 	check(PIEInstanceID != INDEX_NONE);
 
 	UPackage* InPackage = InWorld->GetOutermost();
