Remove unused FFixupSmartPointersForPIEArchive class.

diff --git a/Engine/Source/Runtime/Engine/Private/World.cpp b/Engine/Source/Runtime/Engine/Private/World.cpp
--- a/Engine/Source/Runtime/Engine/Private/World.cpp
+++ b/Engine/Source/Runtime/Engine/Private/World.cpp
@@ -3263,62 +3263,6 @@ UWorld* UWorld::GetDuplicatedWorldForPIE(UWorld* InWorld, UPackage* InPIEackage,
 	return DuplicatedWorld;
 }
 
-/**
- * Simple archive for updating lazy pointer GUIDs when a sub-level gets duplicated for PIE
- */
-class FFixupSmartPointersForPIEArchive : public FArchiveUObject
-{
-	/** Keeps track of objects that have already been serialized */
-	TSet<UObject*> VisitedObjects;
-
-public:
-
-	FFixupSmartPointersForPIEArchive()
-	{
-		ArIsObjectReferenceCollector = true;
-		ArIsModifyingWeakAndStrongReferences = true;
-		this->SetIsPersistent(false);
-		ArIgnoreArchetypeRef = true;
-	}
-	virtual FArchive& operator<<(FLazyObjectPtr& LazyObjectPtr) override
-	{
-		FArchive& Ar = *this;
-		FUniqueObjectGuid ID = LazyObjectPtr.GetUniqueID();
-		
-		// Remap unique ID if necessary
-		ID = ID.FixupForPIE();
-
-		LazyObjectPtr = ID;
-		return Ar;
-	}
-
-	virtual FArchive& operator<<(UObject*& Object) override
-	{
-		if (Object && (Object->IsA(UWorld::StaticClass()) || Object->IsInA(UWorld::StaticClass())) && !VisitedObjects.Contains(Object))
-		{
-			VisitedObjects.Add(Object);
-			UPackage* Outermost = Object->GetOutermost();
-			if (!Outermost || Outermost->HasAnyPackageFlags(PKG_PlayInEditor))
-			{
-				Object->Serialize(*this);
-			}
-		}
-		return *this;
-	}
-
-	virtual FArchive& operator<<(FSoftObjectPtr& Value) override
-	{
-		// Explicitly do nothing, we don't want to accidentally do PIE fixups
-		return *this;
-	}
-
-	virtual FArchive& operator<<(FSoftObjectPath& Value) override
-	{
-		// Explicitly do nothing, we don't want to accidentally do PIE fixups
-		return *this;
-	}
-};
-
 UWorld* UWorld::DuplicateWorldForPIE(const FString& PackageName, UWorld* OwningWorld)
 {
 #if WITH_EDITOR
