WorldPartition: Fix issue where the references of references of pinned actors were not being added to the world

diff --git a/Engine/Source/Runtime/Engine/Private/WorldPartition/ActorDescContainer.cpp b/Engine/Source/Runtime/Engine/Private/WorldPartition/ActorDescContainer.cpp
--- a/Engine/Source/Runtime/Engine/Private/WorldPartition/ActorDescContainer.cpp
+++ b/Engine/Source/Runtime/Engine/Private/WorldPartition/ActorDescContainer.cpp
@@ -243,6 +243,25 @@ void UActorDescContainer::PinActor(const FGuid& ActorGuid)
 	{
 		return;
 	}
+
+	TFunction<void(const FGuid&, TMap<FGuid, FWorldPartitionReference>&)> AddReferences = [this, &AddReferences](const FGuid& ActorGuid, TMap<FGuid, FWorldPartitionReference>& ReferenceMap)
+	{
+		if (ReferenceMap.Contains(ActorGuid))
+		{
+			return;
+		}
+
+		if (TUniquePtr<FWorldPartitionActorDesc>* const ActorDescPtr = GetActorDescriptor(ActorGuid); ActorDescPtr != nullptr && ActorDescPtr->IsValid())
+		{
+			ReferenceMap.Emplace(ActorGuid, ActorDescPtr);
+			
+			const FWorldPartitionActorDesc* const ActorDesc = ActorDescPtr->Get(); 
+			for (const FGuid& ReferencedActor : ActorDesc->GetReferences())
+			{
+				AddReferences(ReferencedActor, ReferenceMap);
+			}
+		}
+	};
 	
 	if (TUniquePtr<FWorldPartitionActorDesc>* const ActorDescPtr = GetActorDescriptor(ActorGuid))
 	{
@@ -250,12 +269,12 @@ void UActorDescContainer::PinActor(const FGuid& ActorGuid)
 		{
 			PinnedActors.Emplace(ActorGuid, ActorDescPtr);
 
-			// If the pinned actor has references, we must also create references to those to ensure they are loaded
-			TArray<FWorldPartitionReference>& References = PinnedActorRefs.Emplace(ActorGuid);
+			// If the pinned actor has references, we must also create references to those to ensure they are added
+			TMap<FGuid, FWorldPartitionReference>& References = PinnedActorRefs.Emplace(ActorGuid);
 
-			for (const FGuid& ReferencedGuid : ActorDesc->GetReferences())
+			for (const FGuid& ReferencedActor : ActorDesc->GetReferences())
 			{
-				References.Emplace(GetActorDescriptor(ReferencedGuid));
+				AddReferences(ReferencedActor, References);
 			}
 		}
 	}
diff --git a/Engine/Source/Runtime/Engine/Public/WorldPartition/ActorDescContainer.h b/Engine/Source/Runtime/Engine/Public/WorldPartition/ActorDescContainer.h
--- a/Engine/Source/Runtime/Engine/Public/WorldPartition/ActorDescContainer.h
+++ b/Engine/Source/Runtime/Engine/Public/WorldPartition/ActorDescContainer.h
@@ -65,7 +65,7 @@ protected:
 	bool ShouldHandleActorEvent(const AActor* Actor);
 
 	TMap<FGuid, FWorldPartitionReference> PinnedActors;
-	TMap<FGuid, TArray<FWorldPartitionReference>> PinnedActorRefs;
+	TMap<FGuid, TMap<FGuid, FWorldPartitionReference>> PinnedActorRefs;
 	
 	bool bContainerInitialized;
 
