WorldPartition HLOD Builder - Better gathering of in-use HLOD layers
Avoid going through the Asset Registry, as we now have the HLOD layer available in the actor desc, without having to load actors

diff --git a/Engine/Source/Runtime/Engine/Private/WorldPartition/RuntimeSpatialHash/RuntimeSpatialHashHLOD.cpp b/Engine/Source/Runtime/Engine/Private/WorldPartition/RuntimeSpatialHash/RuntimeSpatialHashHLOD.cpp
--- a/Engine/Source/Runtime/Engine/Private/WorldPartition/RuntimeSpatialHash/RuntimeSpatialHashHLOD.cpp
+++ b/Engine/Source/Runtime/Engine/Private/WorldPartition/RuntimeSpatialHash/RuntimeSpatialHashHLOD.cpp
@@ -24,9 +24,6 @@
 #include "MeshCardRepresentation.h"
 #include "StaticMeshCompiler.h"
 
-#include "AssetRegistryModule.h"
-#include "AssetData.h"
-
 
 DEFINE_LOG_CATEGORY_STATIC(LogWorldPartitionRuntimeSpatialHashHLOD, Log, All);
 
@@ -255,38 +252,55 @@ static TArray<FGuid> GenerateHLODsForGrid(UWorldPartition* WorldPartition, const
 	return GridHLODActors;
 }
 
-// Find all HLOD grids from the HLODLayer assets we have and build a dependency graph
-static void GatherHLODGrids(TMap<FName, FSpatialHashRuntimeGrid>& OutHLODGrids, TMap<FName, TSet<FName>>& OutHLODGridsGraph)
+// Find all HLOD grids from the HLODLayer assets we use and build a dependency graph
+static void GatherHLODGrids(UWorldPartition* WorldPartition, TMap<FName, FSpatialHashRuntimeGrid>& OutHLODGrids, TMap<FName, TSet<FName>>& OutHLODGridsGraph)
 {
 	// We don't know the HLOD level for those grids yet, we'll have to compute a graph of dependencies between our HLOD layers first
 	// Start with level 0, we'll rename once we have the correct info
 	const uint32 InitialHLODLevel = 0;
 
-	// Gather up all HLODLayer assets
-	FAssetRegistryModule& AssetRegistryModule = FModuleManager::LoadModuleChecked<FAssetRegistryModule>(TEXT("AssetRegistry"));
-	TArray<FAssetData> HLODLayerAssets;
-	AssetRegistryModule.Get().GetAssetsByClass(UHLODLayer::StaticClass()->GetFName(), HLODLayerAssets);
+	// Gather up all HLODLayer referenced by the actors
+	TSet<UHLODLayer*> HLODLayers;
 
-	// Find all used HLOD grids & build a dependency graph
-	for (const FAssetData& HLODLayerAsset : HLODLayerAssets)
+	for (UActorDescContainer::TIterator<> ActorDescIterator(WorldPartition); ActorDescIterator; ++ActorDescIterator)
 	{
-		if (const UHLODLayer* HLODLayer = Cast<UHLODLayer>(HLODLayerAsset.GetAsset()))
+		if (!ActorDescIterator->GetActorClass()->IsChildOf<AWorldPartitionHLOD>())
 		{
-			FSpatialHashRuntimeGrid HLODGrid;
-			HLODGrid.CellSize = HLODLayer->GetCellSize();
-			HLODGrid.LoadingRange = HLODLayer->GetLoadingRange();
-			HLODGrid.DebugColor = FLinearColor::Red;
-			HLODGrid.GridName = HLODLayer->GetRuntimeGrid(InitialHLODLevel);
-
-			OutHLODGrids.Emplace(HLODGrid.GridName, HLODGrid);
-
-			TSet<FName>& HLODParentGrids = OutHLODGridsGraph.FindOrAdd(HLODGrid.GridName);
-			if (const UHLODLayer* ParentHLODLayer = HLODLayer->GetParentLayer().LoadSynchronous())
+			if (ActorDescIterator->GetActorIsHLODRelevant())
 			{
-				HLODParentGrids.Add(ParentHLODLayer->GetRuntimeGrid(InitialHLODLevel));
+				UHLODLayer* HLODLayer = ActorDescIterator->GetHLODLayer();
+				while (HLODLayer != nullptr)
+				{
+					bool bAlreadyInSet = false;
+					HLODLayers.Add(HLODLayer, &bAlreadyInSet);
+					if (bAlreadyInSet)
+					{
+						break;
+					}
+
+					HLODLayer = HLODLayer->GetParentLayer().LoadSynchronous();
+				}
 			}
 		}
 	}
+
+	// Find all used HLOD grids & build a dependency graph
+	for (const UHLODLayer* HLODLayer : HLODLayers)
+	{
+		FSpatialHashRuntimeGrid HLODGrid;
+		HLODGrid.CellSize = HLODLayer->GetCellSize();
+		HLODGrid.LoadingRange = HLODLayer->GetLoadingRange();
+		HLODGrid.DebugColor = FLinearColor::Red;
+		HLODGrid.GridName = HLODLayer->GetRuntimeGrid(InitialHLODLevel);
+
+		OutHLODGrids.Emplace(HLODGrid.GridName, HLODGrid);
+
+		TSet<FName>& HLODParentGrids = OutHLODGridsGraph.FindOrAdd(HLODGrid.GridName);
+		if (const UHLODLayer* ParentHLODLayer = HLODLayer->GetParentLayer().LoadSynchronous())
+		{
+			HLODParentGrids.Add(ParentHLODLayer->GetRuntimeGrid(InitialHLODLevel));
+		}
+	}
 }
 
 // Sort HLOD grids in the order they'll need to be processed for HLOD generation
@@ -446,10 +460,12 @@ bool UWorldPartitionRuntimeSpatialHash::GenerateHLOD(ISourceControlHelper* Sourc
 		return false;
 	}
 
+	UWorldPartition* WorldPartition = GetOuterUWorldPartition();
+
 	// Find all used HLOD grids & build a dependency graph
 	TMap<FName, FSpatialHashRuntimeGrid> HLODGrids;
 	TMap<FName, TSet<FName>>			 HLODGridsGraph;
-	GatherHLODGrids(HLODGrids, HLODGridsGraph);
+	GatherHLODGrids(WorldPartition, HLODGrids, HLODGridsGraph);
 
 	// Sort HLOD grids in the order they'll need to be processed
 	TArray<FName>	SortedGrids;		// HLOD Grids, sorted in the order they'll need to be processed for HLOD generation
@@ -473,8 +489,6 @@ bool UWorldPartitionRuntimeSpatialHash::GenerateHLOD(ISourceControlHelper* Sourc
 		GridsMapping.Add(Grid.GridName, i);
 	}
 
-	UWorldPartition* WorldPartition = GetOuterUWorldPartition();
-
 	// HLOD creation context
 	FHLODCreationContext Context;
 	TArray<FWorldPartitionHandle> InvalidHLODActors;
@@ -483,13 +497,15 @@ bool UWorldPartitionRuntimeSpatialHash::GenerateHLOD(ISourceControlHelper* Sourc
 	{
 		FWorldPartitionHandle HLODActorHandle(WorldPartition, HLODIterator->GetGuid());
 
-		if (HLODIterator->GetCellHash())
+		uint64 CellHash = HLODIterator->GetCellHash();
+		if (CellHash != 0)
 		{
-			Context.HLODActorDescs.Add(HLODIterator->GetCellHash(), HLODActorHandle);
+			check(!Context.HLODActorDescs.Contains(CellHash));
+			Context.HLODActorDescs.Emplace(CellHash, MoveTemp(HLODActorHandle));
 		}
 		else
 		{
-			InvalidHLODActors.Add(HLODActorHandle);
+			InvalidHLODActors.Emplace(MoveTemp(HLODActorHandle));
 		}
 	}
 
