Setting FSkelMeshRenderSection MaxBoneInfluence. This fixes some issues if we are coping from RenderData to LODModel after merging.

Moving some no-vertex related functions out of the vertex loop.

diff --git a/Engine/Source/Runtime/Engine/Private/SkeletalMeshMerge.cpp b/Engine/Source/Runtime/Engine/Private/SkeletalMeshMerge.cpp
--- a/Engine/Source/Runtime/Engine/Private/SkeletalMeshMerge.cpp
+++ b/Engine/Source/Runtime/Engine/Private/SkeletalMeshMerge.cpp
@@ -476,7 +476,6 @@ void FSkeletalMeshMerge::GenerateLODModel( int32 LODIdx )
 		// keep track of the current base vertex for this section in the merged vertex buffer
 		Section.BaseVertexIndex = MergedVertexBuffer.Num();
 
-
 		// find existing material index
 		check(MergeMesh->GetMaterials().Num() == MaterialIds.Num());
 		int32 MatIndex;
@@ -573,21 +572,37 @@ void FSkeletalMeshMerge::GenerateLODModel( int32 LODIdx )
 			Section.NumVertices += MergeSectionInfo.Section->NumVertices;
 
 			// update total number of vertices 
-			int32 NumTotalVertices = MergeSectionInfo.Section->NumVertices;
+			const int32 NumTotalVertices = MergeSectionInfo.Section->NumVertices;
 
 			// add the vertices from the original source mesh to the merged vertex buffer					
-			int32 MaxVertIdx = FMath::Min<int32>( 
+			const int32 MaxVertIdx = FMath::Min<int32>( 
 				MergeSectionInfo.Section->BaseVertexIndex + NumTotalVertices,
 				SrcLODData.StaticVertexBuffers.PositionVertexBuffer.GetNumVertices()
 				);
 
-			int32 MaxColorIdx = SrcLODData.StaticVertexBuffers.ColorVertexBuffer.GetNumVertices();
+			const int32 MaxColorIdx = SrcLODData.StaticVertexBuffers.ColorVertexBuffer.GetNumVertices();
 
-			// keep track of the current base vertex index before adding any new vertices
-			// this will be needed to remap the index buffer values to the new range
-			int32 CurrentBaseVertexIndex = MergedVertexBuffer.Num();
+			// update max number of influences
 			const uint32 MaxBoneInfluences = SrcLODData.GetSkinWeightVertexBuffer()->GetMaxBoneInfluences();
 			const bool bUse16BitBoneIndex = SrcLODData.GetSkinWeightVertexBuffer()->Use16BitBoneIndex();
+
+			SourceMaxBoneInfluences = FMath::Max(SourceMaxBoneInfluences, MaxBoneInfluences);
+			bSourceUse16BitBoneIndex |= bUse16BitBoneIndex;
+
+			// update RenderSection number of max influences
+			Section.MaxBoneInfluences = MaxBoneInfluences;
+			
+			// update total number of TexCoords
+			const uint32 LODNumTexCoords = SrcLODData.StaticVertexBuffers.StaticMeshVertexBuffer.GetNumTexCoords();
+			if (TotalNumUVs < LODNumTexCoords)
+			{
+				TotalNumUVs = LODNumTexCoords;
+			}
+
+			// keep track of the current base vertex index before adding any new vertices
+			// this will be needed to remap the index buffer values to the new range
+			const int32 CurrentBaseVertexIndex = MergedVertexBuffer.Num();
+			
 			for( int32 VertIdx=MergeSectionInfo.Section->BaseVertexIndex; VertIdx < MaxVertIdx; VertIdx++ )
 			{
 				// add the new vertex
@@ -596,8 +611,6 @@ void FSkeletalMeshMerge::GenerateLODModel( int32 LODIdx )
 
 				CopyVertexFromSource<VertexDataType>(DestVert, SrcLODData, VertIdx, MergeSectionInfo);
 
-				SourceMaxBoneInfluences = FMath::Max(SourceMaxBoneInfluences, MaxBoneInfluences);
-				bSourceUse16BitBoneIndex |= bUse16BitBoneIndex;
 				DestWeight = SrcLODData.GetSkinWeightVertexBuffer()->GetVertexSkinWeights(VertIdx);
 
 				// if the mesh uses vertex colors, copy the source color if possible or default to white
@@ -615,12 +628,6 @@ void FSkeletalMeshMerge::GenerateLODModel( int32 LODIdx )
 					}
 				}
 
-				uint32 LODNumTexCoords = SrcLODData.StaticVertexBuffers.StaticMeshVertexBuffer.GetNumTexCoords();
-				if( TotalNumUVs < LODNumTexCoords )
-				{
-					TotalNumUVs = LODNumTexCoords;
-				}
-
 				// remap the bone index used by this vertex to match the mergedbonemap 
 				for( uint32 Idx=0; Idx < MAX_TOTAL_INFLUENCES; Idx++ )
 				{
@@ -636,7 +643,7 @@ void FSkeletalMeshMerge::GenerateLODModel( int32 LODIdx )
 			Section.NumTriangles += MergeSectionInfo.Section->NumTriangles;
 
 			// add the indices from the original source mesh to the merged index buffer					
-			int32 MaxIndexIdx = FMath::Min<int32>( 
+			const int32 MaxIndexIdx = FMath::Min<int32>( 
 				MergeSectionInfo.Section->BaseIndex + MergeSectionInfo.Section->NumTriangles * 3, 
 				SrcLODData.MultiSizeIndexContainer.GetIndexBuffer()->Num()
 				);
