Directional lights now default to be atmosphere lights.

diff --git a/Engine/Source/Editor/EnvironmentLightingViewer/Private/SEnvironmentLightingViewer.cpp b/Engine/Source/Editor/EnvironmentLightingViewer/Private/SEnvironmentLightingViewer.cpp
--- a/Engine/Source/Editor/EnvironmentLightingViewer/Private/SEnvironmentLightingViewer.cpp
+++ b/Engine/Source/Editor/EnvironmentLightingViewer/Private/SEnvironmentLightingViewer.cpp
@@ -370,7 +370,7 @@ FReply SEnvironmentLightingViewer::OnButtonCreateAtmosphericLight(uint32 Index)
 	DirectionalLight->SetMobility(EComponentMobility::Movable);
 	DirectionalLight->SetActorRotation(FRotator(329, 346, -105));
 #if WITH_EDITORONLY_DATA
-	DirectionalLight->GetComponent()->bUsedAsAtmosphereSunLight = 1;
+	DirectionalLight->GetComponent()->bAtmosphereSunLight = 1;
 	DirectionalLight->GetComponent()->AtmosphereSunLightIndex = Index;
 	// The render proxy is create right after AddActor, so we need to mark the render state as dirty again to get the new values set on the render side too.
 	DirectionalLight->MarkComponentsRenderStateDirty();
diff --git a/Engine/Source/Runtime/Core/Public/UObject/UE5MainStreamObjectVersion.h b/Engine/Source/Runtime/Core/Public/UObject/UE5MainStreamObjectVersion.h
--- a/Engine/Source/Runtime/Core/Public/UObject/UE5MainStreamObjectVersion.h
+++ b/Engine/Source/Runtime/Core/Public/UObject/UE5MainStreamObjectVersion.h
@@ -124,6 +124,9 @@ struct CORE_API FUE5MainStreamObjectVersion
 		// Remove explicit decal blend mode
 		RemoveDecalBlendMode,
 
+		// Made directional lights be atmosphere lights by default
+		DirLightsAreAtmosphereLightsByDefault,
+
         // -----<new versions can be added above this line>-------------------------------------------------
 		VersionPlusOne,
 		LatestVersion = VersionPlusOne - 1
diff --git a/Engine/Source/Runtime/Engine/Classes/Components/DirectionalLightComponent.h b/Engine/Source/Runtime/Engine/Classes/Components/DirectionalLightComponent.h
--- a/Engine/Source/Runtime/Engine/Classes/Components/DirectionalLightComponent.h
+++ b/Engine/Source/Runtime/Engine/Classes/Components/DirectionalLightComponent.h
@@ -146,12 +146,15 @@ class ENGINE_API UDirectionalLightComponent : public ULightComponent
 	/** Determines how far shadows can be cast, in world units.  Larger values increase the shadowing cost. */
 	UPROPERTY(EditAnywhere, AdvancedDisplay, BlueprintReadOnly, Category=DistanceFieldShadows, meta=(UIMin = "1000", UIMax = "100000"), DisplayName = "DistanceField Trace Distance")
 	float TraceDistance;
-	
+
+	UPROPERTY()
+	uint32 bUsedAsAtmosphereSunLight_DEPRECATED : 1;
+
 	/**
 	 * Whether the directional light can interact with the atmosphere, cloud and generate a visual disk. All of which compose the visual sky.
 	 */
 	UPROPERTY(EditAnywhere, BlueprintReadOnly, Category= AtmosphereAndCloud, meta=(DisplayName = "Atmosphere Sun Light", ShouldShowInViewport = true))
-	uint32 bUsedAsAtmosphereSunLight : 1;
+	uint32 bAtmosphereSunLight : 1;
 
 	/**
 	 * Two atmosphere lights are supported. For instance: a sun and a moon, or two suns.
@@ -300,7 +303,7 @@ class ENGINE_API UDirectionalLightComponent : public ULightComponent
 	virtual FLightSceneProxy* CreateSceneProxy() const override;
 	virtual bool IsUsedAsAtmosphereSunLight() const override
 	{
-		return bUsedAsAtmosphereSunLight;
+		return bAtmosphereSunLight; 
 	}
 	virtual uint8 GetAtmosphereSunLightIndex() const override
 	{
diff --git a/Engine/Source/Runtime/Engine/Private/Components/DirectionalLightComponent.cpp b/Engine/Source/Runtime/Engine/Private/Components/DirectionalLightComponent.cpp
--- a/Engine/Source/Runtime/Engine/Private/Components/DirectionalLightComponent.cpp
+++ b/Engine/Source/Runtime/Engine/Private/Components/DirectionalLightComponent.cpp
@@ -12,6 +12,7 @@
 #include "Engine/Texture2D.h"
 #include "SceneManagement.h"
 #include "Engine/DirectionalLight.h"
+#include "UObject/UE5MainStreamObjectVersion.h"
 
 static float GMaxCSMRadiusToAllowPerObjectShadows = 8000;
 static FAutoConsoleVariableRef CVarMaxCSMRadiusToAllowPerObjectShadows(
@@ -991,7 +992,8 @@ UDirectionalLightComponent::UDirectionalLightComponent(const FObjectInitializer&
 	ModulatedShadowColor = FColor(128, 128, 128);
 	ShadowAmount = 1.0f;
 
-	bUsedAsAtmosphereSunLight = false;
+	bUsedAsAtmosphereSunLight_DEPRECATED = false;
+	bAtmosphereSunLight = true;
 	AtmosphereSunLightIndex = 0;
 	AtmosphereSunDiskColorScale = FLinearColor::White;
 
@@ -1106,7 +1108,7 @@ bool UDirectionalLightComponent::CanEditChange(const FProperty* InProperty) cons
 			|| PropertyName == GET_MEMBER_NAME_STRING_CHECKED(UDirectionalLightComponent, CloudShadowOnSurfaceStrength)
 			|| PropertyName == GET_MEMBER_NAME_STRING_CHECKED(UDirectionalLightComponent, CloudShadowDepthBias))
 		{
-			return bCastCloudShadows && bUsedAsAtmosphereSunLight;
+			return bCastCloudShadows && bAtmosphereSunLight;
 		}
 		if (PropertyName == GET_MEMBER_NAME_STRING_CHECKED(UDirectionalLightComponent, bCastCloudShadows)
 			|| PropertyName == GET_MEMBER_NAME_STRING_CHECKED(UDirectionalLightComponent, bCastShadowsOnAtmosphere)
@@ -1115,11 +1117,11 @@ bool UDirectionalLightComponent::CanEditChange(const FProperty* InProperty) cons
 			|| PropertyName == GET_MEMBER_NAME_STRING_CHECKED(UDirectionalLightComponent, AtmosphereSunLightIndex)
 			|| PropertyName == GET_MEMBER_NAME_STRING_CHECKED(UDirectionalLightComponent, AtmosphereSunDiskColorScale))
 		{
-			return bUsedAsAtmosphereSunLight;
+			return bAtmosphereSunLight;
 		}
 		if (PropertyName == GET_MEMBER_NAME_STRING_CHECKED(UDirectionalLightComponent, bCastShadowsOnClouds))
 		{
-			return bUsedAsAtmosphereSunLight && AtmosphereSunLightIndex == 0;	// AtmosphereLight1 has opaque shadow on cloud disabled.
+			return bAtmosphereSunLight && AtmosphereSunLightIndex == 0;	// AtmosphereLight1 has opaque shadow on cloud disabled.
 		}
 
 
@@ -1263,9 +1265,9 @@ void UDirectionalLightComponent::SetShadowAmount(float NewValue)
 void UDirectionalLightComponent::SetAtmosphereSunLight(bool bNewValue)
 {
 	if (AreDynamicDataChangesAllowed()
-		&& bUsedAsAtmosphereSunLight != bNewValue)
+		&& bAtmosphereSunLight != bNewValue)
 	{
-		bUsedAsAtmosphereSunLight = bNewValue;
+		bAtmosphereSunLight = bNewValue;
 		MarkRenderStateDirty();
 	}
 }
@@ -1284,6 +1286,8 @@ void UDirectionalLightComponent::Serialize(FArchive& Ar)
 {
 	Super::Serialize(Ar);
 
+	Ar.UsingCustomVersion(FUE5MainStreamObjectVersion::GUID);
+
 	if(Ar.UEVer() < VER_UE4_REMOVE_LIGHT_MOBILITY_CLASSES)
 	{
 		// If outer is a DirectionalLight, we use the ADirectionalLight::LoadedFromAnotherClass path
@@ -1299,6 +1303,11 @@ void UDirectionalLightComponent::Serialize(FArchive& Ar)
 			}
 		}
 	}
+
+	if ((Ar.CustomVer(FUE5MainStreamObjectVersion::GUID) < FUE5MainStreamObjectVersion::DirLightsAreAtmosphereLightsByDefault))
+	{
+		bAtmosphereSunLight = bUsedAsAtmosphereSunLight_DEPRECATED;
+	}
 }
 
 void UDirectionalLightComponent::InvalidateLightingCacheDetailed(bool bInvalidateBuildEnqueuedLighting, bool bTranslationOnly)
