Add initial facility to convert skeletal mesh LOD import data to mesh description and back.
Add a fallback function to convert from the LOD model directly, in case the bulk data is not available (e.g. for skelmeshes from UE4.24 and eariler).

diff --git a/Engine/Plugins/Runtime/Interchange/Source/Import/Private/Mesh/InterchangeSkeletalMeshFactory.cpp b/Engine/Plugins/Runtime/Interchange/Source/Import/Private/Mesh/InterchangeSkeletalMeshFactory.cpp
--- a/Engine/Plugins/Runtime/Interchange/Source/Import/Private/Mesh/InterchangeSkeletalMeshFactory.cpp
+++ b/Engine/Plugins/Runtime/Interchange/Source/Import/Private/Mesh/InterchangeSkeletalMeshFactory.cpp
@@ -135,436 +135,6 @@ namespace UE
 				return true;
 			}
 
-			void CleanUpUnusedMaterials(FSkeletalMeshImportData& ImportData)
-			{
-				if (ImportData.Materials.Num() <= 0)
-				{
-					return;
-				}
-
-				TArray< SkeletalMeshImportData::FMaterial > ExistingMatList = ImportData.Materials;
-
-				TArray<uint8> UsedMaterialIndex;
-				// Find all material that are use by the mesh faces
-				int32 FaceNum = ImportData.Faces.Num();
-				for (int32 TriangleIndex = 0; TriangleIndex < FaceNum; TriangleIndex++)
-				{
-					SkeletalMeshImportData::FTriangle& Triangle = ImportData.Faces[TriangleIndex];
-					UsedMaterialIndex.AddUnique(Triangle.MatIndex);
-				}
-				//Remove any unused material.
-				if (UsedMaterialIndex.Num() < ExistingMatList.Num())
-				{
-					TArray<int32> RemapIndex;
-					TArray< SkeletalMeshImportData::FMaterial >& NewMatList = ImportData.Materials;
-					NewMatList.Empty();
-					for (int32 ExistingMatIndex = 0; ExistingMatIndex < ExistingMatList.Num(); ++ExistingMatIndex)
-					{
-						if (UsedMaterialIndex.Contains((uint8)ExistingMatIndex))
-						{
-							RemapIndex.Add(NewMatList.Add(ExistingMatList[ExistingMatIndex]));
-						}
-						else
-						{
-							RemapIndex.Add(INDEX_NONE);
-						}
-					}
-					ImportData.MaxMaterialIndex = 0;
-					//Remap the face material index
-					for (int32 TriangleIndex = 0; TriangleIndex < FaceNum; TriangleIndex++)
-					{
-						SkeletalMeshImportData::FTriangle& Triangle = ImportData.Faces[TriangleIndex];
-						check(RemapIndex[Triangle.MatIndex] != INDEX_NONE);
-						Triangle.MatIndex = RemapIndex[Triangle.MatIndex];
-						ImportData.MaxMaterialIndex = FMath::Max<uint32>(ImportData.MaxMaterialIndex, Triangle.MatIndex);
-					}
-				}
-			}
-
-			namespace SmoothGroupHelper
-			{
-				struct tFaceRecord
-				{
-					int32 FaceIndex;
-					int32 HoekIndex;
-					int32 WedgeIndex;
-					uint32 SmoothFlags;
-					uint32 FanFlags;
-				};
-
-				struct VertsFans
-				{
-					TArray<tFaceRecord> FaceRecord;
-					int32 FanGroupCount;
-				};
-
-				struct tInfluences
-				{
-					TArray<int32> RawInfIndices;
-				};
-
-				struct tWedgeList
-				{
-					TArray<int32> WedgeList;
-				};
-
-				struct tFaceSet
-				{
-					TArray<int32> Faces;
-				};
-
-				// Check whether faces have at least two vertices in common. These must be POINTS - don't care about wedges.
-				bool FacesAreSmoothlyConnected(FSkeletalMeshImportData& ImportData, int32 Face1, int32 Face2)
-				{
-
-					//if( ( Face1 >= Thing->SkinData.Faces.Num()) || ( Face2 >= Thing->SkinData.Faces.Num()) ) return false;
-
-					if (Face1 == Face2)
-					{
-						return true;
-					}
-
-					// Smoothing groups match at least one bit in binary AND ?
-					if ((ImportData.Faces[Face1].SmoothingGroups & ImportData.Faces[Face2].SmoothingGroups) == 0)
-					{
-						return false;
-					}
-
-					int32 VertMatches = 0;
-					for (int32 i = 0; i < 3; i++)
-					{
-						int32 Point1 = ImportData.Wedges[ImportData.Faces[Face1].WedgeIndex[i]].VertexIndex;
-
-						for (int32 j = 0; j < 3; j++)
-						{
-							int32 Point2 = ImportData.Wedges[ImportData.Faces[Face2].WedgeIndex[j]].VertexIndex;
-							if (Point2 == Point1)
-							{
-								VertMatches++;
-							}
-						}
-					}
-
-					return (VertMatches >= 2);
-				}
-			} //namespace SmoothGroupHelper
-			int32 DoUnSmoothVerts(FSkeletalMeshImportData& ImportData)
-			{
-				//
-				// Connectivity: triangles with non-matching smoothing groups will be physically split.
-				//
-				// -> Splitting involves: the UV+material-contaning vertex AND the 3d point.
-				//
-				// -> Tally smoothing groups for each and every (textured) vertex.
-				//
-				// -> Collapse: 
-				// -> start from a vertex and all its adjacent triangles - go over
-				// each triangle - if any connecting one (sharing more than one vertex) gives a smoothing match,
-				// accumulate it. Then IF more than one resulting section, 
-				// ensure each boundary 'vert' is split _if not already_ to give each smoothing group
-				// independence from all others.
-				//
-
-				int32 DuplicatedVertCount = 0;
-				int32 RemappedHoeks = 0;
-
-				int32 TotalSmoothMatches = 0;
-				int32 TotalConnexChex = 0;
-
-				// Link _all_ faces to vertices.	
-				TArray<SmoothGroupHelper::VertsFans>  Fans;
-				TArray<SmoothGroupHelper::tInfluences> PointInfluences;
-				TArray<SmoothGroupHelper::tWedgeList>  PointWedges;
-
-				Fans.AddZeroed(ImportData.Points.Num());//Fans.AddExactZeroed(			Thing->SkinData.Points.Num() );
-				PointInfluences.AddZeroed(ImportData.Points.Num());//PointInfluences.AddExactZeroed( Thing->SkinData.Points.Num() );
-				PointWedges.AddZeroed(ImportData.Points.Num());//PointWedges.AddExactZeroed(	 Thing->SkinData.Points.Num() );
-
-				// Existing points map 1:1
-				ImportData.PointToRawMap.AddUninitialized(ImportData.Points.Num());
-				for (int32 i = 0; i < ImportData.Points.Num(); i++)
-				{
-					ImportData.PointToRawMap[i] = i;
-				}
-
-				for (int32 i = 0; i < ImportData.Influences.Num(); i++)
-				{
-					if (PointInfluences.Num() <= ImportData.Influences[i].VertexIndex)
-					{
-						PointInfluences.AddZeroed(ImportData.Influences[i].VertexIndex - PointInfluences.Num() + 1);
-					}
-					PointInfluences[ImportData.Influences[i].VertexIndex].RawInfIndices.Add(i);
-				}
-
-				for (int32 i = 0; i < ImportData.Wedges.Num(); i++)
-				{
-					if (uint32(PointWedges.Num()) <= ImportData.Wedges[i].VertexIndex)
-					{
-						PointWedges.AddZeroed(ImportData.Wedges[i].VertexIndex - PointWedges.Num() + 1);
-					}
-
-					PointWedges[ImportData.Wedges[i].VertexIndex].WedgeList.Add(i);
-				}
-
-				for (int32 f = 0; f < ImportData.Faces.Num(); f++)
-				{
-					// For each face, add a pointer to that face into the Fans[vertex].
-					for (int32 i = 0; i < 3; i++)
-					{
-						int32 WedgeIndex = ImportData.Faces[f].WedgeIndex[i];
-						int32 PointIndex = ImportData.Wedges[WedgeIndex].VertexIndex;
-						SmoothGroupHelper::tFaceRecord NewFR;
-
-						NewFR.FaceIndex = f;
-						NewFR.HoekIndex = i;
-						NewFR.WedgeIndex = WedgeIndex; // This face touches the point courtesy of Wedges[Wedgeindex].
-						NewFR.SmoothFlags = ImportData.Faces[f].SmoothingGroups;
-						NewFR.FanFlags = 0;
-						Fans[PointIndex].FaceRecord.Add(NewFR);
-						Fans[PointIndex].FanGroupCount = 0;
-					}
-				}
-
-				// Investigate connectivity and assign common group numbers (1..+) to the fans' individual FanFlags.
-				for (int32 p = 0; p < Fans.Num(); p++) // The fan of faces for each 3d point 'p'.
-				{
-					// All faces connecting.
-					if (Fans[p].FaceRecord.Num() > 0)
-					{
-						int32 FacesProcessed = 0;
-						TArray<SmoothGroupHelper::tFaceSet> FaceSets; // Sets with indices INTO FANS, not into face array.			
-
-						// Digest all faces connected to this vertex (p) into one or more smooth sets. only need to check 
-						// all faces MINUS one..
-						while (FacesProcessed < Fans[p].FaceRecord.Num())
-						{
-							// One loop per group. For the current ThisFaceIndex, tally all truly connected ones
-							// and put them in a new TArray. Once no more can be connected, stop.
-
-							int32 NewSetIndex = FaceSets.Num(); // 0 to start
-							FaceSets.AddZeroed(1);						// first one will be just ThisFaceIndex.
-
-							// Find the first non-processed face. There will be at least one.
-							int32 ThisFaceFanIndex = 0;
-							{
-								int32 SearchIndex = 0;
-								while (Fans[p].FaceRecord[SearchIndex].FanFlags == -1) // -1 indicates already  processed. 
-								{
-									SearchIndex++;
-								}
-								ThisFaceFanIndex = SearchIndex; //Fans[p].FaceRecord[SearchIndex].FaceIndex; 
-							}
-
-							// Initial face.
-							FaceSets[NewSetIndex].Faces.Add(ThisFaceFanIndex);   // Add the unprocessed Face index to the "local smoothing group" [NewSetIndex].
-							Fans[p].FaceRecord[ThisFaceFanIndex].FanFlags = -1;			  // Mark as processed.
-							FacesProcessed++;
-
-							// Find all faces connected to this face, and if there's any
-							// smoothing group matches, put it in current face set and mark it as processed;
-							// until no more match. 
-							int32 NewMatches = 0;
-							do
-							{
-								NewMatches = 0;
-								// Go over all current faces in this faceset and set if the FaceRecord (local smoothing groups) has any matches.
-								// there will be at least one face already in this faceset - the first face in the fan.
-								for (int32 n = 0; n < FaceSets[NewSetIndex].Faces.Num(); n++)
-								{
-									int32 HookFaceIdx = Fans[p].FaceRecord[FaceSets[NewSetIndex].Faces[n]].FaceIndex;
-
-									//Go over the fan looking for matches.
-									for (int32 s = 0; s < Fans[p].FaceRecord.Num(); s++)
-									{
-										// Skip if same face, skip if face already processed.
-										if ((HookFaceIdx != Fans[p].FaceRecord[s].FaceIndex) && (Fans[p].FaceRecord[s].FanFlags != -1))
-										{
-											TotalConnexChex++;
-											// Process if connected with more than one vertex, AND smooth..
-											if (SmoothGroupHelper::FacesAreSmoothlyConnected(ImportData, HookFaceIdx, Fans[p].FaceRecord[s].FaceIndex))
-											{
-												TotalSmoothMatches++;
-												Fans[p].FaceRecord[s].FanFlags = -1; // Mark as processed.
-												FacesProcessed++;
-												// Add 
-												FaceSets[NewSetIndex].Faces.Add(s); // Store FAN index of this face index into smoothing group's faces. 
-												// Tally
-												NewMatches++;
-											}
-										} // not the same...
-									}// all faces in fan
-								} // all faces in FaceSet
-							} while (NewMatches);
-
-						}// Repeat until all faces processed.
-
-						// For the new non-initialized  face sets, 
-						// Create a new point, influences, and uv-vertex(-ices) for all individual FanFlag groups with an index of 2+ and also remap
-						// the face's vertex into those new ones.
-						if (FaceSets.Num() > 1)
-						{
-							for (int32 f = 1; f < FaceSets.Num(); f++)
-							{
-								check(ImportData.Points.Num() == ImportData.PointToRawMap.Num());
-
-								// We duplicate the current vertex. (3d point)
-								int32 NewPointIndex = ImportData.Points.Num();
-								ImportData.Points.AddUninitialized();
-								ImportData.Points[NewPointIndex] = ImportData.Points[p];
-
-								ImportData.PointToRawMap.AddUninitialized();
-								ImportData.PointToRawMap[NewPointIndex] = p;
-
-								DuplicatedVertCount++;
-
-								// Duplicate all related weights.
-								for (int32 t = 0; t < PointInfluences[p].RawInfIndices.Num(); t++)
-								{
-									// Add new weight
-									int32 NewWeightIndex = ImportData.Influences.Num();
-									ImportData.Influences.AddUninitialized();
-									ImportData.Influences[NewWeightIndex] = ImportData.Influences[PointInfluences[p].RawInfIndices[t]];
-									ImportData.Influences[NewWeightIndex].VertexIndex = NewPointIndex;
-								}
-
-								// Duplicate any and all Wedges associated with it; and all Faces' wedges involved.					
-								for (int32 w = 0; w < PointWedges[p].WedgeList.Num(); w++)
-								{
-									int32 OldWedgeIndex = PointWedges[p].WedgeList[w];
-									int32 NewWedgeIndex = ImportData.Wedges.Num();
-									ImportData.Wedges[OldWedgeIndex].VertexIndex = NewPointIndex;
-								}
-							}
-						} //  if FaceSets.Num(). -> duplicate stuff
-					}//	while( FacesProcessed < Fans[p].FaceRecord.Num() )
-				} // Fans for each 3d point
-
-				if (!ensure(ImportData.Points.Num() == ImportData.PointToRawMap.Num()))
-				{
-					//TODO log a warning to the user
-
-					//Create a valid PointtoRawMap but with bad content
-					int32 PointNum = ImportData.Points.Num();
-					ImportData.PointToRawMap.Empty(PointNum);
-					for (int32 PointIndex = 0; PointIndex < PointNum; ++PointIndex)
-					{
-						ImportData.PointToRawMap[PointIndex] = PointIndex;
-					}
-				}
-
-				return DuplicatedVertCount;
-			}
-
-			/**
-			 * Make sure the SourceMeshDescription is compact before calling this function, we assume all element ID are continuous so we can use them as array index
-			 */
-			bool CopyMeshDescriptionToSkeletalMeshImportData(const FMeshDescription& SourceMeshDescription, FSkeletalMeshImportData& DestinationSkeletalMeshImportData)
-			{
-				DestinationSkeletalMeshImportData.Empty();
-
-				FSkeletalMeshConstAttributes Attributes(SourceMeshDescription);
-				TVertexAttributesConstRef<FVector> VertexPositions = Attributes.GetVertexPositions();
-				FSkinWeightsVertexAttributesConstRef VertexSkinWeights = Attributes.GetVertexSkinWeights();
-
-				TVertexInstanceAttributesConstRef<FVertexID> VertexInstanceVertexIndices = Attributes.GetVertexInstanceVertexIndices();
-				TVertexInstanceAttributesConstRef<FVector2D> VertexInstanceUVs = Attributes.GetVertexInstanceUVs();
-				TVertexInstanceAttributesConstRef<FVector> VertexInstanceNormals = Attributes.GetVertexInstanceNormals();
-				TVertexInstanceAttributesConstRef<FVector> VertexInstanceTangents = Attributes.GetVertexInstanceTangents();
-				TVertexInstanceAttributesConstRef<float> VertexInstanceBiNormalSigns = Attributes.GetVertexInstanceBinormalSigns();
-				TVertexInstanceAttributesConstRef<FVector4> VertexInstanceColors = Attributes.GetVertexInstanceColors();
-
-				TPolygonGroupAttributesConstRef<FName> PolygonGroupMaterialSlotNames = Attributes.GetPolygonGroupMaterialSlotNames();
-				//Get the per face smoothing
-				TArray<uint32> FaceSmoothingMasks;
-				FaceSmoothingMasks.AddZeroed(SourceMeshDescription.Triangles().Num());
-				FSkeletalMeshOperations::ConvertHardEdgesToSmoothGroup(SourceMeshDescription, FaceSmoothingMasks);
-
-				//////////////////////////////////////////////////////////////////////////
-				// Copy the materials
-				DestinationSkeletalMeshImportData.Materials.Reserve(SourceMeshDescription.PolygonGroups().Num());
-				for (FPolygonGroupID PolygonGroupID : SourceMeshDescription.PolygonGroups().GetElementIDs())
-				{
-					SkeletalMeshImportData::FMaterial Material;
-					Material.MaterialImportName = PolygonGroupMaterialSlotNames[PolygonGroupID].ToString();
-					//The material interface will be added later by the factory
-					DestinationSkeletalMeshImportData.Materials.Add(Material);
-				}
-				DestinationSkeletalMeshImportData.MaxMaterialIndex = DestinationSkeletalMeshImportData.Materials.Num()-1;
-
-				//////////////////////////////////////////////////////////////////////////
-				//Copy the vertex positions and the influences
-
-				//Reserve the point and influences
-				DestinationSkeletalMeshImportData.Points.AddZeroed(SourceMeshDescription.Vertices().Num());
-				DestinationSkeletalMeshImportData.Influences.Reserve(SourceMeshDescription.Vertices().Num() * 4);
-				
-				for (FVertexID VertexID : SourceMeshDescription.Vertices().GetElementIDs())
-				{
-					//We can use GetValue because the Meshdescription was compacted before the copy
-					DestinationSkeletalMeshImportData.Points[VertexID.GetValue()] = VertexPositions[VertexID];
-
-					FVertexBoneWeightsConst BoneWeights = VertexSkinWeights.Get(VertexID); 
-					const int32 InfluenceCount = BoneWeights.Num();
-					
-					const int32 InfluenceOffsetIndex = DestinationSkeletalMeshImportData.Influences.Num();
-					DestinationSkeletalMeshImportData.Influences.AddDefaulted(InfluenceCount);
-					for (int32 InfluenceIndex = 0; InfluenceIndex < InfluenceCount; ++InfluenceIndex)
-					{
-						SkeletalMeshImportData::FRawBoneInfluence& BoneInfluence = DestinationSkeletalMeshImportData.Influences[InfluenceOffsetIndex + InfluenceIndex];
-						BoneInfluence.VertexIndex = VertexID.GetValue();
-						BoneInfluence.BoneIndex = BoneWeights[InfluenceIndex].GetBoneIndex();
-						BoneInfluence.Weight = BoneWeights[InfluenceIndex].GetWeight();
-					}
-				}
-
-				//////////////////////////////////////////////////////////////////////////
-				//Copy the triangle and vertex instances
-				DestinationSkeletalMeshImportData.Faces.AddZeroed(SourceMeshDescription.Triangles().Num());
-				DestinationSkeletalMeshImportData.Wedges.Reserve(SourceMeshDescription.VertexInstances().Num());
-				DestinationSkeletalMeshImportData.NumTexCoords = FMath::Min<int32>(VertexInstanceUVs.GetNumChannels(), (int32)MAX_TEXCOORDS);
-				for (FTriangleID TriangleID : SourceMeshDescription.Triangles().GetElementIDs())
-				{
-					FPolygonGroupID PolygonGroupID = SourceMeshDescription.GetTrianglePolygonGroup(TriangleID);
-					TArrayView<const FVertexInstanceID> VertexInstances = SourceMeshDescription.GetTriangleVertexInstances(TriangleID);
-					int32 FaceIndex = TriangleID.GetValue();
-					if (!ensure(DestinationSkeletalMeshImportData.Faces.IsValidIndex(FaceIndex)))
-					{
-						//TODO log an error for the user
-						break;
-					}
-					SkeletalMeshImportData::FTriangle& Face = DestinationSkeletalMeshImportData.Faces[FaceIndex];
-					Face.MatIndex = PolygonGroupID.GetValue();
-					Face.SmoothingGroups = 0;
-					if (FaceSmoothingMasks.IsValidIndex(FaceIndex))
-					{
-						Face.SmoothingGroups = FaceSmoothingMasks[FaceIndex];
-					}
-					//Create the wedges
-					for (int32 Corner = 0; Corner < 3; ++Corner)
-					{
-						FVertexInstanceID VertexInstanceID = VertexInstances[Corner];
-						SkeletalMeshImportData::FVertex Wedge;
-						Wedge.VertexIndex = (uint32)SourceMeshDescription.GetVertexInstanceVertex(VertexInstances[Corner]).GetValue();
-						Wedge.MatIndex = Face.MatIndex;
-						const bool bSRGB = false; //avoid linear to srgb conversion
-						Wedge.Color = FLinearColor(VertexInstanceColors[VertexInstanceID]).ToFColor(bSRGB);
-						for (int32 UVChannelIndex = 0; UVChannelIndex < (int32)(DestinationSkeletalMeshImportData.NumTexCoords); ++UVChannelIndex)
-						{
-							Wedge.UVs[UVChannelIndex] = VertexInstanceUVs.Get(VertexInstanceID, UVChannelIndex);
-						}
-						Face.TangentX[Corner] = VertexInstanceTangents[VertexInstanceID];
-						Face.TangentZ[Corner] = VertexInstanceNormals[VertexInstanceID];
-						Face.TangentY[Corner] = FVector::CrossProduct(VertexInstanceNormals[VertexInstanceID], VertexInstanceTangents[VertexInstanceID]).GetSafeNormal() * VertexInstanceBiNormalSigns[VertexInstanceID];
-
-						Face.WedgeIndex[Corner] = DestinationSkeletalMeshImportData.Wedges.Add(Wedge);
-					}
-				}
-
-				CleanUpUnusedMaterials(DestinationSkeletalMeshImportData);
-				DoUnSmoothVerts(DestinationSkeletalMeshImportData);
-
-				return true;
-			}
 
 			bool CopyBlendShapesMeshDescriptionToSkeletalMeshImportData(const TArray<FString>& BlendShapeToImport, const TMap<FString, FMeshDescription>& LodBlendShapeMeshDescriptions, FSkeletalMeshImportData& DestinationSkeletalMeshImportData)
 			{
@@ -1137,7 +707,8 @@ UObject* UInterchangeSkeletalMeshFactory::CreateAsset(const UInterchangeSkeletal
 					*/
 					
 				}
-				UE::Interchange::Private::CopyMeshDescriptionToSkeletalMeshImportData(LodMeshDescription, SkeletalMeshImportData);
+
+				SkeletalMeshImportData = FSkeletalMeshImportData::CreateFromMeshDescription(LodMeshDescription);
 				SkeletalMeshImportData.RefBonesBinary = RefBonesBinary;
 /*
 
diff --git a/Engine/Source/Runtime/Engine/Engine.Build.cs b/Engine/Source/Runtime/Engine/Engine.Build.cs
--- a/Engine/Source/Runtime/Engine/Engine.Build.cs
+++ b/Engine/Source/Runtime/Engine/Engine.Build.cs
@@ -89,6 +89,7 @@ public class Engine : ModuleRules
 				"AudioPlatformConfiguration",
 				"MeshDescription",
 				"StaticMeshDescription",
+				"SkeletalMeshDescription",
 				"PakFile",
 				"NetworkReplayStreaming",
 				"PhysicsCore",
diff --git a/Engine/Source/Runtime/Engine/Private/Rendering/SkeletalMeshLODImporterData.cpp b/Engine/Source/Runtime/Engine/Private/Rendering/SkeletalMeshLODImporterData.cpp
--- a/Engine/Source/Runtime/Engine/Private/Rendering/SkeletalMeshLODImporterData.cpp
+++ b/Engine/Source/Runtime/Engine/Private/Rendering/SkeletalMeshLODImporterData.cpp
@@ -11,6 +11,9 @@
 #include "Engine/SkeletalMesh.h"
 #include "Factories/FbxSkeletalMeshImportData.h"
 #include "Misc/ScopedSlowTask.h"
+#include "SkeletalMeshAttributes.h"
+#include "SkeletalMeshOperations.h"
+
 
 DEFINE_LOG_CATEGORY_STATIC(LogSkeletalMeshLODImporterData, Log, All);
 
@@ -1260,4 +1263,597 @@ void FSkeletalMeshImportData::ComputeSmoothGroupFromNormals()
 	}
 }
 
-#endif // WITH_EDITOR
\ No newline at end of file
+void FSkeletalMeshImportData::CleanUpUnusedMaterials()
+{
+	if (Materials.Num() <= 0)
+	{
+		return;
+	}
+
+	TArray<SkeletalMeshImportData::FMaterial> ExistingMatList = Materials;
+
+	TArray<uint8> UsedMaterialIndex;
+	// Find all material that are use by the mesh faces
+	int32 FaceNum = Faces.Num();
+	for (int32 TriangleIndex = 0; TriangleIndex < FaceNum; TriangleIndex++)
+	{
+		SkeletalMeshImportData::FTriangle& Triangle = Faces[TriangleIndex];
+		UsedMaterialIndex.AddUnique(Triangle.MatIndex);
+	}
+	//Remove any unused material.
+	if (UsedMaterialIndex.Num() < ExistingMatList.Num())
+	{
+		TArray<int32> RemapIndex;
+		TArray< SkeletalMeshImportData::FMaterial >& NewMatList = Materials;
+		NewMatList.Empty();
+		for (int32 ExistingMatIndex = 0; ExistingMatIndex < ExistingMatList.Num(); ++ExistingMatIndex)
+		{
+			if (UsedMaterialIndex.Contains((uint8)ExistingMatIndex))
+			{
+				RemapIndex.Add(NewMatList.Add(ExistingMatList[ExistingMatIndex]));
+			}
+			else
+			{
+				RemapIndex.Add(INDEX_NONE);
+			}
+		}
+		MaxMaterialIndex = 0;
+		//Remap the face material index
+		for (int32 TriangleIndex = 0; TriangleIndex < FaceNum; TriangleIndex++)
+		{
+			SkeletalMeshImportData::FTriangle& Triangle = Faces[TriangleIndex];
+			check(RemapIndex[Triangle.MatIndex] != INDEX_NONE);
+			Triangle.MatIndex = RemapIndex[Triangle.MatIndex];
+			MaxMaterialIndex = FMath::Max<uint32>(MaxMaterialIndex, Triangle.MatIndex);
+		}
+	}
+}
+
+namespace SmoothGroupHelper
+{
+	struct tFaceRecord
+	{
+		int32 FaceIndex;
+		int32 HoekIndex;
+		int32 WedgeIndex;
+		uint32 SmoothFlags;
+		uint32 FanFlags;
+	};
+
+	struct VertsFans
+	{
+		TArray<tFaceRecord> FaceRecord;
+		int32 FanGroupCount;
+	};
+
+	struct tInfluences
+	{
+		TArray<int32> RawInfIndices;
+	};
+
+	struct tWedgeList
+	{
+		TArray<int32> WedgeList;
+	};
+
+	struct tFaceSet
+	{
+		TArray<int32> Faces;
+	};
+
+	// Check whether faces have at least two vertices in common. These must be POINTS - don't care about wedges.
+	bool FacesAreSmoothlyConnected(FSkeletalMeshImportData& ImportData, int32 Face1, int32 Face2)
+	{
+
+		//if( ( Face1 >= Thing->SkinData.Faces.Num()) || ( Face2 >= Thing->SkinData.Faces.Num()) ) return false;
+
+		if (Face1 == Face2)
+		{
+			return true;
+		}
+
+		// Smoothing groups match at least one bit in binary AND ?
+		if ((ImportData.Faces[Face1].SmoothingGroups & ImportData.Faces[Face2].SmoothingGroups) == 0)
+		{
+			return false;
+		}
+
+		int32 VertMatches = 0;
+		for (int32 i = 0; i < 3; i++)
+		{
+			int32 Point1 = ImportData.Wedges[ImportData.Faces[Face1].WedgeIndex[i]].VertexIndex;
+
+			for (int32 j = 0; j < 3; j++)
+			{
+				int32 Point2 = ImportData.Wedges[ImportData.Faces[Face2].WedgeIndex[j]].VertexIndex;
+				if (Point2 == Point1)
+				{
+					VertMatches++;
+				}
+			}
+		}
+
+		return (VertMatches >= 2);
+	}
+} //namespace SmoothGroupHelper
+
+void FSkeletalMeshImportData::SplitVerticesBySmoothingGroups()
+{
+	//
+	// Connectivity: triangles with non-matching smoothing groups will be physically split.
+	//
+	// -> Splitting involves: the UV+material-contaning vertex AND the 3d point.
+	//
+	// -> Tally smoothing groups for each and every (textured) vertex.
+	//
+	// -> Collapse: 
+	// -> start from a vertex and all its adjacent triangles - go over
+	// each triangle - if any connecting one (sharing more than one vertex) gives a smoothing match,
+	// accumulate it. Then IF more than one resulting section, 
+	// ensure each boundary 'vert' is split _if not already_ to give each smoothing group
+	// independence from all others.
+	//
+
+	int32 TotalSmoothMatches = 0;
+	int32 TotalConnexChex = 0;
+
+	// Link _all_ faces to vertices.	
+	TArray<SmoothGroupHelper::VertsFans>  Fans;
+	TArray<SmoothGroupHelper::tInfluences> PointInfluences;
+	TArray<SmoothGroupHelper::tWedgeList>  PointWedges;
+
+	Fans.AddZeroed(Points.Num());//Fans.AddExactZeroed(			Thing->SkinData.Points.Num() );
+	PointInfluences.AddZeroed(Points.Num());//PointInfluences.AddExactZeroed( Thing->SkinData.Points.Num() );
+	PointWedges.AddZeroed(Points.Num());//PointWedges.AddExactZeroed(	 Thing->SkinData.Points.Num() );
+
+	// Existing points map 1:1
+	PointToRawMap.AddUninitialized(Points.Num());
+	for (int32 i = 0; i < Points.Num(); i++)
+	{
+		PointToRawMap[i] = i;
+	}
+
+	for (int32 i = 0; i < Influences.Num(); i++)
+	{
+		if (PointInfluences.Num() <= Influences[i].VertexIndex)
+		{
+			PointInfluences.AddZeroed(Influences[i].VertexIndex - PointInfluences.Num() + 1);
+		}
+		PointInfluences[Influences[i].VertexIndex].RawInfIndices.Add(i);
+	}
+
+	for (int32 i = 0; i < Wedges.Num(); i++)
+	{
+		if (uint32(PointWedges.Num()) <= Wedges[i].VertexIndex)
+		{
+			PointWedges.AddZeroed(Wedges[i].VertexIndex - PointWedges.Num() + 1);
+		}
+
+		PointWedges[Wedges[i].VertexIndex].WedgeList.Add(i);
+	}
+
+	for (int32 f = 0; f < Faces.Num(); f++)
+	{
+		// For each face, add a pointer to that face into the Fans[vertex].
+		for (int32 i = 0; i < 3; i++)
+		{
+			int32 WedgeIndex = Faces[f].WedgeIndex[i];
+			int32 PointIndex = Wedges[WedgeIndex].VertexIndex;
+			SmoothGroupHelper::tFaceRecord NewFR;
+
+			NewFR.FaceIndex = f;
+			NewFR.HoekIndex = i;
+			NewFR.WedgeIndex = WedgeIndex; // This face touches the point courtesy of Wedges[Wedgeindex].
+			NewFR.SmoothFlags = Faces[f].SmoothingGroups;
+			NewFR.FanFlags = 0;
+			Fans[PointIndex].FaceRecord.Add(NewFR);
+			Fans[PointIndex].FanGroupCount = 0;
+		}
+	}
+
+	// Investigate connectivity and assign common group numbers (1..+) to the fans' individual FanFlags.
+	for (int32 p = 0; p < Fans.Num(); p++) // The fan of faces for each 3d point 'p'.
+	{
+		// All faces connecting.
+		if (Fans[p].FaceRecord.Num() > 0)
+		{
+			int32 FacesProcessed = 0;
+			TArray<SmoothGroupHelper::tFaceSet> FaceSets; // Sets with indices INTO FANS, not into face array.			
+
+			// Digest all faces connected to this vertex (p) into one or more smooth sets. only need to check 
+			// all faces MINUS one..
+			while (FacesProcessed < Fans[p].FaceRecord.Num())
+			{
+				// One loop per group. For the current ThisFaceIndex, tally all truly connected ones
+				// and put them in a new TArray. Once no more can be connected, stop.
+
+				int32 NewSetIndex = FaceSets.Num(); // 0 to start
+				FaceSets.AddZeroed(1);						// first one will be just ThisFaceIndex.
+
+				// Find the first non-processed face. There will be at least one.
+				int32 ThisFaceFanIndex = 0;
+				{
+					int32 SearchIndex = 0;
+					while (Fans[p].FaceRecord[SearchIndex].FanFlags == -1) // -1 indicates already  processed. 
+					{
+						SearchIndex++;
+					}
+					ThisFaceFanIndex = SearchIndex; //Fans[p].FaceRecord[SearchIndex].FaceIndex; 
+				}
+
+				// Initial face.
+				FaceSets[NewSetIndex].Faces.Add(ThisFaceFanIndex);   // Add the unprocessed Face index to the "local smoothing group" [NewSetIndex].
+				Fans[p].FaceRecord[ThisFaceFanIndex].FanFlags = -1;			  // Mark as processed.
+				FacesProcessed++;
+
+				// Find all faces connected to this face, and if there's any
+				// smoothing group matches, put it in current face set and mark it as processed;
+				// until no more match. 
+				int32 NewMatches = 0;
+				do
+				{
+					NewMatches = 0;
+					// Go over all current faces in this faceset and set if the FaceRecord (local smoothing groups) has any matches.
+					// there will be at least one face already in this faceset - the first face in the fan.
+					for (int32 n = 0; n < FaceSets[NewSetIndex].Faces.Num(); n++)
+					{
+						int32 HookFaceIdx = Fans[p].FaceRecord[FaceSets[NewSetIndex].Faces[n]].FaceIndex;
+
+						//Go over the fan looking for matches.
+						for (int32 s = 0; s < Fans[p].FaceRecord.Num(); s++)
+						{
+							// Skip if same face, skip if face already processed.
+							if ((HookFaceIdx != Fans[p].FaceRecord[s].FaceIndex) && (Fans[p].FaceRecord[s].FanFlags != -1))
+							{
+								TotalConnexChex++;
+								// Process if connected with more than one vertex, AND smooth..
+								if (SmoothGroupHelper::FacesAreSmoothlyConnected(*this, HookFaceIdx, Fans[p].FaceRecord[s].FaceIndex))
+								{
+									TotalSmoothMatches++;
+									Fans[p].FaceRecord[s].FanFlags = -1; // Mark as processed.
+									FacesProcessed++;
+									// Add 
+									FaceSets[NewSetIndex].Faces.Add(s); // Store FAN index of this face index into smoothing group's faces. 
+									// Tally
+									NewMatches++;
+								}
+							} // not the same...
+						}// all faces in fan
+					} // all faces in FaceSet
+				} while (NewMatches);
+
+			}// Repeat until all faces processed.
+
+			// For the new non-initialized  face sets, 
+			// Create a new point, influences, and uv-vertex(-ices) for all individual FanFlag groups with an index of 2+ and also remap
+			// the face's vertex into those new ones.
+			if (FaceSets.Num() > 1)
+			{
+				for (int32 f = 1; f < FaceSets.Num(); f++)
+				{
+					check(Points.Num() == PointToRawMap.Num());
+
+					// We duplicate the current vertex. (3d point)
+					const int32 NewPointIndex = Points.Num();
+					Points.AddUninitialized();
+					Points[NewPointIndex] = Points[p];
+
+					PointToRawMap.AddUninitialized();
+					PointToRawMap[NewPointIndex] = p;
+
+					// Duplicate all related weights.
+					for (int32 t = 0; t < PointInfluences[p].RawInfIndices.Num(); t++)
+					{
+						// Add new weight
+						int32 NewWeightIndex = Influences.Num();
+						Influences.AddUninitialized();
+						Influences[NewWeightIndex] = Influences[PointInfluences[p].RawInfIndices[t]];
+						Influences[NewWeightIndex].VertexIndex = NewPointIndex;
+					}
+
+					// Duplicate any and all Wedges associated with it; and all Faces' wedges involved.					
+					for (int32 w = 0; w < PointWedges[p].WedgeList.Num(); w++)
+					{
+						const int32 OldWedgeIndex = PointWedges[p].WedgeList[w];
+						Wedges[OldWedgeIndex].VertexIndex = NewPointIndex;
+					}
+				}
+			} //  if FaceSets.Num(). -> duplicate stuff
+		}//	while( FacesProcessed < Fans[p].FaceRecord.Num() )
+	} // Fans for each 3d point
+
+	if (!ensure(Points.Num() == PointToRawMap.Num()))
+	{
+		//TODO log a warning to the user
+
+		//Create a valid PointtoRawMap but with bad content
+		int32 PointNum = Points.Num();
+		PointToRawMap.Empty(PointNum);
+		for (int32 PointIndex = 0; PointIndex < PointNum; ++PointIndex)
+		{
+			PointToRawMap[PointIndex] = PointIndex;
+		}
+	}
+}
+
+
+bool FSkeletalMeshImportData::GetMeshDescription(FMeshDescription& OutMeshDescription) const
+{
+	using namespace SkeletalMeshImportData;
+	using namespace UE::AnimationCore;
+	
+	OutMeshDescription.Empty();
+	
+	FSkeletalMeshAttributes MeshAttributes(OutMeshDescription);
+
+	MeshAttributes.Register();
+
+	if (Points.IsEmpty() || Faces.IsEmpty())
+	{
+		return false;
+	}
+
+	TVertexAttributesRef<FVector> VertexPositions = MeshAttributes.GetVertexPositions();
+	FSkinWeightsVertexAttributesRef VertexSkinWeights = MeshAttributes.GetVertexSkinWeights();
+	TVertexInstanceAttributesRef<FVector> VertexInstanceNormals = MeshAttributes.GetVertexInstanceNormals();
+	TVertexInstanceAttributesRef<FVector> VertexInstanceTangents = MeshAttributes.GetVertexInstanceTangents();
+	TVertexInstanceAttributesRef<float> VertexInstanceBinormalSigns = MeshAttributes.GetVertexInstanceBinormalSigns();
+	TVertexInstanceAttributesRef<FVector4> VertexInstanceColors = MeshAttributes.GetVertexInstanceColors();
+	TVertexInstanceAttributesRef<FVector2D> VertexInstanceUVs = MeshAttributes.GetVertexInstanceUVs();
+
+	TPolygonGroupAttributesRef<FName> PolygonGroupMaterialSlotNames = MeshAttributes.GetPolygonGroupMaterialSlotNames();
+
+	VertexInstanceUVs.SetNumChannels(NumTexCoords);
+	
+	// Avoid repeated allocations and reserve the target buffers right off the bat.
+	OutMeshDescription.ReserveNewPolygonGroups(Materials.Num());
+	OutMeshDescription.ReserveNewTriangles(Faces.Num());
+	OutMeshDescription.ReserveNewVertexInstances(Wedges.Num());
+	OutMeshDescription.ReserveNewVertices(Points.Num());
+
+	// Copy the vertex positions first and maintain a map so that we can go from the import data's raw vertex index
+	// to the mesh description's VertexID.
+	TArray<FVertexID> VertexIDMap;
+	VertexIDMap.Reserve(Points.Num());
+	for (int32 Idx = 0; Idx < Points.Num(); Idx++)
+	{
+		const FVertexID VertexID = OutMeshDescription.CreateVertex();
+		VertexIDMap.Add(VertexID);
+		VertexPositions.Set(VertexID, Points[Idx]);
+	}
+
+	// The weights are stored with links back to the vertices, rather than being compact.
+	// Make a copy of the weights, sort them by vertex id and go by equal vertex-id strides.
+	// We could do an indirection but the traversal + setup cost is probably not worth it.
+	TArray<FRawBoneInfluence> SortedInfluences(Influences);
+	SortedInfluences.Sort([](const FRawBoneInfluence &A, const FRawBoneInfluence &B)
+	{
+		return A.VertexIndex < B.VertexIndex;
+	});
+
+	// TODO: Alternative profiles.
+	TArray<FBoneWeight> BoneWeights;
+	for(int32 StartStride = 0, EndStride = 0; EndStride != SortedInfluences.Num(); StartStride = EndStride)
+	{
+		const int32 VertexIndex = SortedInfluences[StartStride].VertexIndex;
+		
+		EndStride = StartStride + 1;
+		while (EndStride < SortedInfluences.Num() && VertexIndex == SortedInfluences[EndStride].VertexIndex)
+		{
+			EndStride++;
+		}
+
+		BoneWeights.Reset(0);
+		for (int32 Idx = StartStride; Idx < EndStride; Idx++)
+		{
+			const FRawBoneInfluence &RawInfluence = SortedInfluences[Idx];
+			FBoneWeight BoneWeight(FBoneIndexType(RawInfluence.BoneIndex), RawInfluence.Weight);
+			BoneWeights.Add(BoneWeight);
+		}
+
+		VertexSkinWeights.Set(VertexIDMap[VertexIndex], BoneWeights);		
+	}
+
+	// Partition the faces by material index. Each material index corresponds to a polygon group.
+	// Here, it's worth doing via indirection, due to the size of each FTriangle object.
+	TArray<int32> FaceIndices;
+	FaceIndices.SetNum(Faces.Num());
+	for (int32 Idx = 0; Idx < FaceIndices.Num(); Idx++) { FaceIndices[Idx] = Idx; }
+	FaceIndices.Sort([this](const int32 A, const int32 B)
+	{
+		return Faces[A].MatIndex < Faces[B].MatIndex;
+	});
+
+	TArray<FVertexInstanceID> VertexInstanceIDMap;
+	VertexInstanceIDMap.Init(INDEX_NONE, Wedges.Num());
+
+	TArray<FVertexInstanceID> TriangleVertexInstanceIDs;
+	TriangleVertexInstanceIDs.SetNum(3);
+
+	TArray<uint32> FaceSmoothingMasks;
+	FaceSmoothingMasks.AddZeroed(Faces.Num());
+	
+	for(int32 StartStride = 0, EndStride = 0; EndStride != FaceIndices.Num(); StartStride = EndStride)
+	{
+		const int32 MaterialIndex = Faces[FaceIndices[StartStride]].MatIndex;
+		
+		EndStride = StartStride + 1;
+		while (EndStride < FaceIndices.Num() && MaterialIndex == Faces[FaceIndices[EndStride]].MatIndex)
+		{
+			EndStride++;
+		}
+
+		// Create a section for each material index. We re-use vertex instances if they are
+		// referred to multiple times by the FTriangle object. However, because the tangents
+		// are stored on the triangle, it's possible to end up with a recycled vertex instance
+		// that has a different tangent. This is a limitation for now. 
+		// Along the way we track smoothing groups and use that to define hard edges once the
+		// entire mesh is defined.
+		const FPolygonGroupID PolygonGroupID(MaterialIndex);
+
+		if (!OutMeshDescription.IsPolygonGroupValid(PolygonGroupID))
+		{
+			OutMeshDescription.CreatePolygonGroupWithID(PolygonGroupID);
+		}
+
+		PolygonGroupMaterialSlotNames.Set(PolygonGroupID, FName(*Materials[MaterialIndex].MaterialImportName));
+		
+		for (int32 Idx = StartStride; Idx < EndStride; Idx++)
+		{
+			const FTriangle &Triangle = Faces[FaceIndices[Idx]];
+
+			for (int32 Corner = 0; Corner < 3; Corner++)
+			{
+				const uint32 WedgeId = Triangle.WedgeIndex[Corner];
+				const FVertex &Wedge = Wedges[WedgeId];
+				const FVertexID VertexID = VertexIDMap[Wedge.VertexIndex];
+
+				FVertexInstanceID VertexInstanceID = VertexInstanceIDMap[WedgeId];
+				if (VertexInstanceID == INDEX_NONE)
+				{
+					VertexInstanceID = OutMeshDescription.CreateVertexInstance(VertexID);
+
+					if (bHasVertexColors)
+					{
+						// Don't perform sRGB conversion (which mirrors what CreateFromMeshDescription does).
+						VertexInstanceColors.Set(VertexInstanceID, Wedge.Color.ReinterpretAsLinear());
+					}
+					for (uint32 UVIndex = 0; UVIndex < NumTexCoords; UVIndex++)
+					{
+						VertexInstanceUVs.Set(VertexInstanceID, UVIndex, Wedge.UVs[UVIndex]);
+					}
+					VertexInstanceTangents.Set(VertexInstanceID, Triangle.TangentX[Corner]);
+					VertexInstanceNormals.Set(VertexInstanceID, Triangle.TangentZ[Corner]);
+					VertexInstanceBinormalSigns.Set(VertexInstanceID,
+						((Triangle.TangentZ[Corner] ^ Triangle.TangentX[Corner]) | Triangle.TangentY[Corner]) < 0 ? -1.0f : 1.0f);
+					
+					VertexInstanceIDMap[WedgeId] = VertexInstanceID;
+				}
+
+				TriangleVertexInstanceIDs[Corner] = VertexInstanceID; 
+			}
+
+			const FTriangleID TriangleID = OutMeshDescription.CreateTriangle(PolygonGroupID, TriangleVertexInstanceIDs);
+			const FPolygonID PolygonID = OutMeshDescription.GetTrianglePolygon(TriangleID);
+
+			if (PolygonID.GetValue() >= FaceSmoothingMasks.Num())
+			{
+				FaceSmoothingMasks.SetNum(PolygonID.GetValue() + 1);
+			}
+			FaceSmoothingMasks[PolygonID.GetValue()] = Triangle.SmoothingGroups;
+		}
+	}
+
+	FSkeletalMeshOperations::ConvertSmoothGroupToHardEdges(FaceSmoothingMasks, OutMeshDescription);
+	
+	return true;
+}
+
+FSkeletalMeshImportData FSkeletalMeshImportData::CreateFromMeshDescription(const FMeshDescription& InMeshDescription)
+{
+	FSkeletalMeshImportData SkelMeshImportData;
+	
+	FSkeletalMeshConstAttributes Attributes(InMeshDescription);
+	TVertexAttributesConstRef<FVector> VertexPositions = Attributes.GetVertexPositions();
+	FSkinWeightsVertexAttributesConstRef VertexSkinWeights = Attributes.GetVertexSkinWeights();
+
+	TVertexInstanceAttributesConstRef<FVector2D> VertexInstanceUVs = Attributes.GetVertexInstanceUVs();
+	TVertexInstanceAttributesConstRef<FVector> VertexInstanceNormals = Attributes.GetVertexInstanceNormals();
+	TVertexInstanceAttributesConstRef<FVector> VertexInstanceTangents = Attributes.GetVertexInstanceTangents();
+	TVertexInstanceAttributesConstRef<float> VertexInstanceBiNormalSigns = Attributes.GetVertexInstanceBinormalSigns();
+	TVertexInstanceAttributesConstRef<FVector4> VertexInstanceColors = Attributes.GetVertexInstanceColors();
+
+	TPolygonGroupAttributesConstRef<FName> PolygonGroupMaterialSlotNames = Attributes.GetPolygonGroupMaterialSlotNames();
+	//Get the per face smoothing
+	TArray<uint32> FaceSmoothingMasks;
+	FaceSmoothingMasks.AddZeroed(InMeshDescription.Triangles().Num());
+	FSkeletalMeshOperations::ConvertHardEdgesToSmoothGroup(InMeshDescription, FaceSmoothingMasks);
+
+	//////////////////////////////////////////////////////////////////////////
+	// Copy the materials
+	SkelMeshImportData.Materials.Reserve(InMeshDescription.PolygonGroups().Num());
+	for (FPolygonGroupID PolygonGroupID : InMeshDescription.PolygonGroups().GetElementIDs())
+	{
+		SkeletalMeshImportData::FMaterial Material;
+		Material.MaterialImportName = PolygonGroupMaterialSlotNames[PolygonGroupID].ToString();
+		//The material interface will be added later by the factory
+		SkelMeshImportData.Materials.Add(Material);
+	}
+	SkelMeshImportData.MaxMaterialIndex = SkelMeshImportData.Materials.Num()-1;
+
+	//////////////////////////////////////////////////////////////////////////
+	//Copy the vertex positions and the influences
+
+	//Reserve the point and influences
+	SkelMeshImportData.Points.AddZeroed(InMeshDescription.Vertices().Num());
+	SkelMeshImportData.Influences.Reserve(InMeshDescription.Vertices().Num() * 4);
+	
+	for (FVertexID VertexID : InMeshDescription.Vertices().GetElementIDs())
+	{
+		//We can use GetValue because the Meshdescription was compacted before the copy
+		SkelMeshImportData.Points[VertexID.GetValue()] = VertexPositions[VertexID];
+
+		FVertexBoneWeightsConst BoneWeights = VertexSkinWeights.Get(VertexID); 
+		const int32 InfluenceCount = BoneWeights.Num();
+		
+		const int32 InfluenceOffsetIndex = SkelMeshImportData.Influences.Num();
+		SkelMeshImportData.Influences.AddDefaulted(InfluenceCount);
+		for (int32 InfluenceIndex = 0; InfluenceIndex < InfluenceCount; ++InfluenceIndex)
+		{
+			SkeletalMeshImportData::FRawBoneInfluence& BoneInfluence = SkelMeshImportData.Influences[InfluenceOffsetIndex + InfluenceIndex];
+			BoneInfluence.VertexIndex = VertexID.GetValue();
+			BoneInfluence.BoneIndex = BoneWeights[InfluenceIndex].GetBoneIndex();
+			BoneInfluence.Weight = BoneWeights[InfluenceIndex].GetWeight();
+		}
+	}
+
+	//////////////////////////////////////////////////////////////////////////
+	//Copy the triangle and vertex instances
+	SkelMeshImportData.Faces.AddZeroed(InMeshDescription.Triangles().Num());
+	SkelMeshImportData.Wedges.Reserve(InMeshDescription.VertexInstances().Num());
+	SkelMeshImportData.NumTexCoords = FMath::Min<int32>(VertexInstanceUVs.GetNumChannels(), (int32)MAX_TEXCOORDS);
+	for (FTriangleID TriangleID : InMeshDescription.Triangles().GetElementIDs())
+	{
+		FPolygonGroupID PolygonGroupID = InMeshDescription.GetTrianglePolygonGroup(TriangleID);
+		TArrayView<const FVertexInstanceID> VertexInstances = InMeshDescription.GetTriangleVertexInstances(TriangleID);
+		int32 FaceIndex = TriangleID.GetValue();
+		if (!ensure(SkelMeshImportData.Faces.IsValidIndex(FaceIndex)))
+		{
+			//TODO log an error for the user
+			break;
+		}
+		SkeletalMeshImportData::FTriangle& Face = SkelMeshImportData.Faces[FaceIndex];
+		Face.MatIndex = PolygonGroupID.GetValue();
+		Face.SmoothingGroups = 0;
+		if (FaceSmoothingMasks.IsValidIndex(FaceIndex))
+		{
+			Face.SmoothingGroups = FaceSmoothingMasks[FaceIndex];
+		}
+		//Create the wedges
+		for (int32 Corner = 0; Corner < 3; ++Corner)
+		{
+			FVertexInstanceID VertexInstanceID = VertexInstances[Corner];
+			SkeletalMeshImportData::FVertex Wedge;
+			Wedge.VertexIndex = (uint32)InMeshDescription.GetVertexInstanceVertex(VertexInstances[Corner]).GetValue();
+			Wedge.MatIndex = Face.MatIndex;
+			const bool bSRGB = false; //avoid linear to srgb conversion
+			Wedge.Color = FLinearColor(VertexInstanceColors[VertexInstanceID]).ToFColor(bSRGB);
+			for (int32 UVChannelIndex = 0; UVChannelIndex < (int32)(SkelMeshImportData.NumTexCoords); ++UVChannelIndex)
+			{
+				Wedge.UVs[UVChannelIndex] = VertexInstanceUVs.Get(VertexInstanceID, UVChannelIndex);
+			}
+			Face.TangentX[Corner] = VertexInstanceTangents[VertexInstanceID];
+			Face.TangentZ[Corner] = VertexInstanceNormals[VertexInstanceID];
+			Face.TangentY[Corner] = FVector::CrossProduct(VertexInstanceNormals[VertexInstanceID], VertexInstanceTangents[VertexInstanceID]).GetSafeNormal() * VertexInstanceBiNormalSigns[VertexInstanceID];
+
+			Face.WedgeIndex[Corner] = SkelMeshImportData.Wedges.Add(Wedge);
+		}
+	}
+
+	SkelMeshImportData.CleanUpUnusedMaterials();
+	SkelMeshImportData.SplitVerticesBySmoothingGroups();
+
+	return SkelMeshImportData;
+}
+
+#endif // WITH_EDITOR
diff --git a/Engine/Source/Runtime/Engine/Private/SkeletalMeshLODModel.cpp b/Engine/Source/Runtime/Engine/Private/SkeletalMeshLODModel.cpp
--- a/Engine/Source/Runtime/Engine/Private/SkeletalMeshLODModel.cpp
+++ b/Engine/Source/Runtime/Engine/Private/SkeletalMeshLODModel.cpp
@@ -21,6 +21,8 @@
 #include "GPUSkinVertexFactory.h"
 #include "UObject/AnimObjectVersion.h"
 #include "Misc/ScopeLock.h"
+#include "MeshDescription.h"
+#include "SkeletalMeshAttributes.h"
 
 /*-----------------------------------------------------------------------------
 FSoftSkinVertex
@@ -1233,4 +1235,122 @@ void FSkeletalMeshLODModel::CopyStructure(FSkeletalMeshLODModel* Destination, co
 	Destination->BulkDataReadMutex = DestinationBulkDataReadMutex;
 }
 
+void FSkeletalMeshLODModel::GetMeshDescription(FMeshDescription& MeshDescription, const USkeletalMesh *Owner) const
+{
+	using UE::AnimationCore::FBoneWeights;
+
+	MeshDescription.Empty();
+	
+	FSkeletalMeshAttributes MeshAttributes(MeshDescription);			
+	
+	// Register extra attributes for us.
+	MeshAttributes.Register();
+
+	TVertexAttributesRef<FVector> VertexPositions = MeshAttributes.GetVertexPositions();
+	FSkinWeightsVertexAttributesRef VertexSkinWeights = MeshAttributes.GetVertexSkinWeights();
+	TVertexInstanceAttributesRef<FVector> VertexInstanceNormals = MeshAttributes.GetVertexInstanceNormals();
+	TVertexInstanceAttributesRef<FVector> VertexInstanceTangents = MeshAttributes.GetVertexInstanceTangents();
+	TVertexInstanceAttributesRef<float> VertexInstanceBinormalSigns = MeshAttributes.GetVertexInstanceBinormalSigns();
+	TVertexInstanceAttributesRef<FVector4> VertexInstanceColors = MeshAttributes.GetVertexInstanceColors();
+	TVertexInstanceAttributesRef<FVector2D> VertexInstanceUVs = MeshAttributes.GetVertexInstanceUVs();
+
+	TPolygonGroupAttributesRef<FName> PolygonGroupMaterialSlotNames = MeshAttributes.GetPolygonGroupMaterialSlotNames();
+	
+	const int32 NumTriangles = IndexBuffer.Num() / 3;
+
+	MeshDescription.ReserveNewPolygonGroups(Sections.Num());
+	MeshDescription.ReserveNewTriangles(NumTriangles);
+	MeshDescription.ReserveNewVertexInstances(NumTriangles * 3);
+	MeshDescription.ReserveNewVertices(static_cast<int32>(NumVertices));
+
+	TArray<FVertexID> VertexIDs;
+	VertexIDs.Reserve(NumVertices);
+	for (int32 VertexIndex = 0; VertexIndex < int32(NumVertices); VertexIndex++)
+	{
+		VertexIDs.Add(MeshDescription.CreateVertex());
+	}
+
+	const TArray<FSkeletalMaterial>& Materials = Owner->GetMaterials();
+	const bool bHasVertexColors = (Owner->GetVertexBufferFlags() & ESkeletalMeshVertexFlags::HasVertexColors);
+
+	// Convert sections to polygon groups, each with their own material.
+	for (int32 SectionIndex = 0; SectionIndex < Sections.Num(); SectionIndex++)
+	{
+		const FSkelMeshSection& Section = Sections[SectionIndex];
+
+		// Convert positions and bone weights
+		const TArray<FSoftSkinVertex>& SourceVertices = Section.SoftVertices;
+		for (int32 VertexIndex = 0; VertexIndex < SourceVertices.Num(); VertexIndex++)
+		{
+			const FVertexID VertexID = VertexIDs[VertexIndex + Section.BaseVertexIndex];
+
+			VertexPositions.Set(VertexID, SourceVertices[VertexIndex].Position);
+
+			// Skeleton bone indexes translated from the render mesh compact indexes.
+			FBoneIndexType	InfluenceBones[MAX_TOTAL_INFLUENCES];
+
+			for (int32 InfluenceIndex = 0; SourceVertices[VertexIndex].InfluenceWeights[InfluenceIndex] && InfluenceIndex < MAX_TOTAL_INFLUENCES; InfluenceIndex++)
+			{
+				const int32 BoneId = SourceVertices[VertexIndex].InfluenceBones[InfluenceIndex];
+
+				InfluenceBones[InfluenceIndex] = Section.BoneMap[BoneId];
+			}
+
+			VertexSkinWeights.Set(VertexID, FBoneWeights::Create(InfluenceBones, SourceVertices[VertexIndex].InfluenceWeights));
+		}
+
+
+		const FPolygonGroupID PolygonGroupID(Section.MaterialIndex);
+
+		if (!MeshDescription.IsPolygonGroupValid(PolygonGroupID))
+		{
+			MeshDescription.CreatePolygonGroupWithID(FPolygonGroupID(Section.MaterialIndex));
+		}
+
+		if (ensure(Materials.IsValidIndex(Section.MaterialIndex)))
+		{
+			PolygonGroupMaterialSlotNames.Set(PolygonGroupID, Materials[Section.MaterialIndex].MaterialSlotName);
+		}
+
+		for (int32 TriangleID = 0; TriangleID < int32(Section.NumTriangles); TriangleID++)
+		{
+			const int32 VertexIndexBase = TriangleID * 3 + Section.BaseIndex;
+
+			TArray<FVertexInstanceID> TriangleVertexInstanceIDs;
+			TriangleVertexInstanceIDs.SetNum(3);
+
+			for (int32 Corner = 0; Corner < 3; Corner++)
+			{
+				const int32 SourceVertexIndex = IndexBuffer[VertexIndexBase + Corner];
+				const FVertexID VertexID = VertexIDs[SourceVertexIndex];
+				const FVertexInstanceID VertexInstanceID = MeshDescription.CreateVertexInstance(VertexID);
+
+				const FSoftSkinVertex& SourceVertex = SourceVertices[SourceVertexIndex - Section.BaseVertexIndex];
+
+				VertexInstanceNormals.Set(VertexInstanceID, SourceVertex.TangentZ);
+				VertexInstanceTangents.Set(VertexInstanceID, SourceVertex.TangentX);
+				VertexInstanceBinormalSigns.Set(VertexInstanceID, FMatrix(
+					SourceVertex.TangentX.GetSafeNormal(),
+					SourceVertex.TangentY.GetSafeNormal(),
+					SourceVertex.TangentZ.GetSafeNormal(),
+					FVector::ZeroVector).Determinant() < 0.0f ? -1.0f : +1.0f);
+
+				for (int32 UVIndex = 0; UVIndex < int32(NumTexCoords); UVIndex++)
+				{
+					VertexInstanceUVs.Set(VertexInstanceID, UVIndex, SourceVertex.UVs[UVIndex]);
+				}
+
+				if (bHasVertexColors)
+				{
+					VertexInstanceColors.Set(VertexInstanceID, FLinearColor(SourceVertex.Color));
+				}
+
+				TriangleVertexInstanceIDs[Corner] = VertexInstanceID;
+			}
+
+			MeshDescription.CreateTriangle(PolygonGroupID, TriangleVertexInstanceIDs);
+		}
+	}
+}
+
 #endif // WITH_EDITOR
diff --git a/Engine/Source/Runtime/Engine/Public/Rendering/SkeletalMeshLODImporterData.h b/Engine/Source/Runtime/Engine/Public/Rendering/SkeletalMeshLODImporterData.h
--- a/Engine/Source/Runtime/Engine/Public/Rendering/SkeletalMeshLODImporterData.h
+++ b/Engine/Source/Runtime/Engine/Public/Rendering/SkeletalMeshLODImporterData.h
@@ -14,6 +14,7 @@
 #include "Animation/MorphTarget.h"
 #include "Templates/DontCopy.h"
 
+struct FMeshDescription;
 class FSkeletalMeshLODModel;
 
 #endif
@@ -418,6 +419,17 @@ public:
 	 * Use the faces corner normals to create the face smooth groups data
 	 */
 	void ComputeSmoothGroupFromNormals();
+
+	/**
+	 * Returns a mesh description from the import data
+	 */
+	bool GetMeshDescription(FMeshDescription &OutMeshDescription) const;
+
+	static FSkeletalMeshImportData CreateFromMeshDescription(const FMeshDescription &InMeshDescription);
+
+private:
+	void CleanUpUnusedMaterials();
+	void SplitVerticesBySmoothingGroups();
 };
 
 /**
diff --git a/Engine/Source/Runtime/Engine/Public/Rendering/SkeletalMeshLODModel.h b/Engine/Source/Runtime/Engine/Public/Rendering/SkeletalMeshLODModel.h
--- a/Engine/Source/Runtime/Engine/Public/Rendering/SkeletalMeshLODModel.h
+++ b/Engine/Source/Runtime/Engine/Public/Rendering/SkeletalMeshLODModel.h
@@ -496,6 +496,12 @@ public:
 		FSkeletalMeshLODModel::CopyStructure(Destination, Other);
 		return Destination;
 	}
+
+	/**
+	 * Fills in a representation of this model into the given mesh description object. Existing
+	 * mesh description data is emptied.
+	 */
+	void ENGINE_API GetMeshDescription(FMeshDescription& MeshDescription, const USkeletalMesh *Owner) const;
 };
 
 #endif // WITH_EDITOR
\ No newline at end of file
