Fixed issue with CL16184912 causing HLOD grid actors not to be created under some conditions

diff --git a/Engine/Source/Runtime/Engine/Private/WorldPartition/RuntimeSpatialHash/RuntimeSpatialHashHLOD.cpp b/Engine/Source/Runtime/Engine/Private/WorldPartition/RuntimeSpatialHash/RuntimeSpatialHashHLOD.cpp
--- a/Engine/Source/Runtime/Engine/Private/WorldPartition/RuntimeSpatialHash/RuntimeSpatialHashHLOD.cpp
+++ b/Engine/Source/Runtime/Engine/Private/WorldPartition/RuntimeSpatialHash/RuntimeSpatialHashHLOD.cpp
@@ -264,11 +264,12 @@ static void GatherHLODGrids(UWorldPartition* WorldPartition, TMap<FName, FSpatia
 
 	for (UActorDescContainer::TIterator<> ActorDescIterator(WorldPartition); ActorDescIterator; ++ActorDescIterator)
 	{
-		if (!ActorDescIterator->GetActorClass()->IsChildOf<AWorldPartitionHLOD>())
+		const FWorldPartitionActorDesc& ActorDesc = **ActorDescIterator;
+		if (!ActorDesc.GetActorClass()->IsChildOf<AWorldPartitionHLOD>())
 		{
-			if (ActorDescIterator->GetActorIsHLODRelevant())
+			if (ActorDesc.GetActorIsHLODRelevant())
 			{
-				UHLODLayer* HLODLayer = ActorDescIterator->GetHLODLayer();
+				UHLODLayer* HLODLayer = UHLODLayer::GetHLODLayer(ActorDesc, WorldPartition);
 				while (HLODLayer != nullptr)
 				{
 					bool bAlreadyInSet = false;
