Fix namepsacing ambiguity issues revealed in the unity stress test.

#trivial
#rnx

diff --git a/Engine/Source/Runtime/Engine/Private/Rendering/SkeletalMeshLODImporterData.cpp b/Engine/Source/Runtime/Engine/Private/Rendering/SkeletalMeshLODImporterData.cpp
--- a/Engine/Source/Runtime/Engine/Private/Rendering/SkeletalMeshLODImporterData.cpp
+++ b/Engine/Source/Runtime/Engine/Private/Rendering/SkeletalMeshLODImporterData.cpp
@@ -1579,7 +1579,6 @@ void FSkeletalMeshImportData::SplitVerticesBySmoothingGroups()
 
 bool FSkeletalMeshImportData::GetMeshDescription(FMeshDescription& OutMeshDescription) const
 {
-	using namespace SkeletalMeshImportData;
 	using namespace UE::AnimationCore;
 	
 	OutMeshDescription.Empty();
@@ -1625,8 +1624,8 @@ bool FSkeletalMeshImportData::GetMeshDescription(FMeshDescription& OutMeshDescri
 	// The weights are stored with links back to the vertices, rather than being compact.
 	// Make a copy of the weights, sort them by vertex id and go by equal vertex-id strides.
 	// We could do an indirection but the traversal + setup cost is probably not worth it.
-	TArray<FRawBoneInfluence> SortedInfluences(Influences);
-	SortedInfluences.Sort([](const FRawBoneInfluence &A, const FRawBoneInfluence &B)
+	TArray<SkeletalMeshImportData::FRawBoneInfluence> SortedInfluences(Influences);
+	SortedInfluences.Sort([](const SkeletalMeshImportData::FRawBoneInfluence &A, const SkeletalMeshImportData::FRawBoneInfluence &B)
 	{
 		return A.VertexIndex < B.VertexIndex;
 	});
@@ -1646,7 +1645,7 @@ bool FSkeletalMeshImportData::GetMeshDescription(FMeshDescription& OutMeshDescri
 		BoneWeights.Reset(0);
 		for (int32 Idx = StartStride; Idx < EndStride; Idx++)
 		{
-			const FRawBoneInfluence &RawInfluence = SortedInfluences[Idx];
+			const SkeletalMeshImportData::FRawBoneInfluence &RawInfluence = SortedInfluences[Idx];
 			FBoneWeight BoneWeight(FBoneIndexType(RawInfluence.BoneIndex), RawInfluence.Weight);
 			BoneWeights.Add(BoneWeight);
 		}
@@ -1700,12 +1699,12 @@ bool FSkeletalMeshImportData::GetMeshDescription(FMeshDescription& OutMeshDescri
 		
 		for (int32 Idx = StartStride; Idx < EndStride; Idx++)
 		{
-			const FTriangle &Triangle = Faces[FaceIndices[Idx]];
+			const SkeletalMeshImportData::FTriangle &Triangle = Faces[FaceIndices[Idx]];
 
 			for (int32 Corner = 0; Corner < 3; Corner++)
 			{
 				const uint32 WedgeId = Triangle.WedgeIndex[Corner];
-				const FVertex &Wedge = Wedges[WedgeId];
+				const SkeletalMeshImportData::FVertex &Wedge = Wedges[WedgeId];
 				const FVertexID VertexID = VertexIDMap[Wedge.VertexIndex];
 
 				FVertexInstanceID VertexInstanceID = VertexInstanceIDMap[WedgeId];
@@ -1718,7 +1717,7 @@ bool FSkeletalMeshImportData::GetMeshDescription(FMeshDescription& OutMeshDescri
 						// Don't perform sRGB conversion (which mirrors what CreateFromMeshDescription does).
 						VertexInstanceColors.Set(VertexInstanceID, Wedge.Color.ReinterpretAsLinear());
 					}
-					for (uint32 UVIndex = 0; UVIndex < NumTexCoords; UVIndex++)
+					for (int32 UVIndex = 0; UVIndex < int32(NumTexCoords); UVIndex++)
 					{
 						VertexInstanceUVs.Set(VertexInstanceID, UVIndex, Wedge.UVs[UVIndex]);
 					}
