Fix sync groups not working in certain circumstances

Sync groups use an old-style enum, so the 'get' function ended up reading 4 bytes rather than one (as a TEnumAsByte), but the two types covnert silently so no warnings/errors are given. Lesson is: use enum classes as much as possible.

diff --git a/Engine/Source/Runtime/Engine/Private/Animation/AnimNode_SequencePlayer.cpp b/Engine/Source/Runtime/Engine/Private/Animation/AnimNode_SequencePlayer.cpp
--- a/Engine/Source/Runtime/Engine/Private/Animation/AnimNode_SequencePlayer.cpp
+++ b/Engine/Source/Runtime/Engine/Private/Animation/AnimNode_SequencePlayer.cpp
@@ -225,7 +225,7 @@ FName FAnimNode_SequencePlayer::GetGroupName() const
 
 EAnimGroupRole::Type FAnimNode_SequencePlayer::GetGroupRole() const
 {
-	return GET_ANIM_NODE_DATA(EAnimGroupRole::Type, GroupRole);
+	return GET_ANIM_NODE_DATA(TEnumAsByte<EAnimGroupRole::Type>, GroupRole);
 }
 
 EAnimSyncMethod FAnimNode_SequencePlayer::GetGroupMethod() const
