//UE5/Main - Cook Chunked Switch Client - New play length is same as existing one
#fix check for play length delta before setting it when refreshing data according to Parent Montage data

diff --git a/Engine/Source/Runtime/Engine/Private/Animation/AnimSequenceBase.cpp b/Engine/Source/Runtime/Engine/Private/Animation/AnimSequenceBase.cpp
--- a/Engine/Source/Runtime/Engine/Private/Animation/AnimSequenceBase.cpp
+++ b/Engine/Source/Runtime/Engine/Private/Animation/AnimSequenceBase.cpp
@@ -834,7 +834,10 @@ void UAnimSequenceBase::RefreshParentAssetData()
 	{
 		const UAnimDataModel* ParentDataModel = ParentSeqBase->GetDataModel();
 
-		Controller->SetPlayLength(ParentDataModel->GetPlayLength());
+		if (!FMath::IsNearlyEqual(DataModel->GetPlayLength(), ParentDataModel->GetPlayLength()))
+		{
+			Controller->SetPlayLength(ParentDataModel->GetPlayLength());
+		}
 		
 		Controller->RemoveAllCurvesOfType(ERawCurveTrackTypes::RCT_Float);
 		for (const FFloatCurve& FloatCurve : ParentDataModel->GetFloatCurves())
