- Moved World Parition RuntimeCell priority and sorting in base class UWorldPartitionRuntimeHash.
- Added SourcePriority in UWorldPartitionRuntimeCell to prioritize runtime cell streaming based on which streaming source intersect them.
- Refactored cell sorting so that devired class can override sorting predicate. Overridden default sorting for UWorldPartitionRuntimeSpatialHashCell.
- Added profile tags to track runtime performance of World Partition streaming.

diff --git a/Engine/Source/Runtime/Engine/Private/WorldPartition/WorldPartitionLevelStreamingPolicy.cpp b/Engine/Source/Runtime/Engine/Private/WorldPartition/WorldPartitionLevelStreamingPolicy.cpp
--- a/Engine/Source/Runtime/Engine/Private/WorldPartition/WorldPartitionLevelStreamingPolicy.cpp
+++ b/Engine/Source/Runtime/Engine/Private/WorldPartition/WorldPartitionLevelStreamingPolicy.cpp
@@ -82,6 +82,8 @@ void UWorldPartitionLevelStreamingPolicy::SetTargetStateForCells(EWorldPartition
 
 void UWorldPartitionLevelStreamingPolicy::SetCellsStateToLoaded(const TSet<const UWorldPartitionRuntimeCell*>& ToLoadCells)
 {
+	TRACE_CPUPROFILER_EVENT_SCOPE(UWorldPartitionLevelStreamingPolicy::SetCellsStateToLoaded);
+
 	int32 MaxCellsToLoad = GetMaxCellsToLoad();
 	
 	// Sort cells based on importance
@@ -113,6 +115,8 @@ void UWorldPartitionLevelStreamingPolicy::SetCellsStateToLoaded(const TSet<const
 
 void UWorldPartitionLevelStreamingPolicy::SetCellsStateToActivated(const TSet<const UWorldPartitionRuntimeCell*>& ToActivateCells)
 {
+	TRACE_CPUPROFILER_EVENT_SCOPE(UWorldPartitionLevelStreamingPolicy::SetCellsStateToActivated);
+
 	int32 MaxCellsToLoad = GetMaxCellsToLoad();
 
 	// Sort cells based on importance
@@ -144,6 +148,8 @@ void UWorldPartitionLevelStreamingPolicy::SetCellsStateToActivated(const TSet<co
 
 void UWorldPartitionLevelStreamingPolicy::SetCellsStateToUnloaded(const TSet<const UWorldPartitionRuntimeCell*>& ToUnloadCells)
 {
+	TRACE_CPUPROFILER_EVENT_SCOPE(UWorldPartitionLevelStreamingPolicy::SetCellsStateToUnloaded);
+
 	for (const UWorldPartitionRuntimeCell* Cell : ToUnloadCells)
 	{
 		UE_LOG(LogWorldPartition, Verbose, TEXT("UWorldPartitionLevelStreamingPolicy::UnloadCells %s"), *Cell->GetName());
@@ -156,6 +162,8 @@ void UWorldPartitionLevelStreamingPolicy::SetCellsStateToUnloaded(const TSet<con
 
 ULevel* UWorldPartitionLevelStreamingPolicy::GetPreferredLoadedLevelToAddToWorld() const
 {
+	TRACE_CPUPROFILER_EVENT_SCOPE(UWorldPartitionLevelStreamingPolicy::GetPreferredLoadedLevelToAddToWorld);
+
 	check(WorldPartition->IsInitialized());
 	UWorld* World = WorldPartition->GetWorld();
 	if (World->GetNetMode() == NM_DedicatedServer)
diff --git a/Engine/Source/Runtime/Engine/Private/WorldPartition/WorldPartitionRuntimeCell.cpp b/Engine/Source/Runtime/Engine/Private/WorldPartition/WorldPartitionRuntimeCell.cpp
--- a/Engine/Source/Runtime/Engine/Private/WorldPartition/WorldPartitionRuntimeCell.cpp
+++ b/Engine/Source/Runtime/Engine/Private/WorldPartition/WorldPartitionRuntimeCell.cpp
@@ -9,9 +9,14 @@
 #include "WorldPartition/DataLayer/DataLayersID.h"
 #include "Engine/World.h"
 
+int32 UWorldPartitionRuntimeCell::StreamingSourceCacheEpoch = 0;
+
 UWorldPartitionRuntimeCell::UWorldPartitionRuntimeCell(const FObjectInitializer& ObjectInitializer)
 : Super(ObjectInitializer)
 , bIsAlwaysLoaded(false)
+, Priority(0)
+, CachedSourcePriority(0)
+, CachedSourceInfoEpoch(INT_MIN)
 {}
 
 #if WITH_EDITOR
@@ -68,6 +73,28 @@ void UWorldPartitionRuntimeCell::UpdateDebugName()
 
 #endif
 
+bool UWorldPartitionRuntimeCell::CacheStreamingSourceInfo(const FWorldPartitionStreamingSource& Source) const
+{
+	bool bWasCacheDirtied = false;
+	if (CachedSourceInfoEpoch != UWorldPartitionRuntimeCell::StreamingSourceCacheEpoch)
+	{
+		CachedSourceInfoEpoch = UWorldPartitionRuntimeCell::StreamingSourceCacheEpoch;
+		bWasCacheDirtied = true;
+	}
+
+	// If cache was dirtied, use value, else use minimum with existing cached value
+	CachedSourcePriority = bWasCacheDirtied ? (int32)Source.Priority : FMath::Min((int32)Source.Priority, CachedSourcePriority);
+	return bWasCacheDirtied;
+}
+
+int32 UWorldPartitionRuntimeCell::SortCompare(const UWorldPartitionRuntimeCell* Other) const
+{
+	// Source priority (lower value is higher prio)
+	const int32 Comparison = CachedSourcePriority - Other->CachedSourcePriority;
+	// Cell priority (lower value is higher prio)
+	return (Comparison != 0) ? Comparison : (Priority - Other->Priority);
+}
+
 bool UWorldPartitionRuntimeCell::IsDebugShown() const
 {
 	return FWorldPartitionDebugHelper::IsDebugRuntimeHashGridShown(GridName) &&
diff --git a/Engine/Source/Runtime/Engine/Private/WorldPartition/WorldPartitionRuntimeHash.cpp b/Engine/Source/Runtime/Engine/Private/WorldPartition/WorldPartitionRuntimeHash.cpp
--- a/Engine/Source/Runtime/Engine/Private/WorldPartition/WorldPartitionRuntimeHash.cpp
+++ b/Engine/Source/Runtime/Engine/Private/WorldPartition/WorldPartitionRuntimeHash.cpp
@@ -267,4 +267,15 @@ void UWorldPartitionRuntimeHash::CheckForErrorsInternal(const TMap<FGuid, FWorld
 }
 #endif
 
+void UWorldPartitionRuntimeHash::SortStreamingCellsByImportance(const TSet<const UWorldPartitionRuntimeCell*>& InCells, const TArray<FWorldPartitionStreamingSource>& InSources, TArray<const UWorldPartitionRuntimeCell*, TInlineAllocator<256>>& OutSortedCells) const
+{
+	TRACE_CPUPROFILER_EVENT_SCOPE(UWorldPartitionRuntimeHash::SortStreamingCellsByImportance);
+
+	OutSortedCells = InCells.Array();
+	Algo::Sort(OutSortedCells, [](const UWorldPartitionRuntimeCell* CellA, const UWorldPartitionRuntimeCell* CellB)
+	{
+		return CellA->SortCompare(CellB) < 0;
+	});
+}
+
 #undef LOCTEXT_NAMESPACE
\ No newline at end of file
diff --git a/Engine/Source/Runtime/Engine/Private/WorldPartition/WorldPartitionRuntimeSpatialHash.cpp b/Engine/Source/Runtime/Engine/Private/WorldPartition/WorldPartitionRuntimeSpatialHash.cpp
--- a/Engine/Source/Runtime/Engine/Private/WorldPartition/WorldPartitionRuntimeSpatialHash.cpp
+++ b/Engine/Source/Runtime/Engine/Private/WorldPartition/WorldPartitionRuntimeSpatialHash.cpp
@@ -65,12 +65,6 @@ static FAutoConsoleVariableRef CVarShowRuntimeSpatialHashGridLevelCount(
 	GShowRuntimeSpatialHashGridLevelCount,
 	TEXT("Used to choose how many grid levels to display when showing world partition runtime hash."));
 
-static float GRuntimeSpatialHashCellToSourceAngleContributionToCellImportance = 0.4f; // Value between [0, 1]
-static FAutoConsoleVariableRef CVarRuntimeSpatialHashCellToSourceAngleContributionToCellImportance(
-	TEXT("wp.Runtime.RuntimeSpatialHashCellToSourceAngleContributionToCellImportance"),
-	GRuntimeSpatialHashCellToSourceAngleContributionToCellImportance,
-	TEXT("Value between 0 and 1 that modulates the contribution of the angle between streaming source-to-cell vector and source-forward vector to the cell importance. The closest to 0, the less the angle will contribute to the cell importance."));
-
 // ------------------------------------------------------------------------------------------------
 FSpatialHashStreamingGrid::FSpatialHashStreamingGrid()
 	: Origin(ForceInitToZero)
@@ -108,6 +102,8 @@ const FSquare2DGridHelper& FSpatialHashStreamingGrid::GetGridHelper() const
 
 void FSpatialHashStreamingGrid::GetCells(const FWorldPartitionStreamingQuerySource& QuerySource, TSet<const UWorldPartitionRuntimeCell*>& OutCells) const
 {
+	TRACE_CPUPROFILER_EVENT_SCOPE(FSpatialHashStreamingGrid::GetCells_QuerySource);
+
 	auto ShouldAddCell = [](const UWorldPartitionRuntimeCell* Cell, const FWorldPartitionStreamingQuerySource& QuerySource)
 	{
 		if (Cell->HasDataLayers())
@@ -160,8 +156,10 @@ void FSpatialHashStreamingGrid::GetCells(const FWorldPartitionStreamingQuerySour
 	}	
 }
 
-void FSpatialHashStreamingGrid::GetCells(const TArray<FWorldPartitionStreamingSource>& Sources, const UDataLayerSubsystem* DataLayerSubsystem, TSet<const UWorldPartitionRuntimeCell*>& OutActivateCells, TSet<const UWorldPartitionRuntimeCell*>& OutLoadCells) const
+void FSpatialHashStreamingGrid::GetCells(const TArray<FWorldPartitionStreamingSource>& Sources, const UDataLayerSubsystem* DataLayerSubsystem, UWorldPartitionRuntimeHash::FStreamingSourceCells& OutActivateCells, UWorldPartitionRuntimeHash::FStreamingSourceCells& OutLoadCells) const
 {
+	TRACE_CPUPROFILER_EVENT_SCOPE(FSpatialHashStreamingGrid::GetCells);
+
 	const FSquare2DGridHelper& Helper = GetGridHelper();
 	for (const FWorldPartitionStreamingSource& Source : Sources)
 	{
@@ -177,24 +175,24 @@ void FSpatialHashStreamingGrid::GetCells(const TArray<FWorldPartitionStreamingSo
 					{
 						if (Source.TargetState == EStreamingSourceTargetState::Loaded)
 						{
-							OutLoadCells.Add(Cell);
+							OutLoadCells.AddCell(Cell, Source);
 						}
 						else
 						{
 							check(Source.TargetState == EStreamingSourceTargetState::Activated);
-							OutActivateCells.Add(Cell);
+							OutActivateCells.AddCell(Cell, Source);
 						}
 					}
 					else if (DataLayerSubsystem && DataLayerSubsystem->IsAnyDataLayerInState(Cell->GetDataLayers(), EDataLayerState::Loaded))
 					{
-						OutLoadCells.Add(Cell);
+						OutLoadCells.AddCell(Cell, Source);
 					}
 				}
 			}
 		});
 	}
 
-	GetAlwaysLoadedCells(DataLayerSubsystem, OutActivateCells, OutLoadCells);
+	GetAlwaysLoadedCells(DataLayerSubsystem, OutActivateCells.GetCells(), OutLoadCells.GetCells());
 }
 
 void FSpatialHashStreamingGrid::GetAlwaysLoadedCells(const UDataLayerSubsystem* DataLayerSubsystem, TSet<const UWorldPartitionRuntimeCell*>& OutActivateCells, TSet<const UWorldPartitionRuntimeCell*>& OutLoadCells) const
@@ -430,10 +428,11 @@ void FSpatialHashStreamingGrid::Draw2D(UCanvas* Canvas, const TArray<FWorldParti
 
 		FCanvasLineItem LineItem;
 		LineItem.LineThickness = 2;
-		LineItem.SetColor(FLinearColor::White);
 
 		for (const FWorldPartitionStreamingSource& Source : Sources)
 		{
+			LineItem.SetColor(Source.GetDebugColor());
+
 			TArray<FVector> LinePoints;
 			LinePoints.SetNum(2);
 
@@ -833,7 +832,7 @@ bool UWorldPartitionRuntimeSpatialHash::CreateStreamingGrid(const FSpatialHashRu
 				StreamingCell->SetIsAlwaysLoaded(bIsCellAlwaysLoaded);
 				StreamingCell->SetDataLayers(GridCellDataChunk.GetDataLayers());
 				StreamingCell->Level = Level;
-				StreamingCell->Priority = RuntimeGrid.Priority;
+				StreamingCell->SetPriority(RuntimeGrid.Priority);
 				FBox2D Bounds;
 				verify(TempLevel.GetCellBounds(FIntVector2(CellCoordX, CellCoordY), Bounds));
 				StreamingCell->Position = FVector(Bounds.GetCenter(), 0.f);
@@ -1018,8 +1017,10 @@ bool UWorldPartitionRuntimeSpatialHash::GetStreamingCells(const FWorldPartitionS
 	return !!OutCells.Num();
 }
 
-bool UWorldPartitionRuntimeSpatialHash::GetStreamingCells(const TArray<FWorldPartitionStreamingSource>& Sources, TSet<const UWorldPartitionRuntimeCell*>& OutActivateCells, TSet<const UWorldPartitionRuntimeCell*>& OutLoadCells) const
+bool UWorldPartitionRuntimeSpatialHash::GetStreamingCells(const TArray<FWorldPartitionStreamingSource>& Sources, UWorldPartitionRuntimeHash::FStreamingSourceCells& OutActivateCells, UWorldPartitionRuntimeHash::FStreamingSourceCells& OutLoadCells) const
 {
+	TRACE_CPUPROFILER_EVENT_SCOPE(UWorldPartitionRuntimeSpatialHash::GetStreamingCells);
+
 	UDataLayerSubsystem* DataLayerSubsystem = GetWorld()->GetSubsystem<UDataLayerSubsystem>();
 
 	if (Sources.Num() == 0)
@@ -1027,7 +1028,7 @@ bool UWorldPartitionRuntimeSpatialHash::GetStreamingCells(const TArray<FWorldPar
 		// Get always loaded cells
 		for (const FSpatialHashStreamingGrid& StreamingGrid : StreamingGrids)
 		{
-			StreamingGrid.GetAlwaysLoadedCells(DataLayerSubsystem, OutActivateCells, OutLoadCells);
+			StreamingGrid.GetAlwaysLoadedCells(DataLayerSubsystem, OutActivateCells.GetCells(), OutLoadCells.GetCells());
 		}
 	}
 	else
@@ -1042,63 +1043,6 @@ bool UWorldPartitionRuntimeSpatialHash::GetStreamingCells(const TArray<FWorldPar
 	return !!(OutActivateCells.Num() + OutLoadCells.Num());
 }
 
-void UWorldPartitionRuntimeSpatialHash::SortStreamingCellsByImportance(const TSet<const UWorldPartitionRuntimeCell*>& InCells, const TArray<FWorldPartitionStreamingSource>& InSources, TArray<const UWorldPartitionRuntimeCell*, TInlineAllocator<256>>& OutSortedCells) const
-{
-	struct FCellShortestDist
-	{
-		FCellShortestDist(const UWorldPartitionRuntimeSpatialHashCell* InCell, float InMinDistance)
-			: Cell(InCell)
-			, SourceMinDistance(InMinDistance)
-		{}
-		const UWorldPartitionRuntimeSpatialHashCell* Cell;
-		float SourceMinDistance;
-	};
-
-	TArray<FCellShortestDist> SortedCells;
-	SortedCells.Reserve(InCells.Num());
-	for (const UWorldPartitionRuntimeCell* ToLoadCell : InCells)
-	{
-		const UWorldPartitionRuntimeSpatialHashCell* Cell = Cast<const UWorldPartitionRuntimeSpatialHashCell>(ToLoadCell);
-		FCellShortestDist& SortedCell = SortedCells.Emplace_GetRef(Cell, FLT_MAX);
-
-		const float AngleContribution = FMath::Clamp(GRuntimeSpatialHashCellToSourceAngleContributionToCellImportance, 0.f, 1.f);
-		for (const FWorldPartitionStreamingSource& Source : InSources)
-		{
-			const float SqrDistance = FVector::DistSquared(Source.Location, Cell->Position);
-			float AngleFactor = 1.f;
-			if (!FMath::IsNearlyZero(AngleContribution))
-			{
-				const FVector2D SourceForward(Source.Rotation.Quaternion().GetForwardVector());
-				const FVector2D SourceToCell(Cell->Position - Source.Location);
-				const float Dot = FVector2D::DotProduct(SourceForward.GetSafeNormal(), SourceToCell.GetSafeNormal());
-				const float NormalizedAngle = FMath::Clamp(FMath::Abs(FMath::Acos(Dot)/PI), 0.f, 1.f);
-				AngleFactor = FMath::Pow(NormalizedAngle, AngleContribution);
-			}
-			// Modulate distance to cell by angle relative to source forward vector (to prioritize cells in front)
-			SortedCell.SourceMinDistance = FMath::Min(SqrDistance * AngleFactor, SortedCell.SourceMinDistance);
-		}
-	}
-
-	Algo::Sort(SortedCells, [](const FCellShortestDist& A, const FCellShortestDist& B)
-	{
-		if (A.Cell->Priority == B.Cell->Priority)
-		{
-			if (A.Cell->Level == B.Cell->Level)
-			{
-				return A.SourceMinDistance < B.SourceMinDistance;
-			}
-			return A.Cell->Level > B.Cell->Level;
-		}
-		return A.Cell->Priority < B.Cell->Priority;
-	});
-
-	OutSortedCells.Reserve(InCells.Num());
-	for (const FCellShortestDist& SortedCell : SortedCells)
-	{
-		OutSortedCells.Add(SortedCell.Cell);
-	}
-}
-
 FVector2D UWorldPartitionRuntimeSpatialHash::GetDraw2DDesiredFootprint(const FVector2D& CanvasSize) const
 {
 	return FVector2D(CanvasSize.X * GetFilteredStreamingGrids().Num(), CanvasSize.Y);
diff --git a/Engine/Source/Runtime/Engine/Private/WorldPartition/WorldPartitionRuntimeSpatialHashCell.cpp b/Engine/Source/Runtime/Engine/Private/WorldPartition/WorldPartitionRuntimeSpatialHashCell.cpp
--- a/Engine/Source/Runtime/Engine/Private/WorldPartition/WorldPartitionRuntimeSpatialHashCell.cpp
+++ b/Engine/Source/Runtime/Engine/Private/WorldPartition/WorldPartitionRuntimeSpatialHashCell.cpp
@@ -1,9 +1,17 @@
 // Copyright Epic Games, Inc. All Rights Reserved.
 #include "WorldPartition/WorldPartitionRuntimeSpatialHashCell.h"
+#include "WorldPartition/WorldPartitionStreamingSource.h"
+
+static float GRuntimeSpatialHashCellToSourceAngleContributionToCellImportance = 0.4f; // Value between [0, 1]
+static FAutoConsoleVariableRef CVarRuntimeSpatialHashCellToSourceAngleContributionToCellImportance(
+	TEXT("wp.Runtime.RuntimeSpatialHashCellToSourceAngleContributionToCellImportance"),
+	GRuntimeSpatialHashCellToSourceAngleContributionToCellImportance,
+	TEXT("Value between 0 and 1 that modulates the contribution of the angle between streaming source-to-cell vector and source-forward vector to the cell importance. The closest to 0, the less the angle will contribute to the cell importance."));
 
 UWorldPartitionRuntimeSpatialHashCell::UWorldPartitionRuntimeSpatialHashCell(const FObjectInitializer& ObjectInitializer)
 : Super(ObjectInitializer)
 , Level(0)
+, CachedSourceMinDistance(0.f)
 {}
 
 #if WITH_EDITOR
@@ -22,3 +30,46 @@ void UWorldPartitionRuntimeSpatialHashCell::PostDuplicate(bool bDuplicateForPIE)
 }
 #endif
 
+bool UWorldPartitionRuntimeSpatialHashCell::CacheStreamingSourceInfo(const FWorldPartitionStreamingSource& Source) const
+{
+	const bool bWasCacheDirtied = Super::CacheStreamingSourceInfo(Source);
+	
+	const float AngleContribution = FMath::Clamp(GRuntimeSpatialHashCellToSourceAngleContributionToCellImportance, 0.f, 1.f);
+	const float SqrDistance = FVector::DistSquared(Source.Location, Position);
+	float AngleFactor = 1.f;
+	if (!FMath::IsNearlyZero(AngleContribution))
+	{
+		const FVector2D SourceForward(Source.Rotation.Quaternion().GetForwardVector());
+		const FVector2D SourceToCell(Position - Source.Location);
+		const float Dot = FVector2D::DotProduct(SourceForward.GetSafeNormal(), SourceToCell.GetSafeNormal());
+		const float NormalizedAngle = FMath::Clamp(FMath::Abs(FMath::Acos(Dot) / PI), 0.f, 1.f);
+		AngleFactor = FMath::Pow(NormalizedAngle, AngleContribution);
+	}
+	// Modulate distance to cell by angle relative to source forward vector (to prioritize cells in front)
+	float ModulatedDistance = SqrDistance * AngleFactor;
+
+	// If cache was dirtied, use value, else use minimum with existing cached value
+	CachedSourceMinDistance = bWasCacheDirtied ? ModulatedDistance : FMath::Min(ModulatedDistance, CachedSourceMinDistance);
+
+	return bWasCacheDirtied;
+}
+
+int32 UWorldPartitionRuntimeSpatialHashCell::SortCompare(const UWorldPartitionRuntimeCell* InOther) const
+{
+	int32 Result = Super::SortCompare(InOther);
+	if (Result == 0)
+	{
+		const UWorldPartitionRuntimeSpatialHashCell* Other = Cast<const UWorldPartitionRuntimeSpatialHashCell>(InOther);
+		check(Other);
+		
+		// Level (higher value is higher prio)
+		Result = Other->Level - Level;
+		if (Result == 0)
+		{
+			// Closest distance (lower value is higher prio)
+			const float Diff = CachedSourceMinDistance - Other->CachedSourceMinDistance;
+			Result = Diff < 0.f ? -1 : (Diff > 0.f ? 1 : 0);
+		}
+	}
+	return Result;
+}
\ No newline at end of file
diff --git a/Engine/Source/Runtime/Engine/Private/WorldPartition/WorldPartitionStreamingPolicy.cpp b/Engine/Source/Runtime/Engine/Private/WorldPartition/WorldPartitionStreamingPolicy.cpp
--- a/Engine/Source/Runtime/Engine/Private/WorldPartition/WorldPartitionStreamingPolicy.cpp
+++ b/Engine/Source/Runtime/Engine/Private/WorldPartition/WorldPartitionStreamingPolicy.cpp
@@ -40,6 +40,8 @@ UWorldPartitionStreamingPolicy::UWorldPartitionStreamingPolicy(const FObjectInit
 
 void UWorldPartitionStreamingPolicy::UpdateStreamingSources()
 {
+	TRACE_CPUPROFILER_EVENT_SCOPE(UWorldPartitionStreamingPolicy::UpdateStreamingSources);
+
 	StreamingSources.Reset();
 
 	if (!WorldPartition->IsInitialized())
@@ -99,6 +101,11 @@ void UWorldPartitionStreamingPolicy::UpdateStreamingSources()
 			// Transform to Local
 			StreamingSource.Location = WorldToLocal.TransformPosition(StreamingSource.Location);
 			StreamingSource.Rotation = WorldToLocal.TransformRotation(StreamingSource.Rotation.Quaternion()).Rotator();
+			// If none is provided, default Streaming Source provider's priority to be less than those based on player controllers
+			if (StreamingSource.Priority == EStreamingSourcePriority::Default)
+			{
+				StreamingSource.Priority = EStreamingSourcePriority::Low;
+			}
 
 			StreamingSources.Add(StreamingSource);
 		}
@@ -107,8 +114,8 @@ void UWorldPartitionStreamingPolicy::UpdateStreamingSources()
 
 void UWorldPartitionStreamingPolicy::UpdateStreamingState()
 {
-	TSet<const UWorldPartitionRuntimeCell*> LoadStreamingCells;
-	TSet<const UWorldPartitionRuntimeCell*> ActivateStreamingCells;
+	TRACE_CPUPROFILER_EVENT_SCOPE(UWorldPartitionStreamingPolicy::UpdateStreamingState);
+
 	UWorld* World = WorldPartition->GetWorld();
 	check(World && World->IsGameWorld());
 	
@@ -120,8 +127,9 @@ void UWorldPartitionStreamingPolicy::UpdateStreamingState()
 		// Load all cells on server (do this once)
 		if (!bIsServerLoadingDone)
 		{
-			WorldPartition->RuntimeHash->GetAllStreamingCells(ActivateStreamingCells);
-			TSet<const UWorldPartitionRuntimeCell*> ToActivateCells = ActivateStreamingCells.Difference(ActivatedCells);
+			TSet<const UWorldPartitionRuntimeCell*> AllStreamingCells;
+			WorldPartition->RuntimeHash->GetAllStreamingCells(AllStreamingCells);
+			TSet<const UWorldPartitionRuntimeCell*> ToActivateCells = AllStreamingCells.Difference(ActivatedCells);
 
 			// Mark Runtime Cells to be Always Loaded
 			for (const UWorldPartitionRuntimeCell* Cell : ToActivateCells)
@@ -142,10 +150,17 @@ void UWorldPartitionStreamingPolicy::UpdateStreamingState()
 			return;
 		}
 
+		UWorldPartitionRuntimeHash::FStreamingSourceCells ActivateCells;
+		UWorldPartitionRuntimeHash::FStreamingSourceCells LoadCells;
+		TSet<const UWorldPartitionRuntimeCell*>& LoadStreamingCells = LoadCells.GetCells();
+		TSet<const UWorldPartitionRuntimeCell*>& ActivateStreamingCells = ActivateCells.GetCells();
+
 		// When uninitializing, UpdateStreamingState is called, but we don't want any cells to be loaded
 		if (WorldPartition->IsInitialized())
 		{
-			WorldPartition->RuntimeHash->GetStreamingCells(StreamingSources, ActivateStreamingCells, LoadStreamingCells);
+			UWorldPartitionRuntimeCell::DirtyStreamingSourceCacheEpoch();
+			WorldPartition->RuntimeHash->GetStreamingCells(StreamingSources, ActivateCells, LoadCells);
+
 			// Activation superseeds Loading
 			LoadStreamingCells = LoadStreamingCells.Difference(ActivateStreamingCells);
 		}
diff --git a/Engine/Source/Runtime/Engine/Private/WorldPartition/WorldPartitionSubsystem.cpp b/Engine/Source/Runtime/Engine/Private/WorldPartition/WorldPartitionSubsystem.cpp
--- a/Engine/Source/Runtime/Engine/Private/WorldPartition/WorldPartitionSubsystem.cpp
+++ b/Engine/Source/Runtime/Engine/Private/WorldPartition/WorldPartitionSubsystem.cpp
@@ -266,13 +266,13 @@ void UWorldPartitionSubsystem::Draw(UCanvas* Canvas, class APlayerController* PC
 			float MaxTextWidth = 0;
 			for (const FWorldPartitionStreamingSource& StreamingSource : *StreamingSources)
 			{
-				FWorldPartitionDebugHelper::DrawText(Canvas, StreamingSource.Name.ToString(), GEngine->GetTinyFont(), FColor::White, Pos, &MaxTextWidth);
+				FWorldPartitionDebugHelper::DrawText(Canvas, StreamingSource.Name.ToString(), GEngine->GetSmallFont(), StreamingSource.GetDebugColor(), Pos, &MaxTextWidth);
 			}
 			Pos = CurrentOffset + FVector2D(MaxTextWidth + 10, 0.f);
 			for (const FWorldPartitionStreamingSource& StreamingSource : *StreamingSources)
 			{
-				const FString Text = FString::Printf(TEXT("Pos: %s | Rot: %s"), *StreamingSource.Location.ToString(), *StreamingSource.Rotation.ToString());
-				FWorldPartitionDebugHelper::DrawText(Canvas, Text, GEngine->GetTinyFont(), FColor::White, Pos);
+				const FString Text = FString::Printf(TEXT("Priority: %d | Pos: %s | Rot: %s"), StreamingSource.Priority, *StreamingSource.Location.ToString(), *StreamingSource.Rotation.ToString());
+				FWorldPartitionDebugHelper::DrawText(Canvas, Text, GEngine->GetSmallFont(), FColor::White, Pos);
 			}
 			CurrentOffset.Y = Pos.Y;
 		}
diff --git a/Engine/Source/Runtime/Engine/Public/WorldPartition/WorldPartitionRuntimeCell.h b/Engine/Source/Runtime/Engine/Public/WorldPartition/WorldPartitionRuntimeCell.h
--- a/Engine/Source/Runtime/Engine/Public/WorldPartition/WorldPartitionRuntimeCell.h
+++ b/Engine/Source/Runtime/Engine/Public/WorldPartition/WorldPartitionRuntimeCell.h
@@ -114,12 +114,19 @@ class UWorldPartitionRuntimeCell : public UObject
 	virtual FLinearColor GetDebugColor() const { return FLinearColor::Black; }
 	virtual bool IsAlwaysLoaded() const { return bIsAlwaysLoaded; }
 	virtual void SetIsAlwaysLoaded(bool bInIsAlwaysLoaded) { bIsAlwaysLoaded = bInIsAlwaysLoaded; }
+	virtual void SetPriority(int32 InPriority) { Priority = InPriority; }
 	bool HasDataLayers() const { return !DataLayers.IsEmpty(); }
 	const TArray<FName>& GetDataLayers() const { return DataLayers; }
 	virtual EStreamingStatus GetStreamingStatus() const { return LEVEL_Unloaded; }
 	virtual bool IsLoading() const { return false; }
 	virtual const FString& GetDebugName() const { return DebugName; }
 	virtual bool IsDebugShown() const;
+	virtual int32 SortCompare(const UWorldPartitionRuntimeCell* Other) const;
+
+	/** Caches information on streaming source that will be used later on to sort cell. Returns true if cache was reset, else returns false. */
+	virtual bool CacheStreamingSourceInfo(const struct FWorldPartitionStreamingSource& Source) const;
+
+	static void DirtyStreamingSourceCacheEpoch() { ++UWorldPartitionRuntimeCell::StreamingSourceCacheEpoch; }
 
 #if WITH_EDITOR
 	void SetDataLayers(const TArray<const UDataLayer*>& InDataLayers);
@@ -163,4 +170,16 @@ private:
 
 	UPROPERTY()
 	FString DebugName;
+
+	// Custom Priority
+	UPROPERTY()
+	int32 Priority;
+
+	// Source Priority
+	mutable int32 CachedSourcePriority;
+
+	// Epoch used to dirty cache
+	mutable int32 CachedSourceInfoEpoch;
+
+	static int32 StreamingSourceCacheEpoch;
 };
diff --git a/Engine/Source/Runtime/Engine/Public/WorldPartition/WorldPartitionRuntimeHash.h b/Engine/Source/Runtime/Engine/Public/WorldPartition/WorldPartitionRuntimeHash.h
--- a/Engine/Source/Runtime/Engine/Public/WorldPartition/WorldPartitionRuntimeHash.h
+++ b/Engine/Source/Runtime/Engine/Public/WorldPartition/WorldPartitionRuntimeHash.h
@@ -41,11 +41,27 @@ class ENGINE_API UWorldPartitionRuntimeHash : public UObject
 	void OnEndPlay();
 #endif
 
+	class FStreamingSourceCells
+	{
+	public:
+		void AddCell(const UWorldPartitionRuntimeCell* InCell, const FWorldPartitionStreamingSource& InSource)
+		{
+			InCell->CacheStreamingSourceInfo(InSource);
+			Cells.Add(InCell);
+		}
+
+		int32 Num() const { return Cells.Num(); }
+		TSet<const UWorldPartitionRuntimeCell*>& GetCells() { return Cells; }
+
+	private:
+		TSet<const UWorldPartitionRuntimeCell*> Cells;
+	};
+
 	// Streaming interface
 	virtual int32 GetAllStreamingCells(TSet<const UWorldPartitionRuntimeCell*>& Cells, bool bIncludeDataLayers = false) const { return 0; }
 	virtual bool GetStreamingCells(const FWorldPartitionStreamingQuerySource& QuerySource, TSet<const UWorldPartitionRuntimeCell*>& OutCells) const { return false; }
-	virtual bool GetStreamingCells(const TArray<FWorldPartitionStreamingSource>& Sources, TSet<const UWorldPartitionRuntimeCell*>& OutActivateCells, TSet<const UWorldPartitionRuntimeCell*>& OutLoadCells) const { return false; };
-	virtual void SortStreamingCellsByImportance(const TSet<const UWorldPartitionRuntimeCell*>& InCells, const TArray<FWorldPartitionStreamingSource>& InSources, TArray<const UWorldPartitionRuntimeCell*, TInlineAllocator<256>>& OutSortedCells) const {}
+	virtual bool GetStreamingCells(const TArray<FWorldPartitionStreamingSource>& Sources, UWorldPartitionRuntimeHash::FStreamingSourceCells& OutActivateCells, UWorldPartitionRuntimeHash::FStreamingSourceCells& OutLoadCells) const { return false; };
+	virtual void SortStreamingCellsByImportance(const TSet<const UWorldPartitionRuntimeCell*>& InCells, const TArray<FWorldPartitionStreamingSource>& InSources, TArray<const UWorldPartitionRuntimeCell*, TInlineAllocator<256>>& OutSortedCells) const;
 
 	/* Returns desired footprint that Draw2D should take relative to given Canvas size (the value can exceed the given size).
 	 * UWorldPartitionSubSystem will re-adapt the size relative to all others UWorldPartitionRuntimeHash and provide the correct size to Draw2D.
diff --git a/Engine/Source/Runtime/Engine/Public/WorldPartition/WorldPartitionRuntimeSpatialHash.h b/Engine/Source/Runtime/Engine/Public/WorldPartition/WorldPartitionRuntimeSpatialHash.h
--- a/Engine/Source/Runtime/Engine/Public/WorldPartition/WorldPartitionRuntimeSpatialHash.h
+++ b/Engine/Source/Runtime/Engine/Public/WorldPartition/WorldPartitionRuntimeSpatialHash.h
@@ -125,7 +125,7 @@ struct FSpatialHashStreamingGrid
 
 	// Used by PIE/Game
 	void GetCells(const FWorldPartitionStreamingQuerySource& QuerySource, TSet<const UWorldPartitionRuntimeCell*>& OutCells) const;
-	void GetCells(const TArray<FWorldPartitionStreamingSource>& Sources, const class UDataLayerSubsystem* DataLayerSubsystem, TSet<const UWorldPartitionRuntimeCell*>& OutActivateCells, TSet<const UWorldPartitionRuntimeCell*>& OutLoadCells) const;
+	void GetCells(const TArray<FWorldPartitionStreamingSource>& Sources, const class UDataLayerSubsystem* DataLayerSubsystem, UWorldPartitionRuntimeHash::FStreamingSourceCells& OutActivateCells, UWorldPartitionRuntimeHash::FStreamingSourceCells& OutLoadCells) const;
 	void GetAlwaysLoadedCells(const UDataLayerSubsystem* DataLayerSubsystem, TSet<const UWorldPartitionRuntimeCell*>& OutActivateCells, TSet<const UWorldPartitionRuntimeCell*>& OutLoadCells) const;
 	void Draw2D(UCanvas* Canvas, const TArray<FWorldPartitionStreamingSource>& Sources, const FBox& Region, const FBox2D& GridScreenBounds, TFunctionRef<FVector2D(const FVector2D&)> WorldToScreen, const TMap<FName, FColor>& DataLayerDebugColors) const;
 	void Draw3D(UWorld* World, const TArray<FWorldPartitionStreamingSource>& Sources, const FTransform& Transform, const TMap<FName, FColor>& DataLayerDebugColors) const;
@@ -219,9 +219,8 @@ public:
 	// streaming interface
 	virtual int32 GetAllStreamingCells(TSet<const UWorldPartitionRuntimeCell*>& Cells, bool bIncludeDataLayers = false) const override;
 	virtual bool GetStreamingCells(const FWorldPartitionStreamingQuerySource& QuerySource, TSet<const UWorldPartitionRuntimeCell*>& OutCells) const override;
-	virtual bool GetStreamingCells(const TArray<FWorldPartitionStreamingSource>& Sources, TSet<const UWorldPartitionRuntimeCell*>& OutActivateCells, TSet<const UWorldPartitionRuntimeCell*>& OutLoadCells) const override;
-	virtual void SortStreamingCellsByImportance(const TSet<const UWorldPartitionRuntimeCell*>& InCells, const TArray<FWorldPartitionStreamingSource>& InSources, TArray<const UWorldPartitionRuntimeCell*, TInlineAllocator<256>>& OutSortedCells) const override;
-	
+	virtual bool GetStreamingCells(const TArray<FWorldPartitionStreamingSource>& Sources, UWorldPartitionRuntimeHash::FStreamingSourceCells& OutActivateCells, UWorldPartitionRuntimeHash::FStreamingSourceCells& OutLoadCells) const override;
+		
 	static FString GetCellCoordString(const FIntVector& InCellGlobalCoord);
 
 protected:
diff --git a/Engine/Source/Runtime/Engine/Public/WorldPartition/WorldPartitionRuntimeSpatialHashCell.h b/Engine/Source/Runtime/Engine/Public/WorldPartition/WorldPartitionRuntimeSpatialHashCell.h
--- a/Engine/Source/Runtime/Engine/Public/WorldPartition/WorldPartitionRuntimeSpatialHashCell.h
+++ b/Engine/Source/Runtime/Engine/Public/WorldPartition/WorldPartitionRuntimeSpatialHashCell.h
@@ -13,6 +13,8 @@ class UWorldPartitionRuntimeSpatialHashCell : public UWorldPartitionRuntimeCell
 #if WITH_EDITOR
 	virtual void PostDuplicate(bool bDuplicateForPIE) override;
 #endif
+	virtual bool CacheStreamingSourceInfo(const FWorldPartitionStreamingSource& Source) const override;
+	virtual int32 SortCompare(const UWorldPartitionRuntimeCell* InOther) const override;
 
 	UPROPERTY()
 	FVector Position;
@@ -20,8 +22,8 @@ class UWorldPartitionRuntimeSpatialHashCell : public UWorldPartitionRuntimeCell
 	UPROPERTY()
 	int32 Level;
 
-	UPROPERTY()
-	int32 Priority;
+	// Computed and cached value used by SortCompare to sort Cells
+	mutable float CachedSourceMinDistance;
 
 #if WITH_EDITORONLY_DATA
 	UPROPERTY()
diff --git a/Engine/Source/Runtime/Engine/Public/WorldPartition/WorldPartitionStreamingSource.h b/Engine/Source/Runtime/Engine/Public/WorldPartition/WorldPartitionStreamingSource.h
--- a/Engine/Source/Runtime/Engine/Public/WorldPartition/WorldPartitionStreamingSource.h
+++ b/Engine/Source/Runtime/Engine/Public/WorldPartition/WorldPartitionStreamingSource.h
@@ -2,6 +2,8 @@
 #pragma once
 
 #include "CoreMinimal.h"
+#include "Math/Color.h"
+#include "Math/RandomStream.h"
 #include "WorldPartitionStreamingSource.generated.h"
 
 /**
@@ -63,25 +65,44 @@ struct FWorldPartitionStreamingQuerySource
 	bool bSpatialQuery;
 };
 
+/**
+ * Streaming Source Priority
+ */
+enum class EStreamingSourcePriority : int32
+{
+	Highest = INT_MIN,
+	High = -4096,
+	Normal = 0,
+	Low = 4096,
+	Lowest = INT_MAX,
+	Default = Normal
+};
+
 /**
  * Structure containing all properties required to stream from a source
  */
 struct ENGINE_API FWorldPartitionStreamingSource
 {
 	FWorldPartitionStreamingSource()
+		: Priority(EStreamingSourcePriority::Default)
 	{}
 
-	FWorldPartitionStreamingSource(FName InName, const FVector& InLocation, const FRotator& InRotation, EStreamingSourceTargetState InTargetState)
+	FWorldPartitionStreamingSource(FName InName, const FVector& InLocation, const FRotator& InRotation, EStreamingSourceTargetState InTargetState, EStreamingSourcePriority InPriority = EStreamingSourcePriority::Default)
 		: Name(InName)
 		, Location(InLocation)
 		, Rotation(InRotation)
 		, TargetState(InTargetState)
-	{}
+		, Priority(InPriority)
+	{
+	}
+
+	FColor GetDebugColor() const { return FColor::MakeRedToGreenColorFromScalar(FRandomStream(Name).GetFraction()); }
 
 	FName Name;
 	FVector Location;
 	FRotator Rotation;
 	EStreamingSourceTargetState TargetState;
+	EStreamingSourcePriority Priority;
 };
 
 /**
