HLOD actors are referencing runtime cells using their FName
Now storing the source cell as a TSoftObjectPtr<>
Removed custom HLOD rename code in UWorldPartitionRenameCommandlet

Tested PIE, -game, cooked & WP rename commandlet

diff --git a/Engine/Source/Editor/UnrealEd/Private/Commandlets/WorldPartitionRenameCommandlet.cpp b/Engine/Source/Editor/UnrealEd/Private/Commandlets/WorldPartitionRenameCommandlet.cpp
--- a/Engine/Source/Editor/UnrealEd/Private/Commandlets/WorldPartitionRenameCommandlet.cpp
+++ b/Engine/Source/Editor/UnrealEd/Private/Commandlets/WorldPartitionRenameCommandlet.cpp
@@ -236,15 +236,6 @@ int32 UWorldPartitionRenameCommandlet::Main(const FString& Params)
 		PackagesToSave.Add(ActorDescIterator->GetActor()->GetPackage());
 	}
 
-	// Fixup cell names referenced by HLODs
-	for (FActorDescList::TIterator<AWorldPartitionHLOD> HLODActorDescIterator(WorldPartition); HLODActorDescIterator; ++HLODActorDescIterator)
-	{
-		AWorldPartitionHLOD* HLODActor = CastChecked<AWorldPartitionHLOD>(HLODActorDescIterator->GetActor());
-		FString OldCellName = HLODActor->GetCellName().ToString();
-		FString NewCellName = OldCellName.Replace(*OldWorldName, *NewWorldName);
-		HLODActor->SetCellName(FName(NewCellName));
-	}
-
 	PackagesToSave.Add(World->GetPackage());
 
 	// Replace old map package with a redirector to the new map package
diff --git a/Engine/Source/Runtime/Engine/Private/WorldPartition/HLOD/HLODActor.cpp b/Engine/Source/Runtime/Engine/Private/WorldPartition/HLOD/HLODActor.cpp
--- a/Engine/Source/Runtime/Engine/Private/WorldPartition/HLOD/HLODActor.cpp
+++ b/Engine/Source/Runtime/Engine/Private/WorldPartition/HLOD/HLODActor.cpp
@@ -49,20 +49,6 @@ void AWorldPartitionHLOD::EndPlay(const EEndPlayReason::Type EndPlayReason)
 	Super::EndPlay(EndPlayReason);
 }
 
-void AWorldPartitionHLOD::PostLoad()
-{
-	Super::PostLoad();
-
-#if WITH_EDITOR
-	if (CellName == NAME_None)
-	{
-		// Prior to the addition of the cell name member, the cell name was actually stored in the label
-		CellName = *GetActorLabel();
-	}
-#endif
-}
-
-
 void AWorldPartitionHLOD::RerunConstructionScripts()
 {}
 
@@ -133,9 +119,9 @@ const TArray<FGuid>& AWorldPartitionHLOD::GetSubActors() const
 	return SubActors;
 }
 
-void AWorldPartitionHLOD::SetCellName(FName InCellName)
+void AWorldPartitionHLOD::SetSourceCell(const TSoftObjectPtr<UWorldPartitionRuntimeCell>& InSourceCell)
 {
-	CellName = InCellName;
+	SourceCell = InSourceCell;
 }
 
 const FBox& AWorldPartitionHLOD::GetHLODBounds() const
diff --git a/Engine/Source/Runtime/Engine/Private/WorldPartition/HLOD/HLODBuilder.cpp b/Engine/Source/Runtime/Engine/Private/WorldPartition/HLOD/HLODBuilder.cpp
--- a/Engine/Source/Runtime/Engine/Private/WorldPartition/HLOD/HLODBuilder.cpp
+++ b/Engine/Source/Runtime/Engine/Private/WorldPartition/HLOD/HLODBuilder.cpp
@@ -469,9 +469,13 @@ TArray<AWorldPartitionHLOD*> FHLODBuilderUtilities::CreateHLODActors(FHLODCreati
 			SpawnParams.NameMode = FActorSpawnParameters::ESpawnActorNameMode::Required_Fatal;
 			HLODActor = InCreationParams.WorldPartition->GetWorld()->SpawnActor<AWorldPartitionHLOD>(SpawnParams);
 
+			FSoftObjectPath WorldPartitionPath(InCreationParams.WorldPartition);
+			FString CellObjectPath = FString::Printf(TEXT("%s.%s"), *WorldPartitionPath.ToString(), *InCreationParams.CellName.ToString());
+			TSoftObjectPtr<UWorldPartitionRuntimeCell> RuntimeCell(CellObjectPath);
+
 			HLODActor->SetActorLabel(FString::Printf(TEXT("HLOD%d_%s"), InCreationParams.HLODLevel, *InCreationParams.CellName.ToString()));
 			HLODActor->SetFolderPath(*FString::Printf(TEXT("HLOD/HLOD%d"), InCreationParams.HLODLevel));
-			HLODActor->SetCellName(InCreationParams.CellName);
+			HLODActor->SetSourceCell(RuntimeCell);
 			HLODActor->SetSubActorsHLODLayer(HLODLayer);
 			HLODActor->SetGridIndices(InCreationParams.GridIndexX, InCreationParams.GridIndexY, InCreationParams.GridIndexZ);
 
diff --git a/Engine/Source/Runtime/Engine/Private/WorldPartition/HLOD/HLODSubsystem.cpp b/Engine/Source/Runtime/Engine/Private/WorldPartition/HLOD/HLODSubsystem.cpp
--- a/Engine/Source/Runtime/Engine/Private/WorldPartition/HLOD/HLODSubsystem.cpp
+++ b/Engine/Source/Runtime/Engine/Private/WorldPartition/HLOD/HLODSubsystem.cpp
@@ -99,7 +99,7 @@ void UHLODSubsystem::OnWorldPartitionRegistered(UWorldPartition* InWorldPartitio
 	// Build cell to HLOD mapping
 	for (const UWorldPartitionRuntimeCell* Cell : StreamingCells)
 	{
-		CellsHLODMapping.Emplace(Cell->GetFName());
+		CellsHLODMapping.Emplace(Cell);
 	}
 }
 
@@ -113,11 +113,11 @@ void UHLODSubsystem::RegisterHLODActor(AWorldPartitionHLOD* InWorldPartitionHLOD
 {
 	TRACE_CPUPROFILER_EVENT_SCOPE(UHLODSubsystem::RegisterHLODActor);
 
-	FName CellName = InWorldPartitionHLOD->GetCellName();
-	FCellHLODMapping* CellHLODs = CellsHLODMapping.Find(CellName);
+	const TSoftObjectPtr<UWorldPartitionRuntimeCell>& RuntimeCell = InWorldPartitionHLOD->GetSourceCell();
+	FCellHLODMapping* CellHLODs = CellsHLODMapping.Find(RuntimeCell);
 
 #if WITH_EDITOR
-	UE_LOG(LogHLODSubsystem, Verbose, TEXT("Registering HLOD %s (%s) for cell %s"), *InWorldPartitionHLOD->GetActorLabel(), *InWorldPartitionHLOD->GetActorGuid().ToString(), *CellName.ToString());
+	UE_LOG(LogHLODSubsystem, Verbose, TEXT("Registering HLOD %s (%s) for cell %s"), *InWorldPartitionHLOD->GetActorLabel(), *InWorldPartitionHLOD->GetActorGuid().ToString(), *RuntimeCell.ToString());
 #endif
 
 	if (CellHLODs)
@@ -127,7 +127,7 @@ void UHLODSubsystem::RegisterHLODActor(AWorldPartitionHLOD* InWorldPartitionHLOD
 	}
 	else
 	{
-		UE_LOG(LogHLODSubsystem, Warning, TEXT("Found HLOD referencing nonexistent cell '%s'"), *CellName.ToString());
+		UE_LOG(LogHLODSubsystem, Verbose, TEXT("Found HLOD referencing nonexistent cell '%s'"), *RuntimeCell.ToString());
 		InWorldPartitionHLOD->SetVisibility(false);
 	}
 }
@@ -136,12 +136,13 @@ void UHLODSubsystem::UnregisterHLODActor(AWorldPartitionHLOD* InWorldPartitionHL
 {
 	TRACE_CPUPROFILER_EVENT_SCOPE(UHLODSubsystem::UnregisterHLODActor);
 
+	const TSoftObjectPtr<UWorldPartitionRuntimeCell>& RuntimeCell = InWorldPartitionHLOD->GetSourceCell();
+	FCellHLODMapping* CellHLODs = CellsHLODMapping.Find(RuntimeCell);
+
 #if WITH_EDITOR
-	UE_LOG(LogHLODSubsystem, Verbose, TEXT("Unregistering HLOD %s (%s) for cell %s"), *InWorldPartitionHLOD->GetActorLabel(), *InWorldPartitionHLOD->GetActorGuid().ToString(), *InWorldPartitionHLOD->GetCellName().ToString());
+	UE_LOG(LogHLODSubsystem, Verbose, TEXT("Unregistering HLOD %s (%s) for cell %s"), *InWorldPartitionHLOD->GetActorLabel(), *InWorldPartitionHLOD->GetActorGuid().ToString(), *RuntimeCell.ToString());
 #endif
 
-	FName CellName = InWorldPartitionHLOD->GetCellName();
-	FCellHLODMapping* CellHLODs = CellsHLODMapping.Find(CellName);
 	if (CellHLODs)
 	{
 		int32 NumRemoved = CellHLODs->LoadedHLODs.Remove(InWorldPartitionHLOD);
@@ -151,7 +152,7 @@ void UHLODSubsystem::UnregisterHLODActor(AWorldPartitionHLOD* InWorldPartitionHL
 
 void UHLODSubsystem::OnCellShown(const UWorldPartitionRuntimeCell* InCell)
 {
-	FCellHLODMapping& CellHLODs = CellsHLODMapping.FindChecked(InCell->GetFName());
+	FCellHLODMapping& CellHLODs = CellsHLODMapping.FindChecked(InCell);
 	CellHLODs.bIsCellVisible = true;
 
 #if WITH_EDITOR
@@ -169,7 +170,7 @@ void UHLODSubsystem::OnCellShown(const UWorldPartitionRuntimeCell* InCell)
 
 void UHLODSubsystem::OnCellHidden(const UWorldPartitionRuntimeCell* InCell)
 {
-	FCellHLODMapping& CellHLODs = CellsHLODMapping.FindChecked(InCell->GetFName());
+	FCellHLODMapping& CellHLODs = CellsHLODMapping.FindChecked(InCell);
 	CellHLODs.bIsCellVisible = false;
 
 #if WITH_EDITOR
diff --git a/Engine/Source/Runtime/Engine/Public/WorldPartition/HLOD/HLODActor.h b/Engine/Source/Runtime/Engine/Public/WorldPartition/HLOD/HLODActor.h
--- a/Engine/Source/Runtime/Engine/Public/WorldPartition/HLOD/HLODActor.h
+++ b/Engine/Source/Runtime/Engine/Public/WorldPartition/HLOD/HLODActor.h
@@ -19,7 +19,7 @@ class ENGINE_API AWorldPartitionHLOD : public AActor
 public:
 	void SetVisibility(bool bInVisible);
 
-	inline FName GetCellName() const { return CellName; }
+	inline const TSoftObjectPtr<UWorldPartitionRuntimeCell>& GetSourceCell() const { return SourceCell; }
 	inline uint32 GetLODLevel() const { return LODLevel; }
 
 	virtual bool IsHLODRelevant() const override { return true; }
@@ -47,7 +47,7 @@ public:
 		OutGridIndexZ = GridIndexZ;
 	}
 
-	void SetCellName(FName InCellName);
+	void SetSourceCell(const TSoftObjectPtr<UWorldPartitionRuntimeCell>& InSourceCell);
 	inline void SetLODLevel(uint32 InLODLevel) { LODLevel = InLODLevel; }
 
 	const FBox& GetHLODBounds() const;
@@ -59,7 +59,6 @@ public:
 
 protected:
 	//~ Begin UObject Interface.
-	virtual void PostLoad() override;
 	virtual void RerunConstructionScripts() override;
 	//~ End UObject Interface.
 
@@ -107,7 +106,7 @@ private:
 	uint32 LODLevel;
 
 	UPROPERTY()
-	FName CellName;
+	TSoftObjectPtr<UWorldPartitionRuntimeCell> SourceCell;
 };
 
 DEFINE_ACTORDESC_TYPE(AWorldPartitionHLOD, FHLODActorDesc);
diff --git a/Engine/Source/Runtime/Engine/Public/WorldPartition/HLOD/HLODSubsystem.h b/Engine/Source/Runtime/Engine/Public/WorldPartition/HLOD/HLODSubsystem.h
--- a/Engine/Source/Runtime/Engine/Public/WorldPartition/HLOD/HLODSubsystem.h
+++ b/Engine/Source/Runtime/Engine/Public/WorldPartition/HLOD/HLODSubsystem.h
@@ -52,11 +52,12 @@ private:
 	{
 		bool						bIsCellVisible;
 		TSet<AWorldPartitionHLOD*>	LoadedHLODs;
+		
 
 		FCellHLODMapping() : bIsCellVisible(false)
 		{
 		}
 	};
 
-	TMap<FName, FCellHLODMapping> CellsHLODMapping;
+	TMap<TSoftObjectPtr<UWorldPartitionRuntimeCell>, FCellHLODMapping> CellsHLODMapping;
 };
