Fix a GC hang due to IsReadyForFinishDestroy never reports true for some static meshes. UStreamableRenderAsset::IsReadyForFinishDestroy has a side effect of ticking the pending streaming update and clearing the pointer when the update is done. So it needs to be called before HasPendingInitOrStreaming. Otherwise, some pending streaming updates may never finish during GC.

diff --git a/Engine/Source/Runtime/Engine/Private/StaticMesh.cpp b/Engine/Source/Runtime/Engine/Private/StaticMesh.cpp
--- a/Engine/Source/Runtime/Engine/Private/StaticMesh.cpp
+++ b/Engine/Source/Runtime/Engine/Private/StaticMesh.cpp
@@ -3813,7 +3813,7 @@ bool UStaticMesh::IsReadyForFinishDestroy()
 #endif
 
 	// Match BeginDestroy() by checking for HasPendingInitOrStreaming().
-	if (HasPendingInitOrStreaming() || !Super::IsReadyForFinishDestroy())
+	if (!Super::IsReadyForFinishDestroy() || HasPendingInitOrStreaming())
 	{
 		return false;
 	}
