Anim Sequence - Sequence editing option "remove frame 0 to frame X" leaves behind frame 0
#fix redo cropping fix - actually need to round the frame time to numbers rather than offset by one. Previous fix worked for cropping from end, but broke cropping from start.

diff --git a/Engine/Source/Runtime/Engine/Private/Animation/AnimSequenceHelpers.cpp b/Engine/Source/Runtime/Engine/Private/Animation/AnimSequenceHelpers.cpp
--- a/Engine/Source/Runtime/Engine/Private/Animation/AnimSequenceHelpers.cpp
+++ b/Engine/Source/Runtime/Engine/Private/Animation/AnimSequenceHelpers.cpp
@@ -773,10 +773,10 @@ bool AnimationData::Trim(UAnimSequence* InSequence, float TrimStart, float TrimE
 			return false;
 		}
 
-		const FFrameNumber StartFrameTrim = FrameRate.AsFrameNumber(TrimStart);
-		const FFrameNumber EndFrameTrim = FrameRate.AsFrameNumber(TrimEnd);
+		const FFrameNumber StartFrameTrim = FrameRate.AsFrameTime(TrimStart).RoundToFrame();
+		const FFrameNumber EndFrameTrim = FrameRate.AsFrameTime(TrimEnd).RoundToFrame();
 
-		const int32 StartTrimKeyIndex = StartFrameTrim.Value + 1;
+		const int32 StartTrimKeyIndex = StartFrameTrim.Value;
 		const int32 NumTrimmedFrames = EndFrameTrim.Value - StartFrameTrim.Value;
 		
 		IAnimationDataController::FScopedBracket ScopedBracket(Controller, LOCTEXT("TrimRawAnimation_Bracket", "Trimming Animation Track Data"));
@@ -887,7 +887,7 @@ void AnimationData::RemoveKeys(UAnimSequence* InSequence, int32 StartKeyIndex, i
 		const int32 NewNumberOfFrames = FMath::Max(NewNumberOfKeys - 1, 0);
 		const float NewSequenceLength = FrameRate.AsSeconds(NewNumberOfFrames);
 
-		const float StartFrameTime = FrameRate.AsSeconds(FMath::Max(StartKeyIndex - 1, 0));
+		const float StartFrameTime = FrameRate.AsSeconds(FMath::Max(StartKeyIndex, 0));
 		const float RemovedTime = NumKeysToRemove * FrameRate.AsInterval();
 		const float EndTime = StartFrameTime + RemovedTime;
 
