Removed LevelInstanceSubsystem ShouldLoadInstancesOnRegistration() / SetLoadInstancesOnRegistration()
Tested using internal demo in PIE, standalone game and a cooked game

diff --git a/Engine/Source/Editor/UnrealEd/Private/Commandlets/DerivedDataCacheCommandlet.cpp b/Engine/Source/Editor/UnrealEd/Private/Commandlets/DerivedDataCacheCommandlet.cpp
--- a/Engine/Source/Editor/UnrealEd/Private/Commandlets/DerivedDataCacheCommandlet.cpp
+++ b/Engine/Source/Editor/UnrealEd/Private/Commandlets/DerivedDataCacheCommandlet.cpp
@@ -260,12 +260,6 @@ void UDerivedDataCacheCommandlet::CacheWorldPackages(UWorld* World, uint8 Packag
 	bool bResult = true;
 	if (World->HasSubsystem<UWorldPartitionSubsystem>())
 	{
-		// Commandlets aren't loading level instances by default, change that behavior
-		if (ULevelInstanceSubsystem* LevelInstanceSubsystem = World->GetSubsystem<ULevelInstanceSubsystem>())
-		{
-			LevelInstanceSubsystem->SetLoadInstancesOnRegistration(true);
-		}
-
 		// Ensure the world has a valid world partition.
 		UWorldPartition* WorldPartition = World->GetWorldPartition();
 		check(WorldPartition);
diff --git a/Engine/Source/Editor/UnrealEd/Private/WorldPartition/WorldPartitionBuilder.cpp b/Engine/Source/Editor/UnrealEd/Private/WorldPartition/WorldPartitionBuilder.cpp
--- a/Engine/Source/Editor/UnrealEd/Private/WorldPartition/WorldPartitionBuilder.cpp
+++ b/Engine/Source/Editor/UnrealEd/Private/WorldPartition/WorldPartitionBuilder.cpp
@@ -114,12 +114,6 @@ bool UWorldPartitionBuilder::RunBuilder(UWorldPartitionBuilder* Builder, UWorld*
 	bool bResult = true;
 	if (World->HasSubsystem<UWorldPartitionSubsystem>())
 	{
-		// Commandlets aren't loading level instances by default, change that behavior
-		if (ULevelInstanceSubsystem* LevelInstanceSubsystem = World->GetSubsystem<ULevelInstanceSubsystem>())
-		{
-			LevelInstanceSubsystem->SetLoadInstancesOnRegistration(true);
-		}
-
 		// Ensure the world has a valid world partition.
 		UWorldPartition* WorldPartition = World->GetWorldPartition();
 		check(WorldPartition);
diff --git a/Engine/Source/Runtime/Engine/Private/LevelInstance/LevelInstanceActor.cpp b/Engine/Source/Runtime/Engine/Private/LevelInstance/LevelInstanceActor.cpp
--- a/Engine/Source/Runtime/Engine/Private/LevelInstance/LevelInstanceActor.cpp
+++ b/Engine/Source/Runtime/Engine/Private/LevelInstance/LevelInstanceActor.cpp
@@ -80,10 +80,7 @@ void ALevelInstance::PostRegisterAllComponents()
 	{
 		LevelInstanceID = LevelInstanceSubsystem->RegisterLevelInstance(this);
 
-		if (LevelInstanceSubsystem->ShouldLoadInstancesOnRegistration())
-		{
-			LoadLevelInstance();
-		}
+		LoadLevelInstance();
 
 #if WITH_EDITOR
 		// Make sure transformation is up to date after registration as its possible LevelInstance actor can get unregistered when editing properties
@@ -107,10 +104,7 @@ void ALevelInstance::PostUnregisterAllComponents()
 
 		LevelInstanceSubsystem->UnregisterLevelInstance(this);
 
-		if (LevelInstanceSubsystem->ShouldLoadInstancesOnRegistration())
-		{
-			UnloadLevelInstance();
-		}
+		UnloadLevelInstance();
 
 		// To avoid processing PostUnregisterAllComponents multiple times (BP Recompile is one use case)
 		LevelInstanceID = FLevelInstanceID();
diff --git a/Engine/Source/Runtime/Engine/Private/LevelInstance/LevelInstanceSubsystem.cpp b/Engine/Source/Runtime/Engine/Private/LevelInstance/LevelInstanceSubsystem.cpp
--- a/Engine/Source/Runtime/Engine/Private/LevelInstance/LevelInstanceSubsystem.cpp
+++ b/Engine/Source/Runtime/Engine/Private/LevelInstance/LevelInstanceSubsystem.cpp
@@ -64,8 +64,6 @@ void ULevelInstanceSubsystem::Initialize(FSubsystemCollectionBase& Collection)
 	{
 		FModuleManager::LoadModuleChecked<ILevelInstanceEditorModule>("LevelInstanceEditor");
 	}
-
-	bLoadInstancesOnRegistration = !IsRunningCommandlet();
 #endif
 }
 
@@ -1728,21 +1726,6 @@ void ULevelInstanceSubsystem::BlockUnloadLevelInstance(ALevelInstance* LevelInst
 	}
 }
 
-void ULevelInstanceSubsystem::SetLoadInstancesOnRegistration(bool bInLoadInstancesOnRegistration)
-{
-	bLoadInstancesOnRegistration = bInLoadInstancesOnRegistration;
-}
-
-#endif
-
-bool ULevelInstanceSubsystem::ShouldLoadInstancesOnRegistration() const
-{
-#if WITH_EDITOR
-	return bLoadInstancesOnRegistration;
-#else
-	return !IsRunningCommandlet();
 #endif
-}
 
 #undef LOCTEXT_NAMESPACE
-
diff --git a/Engine/Source/Runtime/Engine/Public/LevelInstance/LevelInstanceSubsystem.h b/Engine/Source/Runtime/Engine/Public/LevelInstance/LevelInstanceSubsystem.h
--- a/Engine/Source/Runtime/Engine/Public/LevelInstance/LevelInstanceSubsystem.h
+++ b/Engine/Source/Runtime/Engine/Public/LevelInstance/LevelInstanceSubsystem.h
@@ -43,7 +43,6 @@ public:
 	void RequestLoadLevelInstance(ALevelInstance* LevelInstanceActor, bool bUpdate);
 	void RequestUnloadLevelInstance(ALevelInstance* LevelInstanceActor);
 	bool IsLoaded(const ALevelInstance* LevelInstanceActor) const;
-	bool ShouldLoadInstancesOnRegistration() const;
 	void ForEachLevelInstanceAncestorsAndSelf(AActor* Actor, TFunctionRef<bool(ALevelInstance*)> Operation) const;
 
 #if WITH_EDITOR
@@ -90,8 +89,6 @@ public:
 	ALevelInstance* GetParentLevelInstance(const AActor* Actor) const;
 	void BlockLoadLevelInstance(ALevelInstance* LevelInstanceActor);
 	void BlockUnloadLevelInstance(ALevelInstance* LevelInstanceActor);
-
-	void SetLoadInstancesOnRegistration(bool bLoadInstancesOnRegistration);
 #endif
 
 private:
@@ -137,8 +134,6 @@ private:
 	TUniquePtr<FLevelsToRemoveScope> LevelsToRemoveScope;
 	
 	TMap<FName, FLevelInstanceEdit> LevelInstanceEdits;
-
-	bool bLoadInstancesOnRegistration;
 #endif
 
 	struct FLevelInstance
