DataLayerSubsystem: Expose Active Layer Names to BP

diff --git a/Engine/Source/Runtime/Engine/Public/WorldPartition/DataLayer/DataLayerSubsystem.h b/Engine/Source/Runtime/Engine/Public/WorldPartition/DataLayer/DataLayerSubsystem.h
--- a/Engine/Source/Runtime/Engine/Public/WorldPartition/DataLayer/DataLayerSubsystem.h
+++ b/Engine/Source/Runtime/Engine/Public/WorldPartition/DataLayer/DataLayerSubsystem.h
@@ -77,6 +77,12 @@ public:
 	UDataLayer* GetDataLayerFromName(FName InDataLayerName) const;
 	//~ End Blueprint callable functions
 
+	UFUNCTION(BlueprintCallable, Category = DataLayers)
+	const TSet<FName>& GetActiveDataLayerNames() const { return ActiveDataLayerNames; }
+
+	UFUNCTION(BlueprintCallable, Category = DataLayers)
+	const TSet<FName>& GetLoadedDataLayerNames() const { return LoadedDataLayerNames; }
+
 	void SetDataLayerState(const UDataLayer* InDataLayer, EDataLayerState InState);
 	void SetDataLayerStateByName(const FName& InDataLayerName, EDataLayerState InState);
 	EDataLayerState GetDataLayerState(const UDataLayer* InDataLayer) const;
