Added Nanite Proxy Percent option to MeshMergeSettings / MeshProxySettings

diff --git a/Engine/Source/Developer/MeshMergeUtilities/Private/MeshMergeUtilities.cpp b/Engine/Source/Developer/MeshMergeUtilities/Private/MeshMergeUtilities.cpp
--- a/Engine/Source/Developer/MeshMergeUtilities/Private/MeshMergeUtilities.cpp
+++ b/Engine/Source/Developer/MeshMergeUtilities/Private/MeshMergeUtilities.cpp
@@ -3095,8 +3095,8 @@ void FMeshMergeUtilities::MergeComponentsToStaticMesh(const TArray<UPrimitiveCom
 
 		// Nanite settings
 		StaticMesh->NaniteSettings.bEnabled = InSettings.bGenerateNaniteEnabledMesh;
+		StaticMesh->NaniteSettings.PercentTriangles = InSettings.NaniteProxyTrianglePercent * 0.01f;
 		StaticMesh->NaniteSettings.PositionPrecision = MIN_int32;
-		StaticMesh->NaniteSettings.PercentTriangles = 1.0f;	// For now, assume meshes generated here are low poly enough to be used directly as the nanite proxy.
 
 #if WITH_EDITOR
 		//If we are running the automation test
diff --git a/Engine/Source/Developer/MeshMergeUtilities/Private/ProxyGenerationProcessor.cpp b/Engine/Source/Developer/MeshMergeUtilities/Private/ProxyGenerationProcessor.cpp
--- a/Engine/Source/Developer/MeshMergeUtilities/Private/ProxyGenerationProcessor.cpp
+++ b/Engine/Source/Developer/MeshMergeUtilities/Private/ProxyGenerationProcessor.cpp
@@ -290,8 +290,8 @@ void FProxyGenerationProcessor::ProcessJob(const FGuid& JobGuid, FProxyGeneratio
 
 	// Nanite settings
 	StaticMesh->NaniteSettings.bEnabled = Data->MergeData->InProxySettings.bGenerateNaniteEnabledMesh;
+	StaticMesh->NaniteSettings.PercentTriangles = Data->MergeData->InProxySettings.NaniteProxyTrianglePercent * 0.01f;
 	StaticMesh->NaniteSettings.PositionPrecision = MIN_int32;
-	StaticMesh->NaniteSettings.PercentTriangles = 1.0f;	// For now, assume meshes generated here are low poly enough to be used directly as the nanite proxy.
 
 	//Set the Imported version before calling the build
 	StaticMesh->ImportVersion = EImportStaticMeshVersion::LastVersion;
diff --git a/Engine/Source/Runtime/Engine/Classes/Engine/MeshMerging.h b/Engine/Source/Runtime/Engine/Classes/Engine/MeshMerging.h
--- a/Engine/Source/Runtime/Engine/Classes/Engine/MeshMerging.h
+++ b/Engine/Source/Runtime/Engine/Classes/Engine/MeshMerging.h
@@ -324,9 +324,13 @@ struct FMeshProxySettings
 	uint8 bGenerateLightmapUVs:1;
 
 	/** Whether to generate a nanite-enabled mesh */
-	UPROPERTY(BlueprintReadWrite, EditAnywhere, Category = ProxySettings)
+	UPROPERTY(BlueprintReadWrite, EditAnywhere, Category = NaniteSettings)
 	uint8 bGenerateNaniteEnabledMesh : 1;
 
+	/** Percentage of triangles to reduce down to for generating a coarse proxy mesh from the Nanite mesh */
+	UPROPERTY(BlueprintReadWrite, EditAnywhere, AdvancedDisplay, Category = NaniteSettings, meta = (EditConditionHides, EditCondition = "bGenerateNaniteEnabledMesh", ClampMin = 0, ClampMax = 100))
+	float NaniteProxyTrianglePercent;
+
 	/** Default settings. */
 	FMeshProxySettings()
 		: ScreenSize(300)
@@ -361,6 +365,7 @@ struct FMeshProxySettings
 		, bAllowVertexColors(false)
 		, bGenerateLightmapUVs(false)
 		, bGenerateNaniteEnabledMesh(false)
+		, NaniteProxyTrianglePercent(100)
 	{
 		MaterialSettings.MaterialMergeType = EMaterialMergeType::MaterialMergeType_Simplygon;
 	}
@@ -391,7 +396,8 @@ struct FMeshProxySettings
 			&& bCreateCollision == Other.bCreateCollision
 			&& bAllowVertexColors == Other.bAllowVertexColors
 			&& bGenerateLightmapUVs == Other.bGenerateLightmapUVs
-			&& bGenerateNaniteEnabledMesh == Other.bGenerateNaniteEnabledMesh;
+			&& bGenerateNaniteEnabledMesh == Other.bGenerateNaniteEnabledMesh
+			&& NaniteProxyTrianglePercent == Other.NaniteProxyTrianglePercent;
 	}
 
 	/** Inequality. */
@@ -524,9 +530,13 @@ struct FMeshMergingSettings
 	uint8 bAllowDistanceField:1;
 
 	/** Whether to generate a nanite-enabled mesh */
-	UPROPERTY(EditAnywhere, BlueprintReadWrite, Category = MeshSettings)
+	UPROPERTY(EditAnywhere, BlueprintReadWrite, Category = NaniteSettings)
 	uint8 bGenerateNaniteEnabledMesh : 1;
 
+	/** Percentage of triangles to reduce down to for generating a coarse proxy mesh from the Nanite mesh */
+	UPROPERTY(BlueprintReadWrite, EditAnywhere, AdvancedDisplay, Category = NaniteSettings, meta = (EditConditionHides, EditCondition = "bGenerateNaniteEnabledMesh", ClampMin = 0, ClampMax = 100))
+	float NaniteProxyTrianglePercent;
+
 #if WITH_EDITORONLY_DATA
 	/** Whether we should import vertex colors into merged mesh */
 	UPROPERTY()
@@ -575,6 +585,7 @@ struct FMeshMergingSettings
 		, bIncludeImposters(true)
 		, bAllowDistanceField(false)
 		, bGenerateNaniteEnabledMesh(false)
+		, NaniteProxyTrianglePercent(100)
 #if WITH_EDITORONLY_DATA
 		, bImportVertexColors_DEPRECATED(false)
 		, bCalculateCorrectLODModel_DEPRECATED(false)
