Initial load time optimization

Reduces Engine::Init() from 5.69s -> 4.94s on a Gen5 console for an internal project together with another UpdateTemporaryCaches() optimization

diff --git a/Engine/Plugins/Experimental/GameFeatures/Source/GameFeatures/Private/GameFeaturesSubsystem.cpp b/Engine/Plugins/Experimental/GameFeatures/Source/GameFeatures/Private/GameFeaturesSubsystem.cpp
--- a/Engine/Plugins/Experimental/GameFeatures/Source/GameFeatures/Private/GameFeaturesSubsystem.cpp
+++ b/Engine/Plugins/Experimental/GameFeatures/Source/GameFeatures/Private/GameFeaturesSubsystem.cpp
@@ -200,7 +200,7 @@ void UGameFeaturesSubsystem::AddGameFeatureToAssetManager(const UGameFeatureData
 	check(GameFeatureToAdd);
 	UAssetManager& LocalAssetManager = UAssetManager::Get();
 
-	LocalAssetManager.StartBulkScanning();
+	LocalAssetManager.PushBulkScanning();
 
 	for (FPrimaryAssetTypeInfo TypeInfo : GameFeatureToAdd->GetPrimaryAssetTypesToScan())
 	{
@@ -235,7 +235,7 @@ void UGameFeaturesSubsystem::AddGameFeatureToAssetManager(const UGameFeatureData
 		}
 	}
 
-	LocalAssetManager.StopBulkScanning();
+	LocalAssetManager.PopBulkScanning();
 }
 
 void UGameFeaturesSubsystem::RemoveGameFeatureFromAssetManager(const UGameFeatureData* GameFeatureToRemove)
@@ -582,11 +582,15 @@ void UGameFeaturesSubsystem::LoadBuiltInGameFeaturePlugin(const TSharedRef<IPlug
 
 void UGameFeaturesSubsystem::LoadBuiltInGameFeaturePlugins(FBuiltInPluginAdditionalFilters AdditionalFilter)
 {
+	UAssetManager::Get().PushBulkScanning();
+
 	TArray<TSharedRef<IPlugin>> EnabledPlugins = IPluginManager::Get().GetEnabledPlugins();
 	for (const TSharedRef<IPlugin>& Plugin : EnabledPlugins)
 	{
 		LoadBuiltInGameFeaturePlugin(Plugin, AdditionalFilter);
 	}
+
+	UAssetManager::Get().PopBulkScanning();
 }
 
 bool UGameFeaturesSubsystem::GetPluginURLForBuiltInPluginByName(const FString& PluginName, FString& OutPluginURL)
diff --git a/Engine/Source/Runtime/Engine/Classes/Engine/AssetManager.h b/Engine/Source/Runtime/Engine/Classes/Engine/AssetManager.h
--- a/Engine/Source/Runtime/Engine/Classes/Engine/AssetManager.h
+++ b/Engine/Source/Runtime/Engine/Classes/Engine/AssetManager.h
@@ -97,10 +97,16 @@ public:
 	/** Single path wrapper */
 	virtual int32 ScanPathForPrimaryAssets(FPrimaryAssetType PrimaryAssetType, const FString& Path, UClass* BaseClass, bool bHasBlueprintClasses, bool bIsEditorOnly = false, bool bForceSynchronousScan = true);
 	
-	/** Call this before many calls to ScanPaths to improve load performance. It will defer expensive updates until StopBulkScanning is called */
+	/** Call before many calls to ScanPaths to improve load performance. Match each call with PopBulkScanning(). */
+	void PushBulkScanning();
+	void PopBulkScanning();
+
+protected:
 	virtual void StartBulkScanning();
 	virtual void StopBulkScanning();
+	bool IsBulkScanning() const { return NumBulkScanRequests > 0 ; }
 
+public:
 	/** 
 	 * Adds or updates a Dynamic asset, which is a runtime-specified asset that has no on disk representation, so has no FAssetData. But it can have bundle state and a path.
 	 *
@@ -738,9 +744,9 @@ protected:
 	UPROPERTY()
 	bool bOnlyCookProductionAssets;
 
-	/** True if we are currently in bulk scanning mode */
+	/** >0 if we are currently in bulk scanning mode */
 	UPROPERTY()
-	bool bIsBulkScanning;
+	int32 NumBulkScanRequests;
 
 	/** True if asset data is current, if false it will need to rescan before PIE */
 	UPROPERTY()
diff --git a/Engine/Source/Runtime/Engine/Private/AssetManager.cpp b/Engine/Source/Runtime/Engine/Private/AssetManager.cpp
--- a/Engine/Source/Runtime/Engine/Private/AssetManager.cpp
+++ b/Engine/Source/Runtime/Engine/Private/AssetManager.cpp
@@ -213,6 +213,47 @@ const FString UAssetManager::DynamicSearchRootsVirtualPath = TEXT("$DynamicSearc
 FSimpleMulticastDelegate UAssetManager::OnCompletedInitialScanDelegate;
 FSimpleMulticastDelegate UAssetManager::OnAssetManagerCreatedDelegate;
 
+// Allow StartInitialLoading() bulk scan to continue past FinishInitialLoading() to cover:
+// 1. Loading in subclass FinishInitialLoading() override after calling base FinishInitialLoading() 
+// 2. Plugin loading in OnPostEngine callback, e.g. UDataRegistrySubsystem::LoadAllRegistries()
+// 
+// Bulk scanning should stop before the early GC when loading the first map, to ensure
+// RebuildObjectReferenceList() stops INI-referenced UClasses from being GCed.
+//
+// Ideally some future asset manager refactoring can get rid of this complexity / hack.
+static class FInitialBulkScanHelper
+{
+	UAssetManager* Scanner = nullptr;
+
+public:
+	void StartOnce(UAssetManager* Manager)
+	{
+		if (Scanner == nullptr)
+		{
+			Manager->PushBulkScanning();
+			Scanner = Manager;
+		}
+	}
+
+	void StopOnce()
+	{
+		if (Scanner)
+		{
+			Scanner->PopBulkScanning();
+		}
+		else
+		{
+			++Scanner; // Avoid starting bulk scan after OnPostEngineInit
+		}
+	}
+}
+GInitialBulkScan;
+
+// OnPostEngineInit fires in reverse order, register early to get a late callback
+static FDelayedAutoRegisterHelper GInitialBulkScanStopper(EDelayedRegisterRunPhase::StartOfEnginePreInit,
+	[] { FCoreDelegates::OnPostEngineInit.AddLambda([] { GInitialBulkScan.StopOnce(); }); } );
+
+
 UAssetManager::UAssetManager()
 {
 	bIsGlobalAsyncScanEnvironment = false;
@@ -220,7 +261,7 @@ UAssetManager::UAssetManager()
 	bShouldUseSynchronousLoad = false;
 	bIsLoadingFromPakFiles = false;
 	bShouldAcquireMissingChunksOnLoad = false;
-	bIsBulkScanning = false;
+	NumBulkScanRequests = 0;
 	bIsManagementDatabaseCurrent = false;
 	bIsPrimaryAssetDirectoryCurrent = false;
 	bUpdateManagementDatabaseAfterScan = false;
@@ -857,7 +898,7 @@ int32 UAssetManager::ScanPathsForPrimaryAssets(FPrimaryAssetType PrimaryAssetTyp
 		UpdateCachedAssetData(PrimaryAssetId, Data, false);
 	}
 
-	if (!bIsBulkScanning)
+	if (!IsBulkScanning())
 	{
 		RebuildObjectReferenceList();
 	}
@@ -865,28 +906,38 @@ int32 UAssetManager::ScanPathsForPrimaryAssets(FPrimaryAssetType PrimaryAssetTyp
 	return NumAdded;
 }
 
-void UAssetManager::StartBulkScanning()
+void UAssetManager::PushBulkScanning()
 {
-	// TODO switch to an int and support nesting
-	if (ensure(!bIsBulkScanning))
+	if (++NumBulkScanRequests == 1)
 	{
-		bIsBulkScanning = true;
-		NumberOfSpawnedNotifications = 0;
-		bOldTemporaryCachingMode = GetAssetRegistry().GetTemporaryCachingMode();
-		// Go into temporary caching mode to speed up class queries
-		GetAssetRegistry().SetTemporaryCachingMode(true);
+		StartBulkScanning();
 	}
 }
 
-void UAssetManager::StopBulkScanning()
+void UAssetManager::PopBulkScanning()
 {
-	if (ensure(bIsBulkScanning))
+	ensure(NumBulkScanRequests > 0);
+	if (--NumBulkScanRequests == 0)
 	{
-		bIsBulkScanning = false;
-
-		// Leave temporary caching mode
-		GetAssetRegistry().SetTemporaryCachingMode(bOldTemporaryCachingMode);
+		StopBulkScanning();
 	}
+}
+
+void UAssetManager::StartBulkScanning()
+{
+	check(IsBulkScanning());
+
+	NumberOfSpawnedNotifications = 0;
+	bOldTemporaryCachingMode = GetAssetRegistry().GetTemporaryCachingMode();
+	// Go into temporary caching mode to speed up class queries
+	GetAssetRegistry().SetTemporaryCachingMode(true);
+}
+
+void UAssetManager::StopBulkScanning()
+{
+	check(!IsBulkScanning());
+
+	GetAssetRegistry().SetTemporaryCachingMode(bOldTemporaryCachingMode);
 	
 	RebuildObjectReferenceList();
 }
@@ -907,7 +958,7 @@ bool UAssetManager::RegisterSpecificPrimaryAsset(const FPrimaryAssetId& PrimaryA
 	// If we got this far, it will succeed but might warn
 	UpdateCachedAssetData(PrimaryAssetId, NewAssetData, false);
 
-	if (!bIsBulkScanning)
+	if (!IsBulkScanning())
 	{
 		RebuildObjectReferenceList();
 	}
@@ -981,7 +1032,7 @@ void UAssetManager::UpdateCachedAssetData(const FPrimaryAssetId& PrimaryAssetId,
 			LocalAssetRegistry.SetPrimaryAssetIdForObjectPath(NameData.AssetDataPath, PrimaryAssetId);
 		}
 
-		if (bIsBulkScanning)
+		if (IsBulkScanning())
 		{
 			// Do a partial update, add to the path->asset map
 			AssetPathMap.Add(NewAssetPath.GetAssetPathName(), PrimaryAssetId);
@@ -1075,7 +1126,7 @@ bool UAssetManager::AddDynamicAsset(const FPrimaryAssetId& PrimaryAssetId, const
 
 	NameData.AssetPtr = FSoftObjectPtr(AssetPath);
 
-	if (bIsBulkScanning && AssetPath.IsValid())
+	if (IsBulkScanning() && AssetPath.IsValid())
 	{
 		// Do a partial update, add to the path->asset map
 		AssetPathMap.Add(AssetPath.GetAssetPathName(), PrimaryAssetId);
@@ -2988,7 +3039,7 @@ void UAssetManager::ScanPrimaryAssetTypesFromConfig()
 	IAssetRegistry& AssetRegistry = GetAssetRegistry();
 	const UAssetManagerSettings& Settings = GetSettings();
 
-	StartBulkScanning();
+	PushBulkScanning();
 
 	double LastPumpTime = FPlatformTime::Seconds();
 	for (FPrimaryAssetTypeInfo TypeInfo : Settings.PrimaryAssetTypesToScan)
@@ -3012,7 +3063,7 @@ void UAssetManager::ScanPrimaryAssetTypesFromConfig()
 		}
 	}
 
-	StopBulkScanning();
+	PopBulkScanning();
 }
 
 void UAssetManager::ScanPrimaryAssetRulesFromConfig()
@@ -3269,6 +3320,8 @@ bool UAssetManager::GetPackageManagers(FName PackageName, bool bRecurseToParents
 
 void UAssetManager::StartInitialLoading()
 {
+	GInitialBulkScan.StartOnce(this);
+
 	ScanPrimaryAssetTypesFromConfig();
 
 	OnAssetManagerCreatedDelegate.Broadcast();
@@ -3415,7 +3468,7 @@ EAssetSetManagerResult::Type UAssetManager::ShouldSetManager(const FAssetIdentif
 
 void UAssetManager::OnAssetRegistryFilesLoaded()
 {
-	StartBulkScanning();
+	PushBulkScanning();
 
 	for (TPair<FName, TSharedRef<FPrimaryAssetTypeData>>& TypePair : AssetTypeMap)
 	{
@@ -3430,7 +3483,7 @@ void UAssetManager::OnAssetRegistryFilesLoaded()
 		}
 	}
 
-	StopBulkScanning();
+	PopBulkScanning();
 
 	PostInitialAssetScan();
 }
@@ -3975,7 +4028,7 @@ void UAssetManager::RefreshPrimaryAssetDirectory(bool bForceRefresh)
 {
 	if (bForceRefresh || !bIsPrimaryAssetDirectoryCurrent)
 	{
-		StartBulkScanning();
+		PushBulkScanning();
 
 		for (TPair<FName, TSharedRef<FPrimaryAssetTypeData>>& TypePair : AssetTypeMap)
 		{
@@ -4006,7 +4059,7 @@ void UAssetManager::RefreshPrimaryAssetDirectory(bool bForceRefresh)
 			}
 		}
 
-		StopBulkScanning();
+		PopBulkScanning();
 
 		PostInitialAssetScan();
 	}
