New Strata "shared tangent basis foundation" improved to support parameter blending.
- Any shared tangent basis can be registered
- The effective one are evaluated at the end using the cpu compilater kowledge
- Both items above now allows to execute horizontal blending leveraging the new shared tangent basis code
- Only horizontal parameter blending is working today. It is a tick box on the horizontal parameter blending node.

A bit of a gymnastic because GetParameter -Code and -Hash cannot be done at anytime so values must be cached.
Some renaming will follow, add and vertical will be implemented later.

Shooter game ran successfully.

diff --git a/Engine/Shaders/Private/Strata/Strata.ush b/Engine/Shaders/Private/Strata/Strata.ush
--- a/Engine/Shaders/Private/Strata/Strata.ush
+++ b/Engine/Shaders/Private/Strata/Strata.ush
@@ -445,11 +445,19 @@ uint GetStrataSubsurfaceProfile(float NormalizedId)
 	return uint(NormalizedId * 255.0f + 0.5f);
 
 }
+
 FStrataData GetInitialisedStrataData()
 {
 	FStrataData StrataData = (FStrataData)0;
 	return StrataData;
 }
+FStrataData GetInitialisedStrataDataWithReservedFirstSlot()
+{
+	FStrataData StrataData = GetInitialisedStrataData();
+	StrataData.LayerCount++;
+	StrataData.Layers[0].BSDFCount++;
+	return StrataData;
+}
 
 FStrataData GetStrataSlabBSDF(
 	float3 BaseColor, float3 EdgeColor, float Specular, float Metallic, float Roughness, 
@@ -462,8 +470,7 @@ FStrataData GetStrataSlabBSDF(
 	float Thickness, 
 	uint SharedNormalIndex, inout uint SharedNormalTypes)
 {
-	FStrataData StrataData = GetInitialisedStrataData();
-	StrataData.LayerCount++;
+	FStrataData StrataData = GetInitialisedStrataDataWithReservedFirstSlot();
 
 	const uint bHasEdgeColor = any(EdgeColor < 1);
 	const uint bHasHaziness = Haziness > 0.0f ? 1 : 0;
@@ -491,7 +498,6 @@ FStrataData GetStrataSlabBSDF(
 	//  - We also only enabled thin lighting is the MFP input pin is plugged in.
 	//  - It will only be ran for the bottom layer only. Otherwise, the simple volume should be used. As such, Two-sided lighting can also be disabled later if part of the bottom layer.
 
-	StrataData.Layers[0].BSDFCount++;
 	BSDF_SETTYPE			(StrataData.Layers[0].BSDFs[0],		STRATA_BSDF_TYPE_SLAB);
 	BSDF_SETNORMALID		(StrataData.Layers[0].BSDFs[0],		SharedNormalIndex);
 	BSDF_SETEMISSIVE		(StrataData.Layers[0].BSDFs[0],		Emissive);
@@ -540,9 +546,7 @@ FStrataData GetStrataSlabBSDF(
 
 FStrataData GetStrataVolumeFogCloudBSDF(float3 Albedo, float3 Extinction, float3 Emissive, float AmbientOcclusion)
 {
-	FStrataData StrataData = GetInitialisedStrataData();
-	StrataData.LayerCount++;
-	StrataData.Layers[0].BSDFCount++;
+	FStrataData StrataData = GetInitialisedStrataDataWithReservedFirstSlot();
 
 	BSDF_SETTYPE					(StrataData.Layers[0].BSDFs[0], STRATA_BSDF_TYPE_VOLUMETRICFOGCLOUD);
 	BSDF_SETEMISSIVE				(StrataData.Layers[0].BSDFs[0], Emissive);
@@ -557,9 +561,7 @@ FStrataData GetStrataVolumeFogCloudBSDF(float3 Albedo, float3 Extinction, float3
 
 FStrataData GetStrataUnlitBSDF(float3 Emissive, float3 TransmittanceColor)
 {
-	FStrataData StrataData = GetInitialisedStrataData();
-	StrataData.LayerCount++;
-	StrataData.Layers[0].BSDFCount++;
+	FStrataData StrataData = GetInitialisedStrataDataWithReservedFirstSlot();
 
 	BSDF_SETTYPE					(StrataData.Layers[0].BSDFs[0], STRATA_BSDF_TYPE_UNLIT);
 	BSDF_SETEMISSIVE				(StrataData.Layers[0].BSDFs[0], Emissive);
@@ -572,9 +574,7 @@ FStrataData GetStrataUnlitBSDF(float3 Emissive, float3 TransmittanceColor)
 
 FStrataData GetStrataHairBSDF(float3 BaseColor, float Scatter, float Specular, float Roughness, float Backlit, float3 Emissive, uint SharedNormalIndex)
 {
-	FStrataData StrataData = GetInitialisedStrataData();
-	StrataData.LayerCount++;
-	StrataData.Layers[0].BSDFCount++;
+	FStrataData StrataData = GetInitialisedStrataDataWithReservedFirstSlot();
 
 	BSDF_SETTYPE					(StrataData.Layers[0].BSDFs[0], STRATA_BSDF_TYPE_HAIR);
 	BSDF_SETNORMALID				(StrataData.Layers[0].BSDFs[0], SharedNormalIndex);
@@ -594,9 +594,7 @@ FStrataData GetStrataSingleLayerWaterBSDF(
 	float3 BaseColor, float Metallic, float Specular, float Roughness, float3 Emissive, float TopMaterialOpacity,
 	float3 WaterAlbedo, float3 WaterExtinction, float WaterPhaseG, float3 ColorScaleBehindWater, uint SharedNormalIndex)
 {
-	FStrataData StrataData = GetInitialisedStrataData();
-	StrataData.LayerCount++;
-	StrataData.Layers[0].BSDFCount++;
+	FStrataData StrataData = GetInitialisedStrataDataWithReservedFirstSlot();
 
 	BSDF_SETTYPE					(StrataData.Layers[0].BSDFs[0], STRATA_BSDF_TYPE_SINGLELAYERWATER);
 	BSDF_SETNORMALID				(StrataData.Layers[0].BSDFs[0], SharedNormalIndex);
@@ -874,6 +872,78 @@ FStrataData StrataVerticalLayering(FStrataData Top, FStrataData Base)
 	return StrataData;
 }
 
+// Note about parameter blending
+//  - We can only parameter blend Slab nodes. This error is handled in the compiler.
+//  - We can only parameter blend one BSDF into another single BSDF. More complex topology are not handled and result in a compiler error. This could be extended to a single BSDF per layer.
+//  - Maybe at some point we can propose a concatenate node, compressing a complex topology into a single one.
+//  - Slab node with SSSprofile cannot be blendend. Those will have to be specifically handled by the user and not mixed with any other slab if it can be demoted to parameter blending
+
+FStrataData StrataHorizontalMixingParameterBlending(FStrataData Foreground, FStrataData Background, float Mix, uint NewNormalIndex)
+{
+	const float ForegroundMixFactor = saturate(Mix);
+	const float BackgroundMixFactor = 1.0 - ForegroundMixFactor;
+
+	// This function can only parameter blend two BSDF and not more. We cannot blend complex topology.
+	// However, maybe at some point we will
+	// This function only blend one we only parameter blend one slab with anothe
+
+	FStrataData Result = GetInitialisedStrataDataWithReservedFirstSlot();
+#define ResultBSDF		Result.Layers[0].BSDFs[0]
+#define ForegroundBSDF	Foreground.Layers[0].BSDFs[0]
+#define BackgroundBSDF	Background.Layers[0].BSDFs[0]
+
+	// We can only parameter blend the slab 
+	BSDF_SETTYPE(ResultBSDF,			STRATA_BSDF_TYPE_SLAB);
+	BSDF_SETNORMALID(ResultBSDF,		NewNormalIndex);
+
+	BSDF_SETISTOPLAYER(ResultBSDF,		0);
+	BSDF_SETANISOTROPY(ResultBSDF,		BSDF_GETANISOTROPY(ForegroundBSDF)				| BSDF_GETANISOTROPY(BackgroundBSDF));
+	BSDF_SETHASSSS(ResultBSDF,			BSDF_GETHASSSS(ForegroundBSDF)					| BSDF_GETHASSSS(BackgroundBSDF));
+	BSDF_SETHASSSSPROFILE(ResultBSDF,	BSDF_GETHASSSSPROFILE(ForegroundBSDF)			| BSDF_GETHASSSSPROFILE(BackgroundBSDF));	// Problem: cannot blend SSS profiles...
+	BSDF_SETISTHIN(ResultBSDF,			BSDF_GETISTHIN(ForegroundBSDF)					| BSDF_GETISTHIN(BackgroundBSDF));			// STRATA_TODO fix the MPF/DMFP situation for that to really work
+	BSDF_SETHASDMFP(ResultBSDF,			BSDF_GETHASDMFP(ForegroundBSDF)					| BSDF_GETHASDMFP(BackgroundBSDF));
+	BSDF_SETHASHAZINESS(ResultBSDF,		BSDF_GETHASHAZINESS(ForegroundBSDF)				| BSDF_GETHASHAZINESS(BackgroundBSDF));
+	BSDF_SETHASEDGECOLOR(ResultBSDF,	BSDF_GETHASEDGECOLOR(ForegroundBSDF)			| BSDF_GETHASEDGECOLOR(BackgroundBSDF));
+	BSDF_SETHASTHINFILM(ResultBSDF,		BSDF_GETHASTHINFILM(ForegroundBSDF)				| BSDF_GETHASTHINFILM(BackgroundBSDF));
+	BSDF_SETHASFUZZ(ResultBSDF,			BSDF_GETHASFUZZ(ForegroundBSDF)					| BSDF_GETHASFUZZ(BackgroundBSDF));
+	BSDF_SETEMISSIVE(ResultBSDF,		lerp(BSDF_GETEMISSIVE(BackgroundBSDF),			BSDF_GETEMISSIVE(ForegroundBSDF),			ForegroundMixFactor));
+	BSDF_SETTHICKNESSCM(ResultBSDF,		lerp(BSDF_GETTHICKNESSCM(BackgroundBSDF),		BSDF_GETTHICKNESSCM(ForegroundBSDF),		ForegroundMixFactor));
+
+	SLAB_BASECOLOR(ResultBSDF) =		lerp(SLAB_BASECOLOR(BackgroundBSDF),			SLAB_BASECOLOR(ForegroundBSDF),				ForegroundMixFactor);
+	SLAB_EDGECOLOR(ResultBSDF)=			lerp(SLAB_EDGECOLOR(BackgroundBSDF),			SLAB_EDGECOLOR(ForegroundBSDF),				ForegroundMixFactor);
+	SLAB_SPECULAR(ResultBSDF) =			lerp(SLAB_SPECULAR(BackgroundBSDF),				SLAB_SPECULAR(ForegroundBSDF),				ForegroundMixFactor);
+	SLAB_METALLIC(ResultBSDF) =			lerp(SLAB_METALLIC(BackgroundBSDF),				SLAB_METALLIC(ForegroundBSDF),				ForegroundMixFactor);
+	SLAB_ROUGHNESSX(ResultBSDF) =		lerp(SLAB_ROUGHNESSX(BackgroundBSDF),			SLAB_ROUGHNESSX(ForegroundBSDF),			ForegroundMixFactor);
+	SLAB_ROUGHNESSY(ResultBSDF) =		lerp(SLAB_ROUGHNESSY(BackgroundBSDF),			SLAB_ROUGHNESSY(ForegroundBSDF),			ForegroundMixFactor);
+	SLAB_HAZINESS(ResultBSDF) =			lerp(SLAB_HAZINESS(BackgroundBSDF),				SLAB_HAZINESS(ForegroundBSDF),				ForegroundMixFactor);
+	SLAB_THINFILM_THICKNESS(ResultBSDF)=lerp(SLAB_THINFILM_THICKNESS(BackgroundBSDF),	SLAB_THINFILM_THICKNESS(ForegroundBSDF),	ForegroundMixFactor);
+	SLAB_FUZZ_AMOUNT(ResultBSDF) =		lerp(SLAB_FUZZ_AMOUNT(BackgroundBSDF),			SLAB_FUZZ_AMOUNT(ForegroundBSDF),			ForegroundMixFactor);
+	SLAB_FUZZ_COLOR(ResultBSDF)	=		lerp(SLAB_FUZZ_COLOR(BackgroundBSDF),			SLAB_FUZZ_COLOR(ForegroundBSDF),			ForegroundMixFactor);
+
+	if (BSDF_GETHASSSSPROFILE(ResultBSDF))
+	{
+		// As mentioned above, we cannot blend SSSProfiles.
+		SLAB_SSSPROFILEID(ResultBSDF) = SLAB_SSSPROFILEID(ForegroundBSDF);
+		SLAB_SSSPROFILERADIUSSCALE(ResultBSDF) = SLAB_SSSPROFILERADIUSSCALE(ForegroundBSDF);
+	}
+	else
+	{
+		SLAB_SSSDMFP(ResultBSDF) = lerp(SLAB_SSSDMFP(BackgroundBSDF),				SLAB_SSSDMFP(ForegroundBSDF),				ForegroundMixFactor);
+	}
+
+
+#if STRATA_INLINE_SHADING
+	// At this stage we do not know if this strata BSDF is going to be using simple volume lighting so always store the DMFP on the side to not override SSSPROFILE data.
+	ResultBSDF.TmpDMFP =				lerp(BackgroundBSDF.TmpDMFP,					ForegroundBSDF.TmpDMFP,						ForegroundMixFactor);
+#endif
+
+	ResultBSDF.Weight =					lerp(BackgroundBSDF.Weight,						ForegroundBSDF.Weight,						ForegroundMixFactor);
+
+#undef ResultBSDF
+#undef ForegroundBSDF
+#undef BackgroundBSDF
+	return Result;
+}
 
 
 ///////////////////////////////////////////////////////////////////////////////
diff --git a/Engine/Shaders/Public/ShaderVersion.ush b/Engine/Shaders/Public/ShaderVersion.ush
--- a/Engine/Shaders/Public/ShaderVersion.ush
+++ b/Engine/Shaders/Public/ShaderVersion.ush
@@ -3,4 +3,4 @@
 // in Platform.ush (which should be included in any shader) it allows to invalidate the shader DDC.
 // 
 // If you are merging streams and there is a conflict with this GUID you should make a new GUID rather than taking one or the other.
-// GUID = 24BB391CB1F842E9922C0FEDC88DF989
\ No newline at end of file
+// GUID = 5340A214CE4747139F30298B39D33498
diff --git a/Engine/Source/Runtime/Engine/Classes/Materials/MaterialExpressionStrata.h b/Engine/Source/Runtime/Engine/Classes/Materials/MaterialExpressionStrata.h
--- a/Engine/Source/Runtime/Engine/Classes/Materials/MaterialExpressionStrata.h
+++ b/Engine/Source/Runtime/Engine/Classes/Materials/MaterialExpressionStrata.h
@@ -401,6 +401,9 @@ class UMaterialExpressionStrataHorizontalMixing : public UMaterialExpressionStra
 	UPROPERTY()
 	FExpressionInput Mix;
 
+	UPROPERTY(EditAnywhere, Category = Mode)
+	uint32 bUseParameterBlending : 1;
+
 	//~ Begin UMaterialExpression Interface
 #if WITH_EDITOR
 	virtual int32 Compile(class FMaterialCompiler* Compiler, int32 OutputIndex) override;
diff --git a/Engine/Source/Runtime/Engine/Private/Materials/HLSLMaterialTranslator.cpp b/Engine/Source/Runtime/Engine/Private/Materials/HLSLMaterialTranslator.cpp
--- a/Engine/Source/Runtime/Engine/Private/Materials/HLSLMaterialTranslator.cpp
+++ b/Engine/Source/Runtime/Engine/Private/Materials/HLSLMaterialTranslator.cpp
@@ -158,7 +158,9 @@ FHLSLMaterialTranslator::FHLSLMaterialTranslator(FMaterial* InMaterial,
 ,	AllocatedUserTexCoords()
 ,	AllocatedUserVertexTexCoords()
 ,	DynamicParticleParameterMask(0)
+,	StrataValidFrontMaterialCodeChunkPostTranslate(INDEX_NONE)
 ,	NextFreeStrataShaderNormalIndex(0)
+,	FinalUsedSharedNormalCount(0)
 ,	NumVtSamples(0)
 ,	TargetPlatform(InTargetPlatform)
 {
@@ -1201,7 +1203,8 @@ bool FHLSLMaterialTranslator::Translate()
 		bool bStrataFrontMaterialIsValid = bStrataFrontMaterialProvided;
 		if (bStrataFrontMaterialIsValid)
 		{
-			const FStrataMaterialCompilationInfo& StrataCompilationInfo = CodeChunkToStrataCompilationInfoMap[Chunk[MP_FrontMaterial]];
+			StrataValidFrontMaterialCodeChunkPostTranslate = Chunk[MP_FrontMaterial];
+			const FStrataMaterialCompilationInfo& StrataCompilationInfo = CodeChunkToStrataCompilationInfoMap[StrataValidFrontMaterialCodeChunkPostTranslate];
 
 			// The material can be null when some entries are automatically generated, for instance in the material layer blending system
 			const bool bIsNullMaterial = (StrataCompilationInfo.LayerCount == 0 && StrataCompilationInfo.TotalBSDFCount == 0);
@@ -1211,7 +1214,7 @@ bool FHLSLMaterialTranslator::Translate()
 		if (bStrataFrontMaterialIsValid)
 		{
 			bMaterialIsStrata = true;
-			const FStrataMaterialCompilationInfo& StrataCompilationInfo = CodeChunkToStrataCompilationInfoMap[Chunk[MP_FrontMaterial]];
+			const FStrataMaterialCompilationInfo& StrataCompilationInfo = CodeChunkToStrataCompilationInfoMap[StrataValidFrontMaterialCodeChunkPostTranslate];
 
 			// VolumetricFogCloud must be used in isolation
 			if (StrataMaterialContainsAnyBSDF(this, StrataCompilationInfo, STRATA_BSDF_TYPE_VOLUMETRICFOGCLOUD))
@@ -1277,7 +1280,7 @@ bool FHLSLMaterialTranslator::Translate()
 				const FStrataMaterialCompilationInfo::FLayer& Layer = StrataCompilationInfo.Layers[LayerIt];
 				for (uint32 BSDFIt = 0; BSDFIt < Layer.BSDFCount; ++BSDFIt)
 				{
-					StrataMaterialDescription += FString::Printf(TEXT("     - %s - SharedNormalIndex = %u \r\n"), *GetStrataBSDFName(Layer.BSDFs[BSDFIt].Type), Layer.BSDFs[BSDFIt].SharedNormalIndex);
+					StrataMaterialDescription += FString::Printf(TEXT("     - %s - SharedNormalIndexMacro = %s \r\n"), *GetStrataBSDFName(Layer.BSDFs[BSDFIt].Type), *GetStrataSharedNormalIndexMacro(Layer.BSDFs[BSDFIt].RegisteredSharedNormal));
 				}
 			}
 
@@ -1682,6 +1685,113 @@ void FHLSLMaterialTranslator::GetMaterialEnvironment(EShaderPlatform InPlatform,
 		OutEnvironment.SetDefine(TEXT("STRATA_CLAMPED_LAYER_COUNT"), StrataMaterialAnalysis.ClampedLayerCount);
 	}
 
+	if (bMaterialIsStrata)
+	{
+		StrataPixelNormalInitializerValues += FString::Printf(TEXT("\n\n\n\t// Strata normal and tangent\n"));
+		FinalUsedSharedNormalCount = 0;
+		FStrataRegisteredSharedNormal UsedSharedNormalInfo[STRATA_MAX_SHARED_NORMAL_REGISTERS];
+
+
+		/*
+		* The final output code/workflow for shared tangent basis should look like
+		*
+		* #define NORMAL_INDEX_0 0		// default, unused
+		* #define NORMAL_INDEX_1 0		// default, unused
+		* #define NORMAL_INDEX_2 0
+		*
+		* FStrataData BSDF0 = GetStrataSlabBSDF(... NORMAL_INDEX_0, NormalCode0 ...)
+		* FStrataData BSDF1 = GetStrataSlabBSDF(... NORMAL_INDEX_1, NormalCode1 ...)
+		*
+		* float3 NormalCode2 = lerp(NormalCode0, NormalCode1, mix)
+		* FStrataData BSDF2 = StrataHorizontalMixingParameterBlending(BSDF0, BSDF1, mix, NormalCode2, NORMAL_INDEX_2) // will internally create NormalCode2
+		*
+		* tParameters.SharedNormals.Normals[NORMAL_INDEX_2] = NormalCode2;
+		* #if MATERIAL_TANGENTSPACENORMAL
+		* Parameters.SharedNormals.Normals[NORMAL_INDEX_2] *= Parameters.TwoSidedSign;
+		* #endif
+		*/
+
+
+		// We go over all output BSDF and reserve actualspace for there shared tangent basis. Only up to STRATA_MAX_SHARED_NORMAL_REGISTERS.
+		// When we know we will need a basis, it is allocated as FinalSharedNormalIndex. A macro is created with propoer output index for the correct value to be output ot the correct slot by the BSDF.
+		// If too many output are required, we faile the compilation.
+		// STRATA_TODO: parameter blend in case it is not possible to output according to per pixel cost retriction.
+		const FStrataMaterialCompilationInfo& StrataCompilationInfo = CodeChunkToStrataCompilationInfoMap[StrataValidFrontMaterialCodeChunkPostTranslate];
+		for (uint32 LayerIt = 0; LayerIt < StrataCompilationInfo.LayerCount; ++LayerIt)
+		{
+			const FStrataMaterialCompilationInfo::FLayer& Layer = StrataCompilationInfo.Layers[LayerIt];
+			for (uint32 BSDFIt = 0; BSDFIt < Layer.BSDFCount; ++BSDFIt)
+			{
+				const FStrataMaterialCompilationInfo::FBSDF& BSDF = Layer.BSDFs[BSDFIt];
+				if (BSDF.RegisteredSharedNormal.NormalCodeChunk == INDEX_NONE && BSDF.RegisteredSharedNormal.TangentCodeChunk == INDEX_NONE)
+				{
+					// We skip null normal on certain BSDF, for instance unlit.
+					continue;
+				}
+
+				const FStrataSharedNormalInfo StrataSharedNormalInfo = StrataCompilationInfoGetMatchingSharedNormalInfo(BSDF.RegisteredSharedNormal);
+
+				// First, we check that the normal/tangent has not already written out (avoid 2 BSDFs sharing the same normal to note generate the same code twice)
+				bool bAlreadyProcessed = false;
+				for (uint8 i = 0; i < FinalUsedSharedNormalCount; ++i)
+				{
+					if (UsedSharedNormalInfo[i].NormalCodeChunkHash == StrataSharedNormalInfo.SharedData.NormalCodeChunkHash &&
+						(UsedSharedNormalInfo[i].TangentCodeChunkHash == StrataSharedNormalInfo.SharedData.TangentCodeChunkHash || BSDF.RegisteredSharedNormal.TangentCodeChunk == INDEX_NONE))
+					{
+						bAlreadyProcessed = true;
+						break;
+					}
+				}
+				if (bAlreadyProcessed)
+				{
+					continue;
+				}
+
+				check(FinalUsedSharedNormalCount < STRATA_MAX_SHARED_NORMAL_REGISTERS);
+				const uint8 FinalSharedNormalIndex = FinalUsedSharedNormalCount++;
+				UsedSharedNormalInfo[FinalSharedNormalIndex] = StrataSharedNormalInfo.SharedData;
+
+				// Write out normals
+				StrataPixelNormalInitializerValues += FString::Printf(TEXT("\tParameters.SharedNormals.Normals[%u] = %s;\n"), FinalSharedNormalIndex, *StrataSharedNormalInfo.SharedNormalCode);
+				StrataPixelNormalInitializerValues += FString::Printf(TEXT("\t#if MATERIAL_TANGENTSPACENORMAL\n"));
+				StrataPixelNormalInitializerValues += FString::Printf(TEXT("\tParameters.SharedNormals.Normals[%u] *= Parameters.TwoSidedSign;\n"), FinalSharedNormalIndex);
+				StrataPixelNormalInitializerValues += FString::Printf(TEXT("\t#endif\n"));
+
+				// Write out tangents
+				if (StrataSharedNormalInfo.SharedData.TangentCodeChunk != INDEX_NONE)
+				{
+					StrataPixelNormalInitializerValues += FString::Printf(TEXT("\tParameters.SharedNormals.Tangents[%u] = %s;\n"), FinalSharedNormalIndex, *StrataSharedNormalInfo.SharedTangentCode);
+				}
+				else
+				{
+					StrataPixelNormalInitializerValues += FString::Printf(TEXT("\tParameters.SharedNormals.Tangents[%u] = float3(1,0,0);\n"), FinalSharedNormalIndex);
+				}
+				StrataPixelNormalInitializerValues += FString::Printf(TEXT("\t#if MATERIAL_TANGENTSPACENORMAL\n"));
+				StrataPixelNormalInitializerValues += FString::Printf(TEXT("\tParameters.SharedNormals.Tangents[%u] *= Parameters.TwoSidedSign;\n"), FinalSharedNormalIndex);
+				StrataPixelNormalInitializerValues += FString::Printf(TEXT("\t#endif\n"));
+			}
+		}
+		StrataPixelNormalInitializerValues += FString::Printf(TEXT("\tParameters.SharedNormals.NormalCount = %u;\n"), FinalUsedSharedNormalCount);
+
+		// Now write out all the macros, them mapping from the BSDF to the effective position/index in the shared normal array they should write to.
+		for (TMultiMap<uint64, FStrataSharedNormalInfo>::TIterator It(CodeChunkToStrataSharedNormal); It; ++It)
+		{
+			// The default linear output index will be 0 by default, and different if in fact the shared normal point to one that is effectively in used in the array of shared normals.
+			uint8 LinearIndex = 0;
+			for (uint8 i = 0; i < FinalUsedSharedNormalCount; ++i)
+			{
+				if (UsedSharedNormalInfo[i].NormalCodeChunkHash == It->Value.SharedData.NormalCodeChunkHash &&
+					(UsedSharedNormalInfo[i].TangentCodeChunkHash == It->Value.SharedData.TangentCodeChunkHash || It->Value.SharedData.TangentCodeChunk == INDEX_NONE))
+				{
+					LinearIndex = i;
+					break;
+				}
+			}
+
+			OutEnvironment.SetDefine(*GetStrataSharedNormalIndexMacro(It->Value.SharedData), LinearIndex);
+		}
+	}
+
 	OutEnvironment.SetDefine(TEXT("TEXTURE_SAMPLE_DEBUG"), IsDebugTextureSampleEnabled() ? TEXT("1") : TEXT("0"));
 }
 
@@ -1811,31 +1921,7 @@ void FHLSLMaterialTranslator::GetSharedInputsMaterialCode(FString& PixelMembersD
 
 		if (CodeChunkToStrataSharedNormal.Num() > 0)
 		{
-			// Generate code to fill up the array containing normal shared between Strata BSDFs
-			for (TMultiMap<uint64, FStrataSharedNormalInfo>::TIterator It(CodeChunkToStrataSharedNormal); It; ++It)
-			{
-				PixelInputInitializerValues += FString::Printf(TEXT("\tParameters.SharedNormals.Normals[%u] = %s;\n"), It->Value.SharedNormalIndex, *It->Value.SharedNormalCode);
-
-				PixelInputInitializerValues += FString::Printf(TEXT("\t#if MATERIAL_TANGENTSPACENORMAL\n"));
-				PixelInputInitializerValues += FString::Printf(TEXT("\tParameters.SharedNormals.Normals[%u] *= Parameters.TwoSidedSign;\n"), It->Value.SharedNormalIndex);
-				PixelInputInitializerValues += FString::Printf(TEXT("\t#endif\n"));
-			}
-			for (TMultiMap<uint64, FStrataSharedNormalInfo>::TIterator It(CodeChunkToStrataSharedNormal); It; ++It)
-			{
-				if (It->Value.TangentCodeChunk != -1)
-				{
-					PixelInputInitializerValues += FString::Printf(TEXT("\tParameters.SharedNormals.Tangents[%u] = %s;\n"), It->Value.SharedNormalIndex, *It->Value.SharedTangentCode);
-				}
-				else
-				{
-					PixelInputInitializerValues += FString::Printf(TEXT("\tParameters.SharedNormals.Tangents[%u] = float3(1,0,0);\n"), It->Value.SharedNormalIndex);
-				}
-
-				PixelInputInitializerValues += FString::Printf(TEXT("\t#if MATERIAL_TANGENTSPACENORMAL\n"));
-				PixelInputInitializerValues += FString::Printf(TEXT("\tParameters.SharedNormals.Tangents[%u] *= Parameters.TwoSidedSign;\n"), It->Value.SharedNormalIndex);
-				PixelInputInitializerValues += FString::Printf(TEXT("\t#endif\n"));
-			}
-			PixelInputInitializerValues += FString::Printf(TEXT("\tParameters.SharedNormals.NormalCount = %u;\n"), CodeChunkToStrataSharedNormal.Num());
+			PixelInputInitializerValues += StrataPixelNormalInitializerValues;
 		}
 
 		PixelMembersInitializationEpilog += PixelInputInitializerValues;
@@ -8761,79 +8847,98 @@ const FStrataMaterialCompilationInfo& FHLSLMaterialTranslator::GetStrataCompilat
 	return CodeChunkToStrataCompilationInfoMap[CodeChunk];
 }
 
-uint8 FHLSLMaterialTranslator::StrataCompilationInfoRegisterSharedNormalIndex(int32 NormalCodeChunk)
+FStrataRegisteredSharedNormal FHLSLMaterialTranslator::StrataCompilationInfoRegisterSharedNormal(int32 NormalCodeChunk)
 {
 	check(NormalCodeChunk != INDEX_NONE);
 	check(NextFreeStrataShaderNormalIndex < 255);	// Out of shared normal slots
-	check(NextFreeStrataShaderNormalIndex < STRATA_MAX_SHARED_NORMAL_REGISTERS);	// This is the size of the HLSL array limiting the maximum shared normal count.
 
-	const uint64 NormalCodeChunkHash = GetParameterHash(NormalCodeChunk);
+	FStrataRegisteredSharedNormal StrataRegisteredSharedNormal;
+	StrataRegisteredSharedNormal.NormalCodeChunk		= NormalCodeChunk;
+	StrataRegisteredSharedNormal.NormalCodeChunkHash	= GetParameterHash(NormalCodeChunk); 
+	StrataRegisteredSharedNormal.TangentCodeChunk		= INDEX_NONE;
+	StrataRegisteredSharedNormal.TangentCodeChunkHash	= GetParameterHash(INDEX_NONE);
+	StrataRegisteredSharedNormal.GraphSharedNormalIndex = NextFreeStrataShaderNormalIndex++;
 
 	// Find any basis which match the Normal code chunk
 	// A normal can be duplicated when it is paired with different tangent, so find the first one which matches
 	TArray<FStrataSharedNormalInfo*> NormalInfos;
-	CodeChunkToStrataSharedNormal.MultiFindPointer(NormalCodeChunkHash, NormalInfos);
+	CodeChunkToStrataSharedNormal.MultiFindPointer(StrataRegisteredSharedNormal.NormalCodeChunkHash, NormalInfos);
 	if (NormalInfos.Num() == 0)
 	{
-		// Allocate a new slot for a new shared normal
-		uint8 NewSharedNormalIndex = NextFreeStrataShaderNormalIndex++;
-		CodeChunkToStrataSharedNormal.Add(NormalCodeChunkHash, { NormalCodeChunk , INDEX_NONE, NewSharedNormalIndex, *GetParameterCode(NormalCodeChunk), FString() });
-		return NewSharedNormalIndex;
+		CodeChunkToStrataSharedNormal.Add(StrataRegisteredSharedNormal.NormalCodeChunkHash, { StrataRegisteredSharedNormal, *GetParameterCode(StrataRegisteredSharedNormal.NormalCodeChunk), FString() });
+		return StrataRegisteredSharedNormal;
 	}
 	// Return the first existing code chunk which match the normal chunk code
-	return NormalInfos[0]->SharedNormalIndex;
+	return NormalInfos[0]->SharedData;
 }
 
-uint8 FHLSLMaterialTranslator::StrataCompilationInfoRegisterSharedNormalIndex(int32 NormalCodeChunk, int32 TangentCodeChunk)
+FStrataRegisteredSharedNormal FHLSLMaterialTranslator::StrataCompilationInfoRegisterSharedNormal(int32 NormalCodeChunk, int32 TangentCodeChunk)
 {
 	check(NormalCodeChunk != INDEX_NONE);
-	check(TangentCodeChunk != INDEX_NONE);
 	check(NextFreeStrataShaderNormalIndex < 255);	// Out of shared normal slots
-	check(NextFreeStrataShaderNormalIndex < STRATA_MAX_SHARED_NORMAL_REGISTERS);	// This is the size of the HLSL array limiting the maximum shared normal count.
+	
 
-	const uint64 NormalCodeChunkHash = GetParameterHash(NormalCodeChunk);
-	const uint64 TangentCodeChunkHash = GetParameterHash(TangentCodeChunk);
+	FStrataRegisteredSharedNormal StrataRegisteredSharedNormal;
+	StrataRegisteredSharedNormal.NormalCodeChunk		= NormalCodeChunk;
+	StrataRegisteredSharedNormal.NormalCodeChunkHash	= GetParameterHash(NormalCodeChunk); 
+	StrataRegisteredSharedNormal.TangentCodeChunk		= TangentCodeChunk;
+	StrataRegisteredSharedNormal.TangentCodeChunkHash	= GetParameterHash(TangentCodeChunk);
+	StrataRegisteredSharedNormal.GraphSharedNormalIndex = NextFreeStrataShaderNormalIndex++;
 
 	// Find a basis which matches both the Normal & the Tangent code chunks
 	TArray<FStrataSharedNormalInfo*> NormalInfos;
-	CodeChunkToStrataSharedNormal.MultiFindPointer(NormalCodeChunkHash, NormalInfos);
-	bool bFound = false;
-	uint8 SharedNormalIndex = 0;
+	CodeChunkToStrataSharedNormal.MultiFindPointer(StrataRegisteredSharedNormal.NormalCodeChunkHash, NormalInfos);
 	for (FStrataSharedNormalInfo* NormalInfo : NormalInfos)
 	{
 		// * Either we find a perfect match (normal & tangent matches)
 		// * Or we find a normal which doesn't have a tangent associated with, and we set the tangent for code
-		if (TangentCodeChunkHash == GetParameterHash(NormalInfo->TangentCodeChunk))
+		if (StrataRegisteredSharedNormal.TangentCodeChunkHash == NormalInfo->SharedData.TangentCodeChunk)
 		{
-			SharedNormalIndex = NormalInfo->SharedNormalIndex;
-			bFound = true;
+			return NormalInfo->SharedData;
 			break;
 		}
-		else if (NormalInfo->TangentCodeChunk == INDEX_NONE)
+		else if (NormalInfo->SharedData.TangentCodeChunk == INDEX_NONE)
 		{
-			NormalInfo->TangentCodeChunk = TangentCodeChunk;
-			NormalInfo->SharedTangentCode = *GetParameterCode(TangentCodeChunk);
-			SharedNormalIndex = NormalInfo->SharedNormalIndex;
-			bFound = true;
+			NormalInfo->SharedData.TangentCodeChunk		= TangentCodeChunk;
+			NormalInfo->SharedData.TangentCodeChunkHash = StrataRegisteredSharedNormal.TangentCodeChunkHash;
+			NormalInfo->SharedTangentCode				= *GetParameterCode(TangentCodeChunk);
+			return NormalInfo->SharedData;
 			break;
 		}
 	}
-	
-	if (NormalInfos.Num() == 0 || !bFound)
+
+	// Allocate a new slot for a new shared normal & tangent
+	const FString NewSharedNormalIndex = GetStrataSharedNormalIndexMacro(StrataRegisteredSharedNormal);
+	CodeChunkToStrataSharedNormal.Add(StrataRegisteredSharedNormal.NormalCodeChunkHash, { StrataRegisteredSharedNormal, *GetParameterCode(StrataRegisteredSharedNormal.NormalCodeChunk), *GetParameterCode(StrataRegisteredSharedNormal.TangentCodeChunk) });
+	return StrataRegisteredSharedNormal;
+}
+
+FHLSLMaterialTranslator::FStrataSharedNormalInfo FHLSLMaterialTranslator::StrataCompilationInfoGetMatchingSharedNormalInfo(const FStrataRegisteredSharedNormal& SearchedSharedNormal)
+{
+	check(NextFreeStrataShaderNormalIndex < 255);	// Out of shared normal slots
+
+	// Find a basis which matches both the Normal & the Tangent code chunks
+	TArray<FStrataSharedNormalInfo*> NormalInfos;
+	CodeChunkToStrataSharedNormal.MultiFindPointer(SearchedSharedNormal.NormalCodeChunkHash, NormalInfos);
+
+	// We first try to find a perfect match for normal and tangent from all the registered element.
+	for (FStrataSharedNormalInfo* NormalInfo : NormalInfos)
 	{
-		// Allocate a new slot for a new shared normal & tangent
-		uint8 NewSharedNormalIndex = NextFreeStrataShaderNormalIndex++;
-		CodeChunkToStrataSharedNormal.Add(NormalCodeChunkHash, { NormalCodeChunk , TangentCodeChunk, NewSharedNormalIndex, *GetParameterCode(NormalCodeChunk), *GetParameterCode(TangentCodeChunk) });
-		return NewSharedNormalIndex;
+		if (SearchedSharedNormal.TangentCodeChunk == INDEX_NONE ||											// We selected the first available normal if there is no tangent specified on the material.
+			SearchedSharedNormal.TangentCodeChunkHash == NormalInfo->SharedData.TangentCodeChunkHash)		// Otherwise we select the normal+tangent that exactly matches the request.
+		{
+			return *NormalInfo;
+			break;
+		}
 	}
 
-	// Return the existing code chunk
-	return SharedNormalIndex;
+	check(0);	// When the compiler is querying, this is to get a result to generate code from a fully processed graph. No result means a bug happened during graph processing.
+	return FStrataSharedNormalInfo();
 }
 
 uint8 FHLSLMaterialTranslator::StrataCompilationInfoGetSharedNormalCount()
 {
-	return CodeChunkToStrataSharedNormal.Num();
+	return FinalUsedSharedNormalCount;
 }
 
 int32 FHLSLMaterialTranslator::StrataCreateAndRegisterNullMaterial()
@@ -8852,10 +8957,10 @@ int32 FHLSLMaterialTranslator::StrataSlabBSDF(
 	int32 ThinFilmThickness, 
 	int32 FuzzAmount, int32 FuzzColor,
 	int32 Thickness, 
-	int32 Normal, int32 Tangent, uint8 SharedNormalIndex)
+	int32 Normal, int32 Tangent, const FString& SharedNormalIndexMacro)
 {
 	return AddCodeChunk(
-		MCT_Strata, TEXT("GetStrataSlabBSDF(%s, %s, %s, %s, %s, %s, %s, %s, %s, %s, %s, %s, %s, %s, %s, %u, Parameters.SharedNormals.NormalTypes) /* Normal = %s ; Tangent = %s */"),
+		MCT_Strata, TEXT("GetStrataSlabBSDF(%s, %s, %s, %s, %s, %s, %s, %s, %s, %s, %s, %s, %s, %s, %s, %s, Parameters.SharedNormals.NormalTypes) /* Normal = %s ; Tangent = %s */"),
 		*GetParameterCode(BaseColor),
 		*GetParameterCode(EdgeColor),
 		*GetParameterCode(Specular),
@@ -8871,19 +8976,19 @@ int32 FHLSLMaterialTranslator::StrataSlabBSDF(
 		*GetParameterCode(FuzzAmount),
 		*GetParameterCode(FuzzColor),
 		*GetParameterCode(Thickness),
-		SharedNormalIndex,
+		*SharedNormalIndexMacro,
 		*GetParameterCode(Normal),
-		*GetParameterCode(Tangent)
+		Tangent!=INDEX_NONE ? *GetParameterCode(Tangent) : TEXT("NONE")
 	);
 }
 
-int32 FHLSLMaterialTranslator::StrataSheenBSDF(int32 BaseColor, int32 Roughness, int32 Normal, uint8 SharedNormalIndex)
+int32 FHLSLMaterialTranslator::StrataSheenBSDF(int32 BaseColor, int32 Roughness, int32 Normal, const FString& SharedNormalIndexMacro)
 {
 	return AddCodeChunk(
-		MCT_Strata, TEXT("GetStrataSheenBSDF(%s, %s, %u) /* %s */"),
+		MCT_Strata, TEXT("GetStrataSheenBSDF(%s, %s, %s) /* %s */"),
 		*GetParameterCode(BaseColor),
 		*GetParameterCode(Roughness),
-		SharedNormalIndex,
+		*SharedNormalIndexMacro,
 		*GetParameterCode(Normal)
 	);
 }
@@ -8908,27 +9013,27 @@ int32 FHLSLMaterialTranslator::StrataUnlitBSDF(int32 EmissiveColor, int32 Transm
 	);
 }
 
-int32 FHLSLMaterialTranslator::StrataHairBSDF(int32 BaseColor, int32 Scatter, int32 Specular, int32 Roughness, int32 Backlit, int32 EmissiveColor, int32 Tangent, uint8 SharedNormalIndex)
+int32 FHLSLMaterialTranslator::StrataHairBSDF(int32 BaseColor, int32 Scatter, int32 Specular, int32 Roughness, int32 Backlit, int32 EmissiveColor, int32 Tangent, const FString& SharedNormalIndexMacro)
 {
 	return AddCodeChunk(
-		MCT_Strata, TEXT("GetStrataHairBSDF(%s, %s, %s, %s, %s, %s, %u) /* %s */"),
+		MCT_Strata, TEXT("GetStrataHairBSDF(%s, %s, %s, %s, %s, %s, %s) /* %s */"),
 		*GetParameterCode(BaseColor),
 		*GetParameterCode(Scatter),
 		*GetParameterCode(Specular),
 		*GetParameterCode(Roughness),
 		*GetParameterCode(Backlit),
 		*GetParameterCode(EmissiveColor),
-		SharedNormalIndex,
+		*SharedNormalIndexMacro,
 		*GetParameterCode(Tangent)
 	);
 }
 
 int32 FHLSLMaterialTranslator::StrataSingleLayerWaterBSDF(
 	int32 BaseColor, int32 Metallic, int32 Specular, int32 Roughness, int32 EmissiveColor, int32 TopMaterialOpacity,
-	int32 WaterAlbedo, int32 WaterExtinction, int32 WaterPhaseG, int32 ColorScaleBehindWater, int32 Normal, uint8 SharedNormalIndex)
+	int32 WaterAlbedo, int32 WaterExtinction, int32 WaterPhaseG, int32 ColorScaleBehindWater, int32 Normal, const FString& SharedNormalIndexMacro)
 {
 	return AddCodeChunk(
-		MCT_Strata, TEXT("GetStrataSingleLayerWaterBSDF(%s, %s, %s, %s, %s, %s, %s, %s, %s, %s, %u) /* %s */"),
+		MCT_Strata, TEXT("GetStrataSingleLayerWaterBSDF(%s, %s, %s, %s, %s, %s, %s, %s, %s, %s, %s) /* %s */"),
 		*GetParameterCode(BaseColor),
 		*GetParameterCode(Metallic),
 		*GetParameterCode(Specular),
@@ -8939,7 +9044,7 @@ int32 FHLSLMaterialTranslator::StrataSingleLayerWaterBSDF(
 		*GetParameterCode(WaterExtinction),
 		*GetParameterCode(WaterPhaseG),
 		*GetParameterCode(ColorScaleBehindWater),
-		SharedNormalIndex,
+		*SharedNormalIndexMacro,
 		*GetParameterCode(Normal)
 	);
 }
@@ -8958,6 +9063,21 @@ int32 FHLSLMaterialTranslator::StrataHorizontalMixing(int32 Foreground, int32 Ba
 	);
 }
 
+int32 FHLSLMaterialTranslator::StrataHorizontalMixingParameterBlending(int32 Foreground, int32 Background, int32 Mix, const FString& SharedNormalIndexMacro)
+{
+	if (Foreground == INDEX_NONE || Background == INDEX_NONE || Mix == INDEX_NONE)
+	{
+		return INDEX_NONE;
+	}
+	return AddCodeChunk(
+		MCT_Strata, TEXT("StrataHorizontalMixingParameterBlending(%s, %s, %s, %s)"),
+		*GetParameterCode(Foreground),
+		*GetParameterCode(Background),
+		*GetParameterCode(Mix),
+		*SharedNormalIndexMacro
+	);
+}
+
 int32 FHLSLMaterialTranslator::StrataVerticalLayering(int32 Top, int32 Base)
 {
 	if (Top == INDEX_NONE || Base == INDEX_NONE)
@@ -8971,6 +9091,11 @@ int32 FHLSLMaterialTranslator::StrataVerticalLayering(int32 Top, int32 Base)
 	);
 }
 
+//virtual int32 FHLSLMaterialTranslator::StrataVerticalLayeringParameterBlending(int32 Top, int32 Base, const FString& SharedNormalIndexMacro)
+//{
+//	// STRATA_TODO
+//}
+
 int32 FHLSLMaterialTranslator::StrataAdd(int32 A, int32 B)
 {
 	if (A == INDEX_NONE || B == INDEX_NONE)
@@ -8984,6 +9109,11 @@ int32 FHLSLMaterialTranslator::StrataAdd(int32 A, int32 B)
 	);
 }
 
+//virtual int32 FHLSLMaterialTranslator::StrataAddParameterBlending(int32 A, int32 B, const FString& SharedNormalIndexMacro)
+//{
+//	// STRATA_TODO
+//}
+
 int32 FHLSLMaterialTranslator::StrataMultiply(int32 A, int32 Weight)
 {
 	if (A == INDEX_NONE || Weight == INDEX_NONE)
diff --git a/Engine/Source/Runtime/Engine/Private/Materials/HLSLMaterialTranslator.h b/Engine/Source/Runtime/Engine/Private/Materials/HLSLMaterialTranslator.h
--- a/Engine/Source/Runtime/Engine/Private/Materials/HLSLMaterialTranslator.h
+++ b/Engine/Source/Runtime/Engine/Private/Materials/HLSLMaterialTranslator.h
@@ -392,18 +392,27 @@ protected:
 	/** Will contain all the shading models picked up from the material expression graph */
 	FMaterialShadingModelField ShadingModelsFromCompilation;
 
-	/** Tracks shared normal & tangent for strata materials. */
-	TMap<int32, FStrataMaterialCompilationInfo> CodeChunkToStrataCompilationInfoMap; // Mapping from node output to StrataMaterial
+	/** Maps a code chunk to a Strata material topology, if the code chunk represent a StrataData. */
+	TMap<int32, FStrataMaterialCompilationInfo> CodeChunkToStrataCompilationInfoMap;
+	/** Only valid after Translate() is called. This is the code chunk representing the front material. It can be used to recover the material topology. */
+	int32 StrataValidFrontMaterialCodeChunkPostTranslate;
+	/** The code initializing the array of shared normal/tangent. */
+	FString StrataPixelNormalInitializerValues;
+	/** The next free index that can be used to represent a unique macros pointing to he shared normal position in the array of shared normal written to memory once the shader is executed. */
 	uint8 NextFreeStrataShaderNormalIndex;
+	/** The effective final shared normal count used by the final shader. */
+	uint8 FinalUsedSharedNormalCount;
+	/** Represent a shared normal description with its associated code. */
 	struct FStrataSharedNormalInfo
 	{
-		int32 NormalCodeChunk;
-		int32 TangentCodeChunk;
-		uint8 SharedNormalIndex;
+		FStrataRegisteredSharedNormal SharedData;
 		FString SharedNormalCode;
 		FString SharedTangentCode;
 	};
-	TMultiMap<uint64, FStrataSharedNormalInfo> CodeChunkToStrataSharedNormal; // Stored into to a multi-map as a normal can be shared accross several tangent basis (i.e. having different tangent)
+	/** Tracks shared normal & tangent used by strata materials, mapping a normal code chunk hash to a SharedMaterialInfo.
+	 * A normal code chunk hash can point to multiple shared info in case it is paired with different tangents. */
+	TMultiMap<uint64, FStrataSharedNormalInfo> CodeChunkToStrataSharedNormal;
+	/** The final material topology and memory analysis that can be feedback to the user in the editor UI */
 	FStrataMaterialAnalysisResult StrataMaterialAnalysis;
 
 	/** Tracks the total number of vt samples in the shader. */
@@ -975,28 +984,33 @@ protected:
 		int32 ThinFilmThickness, 
 		int32 FuzzAmount, int32 FuzzColor, 
 		int32 Thickness,
-		int32 Normal, int32 Tangent, uint8 SharedNormalIndex) override;
-	virtual int32 StrataSheenBSDF(int32 BaseColor, int32 Roughness, int32 Normal, uint8 SharedNormalIndex) override;
+		int32 Normal, int32 Tangent, const FString& SharedNormalIndexMacro) override;
+	virtual int32 StrataSheenBSDF(int32 BaseColor, int32 Roughness, int32 Normal, const FString& SharedNormalIndexMacro) override;
 	virtual int32 StrataVolumetricFogCloudBSDF(int32 Albedo, int32 Extinction, int32 EmissiveColor, int32 AmbientOcclusion) override;
 	virtual int32 StrataUnlitBSDF(int32 EmissiveColor, int32 TransmittanceColor) override;
-	virtual int32 StrataHairBSDF(int32 BaseColor, int32 Scatter, int32 Specular, int32 Roughness, int32 Backlit, int32 EmissiveColor, int32 Tangent, uint8 SharedNormalIndex) override;
+	virtual int32 StrataHairBSDF(int32 BaseColor, int32 Scatter, int32 Specular, int32 Roughness, int32 Backlit, int32 EmissiveColor, int32 Tangent, const FString& SharedNormalIndexMacro) override;
 	virtual int32 StrataSingleLayerWaterBSDF(
 		int32 BaseColor, int32 Metallic, int32 Specular, int32 Roughness, 
 		int32 EmissiveColor, int32 TopMaterialOpacity, int32 WaterAlbedo, int32 WaterExtinction, int32 WaterPhaseG, 
-		int32 ColorScaleBehindWater, int32 Normal, uint8 SharedNormalIndex) override;
+		int32 ColorScaleBehindWater, int32 Normal, const FString& SharedNormalIndexMacro) override;
 	virtual int32 StrataHorizontalMixing(int32 Foreground, int32 Background, int32 Mix) override;
+	virtual int32 StrataHorizontalMixingParameterBlending(int32 Foreground, int32 Background, int32 Mix, const FString& SharedNormalIndexMacro) override;
 	virtual int32 StrataVerticalLayering(int32 Top, int32 Base) override;
+	// STRATA_TODO virtual int32 StrataVerticalLayeringParameterBlending(int32 Top, int32 Base, const FString& SharedNormalIndexMacro) override;
 	virtual int32 StrataAdd(int32 A, int32 B) override;
+	// STRATA_TODO virtual int32 StrataAddParameterBlending(int32 A, int32 B, const FString& SharedNormalIndexMacro) override;
 	virtual int32 StrataMultiply(int32 A, int32 Weight) override;
 	virtual int32 StrataTransmittanceToMFP(int32 TransmittanceColor, int32 DesiredThickness, int32 OutputIndex) override;
 
 	virtual void StrataCompilationInfoRegisterCodeChunk(int32 CodeChunk, FStrataMaterialCompilationInfo& StrataMaterialCompilationInfo) override;
 	virtual bool StrataCompilationInfoContainsCodeChunk(int32 CodeChunk) override;
 	virtual const FStrataMaterialCompilationInfo& GetStrataCompilationInfo(int32 CodeChunk) override;
-	virtual uint8 StrataCompilationInfoRegisterSharedNormalIndex(int32 NormalCodeChunk) override;
-	virtual uint8 StrataCompilationInfoRegisterSharedNormalIndex(int32 NormalCodeChunk, int32 TangentCodeChunk) override;
+	virtual FStrataRegisteredSharedNormal StrataCompilationInfoRegisterSharedNormal(int32 NormalCodeChunk) override;
+	virtual FStrataRegisteredSharedNormal StrataCompilationInfoRegisterSharedNormal(int32 NormalCodeChunk, int32 TangentCodeChunk) override;
 	virtual uint8 StrataCompilationInfoGetSharedNormalCount() override;
 
+	FStrataSharedNormalInfo StrataCompilationInfoGetMatchingSharedNormalInfo(const FStrataRegisteredSharedNormal& SearchedSharedNormal);
+
 #if HANDLE_CUSTOM_OUTPUTS_AS_MATERIAL_ATTRIBUTES
 	/** Used to translate code for custom output attributes such as ClearCoatBottomNormal */
 	void GenerateCustomAttributeCode(int32 OutputIndex, int32 OutputCode, EMaterialValueType OutputType, FString& DisplayName);
diff --git a/Engine/Source/Runtime/Engine/Private/Materials/MaterialExpressions.cpp b/Engine/Source/Runtime/Engine/Private/Materials/MaterialExpressions.cpp
--- a/Engine/Source/Runtime/Engine/Private/Materials/MaterialExpressions.cpp
+++ b/Engine/Source/Runtime/Engine/Private/Materials/MaterialExpressions.cpp
@@ -19720,6 +19720,15 @@ FName UMaterialExpressionSceneDepthWithoutWater::GetInputName(int32 InputIndex)
 ///////////////////////////////////////////////////////////////////////////////
 // Strata
 
+#if WITH_EDITOR
+static int32 StrataBlendNormal(class FMaterialCompiler* Compiler, int32 NormalCodeChunk0, int32 NormalCodeChunk1, int32 MixCodeChunk)
+{
+	int32 SafeMixCodeChunk = Compiler->Saturate(MixCodeChunk);
+	int32 LerpedNormal = Compiler->Lerp(NormalCodeChunk0, NormalCodeChunk1, SafeMixCodeChunk);
+	int32 BlendedNormalCodeChunk = Compiler->Div(LerpedNormal, Compiler->SquareRoot(Compiler->Dot(LerpedNormal, LerpedNormal)));
+	return BlendedNormalCodeChunk;
+}
+#endif // WITH_EDITOR
 
 #if WITH_EDITOR
 
@@ -19853,8 +19862,8 @@ int32 UMaterialExpressionStrataSlabBSDF::Compile(class FMaterialCompiler* Compil
 
 	int32 NormalCodeChunk = CompileWithDefaultNormalWS(Compiler, Normal);
 
-	int32 TangentCodeChunk = CompileWithDefaultTangentWS(Compiler, Tangent);
-	uint8 SharedNormalIndex = bHasAnisotropy ? StrataCompilationInfoCreateSharedNormal(Compiler, NormalCodeChunk, TangentCodeChunk) : StrataCompilationInfoCreateSharedNormal(Compiler, NormalCodeChunk);
+	int32 TangentCodeChunk = bHasAnisotropy ? CompileWithDefaultTangentWS(Compiler, Tangent) : INDEX_NONE;
+	const FStrataRegisteredSharedNormal NewRegisteredSharedNormal = StrataCompilationInfoCreateSharedNormal(Compiler, NormalCodeChunk, TangentCodeChunk);
 
 	const bool bHasEdgeColor = HasEdgeColor();
 	const bool bHasThinFilm = HasThinFilm();
@@ -19887,8 +19896,8 @@ int32 UMaterialExpressionStrataSlabBSDF::Compile(class FMaterialCompiler* Compil
 		CompileWithDefaultFloat1(Compiler, Thickness, STRATA_LAYER_DEFAULT_THICKNESS_CM),
 		NormalCodeChunk,
 		TangentCodeChunk,
-		SharedNormalIndex);
-	StrataCompilationInfoCreateSingleBSDFMaterial(Compiler, OutputCodeChunk, SharedNormalIndex, STRATA_BSDF_TYPE_SLAB, bHasSSS, bHasDMFPPluggedIn, bHasEdgeColor, bHasThinFilm, bHasFuzz);
+		Compiler->GetStrataSharedNormalIndexMacro(NewRegisteredSharedNormal));
+	StrataCompilationInfoCreateSingleBSDFMaterial(Compiler, OutputCodeChunk, NewRegisteredSharedNormal, STRATA_BSDF_TYPE_SLAB, bHasSSS, bHasDMFPPluggedIn, bHasEdgeColor, bHasThinFilm, bHasFuzz);
 
 	return OutputCodeChunk;
 }
@@ -20113,8 +20122,7 @@ int32 UMaterialExpressionStrataVolumetricFogCloudBSDF::Compile(class FMaterialCo
 		CompileWithDefaultFloat3(Compiler, EmissiveColor, 0.0f, 0.0f, 0.0f),
 		CompileWithDefaultFloat1(Compiler, AmbientOcclusion, 1.0f));
 
-	uint8 FakeSharedNormalIndex = 0;
-	StrataCompilationInfoCreateSingleBSDFMaterial(Compiler, OutputCodeChunk, FakeSharedNormalIndex, STRATA_BSDF_TYPE_VOLUMETRICFOGCLOUD);
+	StrataCompilationInfoCreateSingleBSDFMaterial(Compiler, OutputCodeChunk, StrataCompilationInfoCreateNullSharedNormal(), STRATA_BSDF_TYPE_VOLUMETRICFOGCLOUD);
 
 	return OutputCodeChunk;
 }
@@ -20185,8 +20193,7 @@ int32 UMaterialExpressionStrataUnlitBSDF::Compile(class FMaterialCompiler* Compi
 		CompileWithDefaultFloat3(Compiler, EmissiveColor, 0.0f, 0.0f, 0.0f),
 		CompileWithDefaultFloat3(Compiler, TransmittanceColor, 1.0f, 1.0f, 1.0f));
 
-	uint8 FakeSharedNormalIndex = 0;
-	StrataCompilationInfoCreateSingleBSDFMaterial(Compiler, OutputCodeChunk, FakeSharedNormalIndex, STRATA_BSDF_TYPE_UNLIT);
+	StrataCompilationInfoCreateSingleBSDFMaterial(Compiler, OutputCodeChunk, StrataCompilationInfoCreateNullSharedNormal(), STRATA_BSDF_TYPE_UNLIT);
 
 	return OutputCodeChunk;
 }
@@ -20249,7 +20256,7 @@ int32 UMaterialExpressionStrataHairBSDF::Compile(class FMaterialCompiler* Compil
 {
 	// For hair, the shared normal index in fact represent the tangent
 	int32 TangentCodeChunk = CompileWithDefaultTangentWS(Compiler, Tangent);
-	uint8 SharedNormalIndex = StrataCompilationInfoCreateSharedNormal(Compiler, TangentCodeChunk);
+	const FStrataRegisteredSharedNormal NewRegisteredSharedNormal = StrataCompilationInfoCreateSharedNormal(Compiler, TangentCodeChunk);
 
 	int32 OutputCodeChunk = Compiler->StrataHairBSDF(
 		CompileWithDefaultFloat3(Compiler, BaseColor,	0.0f, 0.0f, 0.0f),
@@ -20259,9 +20266,9 @@ int32 UMaterialExpressionStrataHairBSDF::Compile(class FMaterialCompiler* Compil
 		CompileWithDefaultFloat1(Compiler, Backlit,		0.0f),
 		CompileWithDefaultFloat3(Compiler, EmissiveColor,1.0f, 0.0f, 0.0f),
 		TangentCodeChunk,
-		SharedNormalIndex);
+		Compiler->GetStrataSharedNormalIndexMacro(NewRegisteredSharedNormal));
 
-	StrataCompilationInfoCreateSingleBSDFMaterial(Compiler, OutputCodeChunk, SharedNormalIndex, STRATA_BSDF_TYPE_HAIR);
+	StrataCompilationInfoCreateSingleBSDFMaterial(Compiler, OutputCodeChunk, NewRegisteredSharedNormal, STRATA_BSDF_TYPE_HAIR);
 
 	return OutputCodeChunk;
 }
@@ -20338,7 +20345,7 @@ UMaterialExpressionStrataSingleLayerWaterBSDF::UMaterialExpressionStrataSingleLa
 int32 UMaterialExpressionStrataSingleLayerWaterBSDF::Compile(class FMaterialCompiler* Compiler, int32 OutputIndex)
 {
 	int32 NormalCodeChunk = CompileWithDefaultNormalWS(Compiler, Normal);
-	uint8 SharedNormalIndex = StrataCompilationInfoCreateSharedNormal(Compiler, NormalCodeChunk);
+	const FStrataRegisteredSharedNormal NewRegisteredSharedNormal = StrataCompilationInfoCreateSharedNormal(Compiler, NormalCodeChunk);
 
 	int32 OutputCodeChunk = Compiler->StrataSingleLayerWaterBSDF(
 		CompileWithDefaultFloat3(Compiler, BaseColor, 0.0f, 0.0f, 0.0f),
@@ -20352,9 +20359,9 @@ int32 UMaterialExpressionStrataSingleLayerWaterBSDF::Compile(class FMaterialComp
 		CompileWithDefaultFloat1(Compiler, WaterPhaseG, 0.0f),
 		CompileWithDefaultFloat3(Compiler, ColorScaleBehindWater, 1.0f, 1.0f, 1.0f),
 		NormalCodeChunk,
-		SharedNormalIndex);
+		Compiler->GetStrataSharedNormalIndexMacro(NewRegisteredSharedNormal));
 
-	StrataCompilationInfoCreateSingleBSDFMaterial(Compiler, OutputCodeChunk, SharedNormalIndex, STRATA_BSDF_TYPE_SINGLELAYERWATER);
+	StrataCompilationInfoCreateSingleBSDFMaterial(Compiler, OutputCodeChunk, NewRegisteredSharedNormal, STRATA_BSDF_TYPE_SINGLELAYERWATER);
 
 	return OutputCodeChunk;
 }
@@ -20427,6 +20434,7 @@ void UMaterialExpressionStrataSingleLayerWaterBSDF::GatherStrataMaterialInfo(FSt
 
 UMaterialExpressionStrataHorizontalMixing::UMaterialExpressionStrataHorizontalMixing(const FObjectInitializer& ObjectInitializer)
 	: Super(ObjectInitializer)
+	, bUseParameterBlending(false)
 {
 	struct FConstructorStatics
 	{
@@ -20454,24 +20462,88 @@ int32 UMaterialExpressionStrataHorizontalMixing::Compile(class FMaterialCompiler
 	int32 ForegroundCodeChunk = Foreground.Compile(Compiler);
 	int32 BackgroundCodeChunk = Background.Compile(Compiler);
 
-	int32 OutputCodeChunk = Compiler->StrataHorizontalMixing(
-		ForegroundCodeChunk,
-		BackgroundCodeChunk,
-		Mix.GetTracedInput().Expression ? Mix.Compile(Compiler) : Compiler->Constant(0.0f));
-
 	if (!Compiler->StrataCompilationInfoContainsCodeChunk(ForegroundCodeChunk) || !Compiler->StrataCompilationInfoContainsCodeChunk(BackgroundCodeChunk))
 	{
 		return Compiler->Errorf(TEXT("Could not find ForegroundCodeChunk or BackgroundCodeChunk to mix"));
 	}
-	FStrataMaterialCompilationInfo StrataInfo = StrataCompilationInfoHorizontalMixing(Compiler, Compiler->GetStrataCompilationInfo(ForegroundCodeChunk), Compiler->GetStrataCompilationInfo(BackgroundCodeChunk));
-	Compiler->StrataCompilationInfoRegisterCodeChunk(OutputCodeChunk, StrataInfo);
+	const FStrataMaterialCompilationInfo& ForegroundStrataData = Compiler->GetStrataCompilationInfo(ForegroundCodeChunk);
+	const FStrataMaterialCompilationInfo& BackgroundStrataData = Compiler->GetStrataCompilationInfo(BackgroundCodeChunk);
+
+	const int32 MixCodeChunk = Mix.GetTracedInput().Expression ? Mix.Compile(Compiler) : Compiler->Constant(0.0f);
+
+	int32 OutputCodeChunk = INDEX_NONE;
+	if (bUseParameterBlending)
+	{
+		if (ForegroundStrataData.TotalBSDFCount != 1)
+		{
+			return Compiler->Errorf(TEXT("Foreground: cannot parameter horizontal blend complex material topology using parameter blending. Only a one-to-one blending is supported today."));
+		}
+		if (BackgroundStrataData.TotalBSDFCount != 1)
+		{
+			return Compiler->Errorf(TEXT("Background: cannot parameter horizontal blend complex material topology using parameter blending. Only a one-to-one blending is supported today."));
+		}
+
+		const FStrataMaterialCompilationInfo::FBSDF& ForegroundBSDF = ForegroundStrataData.Layers[0].BSDFs[0];
+		const FStrataMaterialCompilationInfo::FBSDF& BackgroundBSDF = BackgroundStrataData.Layers[0].BSDFs[0];
+		if (ForegroundBSDF.Type != STRATA_BSDF_TYPE_SLAB)
+		{
+			return Compiler->Errorf(TEXT("Foreground: we can only apply parameter blending to Slab BSDF."));
+		}
+		if (BackgroundBSDF.Type != STRATA_BSDF_TYPE_SLAB)
+		{
+			return Compiler->Errorf(TEXT("Background: we can only apply parameter blending to Slab BSDF."));
+		}
+
+		// Compute the new Normal and Tangent resulting from the blending using code chunk
+		const int32 NewNormalCodeChunk = StrataBlendNormal(Compiler, BackgroundBSDF.RegisteredSharedNormal.NormalCodeChunk, ForegroundBSDF.RegisteredSharedNormal.NormalCodeChunk, MixCodeChunk);
+		// The tangent is optional so we treat it differently if INDEX_NONE is specified
+		int32 NewTangentCodeChunk = INDEX_NONE;
+		if (ForegroundBSDF.RegisteredSharedNormal.TangentCodeChunk != INDEX_NONE && BackgroundBSDF.RegisteredSharedNormal.TangentCodeChunk != INDEX_NONE)
+		{
+			NewTangentCodeChunk = StrataBlendNormal(Compiler, ForegroundBSDF.RegisteredSharedNormal.TangentCodeChunk, BackgroundBSDF.RegisteredSharedNormal.TangentCodeChunk, MixCodeChunk);
+		}
+		else if (ForegroundBSDF.RegisteredSharedNormal.TangentCodeChunk != INDEX_NONE)
+		{
+			NewTangentCodeChunk = ForegroundBSDF.RegisteredSharedNormal.TangentCodeChunk;
+		}
+		else if (BackgroundBSDF.RegisteredSharedNormal.TangentCodeChunk != INDEX_NONE)
+		{
+			NewTangentCodeChunk = BackgroundBSDF.RegisteredSharedNormal.TangentCodeChunk;
+		}
+		const FStrataRegisteredSharedNormal NewRegisteredSharedNormal = StrataCompilationInfoCreateSharedNormal(Compiler, NewNormalCodeChunk, NewTangentCodeChunk);
+
+		OutputCodeChunk = Compiler->StrataHorizontalMixingParameterBlending(ForegroundCodeChunk, BackgroundCodeChunk, MixCodeChunk, Compiler->GetStrataSharedNormalIndexMacro(NewRegisteredSharedNormal));
+
+		StrataCompilationInfoCreateSingleBSDFMaterial(Compiler, OutputCodeChunk, NewRegisteredSharedNormal, STRATA_BSDF_TYPE_SLAB);
+		FStrataMaterialCompilationInfo StrataInfo = StrataCompilationInfoHorizontalMixingParamBlend(Compiler, ForegroundStrataData, BackgroundStrataData, NewRegisteredSharedNormal);
+
+		Compiler->StrataCompilationInfoRegisterCodeChunk(OutputCodeChunk, StrataInfo);
+	}
+	else
+	{
+		OutputCodeChunk = Compiler->StrataHorizontalMixing(
+			ForegroundCodeChunk,
+			BackgroundCodeChunk,
+			MixCodeChunk);
+
+		FStrataMaterialCompilationInfo StrataInfo = StrataCompilationInfoHorizontalMixing(Compiler, ForegroundStrataData, BackgroundStrataData);
+
+		Compiler->StrataCompilationInfoRegisterCodeChunk(OutputCodeChunk, StrataInfo);
+	}
 
 	return OutputCodeChunk;
 }
 
 void UMaterialExpressionStrataHorizontalMixing::GetCaption(TArray<FString>& OutCaptions) const
 {
-	OutCaptions.Add(TEXT("Strata BSDF Horizontal Blend"));
+	if (bUseParameterBlending)
+	{
+		OutCaptions.Add(TEXT("Strata BSDF Horizontal Parameter Blend"));
+	}
+	else
+	{
+		OutCaptions.Add(TEXT("Strata BSDF Horizontal Blend"));
+	}
 }
 
 uint32 UMaterialExpressionStrataHorizontalMixing::GetOutputType(int32 OutputIndex)
diff --git a/Engine/Source/Runtime/Engine/Private/Materials/StrataMaterial.cpp b/Engine/Source/Runtime/Engine/Private/Materials/StrataMaterial.cpp
--- a/Engine/Source/Runtime/Engine/Private/Materials/StrataMaterial.cpp
+++ b/Engine/Source/Runtime/Engine/Private/Materials/StrataMaterial.cpp
@@ -43,24 +43,35 @@ static void UpdateTotalBSDFCount(FMaterialCompiler* Compiler, FStrataMaterialCom
 	}
 }
 
-
-uint8 StrataCompilationInfoCreateSharedNormal(FMaterialCompiler* Compiler, int32 NormalCodeChunk)
+FStrataRegisteredSharedNormal StrataCompilationInfoCreateNullSharedNormal()
 {
-	return Compiler->StrataCompilationInfoRegisterSharedNormalIndex(NormalCodeChunk);
+	FStrataRegisteredSharedNormal FakeSharedNormal;
+	FakeSharedNormal.NormalCodeChunk = INDEX_NONE;
+	FakeSharedNormal.TangentCodeChunk = INDEX_NONE;
+	FakeSharedNormal.NormalCodeChunkHash = 0;
+	FakeSharedNormal.TangentCodeChunkHash = 0;
+	FakeSharedNormal.GraphSharedNormalIndex = 0;
+	return FakeSharedNormal;
 }
 
-uint8 StrataCompilationInfoCreateSharedNormal(FMaterialCompiler* Compiler, int32 NormalCodeChunk, int32 TangentCodeChunk)
+FStrataRegisteredSharedNormal StrataCompilationInfoCreateSharedNormal(FMaterialCompiler* Compiler, int32 NormalCodeChunk, int32 TangentCodeChunk)
 {
-	return Compiler->StrataCompilationInfoRegisterSharedNormalIndex(NormalCodeChunk, TangentCodeChunk);
+	if (TangentCodeChunk == INDEX_NONE)
+	{
+		return Compiler->StrataCompilationInfoRegisterSharedNormal(NormalCodeChunk);
+	}
+	return Compiler->StrataCompilationInfoRegisterSharedNormal(NormalCodeChunk, TangentCodeChunk);
 }
 
-void StrataCompilationInfoCreateSingleBSDFMaterial(FMaterialCompiler* Compiler, int32 CodeChunk, uint8 SharedNormalIndex, uint8 BSDFType, bool bHasSSS, bool bHasDMFPPluggedIn, bool bHasEdgeColor, bool bHasThinFilm, bool bHasFuzz)
+void StrataCompilationInfoCreateSingleBSDFMaterial(FMaterialCompiler* Compiler, int32 CodeChunk,
+	const FStrataRegisteredSharedNormal& RegisteredSharedNormal,
+	uint8 BSDFType, bool bHasSSS, bool bHasDMFPPluggedIn, bool bHasEdgeColor, bool bHasThinFilm, bool bHasFuzz)
 {
 	FStrataMaterialCompilationInfo StrataInfo;
 	StrataInfo.LayerCount = 1;
 	StrataInfo.Layers[0].BSDFCount = 1;
 	StrataInfo.Layers[0].BSDFs[0].Type = BSDFType;
-	StrataInfo.Layers[0].BSDFs[0].SharedNormalIndex = SharedNormalIndex;
+	StrataInfo.Layers[0].BSDFs[0].RegisteredSharedNormal = RegisteredSharedNormal;
 	StrataInfo.Layers[0].BSDFs[0].bHasSSS = bHasSSS;
 	StrataInfo.Layers[0].BSDFs[0].bHasDMFPPluggedIn = bHasDMFPPluggedIn;
 	StrataInfo.Layers[0].BSDFs[0].bHasEdgeColor = bHasEdgeColor;
@@ -122,6 +133,28 @@ FStrataMaterialCompilationInfo StrataCompilationInfoHorizontalMixing(FMaterialCo
 }
 
 
+FStrataMaterialCompilationInfo StrataCompilationInfoHorizontalMixingParamBlend(FMaterialCompiler* Compiler, const FStrataMaterialCompilationInfo& A, const FStrataMaterialCompilationInfo& B, const FStrataRegisteredSharedNormal& RegisteredSharedNormal)
+{
+	check(A.TotalBSDFCount == 1);
+	check(B.TotalBSDFCount == 1);
+
+	FStrataMaterialCompilationInfo StrataInfo = A;
+	FStrataMaterialCompilationInfo::FBSDF& NewBSDF = StrataInfo.Layers[0].BSDFs[0];
+	const FStrataMaterialCompilationInfo::FBSDF& OtherBSDF = B.Layers[0].BSDFs[0];
+
+	NewBSDF.RegisteredSharedNormal = RegisteredSharedNormal;
+
+	// When parameter blending is used, we take the union of all the features activated by input BSDFs.
+	NewBSDF.bHasSSS				|=	OtherBSDF.bHasSSS;
+	NewBSDF.bHasDMFPPluggedIn	|=	OtherBSDF.bHasDMFPPluggedIn;
+	NewBSDF.bHasEdgeColor		|=	OtherBSDF.bHasEdgeColor;
+	NewBSDF.bHasThinFilm		|=	OtherBSDF.bHasThinFilm;
+	NewBSDF.bHasFuzz			|=	OtherBSDF.bHasFuzz;
+	UpdateTotalBSDFCount(Compiler, StrataInfo);
+	return StrataInfo;
+}
+
+
 FStrataMaterialCompilationInfo StrataCompilationInfoVerticalLayering(FMaterialCompiler* Compiler, const FStrataMaterialCompilationInfo& Top, const FStrataMaterialCompilationInfo& Base)
 {
 	if ((Top.LayerCount + Base.LayerCount) > STRATA_MAX_LAYER_COUNT)
diff --git a/Engine/Source/Runtime/Engine/Private/Materials/StrataMaterial.h b/Engine/Source/Runtime/Engine/Private/Materials/StrataMaterial.h
--- a/Engine/Source/Runtime/Engine/Private/Materials/StrataMaterial.h
+++ b/Engine/Source/Runtime/Engine/Private/Materials/StrataMaterial.h
@@ -5,6 +5,7 @@
 #include "CoreMinimal.h"
 #include "StrataDefinitions.h"
 #include "Containers/Map.h"
+#include "MaterialCompiler.h"
 
 
 
@@ -17,7 +18,7 @@ struct FStrataMaterialCompilationInfo
 	struct FBSDF
 	{
 		uint8 Type;
-		uint8 SharedNormalIndex;
+		FStrataRegisteredSharedNormal RegisteredSharedNormal;
 
 		// Those properties are centered around the SLAB node. Maybe this can be abstracted at some point.
 		bool  bHasSSS;
@@ -48,14 +49,18 @@ struct FStrataMaterialCompilationInfo
 
 FString GetStrataBSDFName(uint8 BSDFType);
 
-uint8 StrataCompilationInfoCreateSharedNormal(FMaterialCompiler* Compiler, int32 NormalCodeChunk);
-uint8 StrataCompilationInfoCreateSharedNormal(FMaterialCompiler* Compiler, int32 NormalCodeChunk, int32 TangentCodeChunk);
+FStrataRegisteredSharedNormal StrataCompilationInfoCreateNullSharedNormal();
+FStrataRegisteredSharedNormal StrataCompilationInfoCreateSharedNormal(FMaterialCompiler* Compiler, int32 NormalCodeChunk, int32 TangentCodeChunk = INDEX_NONE);
+
 void  StrataCompilationInfoCreateNullBSDF(FMaterialCompiler* Compiler, int32 CodeChunk);
-void  StrataCompilationInfoCreateSingleBSDFMaterial(FMaterialCompiler* Compiler, int32 CodeChunk, uint8 SharedNormalIndex, uint8 BSDFType, bool bHasSSS = false, bool bHasDMFPPluggedIn = false, bool bHasEdgeColor = false, bool bHasThinFilm = false, bool bHasFuzz = false);
+void  StrataCompilationInfoCreateSingleBSDFMaterial(FMaterialCompiler* Compiler, int32 CodeChunk, 
+	const FStrataRegisteredSharedNormal& RegisteredSharedNormal,
+	uint8 BSDFType, bool bHasSSS = false, bool bHasDMFPPluggedIn = false, bool bHasEdgeColor = false, bool bHasThinFilm = false, bool bHasFuzz = false);
 
 FStrataMaterialCompilationInfo StrataCompilationInfoAdd(FMaterialCompiler* Compiler, const FStrataMaterialCompilationInfo& A, const FStrataMaterialCompilationInfo& B);
 FStrataMaterialCompilationInfo StrataCompilationInfoMultiply(FMaterialCompiler* Compiler, const FStrataMaterialCompilationInfo& A);
 FStrataMaterialCompilationInfo StrataCompilationInfoHorizontalMixing(FMaterialCompiler* Compiler, const FStrataMaterialCompilationInfo& A, const FStrataMaterialCompilationInfo& B);
+FStrataMaterialCompilationInfo StrataCompilationInfoHorizontalMixingParamBlend(FMaterialCompiler* Compiler, const FStrataMaterialCompilationInfo& A, const FStrataMaterialCompilationInfo& B, const FStrataRegisteredSharedNormal& RegisteredSharedNormal);
 FStrataMaterialCompilationInfo StrataCompilationInfoVerticalLayering(FMaterialCompiler* Compiler, const FStrataMaterialCompilationInfo& Top, const FStrataMaterialCompilationInfo& Base);
 
 bool StrataIsVolumetricFogCloudOnly(FMaterialCompiler* Compiler, const FStrataMaterialCompilationInfo& Material);
diff --git a/Engine/Source/Runtime/Engine/Public/MaterialCompiler.h b/Engine/Source/Runtime/Engine/Public/MaterialCompiler.h
--- a/Engine/Source/Runtime/Engine/Public/MaterialCompiler.h
+++ b/Engine/Source/Runtime/Engine/Public/MaterialCompiler.h
@@ -48,6 +48,15 @@ enum class EMaterialCompilerType
 	MaterialProxy, /** Flat material proxy compiler */
 };
 
+struct FStrataRegisteredSharedNormal
+{
+	int32 NormalCodeChunk;
+	int32 TangentCodeChunk;
+	uint64 NormalCodeChunkHash;
+	uint64 TangentCodeChunkHash;
+	uint8 GraphSharedNormalIndex;
+};
+
 /** 
  * The interface used to translate material expressions into executable code. 
  * Note: Most member functions should be pure virtual to force a FProxyMaterialCompiler override!
@@ -418,15 +427,16 @@ public:
 		int32 ThinFilmThickness, 
 		int32 FuzzAmount, int32 FuzzColor,
 		int32 Thickness,
-		int32 Normal, int32 Tangent, uint8 SharedNormalIndex) = 0;
-	virtual int32 StrataSheenBSDF(int32 BaseColor, int32 Roughness, int32 Normal, uint8 SharedNormalIndex) = 0;
+		int32 Normal, int32 Tangent, const FString& SharedNormalIndexMacro) = 0;
+	virtual int32 StrataSheenBSDF(int32 BaseColor, int32 Roughness, int32 Normal, const FString& SharedNormalIndexMacro) = 0;
 	virtual int32 StrataVolumetricFogCloudBSDF(int32 Albedo, int32 Extinction, int32 EmissiveColor, int32 AmbientOcclusion) = 0;
 	virtual int32 StrataUnlitBSDF(int32 EmissiveColor, int32 TransmittanceColor) = 0;
-	virtual int32 StrataHairBSDF(int32 BaseColor, int32 Scatter, int32 Specular, int32 Roughness, int32 Backlit, int32 EmissiveColor, int32 Tangent, uint8 SharedNormalIndex) = 0;
+	virtual int32 StrataHairBSDF(int32 BaseColor, int32 Scatter, int32 Specular, int32 Roughness, int32 Backlit, int32 EmissiveColor, int32 Tangent, const FString& SharedNormalIndexMacro) = 0;
 	virtual int32 StrataSingleLayerWaterBSDF(
 		int32 BaseColor, int32 Metallic, int32 Specular, int32 Roughness, int32 EmissiveColor, int32 TopMaterialOpacity, 
-		int32 WaterAlbedo, int32 WaterExtinction, int32 WaterPhaseG, int32 ColorScaleBehindWater, int32 Normal, uint8 SharedNormalIndex) = 0;
+		int32 WaterAlbedo, int32 WaterExtinction, int32 WaterPhaseG, int32 ColorScaleBehindWater, int32 Normal, const FString& SharedNormalIndexMacro) = 0;
 	virtual int32 StrataHorizontalMixing(int32 Foreground, int32 Background, int32 Mix) = 0;
+	virtual int32 StrataHorizontalMixingParameterBlending(int32 Foreground, int32 Background, int32 Mix, const FString& SharedNormalIndexMacro) = 0;
 	virtual int32 StrataVerticalLayering(int32 Top, int32 Base) = 0;
 	virtual int32 StrataAdd(int32 A, int32 B) = 0;
 	virtual int32 StrataMultiply(int32 A, int32 Weight) = 0;
@@ -435,9 +445,13 @@ public:
 	virtual void StrataCompilationInfoRegisterCodeChunk(int32 CodeChunk, FStrataMaterialCompilationInfo& StrataMaterialCompilationInfo) = 0;
 	virtual bool StrataCompilationInfoContainsCodeChunk(int32 CodeChunk) = 0;
 	virtual const FStrataMaterialCompilationInfo& GetStrataCompilationInfo(int32 CodeChunk) = 0;
-	virtual uint8 StrataCompilationInfoRegisterSharedNormalIndex(int32 NormalCodeChunk) = 0;
-	virtual uint8 StrataCompilationInfoRegisterSharedNormalIndex(int32 NormalCodeChunk, int32 TangentCodeChunk) = 0;
+	virtual FStrataRegisteredSharedNormal StrataCompilationInfoRegisterSharedNormal(int32 NormalCodeChunk) = 0;
+	virtual FStrataRegisteredSharedNormal StrataCompilationInfoRegisterSharedNormal(int32 NormalCodeChunk, int32 TangentCodeChunk) = 0;
 	virtual uint8 StrataCompilationInfoGetSharedNormalCount() = 0;
+	FString GetStrataSharedNormalIndexMacro(const FStrataRegisteredSharedNormal& SharedNormal)
+	{
+		return FString::Printf(TEXT("NORMAL_INDEX_%u"), SharedNormal.GraphSharedNormalIndex);
+	}
 
 	// Water
 	virtual int32 SceneDepthWithoutWater(int32 Offset, int32 ViewportUV, bool bUseOffset, float FallbackDepth) = 0;
@@ -907,7 +921,7 @@ public:
 		int32 ThinFilmThickness, 
 		int32 FuzzAmount, int32 FuzzColor,
 		int32 Thickness,
-		int32 Normal, int32 Tangent, uint8 SharedNormalIndex) override
+		int32 Normal, int32 Tangent, const FString& SharedNormalIndexMacro) override
 	{
 		return Compiler->StrataSlabBSDF(
 			BaseColor, EdgeColor, Specular, Metallic, Roughness, Anisotropy,
@@ -917,12 +931,12 @@ public:
 			ThinFilmThickness,
 			FuzzAmount, FuzzColor,
 			Thickness,
-			Normal, Tangent, SharedNormalIndex);
+			Normal, Tangent, SharedNormalIndexMacro);
 	}
 
-	virtual int32 StrataSheenBSDF(int32 BaseColor, int32 Roughness, int32 Normal, uint8 SharedNormalIndex) override
+	virtual int32 StrataSheenBSDF(int32 BaseColor, int32 Roughness, int32 Normal, const FString& SharedNormalIndexMacro) override
 	{
-		return Compiler->StrataSheenBSDF(BaseColor, Roughness, Normal, SharedNormalIndex);
+		return Compiler->StrataSheenBSDF(BaseColor, Roughness, Normal, SharedNormalIndexMacro);
 	}
 
 	virtual int32 StrataVolumetricFogCloudBSDF(int32 Albedo, int32 Extinction, int32 EmissiveColor, int32 AmbientOcclusion) override
@@ -935,18 +949,18 @@ public:
 		return Compiler->StrataUnlitBSDF(EmissiveColor, TransmittanceColor);
 	}
 
-	virtual int32 StrataHairBSDF(int32 BaseColor, int32 Scatter, int32 Specular, int32 Roughness, int32 Backlit, int32 EmissiveColor, int32 Tangent, uint8 SharedNormalIndex) override
+	virtual int32 StrataHairBSDF(int32 BaseColor, int32 Scatter, int32 Specular, int32 Roughness, int32 Backlit, int32 EmissiveColor, int32 Tangent, const FString& SharedNormalIndexMacro) override
 	{
-		return Compiler->StrataHairBSDF(BaseColor, Scatter, Specular, Roughness, Backlit, EmissiveColor, Tangent, SharedNormalIndex);
+		return Compiler->StrataHairBSDF(BaseColor, Scatter, Specular, Roughness, Backlit, EmissiveColor, Tangent, SharedNormalIndexMacro);
 	}
 
 	virtual int32 StrataSingleLayerWaterBSDF(
 		int32 BaseColor, int32 Metallic, int32 Specular, int32 Roughness, int32 EmissiveColor, int32 TopMaterialOpacity, 
-		int32 WaterAlbedo, int32 WaterExtinction, int32 WaterPhaseG, int32 ColorScaleBehindWater, int32 Normal, uint8 SharedNormalIndex) override
+		int32 WaterAlbedo, int32 WaterExtinction, int32 WaterPhaseG, int32 ColorScaleBehindWater, int32 Normal, const FString& SharedNormalIndexMacro) override
 	{
 		return Compiler->StrataSingleLayerWaterBSDF(
 			BaseColor, Metallic, Specular, Roughness, EmissiveColor, TopMaterialOpacity, 
-			WaterAlbedo, WaterExtinction, WaterPhaseG, ColorScaleBehindWater, Normal, SharedNormalIndex);
+			WaterAlbedo, WaterExtinction, WaterPhaseG, ColorScaleBehindWater, Normal, SharedNormalIndexMacro);
 	}
 
 	virtual int32 StrataHorizontalMixing(int32 Foreground, int32 Background, int32 Mix) override
@@ -954,6 +968,11 @@ public:
 		return Compiler->StrataHorizontalMixing(Foreground, Background, Mix);
 	}
 
+	virtual int32 StrataHorizontalMixingParameterBlending(int32 Foreground, int32 Background, int32 Mix, const FString& SharedNormalIndexMacro) override
+	{
+		return Compiler->StrataHorizontalMixingParameterBlending(Foreground, Background, Mix, SharedNormalIndexMacro);
+	}
+
 	virtual int32 StrataVerticalLayering(int32 Top, int32 Base) override
 	{
 		return Compiler->StrataVerticalLayering(Top, Base);
@@ -989,14 +1008,14 @@ public:
 		return Compiler->GetStrataCompilationInfo(CodeChunk);
 	}
 
-	virtual uint8 StrataCompilationInfoRegisterSharedNormalIndex(int32 NormalCodeChunk) override
+	virtual FStrataRegisteredSharedNormal StrataCompilationInfoRegisterSharedNormal(int32 NormalCodeChunk) override
 	{
-		return Compiler->StrataCompilationInfoRegisterSharedNormalIndex(NormalCodeChunk);
+		return Compiler->StrataCompilationInfoRegisterSharedNormal(NormalCodeChunk);
 	}
 
-	virtual uint8 StrataCompilationInfoRegisterSharedNormalIndex(int32 NormalCodeChunk, int32 TangentCodeChunk) override
+	virtual FStrataRegisteredSharedNormal StrataCompilationInfoRegisterSharedNormal(int32 NormalCodeChunk, int32 TangentCodeChunk) override
 	{
-		return Compiler->StrataCompilationInfoRegisterSharedNormalIndex(NormalCodeChunk, TangentCodeChunk);
+		return Compiler->StrataCompilationInfoRegisterSharedNormal(NormalCodeChunk, TangentCodeChunk);
 	}
 
 	virtual uint8 StrataCompilationInfoGetSharedNormalCount() override
