Appease PVSStudio and add comments

#rnx

diff --git a/Engine/Source/Runtime/Engine/Classes/Engine/AssetManager.h b/Engine/Source/Runtime/Engine/Classes/Engine/AssetManager.h
--- a/Engine/Source/Runtime/Engine/Classes/Engine/AssetManager.h
+++ b/Engine/Source/Runtime/Engine/Classes/Engine/AssetManager.h
@@ -102,7 +102,9 @@ public:
 	void PopBulkScanning();
 
 protected:
+	/** Should only be called from PushBulkScanning() and override */
 	virtual void StartBulkScanning();
+	/** Should only be called from PopBulkScanning() and override */
 	virtual void StopBulkScanning();
 	bool IsBulkScanning() const { return NumBulkScanRequests > 0 ; }
 
diff --git a/Engine/Source/Runtime/Engine/Private/AssetManager.cpp b/Engine/Source/Runtime/Engine/Private/AssetManager.cpp
--- a/Engine/Source/Runtime/Engine/Private/AssetManager.cpp
+++ b/Engine/Source/Runtime/Engine/Private/AssetManager.cpp
@@ -243,7 +243,7 @@ public:
 		}
 		else
 		{
-			++Scanner; // Avoid starting bulk scan after OnPostEngineInit
+			Scanner = reinterpret_cast<UAssetManager*>(8); // Avoid starting bulk scan after OnPostEngineInit
 		}
 	}
 }
