Typed Element Common Actions now gather actions based on element list, instead of selection set.

diff --git a/Engine/Source/Editor/LevelEditor/Private/Elements/Actor/ActorElementLevelEditorCommonActionsCustomization.cpp b/Engine/Source/Editor/LevelEditor/Private/Elements/Actor/ActorElementLevelEditorCommonActionsCustomization.cpp
--- a/Engine/Source/Editor/LevelEditor/Private/Elements/Actor/ActorElementLevelEditorCommonActionsCustomization.cpp
+++ b/Engine/Source/Editor/LevelEditor/Private/Elements/Actor/ActorElementLevelEditorCommonActionsCustomization.cpp
@@ -10,17 +10,17 @@
 #include "EditorModeManager.h"
 #include "Toolkits/IToolkitHost.h"
 
-void FActorElementLevelEditorCommonActionsCustomization::GetElementsForAction(const TTypedElement<UTypedElementWorldInterface>& InElementWorldHandle, const UTypedElementSelectionSet* InSelectionSet, UTypedElementList* OutElementsToDelete)
+void FActorElementLevelEditorCommonActionsCustomization::GetElementsForAction(const TTypedElement<UTypedElementWorldInterface>& InElementWorldHandle, const UTypedElementList* InElementList, UTypedElementList* OutElementsForAction)
 {
 	AActor* Actor = ActorElementDataUtil::GetActorFromHandleChecked(InElementWorldHandle);
 
-	if (UComponentElementSelectionInterface::HasSelectedComponents(InSelectionSet->GetElementList()))
+	if (UComponentElementSelectionInterface::HasSelectedComponents(InElementList))
 	{
 		// If we have components selected then we will perform the action on those rather than the actors
 		return;
 	}
 
-	FTypedElementCommonActionsCustomization::GetElementsForAction(InElementWorldHandle, InSelectionSet, OutElementsToDelete);
+	FTypedElementCommonActionsCustomization::GetElementsForAction(InElementWorldHandle, InElementList, OutElementsForAction);
 }
 
 bool FActorElementLevelEditorCommonActionsCustomization::DeleteElements(UTypedElementWorldInterface* InWorldInterface, TArrayView<const FTypedElementHandle> InElementHandles, UWorld* InWorld, UTypedElementSelectionSet* InSelectionSet, const FTypedElementDeletionOptions& InDeletionOptions)
diff --git a/Engine/Source/Editor/LevelEditor/Private/Elements/Actor/ActorElementLevelEditorCommonActionsCustomization.h b/Engine/Source/Editor/LevelEditor/Private/Elements/Actor/ActorElementLevelEditorCommonActionsCustomization.h
--- a/Engine/Source/Editor/LevelEditor/Private/Elements/Actor/ActorElementLevelEditorCommonActionsCustomization.h
+++ b/Engine/Source/Editor/LevelEditor/Private/Elements/Actor/ActorElementLevelEditorCommonActionsCustomization.h
@@ -8,7 +8,7 @@
 class FActorElementLevelEditorCommonActionsCustomization : public FTypedElementCommonActionsCustomization, public FTypedElementAssetEditorToolkitHostMixin
 {
 public:
-	virtual void GetElementsForAction(const TTypedElement<UTypedElementWorldInterface>& InElementWorldHandle, const UTypedElementSelectionSet* InSelectionSet, UTypedElementList* OutElementsToDelete) override;
+	virtual void GetElementsForAction(const TTypedElement<UTypedElementWorldInterface>& InElementWorldHandle, const UTypedElementList* InElementList, UTypedElementList* OutElementsForAction) override;
 	virtual bool DeleteElements(UTypedElementWorldInterface* InWorldInterface, TArrayView<const FTypedElementHandle> InElementHandles, UWorld* InWorld, UTypedElementSelectionSet* InSelectionSet, const FTypedElementDeletionOptions& InDeletionOptions) override;
 	virtual void DuplicateElements(UTypedElementWorldInterface* InWorldInterface, TArrayView<const FTypedElementHandle> InElementHandles, UWorld* InWorld, const FVector& InLocationOffset, TArray<FTypedElementHandle>& OutNewElements) override;
 };
diff --git a/Engine/Source/Editor/UnrealEd/Private/LevelEditorViewport.cpp b/Engine/Source/Editor/UnrealEd/Private/LevelEditorViewport.cpp
--- a/Engine/Source/Editor/UnrealEd/Private/LevelEditorViewport.cpp
+++ b/Engine/Source/Editor/UnrealEd/Private/LevelEditorViewport.cpp
@@ -2610,7 +2610,7 @@ bool FLevelEditorViewportClient::InputWidgetDelta(FViewport* InViewport, EAxisLi
 							if (CommonActions)
 							{
 								const UTypedElementList* ElementsToManipulate = GetElementsToManipulate();
-								TArray<FTypedElementHandle> DuplicatedElements = CommonActions->DuplicateElements(ElementsToManipulate, GetWorld(), FVector::ZeroVector);
+								TArray<FTypedElementHandle> DuplicatedElements = CommonActions->DuplicateElementsInList(ElementsToManipulate, GetWorld(), FVector::ZeroVector);
 
 								// Exclusively select the new elements, so that future gizmo interaction manipulates those items instead
 								if (DuplicatedElements.Num() > 0)
diff --git a/Engine/Source/Runtime/Engine/Private/Elements/Framework/TypedElementCommonActions.cpp b/Engine/Source/Runtime/Engine/Private/Elements/Framework/TypedElementCommonActions.cpp
--- a/Engine/Source/Runtime/Engine/Private/Elements/Framework/TypedElementCommonActions.cpp
+++ b/Engine/Source/Runtime/Engine/Private/Elements/Framework/TypedElementCommonActions.cpp
@@ -7,7 +7,7 @@
 
 #include "UObject/GCObjectScopeGuard.h"
 
-void FTypedElementCommonActionsCustomization::GetElementsForAction(const TTypedElement<UTypedElementWorldInterface>& InElementWorldHandle, const UTypedElementSelectionSet* InSelectionSet, UTypedElementList* OutElementsToDelete)
+void FTypedElementCommonActionsCustomization::GetElementsForAction(const TTypedElement<UTypedElementWorldInterface>& InElementWorldHandle, const UTypedElementList* InElementList, UTypedElementList* OutElementsToDelete)
 {
 	OutElementsToDelete->Add(InElementWorldHandle);
 }
@@ -23,14 +23,14 @@ void FTypedElementCommonActionsCustomization::DuplicateElements(UTypedElementWor
 }
 
 
-void UTypedElementCommonActions::GetSelectedElementsForAction(const UTypedElementSelectionSet* InSelectionSet, UTypedElementList* OutElementsToDelete) const
+void UTypedElementCommonActions::GetElementsForAction(const UTypedElementList* InElementList, UTypedElementList* OutElementsForAction) const
 {
-	OutElementsToDelete->Reset();
-	InSelectionSet->ForEachSelectedElement<UTypedElementWorldInterface>([this, InSelectionSet, OutElementsToDelete](const TTypedElement<UTypedElementWorldInterface>& InElementWorldHandle)
+	OutElementsForAction->Reset();
+	InElementList->ForEachElement<UTypedElementWorldInterface>([this, InElementList, OutElementsForAction](const TTypedElement<UTypedElementWorldInterface>& InElementWorldHandle)
 	{
 		FTypedElementCommonActionsElement CommonActionsElement(InElementWorldHandle, GetInterfaceCustomizationByTypeId(InElementWorldHandle.GetId().GetTypeId()));
 		check(CommonActionsElement.IsSet());
-		CommonActionsElement.GetElementsForAction(InSelectionSet, OutElementsToDelete);
+		CommonActionsElement.GetElementsForAction(InElementList, OutElementsForAction);
 		return true;
 	});
 }
@@ -42,15 +42,15 @@ bool UTypedElementCommonActions::DeleteElements(const TArray<FTypedElementHandle
 
 bool UTypedElementCommonActions::DeleteElements(TArrayView<const FTypedElementHandle> ElementHandles, UWorld* World, UTypedElementSelectionSet* SelectionSet, const FTypedElementDeletionOptions& DeletionOptions)
 {
-	TMap<FTypedHandleTypeId, TArray<FTypedElementHandle>> ElementsToDuplicateByType;
-	TypedElementUtil::BatchElementsByType(ElementHandles, ElementsToDuplicateByType);
+	TMap<FTypedHandleTypeId, TArray<FTypedElementHandle>> ElementsToDeleteByType;
+	TypedElementUtil::BatchElementsByType(ElementHandles, ElementsToDeleteByType);
 
 	bool bSuccess = false;
 
 	UTypedElementRegistry* Registry = UTypedElementRegistry::GetInstance();
 	UTypedElementRegistry::FDisableElementDestructionOnGC GCGuard(Registry);
 
-	for (const auto& ElementsByTypePair : ElementsToDuplicateByType)
+	for (const auto& ElementsByTypePair : ElementsToDeleteByType)
 	{
 		FTypedElementCommonActionsCustomization* CommonActionsCustomization = GetInterfaceCustomizationByTypeId(ElementsByTypePair.Key);
 		UTypedElementWorldInterface* WorldInterface = Registry->GetElementInterface<UTypedElementWorldInterface>(ElementsByTypePair.Key);
@@ -65,15 +65,15 @@ bool UTypedElementCommonActions::DeleteElements(TArrayView<const FTypedElementHa
 
 bool UTypedElementCommonActions::DeleteElements(const UTypedElementList* ElementList, UWorld* World, UTypedElementSelectionSet* SelectionSet, const FTypedElementDeletionOptions& DeletionOptions)
 {
-	TMap<FTypedHandleTypeId, TArray<FTypedElementHandle>> ElementsToDuplicateByType;
-	TypedElementUtil::BatchElementsByType(ElementList, ElementsToDuplicateByType);
+	TMap<FTypedHandleTypeId, TArray<FTypedElementHandle>> ElementsToDeleteByType;
+	TypedElementUtil::BatchElementsByType(ElementList, ElementsToDeleteByType);
 
 	bool bSuccess = false;
 
 	UTypedElementRegistry* Registry = UTypedElementRegistry::GetInstance();
 	UTypedElementRegistry::FDisableElementDestructionOnGC GCGuard(Registry);
 
-	for (const auto& ElementsByTypePair : ElementsToDuplicateByType)
+	for (const auto& ElementsByTypePair : ElementsToDeleteByType)
 	{
 		FTypedElementCommonActionsCustomization* CommonActionsCustomization = GetInterfaceCustomizationByTypeId(ElementsByTypePair.Key);
 		UTypedElementWorldInterface* WorldInterface = Registry->GetElementInterface<UTypedElementWorldInterface>(ElementsByTypePair.Key);
@@ -87,9 +87,14 @@ bool UTypedElementCommonActions::DeleteElements(const UTypedElementList* Element
 }
 
 bool UTypedElementCommonActions::DeleteSelectedElements(UTypedElementSelectionSet* SelectionSet, UWorld* World, const FTypedElementDeletionOptions& DeletionOptions)
+{
+	return DeleteElementsInList(SelectionSet->GetElementList(), World, SelectionSet, DeletionOptions);
+}
+
+bool UTypedElementCommonActions::DeleteElementsInList(const UTypedElementList* ElementList, UWorld* World, UTypedElementSelectionSet* SelectionSet, const FTypedElementDeletionOptions& DeletionOptions)
 {
 	TGCObjectScopeGuard<UTypedElementList> ElementsForAction(UTypedElementRegistry::GetInstance()->CreateElementList());
-	GetSelectedElementsForAction(SelectionSet, ElementsForAction.Get());
+	GetElementsForAction(ElementList, ElementsForAction.Get());
 
 	const bool bSuccess = DeleteElements(ElementsForAction.Get(), World, SelectionSet, DeletionOptions);
 	ElementsForAction.Get()->Reset();
@@ -146,11 +151,16 @@ TArray<FTypedElementHandle> UTypedElementCommonActions::DuplicateElements(const
 }
 
 TArray<FTypedElementHandle> UTypedElementCommonActions::DuplicateSelectedElements(const UTypedElementSelectionSet* SelectionSet, UWorld* World, const FVector& LocationOffset)
+{
+	return DuplicateElementsInList(SelectionSet->GetElementList(), World, LocationOffset);
+}
+
+TArray<FTypedElementHandle> UTypedElementCommonActions::DuplicateElementsInList(const UTypedElementList* ElementList, UWorld* World, const FVector& LocationOffset)
 {
 	TGCObjectScopeGuard<UTypedElementList> ElementsForAction(UTypedElementRegistry::GetInstance()->CreateElementList());
-	GetSelectedElementsForAction(SelectionSet, ElementsForAction.Get());
+	GetElementsForAction(ElementList, ElementsForAction.Get());
 
-	TArray<FTypedElementHandle> NewElements = DuplicateElements(ElementsForAction.Get(), World, LocationOffset);
+	TArray<FTypedElementHandle> NewElements = DuplicateElements(MakeArrayView(ElementsForAction.Get()->GetElementHandles()), World, LocationOffset);
 	ElementsForAction.Get()->Reset();
 	return NewElements;
 }
diff --git a/Engine/Source/Runtime/Engine/Public/Elements/Framework/TypedElementCommonActions.h b/Engine/Source/Runtime/Engine/Public/Elements/Framework/TypedElementCommonActions.h
--- a/Engine/Source/Runtime/Engine/Public/Elements/Framework/TypedElementCommonActions.h
+++ b/Engine/Source/Runtime/Engine/Public/Elements/Framework/TypedElementCommonActions.h
@@ -18,7 +18,7 @@ public:
 	virtual ~FTypedElementCommonActionsCustomization() = default;
 
 	//~ See UTypedElementCommonActions for API docs
-	virtual void GetElementsForAction(const TTypedElement<UTypedElementWorldInterface>& InElementWorldHandle, const UTypedElementSelectionSet* InSelectionSet, UTypedElementList* OutElementsToDelete);
+	virtual void GetElementsForAction(const TTypedElement<UTypedElementWorldInterface>& InElementWorldHandle, const UTypedElementList* InElementList, UTypedElementList* OutElementsToDelete);
 	virtual bool DeleteElements(UTypedElementWorldInterface* InWorldInterface, TArrayView<const FTypedElementHandle> InElementHandles, UWorld* InWorld, UTypedElementSelectionSet* InSelectionSet, const FTypedElementDeletionOptions& InDeletionOptions);
 	virtual void DuplicateElements(UTypedElementWorldInterface* InWorldInterface, TArrayView<const FTypedElementHandle> InElementHandles, UWorld* InWorld, const FVector& InLocationOffset, TArray<FTypedElementHandle>& OutNewElements);
 };
@@ -55,7 +55,7 @@ public:
 	}
 
 	//~ See UTypedElementCommonActions for API docs
-	void GetElementsForAction(const UTypedElementSelectionSet* InSelectionSet, UTypedElementList* OutElementsToDelete) { CommonActionsCustomization->GetElementsForAction(ElementWorldHandle, InSelectionSet, OutElementsToDelete); }
+	void GetElementsForAction(const UTypedElementList* InElementList, UTypedElementList* OutElementsToDelete) { CommonActionsCustomization->GetElementsForAction(ElementWorldHandle, InElementList, OutElementsToDelete); }
 
 private:
 	TTypedElement<UTypedElementWorldInterface> ElementWorldHandle;
@@ -73,10 +73,10 @@ class ENGINE_API UTypedElementCommonActions : public UObject, public TTypedEleme
 
 public:
 	/**
-	 * Get the elements from the given selection set that a common action should operate on.
+	 * Get the elements from the given element list that a common action should operate on.
 	 * @note This allows coarse filtering of the elements based on the selection state (eg, favoring components over actors), however it doesn't mean that the returned elements will actually be valid to perform a given action on.
 	 */
-	void GetSelectedElementsForAction(const UTypedElementSelectionSet* InSelectionSet, UTypedElementList* OutElementsForAction) const;
+	void GetElementsForAction(const UTypedElementList* InElementList, UTypedElementList* OutElementsForAction) const;
 
 	UFUNCTION(BlueprintCallable, Category="TypedElementFramework|Common")
 	bool DeleteElements(const TArray<FTypedElementHandle>& ElementHandles, UWorld* World, UTypedElementSelectionSet* SelectionSet, const FTypedElementDeletionOptions& DeletionOptions);
@@ -85,6 +85,7 @@ public:
 
 	UFUNCTION(BlueprintCallable, Category="TypedElementFramework|Common")
 	bool DeleteSelectedElements(UTypedElementSelectionSet* SelectionSet, UWorld* World, const FTypedElementDeletionOptions& DeletionOptions);
+	bool DeleteElementsInList(const UTypedElementList* ElementList, UWorld* World, UTypedElementSelectionSet* InSelectionSet, const FTypedElementDeletionOptions& DeletionOptions);
 
 	UFUNCTION(BlueprintCallable, Category="TypedElementFramework|Common")
 	TArray<FTypedElementHandle> DuplicateElements(const TArray<FTypedElementHandle>& ElementHandles, UWorld* World, const FVector& LocationOffset);
@@ -93,6 +94,7 @@ public:
 
 	UFUNCTION(BlueprintCallable, Category="TypedElementFramework|Common")
 	TArray<FTypedElementHandle> DuplicateSelectedElements(const UTypedElementSelectionSet* SelectionSet, UWorld* World, const FVector& LocationOffset);
+	TArray<FTypedElementHandle> DuplicateElementsInList(const UTypedElementList* ElementList, UWorld* World, const FVector& LocationOffset);
 
 private:
 	/**
