Shorten the cell level names to show more in STAT LEVELS.

diff --git a/Engine/Source/Runtime/Core/Public/UObject/UE5MainStreamObjectVersion.h b/Engine/Source/Runtime/Core/Public/UObject/UE5MainStreamObjectVersion.h
--- a/Engine/Source/Runtime/Core/Public/UObject/UE5MainStreamObjectVersion.h
+++ b/Engine/Source/Runtime/Core/Public/UObject/UE5MainStreamObjectVersion.h
@@ -127,6 +127,9 @@ struct CORE_API FUE5MainStreamObjectVersion
 		// Made directional lights be atmosphere lights by default
 		DirLightsAreAtmosphereLightsByDefault,
 
+		// Changed how world partition streaming cells are named
+		WorldPartitionStreamingCellsNamingShortened,
+
         // -----<new versions can be added above this line>-------------------------------------------------
 		VersionPlusOne,
 		LatestVersion = VersionPlusOne - 1
diff --git a/Engine/Source/Runtime/Engine/Private/UnrealEngine.cpp b/Engine/Source/Runtime/Engine/Private/UnrealEngine.cpp
--- a/Engine/Source/Runtime/Engine/Private/UnrealEngine.cpp
+++ b/Engine/Source/Runtime/Engine/Private/UnrealEngine.cpp
@@ -15453,10 +15453,12 @@ int32 UEngine::RenderStatLevels(UWorld* World, FViewport* Viewport, FCanvas* Can
 	Canvas->DrawShadowedString(X, Y, TEXT("Levels"), GetSmallFont(), FLinearColor::White);
 	Y += 12;
 
+	FString MapShortName;
 	if (SubLevelsStatusList.Num())
 	{
 		// First entry - always persistent level
 		FString MapName	= SubLevelsStatusList[0].PackageName.ToString();
+		MapShortName = FPackageName::GetShortName(MapName);
 		UPackage* LevelPackage = FindObjectFast<UPackage>(NULL, SubLevelsStatusList[0].PackageName);
 		if (SubLevelsStatusList[0].bPlayerInside)
 		{
@@ -15493,7 +15495,22 @@ int32 UEngine::RenderStatLevels(UWorld* World, FViewport* Viewport, FCanvas* Can
 
 		FColor	Color = ULevelStreaming::GetLevelStreamingStatusColor(LevelStatus.StreamingStatus);
 		const TCHAR* StatusDisplayName = ULevelStreaming::GetLevelStreamingStatusDisplayName(LevelStatus.StreamingStatus);
-		FString DisplayName = LevelStatus.PackageName.ToString();
+		FString LevelPackageName = LevelStatus.PackageName.ToString();
+		FString DisplayName = LevelPackageName;
+
+		if (World->GetWorldPartition())
+		{
+			int32 CutPos = LevelPackageName.Find(*MapShortName);
+			if (CutPos != INDEX_NONE)
+			{
+				CutPos += MapShortName.Len();
+				while (LevelPackageName[CutPos] == TEXT('_'))
+				{
+					CutPos++;
+				}
+				DisplayName = LevelPackageName.Right(LevelPackageName.Len() - CutPos);
+			}
+		}
 
 		if (LevelStatus.LODIndex != INDEX_NONE)
 		{
diff --git a/Engine/Source/Runtime/Engine/Private/WorldPartition/HLOD/HLODActor.cpp b/Engine/Source/Runtime/Engine/Private/WorldPartition/HLOD/HLODActor.cpp
--- a/Engine/Source/Runtime/Engine/Private/WorldPartition/HLOD/HLODActor.cpp
+++ b/Engine/Source/Runtime/Engine/Private/WorldPartition/HLOD/HLODActor.cpp
@@ -3,6 +3,7 @@
 #include "WorldPartition/HLOD/HLODActor.h"
 #include "WorldPartition/HLOD/HLODSubsystem.h"
 #include "Components/PrimitiveComponent.h"
+#include "UObject/UE5MainStreamObjectVersion.h"
 
 #if WITH_EDITOR
 #include "WorldPartition/WorldPartition.h"
@@ -49,6 +50,19 @@ void AWorldPartitionHLOD::EndPlay(const EEndPlayReason::Type EndPlayReason)
 	Super::EndPlay(EndPlayReason);
 }
 
+void AWorldPartitionHLOD::Serialize(FArchive& Ar)
+{
+	Ar.UsingCustomVersion(FUE5MainStreamObjectVersion::GUID);
+	Super::Serialize(Ar);
+
+#if WITH_EDITOR
+	if(Ar.IsLoading() && Ar.CustomVer(FUE5MainStreamObjectVersion::GUID) < FUE5MainStreamObjectVersion::WorldPartitionStreamingCellsNamingShortened)
+	{
+		SourceCell = SourceCell.ToString().Replace(TEXT("WPRT_"), TEXT(""), ESearchCase::CaseSensitive).Replace(TEXT("Cell_"), TEXT(""), ESearchCase::CaseSensitive);
+	}
+#endif
+}
+
 void AWorldPartitionHLOD::RerunConstructionScripts()
 {}
 
diff --git a/Engine/Source/Runtime/Engine/Private/WorldPartition/WorldPartitionRuntimeSpatialHash.cpp b/Engine/Source/Runtime/Engine/Private/WorldPartition/WorldPartitionRuntimeSpatialHash.cpp
--- a/Engine/Source/Runtime/Engine/Private/WorldPartition/WorldPartitionRuntimeSpatialHash.cpp
+++ b/Engine/Source/Runtime/Engine/Private/WorldPartition/WorldPartitionRuntimeSpatialHash.cpp
@@ -663,7 +663,7 @@ FName UWorldPartitionRuntimeSpatialHash::GetCellName(UWorldPartition* WorldParti
 	const FString PackageName = FPackageName::GetShortName(WorldPartition->GetPackage());
 	const FString PackageNameNoPIEPrefix = UWorld::RemovePIEPrefix(PackageName);
 
-	return FName(*FString::Printf(TEXT("WPRT_%s_%s_Cell_%s_DL%X"), *PackageNameNoPIEPrefix, *InGridName.ToString(), *GetCellCoordString(InCellGlobalCoord), InDataLayerID.GetHash()));
+	return FName(*FString::Printf(TEXT("%s_%s_%s_DL%X"), *PackageNameNoPIEPrefix, *InGridName.ToString(), *GetCellCoordString(InCellGlobalCoord), InDataLayerID.GetHash()));
 }
 
 FName UWorldPartitionRuntimeSpatialHash::GetCellName(FName InGridName, const FIntVector& InCellGlobalCoord, const FDataLayersID& InDataLayerID) const
diff --git a/Engine/Source/Runtime/Engine/Public/WorldPartition/HLOD/HLODActor.h b/Engine/Source/Runtime/Engine/Public/WorldPartition/HLOD/HLODActor.h
--- a/Engine/Source/Runtime/Engine/Public/WorldPartition/HLOD/HLODActor.h
+++ b/Engine/Source/Runtime/Engine/Public/WorldPartition/HLOD/HLODActor.h
@@ -59,6 +59,7 @@ public:
 
 protected:
 	//~ Begin UObject Interface.
+	virtual void Serialize(FArchive& Ar) override;
 	virtual void RerunConstructionScripts() override;
 	//~ End UObject Interface.
 
