Build fix for 16207912

diff --git a/Engine/Source/Runtime/Engine/Private/WorldPartition/HLOD/HLODActor.cpp b/Engine/Source/Runtime/Engine/Private/WorldPartition/HLOD/HLODActor.cpp
--- a/Engine/Source/Runtime/Engine/Private/WorldPartition/HLOD/HLODActor.cpp
+++ b/Engine/Source/Runtime/Engine/Private/WorldPartition/HLOD/HLODActor.cpp
@@ -9,8 +9,11 @@
 #include "WorldPartition/WorldPartition.h"
 #include "WorldPartition/WorldPartitionActorDesc.h"
 #include "WorldPartition/HLOD/HLODActorDesc.h"
-#include "WorldPartition/HLOD/HLODBuilder.h"
 #include "WorldPartition/HLOD/HLODLayer.h"
+
+#include "Modules/ModuleManager.h"
+#include "IWorldPartitionHLODUtilities.h"
+#include "WorldPartitionHLODUtilitiesModule.h"
 #endif
 
 AWorldPartitionHLOD::AWorldPartitionHLOD(const FObjectInitializer& ObjectInitializer)
@@ -165,12 +168,20 @@ uint32 AWorldPartitionHLOD::GetHLODHash() const
 
 void AWorldPartitionHLOD::BuildHLOD(bool bForceBuild)
 {
-	if (bForceBuild)
+	IWorldPartitionHLODUtilities* WPHLODUtilities = FModuleManager::Get().LoadModuleChecked<IWorldPartitionHLODUtilitiesModule>("WorldPartitionHLODUtilities").GetUtilities();
+	if (WPHLODUtilities)
 	{
-		HLODHash = 0;
+		if (bForceBuild)
+		{
+			HLODHash = 0;
+		}
+
+		HLODHash = WPHLODUtilities->BuildHLOD(this);
 	}
+	else
+	{
 
-	HLODHash = FHLODBuilderUtilities::BuildHLOD(this);
+	}
 }
 
 #endif // WITH_EDITOR
