Allow per-instance selection of lightweight instances in the editor

#rnx

diff --git a/Engine/Source/Runtime/Engine/Private/LightWeightInstanceStaticMeshManager.cpp b/Engine/Source/Runtime/Engine/Private/LightWeightInstanceStaticMeshManager.cpp
--- a/Engine/Source/Runtime/Engine/Private/LightWeightInstanceStaticMeshManager.cpp
+++ b/Engine/Source/Runtime/Engine/Private/LightWeightInstanceStaticMeshManager.cpp
@@ -155,6 +155,8 @@ void ALightWeightInstanceStaticMeshManager::SetInstancedStaticMeshParams()
 	InstancedStaticMeshComponent->PrimaryComponentTick.bCanEverTick = false;
 	// Allows updating in game, while optimizing rendering for the case that it is not modified
 	InstancedStaticMeshComponent->Mobility = EComponentMobility::Movable;
+	// Allows per-instance selection in the editor
+	InstancedStaticMeshComponent->bHasPerInstanceHitProxies = true;
 }
 
 void ALightWeightInstanceStaticMeshManager::SetStaticMeshFromActor(AActor* InActor)
