Fix non-unity compile error.
#rnx

diff --git a/Engine/Source/Runtime/Engine/Private/WorldPartition/RuntimeSpatialHash/RuntimeSpatialHashHLOD.cpp b/Engine/Source/Runtime/Engine/Private/WorldPartition/RuntimeSpatialHash/RuntimeSpatialHashHLOD.cpp
--- a/Engine/Source/Runtime/Engine/Private/WorldPartition/RuntimeSpatialHash/RuntimeSpatialHashHLOD.cpp
+++ b/Engine/Source/Runtime/Engine/Private/WorldPartition/RuntimeSpatialHash/RuntimeSpatialHashHLOD.cpp
@@ -14,6 +14,7 @@
 #include "HAL/PlatformFileManager.h"
 #include "Misc/PackageName.h"
 #include "Misc/ScopedSlowTask.h"
+#include "Modules/ModuleManager.h"
 
 #include "Algo/AnyOf.h"
 #include "Algo/ForEach.h"
