Change default editor cell sizes from 512M to 128M for new / converted partitioned worlds.

diff --git a/Engine/Source/Runtime/Engine/Private/WorldPartition/WorldPartitionEditorSpatialHash.cpp b/Engine/Source/Runtime/Engine/Private/WorldPartition/WorldPartitionEditorSpatialHash.cpp
--- a/Engine/Source/Runtime/Engine/Private/WorldPartition/WorldPartitionEditorSpatialHash.cpp
+++ b/Engine/Source/Runtime/Engine/Private/WorldPartition/WorldPartitionEditorSpatialHash.cpp
@@ -40,7 +40,7 @@ void UWorldPartitionEditorSpatialHash::Initialize()
 
 void UWorldPartitionEditorSpatialHash::SetDefaultValues()
 {
-	CellSize = 51200;
+	CellSize = 12800;
 }
 
 FName UWorldPartitionEditorSpatialHash::GetWorldPartitionEditorName()
