Restrict material compilation error for decal read/write normal to GBuffer decals only.

diff --git a/Engine/Source/Runtime/Engine/Private/Materials/HLSLMaterialTranslator.cpp b/Engine/Source/Runtime/Engine/Private/Materials/HLSLMaterialTranslator.cpp
--- a/Engine/Source/Runtime/Engine/Private/Materials/HLSLMaterialTranslator.cpp
+++ b/Engine/Source/Runtime/Engine/Private/Materials/HLSLMaterialTranslator.cpp
@@ -6131,22 +6131,24 @@ void FHLSLMaterialTranslator::UseSceneTextureId(ESceneTextureId SceneTextureId,
 
 	if (Material->GetMaterialDomain() == MD_DeferredDecal)
 	{
-		if (!(SceneTextureId == PPI_SceneDepth || SceneTextureId == PPI_WorldNormal || SceneTextureId == PPI_CustomDepth || SceneTextureId == PPI_CustomStencil))
+		const bool bSceneTextureSupportsDecal = SceneTextureId == PPI_SceneDepth || SceneTextureId == PPI_WorldNormal || SceneTextureId == PPI_CustomDepth || SceneTextureId == PPI_CustomStencil;
+		if (!bSceneTextureSupportsDecal)
 		{
-			// Note: For DBuffer decals: CustomDepth and CustomStencil are only available if r.CustomDepth.Order = 0
-			Errorf(TEXT("Decals can only access SceneDepth, CustomDepth, CustomStencil, and WorldNormal"));
+			// Note: For DBuffer decals CustomDepth and CustomStencil are only available if r.CustomDepth.Order = 0
+			Errorf(TEXT("Decals can only access SceneDepth, CustomDepth, CustomStencil, and WorldNormal."));
 		}
 
-		if (SceneTextureId == PPI_WorldNormal && Material->HasNormalConnected())
+		const bool bSceneTextureRequiresSM5 = SceneTextureId == PPI_WorldNormal || SceneTextureId == PPI_CustomDepth || SceneTextureId == PPI_CustomStencil;
+		if (bSceneTextureRequiresSM5)
 		{
-			// GBuffer can only relate to WorldNormal here.
-			Errorf(TEXT("Decals that read WorldNormal cannot output to normal at the same time"));
+			ErrorUnlessFeatureLevelSupported(ERHIFeatureLevel::SM5);
 		}
 
-		const bool bRequiresSM5 = SceneTextureId == PPI_WorldNormal || SceneTextureId == PPI_CustomDepth || SceneTextureId == PPI_CustomStencil;
-		if (bRequiresSM5)
+		if (SceneTextureId == PPI_WorldNormal && Material->HasNormalConnected() && !IsUsingDBuffers(Platform))
 		{
-			ErrorUnlessFeatureLevelSupported(ERHIFeatureLevel::SM5);
+			// GBuffer decals can't bind Normal for read and write.
+			// Note: DBuffer decals can support this but only if the sampled WorldNormal isn't connected to the output normal.
+			Errorf(TEXT("Decals that read WorldNormal cannot output to normal at the same time. Enable DBuffer to support this."));
 		}
 	}
 
