Added asset data typed element interface for ISM element types.

diff --git a/Engine/Source/Runtime/Engine/Private/Elements/Framework/EngineElements.cpp b/Engine/Source/Runtime/Engine/Private/Elements/Framework/EngineElements.cpp
--- a/Engine/Source/Runtime/Engine/Private/Elements/Framework/EngineElements.cpp
+++ b/Engine/Source/Runtime/Engine/Private/Elements/Framework/EngineElements.cpp
@@ -25,6 +25,7 @@
 #include "Elements/SMInstance/SMInstanceElementData.h"
 #include "Elements/SMInstance/SMInstanceElementWorldInterface.h"
 #include "Elements/SMInstance/SMInstanceElementSelectionInterface.h"
+#include "Elements/SMInstance/SMInstanceElementAssetDataInterface.h"
 
 #include "Modules/ModuleManager.h"
 
@@ -71,6 +72,7 @@ void RegisterEngineSMInstanceElements()
 	Registry->RegisterElementType<FSMInstanceElementData>(NAME_SMInstance);
 	Registry->RegisterElementInterface<UTypedElementWorldInterface>(NAME_SMInstance, NewObject<USMInstanceElementWorldInterface>());
 	Registry->RegisterElementInterface<UTypedElementSelectionInterface>(NAME_SMInstance, NewObject<USMInstanceElementSelectionInterface>());
+	Registry->RegisterElementInterface<UTypedElementAssetDataInterface>(NAME_SMInstance, NewObject<USMInstanceElementAssetDataInterface>());
 }
 
 void RegisterEngineElements()
diff --git a/Engine/Source/Runtime/Engine/Private/Elements/SMInstance/SMInstanceElementAssetDataInterface.cpp b/Engine/Source/Runtime/Engine/Private/Elements/SMInstance/SMInstanceElementAssetDataInterface.cpp
new file mode 100644
--- /dev/null
+++ b/Engine/Source/Runtime/Engine/Private/Elements/SMInstance/SMInstanceElementAssetDataInterface.cpp
@@ -0,0 +1,20 @@
+// Copyright Epic Games, Inc. All Rights Reserved.
+
+#include "Elements/SMInstance/SMInstanceElementAssetDataInterface.h"
+
+#include "Elements/SMInstance/SMInstanceElementData.h"
+#include "Components/InstancedStaticMeshComponent.h"
+#include "Engine/StaticMesh.h"
+
+FAssetData USMInstanceElementAssetDataInterface::GetAssetData(const FTypedElementHandle& InElementHandle)
+{
+	if (FSMInstanceId SMInstance = SMInstanceElementDataUtil::GetSMInstanceFromHandle(InElementHandle))
+	{
+		if (SMInstance.ISMComponent->GetStaticMesh())
+		{
+			return FAssetData(SMInstance.ISMComponent->GetStaticMesh());
+		}
+	}
+
+	return FAssetData();
+}
diff --git a/Engine/Source/Runtime/Engine/Public/Elements/SMInstance/SMInstanceElementAssetDataInterface.h b/Engine/Source/Runtime/Engine/Public/Elements/SMInstance/SMInstanceElementAssetDataInterface.h
new file mode 100644
--- /dev/null
+++ b/Engine/Source/Runtime/Engine/Public/Elements/SMInstance/SMInstanceElementAssetDataInterface.h
@@ -0,0 +1,16 @@
+// Copyright Epic Games, Inc. All Rights Reserved.
+
+#pragma once
+
+#include "Elements/Interfaces/TypedElementAssetDataInterface.h"
+
+#include "SMInstanceElementAssetDataInterface.generated.h"
+
+UCLASS()
+class ENGINE_API USMInstanceElementAssetDataInterface : public UTypedElementAssetDataInterface
+{
+	GENERATED_BODY()
+
+public:
+	virtual FAssetData GetAssetData(const FTypedElementHandle& InElementHandle) override;
+};
