DebugRenderSceneProxy:
- fixed text color not been applied

diff --git a/Engine/Source/Runtime/Engine/Private/DebugRenderSceneProxy.cpp b/Engine/Source/Runtime/Engine/Private/DebugRenderSceneProxy.cpp
--- a/Engine/Source/Runtime/Engine/Private/DebugRenderSceneProxy.cpp
+++ b/Engine/Source/Runtime/Engine/Private/DebugRenderSceneProxy.cpp
@@ -88,6 +88,7 @@ void FDebugDrawDelegateHelper::DrawDebugLabels(UCanvas* Canvas, APlayerControlle
 		{
 			const FVector ScreenLoc = Canvas->Project(It->Location);
 			const FFontRenderInfo& FontInfo = TextWithoutShadowDistance >= 0 ? (FDebugRenderSceneProxy::PointInRange(It->Location, View, TextWithoutShadowDistance) ? FontRenderInfoWithShadow : FontRenderInfo) : FontRenderInfo;
+			Canvas->SetDrawColor(It->Color);
 			Canvas->DrawText(RenderFont, It->Text, ScreenLoc.X, ScreenLoc.Y, 1, 1, FontInfo);
 		}
 	}
diff --git a/Engine/Source/Runtime/Engine/Public/DebugRenderSceneProxy.h b/Engine/Source/Runtime/Engine/Public/DebugRenderSceneProxy.h
--- a/Engine/Source/Runtime/Engine/Public/DebugRenderSceneProxy.h
+++ b/Engine/Source/Runtime/Engine/Public/DebugRenderSceneProxy.h
@@ -178,7 +178,7 @@ public:
 
 		FString Text;
 		FVector Location;
-		FLinearColor Color;
+		FColor Color;
 	};
 
 	struct FCone
