Don't force a full rescan of the external actors folder when initializing the actor descriptors containers.
This was required when these folders were blacklisted but that's not the case anymore.
Saves ~20s when loading a map containing ~90K external actors.

diff --git a/Engine/Source/Runtime/Engine/Private/WorldPartition/ActorDescContainer.cpp b/Engine/Source/Runtime/Engine/Private/WorldPartition/ActorDescContainer.cpp
--- a/Engine/Source/Runtime/Engine/Private/WorldPartition/ActorDescContainer.cpp
+++ b/Engine/Source/Runtime/Engine/Private/WorldPartition/ActorDescContainer.cpp
@@ -33,7 +33,7 @@ void UActorDescContainer::Initialize(UWorld* InWorld, FName InPackageName)
 
 		// Do a synchronous scan of the level external actors path.			
 		IAssetRegistry& AssetRegistry = FModuleManager::LoadModuleChecked<FAssetRegistryModule>(TEXT("AssetRegistry")).Get();
-		AssetRegistry.ScanPathsSynchronous({ LevelExternalActorsPath }, /*bForceRescan*/true, /*bIgnoreBlackListScanFilters*/true);
+		AssetRegistry.ScanPathsSynchronous({ LevelExternalActorsPath }, /*bForceRescan*/false, /*bIgnoreBlackListScanFilters*/false);
 
 		FARFilter Filter;
 		Filter.bRecursivePaths = true;
@@ -315,4 +315,4 @@ void UActorDescContainer::OnActorDescRemoved(FWorldPartitionActorDesc* ActorDesc
 {
 	OnActorDescRemovedEvent.Broadcast(ActorDesc);
 }
-#endif
\ No newline at end of file
+#endif
