Added support for mirroring animation attributes

diff --git a/Engine/Source/Runtime/AnimGraphRuntime/Private/AnimNodes/AnimNode_Mirror.cpp b/Engine/Source/Runtime/AnimGraphRuntime/Private/AnimNodes/AnimNode_Mirror.cpp
--- a/Engine/Source/Runtime/AnimGraphRuntime/Private/AnimNodes/AnimNode_Mirror.cpp
+++ b/Engine/Source/Runtime/AnimGraphRuntime/Private/AnimNodes/AnimNode_Mirror.cpp
@@ -4,6 +4,7 @@
 #include "Animation/AnimInstanceProxy.h"
 #include "Animation/AnimNode_Inertialization.h"
 #include "AnimationRuntime.h"
+#include "Animation/AttributesRuntime.h"
 #include "Animation/MirrorSyncScope.h"
 
 #define LOCTEXT_NAMESPACE "AnimNode_Mirror"
@@ -12,6 +13,9 @@ FAnimNode_Mirror::FAnimNode_Mirror()
 	: MirrorDataTable(nullptr)
 	, BlendTimeOnMirrorStateChange(0.0f)
 	, bMirror(true)
+	, bBoneMirroring(true)
+	, bCurveMirroring(true)
+	, bAttributeMirroring(true)
 	, bResetChildOnMirrorStateChange(false)
 	, bMirrorState(false)
 	, bMirrorStateIsValid(false)
@@ -119,8 +123,20 @@ void FAnimNode_Mirror::Evaluate_AnyThread(FPoseContext& Output)
 
 	if (bMirrorState && MirrorDataTable)
 	{
-		FAnimationRuntime::MirrorPose(Output.Pose, MirrorDataTable->MirrorAxis, CompactPoseMirrorBones, ComponentSpaceRefRotations);
-		FAnimationRuntime::MirrorCurves(Output.Curve, *MirrorDataTable);
+		if (bBoneMirroring)
+		{
+			FAnimationRuntime::MirrorPose(Output.Pose, MirrorDataTable->MirrorAxis, CompactPoseMirrorBones, ComponentSpaceRefRotations);
+		}
+
+		if (bCurveMirroring)
+		{
+			FAnimationRuntime::MirrorCurves(Output.Curve, *MirrorDataTable);
+		}
+
+		if (bAttributeMirroring)
+		{
+			UE::Anim::Attributes::MirrorAttributes(Output.CustomAttributes, *MirrorDataTable);
+		}
 	}
 }
 
diff --git a/Engine/Source/Runtime/AnimGraphRuntime/Public/AnimNodes/AnimNode_Mirror.h b/Engine/Source/Runtime/AnimGraphRuntime/Public/AnimNodes/AnimNode_Mirror.h
--- a/Engine/Source/Runtime/AnimGraphRuntime/Public/AnimNodes/AnimNode_Mirror.h
+++ b/Engine/Source/Runtime/AnimGraphRuntime/Public/AnimNodes/AnimNode_Mirror.h
@@ -46,6 +46,15 @@ private:
 	UPROPERTY(EditAnywhere,  Category = Settings, meta = (PinShownByDefault))
 	bool bMirror;
 
+	UPROPERTY(EditAnywhere, Category = MirroredChannels, meta=(DisplayName="Bone"))
+	bool bBoneMirroring;
+
+	UPROPERTY(EditAnywhere, Category = MirroredChannels, meta=(DisplayName = "Curve"))
+	bool bCurveMirroring;
+
+	UPROPERTY(EditAnywhere, Category = MirroredChannels, meta=(DisplayName = "Attributes"))
+	bool bAttributeMirroring;
+
 	// Whether to reset (reinitialize) the child (source) pose when the mirror state changes
 	UPROPERTY(EditAnywhere, Category = Settings)
 	bool bResetChildOnMirrorStateChange;
diff --git a/Engine/Source/Runtime/Engine/Classes/Animation/AttributesContainer.h b/Engine/Source/Runtime/Engine/Classes/Animation/AttributesContainer.h
--- a/Engine/Source/Runtime/Engine/Classes/Animation/AttributesContainer.h
+++ b/Engine/Source/Runtime/Engine/Classes/Animation/AttributesContainer.h
@@ -638,6 +638,11 @@ namespace UE
 				return false;
 			}
 
+		/*
+		* @return Unique bone indices for all contained entries of a specific attribute type
+		*/
+		const TArray<int32, InAllocator>& GetUniqueTypedBoneIndices(int32 TypeIndex) const { return UniqueTypedBoneIndices[TypeIndex]; }
+
 		public:
 			/* Deprecated API */
 			template<typename DataType> 
@@ -693,6 +698,17 @@ namespace UE
 				return Array;
 			}
 		protected:
+
+			/*
+			* @return Array of all the contained values for the provided TypeIndex
+			*/
+			TArray<TWrappedAttribute<InAllocator>, InAllocator>& GetValuesInternal(int32 TypeIndex)  { return Values[TypeIndex]; }
+
+			/*
+			* @return Array of all the contained keys for the provided TypeIndex
+			*/
+			TArray<FAttributeId, InAllocator>& GetKeysInternal(int32 TypeIndex) { return AttributeIdentifiers[TypeIndex]; }
+			
 			/** Find or add a new root-level entry for the provided attribute data type, returning the index into the arrays representing the type */
 			int32 FindOrAddTypeIndex(const UScriptStruct* InScriptStruct)
 			{
@@ -720,6 +736,9 @@ namespace UE
 
 			template <typename OtherBoneIndexType, typename OtherAllocator>
 			friend struct TAttributeContainer;
+
+			template<class OtherBoneIndexType, typename OtherAllocator>
+			friend struct TAttributeContainerAccessor;
 		};
 	}
 }
diff --git a/Engine/Source/Runtime/Engine/Private/Animation/AttributesRuntime.cpp b/Engine/Source/Runtime/Engine/Private/Animation/AttributesRuntime.cpp
--- a/Engine/Source/Runtime/Engine/Private/Animation/AttributesRuntime.cpp
+++ b/Engine/Source/Runtime/Engine/Private/Animation/AttributesRuntime.cpp
@@ -5,6 +5,7 @@
 #include "Animation/AnimTypes.h"
 #include "Animation/AnimationSettings.h"
 #include "Animation/AnimationAsset.h"
+#include "Animation/MirrorDataTable.h"
 #include "AnimationRuntime.h"
 
 #include "Containers/UnrealString.h"
@@ -441,6 +442,110 @@ void Attributes::BlendAttributesPerBoneFilter(const FStackAttributeContainer& Ba
 	}
 }
 
+void Attributes::MirrorAttributes(FStackAttributeContainer& CustomAttributes, const UMirrorDataTable& MirrorTable)
+{
+	struct FMirrorSet
+	{
+		/** Pointers to attribute values */
+		TArray<uint8*, FAnimStackAllocator> DataPtrsA;
+		TArray<uint8*, FAnimStackAllocator> DataPtrsB;
+		/** Identifier of the attribute */
+		TArray <const FAttributeId*, FAnimStackAllocator> IdentifiersA;
+		TArray <const FAttributeId*, FAnimStackAllocator> IdentifiersB;
+		int32 MirrorIndexA;
+		int32 MirrorIndexB;
+		void Reset() { DataPtrsA.Reset(); DataPtrsB.Reset(), IdentifiersA.Reset(); IdentifiersB.Reset(); }
+	};
+
+	TArray<FMirrorSet, FAnimStackAllocator> MirrorSets;
+
+	auto AddUnique = [&MirrorSets](int32 IndexA, int32 IndexB) -> bool
+	{
+		for (FMirrorSet& CurSet : MirrorSets)
+		{
+			if ( (CurSet.MirrorIndexA == IndexA && CurSet.MirrorIndexB == IndexB) ||
+				 (CurSet.MirrorIndexB == IndexA && CurSet.MirrorIndexA == IndexB) )
+			{
+				return false;
+			}
+		}
+		FMirrorSet& NewSet = MirrorSets.Emplace_GetRef();
+		NewSet.MirrorIndexA = IndexA;
+		NewSet.MirrorIndexB = IndexB;
+		return true;
+	};
+
+	for (const TWeakObjectPtr<const UScriptStruct>& WeakScriptStruct : CustomAttributes.GetUniqueTypes())
+	{
+		const int32 TypeIndex = CustomAttributes.FindTypeIndex(WeakScriptStruct.Get());
+		if (TypeIndex != INDEX_NONE)
+		{
+			MirrorSets.Reset();
+			const TArray<int32, FAnimStackAllocator>& UniqueBoneIndices = CustomAttributes.GetUniqueTypedBoneIndices(TypeIndex);
+			TArray<int32, FAnimStackAllocator> SortedBoneIndices(UniqueBoneIndices); 
+			SortedBoneIndices.Sort();
+
+			for (int32 CurBoneIndex : UniqueBoneIndices)
+			{
+				int32 MirrorBoneIndex = MirrorTable.BoneToMirrorBoneIndex[CurBoneIndex];
+				int32 SwapCheckBoneIndex = MirrorTable.BoneToMirrorBoneIndex[MirrorBoneIndex];
+
+				// only handle simple swaps - both entries exist and are different
+				if (MirrorBoneIndex != INDEX_NONE && CurBoneIndex != INDEX_NONE && CurBoneIndex != MirrorBoneIndex &&
+					SwapCheckBoneIndex == CurBoneIndex && Algo::BinarySearch(SortedBoneIndices, MirrorBoneIndex))
+				{
+					AddUnique(CurBoneIndex, MirrorBoneIndex);
+				}
+			}
+
+			TArray<TWrappedAttribute<FAnimStackAllocator>, FAnimStackAllocator>& ValuesArray = FStackAttributeContainerAccessor::GetValues(CustomAttributes, TypeIndex);
+			const TArray<FAttributeId, FAnimStackAllocator>& AttributeIdentifiers = CustomAttributes.GetKeys(TypeIndex);
+
+			// gather attributes that are on mirrored bones 
+			for (int32 AttributeIndex = 0; AttributeIndex < AttributeIdentifiers.Num(); ++AttributeIndex)
+			{
+				const FAttributeId& AttributeIdentifier = AttributeIdentifiers[AttributeIndex];
+				uint8* DataPtr = ValuesArray[AttributeIndex].GetPtr<uint8>();
+				for (FMirrorSet& CurSet : MirrorSets)
+				{
+					if (CurSet.MirrorIndexA == AttributeIdentifier.GetIndex())
+					{
+						CurSet.IdentifiersA.Emplace(&AttributeIdentifier);
+						CurSet.DataPtrsA.Emplace(DataPtr);
+					}
+					else if (CurSet.MirrorIndexB == AttributeIdentifier.GetIndex())
+					{
+						CurSet.IdentifiersB.Emplace(&AttributeIdentifier);
+						CurSet.DataPtrsB.Emplace(DataPtr);
+					}
+				}
+			}
+
+			for (FMirrorSet& CurSet : MirrorSets)
+			{
+				for (int32 IndexA = 0; IndexA < CurSet.IdentifiersA.Num(); ++IndexA)
+				{
+					const FAttributeId* IndexAId = CurSet.IdentifiersA[IndexA];
+					uint8* IndexADataPtr = CurSet.DataPtrsA[IndexA];
+					for (int32 IndexB = 0; IndexB < CurSet.IdentifiersB.Num(); ++IndexB)
+					{
+						const FAttributeId* IndexBId = CurSet.IdentifiersB[IndexB];
+						if (IndexAId->GetName() == IndexBId->GetName() && IndexAId->GetNamespace() == IndexBId->GetNamespace())
+						{
+							uint8* IndexBDataPtr = CurSet.DataPtrsB[IndexB];
+							TWrappedAttribute<FAnimStackAllocator> SwapStruct(WeakScriptStruct.Get());
+							WeakScriptStruct->CopyScriptStruct(SwapStruct.GetPtr<void>(), IndexADataPtr);
+							WeakScriptStruct->CopyScriptStruct(IndexADataPtr, IndexBDataPtr);
+							WeakScriptStruct->CopyScriptStruct(IndexBDataPtr, SwapStruct.GetPtr<void>());
+							break;
+						}
+					}
+				}
+			}
+		} 
+	}
+}
+
 ECustomAttributeBlendType Attributes::GetAttributeBlendType(const FName& InName)
 {
 	const UAnimationSettings* Settings = UAnimationSettings::Get();
diff --git a/Engine/Source/Runtime/Engine/Public/Animation/AttributesRuntime.h b/Engine/Source/Runtime/Engine/Public/Animation/AttributesRuntime.h
--- a/Engine/Source/Runtime/Engine/Public/Animation/AttributesRuntime.h
+++ b/Engine/Source/Runtime/Engine/Public/Animation/AttributesRuntime.h
@@ -31,6 +31,8 @@ struct FAnimatedBoneAttribute;
 struct FCompactPoseBoneIndex;
 struct FAnimExtractContext;
 
+class UMirrorDataTable;
+
 namespace UE
 {
 	namespace Anim
@@ -38,6 +40,17 @@ namespace UE
 		struct ENGINE_API FStackAttributeContainer : public TAttributeContainer<FCompactPoseBoneIndex, FAnimStackAllocator> {};
 		struct ENGINE_API FHeapAttributeContainer : public TAttributeContainer<FCompactPoseBoneIndex, FDefaultAllocator> {};
 
+		/** Accessor for internal data of the TAttributeContainer.  The TAttributeContainer methods are preferred - this accessor should only be used in limited situations */
+		template<class BoneIndexType, typename InAllocator>
+        struct TAttributeContainerAccessor
+		{
+			static TArray<TWrappedAttribute<InAllocator>, InAllocator>& GetValues(TAttributeContainer<BoneIndexType, InAllocator>& Attributes,  int32 TypeIndex) {return Attributes.GetValuesInternal(TypeIndex);}
+			static TArray<FAttributeId, InAllocator>& GetKeys(TAttributeContainer<BoneIndexType, InAllocator>& Attributes,  int32 TypeIndex) {return Attributes.GetKeysInternal(TypeIndex);}
+		};
+
+		struct ENGINE_API FStackAttributeContainerAccessor : public TAttributeContainerAccessor<FCompactPoseBoneIndex, FAnimStackAllocator> {};
+		struct ENGINE_API FHeapAttributeContainerAccessor : public TAttributeContainerAccessor<FCompactPoseBoneIndex, FDefaultAllocator> {};
+		
 		/** Helper functionality for attributes animation runtime */
 		struct ENGINE_API Attributes
 		{
@@ -81,6 +94,9 @@ namespace UE
 			/** Interpolates between two sets of attributes */
 			static void InterpolateAttributes(FHeapAttributeContainer& FromAttributes, const FHeapAttributeContainer& ToAttributes, float Alpha);
 
+			/** Mirror (swap) attributes with the specified MirrorDataTable. Attributes are swapped using the bone mapping such that bones which are mirrored swap attributes */
+			static void MirrorAttributes(FStackAttributeContainer& Attributes, const UMirrorDataTable& MirrorDataTable);
+			
 			/** Helper functionality to retrieve the correct blend type (from UAnimationSettings) for the provided attribute name */
 			static ECustomAttributeBlendType GetAttributeBlendType(const FName& InName);
 
