Don't hydrate light weight instances when an instance is selected in the editor.

diff --git a/Engine/Source/Editor/UnrealEd/Private/Editor/ActorPositioning.cpp b/Engine/Source/Editor/UnrealEd/Private/Editor/ActorPositioning.cpp
--- a/Engine/Source/Editor/UnrealEd/Private/Editor/ActorPositioning.cpp
+++ b/Engine/Source/Editor/UnrealEd/Private/Editor/ActorPositioning.cpp
@@ -60,10 +60,9 @@ bool IsHitIgnored(const FHitResult& InHit, const FSceneView& InSceneView)
 	// we're also accessing UObjects, the game thread should be blocked during this call.
 	check(IsInParallelRenderingThread());
 	const FActorInstanceHandle& HitObjHandle = InHit.HitObjectHandle;
-	AActor* Actor = HitObjHandle.FetchActor();
 	
 	// Try and find a primitive component for the hit
-	const UPrimitiveComponent* PrimitiveComponent = Actor ? Cast<UPrimitiveComponent>(Actor->GetRootComponent()) : nullptr;
+	const UPrimitiveComponent* PrimitiveComponent = Cast<UPrimitiveComponent>(HitObjHandle.GetRootComponent());
 
 	if (!PrimitiveComponent)
 	{
@@ -150,7 +149,7 @@ FActorPositionTraceResult FActorPositioning::TraceWorldForPosition(const UWorld&
 				Results.Location = Hit.Location;
 				Results.SurfaceNormal = Hit.Normal.GetSafeNormal();
 				Results.State = FActorPositionTraceResult::HitSuccess;
-				Results.HitActor = Hit.HitObjectHandle.FetchActor();
+				Results.HitActor = Hit.HitObjectHandle.GetManagingActor();
 			}
 		}
 	}
diff --git a/Engine/Source/Runtime/Engine/Classes/Engine/EngineTypes.h b/Engine/Source/Runtime/Engine/Classes/Engine/EngineTypes.h
--- a/Engine/Source/Runtime/Engine/Classes/Engine/EngineTypes.h
+++ b/Engine/Source/Runtime/Engine/Classes/Engine/EngineTypes.h
@@ -2047,6 +2047,7 @@ struct ENGINE_API FActorInstanceHandle
 
 	friend struct FLightWeightInstanceSubsystem;
 	friend class ALightWeightInstanceManager;
+	friend class UActorInstanceHandleInterface;
 
 	FActorInstanceHandle();
 
@@ -2071,7 +2072,13 @@ struct ENGINE_API FActorInstanceHandle
 	FName GetFName() const;
 	FString GetName() const;
 
-	// Returns the actor specified by this handle. This may require loading and creating the actor object.
+	/** If this handle has a valid actor, return it; otherwise return the actor responsible for managing the instances. */
+	AActor* GetManagingActor() const;
+
+	/** Returns either the actor's root component or the root component for the manager associated with the handle */
+	USceneComponent* GetRootComponent() const;
+
+	/** Returns the actor specified by this handle. This may require loading and creating the actor object. */
 	AActor* FetchActor() const;
 	template <typename T>
 	T* FetchActor() const;
diff --git a/Engine/Source/Runtime/Engine/Private/EngineTypes.cpp b/Engine/Source/Runtime/Engine/Private/EngineTypes.cpp
--- a/Engine/Source/Runtime/Engine/Private/EngineTypes.cpp
+++ b/Engine/Source/Runtime/Engine/Private/EngineTypes.cpp
@@ -385,6 +385,27 @@ FString FActorInstanceHandle::GetName() const
 	return FString();
 }
 
+AActor* FActorInstanceHandle::GetManagingActor() const
+{
+	if (Actor.IsValid())
+	{
+		return Actor.Get();
+	}
+
+	return const_cast<ALightWeightInstanceManager*>(FLightWeightInstanceSubsystem::Get().GetManagerAt(ManagerIndex));
+}
+
+USceneComponent* FActorInstanceHandle::GetRootComponent() const
+{
+	if (Actor.IsValid())
+	{
+		return Actor->GetRootComponent();
+	}
+
+	const AActor* Manager = FLightWeightInstanceSubsystem::Get().GetManagerAt(ManagerIndex);
+	return Manager ? Manager->GetRootComponent() : nullptr;
+}
+
 AActor* FActorInstanceHandle::FetchActor() const
 {
 	if (Actor.IsValid())
@@ -413,6 +434,12 @@ bool FActorInstanceHandle::operator==(const AActor* OtherActor) const
 		return Actor.Get() == OtherActor;
 	}
 
+	// if OtherActor is null then we're only equal if this doesn't refer to a valid instance
+	if (OtherActor == nullptr)
+	{
+		return ManagerIndex == INDEX_NONE && InstanceIndex == INDEX_NONE;
+	}
+
 	// we don't have an actor so see if we can look up an instance associated with OtherActor and see if we refer to the same instance
 	if (ALightWeightInstanceManager* Manager = FLightWeightInstanceSubsystem::Get().FindLightWeightInstanceManager(OtherActor->StaticClass(), OtherActor->GetLevel()))
 	{
