WorldPartition: Add DataLayers to DumpActorDesc(s) commands

diff --git a/Engine/Source/Runtime/Engine/Private/WorldPartition/WorldPartition.cpp b/Engine/Source/Runtime/Engine/Private/WorldPartition/WorldPartition.cpp
--- a/Engine/Source/Runtime/Engine/Private/WorldPartition/WorldPartition.cpp
+++ b/Engine/Source/Runtime/Engine/Private/WorldPartition/WorldPartition.cpp
@@ -36,6 +36,7 @@
 #include "WorldPartition/WorldPartitionEditorHash.h"
 #include "WorldPartition/WorldPartitionRuntimeHash.h"
 #include "WorldPartition/WorldPartitionEditorPerProjectUserSettings.h"
+#include "WorldPartition/DataLayer/WorldDataLayers.h"
 #include "Modules/ModuleManager.h"
 #include "GameDelegates.h"
 #endif //WITH_EDITOR
@@ -45,6 +46,58 @@ DEFINE_LOG_CATEGORY(LogWorldPartition);
 #define LOCTEXT_NAMESPACE "WorldPartition"
 
 #if WITH_EDITOR
+TMap<FName, FString> GetDataLayersDumpString(const UWorldPartition* WorldPartition)
+{
+	TMap<FName, FString> DataLayersDumpString;
+	if (const AWorldDataLayers* WorldDataLayers = WorldPartition->GetWorld()->GetWorldDataLayers())
+	{
+		WorldDataLayers->ForEachDataLayer([&DataLayersDumpString](const UDataLayer* DataLayer)
+			{
+				DataLayersDumpString.FindOrAdd(DataLayer->GetFName()) = FString::Format(TEXT("{0}({1})"), { DataLayer->GetDataLayerLabel().ToString(), DataLayer->GetFName().ToString() });
+				return true;
+			});
+	}
+	return DataLayersDumpString;
+}
+
+FString GetActorDescDumpString(const FWorldPartitionActorDesc* ActorDesc, const TMap<FName, FString>& DataLayersDumpString)
+{
+	auto GetDataLayerString = [&DataLayersDumpString](const TArray<FName>& DataLayerNames)
+	{
+		if (DataLayerNames.IsEmpty())
+		{
+			return FString("None");
+		}
+
+		return FString::JoinBy(DataLayerNames, TEXT(", "), 
+			[&DataLayersDumpString](const FName& DataLayerName) 
+			{ 
+				if (const FString* DumpString = DataLayersDumpString.Find(DataLayerName))
+				{
+					return *DumpString;
+				}
+				return DataLayerName.ToString(); 
+			});
+	};
+
+	check(ActorDesc);
+	return FString::Printf(
+		TEXT("%s, %s, %s, %s, %s, %.2f, %.2f, %.2f, %.2f, %.2f, %.2f, %s") LINE_TERMINATOR,
+		*ActorDesc->GetGuid().ToString(),
+		*ActorDesc->GetClass().ToString(),
+		*ActorDesc->GetActorName().ToString(),
+		*ActorDesc->GetActorPackage().ToString(),
+		*ActorDesc->GetActorLabel().ToString(),
+		ActorDesc->GetBounds().GetCenter().X,
+		ActorDesc->GetBounds().GetCenter().Y,
+		ActorDesc->GetBounds().GetCenter().Z,
+		ActorDesc->GetBounds().GetExtent().X,
+		ActorDesc->GetBounds().GetExtent().Y,
+		ActorDesc->GetBounds().GetExtent().Z,
+		*GetDataLayerString(ActorDesc->GetDataLayers())
+	);
+}
+
 static FAutoConsoleCommand DumpActorDesc(
 	TEXT("wp.Editor.DumpActorDesc"),
 	TEXT("Dump a specific actor descriptor on the console."),
@@ -74,25 +127,13 @@ static FAutoConsoleCommand DumpActorDesc(
 				{
 					if (UWorldPartition* WorldPartition = World->GetWorldPartition())
 					{
-						UE_LOG(LogWorldPartition, Log, TEXT("Guid, Class, Name, Package, BVCenterX, BVCenterY, BVCenterZ, BVExtentX, BVExtentY, BVExtentZ"));
+						TMap<FName, FString> DataLayersDumpString = GetDataLayersDumpString(WorldPartition);
+						UE_LOG(LogWorldPartition, Log, TEXT("Guid, Class, Name, Package, BVCenterX, BVCenterY, BVCenterZ, BVExtentX, BVExtentY, BVExtentZ, DataLayers"));
 						for (const FString& ActorPath : ActorPaths)
 						{
 							if (const FWorldPartitionActorDesc* ActorDesc = WorldPartition->GetActorDesc(ActorPath))
 							{
-
-								UE_LOG(LogWorldPartition, Log, TEXT("%s, %s, %s, %s, %s, %.2f, %.2f, %.2f, %.2f, %.2f, %.2f"),
-									*ActorDesc->GetGuid().ToString(),
-									*ActorDesc->GetClass().ToString(),
-									*ActorDesc->GetActorName().ToString(),
-									*ActorDesc->GetActorPackage().ToString(),
-									*ActorDesc->GetActorLabel().ToString(),
-									ActorDesc->GetBounds().GetCenter().X,
-									ActorDesc->GetBounds().GetCenter().Y,
-									ActorDesc->GetBounds().GetCenter().Z,
-									ActorDesc->GetBounds().GetExtent().X,
-									ActorDesc->GetBounds().GetExtent().Y,
-									ActorDesc->GetBounds().GetExtent().Z
-								);
+								UE_LOG(LogWorldPartition, Log, TEXT("%s"), *GetActorDescDumpString(ActorDesc, DataLayersDumpString));
 							}
 						}
 					}
@@ -1226,10 +1267,11 @@ void UWorldPartition::DumpActorDescs(const FString& Path)
 {
 	if (FArchive* LogFile = IFileManager::Get().CreateFileWriter(*Path))
 	{
-		FString LineEntry = TEXT("Guid, Class, Name, Package, BVCenterX, BVCenterY, BVCenterZ, BVExtentX, BVExtentY, BVExtentZ") LINE_TERMINATOR;
+		FString LineEntry = TEXT("Guid, Class, Name, Package, BVCenterX, BVCenterY, BVCenterZ, BVExtentX, BVExtentY, BVExtentZ, DataLayers") LINE_TERMINATOR;
 		LogFile->Serialize(TCHAR_TO_ANSI(*LineEntry), LineEntry.Len());
 
 		TArray<const FWorldPartitionActorDesc*> ActorDescs;
+		TMap<FName, FString> DataLayersDumpString = GetDataLayersDumpString(this);
 		for (UActorDescContainer::TConstIterator<> ActorDescIterator(this); ActorDescIterator; ++ActorDescIterator)
 		{
 			ActorDescs.Add(*ActorDescIterator);
@@ -1240,20 +1282,7 @@ void UWorldPartition::DumpActorDescs(const FString& Path)
 		});
 		for (const FWorldPartitionActorDesc* ActorDescIterator : ActorDescs)
 		{
-			LineEntry = FString::Printf(
-				TEXT("%s, %s, %s, %s, %s, %.2f, %.2f, %.2f, %.2f, %.2f, %.2f") LINE_TERMINATOR, 
-				*ActorDescIterator->GetGuid().ToString(), 
-				*ActorDescIterator->GetClass().ToString(), 
-				*ActorDescIterator->GetActorName().ToString(), 
-				*ActorDescIterator->GetActorPackage().ToString(), 
-				*ActorDescIterator->GetActorLabel().ToString(),
-				ActorDescIterator->GetBounds().GetCenter().X,
-				ActorDescIterator->GetBounds().GetCenter().Y,
-				ActorDescIterator->GetBounds().GetCenter().Z,
-				ActorDescIterator->GetBounds().GetExtent().X,
-				ActorDescIterator->GetBounds().GetExtent().Y,
-				ActorDescIterator->GetBounds().GetExtent().Z
-			);
+			LineEntry = GetActorDescDumpString(ActorDescIterator, DataLayersDumpString);
 			LogFile->Serialize(TCHAR_TO_ANSI(*LineEntry), LineEntry.Len());
 		}
 
