#fix ensure INDEX_NONE values for bone interpolation indices is handled correctly (using total weight instead)

diff --git a/Engine/Source/Runtime/Engine/Private/Animation/AttributeBlendData.cpp b/Engine/Source/Runtime/Engine/Private/Animation/AttributeBlendData.cpp
--- a/Engine/Source/Runtime/Engine/Private/Animation/AttributeBlendData.cpp
+++ b/Engine/Source/Runtime/Engine/Private/Animation/AttributeBlendData.cpp
@@ -338,10 +338,15 @@ namespace UE { namespace Anim {
 			const FBlendSampleData& BlendSampleData = BlendSampleDataCache[SampleDataIndex];
 			const int32 PerBoneIndex = PerBoneInterpolationIndices[BoneIndex];
 
-			// Blend-sample blending is only performed when they contain per-bone weights, so assert otherwise
-			check(PerBoneIndex != INDEX_NONE && BlendSampleData.PerBoneBlendData.IsValidIndex(PerBoneIndex));
-
-			return BlendSampleData.PerBoneBlendData[PerBoneIndex];
+			// Blend-sample blending is only performed when they contain per-bone weights, if INDEX_NONE or out of range use the total weight instead
+			if (PerBoneIndex != INDEX_NONE && BlendSampleData.PerBoneBlendData.IsValidIndex(PerBoneIndex))
+			{
+				return BlendSampleData.PerBoneBlendData[PerBoneIndex];
+			}
+			else
+			{
+				return BlendSampleData.GetWeight();
+			}
 		}
 
 		return 0.f;
