Fix for selection crash in Frosty.

[REDACTED_HANDLE] Simon.Sherr

diff --git a/Engine/Source/Runtime/Engine/Private/LightWeightInstanceSubsystem.cpp b/Engine/Source/Runtime/Engine/Private/LightWeightInstanceSubsystem.cpp
--- a/Engine/Source/Runtime/Engine/Private/LightWeightInstanceSubsystem.cpp
+++ b/Engine/Source/Runtime/Engine/Private/LightWeightInstanceSubsystem.cpp
@@ -24,7 +24,7 @@ int32 FLightWeightInstanceSubsystem::GetManagerIndex(const ALightWeightInstanceM
 
 const ALightWeightInstanceManager* FLightWeightInstanceSubsystem::GetManagerAt(int32 Index) const
 {
-	return Index < LWInstanceManagers.Num() ? LWInstanceManagers[Index] : nullptr;
+	return (Index >= 0 && Index < LWInstanceManagers.Num()) ? LWInstanceManagers[Index] : nullptr;
 }
 
 ALightWeightInstanceManager* FLightWeightInstanceSubsystem::FindLightWeightInstanceManager(const FActorInstanceHandle& Handle) const
