fix GetAllMobileRelevantLayerNames drops those layers which are only connected to RuntimeVirtualTextureOutput Node

diff --git a/Engine/Source/Runtime/Engine/Private/Materials/Material.cpp b/Engine/Source/Runtime/Engine/Private/Materials/Material.cpp
--- a/Engine/Source/Runtime/Engine/Private/Materials/Material.cpp
+++ b/Engine/Source/Runtime/Engine/Private/Materials/Material.cpp
@@ -5816,6 +5816,21 @@ bool UMaterial::GetAllReferencedExpressions(TArray<UMaterialExpression*>& OutExp
 				}
 			}
 		}
+
+		bool bMobileUseVirtualTexturing = UseVirtualTexturing(ERHIFeatureLevel::ES3_1);
+		if (bMobileUseVirtualTexturing)
+		{
+			TArray<class UMaterialExpressionCustomOutput*> CustomOutputExpressions;
+			GetAllCustomOutputExpressions(CustomOutputExpressions);
+			for (UMaterialExpressionCustomOutput* Expression : CustomOutputExpressions)
+			{
+				if (Expression->IsA<UMaterialExpressionRuntimeVirtualTextureOutput>())
+				{
+					TArray<FExpressionInput*> ProcessedInputs;
+					RecursiveGetExpressionChain(Expression, ProcessedInputs, OutExpressions, InStaticParameterSet, InFeatureLevel, InQuality, InShadingPath);
+				}
+			}
+		}
 	}
 	else
 	{
