Build fix for variable shadowing in lambda function.

diff --git a/Engine/Source/Runtime/Engine/Private/Materials/MaterialExpressions.cpp b/Engine/Source/Runtime/Engine/Private/Materials/MaterialExpressions.cpp
--- a/Engine/Source/Runtime/Engine/Private/Materials/MaterialExpressions.cpp
+++ b/Engine/Source/Runtime/Engine/Private/Materials/MaterialExpressions.cpp
@@ -19879,9 +19879,9 @@ int32 UMaterialExpressionStrataSlabBSDF::Compile(class FMaterialCompiler* Compil
 		SSSProfileCodeChunk = Compiler->ForceCast(Compiler->ScalarParameter(NameSubsurfaceProfile, 1.0f), MCT_Float1);
 	}
 
-	auto DielectricSpecularToF0 = [](float Specular)
+	auto DielectricSpecularToF0 = [](float SpecularIn)
 	{
-		return 0.08f * Specular;
+		return 0.08f * SpecularIn;
 	};
 	const float DefaultSpecular = 0.5f;
 	const float DefaultF0 = DielectricSpecularToF0(DefaultSpecular);
