Strata: added parameter blending version of BSDF add.

diff --git a/Engine/Shaders/Private/Strata/Strata.ush b/Engine/Shaders/Private/Strata/Strata.ush
--- a/Engine/Shaders/Private/Strata/Strata.ush
+++ b/Engine/Shaders/Private/Strata/Strata.ush
@@ -908,10 +908,6 @@ FStrataData StrataHorizontalMixingParameterBlending(FStrataData Foreground, FStr
 	const float ForegroundMixFactor = saturate(Mix);
 	const float BackgroundMixFactor = 1.0 - ForegroundMixFactor;
 
-	// This function can only parameter blend two BSDF and not more. We cannot blend complex topology.
-	// However, maybe at some point we will
-	// This function only blend one we only parameter blend one slab with anothe
-
 	FStrataData Result = GetInitialisedStrataDataWithReservedFirstSlot();
 #define ResultBSDF		Result.Layers[0].BSDFs[0]
 #define ForegroundBSDF	Foreground.Layers[0].BSDFs[0]
@@ -953,7 +949,7 @@ FStrataData StrataHorizontalMixingParameterBlending(FStrataData Foreground, FStr
 	}
 	else
 	{
-		SLAB_SSSDMFP(ResultBSDF) = lerp(SLAB_SSSDMFP(BackgroundBSDF),				SLAB_SSSDMFP(ForegroundBSDF),				ForegroundMixFactor);
+		SLAB_SSSDMFP(ResultBSDF) = lerp(SLAB_SSSDMFP(BackgroundBSDF),				SLAB_SSSDMFP(ForegroundBSDF),					ForegroundMixFactor);
 	}
 
 
@@ -970,6 +966,71 @@ FStrataData StrataHorizontalMixingParameterBlending(FStrataData Foreground, FStr
 	return Result;
 }
 
+FStrataData StrataAddParameterBlending(FStrataData A, FStrataData B, uint NewNormalIndex)
+{
+	const float MixFactor = 0.5f;
+
+	FStrataData Result = GetInitialisedStrataDataWithReservedFirstSlot();
+#define ResultBSDF		Result.Layers[0].BSDFs[0]
+#define ABSDF			A.Layers[0].BSDFs[0]
+#define BBSDF			B.Layers[0].BSDFs[0]
+
+	// We can only parameter blend the slab 
+	BSDF_SETTYPE(ResultBSDF,			STRATA_BSDF_TYPE_SLAB);
+	BSDF_SETNORMALID(ResultBSDF,		NewNormalIndex);
+
+	BSDF_SETISTOPLAYER(ResultBSDF,		0);
+	BSDF_SETANISOTROPY(ResultBSDF,		BSDF_GETANISOTROPY(ABSDF)						| BSDF_GETANISOTROPY(BBSDF));
+	BSDF_SETHASSSS(ResultBSDF,			BSDF_GETHASSSS(ABSDF)							| BSDF_GETHASSSS(BBSDF));
+	BSDF_SETHASSSSPROFILE(ResultBSDF,	BSDF_GETHASSSSPROFILE(ABSDF)					| BSDF_GETHASSSSPROFILE(BBSDF));	// Problem: cannot blend SSS profiles...
+	BSDF_SETISTHIN(ResultBSDF,			BSDF_GETISTHIN(ABSDF)							| BSDF_GETISTHIN(BBSDF));			// STRATA_TODO fix the MPF/DMFP situation for that to really work
+	BSDF_SETHASDMFP(ResultBSDF,			BSDF_GETHASDMFP(ABSDF)							| BSDF_GETHASDMFP(BBSDF));
+	BSDF_SETHASHAZINESS(ResultBSDF,		BSDF_GETHASHAZINESS(ABSDF)						| BSDF_GETHASHAZINESS(BBSDF));
+	BSDF_SETHASF90(ResultBSDF,			BSDF_GETHASF90(ABSDF)							| BSDF_GETHASF90(BBSDF));
+	BSDF_SETHASTHINFILM(ResultBSDF,		BSDF_GETHASTHINFILM(ABSDF)						| BSDF_GETHASTHINFILM(BBSDF));
+	BSDF_SETHASFUZZ(ResultBSDF,			BSDF_GETHASFUZZ(ABSDF)							| BSDF_GETHASFUZZ(BBSDF));
+	BSDF_SETEMISSIVE(ResultBSDF,		BSDF_GETEMISSIVE(BBSDF)							+ BSDF_GETEMISSIVE(ABSDF));
+	BSDF_SETTHICKNESSCM(ResultBSDF,		BSDF_GETTHICKNESSCM(BBSDF)						+ BSDF_GETTHICKNESSCM(ABSDF));
+
+	// Some parameters will contribute to added luminance: in this case we add them together and saturate to not go out of the safe range
+	// Some parameters are not contributing to added luminance: in this case we simply lerp them
+	SLAB_DIFFUSEALBEDO(ResultBSDF) =	saturate(SLAB_DIFFUSEALBEDO(BBSDF)				+ SLAB_DIFFUSEALBEDO(ABSDF));
+	SLAB_F0(ResultBSDF) =				saturate(SLAB_F0(BBSDF)							+ SLAB_F0(ABSDF));
+	SLAB_F90(ResultBSDF) =				saturate(SLAB_F90(BBSDF)						+ SLAB_F90(ABSDF));
+	SLAB_METALLIC(ResultBSDF) =			lerp(SLAB_METALLIC(BBSDF),						SLAB_METALLIC(ABSDF),						MixFactor);
+	SLAB_ROUGHNESSX(ResultBSDF) =		lerp(SLAB_ROUGHNESSX(BBSDF),					SLAB_ROUGHNESSX(ABSDF),						MixFactor);
+	SLAB_ROUGHNESSY(ResultBSDF) =		lerp(SLAB_ROUGHNESSY(BBSDF),					SLAB_ROUGHNESSY(ABSDF),						MixFactor);
+	SLAB_HAZINESS(ResultBSDF) =			lerp(SLAB_HAZINESS(BBSDF),						SLAB_HAZINESS(ABSDF),						MixFactor);
+	SLAB_THINFILM_THICKNESS(ResultBSDF)=lerp(SLAB_THINFILM_THICKNESS(BBSDF),			SLAB_THINFILM_THICKNESS(ABSDF),				MixFactor);
+	SLAB_FUZZ_AMOUNT(ResultBSDF) =		lerp(SLAB_FUZZ_AMOUNT(BBSDF),					SLAB_FUZZ_AMOUNT(ABSDF),					MixFactor);
+	SLAB_FUZZ_COLOR(ResultBSDF) =		saturate(SLAB_FUZZ_COLOR(BBSDF)					+ SLAB_FUZZ_COLOR(ABSDF));
+
+	if (BSDF_GETHASSSSPROFILE(ResultBSDF))
+	{
+		// As mentioned above, we cannot blend SSSProfiles.
+		SLAB_SSSPROFILEID(ResultBSDF) = SLAB_SSSPROFILEID(ABSDF);
+		SLAB_SSSPROFILERADIUSSCALE(ResultBSDF) = SLAB_SSSPROFILERADIUSSCALE(ABSDF);
+	}
+	else
+	{
+		SLAB_SSSDMFP(ResultBSDF) =		lerp(SLAB_SSSDMFP(BBSDF),						SLAB_SSSDMFP(ABSDF),						MixFactor);
+	}
+
+
+#if STRATA_INLINE_SHADING
+	// At this stage we do not know if this strata BSDF is going to be using simple volume lighting so always store the DMFP on the side to not override SSSPROFILE data.
+	ResultBSDF.TmpDMFP =				lerp(BBSDF.TmpDMFP,								ABSDF.TmpDMFP,								MixFactor);
+#endif
+
+	// When it comes to adding the visual contribution, it is tricky to find a matching operator for coverage. We simply lerp for now.
+	ResultBSDF.Weight =					lerp(BBSDF.Weight,								ABSDF.Weight,								MixFactor);
+
+#undef ResultBSDF
+#undef ABSDF
+#undef BBSDF
+	return Result;
+}
+
 
 ///////////////////////////////////////////////////////////////////////////////
 // Functions used to sanitize BSDF before they are used for lighting (when forward) or stored to memory (base pass).
diff --git a/Engine/Source/Runtime/Engine/Classes/Materials/MaterialExpressionStrata.h b/Engine/Source/Runtime/Engine/Classes/Materials/MaterialExpressionStrata.h
--- a/Engine/Source/Runtime/Engine/Classes/Materials/MaterialExpressionStrata.h
+++ b/Engine/Source/Runtime/Engine/Classes/Materials/MaterialExpressionStrata.h
@@ -485,6 +485,9 @@ class UMaterialExpressionStrataAdd : public UMaterialExpressionStrataBSDF
 	UPROPERTY()
 	FExpressionInput B;
 
+	UPROPERTY(EditAnywhere, Category = Mode)
+	uint32 bUseParameterBlending : 1;
+
 	//~ Begin UMaterialExpression Interface
 #if WITH_EDITOR
 	virtual int32 Compile(class FMaterialCompiler* Compiler, int32 OutputIndex) override;
diff --git a/Engine/Source/Runtime/Engine/Private/Materials/HLSLMaterialTranslator.cpp b/Engine/Source/Runtime/Engine/Private/Materials/HLSLMaterialTranslator.cpp
--- a/Engine/Source/Runtime/Engine/Private/Materials/HLSLMaterialTranslator.cpp
+++ b/Engine/Source/Runtime/Engine/Private/Materials/HLSLMaterialTranslator.cpp
@@ -9112,10 +9112,19 @@ int32 FHLSLMaterialTranslator::StrataAdd(int32 A, int32 B)
 	);
 }
 
-//virtual int32 FHLSLMaterialTranslator::StrataAddParameterBlending(int32 A, int32 B, const FString& SharedNormalIndexMacro)
-//{
-//	// STRATA_TODO
-//}
+int32 FHLSLMaterialTranslator::StrataAddParameterBlending(int32 A, int32 B, const FString& SharedNormalIndexMacro)
+{
+	if (A == INDEX_NONE || B == INDEX_NONE)
+	{
+		return INDEX_NONE;
+	}
+	return AddCodeChunk(
+		MCT_Strata, TEXT("StrataAddParameterBlending(%s, %s, %s)"),
+		*GetParameterCode(A),
+		*GetParameterCode(B),
+		*SharedNormalIndexMacro
+	);
+}
 
 int32 FHLSLMaterialTranslator::StrataMultiply(int32 A, int32 Weight)
 {
diff --git a/Engine/Source/Runtime/Engine/Private/Materials/HLSLMaterialTranslator.h b/Engine/Source/Runtime/Engine/Private/Materials/HLSLMaterialTranslator.h
--- a/Engine/Source/Runtime/Engine/Private/Materials/HLSLMaterialTranslator.h
+++ b/Engine/Source/Runtime/Engine/Private/Materials/HLSLMaterialTranslator.h
@@ -1001,7 +1001,7 @@ protected:
 	virtual int32 StrataVerticalLayering(int32 Top, int32 Base) override;
 	// STRATA_TODO virtual int32 StrataVerticalLayeringParameterBlending(int32 Top, int32 Base, const FString& SharedNormalIndexMacro) override;
 	virtual int32 StrataAdd(int32 A, int32 B) override;
-	// STRATA_TODO virtual int32 StrataAddParameterBlending(int32 A, int32 B, const FString& SharedNormalIndexMacro) override;
+	virtual int32 StrataAddParameterBlending(int32 A, int32 B, const FString& SharedNormalIndexMacro) override;
 	virtual int32 StrataMultiply(int32 A, int32 Weight) override;
 	virtual int32 StrataTransmittanceToMFP(int32 TransmittanceColor, int32 DesiredThickness, int32 OutputIndex) override;
 
diff --git a/Engine/Source/Runtime/Engine/Private/Materials/MaterialExpressions.cpp b/Engine/Source/Runtime/Engine/Private/Materials/MaterialExpressions.cpp
--- a/Engine/Source/Runtime/Engine/Private/Materials/MaterialExpressions.cpp
+++ b/Engine/Source/Runtime/Engine/Private/Materials/MaterialExpressions.cpp
@@ -20669,7 +20669,6 @@ int32 UMaterialExpressionStrataHorizontalMixing::Compile(class FMaterialCompiler
 
 	int32 ForegroundCodeChunk = Foreground.Compile(Compiler);
 	int32 BackgroundCodeChunk = Background.Compile(Compiler);
-
 	if (!Compiler->StrataCompilationInfoContainsCodeChunk(ForegroundCodeChunk) || !Compiler->StrataCompilationInfoContainsCodeChunk(BackgroundCodeChunk))
 	{
 		return Compiler->Errorf(TEXT("Could not find ForegroundCodeChunk or BackgroundCodeChunk to mix"));
@@ -20746,7 +20745,7 @@ void UMaterialExpressionStrataHorizontalMixing::GetCaption(TArray<FString>& OutC
 {
 	if (bUseParameterBlending)
 	{
-		OutCaptions.Add(TEXT("Strata BSDF Horizontal Parameter Blend"));
+		OutCaptions.Add(TEXT("Strata BSDF Horizontal Blend (Parameter Blend)"));
 	}
 	else
 	{
@@ -20862,6 +20861,7 @@ void UMaterialExpressionStrataVerticalLayering::GatherStrataMaterialInfo(FStrata
 
 UMaterialExpressionStrataAdd::UMaterialExpressionStrataAdd(const FObjectInitializer& ObjectInitializer)
 	: Super(ObjectInitializer)
+	, bUseParameterBlending(false)
 {
 	struct FConstructorStatics
 	{
@@ -20888,22 +20888,85 @@ int32 UMaterialExpressionStrataAdd::Compile(class FMaterialCompiler* Compiler, i
 
 	int32 ACodeChunk = A.Compile(Compiler);
 	int32 BCodeChunk = B.Compile(Compiler);
-
-	int32 OutputCodeChunk = Compiler->StrataAdd(ACodeChunk, BCodeChunk);
-
 	if (!Compiler->StrataCompilationInfoContainsCodeChunk(ACodeChunk) || !Compiler->StrataCompilationInfoContainsCodeChunk(BCodeChunk))
 	{
 		return Compiler->Errorf(TEXT("Could not find ACodeChunk or BCodeChunk to add"));
 	}
-	FStrataMaterialCompilationInfo StrataInfo = StrataCompilationInfoAdd(Compiler, Compiler->GetStrataCompilationInfo(ACodeChunk), Compiler->GetStrataCompilationInfo(BCodeChunk));
-	Compiler->StrataCompilationInfoRegisterCodeChunk(OutputCodeChunk, StrataInfo);
+	const FStrataMaterialCompilationInfo& AStrataData = Compiler->GetStrataCompilationInfo(ACodeChunk);
+	const FStrataMaterialCompilationInfo& BStrataData = Compiler->GetStrataCompilationInfo(BCodeChunk);
+
+	const int32 MixCodeChunk = Compiler->Constant(0.5f);
+
+	int32 OutputCodeChunk = INDEX_NONE;
+	if (bUseParameterBlending)
+	{
+		if (AStrataData.TotalBSDFCount != 1)
+		{
+			return Compiler->Errorf(TEXT("A input: cannot add complex material topology using parameter blending. Only a one-to-one blending is supported today."));
+		}
+		if (BStrataData.TotalBSDFCount != 1)
+		{
+			return Compiler->Errorf(TEXT("B input: cannot add complex material topology using parameter blending. Only a one-to-one blending is supported today."));
+		}
+
+		const FStrataMaterialCompilationInfo::FBSDF& ABSDF = AStrataData.Layers[0].BSDFs[0];
+		const FStrataMaterialCompilationInfo::FBSDF& BBSDF = BStrataData.Layers[0].BSDFs[0];
+		if (ABSDF.Type != STRATA_BSDF_TYPE_SLAB)
+		{
+			return Compiler->Errorf(TEXT("A input: we can only apply parameter blending to Slab BSDF."));
+		}
+		if (BBSDF.Type != STRATA_BSDF_TYPE_SLAB)
+		{
+			return Compiler->Errorf(TEXT("B input: we can only apply parameter blending to Slab BSDF."));
+		}
+
+		// Compute the new Normal and Tangent resulting from the blending using code chunk
+		const int32 NewNormalCodeChunk = StrataBlendNormal(Compiler, ABSDF.RegisteredSharedNormal.NormalCodeChunk, BBSDF.RegisteredSharedNormal.NormalCodeChunk, MixCodeChunk);
+		// The tangent is optional so we treat it differently if INDEX_NONE is specified
+		int32 NewTangentCodeChunk = INDEX_NONE;
+		if (ABSDF.RegisteredSharedNormal.TangentCodeChunk != INDEX_NONE && BBSDF.RegisteredSharedNormal.TangentCodeChunk != INDEX_NONE)
+		{
+			NewTangentCodeChunk = StrataBlendNormal(Compiler, ABSDF.RegisteredSharedNormal.TangentCodeChunk, BBSDF.RegisteredSharedNormal.TangentCodeChunk, MixCodeChunk);
+		}
+		else if (ABSDF.RegisteredSharedNormal.TangentCodeChunk != INDEX_NONE)
+		{
+			NewTangentCodeChunk = ABSDF.RegisteredSharedNormal.TangentCodeChunk;
+		}
+		else if (BBSDF.RegisteredSharedNormal.TangentCodeChunk != INDEX_NONE)
+		{
+			NewTangentCodeChunk = BBSDF.RegisteredSharedNormal.TangentCodeChunk;
+		}
+		const FStrataRegisteredSharedNormal NewRegisteredSharedNormal = StrataCompilationInfoCreateSharedNormal(Compiler, NewNormalCodeChunk, NewTangentCodeChunk);
+
+		OutputCodeChunk = Compiler->StrataAddParameterBlending(ACodeChunk, BCodeChunk, Compiler->GetStrataSharedNormalIndexMacro(NewRegisteredSharedNormal));
+
+		StrataCompilationInfoCreateSingleBSDFMaterial(Compiler, OutputCodeChunk, NewRegisteredSharedNormal, STRATA_BSDF_TYPE_SLAB);
+		FStrataMaterialCompilationInfo StrataInfo = StrataCompilationInfoAddParamBlend(Compiler, AStrataData, BStrataData, NewRegisteredSharedNormal);
+
+		Compiler->StrataCompilationInfoRegisterCodeChunk(OutputCodeChunk, StrataInfo);
+	}
+	else
+	{
+		OutputCodeChunk = Compiler->StrataAdd(ACodeChunk, BCodeChunk);
+
+		FStrataMaterialCompilationInfo StrataInfo = StrataCompilationInfoAdd(Compiler, Compiler->GetStrataCompilationInfo(ACodeChunk), Compiler->GetStrataCompilationInfo(BCodeChunk));
+
+		Compiler->StrataCompilationInfoRegisterCodeChunk(OutputCodeChunk, StrataInfo);
+	}
 
 	return OutputCodeChunk;
 }
 
 void UMaterialExpressionStrataAdd::GetCaption(TArray<FString>& OutCaptions) const
 {
-	OutCaptions.Add(TEXT("Strata BSDF Add"));
+	if (bUseParameterBlending)
+	{
+		OutCaptions.Add(TEXT("Strata BSDF Add (Parameter Blend)"));
+	}
+	else
+	{
+		OutCaptions.Add(TEXT("Strata BSDF Add"));
+	}
 }
 
 uint32 UMaterialExpressionStrataAdd::GetOutputType(int32 OutputIndex)
diff --git a/Engine/Source/Runtime/Engine/Private/Materials/StrataMaterial.cpp b/Engine/Source/Runtime/Engine/Private/Materials/StrataMaterial.cpp
--- a/Engine/Source/Runtime/Engine/Private/Materials/StrataMaterial.cpp
+++ b/Engine/Source/Runtime/Engine/Private/Materials/StrataMaterial.cpp
@@ -90,6 +90,13 @@ void  StrataCompilationInfoCreateNullBSDF(FMaterialCompiler* Compiler, int32 Cod
 }
 
 
+FStrataMaterialCompilationInfo StrataCompilationInfoMultiply(FMaterialCompiler* Compiler, const FStrataMaterialCompilationInfo& A)
+{
+	FStrataMaterialCompilationInfo StrataInfo = A;
+	return StrataInfo;
+}
+
+
 FStrataMaterialCompilationInfo StrataCompilationInfoAdd(FMaterialCompiler* Compiler, const FStrataMaterialCompilationInfo& A, const FStrataMaterialCompilationInfo& B)
 {
 	FStrataMaterialCompilationInfo StrataInfo = A;
@@ -120,9 +127,24 @@ FStrataMaterialCompilationInfo StrataCompilationInfoAdd(FMaterialCompiler* Compi
 }
 
 
-FStrataMaterialCompilationInfo StrataCompilationInfoMultiply(FMaterialCompiler* Compiler, const FStrataMaterialCompilationInfo& A)
+FStrataMaterialCompilationInfo StrataCompilationInfoAddParamBlend(FMaterialCompiler* Compiler, const FStrataMaterialCompilationInfo& A, const FStrataMaterialCompilationInfo& B, const FStrataRegisteredSharedNormal& RegisteredSharedNormal)
 {
+	check(A.TotalBSDFCount == 1);
+	check(B.TotalBSDFCount == 1);
+
 	FStrataMaterialCompilationInfo StrataInfo = A;
+	FStrataMaterialCompilationInfo::FBSDF& NewBSDF = StrataInfo.Layers[0].BSDFs[0];
+	const FStrataMaterialCompilationInfo::FBSDF& OtherBSDF = B.Layers[0].BSDFs[0];
+
+	NewBSDF.RegisteredSharedNormal = RegisteredSharedNormal;
+
+	// When parameter blending is used, we take the union of all the features activated by input BSDFs.
+	NewBSDF.bHasSSS				|=	OtherBSDF.bHasSSS;
+	NewBSDF.bHasDMFPPluggedIn	|=	OtherBSDF.bHasDMFPPluggedIn;
+	NewBSDF.bHasEdgeColor		|=	OtherBSDF.bHasEdgeColor;
+	NewBSDF.bHasThinFilm		|=	OtherBSDF.bHasThinFilm;
+	NewBSDF.bHasFuzz			|=	OtherBSDF.bHasFuzz;
+	UpdateTotalBSDFCount(Compiler, StrataInfo);
 	return StrataInfo;
 }
 
diff --git a/Engine/Source/Runtime/Engine/Private/Materials/StrataMaterial.h b/Engine/Source/Runtime/Engine/Private/Materials/StrataMaterial.h
--- a/Engine/Source/Runtime/Engine/Private/Materials/StrataMaterial.h
+++ b/Engine/Source/Runtime/Engine/Private/Materials/StrataMaterial.h
@@ -57,8 +57,9 @@ void  StrataCompilationInfoCreateSingleBSDFMaterial(FMaterialCompiler* Compiler,
 	const FStrataRegisteredSharedNormal& RegisteredSharedNormal,
 	uint8 BSDFType, bool bHasSSS = false, bool bHasDMFPPluggedIn = false, bool bHasEdgeColor = false, bool bHasThinFilm = false, bool bHasFuzz = false);
 
-FStrataMaterialCompilationInfo StrataCompilationInfoAdd(FMaterialCompiler* Compiler, const FStrataMaterialCompilationInfo& A, const FStrataMaterialCompilationInfo& B);
 FStrataMaterialCompilationInfo StrataCompilationInfoMultiply(FMaterialCompiler* Compiler, const FStrataMaterialCompilationInfo& A);
+FStrataMaterialCompilationInfo StrataCompilationInfoAdd(FMaterialCompiler* Compiler, const FStrataMaterialCompilationInfo& A, const FStrataMaterialCompilationInfo& B);
+FStrataMaterialCompilationInfo StrataCompilationInfoAddParamBlend(FMaterialCompiler* Compiler, const FStrataMaterialCompilationInfo& A, const FStrataMaterialCompilationInfo& B, const FStrataRegisteredSharedNormal& RegisteredSharedNormal);
 FStrataMaterialCompilationInfo StrataCompilationInfoHorizontalMixing(FMaterialCompiler* Compiler, const FStrataMaterialCompilationInfo& A, const FStrataMaterialCompilationInfo& B);
 FStrataMaterialCompilationInfo StrataCompilationInfoHorizontalMixingParamBlend(FMaterialCompiler* Compiler, const FStrataMaterialCompilationInfo& A, const FStrataMaterialCompilationInfo& B, const FStrataRegisteredSharedNormal& RegisteredSharedNormal);
 FStrataMaterialCompilationInfo StrataCompilationInfoVerticalLayering(FMaterialCompiler* Compiler, const FStrataMaterialCompilationInfo& Top, const FStrataMaterialCompilationInfo& Base);
diff --git a/Engine/Source/Runtime/Engine/Public/MaterialCompiler.h b/Engine/Source/Runtime/Engine/Public/MaterialCompiler.h
--- a/Engine/Source/Runtime/Engine/Public/MaterialCompiler.h
+++ b/Engine/Source/Runtime/Engine/Public/MaterialCompiler.h
@@ -442,6 +442,7 @@ public:
 	virtual int32 StrataHorizontalMixingParameterBlending(int32 Foreground, int32 Background, int32 Mix, const FString& SharedNormalIndexMacro) = 0;
 	virtual int32 StrataVerticalLayering(int32 Top, int32 Base) = 0;
 	virtual int32 StrataAdd(int32 A, int32 B) = 0;
+	virtual int32 StrataAddParameterBlending(int32 A, int32 B, const FString& SharedNormalIndexMacro) = 0;
 	virtual int32 StrataMultiply(int32 A, int32 Weight) = 0;
 	virtual int32 StrataTransmittanceToMFP(int32 TransmittanceColor, int32 DesiredThickness, int32 OutputIndex) = 0;
 
@@ -992,6 +993,11 @@ public:
 		return Compiler->StrataAdd(A, B);
 	}
 
+	virtual int32 StrataAddParameterBlending(int32 A, int32 B, const FString& SharedNormalIndexMacro) override
+	{
+		return Compiler->StrataAddParameterBlending(A, B, SharedNormalIndexMacro);
+	}
+
 	virtual int32 StrataMultiply(int32 A, int32 Weight) override
 	{
 		return Compiler->StrataMultiply(A, Weight);
