Added static_assert to verify that anim node data types are being correctly accessed

diff --git a/Engine/Source/Runtime/AnimGraphRuntime/Private/AnimNodes/AnimNode_BlendSpacePlayer.cpp b/Engine/Source/Runtime/AnimGraphRuntime/Private/AnimNodes/AnimNode_BlendSpacePlayer.cpp
--- a/Engine/Source/Runtime/AnimGraphRuntime/Private/AnimNodes/AnimNode_BlendSpacePlayer.cpp
+++ b/Engine/Source/Runtime/AnimGraphRuntime/Private/AnimNodes/AnimNode_BlendSpacePlayer.cpp
@@ -117,7 +117,7 @@ void FAnimNode_BlendSpacePlayer::UpdateInternal(const FAnimationUpdateContext& C
 
 void FAnimNode_BlendSpacePlayer::Evaluate_AnyThread(FPoseContext& Output)
 {
-	DECLARE_SCOPE_HIERARCHICAL_COUNTER_ANIMNODE(Evaluate_AnyThread)
+	DECLARE_SCOPE_HIERARCHICAL_COUNTER_ANIMNODE(Evaluate_AnyThread) 
 
 	UBlendSpace* CurrentBlendSpace = GetBlendSpace();
 	if ((CurrentBlendSpace != nullptr) && (Output.AnimInstanceProxy->IsSkeletonCompatible(CurrentBlendSpace->GetSkeleton())))
@@ -244,7 +244,7 @@ FName FAnimNode_BlendSpacePlayer::GetGroupName() const
 
 EAnimGroupRole::Type FAnimNode_BlendSpacePlayer::GetGroupRole() const
 {
-	return GET_ANIM_NODE_DATA(EAnimGroupRole::Type, GroupRole);
+	return GET_ANIM_NODE_DATA(TEnumAsByte<EAnimGroupRole::Type>, GroupRole);
 }
 
 EAnimSyncMethod FAnimNode_BlendSpacePlayer::GetGroupMethod() const
diff --git a/Engine/Source/Runtime/AnimGraphRuntime/Private/AnimNodes/AnimNode_SequenceEvaluator.cpp b/Engine/Source/Runtime/AnimGraphRuntime/Private/AnimNodes/AnimNode_SequenceEvaluator.cpp
--- a/Engine/Source/Runtime/AnimGraphRuntime/Private/AnimNodes/AnimNode_SequenceEvaluator.cpp
+++ b/Engine/Source/Runtime/AnimGraphRuntime/Private/AnimNodes/AnimNode_SequenceEvaluator.cpp
@@ -109,7 +109,7 @@ void FAnimNode_SequenceEvaluatorBase::GatherDebugData(FNodeDebugData& DebugData)
 	DebugData.AddDebugItem(DebugLine, true);
 }
 
-float FAnimNode_SequenceEvaluatorBase::GetEffectiveDeltaTime(float ExplicitTime, float PrevExplicitTime) const
+float FAnimNode_SequenceEvaluatorBase::GetEffectiveDeltaTime(float ExplicitTime, float PrevExplicitTime) const 
 {
 	float DeltaTime = ExplicitTime - PrevExplicitTime;
 
@@ -176,7 +176,7 @@ FName FAnimNode_SequenceEvaluator::GetGroupName() const
 
 EAnimGroupRole::Type FAnimNode_SequenceEvaluator::GetGroupRole() const
 {
-	return GET_ANIM_NODE_DATA(EAnimGroupRole::Type, GroupRole);
+	return GET_ANIM_NODE_DATA(TEnumAsByte<EAnimGroupRole::Type>, GroupRole);
 }
 
 EAnimSyncMethod FAnimNode_SequenceEvaluator::GetGroupMethod() const
diff --git a/Engine/Source/Runtime/Engine/Classes/Animation/AnimNodeBase.h b/Engine/Source/Runtime/Engine/Classes/Animation/AnimNodeBase.h
--- a/Engine/Source/Runtime/Engine/Classes/Animation/AnimNodeBase.h
+++ b/Engine/Source/Runtime/Engine/Classes/Animation/AnimNodeBase.h
@@ -953,6 +953,12 @@ private:
 	const FAnimNodeData* NodeData = nullptr;
 };
 
+#if WITH_EDITORONLY_DATA
+#define VERIFY_ANIM_NODE_MEMBER_TYPE(Type, Identifier) static_assert(TIsSame<decltype(Identifier), Type>::Value, "Incorrect return type used");
+#else
+#define VERIFY_ANIM_NODE_MEMBER_TYPE(Type, Identifier)
+#endif
+
 // Get some (potentially folded) anim node data. Only usable from within an anim node.
 // This caches the node data ID in static contained in a local lambda for improved performance
 #define GET_ANIM_NODE_DATA(Type, Identifier) \
@@ -960,6 +966,7 @@ private:
 		GetData<Type>( \
 			[this]() -> UE::Anim::FNodeDataId \
 			{ \
+				VERIFY_ANIM_NODE_MEMBER_TYPE(Type, Identifier) \
 				static UE::Anim::FNodeDataId CachedId_##Identifier; \
 				if(!CachedId_##Identifier.IsValid()) \
 				{ \
@@ -977,6 +984,7 @@ private:
 		GetMutableData<Type>( \
 			[this]() -> UE::Anim::FNodeDataId \
 			{ \
+				VERIFY_ANIM_NODE_MEMBER_TYPE(Type, Identifier) \
 				static UE::Anim::FNodeDataId CachedId_##Identifier; \
 				if(!CachedId_##Identifier.IsValid()) \
 				{ \
diff --git a/Engine/Source/Runtime/Engine/Private/Animation/AnimNode_SequencePlayer.cpp b/Engine/Source/Runtime/Engine/Private/Animation/AnimNode_SequencePlayer.cpp
--- a/Engine/Source/Runtime/Engine/Private/Animation/AnimNode_SequencePlayer.cpp
+++ b/Engine/Source/Runtime/Engine/Private/Animation/AnimNode_SequencePlayer.cpp
@@ -179,13 +179,13 @@ void FAnimNode_SequencePlayer::SetSequence(UAnimSequenceBase* InSequence)
 void FAnimNode_SequencePlayer::SetLoopAnimation(bool bInLoopAnimation)
 {
 #if WITH_EDITORONLY_DATA
-	bLoopAnimation = bInLoopAnimation;
+	bLoopAnimation = bInLoopAnimation; 
 #endif
 }
 
 UAnimSequenceBase* FAnimNode_SequencePlayer::GetSequence() const
 {
-	return GET_ANIM_NODE_DATA(UAnimSequenceBase*, Sequence);
+	return GET_ANIM_NODE_DATA(TObjectPtr<UAnimSequenceBase>, Sequence);
 }
 
 float FAnimNode_SequencePlayer::GetPlayRateBasis() const
