Fix assert "check(SectionData.Section->NumTriangles == Segment.NumPrimitives);"
It happens when switching to a higher detailed LOD, and the higher LOD can't fit in skin cache budget. In that case the RT geometry data is out of sync when skin cache later creates an entry for higher LOD

diff --git a/Engine/Source/Runtime/Engine/Private/SkeletalRenderGPUSkin.cpp b/Engine/Source/Runtime/Engine/Private/SkeletalRenderGPUSkin.cpp
--- a/Engine/Source/Runtime/Engine/Private/SkeletalRenderGPUSkin.cpp
+++ b/Engine/Source/Runtime/Engine/Private/SkeletalRenderGPUSkin.cpp
@@ -375,8 +375,11 @@ void FSkeletalMeshObjectGPUSkin::UpdateDynamicData_RenderThread(FGPUSkinCache* G
 		: true);
 
 #if RHI_RAYTRACING
+	// When switching to a higher detailed LOD, skin cache destroys the old entry, if there isn't enough memory to create the new entry because higher LOD costs more memory, SkinCacheEntry becomes null.
+	// So make sure to set bRequireRecreatingRayTracingGeometry to true to account for this, so RT geometry is updated correctly.
 	bRequireRecreatingRayTracingGeometry = (DynamicData == nullptr || RayTracingGeometry.Initializer.Segments.Num() == 0 || // Newly created
-		(DynamicData != nullptr && DynamicData->LODIndex != InDynamicData->LODIndex)); // LOD level changed
+		(DynamicData != nullptr && DynamicData->LODIndex != InDynamicData->LODIndex) || // LOD level changed
+		SkinCacheEntry == nullptr);
 	
 	if (!bRequireRecreatingRayTracingGeometry)
 	{
