- Removed debug display DataLayer Legend and transfered the legend directly in the DataLayer's debug status
- When the world partition's debug 2D runtime hash is shown, default to also display streaming status legend, streaming sources and DataLayers status

diff --git a/Engine/Source/Runtime/Engine/Private/WorldPartition/DataLayer/DataLayerSubsystem.cpp b/Engine/Source/Runtime/Engine/Private/WorldPartition/DataLayer/DataLayerSubsystem.cpp
--- a/Engine/Source/Runtime/Engine/Private/WorldPartition/DataLayer/DataLayerSubsystem.cpp
+++ b/Engine/Source/Runtime/Engine/Private/WorldPartition/DataLayer/DataLayerSubsystem.cpp
@@ -13,12 +13,6 @@
 #include "WorldPartition/DataLayer/IDataLayerEditorModule.h"
 #endif
 
-static int32 GDrawDataLayers = 0;
-static FAutoConsoleCommand CVarDrawDataLayers(
-	TEXT("wp.Runtime.ToggleDrawDataLayers"),
-	TEXT("Toggles debug display of active data layers."),
-	FConsoleCommandDelegate::CreateLambda([] { GDrawDataLayers = !GDrawDataLayers; }));
-
 UDataLayerSubsystem::UDataLayerSubsystem()
 {}
 
@@ -212,7 +206,7 @@ bool UDataLayerSubsystem::IsAnyDataLayerInState(const TArray<FName>& InDataLayer
 
 void UDataLayerSubsystem::DrawDataLayersStatus(UCanvas* Canvas, FVector2D& Offset) const
 {
-	if (!GDrawDataLayers || !Canvas || !Canvas->SceneView)
+	if (!Canvas || !Canvas->SceneView)
 	{
 		return;
 	}
@@ -220,18 +214,21 @@ void UDataLayerSubsystem::DrawDataLayersStatus(UCanvas* Canvas, FVector2D& Offse
 	FVector2D Pos = Offset;
 	float MaxTextWidth = 0.f;
 
-	auto DrawLayerNames = [this, Canvas, &Pos, &MaxTextWidth](const FString& Title, FColor Color, const TSet<FName>& LayerNames)
+	TMap<FName, FColor> ColorMapping;
+	GetDataLayerDebugColors(ColorMapping);
+
+	auto DrawLayerNames = [this, Canvas, &Pos, &MaxTextWidth, &ColorMapping](const FString& Title, FColor Color, const TSet<FName>& LayerNames)
 	{
 		if (LayerNames.Num() > 0)
 		{
-			FWorldPartitionDebugHelper::DrawText(Canvas, Title, GEngine->GetSmallFont(), FColor::Yellow, Pos, &MaxTextWidth);
+			FWorldPartitionDebugHelper::DrawText(Canvas, Title, GEngine->GetSmallFont(), Color, Pos, &MaxTextWidth);
+
 			UFont* DataLayerFont = GEngine->GetTinyFont();
 			for (const FName& DataLayerName : LayerNames)
 			{
 				if (UDataLayer* DataLayer = GetDataLayerFromName(DataLayerName))
 				{
-					const FString Text = DataLayer->GetDataLayerLabel().ToString();
-					FWorldPartitionDebugHelper::DrawText(Canvas, Text, DataLayerFont, Color, Pos, &MaxTextWidth);
+					FWorldPartitionDebugHelper::DrawLegendItem(Canvas, *DataLayer->GetDataLayerLabel().ToString(), DataLayerFont, ColorMapping[DataLayerName], Pos, &MaxTextWidth);
 				}
 			}
 		}
@@ -355,33 +352,4 @@ void UDataLayerSubsystem::GetDataLayerDebugColors(TMap<FName, FColor>& OutMappin
 		OutMapping.Add(DataLayer->GetFName(), DataLayer->GetDebugColor());
 		return true;
 	});
-}
-
-void UDataLayerSubsystem::DrawDataLayersLegend(UCanvas* Canvas, FVector2D& Offset) const
-{
-	check(Canvas);
-
-	const AWorldDataLayers* WorldDataLayers = GetWorld()->GetWorldDataLayers();
-	if (!WorldDataLayers)
-	{
-		return;
-	}
-
-	TMap<FName, FColor> ColorMapping;
-	GetDataLayerDebugColors(ColorMapping);
-
-	float MaxItemWidth = 0;
-	FVector2D Pos = Offset;
-	FWorldPartitionDebugHelper::DrawText(Canvas, TEXT("DataLayers Legend"), GEngine->GetSmallFont(), FColor::Yellow, Pos, &MaxItemWidth);
-	
-	for (auto& It : ColorMapping)
-	{
-		if (const UDataLayer* DataLayer = WorldDataLayers->GetDataLayerFromName(It.Key))
-		{
-			FString StateString = StaticEnum<EDataLayerState>()->GetDisplayNameTextByValue((int64)GetDataLayerState(DataLayer)).ToString();
-			FWorldPartitionDebugHelper::DrawLegendItem(Canvas, *FString::Printf(TEXT("%s (%s)"), *DataLayer->GetDataLayerLabel().ToString(), *StateString), GEngine->GetTinyFont(), It.Value, Pos, &MaxItemWidth);
-		}
-	}
-
-	Offset.X += MaxItemWidth + 10;
 }
\ No newline at end of file
diff --git a/Engine/Source/Runtime/Engine/Private/WorldPartition/WorldPartitionDebugHelper.cpp b/Engine/Source/Runtime/Engine/Private/WorldPartition/WorldPartitionDebugHelper.cpp
--- a/Engine/Source/Runtime/Engine/Private/WorldPartition/WorldPartitionDebugHelper.cpp
+++ b/Engine/Source/Runtime/Engine/Private/WorldPartition/WorldPartitionDebugHelper.cpp
@@ -163,7 +163,7 @@ void FWorldPartitionDebugHelper::DrawLegendItem(UCanvas* Canvas, const FString&
 
 	FVector2D TextPos(Pos.X + ItemSize.X + 10, Pos.Y);
 	float TextWidth = 0.f;
-	FWorldPartitionDebugHelper::DrawText(Canvas, Text, Font, Color, TextPos, &TextWidth);
+	FWorldPartitionDebugHelper::DrawText(Canvas, Text, Font, FColor::White, TextPos, &TextWidth);
 	if (MaxItemWidth)
 	{
 		*MaxItemWidth = FMath::Max(*MaxItemWidth, TextWidth + ItemSize.X + 10);
diff --git a/Engine/Source/Runtime/Engine/Private/WorldPartition/WorldPartitionRuntimeSpatialHash.cpp b/Engine/Source/Runtime/Engine/Private/WorldPartition/WorldPartitionRuntimeSpatialHash.cpp
--- a/Engine/Source/Runtime/Engine/Private/WorldPartition/WorldPartitionRuntimeSpatialHash.cpp
+++ b/Engine/Source/Runtime/Engine/Private/WorldPartition/WorldPartitionRuntimeSpatialHash.cpp
@@ -219,12 +219,24 @@ void FSpatialHashStreamingGrid::GetAlwaysLoadedCells(const UDataLayerSubsystem*
 	}
 }
 
-void FSpatialHashStreamingGrid::GetFilteredCellsForDebugDraw(const FSpatialHashStreamingGridLayerCell* LayerCell, TArray<const UWorldPartitionRuntimeCell*>& FilteredCells) const
+void FSpatialHashStreamingGrid::GetFilteredCellsForDebugDraw(const FSpatialHashStreamingGridLayerCell* LayerCell, const UDataLayerSubsystem* DataLayerSubsystem, TArray<const UWorldPartitionRuntimeCell*>& FilteredCells) const
 {
 	FilteredCells.Reset();
 	if (LayerCell)
 	{
-		Algo::CopyIf(LayerCell->GridCells, FilteredCells, [](const UWorldPartitionRuntimeCell* GridCell) { return GridCell->IsDebugShown(); });
+		Algo::CopyIf(LayerCell->GridCells, FilteredCells, [DataLayerSubsystem](const UWorldPartitionRuntimeCell* GridCell)
+		{ 
+			if (GridCell->IsDebugShown())
+			{
+				EStreamingStatus StreamingStatus = GridCell->GetStreamingStatus();
+				const TArray<FName>& DataLayers = GridCell->GetDataLayers();
+				return (!DataLayers.Num() || 
+						DataLayerSubsystem->IsAnyDataLayerInState(DataLayers, EDataLayerState::Loaded) || 
+						DataLayerSubsystem->IsAnyDataLayerInState(DataLayers, EDataLayerState::Activated) ||
+						((StreamingStatus != LEVEL_Unloaded) && (StreamingStatus != LEVEL_UnloadedButStillAround)));
+			}
+			return false;
+		});
 	}
 	if (FilteredCells.IsEmpty())
 	{
@@ -233,8 +245,15 @@ void FSpatialHashStreamingGrid::GetFilteredCellsForDebugDraw(const FSpatialHashS
 	}
 }
 
-void FSpatialHashStreamingGrid::Draw3D(UWorld* World, const TArray<FWorldPartitionStreamingSource>& Sources, const FTransform& Transform, const TMap<FName, FColor>& DataLayerDebugColors) const
+void FSpatialHashStreamingGrid::Draw3D(UWorld* World, const TArray<FWorldPartitionStreamingSource>& Sources, const FTransform& Transform) const
 {
+	const UDataLayerSubsystem* DataLayerSubsystem = World->GetSubsystem<UDataLayerSubsystem>();
+	TMap<FName, FColor> DataLayerDebugColors;
+	if (DataLayerSubsystem)
+	{
+		DataLayerSubsystem->GetDataLayerDebugColors(DataLayerDebugColors);
+	}
+
 	const FSquare2DGridHelper& Helper = GetGridHelper();
 	int32 MinGridLevel = FMath::Clamp<int32>(GShowRuntimeSpatialHashGridLevel, 0, GridLevels.Num() - 1);
 	int32 MaxGridLevel = FMath::Clamp<int32>(MinGridLevel + GShowRuntimeSpatialHashGridLevelCount - 1, 0, GridLevels.Num() - 1);
@@ -263,7 +282,7 @@ void FSpatialHashStreamingGrid::Draw3D(UWorld* World, const TArray<FWorldPartiti
 			{
 				const int32* LayerCellIndexPtr = GridLevels[GridLevel].LayerCellsMapping.Find(Coords.Y * Helper.Levels[GridLevel].GridSize + Coords.X);
 				const FSpatialHashStreamingGridLayerCell* LayerCell = LayerCellIndexPtr ? &GridLevels[GridLevel].LayerCells[*LayerCellIndexPtr] : nullptr;
-				GetFilteredCellsForDebugDraw(LayerCell, FilteredCells);
+				GetFilteredCellsForDebugDraw(LayerCell, DataLayerSubsystem, FilteredCells);
 				check(FilteredCells.Num());
 
 				FBox2D CellWorldBounds;
@@ -309,8 +328,15 @@ void FSpatialHashStreamingGrid::Draw3D(UWorld* World, const TArray<FWorldPartiti
 	}
 }
 
-void FSpatialHashStreamingGrid::Draw2D(UCanvas* Canvas, const TArray<FWorldPartitionStreamingSource>& Sources, const FBox& Region, const FBox2D& GridScreenBounds, TFunctionRef<FVector2D(const FVector2D&)> WorldToScreen, const TMap<FName, FColor>& DataLayerDebugColors) const
+void FSpatialHashStreamingGrid::Draw2D(UCanvas* Canvas, UWorld* World, const TArray<FWorldPartitionStreamingSource>& Sources, const FBox& Region, const FBox2D& GridScreenBounds, TFunctionRef<FVector2D(const FVector2D&)> WorldToScreen) const
 {
+	const UDataLayerSubsystem* DataLayerSubsystem = World->GetSubsystem<UDataLayerSubsystem>();
+	TMap<FName, FColor> DataLayerDebugColors;
+	if (DataLayerSubsystem)
+	{
+		DataLayerSubsystem->GetDataLayerDebugColors(DataLayerDebugColors);
+	}
+
 	FCanvas* CanvasObject = Canvas->Canvas;
 	int32 MinGridLevel = FMath::Clamp<int32>(GShowRuntimeSpatialHashGridLevel, 0, GridLevels.Num() - 1);
 	int32 MaxGridLevel = FMath::Clamp<int32>(MinGridLevel + GShowRuntimeSpatialHashGridLevelCount - 1, 0, GridLevels.Num() - 1);
@@ -391,7 +417,7 @@ void FSpatialHashStreamingGrid::Draw2D(UCanvas* Canvas, const TArray<FWorldParti
 
 			const int32* LayerCellIndexPtr = GridLevels[GridLevel].LayerCellsMapping.Find(Coords.Y * Helper.Levels[GridLevel].GridSize + Coords.X);
 			const FSpatialHashStreamingGridLayerCell* LayerCell = LayerCellIndexPtr ? &GridLevels[GridLevel].LayerCells[*LayerCellIndexPtr] : nullptr;
-			GetFilteredCellsForDebugDraw(LayerCell, FilteredCells);
+			GetFilteredCellsForDebugDraw(LayerCell, DataLayerSubsystem, FilteredCells);
 			check(FilteredCells.Num());
 
 			FVector2D CellBoundsSize = CellScreenBounds.GetSize();
@@ -1021,7 +1047,7 @@ bool UWorldPartitionRuntimeSpatialHash::GetStreamingCells(const TArray<FWorldPar
 {
 	TRACE_CPUPROFILER_EVENT_SCOPE(UWorldPartitionRuntimeSpatialHash::GetStreamingCells);
 
-	UDataLayerSubsystem* DataLayerSubsystem = GetWorld()->GetSubsystem<UDataLayerSubsystem>();
+	const UDataLayerSubsystem* DataLayerSubsystem = GetWorld()->GetSubsystem<UDataLayerSubsystem>();
 
 	if (Sources.Num() == 0)
 	{
@@ -1058,13 +1084,6 @@ void UWorldPartitionRuntimeSpatialHash::Draw2D(UCanvas* Canvas, const TArray<FWo
 
 	UWorldPartition* WorldPartition = GetOuterUWorldPartition();
 	UWorld* World = WorldPartition->GetWorld();
-	UDataLayerSubsystem* DataLayerSubsystem = World->GetSubsystem<UDataLayerSubsystem>();
-
-	TMap<FName, FColor> DataLayerDebugColors;
-	if (DataLayerSubsystem)
-	{
-		DataLayerSubsystem->GetDataLayerDebugColors(DataLayerDebugColors);
-	}
 
 	const float CanvasMaxScreenSize = PartitionCanvasSize.X;
 	const float GridMaxScreenSize = CanvasMaxScreenSize / FilteredStreamingGrids.Num();
@@ -1099,7 +1118,7 @@ void UWorldPartitionRuntimeSpatialHash::Draw2D(UCanvas* Canvas, const TArray<FWo
 		const float WorldToScreenScale = (0.5f * GridEffectiveScreenSize) / AverageSphere.W;
 		auto WorldToScreen = [&](const FVector2D& WorldPos) { return (WorldToScreenScale * (WorldPos - GridReferenceWorldPos)) + GridScreenOffset; };
 
-		StreamingGrid->Draw2D(Canvas, Sources, Region, GridScreenBounds, WorldToScreen, DataLayerDebugColors);
+		StreamingGrid->Draw2D(Canvas, World, Sources, Region, GridScreenBounds, WorldToScreen);
 		Offset = GridScreenBounds.Max;
 
 		// Draw WorldPartition name
@@ -1128,19 +1147,12 @@ void UWorldPartitionRuntimeSpatialHash::Draw3D(const TArray<FWorldPartitionStrea
 {
 	UWorld* World = GetWorld();
 	UWorldPartition* WorldPartition = GetOuterUWorldPartition();
-	UDataLayerSubsystem* DataLayerSubsystem = World->GetSubsystem<UDataLayerSubsystem>();
-
-	TMap<FName, FColor> DataLayerDebugColors;
-	if (DataLayerSubsystem)
-	{
-		DataLayerSubsystem->GetDataLayerDebugColors(DataLayerDebugColors);
-	}
 
 	FTransform Transform = WorldPartition->GetInstanceTransform();
 	TArray<const FSpatialHashStreamingGrid*> FilteredStreamingGrids = GetFilteredStreamingGrids();
 	for (const FSpatialHashStreamingGrid* StreamingGrid : FilteredStreamingGrids)
 	{
-		StreamingGrid->Draw3D(World, Sources, Transform, DataLayerDebugColors);
+		StreamingGrid->Draw3D(World, Sources, Transform);
 	}
 }
 
diff --git a/Engine/Source/Runtime/Engine/Private/WorldPartition/WorldPartitionSubsystem.cpp b/Engine/Source/Runtime/Engine/Private/WorldPartition/WorldPartitionSubsystem.cpp
--- a/Engine/Source/Runtime/Engine/Private/WorldPartition/WorldPartitionSubsystem.cpp
+++ b/Engine/Source/Runtime/Engine/Private/WorldPartition/WorldPartitionSubsystem.cpp
@@ -43,6 +43,12 @@ static FAutoConsoleCommand CVarDrawRuntimeCellsDetails(
 	TEXT("Toggles debug display of world partition runtime streaming cells."),
 	FConsoleCommandDelegate::CreateLambda([] { GDrawRuntimeCellsDetails = !GDrawRuntimeCellsDetails; }));
 
+static int32 GDrawDataLayers = 0;
+static FAutoConsoleCommand CVarDrawDataLayers(
+	TEXT("wp.Runtime.ToggleDrawDataLayers"),
+	TEXT("Toggles debug display of active data layers."),
+	FConsoleCommandDelegate::CreateLambda([] { GDrawDataLayers = !GDrawDataLayers; }));
+
 UWorldPartitionSubsystem::UWorldPartitionSubsystem()
 {}
 
@@ -253,7 +259,7 @@ void UWorldPartitionSubsystem::Draw(UCanvas* Canvas, class APlayerController* PC
 		}
 	}
 
-	if (GDrawStreamingSources)
+	if (GDrawStreamingSources || GDrawRuntimeHash2D)
 	{
 		const UWorldPartition* WorldPartition = GetMainWorldPartition();
 		const TArray<FWorldPartitionStreamingSource>* StreamingSources = WorldPartition ? &WorldPartition->GetStreamingSources() : nullptr;
@@ -282,19 +288,13 @@ void UWorldPartitionSubsystem::Draw(UCanvas* Canvas, class APlayerController* PC
 	{
 		UDataLayerSubsystem* DataLayerSubsystem = WorldPartition->GetWorld()->GetSubsystem<UDataLayerSubsystem>();
 
-		if (GDrawLegends)
+		if (GDrawLegends || GDrawRuntimeHash2D)
 		{
 			// Streaming Status Legend
 			WorldPartition->DrawStreamingStatusLegend(Canvas, CurrentOffset);
-
-			// DataLayers Legend
-			if (DataLayerSubsystem)
-			{
-				DataLayerSubsystem->DrawDataLayersLegend(Canvas, CurrentOffset);
-			}
 		}
 
-		if (DataLayerSubsystem)
+		if (DataLayerSubsystem && (GDrawDataLayers || GDrawRuntimeHash2D))
 		{
 			DataLayerSubsystem->DrawDataLayersStatus(Canvas, CurrentOffset);
 		}
diff --git a/Engine/Source/Runtime/Engine/Public/WorldPartition/DataLayer/DataLayerSubsystem.h b/Engine/Source/Runtime/Engine/Public/WorldPartition/DataLayer/DataLayerSubsystem.h
--- a/Engine/Source/Runtime/Engine/Public/WorldPartition/DataLayer/DataLayerSubsystem.h
+++ b/Engine/Source/Runtime/Engine/Public/WorldPartition/DataLayer/DataLayerSubsystem.h
@@ -91,7 +91,6 @@ public:
 	bool IsAnyDataLayerInState(const TArray<FName>& InDataLayerNames, EDataLayerState InState) const;
 	
 	void GetDataLayerDebugColors(TMap<FName, FColor>& OutMapping) const;
-	void DrawDataLayersLegend(UCanvas* Canvas, FVector2D& Offset) const;
 	void DrawDataLayersStatus(UCanvas* Canvas, FVector2D& Offset) const;
 	static TArray<UDataLayer*> ConvertArgsToDataLayers(UWorld* World, const TArray<FString>& InArgs);
 
diff --git a/Engine/Source/Runtime/Engine/Public/WorldPartition/WorldPartitionRuntimeSpatialHash.h b/Engine/Source/Runtime/Engine/Public/WorldPartition/WorldPartitionRuntimeSpatialHash.h
--- a/Engine/Source/Runtime/Engine/Public/WorldPartition/WorldPartitionRuntimeSpatialHash.h
+++ b/Engine/Source/Runtime/Engine/Public/WorldPartition/WorldPartitionRuntimeSpatialHash.h
@@ -127,11 +127,11 @@ struct FSpatialHashStreamingGrid
 	void GetCells(const FWorldPartitionStreamingQuerySource& QuerySource, TSet<const UWorldPartitionRuntimeCell*>& OutCells) const;
 	void GetCells(const TArray<FWorldPartitionStreamingSource>& Sources, const class UDataLayerSubsystem* DataLayerSubsystem, UWorldPartitionRuntimeHash::FStreamingSourceCells& OutActivateCells, UWorldPartitionRuntimeHash::FStreamingSourceCells& OutLoadCells) const;
 	void GetAlwaysLoadedCells(const UDataLayerSubsystem* DataLayerSubsystem, TSet<const UWorldPartitionRuntimeCell*>& OutActivateCells, TSet<const UWorldPartitionRuntimeCell*>& OutLoadCells) const;
-	void Draw2D(UCanvas* Canvas, const TArray<FWorldPartitionStreamingSource>& Sources, const FBox& Region, const FBox2D& GridScreenBounds, TFunctionRef<FVector2D(const FVector2D&)> WorldToScreen, const TMap<FName, FColor>& DataLayerDebugColors) const;
-	void Draw3D(UWorld* World, const TArray<FWorldPartitionStreamingSource>& Sources, const FTransform& Transform, const TMap<FName, FColor>& DataLayerDebugColors) const;
+	void Draw2D(UCanvas* Canvas, UWorld* World, const TArray<FWorldPartitionStreamingSource>& Sources, const FBox& Region, const FBox2D& GridScreenBounds, TFunctionRef<FVector2D(const FVector2D&)> WorldToScreen) const;
+	void Draw3D(UWorld* World, const TArray<FWorldPartitionStreamingSource>& Sources, const FTransform& Transform) const;
 
 private:
-	void GetFilteredCellsForDebugDraw(const FSpatialHashStreamingGridLayerCell* LayerCell, TArray<const UWorldPartitionRuntimeCell*>& FilteredCells) const;
+	void GetFilteredCellsForDebugDraw(const FSpatialHashStreamingGridLayerCell* LayerCell, const UDataLayerSubsystem* DataLayerSubsystem, TArray<const UWorldPartitionRuntimeCell*>& FilteredCells) const;
 	const FSquare2DGridHelper& GetGridHelper() const;
 	mutable FSquare2DGridHelper* GridHelper;
 };
