Fix the error of transition-to-read buffers set to store 'active' tangent buffers which can be intermediate tangent buffers when vertex color blending is turned on recompute tangents, it should be using the output tangent buffers.

diff --git a/Engine/Source/Runtime/Engine/Private/GPUSkinCache.cpp b/Engine/Source/Runtime/Engine/Private/GPUSkinCache.cpp
--- a/Engine/Source/Runtime/Engine/Private/GPUSkinCache.cpp
+++ b/Engine/Source/Runtime/Engine/Private/GPUSkinCache.cpp
@@ -1330,7 +1330,6 @@ void FGPUSkinCache::DoDispatch(FRHICommandListImmediate& RHICmdList)
 		Algo::Sort(BuffersToTransitionForSkinning, [](const FRHIUnorderedAccessView* A, const FRHIUnorderedAccessView* B){return A < B;});
 		BuffersToTransitionForSkinning.SetNum(Algo::Unique(BuffersToTransitionForSkinning));
 		MakeBufferTransitions(RHICmdList, BuffersToTransitionForSkinning, ERHIAccess::Unknown, ERHIAccess::UAVCompute);
-		BuffersToTransitionToRead.Append(BuffersToTransitionForSkinning);
 	}
 
 	RHICmdList.BeginUAVOverlap(BuffersToTransitionForSkinning);
@@ -1424,7 +1423,6 @@ void FGPUSkinCache::DoDispatch(FRHICommandListImmediate& RHICmdList, FGPUSkinCac
 	TArray<FRHIUnorderedAccessView*> BuffersToTransitionForSkinning;
 	PrepareUpdateSkinning(SkinCacheEntry, Section, RevisionNumber, &BuffersToTransitionForSkinning);
 	MakeBufferTransitions(RHICmdList, BuffersToTransitionForSkinning, ERHIAccess::Unknown, ERHIAccess::UAVCompute);
-	BuffersToTransitionToRead.Append(BuffersToTransitionForSkinning);
 
 	RHICmdList.BeginUAVOverlap(BuffersToTransitionForSkinning);
 	DispatchUpdateSkinning(RHICmdList, SkinCacheEntry, Section, RevisionNumber);
@@ -2095,6 +2093,7 @@ void FGPUSkinCache::DispatchUpdateSkinning(FRHICommandListImmediate& RHICmdList,
 		RHICmdList.DispatchComputeShader(VertexCountAlign64, 1, 1);
 		Shader->UnsetParameters(RHICmdList);
 		IncrementDispatchCounter(RHICmdList);
+		BuffersToTransitionToRead.Add(DispatchData.GetPreviousPositionRWBuffer()->UAV);
 	}
 
 	if ((DispatchData.DispatchFlags & (uint32)EGPUSkinCacheDispatchFlags::DispatchPosition) != 0)
@@ -2115,8 +2114,10 @@ void FGPUSkinCache::DispatchUpdateSkinning(FRHICommandListImmediate& RHICmdList,
 		RHICmdList.DispatchComputeShader(VertexCountAlign64, 1, 1);
 		Shader->UnsetParameters(RHICmdList);
 		IncrementDispatchCounter(RHICmdList);
+		BuffersToTransitionToRead.Add(DispatchData.GetPositionRWBuffer()->UAV);
 	}
 
+	BuffersToTransitionToRead.Add(DispatchData.GetTangentRWBuffer()->UAV);
 	check(DispatchData.PreviousPositionBuffer != DispatchData.PositionBuffer);
 }
 
