Added StatId in UInstancedStaticMeshComponent for insight to show component/mesh when CreateAllInstanceBodies/ClearAllInstanceBodies are called.

diff --git a/Engine/Source/Runtime/Engine/Classes/Components/InstancedStaticMeshComponent.h b/Engine/Source/Runtime/Engine/Classes/Components/InstancedStaticMeshComponent.h
--- a/Engine/Source/Runtime/Engine/Classes/Components/InstancedStaticMeshComponent.h
+++ b/Engine/Source/Runtime/Engine/Classes/Components/InstancedStaticMeshComponent.h
@@ -448,6 +448,11 @@ protected:
 	friend FStaticLightingTextureMapping_InstancedStaticMesh;
 	friend FInstancedLightMap2D;
 	friend FInstancedShadowMap2D;
+
+#if STATS
+	/** Used for dynamic stats */
+	TStatId StatId;
+#endif
 };
 
 /** InstancedStaticMeshInstance hit proxy */
diff --git a/Engine/Source/Runtime/Engine/Private/InstancedStaticMesh.cpp b/Engine/Source/Runtime/Engine/Private/InstancedStaticMesh.cpp
--- a/Engine/Source/Runtime/Engine/Private/InstancedStaticMesh.cpp
+++ b/Engine/Source/Runtime/Engine/Private/InstancedStaticMesh.cpp
@@ -1804,6 +1804,17 @@ UInstancedStaticMeshComponent::UInstancedStaticMeshComponent(const FObjectInitia
 	BodyInstance.bSimulatePhysics = false;
 
 	bDisallowMeshPaintPerInstance = true;
+
+#if STATS
+	{
+		UObject const* StatObject = this->AdditionalStatObject();
+		if (!StatObject)
+		{
+			StatObject = this;
+		}
+		StatId = StatObject->GetStatID(true);
+	}
+#endif
 }
 
 UInstancedStaticMeshComponent::UInstancedStaticMeshComponent(FVTableHelper& Helper)
@@ -2086,6 +2097,7 @@ void UInstancedStaticMeshComponent::CreateAllInstanceBodies()
 {
 	TRACE_CPUPROFILER_EVENT_SCOPE(UInstancedStaticMeshComponent::CreateAllInstanceBodies);
 	QUICK_SCOPE_CYCLE_COUNTER(STAT_UInstancedStaticMeshComponent_CreateAllInstanceBodies);
+	STAT(FScopeCycleCounter Context(StatId);)
 
 	const int32 NumBodies = PerInstanceSMData.Num();
 	check(InstanceBodies.Num() == 0);
@@ -2155,6 +2167,8 @@ void UInstancedStaticMeshComponent::CreateAllInstanceBodies()
 void UInstancedStaticMeshComponent::ClearAllInstanceBodies()
 {
 	QUICK_SCOPE_CYCLE_COUNTER(STAT_UInstancedStaticMeshComponent_ClearAllInstanceBodies);
+	STAT(FScopeCycleCounter Context(StatId);)
+
 	for (int32 i = 0; i < InstanceBodies.Num(); i++)
 	{
 		if (InstanceBodies[i])
