Fix merge error in GetPackageGuid.
#rnx

diff --git a/Engine/Source/Runtime/Engine/Private/UnrealEngine.cpp b/Engine/Source/Runtime/Engine/Private/UnrealEngine.cpp
--- a/Engine/Source/Runtime/Engine/Private/UnrealEngine.cpp
+++ b/Engine/Source/Runtime/Engine/Private/UnrealEngine.cpp
@@ -9537,14 +9537,16 @@ void UEngine::AddTextureStreamingSlaveLoc(FVector InLoc, float BoostFactor, bool
 FGuid UEngine::GetPackageGuid(FName PackageName, bool bForPIE)
 {
 	FGuid Result(0,0,0,0);
-	FPackagePath PackagePath;
-	if (!FPackagePath::TryFromMountedName(PackageName.ToString(), PackagePath))
-
 	// There is no package guid support when using the I/O dispatcher
 	if (FIoDispatcher::IsInitialized())
 	{
 		return Result;
 	}
+	FPackagePath PackagePath;
+	if (!FPackagePath::TryFromMountedName(PackageName.ToString(), PackagePath))
+	{
+		return Result;
+	}
 
 	uint32 LoadFlags = LOAD_NoWarn | LOAD_NoVerify;
 	if (bForPIE)
