Fix shader compilation starvation while waiting on async assets when using SN-DBS which requires the http service to be ticked from the game-thread regularly

diff --git a/Engine/Source/Runtime/Engine/Private/AsyncCompilationHelpers.cpp b/Engine/Source/Runtime/Engine/Private/AsyncCompilationHelpers.cpp
--- a/Engine/Source/Runtime/Engine/Private/AsyncCompilationHelpers.cpp
+++ b/Engine/Source/Runtime/Engine/Private/AsyncCompilationHelpers.cpp
@@ -16,6 +16,7 @@
 #include "AssetCompilingManager.h"
 #include "HAL/PlatformStackWalk.h"
 #include "Hash/CityHash.h"
+#include "ShaderCompiler.h"
 #include "Misc/ScopeRWLock.h"
 #include "ProfilingDebugging/StallDetector.h"
 
@@ -143,7 +144,7 @@ namespace AsyncCompilationHelpers
 
 			FText Progress = FormatProgress(NumDone++, Num, Job.GetName());
 
-			// Be nice with the game thread and tick the progress at 60 fps even when no progress is being made...
+			// Be nice with the game thread and tick the progress to keep application responsive even when no progress is being made...
 			bool bLogSlowProgress = true;
 			while (!Job.WaitCompletionWithTimeout(0.016))
 			{
@@ -151,11 +152,22 @@ namespace AsyncCompilationHelpers
 				{
 					SlowTask->EnterProgressFrame(0.0f, Progress);
 				}
+
 				if (bLogSlowProgress)
 				{
 					UE_LOG_REF(LogCategory, Display, TEXT("%s"), *Progress.ToString());
 					bLogSlowProgress = false;
 				}
+
+				// SN-DBS jobs (which make use of the http service) needs to be ticked
+				// from the game-thread to avoid starvation while we wait for other
+				// async tasks to finish.
+				if (GShaderCompilingManager && GShaderCompilingManager->IsCompiling())
+				{
+					const bool bLimitExecutionTime = true;
+					const bool bBlockOnGlobalShaderCompletion = false;
+					GShaderCompilingManager->ProcessAsyncResults(bLimitExecutionTime, bBlockOnGlobalShaderCompletion);
+				}
 			}
 
 			if (SlowTask.IsSet())
