Allow for parallel processing of render commands in editor game modes, seperately from the rest of the editor worlds.

diff --git a/Engine/Source/Runtime/Engine/Private/LevelTick.cpp b/Engine/Source/Runtime/Engine/Private/LevelTick.cpp
--- a/Engine/Source/Runtime/Engine/Private/LevelTick.cpp
+++ b/Engine/Source/Runtime/Engine/Private/LevelTick.cpp
@@ -798,6 +798,11 @@ static TAutoConsoleVariable<int32> CVarAllowAsyncRenderThreadUpdatesDuringGameth
 	1,
 	TEXT("If > 0 then we do the gamethread updates _while_ doing parallel updates."));
 
+static TAutoConsoleVariable<int32> CVarAllowAsyncRenderThreadUpdatesEditorGameWorld(
+	TEXT("AllowAsyncRenderThreadUpdatesEditorGameWorld"),
+	0,
+	TEXT("Used to control async renderthread updates in an editor game world."));
+
 static TAutoConsoleVariable<int32> CVarAllowAsyncRenderThreadUpdatesEditor(
 	TEXT("AllowAsyncRenderThreadUpdatesEditor"),
 	0,
@@ -914,14 +919,23 @@ void UWorld::MarkActorComponentForNeededEndOfFrameUpdate(UActorComponent* Compon
 
 	if (CurrentState == EComponentMarkedForEndOfFrameUpdateState::Unmarked)
 	{
+		// When there is no rendering thread force all updates on game thread,
+		// to avoid modifying scene structures from multiple task threads
+		bForceGameThread = bForceGameThread || !GIsThreadedRendering || !FApp::ShouldUseThreadingForPerformance();
 		if (!bForceGameThread)
 		{
-			bool bAllowConcurrentUpdates = FApp::ShouldUseThreadingForPerformance() && 
-				(GIsEditor ? !!CVarAllowAsyncRenderThreadUpdatesEditor.GetValueOnAnyThread() : !!CVarAllowAsyncRenderThreadUpdates.GetValueOnAnyThread());
-			bForceGameThread = !bAllowConcurrentUpdates 
-								// When there is no rendering thread force all updates on game thread,
-								// to avoid modifying scene structures from multiple task threads
-								|| !GIsThreadedRendering;
+#if WITH_EDITOR
+			if (IsGameWorld())
+			{
+				bForceGameThread = !CVarAllowAsyncRenderThreadUpdatesEditorGameWorld.GetValueOnAnyThread();
+			}
+			else
+			{
+				bForceGameThread = !CVarAllowAsyncRenderThreadUpdatesEditor.GetValueOnAnyThread();
+			}
+#else
+			bForceGameThread = !CVarAllowAsyncRenderThreadUpdates.GetValueOnAnyThread();
+#endif
 		}
 
 		if (bForceGameThread)
