Fix task scoping issue causing engine test failures.

diff --git a/Engine/Source/Runtime/Engine/Private/Streaming/RenderAssetUpdate.cpp b/Engine/Source/Runtime/Engine/Private/Streaming/RenderAssetUpdate.cpp
--- a/Engine/Source/Runtime/Engine/Private/Streaming/RenderAssetUpdate.cpp
+++ b/Engine/Source/Runtime/Engine/Private/Streaming/RenderAssetUpdate.cpp
@@ -288,6 +288,7 @@ void FRenderAssetUpdate::ScheduleAsyncTask()
 
 void FRenderAssetUpdate::FMipUpdateTask::DoWork()
 {
+	FTaskTagScope Scope(ETaskTag::EParallelGameThread);
 	check(PendingUpdate.IsValid());
 
 #if !UE_BUILD_SHIPPING
diff --git a/Engine/Source/Runtime/Engine/Private/Texture.cpp b/Engine/Source/Runtime/Engine/Private/Texture.cpp
--- a/Engine/Source/Runtime/Engine/Private/Texture.cpp
+++ b/Engine/Source/Runtime/Engine/Private/Texture.cpp
@@ -129,20 +129,32 @@ UTexture::UTexture(const FObjectInitializer& ObjectInitializer)
 
 const FTextureResource* UTexture::GetResource() const
 {
-	if (IsInActualRenderingThread() || IsInRHIThread())
+	if (IsInParallelGameThread() || IsInGameThread())
+	{
+		return PrivateResource;
+	}
+	else if (IsInParallelRenderingThread() || IsInRHIThread())
 	{
 		return PrivateResourceRenderThread;
 	}
-	return PrivateResource;
+
+	UE_LOG(LogTexture, Error, TEXT("Attempted to access a texture resource from an unkown thread."));
+	return nullptr;
 }
 
 FTextureResource* UTexture::GetResource()
 {
-	if (IsInActualRenderingThread() || IsInRHIThread())
+	if (IsInParallelGameThread() || IsInGameThread())
+	{
+		return PrivateResource;
+	}
+	else if (IsInParallelRenderingThread() || IsInRHIThread())
 	{
 		return PrivateResourceRenderThread;
 	}
-	return PrivateResource;
+
+	UE_LOG(LogTexture, Error, TEXT("Attempted to access a texture resource from an unkown thread."));
+	return nullptr;
 }
 
 void UTexture::SetResource(FTextureResource* InResource)
