Items now check the category they're in from the parent container.

diff --git a/Engine/Source/Runtime/Engine/Private/ObjectEditorUtils.cpp b/Engine/Source/Runtime/Engine/Private/ObjectEditorUtils.cpp
--- a/Engine/Source/Runtime/Engine/Private/ObjectEditorUtils.cpp
+++ b/Engine/Source/Runtime/Engine/Private/ObjectEditorUtils.cpp
@@ -67,9 +67,21 @@ namespace FObjectEditorUtils
 	FName GetCategoryFName( const FField* InField )
 	{
 		static const FName CategoryKey(TEXT("Category"));
-		if (InField && InField->HasMetaData(CategoryKey))
+		if (InField)
 		{
-			return FName(*InField->GetMetaData(CategoryKey));
+			const FField* FieldToConsider = InField;
+
+			const bool bIsInsideContainerProperty = InField->GetOwner<FArrayProperty>() || InField->GetOwner<FSetProperty>() || InField->GetOwner<FMapProperty>();
+			if (bIsInsideContainerProperty)
+			{
+				FieldToConsider = InField->GetOwner<FField>();
+			}
+
+			FString CategoryString = FieldToConsider->GetMetaData(CategoryKey);
+			if (!CategoryString.IsEmpty())
+			{
+				return FName(*CategoryString);
+			}
 		}
 		return NAME_None;
 	}
