Assigned OwnerWorld  properly for runtime cells levels
- Trying to obtain the correct WP world in PostLoad() calls for actors of this cell level would fail, as the OwnerWorld was still set to the package UWorld.

diff --git a/Engine/Source/Runtime/Engine/Private/WorldPartition/WorldPartitionLevelHelper.cpp b/Engine/Source/Runtime/Engine/Private/WorldPartition/WorldPartitionLevelHelper.cpp
--- a/Engine/Source/Runtime/Engine/Private/WorldPartition/WorldPartitionLevelHelper.cpp
+++ b/Engine/Source/Runtime/Engine/Private/WorldPartition/WorldPartitionLevelHelper.cpp
@@ -139,7 +139,7 @@ void FWorldPartitionLevelHelper::RemapLevelSoftObjectPaths(ULevel* InLevel, UWor
 /**
  * Creates an empty Level used in World Partition
  */
-ULevel* FWorldPartitionLevelHelper::CreateEmptyLevelForRuntimeCell(const UWorld* InWorld, const FString& InWorldAssetName, UPackage* InPackage)
+ULevel* FWorldPartitionLevelHelper::CreateEmptyLevelForRuntimeCell(UWorld* InWorld, const FString& InWorldAssetName, UPackage* InPackage)
 {
 	// Create or use given package
 	UPackage* CellPackage = nullptr;
@@ -180,6 +180,10 @@ ULevel* FWorldPartitionLevelHelper::CreateEmptyLevelForRuntimeCell(const UWorld*
 	check(NewLevel->OwningWorld == NewWorld);
 	check(NewLevel->Model);
 	check(!NewLevel->bIsVisible);
+
+	// Update the persistent level's owning world to the correct world.
+	NewLevel->OwningWorld = InWorld;
+
 	return NewLevel;
 }
 
diff --git a/Engine/Source/Runtime/Engine/Private/WorldPartition/WorldPartitionLevelStreamingDynamic.cpp b/Engine/Source/Runtime/Engine/Private/WorldPartition/WorldPartitionLevelStreamingDynamic.cpp
--- a/Engine/Source/Runtime/Engine/Private/WorldPartition/WorldPartitionLevelStreamingDynamic.cpp
+++ b/Engine/Source/Runtime/Engine/Private/WorldPartition/WorldPartitionLevelStreamingDynamic.cpp
@@ -82,7 +82,7 @@ void UWorldPartitionLevelStreamingDynamic::CreateRuntimeLevel()
 {
 	check(PendingUnloadLevel == nullptr);
 	check(RuntimeLevel == nullptr);
-	const UWorld* PlayWorld = GetWorld();
+	UWorld* PlayWorld = GetWorld();
 	check(PlayWorld && PlayWorld->IsGameWorld());
 
 	// Create streaming cell Level package
diff --git a/Engine/Source/Runtime/Engine/Public/WorldPartition/WorldPartitionLevelHelper.h b/Engine/Source/Runtime/Engine/Public/WorldPartition/WorldPartitionLevelHelper.h
--- a/Engine/Source/Runtime/Engine/Public/WorldPartition/WorldPartitionLevelHelper.h
+++ b/Engine/Source/Runtime/Engine/Public/WorldPartition/WorldPartitionLevelHelper.h
@@ -20,7 +20,7 @@ class UWorldPartition;
 class FWorldPartitionLevelHelper
 {
 public:
-	static ULevel* CreateEmptyLevelForRuntimeCell(const UWorld* InWorld, const FString& InWorldAssetName, UPackage* DestPackage = nullptr);
+	static ULevel* CreateEmptyLevelForRuntimeCell(UWorld* InWorld, const FString& InWorldAssetName, UPackage* DestPackage = nullptr);
 	static void MoveExternalActorsToLevel(const TArray<FWorldPartitionRuntimeCellObjectMapping>& InChildPackages, ULevel* InLevel);
 	static void RemapLevelSoftObjectPaths(ULevel* InLevel, UWorldPartition* InWorldPartition);
 	
