Added conversion of the previous CellName property to the new SourceCell SoftObjectPtr

diff --git a/Engine/Source/Runtime/Engine/Private/WorldPartition/HLOD/HLODActor.cpp b/Engine/Source/Runtime/Engine/Private/WorldPartition/HLOD/HLODActor.cpp
--- a/Engine/Source/Runtime/Engine/Private/WorldPartition/HLOD/HLODActor.cpp
+++ b/Engine/Source/Runtime/Engine/Private/WorldPartition/HLOD/HLODActor.cpp
@@ -61,7 +61,9 @@ void AWorldPartitionHLOD::Serialize(FArchive& Ar)
 #if WITH_EDITOR
 	if(Ar.IsLoading() && Ar.CustomVer(FUE5MainStreamObjectVersion::GUID) < FUE5MainStreamObjectVersion::WorldPartitionStreamingCellsNamingShortened)
 	{
-		SourceCell = SourceCell.ToString().Replace(TEXT("WPRT_"), TEXT(""), ESearchCase::CaseSensitive).Replace(TEXT("Cell_"), TEXT(""), ESearchCase::CaseSensitive);
+		auto FixupCellName = [](const FString& CellName) { return CellName.Replace(TEXT("WPRT_"), TEXT(""), ESearchCase::CaseSensitive).Replace(TEXT("Cell_"), TEXT(""), ESearchCase::CaseSensitive); };
+		CellName_DEPRECATED = FName(FixupCellName(CellName_DEPRECATED.ToString()));
+		SourceCell = FixupCellName(SourceCell.ToString());
 	}
 #endif
 }
@@ -69,6 +71,28 @@ void AWorldPartitionHLOD::Serialize(FArchive& Ar)
 void AWorldPartitionHLOD::RerunConstructionScripts()
 {}
 
+#if WITH_EDITORONLY_DATA
+void AWorldPartitionHLOD::PostLoad()
+{
+	Super::PostLoad();
+
+	// Convert CellName to a soft object path referencing the RuntimeCell object.
+	if (!CellName_DEPRECATED.IsNone() && SourceCell.IsNull())
+	{
+		GetWorld()->ConditionalPostLoad();
+
+		UWorldPartition* WorldPartition = GetWorld()->GetWorldPartition();
+		if (ensure(WorldPartition))
+		{
+			FSoftObjectPath WorldPartitionPath(WorldPartition);
+			FString CellObjectPath = FString::Printf(TEXT("%s.%s"), *WorldPartitionPath.ToString(), *CellName_DEPRECATED.ToString());
+			TSoftObjectPtr<UWorldPartitionRuntimeCell> RuntimeCell(CellObjectPath);
+			SourceCell = RuntimeCell;
+		}
+	}
+}
+#endif
+
 #if WITH_EDITOR
 
 bool AWorldPartitionHLOD::IsHiddenEd() const
diff --git a/Engine/Source/Runtime/Engine/Public/WorldPartition/HLOD/HLODActor.h b/Engine/Source/Runtime/Engine/Public/WorldPartition/HLOD/HLODActor.h
--- a/Engine/Source/Runtime/Engine/Public/WorldPartition/HLOD/HLODActor.h
+++ b/Engine/Source/Runtime/Engine/Public/WorldPartition/HLOD/HLODActor.h
@@ -61,6 +61,9 @@ protected:
 	//~ Begin UObject Interface.
 	virtual void Serialize(FArchive& Ar) override;
 	virtual void RerunConstructionScripts() override;
+#if WITH_EDITORONLY_DATA
+	virtual void PostLoad() override;
+#endif
 	//~ End UObject Interface.
 
 	//~ Begin AActor Interface.
@@ -101,6 +104,10 @@ private:
 
 	UPROPERTY()
 	uint32 HLODHash;
+
+	// Deprecated. Stored in SourceCell instead.
+	UPROPERTY()
+	FName CellName_DEPRECATED;
 #endif
 
 	UPROPERTY()
