Fix  GPU physics field crashing on mobile

diff --git a/Engine/Source/Runtime/Engine/Private/World.cpp b/Engine/Source/Runtime/Engine/Private/World.cpp
--- a/Engine/Source/Runtime/Engine/Private/World.cpp
+++ b/Engine/Source/Runtime/Engine/Private/World.cpp
@@ -2107,7 +2107,7 @@ void UWorld::UpdateWorldComponents(bool bRerunConstructionScripts, bool bCurrent
 		}
 
 		static IConsoleVariable* PhysicsFieldEnableClipmapCVar = IConsoleManager::Get().FindConsoleVariable(TEXT("r.PhysicsField.EnableField"));
-		if (PhysicsFieldEnableClipmapCVar && PhysicsFieldEnableClipmapCVar->GetInt() == 1)
+		if (PhysicsFieldEnableClipmapCVar && PhysicsFieldEnableClipmapCVar->GetInt() == 1 && !IsMobilePlatform(Scene->GetShaderPlatform()))
 		{
 			if (!PhysicsField)
 			{
