Emit largeworldcoordinates status csv metadata and output state to LogInit at startup.

diff --git a/Engine/Source/Runtime/Engine/Private/UnrealEngine.cpp b/Engine/Source/Runtime/Engine/Private/UnrealEngine.cpp
--- a/Engine/Source/Runtime/Engine/Private/UnrealEngine.cpp
+++ b/Engine/Source/Runtime/Engine/Private/UnrealEngine.cpp
@@ -1764,6 +1764,10 @@ void UEngine::Init(IEngineLoop* InEngineLoop)
 		FModuleManager::Get().LoadModule("LevelSequence");
 	}
 
+	// Record large world coordinate state
+	UE_LOG(LogInit, Log, TEXT("Large World Coordinates: %s"), UE_LARGE_WORLD_COORDINATES_DISABLED ? TEXT("Disabled") : TEXT("Enabled") );
+	CSV_METADATA(TEXT("LargeWorldCoordinates"), UE_LARGE_WORLD_COORDINATES_DISABLED ? TEXT("0") : TEXT("1"));
+
 	// Enable the live coding module if it's a developer build
 #if WITH_LIVE_CODING
 	FModuleManager::Get().LoadModule("LiveCoding");
