Fix AActor::GetStreamingBounds not checking for null components.

diff --git a/Engine/Source/Runtime/Engine/Private/Actor.cpp b/Engine/Source/Runtime/Engine/Private/Actor.cpp
--- a/Engine/Source/Runtime/Engine/Private/Actor.cpp
+++ b/Engine/Source/Runtime/Engine/Private/Actor.cpp
@@ -989,6 +989,8 @@ FBox AActor::GetStreamingBounds() const
 
 	auto IsComponentRelevant = [this](const UActorComponent* Component)
 	{
+		check(Component);
+
 		if (!Component->IsRegistered())
 		{
 			return false;
@@ -1017,13 +1019,13 @@ FBox AActor::GetStreamingBounds() const
 
 	if (!StreamingBounds.IsValid)
 	{
-		for (UActorComponent* Component : OwnedComponents)
+		ForEachComponent<UActorComponent>(false, [this, &StreamingBounds, &IsComponentRelevant](UActorComponent* Component)
 		{
 			if (IsComponentRelevant(Component))
 			{
 				StreamingBounds += Component->GetStreamingBounds();
 			}
-		}
+		});
 	}
 
 	return StreamingBounds;
