Removed HLOD Layer setup from Data Layers.

diff --git a/Engine/Source/Runtime/Engine/Private/WorldPartition/DataLayer/DataLayer.cpp b/Engine/Source/Runtime/Engine/Private/WorldPartition/DataLayer/DataLayer.cpp
--- a/Engine/Source/Runtime/Engine/Private/WorldPartition/DataLayer/DataLayer.cpp
+++ b/Engine/Source/Runtime/Engine/Private/WorldPartition/DataLayer/DataLayer.cpp
@@ -16,8 +16,6 @@ UDataLayer::UDataLayer(const FObjectInitializer& ObjectInitializer)
 , bIsInitiallyVisible(true)
 , bIsDynamicallyLoadedInEditor(true)
 , bIsLocked(false)
-, bGeneratesHLODs(true)
-, DefaultHLODLayer()
 #endif
 , DataLayerLabel(GetFName())
 , InitialState(EDataLayerState::Unloaded)
diff --git a/Engine/Source/Runtime/Engine/Private/WorldPartition/HLOD/HLODLayer.cpp b/Engine/Source/Runtime/Engine/Private/WorldPartition/HLOD/HLODLayer.cpp
--- a/Engine/Source/Runtime/Engine/Private/WorldPartition/HLOD/HLODLayer.cpp
+++ b/Engine/Source/Runtime/Engine/Private/WorldPartition/HLOD/HLODLayer.cpp
@@ -16,7 +16,6 @@
 
 #include "WorldPartition/WorldPartition.h"
 #include "WorldPartition/WorldPartitionActorDesc.h"
-#include "WorldPartition/DataLayer/DataLayerSubsystem.h"
 #endif
 
 DEFINE_LOG_CATEGORY_STATIC(LogHLODLayer, Log, All);
@@ -94,16 +93,6 @@ UHLODLayer* UHLODLayer::GetHLODLayer(const AActor* InActor)
 	bool bIsHLOD0 = !InActor->IsA<AWorldPartitionHLOD>();
 	if (bIsHLOD0) 
 	{
-		// Check if this actor is part of a DataLayer that has a default HLOD layer
-		for(const UDataLayer* DataLayer : InActor->GetDataLayerObjects())
-		{
-			UHLODLayer* HLODLayer = DataLayer ? DataLayer->GetDefaultHLODLayer() : nullptr;
-			if (HLODLayer)
-			{
-				return HLODLayer;
-			}
-		}
-
 		// Fallback to the world partition default HLOD layer
 		if (UWorldPartition* WorldPartition = InActor->GetWorld()->GetWorldPartition())
 		{
@@ -127,22 +116,6 @@ UHLODLayer* UHLODLayer::GetHLODLayer(const FWorldPartitionActorDesc& InActorDesc
 	bool bIsHLOD0 = !InActorDesc.GetActorClass()->IsChildOf<AWorldPartitionHLOD>();
 	if (bIsHLOD0)
 	{
-		if (UDataLayerSubsystem* DataLayerSubsystem = InWorldPartition->GetWorld()->GetSubsystem<UDataLayerSubsystem>())
-		{
-			if (const AWorldDataLayers* WorldDataLayers = InWorldPartition->GetWorld()->GetWorldDataLayers())
-			{
-				for (const FName& DataLayerName : InActorDesc.GetDataLayers())
-				{
-					const UDataLayer* DataLayer = WorldDataLayers->GetDataLayerFromName(DataLayerName);
-					UHLODLayer* HLODLayer = DataLayer ? DataLayer->GetDefaultHLODLayer() : nullptr;
-					if (HLODLayer)
-					{
-						return HLODLayer;
-					}
-				}
-			}
-		}
-
 		// Fallback to the world partition default HLOD layer
 		return InWorldPartition->DefaultHLODLayer;
 	}
diff --git a/Engine/Source/Runtime/Engine/Private/WorldPartition/RuntimeSpatialHash/RuntimeSpatialHashHLOD.cpp b/Engine/Source/Runtime/Engine/Private/WorldPartition/RuntimeSpatialHash/RuntimeSpatialHashHLOD.cpp
--- a/Engine/Source/Runtime/Engine/Private/WorldPartition/RuntimeSpatialHash/RuntimeSpatialHashHLOD.cpp
+++ b/Engine/Source/Runtime/Engine/Private/WorldPartition/RuntimeSpatialHash/RuntimeSpatialHashHLOD.cpp
@@ -16,7 +16,6 @@
 #include "Misc/ScopedSlowTask.h"
 #include "Modules/ModuleManager.h"
 
-#include "Algo/AnyOf.h"
 #include "Algo/ForEach.h"
 
 #include "EngineUtils.h"
@@ -113,10 +112,9 @@ static TArray<FGuid> GenerateHLODsForGrid(UWorldPartition* WorldPartition, const
 	auto ShouldGenerateHLODs = [&AlwaysLoadedCell](const FSquare2DGridHelper::FGridLevel::FGridCell& GridCell, const FSquare2DGridHelper::FGridLevel::FGridCellDataChunk& GridCellDataChunk)
 	{
 		const bool bIsCellAlwaysLoaded = &GridCell == &AlwaysLoadedCell;
-		const bool bShouldGenerateHLODForDataLayers = GridCellDataChunk.GetDataLayers().IsEmpty() ? true : Algo::AnyOf(GridCellDataChunk.GetDataLayers(), [](const UDataLayer* DataLayer) { return DataLayer->ShouldGenerateHLODs(); });
 		const bool bChunkHasActors = !GridCellDataChunk.GetActors().IsEmpty();
 
-		const bool bShouldGenerateHLODs = !bIsCellAlwaysLoaded && bShouldGenerateHLODForDataLayers && bChunkHasActors;
+		const bool bShouldGenerateHLODs = !bIsCellAlwaysLoaded && bChunkHasActors;
 		return bShouldGenerateHLODs;
 	};
 
diff --git a/Engine/Source/Runtime/Engine/Public/WorldPartition/DataLayer/DataLayer.h b/Engine/Source/Runtime/Engine/Public/WorldPartition/DataLayer/DataLayer.h
--- a/Engine/Source/Runtime/Engine/Public/WorldPartition/DataLayer/DataLayer.h
+++ b/Engine/Source/Runtime/Engine/Public/WorldPartition/DataLayer/DataLayer.h
@@ -37,8 +37,6 @@ public:
 	void SetIsLocked(bool bInIsLocked) { bIsLocked = bInIsLocked; }
 
 	bool IsDynamicallyLoadedInEditor() const { return bIsDynamicallyLoadedInEditor; }
-	bool ShouldGenerateHLODs() const { return IsDynamicallyLoaded() && bGeneratesHLODs; }
-	class UHLODLayer* GetDefaultHLODLayer() const { return ShouldGenerateHLODs() ? DefaultHLODLayer : nullptr; }
 
 	static FText GetDataLayerText(const UDataLayer* InDataLayer);
 
@@ -92,14 +90,6 @@ private:
 	/** Whether this data layer is locked, which means the user can't change actors assignation, remove or rename it */
 	UPROPERTY()
 	uint32 bIsLocked : 1;
-
-	/** Whether HLODs should be generated for actors in this layer */
-	UPROPERTY(Category = "Data Layer - HLODs", EditAnywhere, meta = (EditConditionHides, EditCondition = "bIsDynamicallyLoaded"))
-	uint32 bGeneratesHLODs : 1;
-
-	// Default HLOD layer
-	UPROPERTY(Category = "Data Layer - HLODs", EditAnywhere, meta = (EditConditionHides, EditCondition = "bIsDynamicallyLoaded && bGeneratesHLODs"))
-	TObjectPtr<class UHLODLayer> DefaultHLODLayer;
 #endif
 
 	/** The display name of the Data Layer */
