U5 - adding missing tag for FUpdateStreamingRenderAssetsTask

diff --git a/Engine/Source/Runtime/Engine/Private/Streaming/StreamingManagerTexture.cpp b/Engine/Source/Runtime/Engine/Private/Streaming/StreamingManagerTexture.cpp
--- a/Engine/Source/Runtime/Engine/Private/Streaming/StreamingManagerTexture.cpp
+++ b/Engine/Source/Runtime/Engine/Private/Streaming/StreamingManagerTexture.cpp
@@ -1569,6 +1569,7 @@ public:
 	}
 	void DoTask(ENamedThreads::Type CurrentThread, const FGraphEventRef& MyCompletionGraphEvent)
 	{
+		FOptionalTaskTagScope Scope(ETaskTag::EParallelGameThread);
 		Manager->UpdateStreamingRenderAssets(StageIdx, NumUpdateStages, bWaitForMipFading, true);
 		CompletionEvent->Trigger();
 	}
