Various import crash fixes: (possible segfault on non-standard-adhearing wave files, using higher channel asset as import template.)

diff --git a/Engine/Source/Editor/AudioEditor/Private/Factories/SoundFactory.cpp b/Engine/Source/Editor/AudioEditor/Private/Factories/SoundFactory.cpp
--- a/Engine/Source/Editor/AudioEditor/Private/Factories/SoundFactory.cpp
+++ b/Engine/Source/Editor/AudioEditor/Private/Factories/SoundFactory.cpp
@@ -344,21 +344,28 @@ UObject* USoundFactory::CreateObject
 		// otherwise create new sound and import raw data.
 		USoundWave* Sound = (bUseExistingSettings && ExistingSound) ? ExistingSound : NewObject<USoundWave>(InParent, Name, Flags, TemplateSoundWave.Get());
 
+		// If we're a multi-channel file, we're going to spoof the behavior of the SoundSurroundFactory
+		int32 ChannelCount = (int32)*WaveInfo.pChannels;
+		check(ChannelCount >0);
+
 		// These get wiped in PostInitProperties by defaults set from Audio Settings,
 		// so set back to template in this specialized case
 		if (TemplateSoundWave.IsValid())
 		{
 			Sound->SoundClassObject = TemplateSoundWave->SoundClassObject;
 			Sound->ConcurrencySet = TemplateSoundWave->ConcurrencySet;
+
+			// we do not want to inherit these values from the template, as the data may be incorrect
+			// rather we re-parse them from the incoming file.
+			Sound->NumChannels = 0;
+			Sound->ChannelOffsets.Reset();
+			Sound->ChannelSizes.Reset();
 		}
 
 		Sound->CompressionQuality = 100;
 		Sound->bStreaming = true;
 		Sound->bSeekableStreaming = true;
 
-		// If we're a multi-channel file, we're going to spoof the behavior of the SoundSurroundFactory
-		int32 ChannelCount = (int32)*WaveInfo.pChannels;
-		check(ChannelCount >0);
 
 		int32 SizeOfSample = (*WaveInfo.pBitsPerSample) / 8;
 
diff --git a/Engine/Source/Runtime/Engine/Private/Audio.cpp b/Engine/Source/Runtime/Engine/Private/Audio.cpp
--- a/Engine/Source/Runtime/Engine/Private/Audio.cpp
+++ b/Engine/Source/Runtime/Engine/Private/Audio.cpp
@@ -1242,6 +1242,89 @@ struct FExtendedFormatChunk
 #pragma pack(pop)
 #endif
 
+bool IsKnownChunkId(uint32 ChunkId)
+{
+	uint32 KnownIds[] =
+	{
+		UE_mmioFOURCC('f', 'm', 't', ' '),
+		UE_mmioFOURCC('d', 'a', 't', 'a'),
+		UE_mmioFOURCC('f', 'a', 'c', 't'),
+		UE_mmioFOURCC('c', 'u', 'e', ' '),
+		UE_mmioFOURCC('p', 'l', 's', 't'),
+		UE_mmioFOURCC('l', 'i', 's', 't'),
+		UE_mmioFOURCC('l', 'a', 'b', 'l'),
+		UE_mmioFOURCC('l', 't', 'x', 't'),
+		UE_mmioFOURCC('n', 'o', 't', 'e'),
+		UE_mmioFOURCC('s', 'm', 'p', 'l'),
+		UE_mmioFOURCC('i', 'n', 's', 't'),
+		UE_mmioFOURCC('a', 'c', 'i', 'd'),
+	};
+
+	return Algo::Find(KnownIds, ChunkId) != nullptr;
+}
+
+FRiffChunkOld* FindRiffChunk(FRiffChunkOld* RiffChunkStart, const uint8* RiffChunkEnd, uint32 ChunkId)
+{
+	// invalid RiffChunkStart
+	if (!ensure(RiffChunkStart))
+	{
+		return nullptr;
+	}
+
+	// invalid RiffChunkEnd
+	if (!ensure(RiffChunkEnd))
+	{
+		return nullptr;
+	}
+
+
+	FRiffChunkOld* RiffChunk = RiffChunkStart;
+
+
+	while ((uint8*)RiffChunk < RiffChunkEnd)
+	{
+		if (INTEL_ORDER32(RiffChunk->ChunkID) == ChunkId)
+		{
+			return RiffChunk;
+		}
+
+		// If the file has non standard chunks identifiers AND is not padding then we cannot be sure what sort of values 
+		// RiffChunk can contain as we will be reading garbage and potentially 'ChunkLen' could be zero which would lead 
+		// to an infinite loop, so check for that here.
+		if (RiffChunk->ChunkLen == 0)
+		{
+			return nullptr;
+		}
+
+		// The format specifies that chunks should be word aligned however some content creation tools seem to ignore this. 
+		// If ChunkLen is even then we can just advance to the next chunk. If it is odd then we need to try and determine if
+		// this file obeys the padding rule.
+		if (INTEL_ORDER32(RiffChunk->ChunkLen) % 2 == 0)
+		{
+			RiffChunk = (FRiffChunkOld*)((uint8*)RiffChunk + INTEL_ORDER32(RiffChunk->ChunkLen) + 8);
+		}
+		else
+		{
+			// First we find the next chunk if the file is not obeying the padding rule, note that we must check to make sure
+			// that the chunk is still within the bounds of our memory as at this point we can no longer trust the data.
+			// If the next chunk (NoPaddingChunk) has an identifier that we recognize then we assume that the current chunk
+			// (RiffChunk) was not padded. If we do not find a chunk id that we recognize we should assume that padding is
+			// required and try that instead.
+			FRiffChunkOld* NoPaddingChunk = (FRiffChunkOld*)((uint8*)RiffChunk + INTEL_ORDER32(RiffChunk->ChunkLen) + 8);
+			if ((uint8*)NoPaddingChunk + 8 < RiffChunkEnd && IsKnownChunkId(NoPaddingChunk->ChunkID))
+			{
+				RiffChunk = NoPaddingChunk;
+			}
+			else
+			{
+				RiffChunk = (FRiffChunkOld*)((uint8*)RiffChunk + FWaveModInfo::Pad16Bit(INTEL_ORDER32(RiffChunk->ChunkLen)) + 8);
+			}
+		}
+	}
+
+	return nullptr;
+}
+
 //
 //	Figure out the WAVE file layout.
 //
@@ -1282,16 +1365,12 @@ bool FWaveModInfo::ReadWaveInfo( const uint8* WaveData, int32 WaveDataSize, FStr
 #endif
 
 	FRiffChunkOld* RiffChunkStart = (FRiffChunkOld*)&WaveData[3 * 4];
-	FRiffChunkOld* RiffChunk = RiffChunkStart;
 	pMasterSize = &RiffHdr->ChunkLen;
 
 	// Look for the 'fmt ' chunk.
-	while ((((uint8*)RiffChunk + 8) < WaveDataEnd) && (INTEL_ORDER32(RiffChunk->ChunkID) != UE_mmioFOURCC('f','m','t',' ')))
-	{
-		RiffChunk = (FRiffChunkOld*)((uint8*)RiffChunk + Pad16Bit(INTEL_ORDER32(RiffChunk->ChunkLen)) + 8);
-	}
+	FRiffChunkOld* RiffChunk = FindRiffChunk(RiffChunkStart, WaveDataEnd, UE_mmioFOURCC('f', 'm', 't', ' '));
 
-	if (INTEL_ORDER32(RiffChunk->ChunkID) != UE_mmioFOURCC('f','m','t',' '))
+	if (RiffChunk == nullptr)
 	{
 		#if !PLATFORM_LITTLE_ENDIAN  // swap them back just in case.
 			if( !AlreadySwapped )
@@ -1301,7 +1380,11 @@ bool FWaveModInfo::ReadWaveInfo( const uint8* WaveData, int32 WaveDataSize, FStr
 				RiffHdr->wID = INTEL_ORDER32( RiffHdr->wID );
 			}
 		#endif
-		if (ErrorReason) *ErrorReason = TEXT("Invalid WAVE file.");
+			if (ErrorReason)
+			{
+				*ErrorReason = TEXT("Invalid WAVE file.");
+			}
+
 		return( false );
 	}
 
@@ -1379,16 +1462,11 @@ bool FWaveModInfo::ReadWaveInfo( const uint8* WaveData, int32 WaveDataSize, FStr
 		pFormatTag = reinterpret_cast<uint16*>(&FmtChunkEx->SubFormat.Data1);
 	}
 
-	// re-initialize the RiffChunk pointer
-	RiffChunk = RiffChunkStart;
 
 	// Look for the 'data' chunk.
-	while ((((uint8*)RiffChunk + 8) < WaveDataEnd) && (INTEL_ORDER32(RiffChunk->ChunkID) != UE_mmioFOURCC('d','a','t','a')))
-	{
-		RiffChunk = (FRiffChunkOld*)((uint8*)RiffChunk + Pad16Bit(INTEL_ORDER32(RiffChunk->ChunkLen)) + 8);
-	}
+	RiffChunk = FindRiffChunk(RiffChunkStart, WaveDataEnd, UE_mmioFOURCC('d', 'a', 't', 'a'));
 
-	if (INTEL_ORDER32(RiffChunk->ChunkID) != UE_mmioFOURCC('d','a','t','a'))
+	if (RiffChunk == nullptr)
 	{
 		#if !PLATFORM_LITTLE_ENDIAN  // swap them back just in case.
 			if (!AlreadySwapped)
@@ -1488,16 +1566,10 @@ bool FWaveModInfo::ReadWaveInfo( const uint8* WaveData, int32 WaveDataSize, FStr
 #endif
 
 	// Look for the cue chunks
-	RiffChunk = RiffChunkStart;
-
-	// Look for the 'cue ' chunk.
-	while ((((uint8*)RiffChunk + 8) < WaveDataEnd) && (INTEL_ORDER32(RiffChunk->ChunkID) != UE_mmioFOURCC('c', 'u', 'e', ' ')))
-	{
-		RiffChunk = (FRiffChunkOld*)((uint8*)RiffChunk + Pad16Bit(INTEL_ORDER32(RiffChunk->ChunkLen)) + 8);
-	}
+	RiffChunk = FindRiffChunk(RiffChunkStart, WaveDataEnd, UE_mmioFOURCC('c', 'u', 'e', ' '));
 
 	// Cue chunks are optional
-	if (reinterpret_cast<uint8*>(RiffChunk) + sizeof(RiffChunk->ChunkID) <= WaveDataEnd && INTEL_ORDER32(RiffChunk->ChunkID) == UE_mmioFOURCC('c', 'u', 'e', ' '))
+	if (RiffChunk != nullptr)
 	{
 		FRiffCueChunk* CueChunk = (FRiffCueChunk*)((uint8*)RiffChunk);
 
@@ -1514,17 +1586,12 @@ bool FWaveModInfo::ReadWaveInfo( const uint8* WaveData, int32 WaveDataSize, FStr
 		
 		}
 		
-		// Now look for labl chunk for labels for cues
-		RiffChunk = RiffChunkStart;
-
+		// Now look for label chunk for labels for cues
 		// Look for the 'list' chunk.
-		while ((((uint8*)RiffChunk + 8) < WaveDataEnd) && (INTEL_ORDER32(RiffChunk->ChunkID) != UE_mmioFOURCC('L', 'I', 'S', 'T')))
-		{
-			RiffChunk = (FRiffChunkOld*)((uint8*)RiffChunk + Pad16Bit(INTEL_ORDER32(RiffChunk->ChunkLen)) + 8);
-		}
+		RiffChunk = FindRiffChunk(RiffChunkStart, WaveDataEnd, UE_mmioFOURCC('L', 'I', 'S', 'T'));
 
 		// Label chunks are also optional
-		if (reinterpret_cast<uint8*>(RiffChunk) + sizeof(RiffChunk->ChunkID) <= WaveDataEnd && INTEL_ORDER32(RiffChunk->ChunkID) == UE_mmioFOURCC('L', 'I', 'S', 'T'))
+		if (RiffChunk != nullptr)
 		{
 			FRiffListChunk* ListChunk = (FRiffListChunk*)((uint8*)RiffChunk);
 			
@@ -1717,3 +1784,4 @@ void SerializeWaveFile(TArray<uint8>& OutWaveFileData, const uint8* InPCMData, c
 	FMemory::Memcpy(&OutWaveFileData[WaveDataByteIndex], InPCMData, NumBytes);
 }
 
+
