Splitting out Sound Concurrency logs into own category - LogAudioConcurrency

diff --git a/Engine/Source/Runtime/Engine/Classes/Sound/SoundConcurrency.h b/Engine/Source/Runtime/Engine/Classes/Sound/SoundConcurrency.h
--- a/Engine/Source/Runtime/Engine/Classes/Sound/SoundConcurrency.h
+++ b/Engine/Source/Runtime/Engine/Classes/Sound/SoundConcurrency.h
@@ -24,6 +24,7 @@ using FSoundOwnerObjectID = uint32;
 /** Sound instance (USoundBase) object ID. */
 using FSoundObjectID = uint32;
 
+ENGINE_API DECLARE_LOG_CATEGORY_EXTERN(LogAudioConcurrency, Verbose, All);
 
 UENUM()
 namespace EMaxConcurrentResolutionRule
diff --git a/Engine/Source/Runtime/Engine/Private/AudioDevice.cpp b/Engine/Source/Runtime/Engine/Private/AudioDevice.cpp
--- a/Engine/Source/Runtime/Engine/Private/AudioDevice.cpp
+++ b/Engine/Source/Runtime/Engine/Private/AudioDevice.cpp
@@ -4705,7 +4705,7 @@ void FAudioDevice::AddNewActiveSoundInternal(const FActiveSound& NewActiveSound,
 			FAudioVirtualLoop VirtualLoop;
 			if (FAudioVirtualLoop::Virtualize(NewActiveSound, *this, bDoRangeCheck, VirtualLoop))
 			{
-				UE_LOG(LogAudio, Verbose, TEXT("New ActiveSound %s Virtualizing: Failed to pass concurrency"), *Sound->GetName());
+				UE_LOG(LogAudioConcurrency, Verbose, TEXT("New ActiveSound %s Virtualizing: Failed to pass concurrency"), *Sound->GetName());
 				AddVirtualLoop(VirtualLoop);
 			}
 			else
diff --git a/Engine/Source/Runtime/Engine/Private/SoundConcurrency.cpp b/Engine/Source/Runtime/Engine/Private/SoundConcurrency.cpp
--- a/Engine/Source/Runtime/Engine/Private/SoundConcurrency.cpp
+++ b/Engine/Source/Runtime/Engine/Private/SoundConcurrency.cpp
@@ -11,6 +11,7 @@
 #include "HAL/IConsoleManager.h"
 #include "Sound/SoundBase.h"
 
+DEFINE_LOG_CATEGORY(LogAudioConcurrency);
 
 static float ConcurrencyMinVolumeScaleCVar = 1.e-3f;
 FAutoConsoleVariableRef CVarConcurrencyMinVolumeScale(
@@ -34,7 +35,7 @@ namespace SoundConcurrency
 		{
 			if (const USoundBase* Sound = ActiveSound.GetSound())
 			{
-				UE_LOG(LogAudio, Verbose,
+				UE_LOG(LogAudioConcurrency, Verbose,
 					TEXT("Sound '%s' concurrency generation '%i' target volume update: %.3f to %.3f."),
 					*Sound->GetName(),
 					SoundData.Generation,
@@ -179,7 +180,7 @@ void FConcurrencyGroup::AddActiveSound(FActiveSound& ActiveSound)
 
 	if (ActiveSound.ConcurrencyGroupData.Contains(GroupID))
 	{
-		UE_LOG(LogAudio, Fatal, TEXT("Attempting to add active sound '%s' to concurrency group multiple times."), *ActiveSound.GetOwnerName());
+		UE_LOG(LogAudioConcurrency, Fatal, TEXT("Attempting to add active sound '%s' to concurrency group multiple times."), *ActiveSound.GetOwnerName());
 		return;
 	}
 
@@ -541,7 +542,7 @@ bool FSoundConcurrencyManager::IsRateLimited(const FConcurrencyHandle& InHandle)
 		if (*LastTimePlayed > 0.0f && InHandle.Settings.RetriggerTime > 0.0f && DeltaTime < InHandle.Settings.RetriggerTime)
 		{
 			// Retrieve the current game time
-			UE_LOG(LogAudio, VeryVerbose, TEXT("Rejected Sound for Group ID (%d) with DeltaTime: %.2f"), InHandle.ObjectID, DeltaTime);
+			UE_LOG(LogAudioConcurrency, VeryVerbose, TEXT("Rejected Sound for Group ID (%d) with DeltaTime: %.2f"), InHandle.ObjectID, DeltaTime);
 			return true;
 		}
 
@@ -557,7 +558,7 @@ FConcurrencyGroup* FSoundConcurrencyManager::CanPlaySound(const FActiveSound& Ne
 	FConcurrencyGroup* ConcurrencyGroup = ConcurrencyGroups.FindRef(GroupID);
 	if (!ConcurrencyGroup)
 	{
-		UE_LOG(LogAudio, Warning, TEXT("Attempting to add active sound '%s' (owner '%s') to invalid concurrency group."),
+		UE_LOG(LogAudioConcurrency, Warning, TEXT("Attempting to add active sound '%s' (owner '%s') to invalid concurrency group."),
 			NewActiveSound.GetSound() ? *NewActiveSound.GetSound()->GetFullName() : TEXT("Unset"),
 			*NewActiveSound.GetOwnerName());
 		return nullptr;
@@ -868,6 +869,11 @@ FActiveSound* FSoundConcurrencyManager::CreateAndEvictActiveSounds(const FActive
 			continue;
 		}
 
+		if (SoundToEvict->Sound)
+		{
+			UE_LOG(LogAudioConcurrency, VeryVerbose, TEXT("Evicting Sound %s due to concurrency"), *(SoundToEvict->Sound->GetName()));
+		}
+
 		StopDueToVoiceStealing(*SoundToEvict);
 	}
 
@@ -885,7 +891,7 @@ void FSoundConcurrencyManager::RemoveActiveSound(FActiveSound& ActiveSound)
 		FConcurrencyGroup* ConcurrencyGroup = ConcurrencyGroups.FindRef(ConcurrencyGroupID);
 		if (!ConcurrencyGroup)
 		{
-			UE_LOG(LogAudio, Error, TEXT("Attempting to remove stopped sound '%s' from inactive concurrency group."),
+			UE_LOG(LogAudioConcurrency, Error, TEXT("Attempting to remove stopped sound '%s' from inactive concurrency group."),
 				ActiveSound.GetSound() ? *ActiveSound.GetSound()->GetName(): TEXT("Unset"));
 			continue;
 		}
@@ -969,7 +975,7 @@ void FSoundConcurrencyManager::StopDueToVoiceStealing(FActiveSound& ActiveSound)
 		ActiveSound.ClearAudioComponent();
 		if (USoundBase* Sound = ActiveSound.GetSound())
 		{
-			UE_LOG(LogAudio, Verbose, TEXT("Playing ActiveSound %s Virtualizing: Sound's voice stollen due to concurrency group maximum met."), *Sound->GetName());
+			UE_LOG(LogAudioConcurrency, Verbose, TEXT("Playing ActiveSound %s Virtualizing: Sound's voice stollen due to concurrency group maximum met."), *Sound->GetName());
 		}
 		ActiveSound.AudioDevice->AddVirtualLoop(VirtualLoop);
 	}
