Renamed Strata UMaterialExpressionStrataMultiply node to UMaterialExpressionStrataWeight.

diff --git a/Engine/Shaders/Private/Strata/Strata.ush b/Engine/Shaders/Private/Strata/Strata.ush
--- a/Engine/Shaders/Private/Strata/Strata.ush
+++ b/Engine/Shaders/Private/Strata/Strata.ush
@@ -853,7 +853,7 @@ FStrataData StrataAdd(FStrataData A, FStrataData B)
 	return StrataData;
 }
 
-FStrataData StrataMultiply(FStrataData A, float Weight)
+FStrataData StrataWeight(FStrataData A, float Weight)
 {
 	const float SafeWeight = saturate(Weight);
 	FStrataData StrataData = A;
@@ -877,8 +877,8 @@ FStrataData StrataHorizontalMixing(FStrataData Foreground, FStrataData Backgroun
 	const float ForegroundMixFactor = saturate(Mix);
 	const float BackgroundMixFactor = 1.0 - ForegroundMixFactor;
 
-	FStrataData A = StrataMultiply(Foreground, ForegroundMixFactor);
-	FStrataData B = StrataMultiply(Background, BackgroundMixFactor);
+	FStrataData A = StrataWeight(Foreground, ForegroundMixFactor);
+	FStrataData B = StrataWeight(Background, BackgroundMixFactor);
 	return StrataAdd(A, B);
 }
 
diff --git a/Engine/Source/Runtime/Engine/Classes/Materials/MaterialExpressionStrata.h b/Engine/Source/Runtime/Engine/Classes/Materials/MaterialExpressionStrata.h
--- a/Engine/Source/Runtime/Engine/Classes/Materials/MaterialExpressionStrata.h
+++ b/Engine/Source/Runtime/Engine/Classes/Materials/MaterialExpressionStrata.h
@@ -504,7 +504,7 @@ class UMaterialExpressionStrataAdd : public UMaterialExpressionStrataBSDF
 };
 
 UCLASS(MinimalAPI, collapsecategories, hidecategories = Object, DisplayName = "Strata BSDF Weight")
-class UMaterialExpressionStrataMultiply : public UMaterialExpressionStrataBSDF
+class UMaterialExpressionStrataWeight : public UMaterialExpressionStrataBSDF
 {
 	GENERATED_UCLASS_BODY()
 
diff --git a/Engine/Source/Runtime/Engine/Private/Materials/HLSLMaterialTranslator.cpp b/Engine/Source/Runtime/Engine/Private/Materials/HLSLMaterialTranslator.cpp
--- a/Engine/Source/Runtime/Engine/Private/Materials/HLSLMaterialTranslator.cpp
+++ b/Engine/Source/Runtime/Engine/Private/Materials/HLSLMaterialTranslator.cpp
@@ -9177,14 +9177,14 @@ int32 FHLSLMaterialTranslator::StrataAddParameterBlending(int32 A, int32 B, cons
 	);
 }
 
-int32 FHLSLMaterialTranslator::StrataMultiply(int32 A, int32 Weight)
+int32 FHLSLMaterialTranslator::StrataWeight(int32 A, int32 Weight)
 {
 	if (A == INDEX_NONE || Weight == INDEX_NONE)
 	{
 		return INDEX_NONE;
 	}
 	return AddCodeChunk(
-		MCT_Strata, TEXT("StrataMultiply(%s, %s)"),
+		MCT_Strata, TEXT("StrataWeight(%s, %s)"),
 		*GetParameterCode(A),
 		*GetParameterCode(Weight)
 	);
diff --git a/Engine/Source/Runtime/Engine/Private/Materials/HLSLMaterialTranslator.h b/Engine/Source/Runtime/Engine/Private/Materials/HLSLMaterialTranslator.h
--- a/Engine/Source/Runtime/Engine/Private/Materials/HLSLMaterialTranslator.h
+++ b/Engine/Source/Runtime/Engine/Private/Materials/HLSLMaterialTranslator.h
@@ -1004,7 +1004,7 @@ protected:
 	virtual int32 StrataVerticalLayeringParameterBlending(int32 Top, int32 Base, const FString& SharedNormalIndexMacro, int32 TopBSDFNormalCodeChunk) override;
 	virtual int32 StrataAdd(int32 A, int32 B) override;
 	virtual int32 StrataAddParameterBlending(int32 A, int32 B, const FString& SharedNormalIndexMacro) override;
-	virtual int32 StrataMultiply(int32 A, int32 Weight) override;
+	virtual int32 StrataWeight(int32 A, int32 Weight) override;
 	virtual int32 StrataTransmittanceToMFP(int32 TransmittanceColor, int32 DesiredThickness, int32 OutputIndex) override;
 
 	virtual void StrataCompilationInfoRegisterCodeChunk(int32 CodeChunk, FStrataMaterialCompilationInfo& StrataMaterialCompilationInfo) override;
diff --git a/Engine/Source/Runtime/Engine/Private/Materials/Material.cpp b/Engine/Source/Runtime/Engine/Private/Materials/Material.cpp
--- a/Engine/Source/Runtime/Engine/Private/Materials/Material.cpp
+++ b/Engine/Source/Runtime/Engine/Private/Materials/Material.cpp
@@ -3874,7 +3874,7 @@ void UMaterial::ConvertMaterialToStrataMaterial()
 			MoveConnectionTo(EmissiveColor, TopSlabBSDF, 13);							// Emissive
 
 			// Now weight the top base material by opacity
-			UMaterialExpressionStrataMultiply* TopMaterialWithCoverage = NewObject<UMaterialExpressionStrataMultiply>(this);
+			UMaterialExpressionStrataWeight* TopMaterialWithCoverage = NewObject<UMaterialExpressionStrataWeight>(this);
 			TopMaterialWithCoverage->GetInput(0)->Connect(0, TopSlabBSDF);				// TopSlabBSDF -> A
 			MoveConnectionTo(Opacity, TopMaterialWithCoverage, 1);						// Opacity -> Weight
 
@@ -3973,7 +3973,7 @@ void UMaterial::ConvertMaterialToStrataMaterial()
 			}
 
 			// Now weight the top base material by ClearCoat
-			UMaterialExpressionStrataMultiply* TopSlabBSDFWithCoverage = NewObject<UMaterialExpressionStrataMultiply>(this);
+			UMaterialExpressionStrataWeight* TopSlabBSDFWithCoverage = NewObject<UMaterialExpressionStrataWeight>(this);
 			TopSlabBSDFWithCoverage->GetInput(0)->Connect(0, TopSlabBSDF);				// TopSlabBSDF -> A
 			MoveConnectionTo(ClearCoat, TopSlabBSDFWithCoverage, 1);					// ClearCoat -> Weight
 
diff --git a/Engine/Source/Runtime/Engine/Private/Materials/MaterialExpressions.cpp b/Engine/Source/Runtime/Engine/Private/Materials/MaterialExpressions.cpp
--- a/Engine/Source/Runtime/Engine/Private/Materials/MaterialExpressions.cpp
+++ b/Engine/Source/Runtime/Engine/Private/Materials/MaterialExpressions.cpp
@@ -21120,7 +21120,7 @@ void UMaterialExpressionStrataAdd::GatherStrataMaterialInfo(FStrataMaterialInfo&
 
 
 
-UMaterialExpressionStrataMultiply::UMaterialExpressionStrataMultiply(const FObjectInitializer& ObjectInitializer)
+UMaterialExpressionStrataWeight::UMaterialExpressionStrataWeight(const FObjectInitializer& ObjectInitializer)
 	: Super(ObjectInitializer)
 {
 	struct FConstructorStatics
@@ -21135,7 +21135,7 @@ UMaterialExpressionStrataMultiply::UMaterialExpressionStrataMultiply(const FObje
 }
 
 #if WITH_EDITOR
-int32 UMaterialExpressionStrataMultiply::Compile(class FMaterialCompiler* Compiler, int32 OutputIndex)
+int32 UMaterialExpressionStrataWeight::Compile(class FMaterialCompiler* Compiler, int32 OutputIndex)
 {
 	if (!A.GetTracedInput().Expression)
 	{
@@ -21145,29 +21145,29 @@ int32 UMaterialExpressionStrataMultiply::Compile(class FMaterialCompiler* Compil
 	int32 ACodeChunk = A.Compile(Compiler);
 	int32 WeightCodeChunk = Weight.GetTracedInput().Expression ? Weight.Compile(Compiler) : Compiler->Constant(1.0f);
 
-	int32 OutputCodeChunk = Compiler->StrataMultiply(ACodeChunk, WeightCodeChunk);
+	int32 OutputCodeChunk = Compiler->StrataWeight(ACodeChunk, WeightCodeChunk);
 
 	if (!Compiler->StrataCompilationInfoContainsCodeChunk(ACodeChunk))
 	{
 		return Compiler->Errorf(TEXT("Could not find ACodeChunk to multiply"));
 	}
-	FStrataMaterialCompilationInfo StrataInfo = StrataCompilationInfoMultiply(Compiler, Compiler->GetStrataCompilationInfo(ACodeChunk));
+	FStrataMaterialCompilationInfo StrataInfo = StrataCompilationInfoWeight(Compiler, Compiler->GetStrataCompilationInfo(ACodeChunk));
 	Compiler->StrataCompilationInfoRegisterCodeChunk(OutputCodeChunk, StrataInfo);
 
 	return OutputCodeChunk;
 }
 
-void UMaterialExpressionStrataMultiply::GetCaption(TArray<FString>& OutCaptions) const
+void UMaterialExpressionStrataWeight::GetCaption(TArray<FString>& OutCaptions) const
 {
 	OutCaptions.Add(TEXT("Strata BSDF Weight"));
 }
 
-uint32 UMaterialExpressionStrataMultiply::GetOutputType(int32 OutputIndex)
+uint32 UMaterialExpressionStrataWeight::GetOutputType(int32 OutputIndex)
 {
 	return MCT_Strata;
 }
 
-uint32 UMaterialExpressionStrataMultiply::GetInputType(int32 InputIndex)
+uint32 UMaterialExpressionStrataWeight::GetInputType(int32 InputIndex)
 {
 	if (InputIndex == 0)
 	{
@@ -21176,12 +21176,12 @@ uint32 UMaterialExpressionStrataMultiply::GetInputType(int32 InputIndex)
 	return MCT_Float1;
 }
 
-bool UMaterialExpressionStrataMultiply::IsResultStrataMaterial(int32 OutputIndex)
+bool UMaterialExpressionStrataWeight::IsResultStrataMaterial(int32 OutputIndex)
 {
 	return true;
 }
 
-void UMaterialExpressionStrataMultiply::GatherStrataMaterialInfo(FStrataMaterialInfo& StrataMaterialInfo, int32 OutputIndex)
+void UMaterialExpressionStrataWeight::GatherStrataMaterialInfo(FStrataMaterialInfo& StrataMaterialInfo, int32 OutputIndex)
 {
 	if (A.GetTracedInput().Expression)
 	{
diff --git a/Engine/Source/Runtime/Engine/Private/Materials/StrataMaterial.cpp b/Engine/Source/Runtime/Engine/Private/Materials/StrataMaterial.cpp
--- a/Engine/Source/Runtime/Engine/Private/Materials/StrataMaterial.cpp
+++ b/Engine/Source/Runtime/Engine/Private/Materials/StrataMaterial.cpp
@@ -90,7 +90,7 @@ void  StrataCompilationInfoCreateNullBSDF(FMaterialCompiler* Compiler, int32 Cod
 }
 
 
-FStrataMaterialCompilationInfo StrataCompilationInfoMultiply(FMaterialCompiler* Compiler, const FStrataMaterialCompilationInfo& A)
+FStrataMaterialCompilationInfo StrataCompilationInfoWeight(FMaterialCompiler* Compiler, const FStrataMaterialCompilationInfo& A)
 {
 	FStrataMaterialCompilationInfo StrataInfo = A;
 	return StrataInfo;
diff --git a/Engine/Source/Runtime/Engine/Private/Materials/StrataMaterial.h b/Engine/Source/Runtime/Engine/Private/Materials/StrataMaterial.h
--- a/Engine/Source/Runtime/Engine/Private/Materials/StrataMaterial.h
+++ b/Engine/Source/Runtime/Engine/Private/Materials/StrataMaterial.h
@@ -57,7 +57,7 @@ void  StrataCompilationInfoCreateSingleBSDFMaterial(FMaterialCompiler* Compiler,
 	const FStrataRegisteredSharedNormal& RegisteredSharedNormal,
 	uint8 BSDFType, bool bHasSSS = false, bool bHasDMFPPluggedIn = false, bool bHasEdgeColor = false, bool bHasThinFilm = false, bool bHasFuzz = false);
 
-FStrataMaterialCompilationInfo StrataCompilationInfoMultiply(FMaterialCompiler* Compiler, const FStrataMaterialCompilationInfo& A);
+FStrataMaterialCompilationInfo StrataCompilationInfoWeight(FMaterialCompiler* Compiler, const FStrataMaterialCompilationInfo& A);
 FStrataMaterialCompilationInfo StrataCompilationInfoAdd(FMaterialCompiler* Compiler, const FStrataMaterialCompilationInfo& A, const FStrataMaterialCompilationInfo& B);
 FStrataMaterialCompilationInfo StrataCompilationInfoAddParamBlend(FMaterialCompiler* Compiler, const FStrataMaterialCompilationInfo& A, const FStrataMaterialCompilationInfo& B, const FStrataRegisteredSharedNormal& RegisteredSharedNormal);
 FStrataMaterialCompilationInfo StrataCompilationInfoHorizontalMixing(FMaterialCompiler* Compiler, const FStrataMaterialCompilationInfo& A, const FStrataMaterialCompilationInfo& B);
diff --git a/Engine/Source/Runtime/Engine/Public/MaterialCompiler.h b/Engine/Source/Runtime/Engine/Public/MaterialCompiler.h
--- a/Engine/Source/Runtime/Engine/Public/MaterialCompiler.h
+++ b/Engine/Source/Runtime/Engine/Public/MaterialCompiler.h
@@ -445,7 +445,7 @@ public:
 	virtual int32 StrataVerticalLayeringParameterBlending(int32 Top, int32 Base, const FString& SharedNormalIndexMacro, int32 TopBSDFNormalCodeChunk) = 0;
 	virtual int32 StrataAdd(int32 A, int32 B) = 0;
 	virtual int32 StrataAddParameterBlending(int32 A, int32 B, const FString& SharedNormalIndexMacro) = 0;
-	virtual int32 StrataMultiply(int32 A, int32 Weight) = 0;
+	virtual int32 StrataWeight(int32 A, int32 Weight) = 0;
 	virtual int32 StrataTransmittanceToMFP(int32 TransmittanceColor, int32 DesiredThickness, int32 OutputIndex) = 0;
 
 	virtual void StrataCompilationInfoRegisterCodeChunk(int32 CodeChunk, FStrataMaterialCompilationInfo& StrataMaterialCompilationInfo) = 0;
@@ -1006,9 +1006,9 @@ public:
 		return Compiler->StrataAddParameterBlending(A, B, SharedNormalIndexMacro);
 	}
 
-	virtual int32 StrataMultiply(int32 A, int32 Weight) override
+	virtual int32 StrataWeight(int32 A, int32 Weight) override
 	{
-		return Compiler->StrataMultiply(A, Weight);
+		return Compiler->StrataWeight(A, Weight);
 	}
 
 	virtual int32 StrataTransmittanceToMFP(int32 TransmittanceColor, int32 DesiredThickness, int32 OutputIndex) override
