Fixed PIE extreme long startup time
Reverting 16267724 + 16267710

diff --git a/Engine/Source/Runtime/Engine/Private/WorldPartition/HLOD/HLODActor.cpp b/Engine/Source/Runtime/Engine/Private/WorldPartition/HLOD/HLODActor.cpp
--- a/Engine/Source/Runtime/Engine/Private/WorldPartition/HLOD/HLODActor.cpp
+++ b/Engine/Source/Runtime/Engine/Private/WorldPartition/HLOD/HLODActor.cpp
@@ -61,9 +61,7 @@ void AWorldPartitionHLOD::Serialize(FArchive& Ar)
 #if WITH_EDITOR
 	if(Ar.IsLoading() && Ar.CustomVer(FUE5MainStreamObjectVersion::GUID) < FUE5MainStreamObjectVersion::WorldPartitionStreamingCellsNamingShortened)
 	{
-		auto FixupCellName = [](const FString& CellName) { return CellName.Replace(TEXT("WPRT_"), TEXT(""), ESearchCase::CaseSensitive).Replace(TEXT("Cell_"), TEXT(""), ESearchCase::CaseSensitive); };
-		CellName_DEPRECATED = FName(FixupCellName(CellName_DEPRECATED.ToString()));
-		SourceCell = FixupCellName(SourceCell.ToString());
+		SourceCell = SourceCell.ToString().Replace(TEXT("WPRT_"), TEXT(""), ESearchCase::CaseSensitive).Replace(TEXT("Cell_"), TEXT(""), ESearchCase::CaseSensitive);
 	}
 #endif
 }
@@ -71,28 +69,6 @@ void AWorldPartitionHLOD::Serialize(FArchive& Ar)
 void AWorldPartitionHLOD::RerunConstructionScripts()
 {}
 
-#if WITH_EDITORONLY_DATA
-void AWorldPartitionHLOD::PostLoad()
-{
-	Super::PostLoad();
-
-	// Convert CellName to a soft object path referencing the RuntimeCell object.
-	if (!CellName_DEPRECATED.IsNone() && SourceCell.IsNull())
-	{
-		GetWorld()->ConditionalPostLoad();
-
-		UWorldPartition* WorldPartition = GetWorld()->GetWorldPartition();
-		if (ensure(WorldPartition))
-		{
-			FSoftObjectPath WorldPartitionPath(WorldPartition);
-			FString CellObjectPath = FString::Printf(TEXT("%s.%s"), *WorldPartitionPath.ToString(), *CellName_DEPRECATED.ToString());
-			TSoftObjectPtr<UWorldPartitionRuntimeCell> RuntimeCell(CellObjectPath);
-			SourceCell = RuntimeCell;
-		}
-	}
-}
-#endif
-
 #if WITH_EDITOR
 
 bool AWorldPartitionHLOD::IsHiddenEd() const
diff --git a/Engine/Source/Runtime/Engine/Private/WorldPartition/WorldPartitionLevelHelper.cpp b/Engine/Source/Runtime/Engine/Private/WorldPartition/WorldPartitionLevelHelper.cpp
--- a/Engine/Source/Runtime/Engine/Private/WorldPartition/WorldPartitionLevelHelper.cpp
+++ b/Engine/Source/Runtime/Engine/Private/WorldPartition/WorldPartitionLevelHelper.cpp
@@ -139,7 +139,7 @@ void FWorldPartitionLevelHelper::RemapLevelSoftObjectPaths(ULevel* InLevel, UWor
 /**
  * Creates an empty Level used in World Partition
  */
-ULevel* FWorldPartitionLevelHelper::CreateEmptyLevelForRuntimeCell(UWorld* InWorld, const FString& InWorldAssetName, UPackage* InPackage)
+ULevel* FWorldPartitionLevelHelper::CreateEmptyLevelForRuntimeCell(const UWorld* InWorld, const FString& InWorldAssetName, UPackage* InPackage)
 {
 	// Create or use given package
 	UPackage* CellPackage = nullptr;
@@ -180,10 +180,6 @@ ULevel* FWorldPartitionLevelHelper::CreateEmptyLevelForRuntimeCell(UWorld* InWor
 	check(NewLevel->OwningWorld == NewWorld);
 	check(NewLevel->Model);
 	check(!NewLevel->bIsVisible);
-
-	// Update the persistent level's owning world to the correct world.
-	NewLevel->OwningWorld = InWorld;
-
 	return NewLevel;
 }
 
diff --git a/Engine/Source/Runtime/Engine/Private/WorldPartition/WorldPartitionLevelStreamingDynamic.cpp b/Engine/Source/Runtime/Engine/Private/WorldPartition/WorldPartitionLevelStreamingDynamic.cpp
--- a/Engine/Source/Runtime/Engine/Private/WorldPartition/WorldPartitionLevelStreamingDynamic.cpp
+++ b/Engine/Source/Runtime/Engine/Private/WorldPartition/WorldPartitionLevelStreamingDynamic.cpp
@@ -82,7 +82,7 @@ void UWorldPartitionLevelStreamingDynamic::CreateRuntimeLevel()
 {
 	check(PendingUnloadLevel == nullptr);
 	check(RuntimeLevel == nullptr);
-	UWorld* PlayWorld = GetWorld();
+	const UWorld* PlayWorld = GetWorld();
 	check(PlayWorld && PlayWorld->IsGameWorld());
 
 	// Create streaming cell Level package
diff --git a/Engine/Source/Runtime/Engine/Public/WorldPartition/HLOD/HLODActor.h b/Engine/Source/Runtime/Engine/Public/WorldPartition/HLOD/HLODActor.h
--- a/Engine/Source/Runtime/Engine/Public/WorldPartition/HLOD/HLODActor.h
+++ b/Engine/Source/Runtime/Engine/Public/WorldPartition/HLOD/HLODActor.h
@@ -61,9 +61,6 @@ protected:
 	//~ Begin UObject Interface.
 	virtual void Serialize(FArchive& Ar) override;
 	virtual void RerunConstructionScripts() override;
-#if WITH_EDITORONLY_DATA
-	virtual void PostLoad() override;
-#endif
 	//~ End UObject Interface.
 
 	//~ Begin AActor Interface.
@@ -104,10 +101,6 @@ private:
 
 	UPROPERTY()
 	uint32 HLODHash;
-
-	// Deprecated. Stored in SourceCell instead.
-	UPROPERTY()
-	FName CellName_DEPRECATED;
 #endif
 
 	UPROPERTY()
diff --git a/Engine/Source/Runtime/Engine/Public/WorldPartition/WorldPartitionLevelHelper.h b/Engine/Source/Runtime/Engine/Public/WorldPartition/WorldPartitionLevelHelper.h
--- a/Engine/Source/Runtime/Engine/Public/WorldPartition/WorldPartitionLevelHelper.h
+++ b/Engine/Source/Runtime/Engine/Public/WorldPartition/WorldPartitionLevelHelper.h
@@ -20,7 +20,7 @@ class UWorldPartition;
 class FWorldPartitionLevelHelper
 {
 public:
-	static ULevel* CreateEmptyLevelForRuntimeCell(UWorld* InWorld, const FString& InWorldAssetName, UPackage* DestPackage = nullptr);
+	static ULevel* CreateEmptyLevelForRuntimeCell(const UWorld* InWorld, const FString& InWorldAssetName, UPackage* DestPackage = nullptr);
 	static void MoveExternalActorsToLevel(const TArray<FWorldPartitionRuntimeCellObjectMapping>& InChildPackages, ULevel* InLevel);
 	static void RemapLevelSoftObjectPaths(ULevel* InLevel, UWorldPartition* InWorldPartition);
 	
