compile fix - use getter instead of private member access

caught in other branch but missed here, oops!

diff --git a/Engine/Source/Runtime/Engine/Private/SoundConcurrency.cpp b/Engine/Source/Runtime/Engine/Private/SoundConcurrency.cpp
--- a/Engine/Source/Runtime/Engine/Private/SoundConcurrency.cpp
+++ b/Engine/Source/Runtime/Engine/Private/SoundConcurrency.cpp
@@ -869,9 +869,9 @@ FActiveSound* FSoundConcurrencyManager::CreateAndEvictActiveSounds(const FActive
 			continue;
 		}
 
-		if (SoundToEvict->Sound)
+		if (SoundToEvict->GetSound())
 		{
-			UE_LOG(LogAudioConcurrency, VeryVerbose, TEXT("Evicting Sound %s due to concurrency"), *(SoundToEvict->Sound->GetName()));
+			UE_LOG(LogAudioConcurrency, VeryVerbose, TEXT("Evicting Sound %s due to concurrency"), *(SoundToEvict->GetSound()->GetName()));
 		}
 
 		StopDueToVoiceStealing(*SoundToEvict);
