Fixed HasEdgeColor and relink node when toggling metalness to fix bugs and connection shifts.

diff --git a/Engine/Source/Runtime/Engine/Private/Materials/MaterialExpressions.cpp b/Engine/Source/Runtime/Engine/Private/Materials/MaterialExpressions.cpp
--- a/Engine/Source/Runtime/Engine/Private/Materials/MaterialExpressions.cpp
+++ b/Engine/Source/Runtime/Engine/Private/Materials/MaterialExpressions.cpp
@@ -19922,7 +19922,7 @@ int32 UMaterialExpressionStrataSlabBSDF::Compile(class FMaterialCompiler* Compil
 	int32 TangentCodeChunk = bHasAnisotropy ? CompileWithDefaultTangentWS(Compiler, Tangent) : INDEX_NONE;
 	const FStrataRegisteredSharedNormal NewRegisteredSharedNormal = StrataCompilationInfoCreateSharedNormal(Compiler, NormalCodeChunk, TangentCodeChunk);
 
-	const bool bHasEdgeColor = HasEdgeColor();
+	const bool bHasEdgeColor = HasEdgeColor(); // This accounts for EdgeColor and also F90 when the non metalness worfklow is selected.
 	const bool bHasThinFilm = HasThinFilm();
 	const bool bHasFuzz = HasFuzz();
 	const bool bHasDMFPPluggedIn = HasDMFPPluggedIn();
@@ -20009,7 +20009,35 @@ const TArray<FExpressionInput*> UMaterialExpressionStrataSlabBSDF::GetInputs()
 
 void UMaterialExpressionStrataSlabBSDF::PostEditChangeProperty(FPropertyChangedEvent& PropertyChangedEvent)
 {
-	if(GraphNode)
+	if (PropertyChangedEvent.Property && (PropertyChangedEvent.Property->GetName() == TEXT("bUseMetalness")))
+	{
+		// Relink to make the node looks somewhat ok when toggling Metalness
+		if (bUseMetalness)
+		{
+			BaseColor.Connect(DiffuseAlbedo.OutputIndex, DiffuseAlbedo.Expression);
+			Specular.Connect(F0.OutputIndex, F0.Expression);
+			EdgeColor.Connect(F90.OutputIndex, F90.Expression);
+
+			DiffuseAlbedo.Expression = nullptr;
+			F0.Expression = nullptr;
+			F90.Expression = nullptr;
+		}
+		else
+		{
+			DiffuseAlbedo.Connect(BaseColor.OutputIndex, BaseColor.Expression);
+			F0.Connect(Specular.OutputIndex, Specular.Expression);
+			F90.Connect(EdgeColor.OutputIndex, EdgeColor.Expression);
+
+			BaseColor.Expression = nullptr;
+			EdgeColor.Expression = nullptr;
+			Specular.Expression = nullptr;
+			Metallic.Expression = nullptr;
+		}
+
+		// Recreate this node and relink according to updated pins.
+		CastChecked<UMaterialGraphNode>(GraphNode)->RecreateAndLinkNode();
+	}
+	if (GraphNode)
 	{
 		GraphNode->ReconstructNode();
 	}
@@ -20342,7 +20370,7 @@ bool UMaterialExpressionStrataSlabBSDF::HasDMFPPluggedIn() const
 
 bool UMaterialExpressionStrataSlabBSDF::HasEdgeColor() const
 {
-	return EdgeColor.IsConnected();
+	return bUseMetalness ? EdgeColor.IsConnected() : F90.IsConnected();
 }
 
 bool UMaterialExpressionStrataSlabBSDF::HasThinFilm() const
