Strata, fixed shared normal count no longer displayed during the feedback to the user.

diff --git a/Engine/Source/Runtime/Engine/Private/Materials/HLSLMaterialTranslator.cpp b/Engine/Source/Runtime/Engine/Private/Materials/HLSLMaterialTranslator.cpp
--- a/Engine/Source/Runtime/Engine/Private/Materials/HLSLMaterialTranslator.cpp
+++ b/Engine/Source/Runtime/Engine/Private/Materials/HLSLMaterialTranslator.cpp
@@ -1276,58 +1276,6 @@
 					Error(*ErrorMsg);
 				}
 			}
-
-			// Output some debug info as comment in code and in the material stat window
-			static const auto CVarStrataBytePerPixel = IConsoleManager::Get().FindTConsoleVariableDataInt(TEXT("r.Strata.BytesPerPixel"));
-			check(CVarStrataBytePerPixel);
-			const uint32 StrataBytePerPixel = CVarStrataBytePerPixel ? CVarStrataBytePerPixel->GetValueOnAnyThread() : 0;
-			StrataMaterialAnalysis = StrataCompilationInfoMaterialAnalysis(this, StrataCompilationInfo, StrataBytePerPixel);
-
-			FString StrataMaterialDescription;
-
-			StrataMaterialDescription += FString::Printf(TEXT("----- STRATA -----\r\n"));
-			StrataMaterialDescription += FString::Printf(TEXT("StrataCompilationInfo -\r\n"));
-			StrataMaterialDescription += FString::Printf(TEXT(" - TotalBSDFCount    = %i\r\n"), StrataCompilationInfo.TotalBSDFCount);
-			StrataMaterialDescription += FString::Printf(TEXT(" - SharedNormalCount = %i\r\n"), StrataCompilationInfoGetSharedNormalCount());
-
-			for (uint32 LayerIt = 0; LayerIt < StrataCompilationInfo.LayerCount; ++LayerIt)
-			{
-				StrataMaterialDescription += FString::Printf(TEXT(" Layer %i BSDFs:\r\n"), LayerIt);
-
-				const FStrataMaterialCompilationInfo::FLayer& Layer = StrataCompilationInfo.Layers[LayerIt];
-				for (uint32 BSDFIt = 0; BSDFIt < Layer.BSDFCount; ++BSDFIt)
-				{
-					StrataMaterialDescription += FString::Printf(TEXT("     - %s - SharedNormalIndexMacro = %s \r\n"), *GetStrataBSDFName(Layer.BSDFs[BSDFIt].Type), *GetStrataSharedNormalIndexMacro(Layer.BSDFs[BSDFIt].RegisteredSharedNormal));
-				}
-			}
-
-			StrataMaterialDescription += FString::Printf(TEXT("Byte Per Pixel Budget      %u\r\n"), StrataBytePerPixel);
-			StrataMaterialDescription += FString::Printf(TEXT("Result.bFitInMemoryBudget  %s\r\n"), StrataMaterialAnalysis.bFitInMemoryBudget ? TEXT("YES") : TEXT("NO"));
-			StrataMaterialDescription += FString::Printf(TEXT("Result.RequestedLayerCount %u\r\n"), StrataMaterialAnalysis.RequestedLayerCount);
-			StrataMaterialDescription += FString::Printf(TEXT("Result.RequestedBSDFCount  %u\r\n"), StrataMaterialAnalysis.RequestedBSDFCount);
-			StrataMaterialDescription += FString::Printf(TEXT("Result.RequestedByteCount  %u\r\n"), StrataMaterialAnalysis.RequestedByteCount);
-			if (!StrataMaterialAnalysis.bFitInMemoryBudget)
-			{
-				StrataMaterialDescription += FString::Printf(TEXT("Result.ClampedLayerCount   %u\r\n"), StrataMaterialAnalysis.ClampedLayerCount);
-				StrataMaterialDescription += FString::Printf(TEXT("Result.ClampedBSDFCount    %u\r\n"), StrataMaterialAnalysis.ClampedBSDFCount);
-				StrataMaterialDescription += FString::Printf(TEXT("Result.UsedByteCount       %u\r\n"), StrataMaterialAnalysis.UsedByteCount);
-			}
-			StrataMaterialDescription += FString::Printf(TEXT("------------------\r\n"));
-
-			ResourcesString += TEXT("/*");
-			ResourcesString += StrataMaterialDescription;
-			ResourcesString += TEXT("*/");
-
-			MaterialCompilationOutput.StrataMaterialDescription = StrataMaterialDescription;
-
-			if (StrataMaterialAnalysis.ClampedLayerCount == 0)
-			{
-				UE_LOG(LogMaterial, Error, TEXT("Material %s cannot have any layers rendered due to its complexity (asset: %s).\r\n"), *Material->GetDebugName(), *Material->GetAssetPath().ToString());
-			}
-			else if (StrataMaterialAnalysis.RequestedLayerCount > StrataMaterialAnalysis.ClampedLayerCount)
-			{
-				UE_LOG(LogMaterial, Warning, TEXT("Material %s is not fitting the allocated byte per pixel budget for starta. Layers will be removed (asset: %s).\r\n"), *Material->GetDebugName(), *Material->GetAssetPath().ToString());
-			}
 		}
 		else
 		{
@@ -1807,6 +1755,62 @@
 
 			OutEnvironment.SetDefine(*GetStrataSharedNormalIndexMacro(It->Value.SharedData), LinearIndex);
 		}
+
+		// Now write some feedback to the user
+		{
+			// Output some debug info as comment in code and in the material stat window
+			static const auto CVarStrataBytePerPixel = IConsoleManager::Get().FindTConsoleVariableDataInt(TEXT("r.Strata.BytesPerPixel"));
+			check(CVarStrataBytePerPixel);
+			const uint32 StrataBytePerPixel = CVarStrataBytePerPixel ? CVarStrataBytePerPixel->GetValueOnAnyThread() : 0;
+			StrataMaterialAnalysis = StrataCompilationInfoMaterialAnalysis(this, StrataCompilationInfo, StrataBytePerPixel);
+
+			FString StrataMaterialDescription;
+
+			StrataMaterialDescription += FString::Printf(TEXT("----- STRATA -----\r\n"));
+			StrataMaterialDescription += FString::Printf(TEXT("StrataCompilationInfo -\r\n"));
+			StrataMaterialDescription += FString::Printf(TEXT(" - TotalBSDFCount    = %i\r\n"), StrataCompilationInfo.TotalBSDFCount);
+			StrataMaterialDescription += FString::Printf(TEXT(" - SharedNormalCount = %i\r\n"), StrataCompilationInfoGetSharedNormalCount());
+
+			for (uint32 LayerIt = 0; LayerIt < StrataCompilationInfo.LayerCount; ++LayerIt)
+			{
+				StrataMaterialDescription += FString::Printf(TEXT(" Layer %i BSDFs:\r\n"), LayerIt);
+
+				const FStrataMaterialCompilationInfo::FLayer& Layer = StrataCompilationInfo.Layers[LayerIt];
+				for (uint32 BSDFIt = 0; BSDFIt < Layer.BSDFCount; ++BSDFIt)
+				{
+					StrataMaterialDescription += FString::Printf(TEXT("     - %s - SharedNormalIndexMacro = %s \r\n"), *GetStrataBSDFName(Layer.BSDFs[BSDFIt].Type), *GetStrataSharedNormalIndexMacro(Layer.BSDFs[BSDFIt].RegisteredSharedNormal));
+				}
+			}
+
+			StrataMaterialDescription += FString::Printf(TEXT("Byte Per Pixel Budget      %u\r\n"), StrataBytePerPixel);
+			StrataMaterialDescription += FString::Printf(TEXT("Result.bFitInMemoryBudget  %s\r\n"), StrataMaterialAnalysis.bFitInMemoryBudget ? TEXT("YES") : TEXT("NO"));
+			StrataMaterialDescription += FString::Printf(TEXT("Result.RequestedLayerCount %u\r\n"), StrataMaterialAnalysis.RequestedLayerCount);
+			StrataMaterialDescription += FString::Printf(TEXT("Result.RequestedBSDFCount  %u\r\n"), StrataMaterialAnalysis.RequestedBSDFCount);
+			StrataMaterialDescription += FString::Printf(TEXT("Result.RequestedByteCount  %u\r\n"), StrataMaterialAnalysis.RequestedByteCount);
+			if (!StrataMaterialAnalysis.bFitInMemoryBudget)
+			{
+				StrataMaterialDescription += FString::Printf(TEXT("Result.ClampedLayerCount   %u\r\n"), StrataMaterialAnalysis.ClampedLayerCount);
+				StrataMaterialDescription += FString::Printf(TEXT("Result.ClampedBSDFCount    %u\r\n"), StrataMaterialAnalysis.ClampedBSDFCount);
+				StrataMaterialDescription += FString::Printf(TEXT("Result.UsedByteCount       %u\r\n"), StrataMaterialAnalysis.UsedByteCount);
+			}
+			StrataMaterialDescription += FString::Printf(TEXT("------------------\r\n"));
+
+			ResourcesString += TEXT("/*");
+			ResourcesString += StrataMaterialDescription;
+			ResourcesString += TEXT("*/");
+
+			MaterialCompilationOutput.StrataMaterialDescription = StrataMaterialDescription;
+
+			if (StrataMaterialAnalysis.ClampedLayerCount == 0)
+			{
+				UE_LOG(LogMaterial, Error, TEXT("Material %s cannot have any layers rendered due to its complexity (asset: %s).\r\n"), *Material->GetDebugName(), *Material->GetAssetPath().ToString());
+			}
+			else if (StrataMaterialAnalysis.RequestedLayerCount > StrataMaterialAnalysis.ClampedLayerCount)
+			{
+				UE_LOG(LogMaterial, Warning, TEXT("Material %s is not fitting the allocated byte per pixel budget for starta. Layers will be removed (asset: %s).\r\n"), *Material->GetDebugName(), *Material->GetAssetPath().ToString());
+			}
+		}
+
 	}
 
 	OutEnvironment.SetDefine(TEXT("TEXTURE_SAMPLE_DEBUG"), IsDebugTextureSampleEnabled() ? TEXT("1") : TEXT("0"));
