Skin Cache cache transition states

diff --git a/Engine/Source/Runtime/Engine/Private/GPUSkinCache.cpp b/Engine/Source/Runtime/Engine/Private/GPUSkinCache.cpp
--- a/Engine/Source/Runtime/Engine/Private/GPUSkinCache.cpp
+++ b/Engine/Source/Runtime/Engine/Private/GPUSkinCache.cpp
@@ -306,11 +306,11 @@ public:
 		uint32 IndexBufferOffsetValue = 0;
 		uint32 NumTriangles = 0;
 
-		FRWBuffer* TangentBuffer = nullptr;
-		FRWBuffer* IntermediateTangentBuffer = nullptr;
-		FRWBuffer* IntermediateAccumulatedTangentBuffer = nullptr;
-		FRWBuffer* PositionBuffer = nullptr;
-		FRWBuffer* PreviousPositionBuffer = nullptr;
+		FGPUSkinCache::FSkinCacheRWBuffer* TangentBuffer = nullptr;
+		FGPUSkinCache::FSkinCacheRWBuffer* IntermediateTangentBuffer = nullptr;
+		FGPUSkinCache::FSkinCacheRWBuffer* IntermediateAccumulatedTangentBuffer = nullptr;
+		FGPUSkinCache::FSkinCacheRWBuffer* PositionBuffer = nullptr;
+		FGPUSkinCache::FSkinCacheRWBuffer* PreviousPositionBuffer = nullptr;
 
         // Handle duplicates
         FShaderResourceViewRHIRef DuplicatedIndicesIndices = nullptr;
@@ -318,13 +318,13 @@ public:
 
 		FSectionDispatchData() = default;
 
-		inline FRWBuffer* GetPreviousPositionRWBuffer()
+		inline FGPUSkinCache::FSkinCacheRWBuffer* GetPreviousPositionRWBuffer()
 		{
 			check(PreviousPositionBuffer);
 			return PreviousPositionBuffer;
 		}
 
-		inline FRWBuffer* GetPositionRWBuffer()
+		inline FGPUSkinCache::FSkinCacheRWBuffer* GetPositionRWBuffer()
 		{
 			check(PositionBuffer);
 			return PositionBuffer;
@@ -338,12 +338,12 @@ public:
 			return SourceVertexFactory->GetPositionsSRV().GetReference();
 		}
 
-		inline FRWBuffer* GetTangentRWBuffer()
+		inline FGPUSkinCache::FSkinCacheRWBuffer* GetTangentRWBuffer()
 		{
 			return TangentBuffer;
 		}
 
-		FRWBuffer* GetActiveTangentRWBuffer()
+		FGPUSkinCache::FSkinCacheRWBuffer* GetActiveTangentRWBuffer()
 		{
 			bool bUseIntermediateTangentBuffer = IndexBuffer && GBlendUsingVertexColorForRecomputeTangents > 0;
 
@@ -357,7 +357,7 @@ public:
 			}
 		}
 
-		inline FRWBuffer* GetIntermediateAccumulatedTangentBuffer()
+		inline FGPUSkinCache::FSkinCacheRWBuffer* GetIntermediateAccumulatedTangentBuffer()
 		{
 			check(IntermediateAccumulatedTangentBuffer);
 			return IntermediateAccumulatedTangentBuffer;
@@ -365,7 +365,7 @@ public:
 
 		void UpdateVertexFactoryDeclaration()
 		{
-			TargetVertexFactory->UpdateVertexDeclaration(SourceVertexFactory, GetPositionRWBuffer(), GetPreSkinPositionSRV(), GetTangentRWBuffer());
+			TargetVertexFactory->UpdateVertexDeclaration(SourceVertexFactory, &GetPositionRWBuffer()->Buffer, GetPreSkinPositionSRV(), &GetTangentRWBuffer()->Buffer);
 		}
 	};
 
@@ -378,9 +378,9 @@ public:
 	{
 		FCachedGeometry::Section MeshSection;
 		const FSkelMeshRenderSection& Section = *DispatchData[SectionIndex].Section;
-		MeshSection.PositionBuffer	= DispatchData[SectionIndex].PositionBuffer->SRV;
+		MeshSection.PositionBuffer	= DispatchData[SectionIndex].PositionBuffer->Buffer.SRV;
 		MeshSection.UVsBuffer		= DispatchData[SectionIndex].UVsBufferSRV;
-		MeshSection.TotalVertexCount= DispatchData[SectionIndex].PositionBuffer->NumBytes / (sizeof(float)*3);
+		MeshSection.TotalVertexCount= DispatchData[SectionIndex].PositionBuffer->Buffer.NumBytes / (sizeof(float)*3);
 		MeshSection.NumPrimitives	= Section.NumTriangles;
 		MeshSection.NumVertices		= Section.NumVertices;
 		MeshSection.IndexBaseIndex	= Section.BaseIndex;
@@ -512,7 +512,7 @@ public:
 			const FSectionDispatchData& SectionData = DispatchData[SectionIdx];
 			FRayTracingGeometrySegment& Segment = OutSegments[SectionIdx];
 
-			Segment.VertexBuffer = SectionData.PositionBuffer->Buffer;
+			Segment.VertexBuffer = SectionData.PositionBuffer->Buffer.Buffer;
 			Segment.VertexBufferOffset = 0;
 
 			check(SectionData.Section->NumTriangles == Segment.NumPrimitives);
@@ -806,9 +806,9 @@ void FGPUSkinCache::TransitionAllToReadable(FRHICommandList& RHICmdList)
 		FMemMark Mark(FMemStack::Get());
 		TArray<FRHITransitionInfo, SceneRenderingAllocator> UAVs;
 		UAVs.Reserve(BuffersToTransitionToRead.Num());
-		for (TSet<FRHIUnorderedAccessView*>::TConstIterator SetIt(BuffersToTransitionToRead); SetIt; ++SetIt)
+		for (TSet<FSkinCacheRWBuffer*>::TConstIterator SetIt(BuffersToTransitionToRead); SetIt; ++SetIt)
 		{
-			UAVs.Add(FRHITransitionInfo(*SetIt, ERHIAccess::Unknown, ERHIAccess::VertexOrIndexBuffer | ERHIAccess::SRVMask));
+			UAVs.Add((*SetIt)->UpdateAccessState(ERHIAccess::VertexOrIndexBuffer | ERHIAccess::SRVMask));
 		}
 		RHICmdList.Transition(UAVs);
 
@@ -964,8 +964,8 @@ public:
 
 		SetShaderValue(RHICmdList, ShaderRHI, NumTriangles, DispatchData.NumTriangles);
 
-		SetSRVParameter(RHICmdList, ShaderRHI, GPUPositionCacheBuffer, DispatchData.GetPositionRWBuffer()->SRV);
-		SetSRVParameter(RHICmdList, ShaderRHI, GPUTangentCacheBuffer, GBlendUsingVertexColorForRecomputeTangents > 0 ? DispatchData.IntermediateTangentBuffer->SRV : DispatchData.GetTangentRWBuffer()->SRV);
+		SetSRVParameter(RHICmdList, ShaderRHI, GPUPositionCacheBuffer, DispatchData.GetPositionRWBuffer()->Buffer.SRV);
+		SetSRVParameter(RHICmdList, ShaderRHI, GPUTangentCacheBuffer, GBlendUsingVertexColorForRecomputeTangents > 0 ? DispatchData.IntermediateTangentBuffer->Buffer.SRV : DispatchData.GetTangentRWBuffer()->Buffer.SRV);
 		SetSRVParameter(RHICmdList, ShaderRHI, UVsInputBuffer, DispatchData.UVsBufferSRV);
 
 		SetShaderValue(RHICmdList, ShaderRHI, SkinCacheStart, DispatchData.OutputStreamStart);
@@ -1102,9 +1102,9 @@ public:
 		// UAVs
 		SetUAVParameter(RHICmdList, ShaderRHI, IntermediateAccumBufferUAV, StagingBuffer.UAV);
 		SetShaderValue(RHICmdList, ShaderRHI, IntermediateAccumBufferOffset, GRecomputeTangentsParallelDispatch * DispatchData.IndexBufferOffsetValue);
-		SetUAVParameter(RHICmdList, ShaderRHI, TangentBufferUAV, DispatchData.GetTangentRWBuffer()->UAV);
+		SetUAVParameter(RHICmdList, ShaderRHI, TangentBufferUAV, DispatchData.GetTangentRWBuffer()->Buffer.UAV);
 
-		SetSRVParameter(RHICmdList, ShaderRHI, TangentInputBuffer, DispatchData.IntermediateTangentBuffer ? DispatchData.IntermediateTangentBuffer->SRV : nullptr);
+		SetSRVParameter(RHICmdList, ShaderRHI, TangentInputBuffer, DispatchData.IntermediateTangentBuffer ? DispatchData.IntermediateTangentBuffer->Buffer.SRV : nullptr);
 
 		SetSRVParameter(RHICmdList, ShaderRHI, ColorInputBuffer, DispatchData.ColorBufferSRV);
 	}
@@ -1146,7 +1146,7 @@ class FRecomputeTangentsPerVertexPassCS : public FBaseRecomputeTangentsPerVertex
 IMPLEMENT_SHADER_TYPE(template<>, FRecomputeTangentsPerVertexPassCS<0>, TEXT("/Engine/Private/RecomputeTangentsPerVertexPass.usf"), TEXT("MainCS"), SF_Compute);
 IMPLEMENT_SHADER_TYPE(template<>, FRecomputeTangentsPerVertexPassCS<1>, TEXT("/Engine/Private/RecomputeTangentsPerVertexPass.usf"), TEXT("MainCS"), SF_Compute);
 
-void FGPUSkinCache::DispatchUpdateSkinTangents(FRHICommandListImmediate& RHICmdList, FGPUSkinCacheEntry* Entry, int32 SectionIndex, FRWBuffer*& StagingBuffer, bool bTrianglePass)
+void FGPUSkinCache::DispatchUpdateSkinTangents(FRHICommandListImmediate& RHICmdList, FGPUSkinCacheEntry* Entry, int32 SectionIndex, FSkinCacheRWBuffer*& StagingBuffer, bool bTrianglePass)
 {
 	FGPUSkinCacheEntry::FSectionDispatchData& DispatchData = Entry->DispatchData[SectionIndex];
 
@@ -1172,17 +1172,17 @@ void FGPUSkinCache::DispatchUpdateSkinTangents(FRHICommandListImmediate& RHICmdL
 			uint32 NumIntsPerBuffer = DispatchData.NumTriangles * 3 * FGPUSkinCache::IntermediateAccumBufferNumInts;
 			CurrentStagingBufferIndex = (CurrentStagingBufferIndex + 1) % StagingBuffers.Num();
 			StagingBuffer = &StagingBuffers[CurrentStagingBufferIndex];
-			if (StagingBuffer->NumBytes < NumIntsPerBuffer * sizeof(uint32))
+			if (StagingBuffer->Buffer.NumBytes < NumIntsPerBuffer * sizeof(uint32))
 			{
 				StagingBuffer->Release();
-				StagingBuffer->Initialize(TEXT("SkinTangentIntermediate"), sizeof(int32), NumIntsPerBuffer, PF_R32_SINT, BUF_UnorderedAccess);
-				RHICmdList.BindDebugLabelName(StagingBuffer->UAV, TEXT("SkinTangentIntermediate"));
+				StagingBuffer->Buffer.Initialize(TEXT("SkinTangentIntermediate"), sizeof(int32), NumIntsPerBuffer, PF_R32_SINT, BUF_UnorderedAccess);
+				RHICmdList.BindDebugLabelName(StagingBuffer->Buffer.UAV, TEXT("SkinTangentIntermediate"));
 
 				const uint32 MemSize = NumIntsPerBuffer * sizeof(uint32);
 				SET_MEMORY_STAT(STAT_GPUSkinCache_TangentsIntermediateMemUsed, MemSize);
 
 				// The UAV must be zero-filled. We leave it zeroed after each round (see RecomputeTangentsPerVertexPass.usf), so this is only needed on when the buffer is first created.
-				RHICmdList.ClearUAVUint(StagingBuffer->UAV, FUintVector4(0, 0, 0, 0));
+				RHICmdList.ClearUAVUint(StagingBuffer->Buffer.UAV, FUintVector4(0, 0, 0, 0));
 			}
 		}
 
@@ -1229,11 +1229,14 @@ void FGPUSkinCache::DispatchUpdateSkinTangents(FRHICommandListImmediate& RHICmdL
 			if (!GRecomputeTangentsParallelDispatch)
 			{
 				// When triangle & vertex passes are interleaved, resource transition is needed in between.
-				RHICmdList.Transition(FRHITransitionInfo(StagingBuffer->UAV, ERHIAccess::Unknown, ERHIAccess::UAVCompute));
+				RHICmdList.Transition({
+					DispatchData.GetActiveTangentRWBuffer()->UpdateAccessState(ERHIAccess::SRVCompute),
+					StagingBuffer->UpdateAccessState(ERHIAccess::UAVCompute)
+				});
 			}
 
 			INC_DWORD_STAT_BY(STAT_GPUSkinCache_NumTrianglesForRecomputeTangents, NumTriangles);
-			Shader->SetParameters(RHICmdList, Entry, DispatchData, GRecomputeTangentsParallelDispatch ? *DispatchData.GetIntermediateAccumulatedTangentBuffer() : *StagingBuffer);
+			Shader->SetParameters(RHICmdList, Entry, DispatchData, GRecomputeTangentsParallelDispatch ? DispatchData.GetIntermediateAccumulatedTangentBuffer()->Buffer : StagingBuffer->Buffer);
 			DispatchComputeShader(RHICmdList, Shader.GetShader(), ThreadGroupCountValue, 1, 1);
 			Shader->UnsetParameters(RHICmdList);
 			IncrementDispatchCounter(RHICmdList);
@@ -1261,12 +1264,12 @@ void FGPUSkinCache::DispatchUpdateSkinTangents(FRHICommandListImmediate& RHICmdL
 		{
 			// When triangle & vertex passes are interleaved, resource transition is needed in between.
 			RHICmdList.Transition({
-				FRHITransitionInfo(DispatchData.GetTangentRWBuffer()->UAV.GetReference(), GBlendUsingVertexColorForRecomputeTangents ? ERHIAccess::Unknown : ERHIAccess::SRVCompute, ERHIAccess::UAVCompute),
-				FRHITransitionInfo(StagingBuffer->UAV, ERHIAccess::UAVCompute, ERHIAccess::UAVCompute)
+				DispatchData.GetTangentRWBuffer()->UpdateAccessState(ERHIAccess::UAVCompute),
+				StagingBuffer->UpdateAccessState(ERHIAccess::UAVCompute)
 				});
 		}
 
-		ComputeShader->SetParameters(RHICmdList, Entry, DispatchData, GRecomputeTangentsParallelDispatch ? *DispatchData.GetIntermediateAccumulatedTangentBuffer() : *StagingBuffer);
+		ComputeShader->SetParameters(RHICmdList, Entry, DispatchData, GRecomputeTangentsParallelDispatch ? DispatchData.GetIntermediateAccumulatedTangentBuffer()->Buffer : StagingBuffer->Buffer);
 		DispatchComputeShader(RHICmdList, ComputeShader.GetShader(), ThreadGroupCountValue, 1, 1);
 		ComputeShader->UnsetParameters(RHICmdList);
 		IncrementDispatchCounter(RHICmdList);
@@ -1296,21 +1299,30 @@ FGPUSkinCache::FRWBuffersAllocation* FGPUSkinCache::TryAllocBuffer(uint32 NumVer
 
 DECLARE_GPU_STAT(GPUSkinCache);
 
-void FGPUSkinCache::MakeBufferTransitions(FRHICommandListImmediate& RHICmdList, TArray<FRHIUnorderedAccessView*>& Buffers, ERHIAccess FromState, ERHIAccess ToState)
+void FGPUSkinCache::MakeBufferTransitions(FRHICommandListImmediate& RHICmdList, TArray<FSkinCacheRWBuffer*>& Buffers, ERHIAccess ToState)
 {
 	if (Buffers.Num() > 0)
 	{
 		FMemMark Mark(FMemStack::Get());
 		TArray<FRHITransitionInfo, SceneRenderingAllocator> UAVs;
 		UAVs.Reserve(Buffers.Num());
-		for (FRHIUnorderedAccessView* Buffer : Buffers)
+		for (FSkinCacheRWBuffer* Buffer : Buffers)
 		{
-			UAVs.Add(FRHITransitionInfo(Buffer, FromState, ToState));
+			UAVs.Add(Buffer->UpdateAccessState(ToState));
 		}
 		RHICmdList.Transition(MakeArrayView(UAVs.GetData(), UAVs.Num()));
 	}
 }
 
+void FGPUSkinCache::GetBufferUAVs(const TArray<FSkinCacheRWBuffer*>& InBuffers, TArray<FRHIUnorderedAccessView*>& OutUAVs)
+{
+	OutUAVs.Reset(InBuffers.Num());
+	for (const FSkinCacheRWBuffer* Buffer : InBuffers)
+	{
+		OutUAVs.Add(Buffer->Buffer.UAV);
+	}
+}
+
 void FGPUSkinCache::DoDispatch(FRHICommandListImmediate& RHICmdList)
 {
 	SCOPED_GPU_STAT(RHICmdList, GPUSkinCache);
@@ -1318,7 +1330,7 @@ void FGPUSkinCache::DoDispatch(FRHICommandListImmediate& RHICmdList)
 	int32 BatchCount = BatchDispatches.Num();
 	INC_DWORD_STAT_BY(STAT_GPUSkinCache_TotalNumChunks, BatchCount);
 
-	TArray<FRHIUnorderedAccessView*> BuffersToTransitionForSkinning;
+	TArray<FSkinCacheRWBuffer*> BuffersToTransitionForSkinning;
 	BuffersToTransitionForSkinning.Reserve(BatchCount * 2);
 	{
 		for (int32 i = 0; i < BatchCount; ++i)
@@ -1326,13 +1338,12 @@ void FGPUSkinCache::DoDispatch(FRHICommandListImmediate& RHICmdList)
 			FDispatchEntry& DispatchItem = BatchDispatches[i];
 			PrepareUpdateSkinning(DispatchItem.SkinCacheEntry, DispatchItem.Section, DispatchItem.RevisionNumber, &BuffersToTransitionForSkinning);
 		}
-
-		Algo::Sort(BuffersToTransitionForSkinning, [](const FRHIUnorderedAccessView* A, const FRHIUnorderedAccessView* B){return A < B;});
-		BuffersToTransitionForSkinning.SetNum(Algo::Unique(BuffersToTransitionForSkinning));
-		MakeBufferTransitions(RHICmdList, BuffersToTransitionForSkinning, ERHIAccess::Unknown, ERHIAccess::UAVCompute);
+		MakeBufferTransitions(RHICmdList, BuffersToTransitionForSkinning, ERHIAccess::UAVCompute);
 	}
 
-	RHICmdList.BeginUAVOverlap(BuffersToTransitionForSkinning);
+	TArray<FRHIUnorderedAccessView*> SkinningBuffersToOverlap;
+	GetBufferUAVs(BuffersToTransitionForSkinning, SkinningBuffersToOverlap);
+	RHICmdList.BeginUAVOverlap(SkinningBuffersToOverlap);
 	{
 		SCOPED_DRAW_EVENT(RHICmdList, GPUSkinCache_UpdateSkinningBatches);
 		for (int32 i = 0; i < BatchCount; ++i)
@@ -1341,32 +1352,34 @@ void FGPUSkinCache::DoDispatch(FRHICommandListImmediate& RHICmdList)
 			DispatchUpdateSkinning(RHICmdList, DispatchItem.SkinCacheEntry, DispatchItem.Section, DispatchItem.RevisionNumber);
 		}
 	}
-	RHICmdList.EndUAVOverlap(BuffersToTransitionForSkinning);
+	RHICmdList.EndUAVOverlap(SkinningBuffersToOverlap);
 
 	// Do necessary buffer transitions before recomputing tangents
-	TArray<FRHIUnorderedAccessView*> BuffersToSRVForRecomputeTangents;
-	TArray<FRHIUnorderedAccessView*> IntermediateAccumulatedTangentBuffers;
+	TArray<FSkinCacheRWBuffer*> BuffersToSRVForRecomputeTangents;
+	TArray<FSkinCacheRWBuffer*> IntermediateAccumulatedTangentBuffers;
 	for (int32 i = 0; i < BatchCount; ++i)
 	{
 		FDispatchEntry& DispatchItem = BatchDispatches[i];
 		FGPUSkinCacheEntry::FSectionDispatchData& DispatchData = DispatchItem.SkinCacheEntry->DispatchData[DispatchItem.Section];
 		if (DispatchData.IndexBuffer)
 		{
-			BuffersToSRVForRecomputeTangents.AddUnique(DispatchData.GetPositionRWBuffer()->UAV);
-			BuffersToSRVForRecomputeTangents.AddUnique(DispatchData.GetActiveTangentRWBuffer()->UAV);
+			BuffersToSRVForRecomputeTangents.AddUnique(DispatchData.GetPositionRWBuffer());
+			BuffersToSRVForRecomputeTangents.AddUnique(DispatchData.GetActiveTangentRWBuffer());
 			if (GRecomputeTangentsParallelDispatch)
 			{
-				IntermediateAccumulatedTangentBuffers.AddUnique(DispatchData.GetIntermediateAccumulatedTangentBuffer()->UAV);
+				IntermediateAccumulatedTangentBuffers.AddUnique(DispatchData.GetIntermediateAccumulatedTangentBuffer());
 			}
 		}	
 	}
-	MakeBufferTransitions(RHICmdList, BuffersToSRVForRecomputeTangents, ERHIAccess::Unknown, ERHIAccess::SRVCompute);
-	MakeBufferTransitions(RHICmdList, IntermediateAccumulatedTangentBuffers, ERHIAccess::Unknown, ERHIAccess::UAVCompute);
+	MakeBufferTransitions(RHICmdList, BuffersToSRVForRecomputeTangents, ERHIAccess::SRVCompute);
+	MakeBufferTransitions(RHICmdList, IntermediateAccumulatedTangentBuffers, ERHIAccess::UAVCompute);
 
-	RHICmdList.BeginUAVOverlap(IntermediateAccumulatedTangentBuffers);
+	TArray<FRHIUnorderedAccessView*> IntermediateAccumulatedTangentBuffersToOverlap;
+	GetBufferUAVs(IntermediateAccumulatedTangentBuffers, IntermediateAccumulatedTangentBuffersToOverlap);
+	RHICmdList.BeginUAVOverlap(IntermediateAccumulatedTangentBuffersToOverlap);
 	{
 		SCOPED_DRAW_EVENT(RHICmdList, GPUSkinCache_RecomputeTangentsBatches);
-		FRWBuffer* StagingBuffer = nullptr;
+		FSkinCacheRWBuffer* StagingBuffer = nullptr;
 		for (int32 i = 0; i < BatchCount; ++i)
 		{
 			FDispatchEntry& DispatchItem = BatchDispatches[i];
@@ -1383,17 +1396,19 @@ void FGPUSkinCache::DoDispatch(FRHICommandListImmediate& RHICmdList)
 		if (GRecomputeTangentsParallelDispatch)
 		{
 			// Do necessary buffer transitions before vertex pass dispatches
-			TArray<FRHIUnorderedAccessView*> TangentBuffers;
+			TArray<FSkinCacheRWBuffer*> TangentBuffers;
 			for (int32 i = 0; i < BatchCount; ++i)
 			{
 				FDispatchEntry& DispatchItem = BatchDispatches[i];
 				FGPUSkinCacheEntry::FSectionDispatchData& DispatchData = DispatchItem.SkinCacheEntry->DispatchData[DispatchItem.Section];
-				TangentBuffers.AddUnique(DispatchData.GetTangentRWBuffer()->UAV);
+				TangentBuffers.AddUnique(DispatchData.GetTangentRWBuffer());
 			}
-			MakeBufferTransitions(RHICmdList, TangentBuffers, GBlendUsingVertexColorForRecomputeTangents ? ERHIAccess::Unknown : ERHIAccess::SRVCompute, ERHIAccess::UAVCompute);
-			MakeBufferTransitions(RHICmdList, IntermediateAccumulatedTangentBuffers, ERHIAccess::UAVCompute, ERHIAccess::UAVCompute);
+			MakeBufferTransitions(RHICmdList, TangentBuffers, ERHIAccess::UAVCompute);
+			MakeBufferTransitions(RHICmdList, IntermediateAccumulatedTangentBuffers, ERHIAccess::UAVCompute);
 		
-			RHICmdList.BeginUAVOverlap(TangentBuffers);
+			TArray<FRHIUnorderedAccessView*> TangentBuffersToOverlap;
+			GetBufferUAVs(TangentBuffers, TangentBuffersToOverlap);
+			RHICmdList.BeginUAVOverlap(TangentBuffersToOverlap);
 			for (int32 i = 0; i < BatchCount; ++i)
 			{
 				FDispatchEntry& DispatchItem = BatchDispatches[i];
@@ -1402,10 +1417,10 @@ void FGPUSkinCache::DoDispatch(FRHICommandListImmediate& RHICmdList)
 					DispatchUpdateSkinTangents(RHICmdList, DispatchItem.SkinCacheEntry, DispatchItem.Section, StagingBuffer, false);
 				}
 			}
-			RHICmdList.EndUAVOverlap(TangentBuffers);
+			RHICmdList.EndUAVOverlap(TangentBuffersToOverlap);
 		}
 	}
-	RHICmdList.EndUAVOverlap(IntermediateAccumulatedTangentBuffers);
+	RHICmdList.EndUAVOverlap(IntermediateAccumulatedTangentBuffersToOverlap);
 
 	for (int32 i = 0; i < BatchCount; ++i)
 	{
@@ -1420,33 +1435,37 @@ void FGPUSkinCache::DoDispatch(FRHICommandListImmediate& RHICmdList, FGPUSkinCac
 
 	INC_DWORD_STAT(STAT_GPUSkinCache_TotalNumChunks);
 
-	TArray<FRHIUnorderedAccessView*> BuffersToTransitionForSkinning;
+	TArray<FSkinCacheRWBuffer*> BuffersToTransitionForSkinning;
 	PrepareUpdateSkinning(SkinCacheEntry, Section, RevisionNumber, &BuffersToTransitionForSkinning);
-	MakeBufferTransitions(RHICmdList, BuffersToTransitionForSkinning, ERHIAccess::Unknown, ERHIAccess::UAVCompute);
+	MakeBufferTransitions(RHICmdList, BuffersToTransitionForSkinning, ERHIAccess::UAVCompute);
 
-	RHICmdList.BeginUAVOverlap(BuffersToTransitionForSkinning);
-	DispatchUpdateSkinning(RHICmdList, SkinCacheEntry, Section, RevisionNumber);
-	RHICmdList.EndUAVOverlap(BuffersToTransitionForSkinning);
+	TArray<FRHIUnorderedAccessView*> SkinningBuffersToOverlap;
+	GetBufferUAVs(BuffersToTransitionForSkinning, SkinningBuffersToOverlap);
+	RHICmdList.BeginUAVOverlap(SkinningBuffersToOverlap);
+	{
+		DispatchUpdateSkinning(RHICmdList, SkinCacheEntry, Section, RevisionNumber);
+	}
+	RHICmdList.EndUAVOverlap(SkinningBuffersToOverlap);
 
 	FGPUSkinCacheEntry::FSectionDispatchData& DispatchData = SkinCacheEntry->DispatchData[Section];
 	if (DispatchData.IndexBuffer)
 	{
 		RHICmdList.Transition({
-			FRHITransitionInfo(DispatchData.GetPositionRWBuffer()->UAV.GetReference(), ERHIAccess::Unknown, ERHIAccess::SRVCompute),
-			FRHITransitionInfo(DispatchData.GetActiveTangentRWBuffer()->UAV.GetReference(), ERHIAccess::Unknown, ERHIAccess::SRVCompute)
+			DispatchData.GetPositionRWBuffer()->UpdateAccessState(ERHIAccess::SRVCompute),
+			DispatchData.GetActiveTangentRWBuffer()->UpdateAccessState(ERHIAccess::SRVCompute)
 		});
 		if (GRecomputeTangentsParallelDispatch)
 		{
-			RHICmdList.Transition(FRHITransitionInfo(DispatchData.GetIntermediateAccumulatedTangentBuffer()->UAV.GetReference(), ERHIAccess::Unknown, ERHIAccess::UAVCompute));
+			RHICmdList.Transition(DispatchData.GetIntermediateAccumulatedTangentBuffer()->UpdateAccessState(ERHIAccess::UAVCompute));
 		}
 
-		FRWBuffer* StagingBuffer = nullptr;
+		FSkinCacheRWBuffer* StagingBuffer = nullptr;
 		DispatchUpdateSkinTangents(RHICmdList, SkinCacheEntry, Section, StagingBuffer, true);
 		if (GRecomputeTangentsParallelDispatch)
 		{
 			RHICmdList.Transition({
-				FRHITransitionInfo(DispatchData.GetTangentRWBuffer()->UAV.GetReference(), GBlendUsingVertexColorForRecomputeTangents ? ERHIAccess::Unknown : ERHIAccess::SRVCompute, ERHIAccess::UAVCompute),
-				FRHITransitionInfo(DispatchData.GetIntermediateAccumulatedTangentBuffer()->UAV.GetReference(), ERHIAccess::UAVCompute, ERHIAccess::UAVCompute)
+				DispatchData.GetTangentRWBuffer()->UpdateAccessState(ERHIAccess::UAVCompute),
+				DispatchData.GetIntermediateAccumulatedTangentBuffer()->UpdateAccessState(ERHIAccess::UAVCompute)
 			});
 		}
 		DispatchUpdateSkinTangents(RHICmdList, SkinCacheEntry, Section, StagingBuffer, false);
@@ -1848,15 +1867,15 @@ void FGPUSkinCache::SetVertexStreams(FGPUSkinCacheEntry* Entry, int32 Section, F
 	FGPUSkinCacheEntry::FSectionDispatchData& DispatchData = Entry->DispatchData[Section];
 
 	//UE_LOG(LogSkinCache, Warning, TEXT("*** SetVertexStreams E %p All %p Sec %d(%p) LOD %d"), Entry, Entry->DispatchData[Section].Allocation, Section, Entry->DispatchData[Section].Section, Entry->LOD);
-	RHICmdList.SetStreamSource(VertexFactory->GetPositionStreamIndex(), DispatchData.GetPositionRWBuffer()->Buffer, 0);
+	RHICmdList.SetStreamSource(VertexFactory->GetPositionStreamIndex(), DispatchData.GetPositionRWBuffer()->Buffer.Buffer, 0);
 	if (VertexFactory->GetTangentStreamIndex() > -1 && DispatchData.GetTangentRWBuffer())
 	{
-		RHICmdList.SetStreamSource(VertexFactory->GetTangentStreamIndex(), DispatchData.GetTangentRWBuffer()->Buffer, 0);
+		RHICmdList.SetStreamSource(VertexFactory->GetTangentStreamIndex(), DispatchData.GetTangentRWBuffer()->Buffer.Buffer, 0);
 	}
 
 	if (ShaderRHI && GPUSkinCachePreviousPositionBuffer.IsBound())
 	{
-		RHICmdList.SetShaderResourceViewParameter(ShaderRHI, GPUSkinCachePreviousPositionBuffer.GetBaseIndex(), DispatchData.GetPreviousPositionRWBuffer()->SRV);
+		RHICmdList.SetShaderResourceViewParameter(ShaderRHI, GPUSkinCachePreviousPositionBuffer.GetBaseIndex(), DispatchData.GetPreviousPositionRWBuffer()->Buffer.SRV);
 	}
 }
 
@@ -1879,18 +1898,18 @@ void FGPUSkinCache::GetShaderBindings(
 
 	//UE_LOG(LogSkinCache, Warning, TEXT("*** SetVertexStreams E %p All %p Sec %d(%p) LOD %d"), Entry, Entry->DispatchData[Section].Allocation, Section, Entry->DispatchData[Section].Section, Entry->LOD);
 
-	VertexStreams.Add(FVertexInputStream(VertexFactory->GetPositionStreamIndex(), 0, DispatchData.GetPositionRWBuffer()->Buffer));
+	VertexStreams.Add(FVertexInputStream(VertexFactory->GetPositionStreamIndex(), 0, DispatchData.GetPositionRWBuffer()->Buffer.Buffer));
 
 	if (VertexFactory->GetTangentStreamIndex() > -1 && DispatchData.GetTangentRWBuffer())
 	{
-		VertexStreams.Add(FVertexInputStream(VertexFactory->GetTangentStreamIndex(), 0, DispatchData.GetTangentRWBuffer()->Buffer));
+		VertexStreams.Add(FVertexInputStream(VertexFactory->GetTangentStreamIndex(), 0, DispatchData.GetTangentRWBuffer()->Buffer.Buffer));
 	}
 
-	ShaderBindings.Add(GPUSkinCachePositionBuffer, DispatchData.GetPositionRWBuffer()->SRV);
-	ShaderBindings.Add(GPUSkinCachePreviousPositionBuffer, DispatchData.GetPreviousPositionRWBuffer()->SRV);
+	ShaderBindings.Add(GPUSkinCachePositionBuffer, DispatchData.GetPositionRWBuffer()->Buffer.SRV);
+	ShaderBindings.Add(GPUSkinCachePreviousPositionBuffer, DispatchData.GetPreviousPositionRWBuffer()->Buffer.SRV);
 }
 
-void FGPUSkinCache::PrepareUpdateSkinning(FGPUSkinCacheEntry* Entry, int32 Section, uint32 RevisionNumber, TArray<FRHIUnorderedAccessView*>* OverlappedUAVs)
+void FGPUSkinCache::PrepareUpdateSkinning(FGPUSkinCacheEntry* Entry, int32 Section, uint32 RevisionNumber, TArray<FSkinCacheRWBuffer*>* OverlappedUAVs)
 {
 	FGPUSkinCacheEntry::FSectionDispatchData& DispatchData = Entry->DispatchData[Section];
 	FGPUBaseSkinVertexFactory::FShaderDataType& ShaderData = DispatchData.SourceVertexFactory->GetShaderData();
@@ -1904,7 +1923,7 @@ void FGPUSkinCache::PrepareUpdateSkinning(FGPUSkinCacheEntry* Entry, int32 Secti
 	DispatchData.DispatchFlags = 0;
 
 	auto BufferUpdate = [&DispatchData, OverlappedUAVs](
-		FRWBuffer*& PositionBuffer,
+		FSkinCacheRWBuffer*& PositionBuffer,
 		const FVertexBufferAndSRV& BoneBuffer, 
 		uint32 Revision,
 		const FVertexBufferAndSRV& PrevBoneBuffer,
@@ -1923,7 +1942,7 @@ void FGPUSkinCache::PrepareUpdateSkinning(FGPUSkinCacheEntry* Entry, int32 Secti
 
 			if (OverlappedUAVs)
 			{
-				(*OverlappedUAVs).Emplace(PositionBuffer->UAV);
+				(*OverlappedUAVs).AddUnique(PositionBuffer);
 			}
 		}
 	};
@@ -1952,7 +1971,7 @@ void FGPUSkinCache::PrepareUpdateSkinning(FGPUSkinCacheEntry* Entry, int32 Secti
 
 	if (OverlappedUAVs && DispatchData.DispatchFlags != 0 && DispatchData.GetActiveTangentRWBuffer())
 	{
-		 (*OverlappedUAVs).Emplace(DispatchData.GetActiveTangentRWBuffer()->UAV);
+		 (*OverlappedUAVs).AddUnique(DispatchData.GetActiveTangentRWBuffer());
 	}
 
 	check(DispatchData.PreviousPositionBuffer != DispatchData.PositionBuffer);
@@ -2084,8 +2103,8 @@ void FGPUSkinCache::DispatchUpdateSkinning(FRHICommandListImmediate& RHICmdList,
 			PrevBoneBuffer,
 			Entry,
 			DispatchData,
-			DispatchData.GetPreviousPositionRWBuffer()->UAV,
-			DispatchData.GetActiveTangentRWBuffer() ? DispatchData.GetActiveTangentRWBuffer()->UAV : nullptr
+			DispatchData.GetPreviousPositionRWBuffer()->Buffer.UAV,
+			DispatchData.GetActiveTangentRWBuffer() ? DispatchData.GetActiveTangentRWBuffer()->Buffer.UAV : nullptr
 			);
 
 		uint32 VertexCountAlign64 = FMath::DivideAndRoundUp(DispatchData.NumVertices, (uint32)64);
@@ -2093,7 +2112,7 @@ void FGPUSkinCache::DispatchUpdateSkinning(FRHICommandListImmediate& RHICmdList,
 		RHICmdList.DispatchComputeShader(VertexCountAlign64, 1, 1);
 		Shader->UnsetParameters(RHICmdList);
 		IncrementDispatchCounter(RHICmdList);
-		BuffersToTransitionToRead.Add(DispatchData.GetPreviousPositionRWBuffer()->UAV);
+		BuffersToTransitionToRead.Add(DispatchData.GetPreviousPositionRWBuffer());
 	}
 
 	if ((DispatchData.DispatchFlags & (uint32)EGPUSkinCacheDispatchFlags::DispatchPosition) != 0)
@@ -2105,8 +2124,8 @@ void FGPUSkinCache::DispatchUpdateSkinning(FRHICommandListImmediate& RHICmdList,
 			BoneBuffer, 
 			Entry, 
 			DispatchData, 
-			DispatchData.GetPositionRWBuffer()->UAV, 
-			DispatchData.GetActiveTangentRWBuffer() ? DispatchData.GetActiveTangentRWBuffer()->UAV : nullptr
+			DispatchData.GetPositionRWBuffer()->Buffer.UAV, 
+			DispatchData.GetActiveTangentRWBuffer() ? DispatchData.GetActiveTangentRWBuffer()->Buffer.UAV : nullptr
 			);
 
 		uint32 VertexCountAlign64 = FMath::DivideAndRoundUp(DispatchData.NumVertices, (uint32)64);
@@ -2114,35 +2133,27 @@ void FGPUSkinCache::DispatchUpdateSkinning(FRHICommandListImmediate& RHICmdList,
 		RHICmdList.DispatchComputeShader(VertexCountAlign64, 1, 1);
 		Shader->UnsetParameters(RHICmdList);
 		IncrementDispatchCounter(RHICmdList);
-		BuffersToTransitionToRead.Add(DispatchData.GetPositionRWBuffer()->UAV);
+		BuffersToTransitionToRead.Add(DispatchData.GetPositionRWBuffer());
 	}
 
-	BuffersToTransitionToRead.Add(DispatchData.GetTangentRWBuffer()->UAV);
+	BuffersToTransitionToRead.Add(DispatchData.GetTangentRWBuffer());
 	check(DispatchData.PreviousPositionBuffer != DispatchData.PositionBuffer);
 }
 
-void FGPUSkinCache::FRWBuffersAllocation::RemoveAllFromTransitionArray(TSet<FRHIUnorderedAccessView*>& InBuffersToTransition)
+void FGPUSkinCache::FRWBuffersAllocation::RemoveAllFromTransitionArray(TSet<FSkinCacheRWBuffer*>& InBuffersToTransition)
 {
 	for (uint32 i = 0; i < NUM_BUFFERS; i++)
 	{
-		FRWBuffer& RWBuffer = PositionBuffers[i];
-		if (RWBuffer.UAV.IsValid())
-		{
-			InBuffersToTransition.Remove(RWBuffer.UAV);
-		}
+		FSkinCacheRWBuffer& RWBuffer = PositionBuffers[i];
+		InBuffersToTransition.Remove(&RWBuffer);
+		
 		if (auto TangentBuffer = GetTangentBuffer())
 		{
-			if (TangentBuffer->UAV.IsValid())
-			{
-				InBuffersToTransition.Remove(TangentBuffer->UAV);
-			}
+			InBuffersToTransition.Remove(TangentBuffer);
 		}
 		if (auto IntermediateTangentBuffer = GetIntermediateTangentBuffer())
 		{
-			if (IntermediateTangentBuffer->UAV.IsValid())
-			{
-				InBuffersToTransition.Remove(IntermediateTangentBuffer->UAV);
-			}
+			InBuffersToTransition.Remove(IntermediateTangentBuffer);
 		}
 	}
 }
diff --git a/Engine/Source/Runtime/Engine/Public/GPUSkinCache.h b/Engine/Source/Runtime/Engine/Public/GPUSkinCache.h
--- a/Engine/Source/Runtime/Engine/Public/GPUSkinCache.h
+++ b/Engine/Source/Runtime/Engine/Public/GPUSkinCache.h
@@ -212,6 +212,26 @@ public:
 		NUM_BUFFERS = 2,
 	};
 
+	struct FSkinCacheRWBuffer
+	{
+		FRWBuffer	Buffer;
+		ERHIAccess	AccessState = ERHIAccess::Unknown;	// Keep track of current access state
+
+		void Release()
+		{
+			Buffer.Release();
+			AccessState = ERHIAccess::Unknown;
+		}
+
+		// Update the access state and return transition info
+		FRHITransitionInfo UpdateAccessState(ERHIAccess NewState)
+		{
+			ERHIAccess OldState = AccessState;
+			AccessState = NewState;
+			return FRHITransitionInfo(Buffer.UAV.GetReference(), OldState, AccessState);
+		}
+	};
+
 	struct FRWBuffersAllocation
 	{
 		friend struct FRWBufferTracker;
@@ -221,21 +241,25 @@ public:
 		{
 			for (int32 Index = 0; Index < NUM_BUFFERS; ++Index)
 			{
-				PositionBuffers[Index].Initialize(TEXT("SkinCachePositions"), PosBufferBytesPerElement, NumVertices * 3, PF_R32_FLOAT, BUF_Static);
+				PositionBuffers[Index].Buffer.Initialize(TEXT("SkinCachePositions"), PosBufferBytesPerElement, NumVertices * 3, PF_R32_FLOAT, BUF_Static);
+				PositionBuffers[Index].AccessState = ERHIAccess::Unknown;
 			}
 			if (WithTangents)
 			{
-				Tangents.Initialize(TEXT("SkinCacheTangents"), TangentBufferBytesPerElement, NumVertices * 2, PF_R16G16B16A16_SNORM, BUF_Static);
+				Tangents.Buffer.Initialize(TEXT("SkinCacheTangents"), TangentBufferBytesPerElement, NumVertices * 2, PF_R16G16B16A16_SNORM, BUF_Static);
+				Tangents.AccessState = ERHIAccess::Unknown;
 				if (FGPUSkinCache::UseIntermediateTangents())
 				{
-					IntermediateTangents.Initialize(TEXT("SkinCacheIntermediateTangents"), TangentBufferBytesPerElement, NumVertices * 2, PF_R16G16B16A16_SNORM, BUF_Static);
+					IntermediateTangents.Buffer.Initialize(TEXT("SkinCacheIntermediateTangents"), TangentBufferBytesPerElement, NumVertices * 2, PF_R16G16B16A16_SNORM, BUF_Static);
+					IntermediateTangents.AccessState = ERHIAccess::Unknown;
 				}
 			}
 			if (NumTriangles > 0)
 			{
-				IntermediateAccumulatedTangents.Initialize(TEXT("SkinCacheIntermediateAccumulatedTangents"), sizeof(int32), NumTriangles * 3 * FGPUSkinCache::IntermediateAccumBufferNumInts, PF_R32_SINT, BUF_UnorderedAccess);
+				IntermediateAccumulatedTangents.Buffer.Initialize(TEXT("SkinCacheIntermediateAccumulatedTangents"), sizeof(int32), NumTriangles * 3 * FGPUSkinCache::IntermediateAccumBufferNumInts, PF_R32_SINT, BUF_UnorderedAccess);
+				IntermediateAccumulatedTangents.AccessState = ERHIAccess::Unknown;
 				// The UAV must be zero-filled. We leave it zeroed after each round (see RecomputeTangentsPerVertexPass.usf), so this is only needed on when the buffer is first created.
-				RHICmdList.ClearUAVUint(IntermediateAccumulatedTangents.UAV, FUintVector4(0, 0, 0, 0));
+				RHICmdList.ClearUAVUint(IntermediateAccumulatedTangents.Buffer.UAV, FUintVector4(0, 0, 0, 0));
 			}
 		}
 
@@ -274,30 +298,30 @@ public:
 			return CalculateRequiredMemory(NumVertices, WithTangents, NumTriangles);
 		}
 
-		FRWBuffer* GetTangentBuffer()
+		FSkinCacheRWBuffer* GetTangentBuffer()
 		{
 			return WithTangents ? &Tangents : nullptr;
 		}
 
-		FRWBuffer* GetIntermediateTangentBuffer()
+		FSkinCacheRWBuffer* GetIntermediateTangentBuffer()
 		{
 			return WithTangents ? &IntermediateTangents : nullptr;
 		}
 
-		FRWBuffer* GetIntermediateAccumulatedTangentBuffer()
+		FSkinCacheRWBuffer* GetIntermediateAccumulatedTangentBuffer()
 		{
 			return NumTriangles > 0 ? &IntermediateAccumulatedTangents : nullptr;
 		}
 
-		void RemoveAllFromTransitionArray(TSet<FRHIUnorderedAccessView*>& BuffersToTransition);
+		void RemoveAllFromTransitionArray(TSet<FSkinCacheRWBuffer*>& BuffersToTransition);
 
 	private:
 		// Output of the GPU skinning (ie Pos, Normals)
-		FRWBuffer PositionBuffers[NUM_BUFFERS];
+		FSkinCacheRWBuffer PositionBuffers[NUM_BUFFERS];
 
-		FRWBuffer Tangents;
-		FRWBuffer IntermediateTangents;
-		FRWBuffer IntermediateAccumulatedTangents;	// Intermediate buffer used to accumulate results of triangle pass to be passed onto vertex pass
+		FSkinCacheRWBuffer Tangents;
+		FSkinCacheRWBuffer IntermediateTangents;
+		FSkinCacheRWBuffer IntermediateAccumulatedTangents;	// Intermediate buffer used to accumulate results of triangle pass to be passed onto vertex pass
 
 		const uint32 NumVertices;
 		const bool WithTangents;
@@ -331,7 +355,7 @@ public:
 			return Allocation->GetNumBytes();
 		}
 
-		FRWBuffer* Find(const FVertexBufferAndSRV& BoneBuffer, uint32 Revision)
+		FSkinCacheRWBuffer* Find(const FVertexBufferAndSRV& BoneBuffer, uint32 Revision)
 		{
 			for (int32 Index = 0; Index < NUM_BUFFERS; ++Index)
 			{
@@ -344,17 +368,17 @@ public:
 			return nullptr;
 		}
 
-		FRWBuffer* GetTangentBuffer()
+		FSkinCacheRWBuffer* GetTangentBuffer()
 		{
 			return Allocation ? Allocation->GetTangentBuffer() : nullptr;
 		}
 
-		FRWBuffer* GetIntermediateTangentBuffer()
+		FSkinCacheRWBuffer* GetIntermediateTangentBuffer()
 		{
 			return Allocation ? Allocation->GetIntermediateTangentBuffer() : nullptr;
 		}
 
-		FRWBuffer* GetIntermediateAccumulatedTangentBuffer()
+		FSkinCacheRWBuffer* GetIntermediateAccumulatedTangentBuffer()
 		{
 			return Allocation ? Allocation->GetIntermediateAccumulatedTangentBuffer() : nullptr;
 		}
@@ -418,9 +442,10 @@ public:
 	inline ERHIFeatureLevel::Type GetFeatureLevel() const { return FeatureLevel; }
 
 protected:
-	void MakeBufferTransitions(FRHICommandListImmediate& RHICmdList, TArray<FRHIUnorderedAccessView*>& Buffers, ERHIAccess FromState, ERHIAccess ToState);
+	void MakeBufferTransitions(FRHICommandListImmediate& RHICmdList, TArray<FSkinCacheRWBuffer*>& Buffers, ERHIAccess ToState);
+	void GetBufferUAVs(const TArray<FSkinCacheRWBuffer*>& InBuffers, TArray<FRHIUnorderedAccessView*>& OutUAVs);
 
-	TSet<FRHIUnorderedAccessView*> BuffersToTransitionToRead;
+	TSet<FSkinCacheRWBuffer*> BuffersToTransitionToRead;
 #if RHI_RAYTRACING
 	TMap<FRayTracingGeometry*, EAccelerationStructureBuildMode> RayTracingGeometriesToUpdate;
 	uint64 RayTracingGeometryMemoryPendingRelease = 0;
@@ -433,13 +458,13 @@ protected:
 	FRWBuffersAllocation* TryAllocBuffer(uint32 NumVertices, bool WithTangnents, uint32 NumTriangles, FRHICommandListImmediate& RHICmdList);
 	void DoDispatch(FRHICommandListImmediate& RHICmdList);
 	void DoDispatch(FRHICommandListImmediate& RHICmdList, FGPUSkinCacheEntry* SkinCacheEntry, int32 Section, int32 RevisionNumber);
-	void DispatchUpdateSkinTangents(FRHICommandListImmediate& RHICmdList, FGPUSkinCacheEntry* Entry, int32 SectionIndex, FRWBuffer*& StagingBuffer, bool bTrianglePass);
+	void DispatchUpdateSkinTangents(FRHICommandListImmediate& RHICmdList, FGPUSkinCacheEntry* Entry, int32 SectionIndex, FSkinCacheRWBuffer*& StagingBuffer, bool bTrianglePass);
 
 	void PrepareUpdateSkinning(
 		FGPUSkinCacheEntry* Entry, 
 		int32 Section, 
 		uint32 RevisionNumber, 
-		TArray<FRHIUnorderedAccessView*>* OverlappedUAVs
+		TArray<FSkinCacheRWBuffer*>* OverlappedUAVs
 		);
 
 	void DispatchUpdateSkinning(
@@ -460,7 +485,7 @@ protected:
 	bool bShouldBatchDispatches = false;
 
 	// For recompute tangents, holds the data required between compute shaders
-	TArray<FRWBuffer> StagingBuffers;
+	TArray<FSkinCacheRWBuffer> StagingBuffers;
 	int32 CurrentStagingBufferIndex;
 
 	ERHIFeatureLevel::Type FeatureLevel;
