Fixed Starta compiler side not taking into account Haziness when evlauating the per pixel cost.

diff --git a/Engine/Source/Runtime/Engine/Classes/Materials/MaterialExpressionStrata.h b/Engine/Source/Runtime/Engine/Classes/Materials/MaterialExpressionStrata.h
--- a/Engine/Source/Runtime/Engine/Classes/Materials/MaterialExpressionStrata.h
+++ b/Engine/Source/Runtime/Engine/Classes/Materials/MaterialExpressionStrata.h
@@ -176,6 +176,7 @@ class UMaterialExpressionStrataSlabBSDF : public UMaterialExpressionStrataBSDF
 	bool HasEdgeColor() const;
 	bool HasThinFilm() const;
 	bool HasFuzz() const;
+	bool HasHaziness() const;
 	bool HasSSS() const;
 	bool HasSSSProfile() const;
 	bool HasDMFPPluggedIn() const;
diff --git a/Engine/Source/Runtime/Engine/Private/Materials/MaterialExpressions.cpp b/Engine/Source/Runtime/Engine/Private/Materials/MaterialExpressions.cpp
--- a/Engine/Source/Runtime/Engine/Private/Materials/MaterialExpressions.cpp
+++ b/Engine/Source/Runtime/Engine/Private/Materials/MaterialExpressions.cpp
@@ -19925,6 +19925,7 @@ int32 UMaterialExpressionStrataSlabBSDF::Compile(class FMaterialCompiler* Compil
 	const bool bHasEdgeColor = HasEdgeColor(); // This accounts for EdgeColor and also F90 when the non metalness worfklow is selected.
 	const bool bHasThinFilm = HasThinFilm();
 	const bool bHasFuzz = HasFuzz();
+	const bool bHasHaziness = HasHaziness();
 	const bool bHasDMFPPluggedIn = HasDMFPPluggedIn();
 
 	int32 SSSProfileCodeChunk = INDEX_NONE;
@@ -19970,7 +19971,7 @@ int32 UMaterialExpressionStrataSlabBSDF::Compile(class FMaterialCompiler* Compil
 		NormalCodeChunk,
 		TangentCodeChunk,
 		Compiler->GetStrataSharedNormalIndexMacro(NewRegisteredSharedNormal));
-	StrataCompilationInfoCreateSingleBSDFMaterial(Compiler, OutputCodeChunk, NewRegisteredSharedNormal, STRATA_BSDF_TYPE_SLAB, bHasSSS, bHasDMFPPluggedIn, bHasEdgeColor, bHasThinFilm, bHasFuzz);
+	StrataCompilationInfoCreateSingleBSDFMaterial(Compiler, OutputCodeChunk, NewRegisteredSharedNormal, STRATA_BSDF_TYPE_SLAB, bHasSSS, bHasDMFPPluggedIn, bHasEdgeColor, bHasThinFilm, bHasFuzz, bHasHaziness);
 
 	return OutputCodeChunk;
 }
@@ -20383,6 +20384,11 @@ bool UMaterialExpressionStrataSlabBSDF::HasFuzz() const
 	return FuzzAmount.IsConnected();
 }
 
+bool UMaterialExpressionStrataSlabBSDF::HasHaziness() const
+{
+	return Haziness.IsConnected();
+}
+
 bool UMaterialExpressionStrataSlabBSDF::HasAnisotropy() const
 {
 	return Anisotropy.IsConnected();
diff --git a/Engine/Source/Runtime/Engine/Private/Materials/StrataMaterial.cpp b/Engine/Source/Runtime/Engine/Private/Materials/StrataMaterial.cpp
--- a/Engine/Source/Runtime/Engine/Private/Materials/StrataMaterial.cpp
+++ b/Engine/Source/Runtime/Engine/Private/Materials/StrataMaterial.cpp
@@ -65,7 +65,7 @@ FStrataRegisteredSharedNormal StrataCompilationInfoCreateSharedNormal(FMaterialC
 
 void StrataCompilationInfoCreateSingleBSDFMaterial(FMaterialCompiler* Compiler, int32 CodeChunk,
 	const FStrataRegisteredSharedNormal& RegisteredSharedNormal,
-	uint8 BSDFType, bool bHasSSS, bool bHasDMFPPluggedIn, bool bHasEdgeColor, bool bHasThinFilm, bool bHasFuzz)
+	uint8 BSDFType, bool bHasSSS, bool bHasDMFPPluggedIn, bool bHasEdgeColor, bool bHasThinFilm, bool bHasFuzz, bool bHasHaziness)
 {
 	FStrataMaterialCompilationInfo StrataInfo;
 	StrataInfo.LayerCount = 1;
@@ -77,6 +77,7 @@ void StrataCompilationInfoCreateSingleBSDFMaterial(FMaterialCompiler* Compiler,
 	StrataInfo.Layers[0].BSDFs[0].bHasEdgeColor = bHasEdgeColor;
 	StrataInfo.Layers[0].BSDFs[0].bHasThinFilm = bHasThinFilm;
 	StrataInfo.Layers[0].BSDFs[0].bHasFuzz = bHasFuzz;
+	StrataInfo.Layers[0].BSDFs[0].bHasHaziness = bHasHaziness;
 	UpdateTotalBSDFCount(Compiler, StrataInfo);
 	Compiler->StrataCompilationInfoRegisterCodeChunk(CodeChunk, StrataInfo);
 }
@@ -145,6 +146,7 @@ FStrataMaterialCompilationInfo StrataCompilationInfoAddParamBlend(FMaterialCompi
 	NewBSDF.bHasEdgeColor		|=	OtherBSDF.bHasEdgeColor;
 	NewBSDF.bHasThinFilm		|=	OtherBSDF.bHasThinFilm;
 	NewBSDF.bHasFuzz			|=	OtherBSDF.bHasFuzz;
+	NewBSDF.bHasHaziness		|=	OtherBSDF.bHasHaziness;
 	UpdateTotalBSDFCount(Compiler, StrataInfo);
 	return StrataInfo;
 }
@@ -174,6 +176,7 @@ FStrataMaterialCompilationInfo StrataCompilationInfoHorizontalMixingParamBlend(F
 	NewBSDF.bHasEdgeColor		|=	OtherBSDF.bHasEdgeColor;
 	NewBSDF.bHasThinFilm		|=	OtherBSDF.bHasThinFilm;
 	NewBSDF.bHasFuzz			|=	OtherBSDF.bHasFuzz;
+	NewBSDF.bHasHaziness		|=	OtherBSDF.bHasHaziness;
 	UpdateTotalBSDFCount(Compiler, StrataInfo);
 	return StrataInfo;
 }
@@ -218,6 +221,7 @@ FStrataMaterialCompilationInfo StrataCompilationInfoVerticalLayeringParamBlend(F
 //	NewBSDF.bHasDMFPPluggedIn	= TopBSDF.bHasDMFPPluggedIn;	// Idem
 	NewBSDF.bHasEdgeColor		|= TopBSDF.bHasEdgeColor;		// We keep the union of both, even though it will be hard to get a perfect match.
 	NewBSDF.bHasFuzz			|= TopBSDF.bHasFuzz;			// Idem
+	NewBSDF.bHasHaziness		|= TopBSDF.bHasHaziness;		// Idem
 	NewBSDF.bHasThinFilm		 = TopBSDF.bHasThinFilm;		// We only keep thin film from the top layer, because its color is otherwise not controlable
 
 	UpdateTotalBSDFCount(Compiler, StrataInfo);
@@ -373,7 +377,7 @@ FStrataMaterialAnalysisResult StrataCompilationInfoMaterialAnalysis(FMaterialCom
 
 				Result.RequestedByteCount += UintByteSize;
 				Result.RequestedByteCount += UintByteSize;
-				if (BSDF.bHasEdgeColor || BSDF.bHasThinFilm)
+				if (BSDF.bHasEdgeColor || BSDF.bHasThinFilm || BSDF.bHasHaziness)
 				{
 					Result.RequestedByteCount += UintByteSize;
 				}
diff --git a/Engine/Source/Runtime/Engine/Private/Materials/StrataMaterial.h b/Engine/Source/Runtime/Engine/Private/Materials/StrataMaterial.h
--- a/Engine/Source/Runtime/Engine/Private/Materials/StrataMaterial.h
+++ b/Engine/Source/Runtime/Engine/Private/Materials/StrataMaterial.h
@@ -26,6 +26,7 @@ struct FStrataMaterialCompilationInfo
 		bool  bHasEdgeColor;
 		bool  bHasThinFilm;
 		bool  bHasFuzz;
+		bool  bHasHaziness;
 	};
 
 	struct FLayer
@@ -55,7 +56,7 @@ FStrataRegisteredSharedNormal StrataCompilationInfoCreateSharedNormal(FMaterialC
 void  StrataCompilationInfoCreateNullBSDF(FMaterialCompiler* Compiler, int32 CodeChunk);
 void  StrataCompilationInfoCreateSingleBSDFMaterial(FMaterialCompiler* Compiler, int32 CodeChunk, 
 	const FStrataRegisteredSharedNormal& RegisteredSharedNormal,
-	uint8 BSDFType, bool bHasSSS = false, bool bHasDMFPPluggedIn = false, bool bHasEdgeColor = false, bool bHasThinFilm = false, bool bHasFuzz = false);
+	uint8 BSDFType, bool bHasSSS = false, bool bHasDMFPPluggedIn = false, bool bHasEdgeColor = false, bool bHasThinFilm = false, bool bHasFuzz = false, bool bHasHaziness = false);
 
 FStrataMaterialCompilationInfo StrataCompilationInfoWeight(FMaterialCompiler* Compiler, const FStrataMaterialCompilationInfo& A);
 FStrataMaterialCompilationInfo StrataCompilationInfoAdd(FMaterialCompiler* Compiler, const FStrataMaterialCompilationInfo& A, const FStrataMaterialCompilationInfo& B);
