U5 - On behalf of Brandon.Dawson: change warning to an ensure

diff --git a/Engine/Source/Runtime/Engine/Private/Texture.cpp b/Engine/Source/Runtime/Engine/Private/Texture.cpp
--- a/Engine/Source/Runtime/Engine/Private/Texture.cpp
+++ b/Engine/Source/Runtime/Engine/Private/Texture.cpp
@@ -138,7 +138,7 @@ const FTextureResource* UTexture::GetResource() const
 		return PrivateResourceRenderThread;
 	}
 
-	UE_LOG(LogTexture, Error, TEXT("Attempted to access a texture resource from an unkown thread."));
+	ensureMsgf(false, TEXT("Attempted to access a texture resource from an unkown thread."));
 	return nullptr;
 }
 
@@ -153,7 +153,7 @@ FTextureResource* UTexture::GetResource()
 		return PrivateResourceRenderThread;
 	}
 
-	UE_LOG(LogTexture, Error, TEXT("Attempted to access a texture resource from an unkown thread."));
+	ensureMsgf(false, TEXT("Attempted to access a texture resource from an unkown thread."));
 	return nullptr;
 }
 
