fix CVar in config files writted as "...=True" turns out to be 0 in mobile editor preview
fix CVar with ECVF_Cheat flag should not be changed by DeviceProfile

diff --git a/Engine/Source/Runtime/Core/Public/Misc/ConfigCacheIni.h b/Engine/Source/Runtime/Core/Public/Misc/ConfigCacheIni.h
--- a/Engine/Source/Runtime/Core/Public/Misc/ConfigCacheIni.h
+++ b/Engine/Source/Runtime/Core/Public/Misc/ConfigCacheIni.h
@@ -1222,3 +1222,8 @@ CORE_API void DumpRecordedConfigReadsFromIni();
  * Helper function to clean up config read history
  */
 CORE_API void DeleteRecordedConfigReadsFromIni();
+
+/**
+ * Helper function to deal with "True","False","Yes","No","On","Off"
+ */
+CORE_API const TCHAR* ConvertValueFromHumanFriendlyValue(const TCHAR* Value);
\ No newline at end of file
diff --git a/Engine/Source/Runtime/Engine/Private/DeviceProfiles/DeviceProfileManager.cpp b/Engine/Source/Runtime/Engine/Private/DeviceProfiles/DeviceProfileManager.cpp
--- a/Engine/Source/Runtime/Engine/Private/DeviceProfiles/DeviceProfileManager.cpp
+++ b/Engine/Source/Runtime/Engine/Private/DeviceProfiles/DeviceProfileManager.cpp
@@ -881,8 +881,12 @@ void UDeviceProfileManager::SetOverrideDeviceProfile(UDeviceProfile* DeviceProfi
 				FString OldValue = CVar->GetString();
 				PushedSettings.Add(Pair.Key, OldValue);
 
-				// set the cvar to the new value, with same priority that it was before (SetByMask means current priority)
-				CVar->SetWithCurrentPriority(*Pair.Value);
+				// cheat CVar can only be set in ConsoleVariables.ini
+				if (!CVar->TestFlags(EConsoleVariableFlags::ECVF_Cheat))
+				{
+					// set the cvar to the new value, with same priority that it was before (SetByMask means current priority)
+					CVar->SetWithCurrentPriority(ConvertValueFromHumanFriendlyValue(*Pair.Value));
+				}
 			}
 		}
 #else
