Make the upload to ISM GPU-Buffers deferred if Nanite is enabled.
 - Nanite ISM / HISM uses a different proxy and does not need this data (it is fetched from GPU-Scene).

diff --git a/Engine/Source/Runtime/Engine/Classes/Components/StaticMeshComponent.h b/Engine/Source/Runtime/Engine/Classes/Components/StaticMeshComponent.h
--- a/Engine/Source/Runtime/Engine/Classes/Components/StaticMeshComponent.h
+++ b/Engine/Source/Runtime/Engine/Classes/Components/StaticMeshComponent.h
@@ -690,6 +690,8 @@ protected:
 		return true;
 	}
 
+	bool ShouldCreateNaniteProxy() const;	
+
 public:
 
 	void ReleaseResources();
diff --git a/Engine/Source/Runtime/Engine/Classes/Engine/InstancedStaticMesh.h b/Engine/Source/Runtime/Engine/Classes/Engine/InstancedStaticMesh.h
--- a/Engine/Source/Runtime/Engine/Classes/Engine/InstancedStaticMesh.h
+++ b/Engine/Source/Runtime/Engine/Classes/Engine/InstancedStaticMesh.h
@@ -60,7 +60,7 @@ class FStaticMeshInstanceBuffer : public FRenderResource
 public:
 
 	/** Default constructor. */
-	FStaticMeshInstanceBuffer(ERHIFeatureLevel::Type InFeatureLevel, bool InRequireCPUAccess);
+	FStaticMeshInstanceBuffer(ERHIFeatureLevel::Type InFeatureLevel, bool InRequireCPUAccess, bool bDeferGPUUploadIn);
 
 	/** Destructor. */
 	~FStaticMeshInstanceBuffer();
@@ -113,29 +113,60 @@ public:
 
 	void BindInstanceVertexBuffer(const class FVertexFactory* VertexFactory, struct FInstancedStaticMeshDataType& InstancedStaticMeshData) const;
 
+	/**
+	 * Call to flush any pending GPU data copies, if bFlushToGPUPending is false it does nothing. Should be called by the Proxy on the render thread
+	 * for example in CreateRenderThreadResources().
+	 */
+	void FlushGPUUpload();
 public:
 	/** The vertex data storage type */
 	TSharedPtr<FStaticMeshInstanceData, ESPMode::ThreadSafe> InstanceData;
 
-	/** Keep CPU copy of instance data*/
+	/** Keep CPU copy of instance data */
 	bool RequireCPUAccess;
 
 	FBufferRHIRef GetInstanceOriginBuffer()
 	{
+		check(!bFlushToGPUPending);
 		return InstanceOriginBuffer.VertexBufferRHI;
 	}
 
 	FBufferRHIRef GetInstanceTransformBuffer()
 	{
+		check(!bFlushToGPUPending);
 		return InstanceTransformBuffer.VertexBufferRHI;
 	}
 
 	FBufferRHIRef GetInstanceLightmapBuffer()
 	{
+		check(!bFlushToGPUPending);
 		return InstanceLightmapBuffer.VertexBufferRHI;
 	}
 
+	/**
+	 * Set flush to GPU as pending if the bDeferGPUUpload flag is true.
+	 * Returns bDeferGPUUpload (so if it returns false, the update should be done at once).
+	 */
+	bool CondSetFlushToGPUPending()
+	{
+		if (bDeferGPUUpload)
+		{
+			bFlushToGPUPending = true;
+			return true;
+		}
+		return bDeferGPUUpload;
+	}
 private:
+
+	/** Defer GPU Upload until we can know if it is needed (that is on the render thread) */
+	bool bDeferGPUUpload;
+
+	/** If true, then we have updates to the host data not yet committed to the GPU. This in turn means
+	 * that bDeferGPUUpload is true, and the Proxy is expected to either call FlushGPUUpload() OR never 
+	 * use the instance data buffers (either is fine).
+	 */
+	bool bFlushToGPUPending;
+
 	class FInstanceOriginBuffer : public FVertexBuffer
 	{
 		virtual FString GetFriendlyName() const override { return TEXT("FInstanceOriginBuffer"); }
@@ -352,8 +383,7 @@ struct FInstanceUpdateCmdBuffer;
 struct FPerInstanceRenderData
 {
 	// Should be always constructed on main thread
-	FPerInstanceRenderData(FStaticMeshInstanceData& Other, ERHIFeatureLevel::Type InFeaureLevel, bool InRequireCPUAccess);
-	FPerInstanceRenderData(FStaticMeshInstanceData& Other, ERHIFeatureLevel::Type InFeaureLevel, bool InRequireCPUAccess, FBox InBounds, bool bTrack);
+	FPerInstanceRenderData(FStaticMeshInstanceData& Other, ERHIFeatureLevel::Type InFeaureLevel, bool InRequireCPUAccess, FBox InBounds, bool bTrack, bool bDeferGPUUploadIn);
 	~FPerInstanceRenderData();
 
 	/**
diff --git a/Engine/Source/Runtime/Engine/Private/Components/StaticMeshComponent.cpp b/Engine/Source/Runtime/Engine/Private/Components/StaticMeshComponent.cpp
--- a/Engine/Source/Runtime/Engine/Private/Components/StaticMeshComponent.cpp
+++ b/Engine/Source/Runtime/Engine/Private/Components/StaticMeshComponent.cpp
@@ -1486,6 +1486,18 @@ void UStaticMeshComponent::UpdatePreCulledData(int32 LODIndex, const TArray<uint
 	}
 }
 
+bool UStaticMeshComponent::ShouldCreateNaniteProxy() const
+{
+	// Whether or not to allow Nanite for this component
+#if WITH_EDITORONLY_DATA
+	const bool bAllowNanite = !bDisplayNaniteProxyMesh;
+#else
+	const bool bAllowNanite = true;
+#endif
+
+	return bAllowNanite && UseNanite(GetScene()->GetShaderPlatform()) && GetStaticMesh()->GetRenderData()->NaniteResources.PageStreamingStates.Num();
+}
+
 void UStaticMeshComponent::ReleaseResources()
 {
 	for(int32 LODIndex = 0;LODIndex < LODData.Num();LODIndex++)
diff --git a/Engine/Source/Runtime/Engine/Private/HierarchicalInstancedStaticMesh.cpp b/Engine/Source/Runtime/Engine/Private/HierarchicalInstancedStaticMesh.cpp
--- a/Engine/Source/Runtime/Engine/Private/HierarchicalInstancedStaticMesh.cpp
+++ b/Engine/Source/Runtime/Engine/Private/HierarchicalInstancedStaticMesh.cpp
@@ -3169,8 +3169,7 @@ FPrimitiveSceneProxy* UHierarchicalInstancedStaticMeshComponent::CreateSceneProx
 		ProxySize = PerInstanceRenderData->ResourceSize;
 		INC_DWORD_STAT_BY(STAT_FoliageInstanceBuffers, ProxySize);
 
-		// TODO: Abstract with a common helper
-		if (UseNanite(GetScene()->GetShaderPlatform()) && GetStaticMesh()->GetRenderData()->NaniteResources.PageStreamingStates.Num())
+		if (ShouldCreateNaniteProxy())
 		{
 			return ::new Nanite::FSceneProxy(this);
 		}
diff --git a/Engine/Source/Runtime/Engine/Private/InstancedStaticMesh.cpp b/Engine/Source/Runtime/Engine/Private/InstancedStaticMesh.cpp
--- a/Engine/Source/Runtime/Engine/Private/InstancedStaticMesh.cpp
+++ b/Engine/Source/Runtime/Engine/Private/InstancedStaticMesh.cpp
@@ -311,9 +311,11 @@ void FInstanceUpdateCmdBuffer::Reset()
 	NumEdits = 0;
 }
 
-FStaticMeshInstanceBuffer::FStaticMeshInstanceBuffer(ERHIFeatureLevel::Type InFeatureLevel, bool InRequireCPUAccess)
+FStaticMeshInstanceBuffer::FStaticMeshInstanceBuffer(ERHIFeatureLevel::Type InFeatureLevel, bool InRequireCPUAccess, bool bDeferGPUUploadIn)
 	: FRenderResource(InFeatureLevel)
 	, RequireCPUAccess(InRequireCPUAccess)
+	, bDeferGPUUpload(bDeferGPUUploadIn)
+	, bFlushToGPUPending(false)
 {
 }
 
@@ -411,7 +413,10 @@ void FStaticMeshInstanceBuffer::UpdateFromCommandBuffer_RenderThread(FInstanceUp
 		}
 	}
 
-	UpdateRHI();
+	if (!CondSetFlushToGPUPending())
+	{
+		UpdateRHI();
+	}
 }
 
 /**
@@ -571,6 +576,24 @@ void FStaticMeshInstanceBuffer::BindInstanceVertexBuffer(const class FVertexFact
 	}
 }
 
+void FStaticMeshInstanceBuffer::FlushGPUUpload()
+{
+	if (bFlushToGPUPending)
+	{
+		check(bDeferGPUUpload);
+
+		if (!IsInitialized())
+		{
+			InitResource();
+		}
+		else
+		{
+			UpdateRHI();
+		}
+		bFlushToGPUPending = false;
+	}
+}
+
 
 void FStaticMeshInstanceData::Serialize(FArchive& Ar)
 {	
@@ -890,6 +913,8 @@ void FInstancedStaticMeshRenderData::InitVertexFactories()
 	ENQUEUE_RENDER_COMMAND(InstancedStaticMeshRenderData_InitVertexFactories)(
 		[this, LightMapCoordinateIndex](FRHICommandListImmediate& RHICmdList)
 	{
+		PerInstanceRenderData->InstanceBuffer.FlushGPUUpload();
+
 		for (int32 LODIndex = 0; LODIndex < VertexFactories.Num(); LODIndex++)
 		{
 			const FStaticMeshLODResources* RenderData = &LODModels[LODIndex];
@@ -925,29 +950,20 @@ void FInstancedStaticMeshRenderData::InitVertexFactories()
 	});
 }
 
-FPerInstanceRenderData::FPerInstanceRenderData(FStaticMeshInstanceData& Other, ERHIFeatureLevel::Type InFeaureLevel, bool InRequireCPUAccess)
-	: ResourceSize(InRequireCPUAccess ? Other.GetResourceSize() : 0)
-	, InstanceBuffer(InFeaureLevel, InRequireCPUAccess)
-	, bTrackBounds(false)
-	, bBoundsTransformsDirty(true)
-{
-	InstanceBuffer.InitFromPreallocatedData(Other);
-	InstanceBuffer_GameThread = InstanceBuffer.InstanceData;
-
-	BeginInitResource(&InstanceBuffer);
-}
 
-FPerInstanceRenderData::FPerInstanceRenderData(FStaticMeshInstanceData& Other, ERHIFeatureLevel::Type InFeaureLevel, bool InRequireCPUAccess, FBox InBounds, bool bTrack)
+FPerInstanceRenderData::FPerInstanceRenderData(FStaticMeshInstanceData& Other, ERHIFeatureLevel::Type InFeaureLevel, bool InRequireCPUAccess, FBox InBounds, bool bTrack, bool bDeferGPUUploadIn)
 	: ResourceSize(InRequireCPUAccess ? Other.GetResourceSize() : 0)
-	, InstanceBuffer(InFeaureLevel, InRequireCPUAccess)
+	, InstanceBuffer(InFeaureLevel, InRequireCPUAccess, bDeferGPUUploadIn)
 	, InstanceLocalBounds(InBounds)
 	, bTrackBounds(bTrack)
 	, bBoundsTransformsDirty(true)
 {
 	InstanceBuffer.InitFromPreallocatedData(Other);
 	InstanceBuffer_GameThread = InstanceBuffer.InstanceData;
-
-	BeginInitResource(&InstanceBuffer);
+	if (!InstanceBuffer.CondSetFlushToGPUPending())
+	{
+		BeginInitResource(&InstanceBuffer);
+	}
 	UpdateBoundsTransforms_Concurrent();
 }
 
@@ -976,7 +992,10 @@ void FPerInstanceRenderData::UpdateFromPreallocatedData(FStaticMeshInstanceData&
 		[InInstanceBufferDataPtr, InInstanceBuffer, this](FRHICommandListImmediate& RHICmdList)
 		{
 			InInstanceBuffer->InstanceData = InInstanceBufferDataPtr;
-			InInstanceBuffer->UpdateRHI();
+			if (!InInstanceBuffer->CondSetFlushToGPUPending())
+			{
+				InInstanceBuffer->UpdateRHI();
+			}
 			UpdateBoundsTransforms_Concurrent();
 		}
 	);
@@ -1265,6 +1284,13 @@ void FInstancedStaticMeshSceneProxy::SetupProxy(UInstancedStaticMeshComponent* I
 void FInstancedStaticMeshSceneProxy::CreateRenderThreadResources()
 {
 	FStaticMeshSceneProxy::CreateRenderThreadResources();
+	
+	// Flush upload of GPU data for ISM/HISM
+	if (ensure(InstancedRenderData.PerInstanceRenderData.IsValid()))
+	{
+		FStaticMeshInstanceBuffer& InstanceBuffer = InstancedRenderData.PerInstanceRenderData->InstanceBuffer;
+		InstanceBuffer.FlushGPUUpload();
+	}
 
 #if GPUCULL_TODO
 	if (UseGPUScene(GetScene().GetShaderPlatform(), GetScene().GetFeatureLevel()))
@@ -1953,8 +1979,7 @@ FPrimitiveSceneProxy* UInstancedStaticMeshComponent::CreateSceneProxy()
 		
 		ProxySize = PerInstanceRenderData->ResourceSize;
 
-		// TODO: Abstract with a common helper
-		if (UseNanite(GetScene()->GetShaderPlatform()) && GetStaticMesh()->GetRenderData()->NaniteResources.PageStreamingStates.Num())
+		if (ShouldCreateNaniteProxy())
 		{
 			return ::new Nanite::FSceneProxy(this);
 		}
@@ -3511,9 +3536,15 @@ void UInstancedStaticMeshComponent::InitPerInstanceRenderData(bool InitializeFro
 	FBox LocalBounds;
 	GetLocalBounds(LocalBounds.Min, LocalBounds.Max);
 
+	// If Nanite is used, we should defer the upload to GPU as the Nanite proxy simply will skip this step.
+	// We can't just disable the upload, because at this point we can't know whether the Nanite proxy will be created in the end
+	// this depends on the static mesh which may still be compiling/loading.
+	// TODO: Perhaps make this specific to if this ISM actually has Nanite (if this can be detected reliably at this point) 
+	const bool bDeferGPUUpload = UseNanite(GetFeatureLevelShaderPlatform(FeatureLevel));
+
 	if (InSharedInstanceBufferData != nullptr)
 	{
-		PerInstanceRenderData = MakeShareable(new FPerInstanceRenderData(*InSharedInstanceBufferData, FeatureLevel, KeepInstanceBufferCPUAccess, LocalBounds, bTrackBounds));
+		PerInstanceRenderData = MakeShareable(new FPerInstanceRenderData(*InSharedInstanceBufferData, FeatureLevel, KeepInstanceBufferCPUAccess, LocalBounds, bTrackBounds, bDeferGPUUpload));
 	}
 	else
 	{
@@ -3527,7 +3558,7 @@ void UInstancedStaticMeshComponent::InitPerInstanceRenderData(bool InitializeFro
 			BuildRenderData(InstanceBufferData, HitProxies);
 		}
 		
-		PerInstanceRenderData = MakeShareable(new FPerInstanceRenderData(InstanceBufferData, FeatureLevel, KeepInstanceBufferCPUAccess, LocalBounds, bTrackBounds));
+		PerInstanceRenderData = MakeShareable(new FPerInstanceRenderData(InstanceBufferData, FeatureLevel, KeepInstanceBufferCPUAccess, LocalBounds, bTrackBounds, bDeferGPUUpload));
 		PerInstanceRenderData->HitProxies = MoveTemp(HitProxies);
 	}
 }
diff --git a/Engine/Source/Runtime/Engine/Private/StaticMeshRender.cpp b/Engine/Source/Runtime/Engine/Private/StaticMeshRender.cpp
--- a/Engine/Source/Runtime/Engine/Private/StaticMeshRender.cpp
+++ b/Engine/Source/Runtime/Engine/Private/StaticMeshRender.cpp
@@ -2310,15 +2310,7 @@ FPrimitiveSceneProxy* UStaticMeshComponent::CreateSceneProxy()
 		return nullptr;
 	}
 
-	// Whether or not to allow Nanite for this component
-#if WITH_EDITORONLY_DATA
-	const bool bAllowNanite = !bDisplayNaniteProxyMesh;
-#else
-	const bool bAllowNanite = true;
-#endif
-
-	// TODO: Abstract with a common helper
-	if (bAllowNanite && UseNanite(GetScene()->GetShaderPlatform()) && GetStaticMesh()->GetRenderData()->NaniteResources.PageStreamingStates.Num())
+	if (ShouldCreateNaniteProxy())
 	{
 		LLM_SCOPE(ELLMTag::StaticMesh);
 		return ::new Nanite::FSceneProxy(this);
