Removed ispc::LerpCurves as it doesn't match FBaseBlendedCurve::LerpTo

diff --git a/Engine/Source/Runtime/Engine/Private/Components/SkeletalMeshComponent.cpp b/Engine/Source/Runtime/Engine/Private/Components/SkeletalMeshComponent.cpp
--- a/Engine/Source/Runtime/Engine/Private/Components/SkeletalMeshComponent.cpp
+++ b/Engine/Source/Runtime/Engine/Private/Components/SkeletalMeshComponent.cpp
@@ -3897,24 +3897,8 @@ void USkeletalMeshComponent::ParallelDuplicateAndInterpolate(FAnimationEvaluatio
 			FAnimationRuntime::LerpBoneTransforms(InAnimEvaluationContext.BoneSpaceTransforms, InAnimEvaluationContext.CachedBoneSpaceTransforms, Alpha, RequiredBones);
 			FillComponentSpaceTransforms(InAnimEvaluationContext.SkeletalMesh, InAnimEvaluationContext.BoneSpaceTransforms, InAnimEvaluationContext.ComponentSpaceTransforms);
 
-			if (INTEL_ISPC)
-			{
-#if INTEL_ISPC
-				ispc::LerpCurves(
-					InAnimEvaluationContext.Curve.CurveWeights.GetData(),
-					InAnimEvaluationContext.Curve.ValidCurveWeights.GetData(),
-					InAnimEvaluationContext.CachedCurve.CurveWeights.GetData(),
-					InAnimEvaluationContext.CachedCurve.ValidCurveWeights.GetData(),
-					InAnimEvaluationContext.Curve.CurveWeights.Num(),
-					Alpha
-				);
-#endif
-			}
-			else
-			{
-				// interpolate curve
-				InAnimEvaluationContext.Curve.LerpTo(InAnimEvaluationContext.CachedCurve, Alpha);
-			}
+			// interpolate curve
+			InAnimEvaluationContext.Curve.LerpTo(InAnimEvaluationContext.CachedCurve, Alpha);
 
 			UE::Anim::Attributes::InterpolateAttributes(InAnimEvaluationContext.CustomAttributes, InAnimEvaluationContext.CachedCustomAttributes, Alpha);
 		}
diff --git a/Engine/Source/Runtime/Engine/Private/Components/SkeletalMeshComponent.ispc b/Engine/Source/Runtime/Engine/Private/Components/SkeletalMeshComponent.ispc
--- a/Engine/Source/Runtime/Engine/Private/Components/SkeletalMeshComponent.ispc
+++ b/Engine/Source/Runtime/Engine/Private/Components/SkeletalMeshComponent.ispc
@@ -34,26 +34,3 @@ export void FillComponentSpaceTransforms(uniform FTransform ComponentSpaceData[]
 		ComponentSpaceData[BoneIndex] = SpaceBase;
 	}
 }
-
-static inline float Lerp(float A, float B, float Alpha)
-{
-	return (1.0f-Alpha)*A + Alpha*B;
-}
-
-export void LerpCurves(uniform float CurveWeights[],
-						uniform uint32 ValidCurveWeights[],
-						const uniform float CachedCurveWeights[],
-						const uniform uint32 CachedValidCurveWeights[],
-						const uniform int CurveLength,
-						const uniform float Alpha)
-{
-	foreach (CurveId = 0 ... CurveLength)
-	{
-		CurveWeights[CurveId] = Lerp(CurveWeights[CurveId], CachedCurveWeights[CurveId], Alpha);
-	}
-
-	foreach (CurveSet = 0 ... CurveLength / 32)
-	{
-		ValidCurveWeights[CurveSet] |= CachedValidCurveWeights[CurveSet];
-	}
-}
