CIS: Bad merge on #if directive

diff --git a/Engine/Source/Runtime/Engine/Private/PhysicsEngine/ConstraintInstance.cpp b/Engine/Source/Runtime/Engine/Private/PhysicsEngine/ConstraintInstance.cpp
--- a/Engine/Source/Runtime/Engine/Private/PhysicsEngine/ConstraintInstance.cpp
+++ b/Engine/Source/Runtime/Engine/Private/PhysicsEngine/ConstraintInstance.cpp
@@ -192,7 +192,6 @@ void FConstraintProfileProperties::UpdatePlasticity_AssumesLocked(const FPhysics
 	const float LinearPlasticityLimit = bLinearPlasticity ? LinearPlasticityThreshold : FLT_MAX;
 	const float AngularPlasticityLimit = bAngularPlasticity ? FMath::DegreesToRadians(AngularPlasticityThreshold) : MAX_FLT;
 
-#if WITH_CHAOS
 	FPhysicsInterface::SetPlasticityLimits_AssumesLocked(InConstraintRef, LinearPlasticityLimit, AngularPlasticityLimit);
 #endif
 }
