Fixes bug where new level script blueprint doesn't run if not compiled when first saved. The core issue is that the Level's LevelScriptActor is not created at save time, and not generated after the level is loaded.

In this fix, we attempt to detect the case after loading the level in the editor, and spawn the LevelScriptActor at this time. The respawn code is moved to its own function so it can be used for multiple cases.

diff --git a/Engine/Source/Runtime/Engine/Classes/Engine/Level.h b/Engine/Source/Runtime/Engine/Classes/Engine/Level.h
--- a/Engine/Source/Runtime/Engine/Classes/Engine/Level.h
+++ b/Engine/Source/Runtime/Engine/Classes/Engine/Level.h
@@ -1191,6 +1191,14 @@ public:
 private:
 	bool IncrementalRegisterComponents(bool bPreRegisterComponents, int32 NumComponentsToUpdate, FRegisterComponentContext* Context);
 	bool IncrementalRunConstructionScripts(bool bProcessAllActors);
+
+#if WITH_EDITOR
+	/** Attempts to detect and fix any issues with the level script blueprint and associated objects */
+	void RepairLevelScript();
+
+	/** Replace the existing LSA (if set) by spawning a new one based on this level's script blueprint */
+	void RegenerateLevelScriptActor();
+#endif // WITH_EDITOR
 };
 
 
diff --git a/Engine/Source/Runtime/Engine/Private/Level.cpp b/Engine/Source/Runtime/Engine/Private/Level.cpp
--- a/Engine/Source/Runtime/Engine/Private/Level.cpp
+++ b/Engine/Source/Runtime/Engine/Private/Level.cpp
@@ -884,13 +884,7 @@ void ULevel::PostLoad()
 		LevelSimplification[Index].PostLoadDeprecated();
 	}
 
-	if (LevelScriptActor)
-	{
-		if (ULevelScriptBlueprint* LevelBlueprint = Cast<ULevelScriptBlueprint>(LevelScriptActor->GetClass()->ClassGeneratedBy))
-		{
-			FBlueprintEditorUtils::FixLevelScriptActorBindings(LevelScriptActor, LevelBlueprint);
-		}
-	}
+	RepairLevelScript();
 #endif
 }
 
@@ -2663,47 +2657,10 @@ void ULevel::CleanupLevelScriptBlueprint()
 
 void ULevel::OnLevelScriptBlueprintChanged(ULevelScriptBlueprint* InBlueprint)
 {
-	if( !InBlueprint->bIsRegeneratingOnLoad && 
-		// Make sure this is OUR level scripting blueprint
-		ensureMsgf(InBlueprint == LevelScriptBlueprint, TEXT("Level ('%s') received OnLevelScriptBlueprintChanged notification for the wrong Blueprint ('%s')."), LevelScriptBlueprint ? *LevelScriptBlueprint->GetPathName() : TEXT("NULL"), *InBlueprint->GetPathName()) )
-	{
-		bool bResetDebugObject = false;
-
-		UClass* SpawnClass = (LevelScriptBlueprint->GeneratedClass) ? LevelScriptBlueprint->GeneratedClass : LevelScriptBlueprint->SkeletonGeneratedClass;
-
-		// Get rid of the old LevelScriptActor
-		if( LevelScriptActor )
-		{
-			// Clear the current debug object and indicate that it needs to be reset (below).
-			if (InBlueprint->GetObjectBeingDebugged() == LevelScriptActor)
-			{
-				bResetDebugObject = true;
-				InBlueprint->SetObjectBeingDebugged(nullptr);
-			}
-
-			LevelScriptActor->MarkPendingKill();
-			LevelScriptActor = nullptr;
-		}
-
-		check( OwningWorld );
-		// Create the new one
-		FActorSpawnParameters SpawnInfo;
-		SpawnInfo.OverrideLevel = this;
-		LevelScriptActor = OwningWorld->SpawnActor<ALevelScriptActor>( SpawnClass, SpawnInfo );
-
-		if( LevelScriptActor )
-		{
-			// Reset the current debug object to the new instance if it was previously set to the old instance.
-			if (bResetDebugObject)
-			{
-				InBlueprint->SetObjectBeingDebugged(LevelScriptActor);
-			}
-
-			LevelScriptActor->ClearFlags(RF_Transactional);
-			check(LevelScriptActor->GetLevel() == this);
-			// Finally, fixup all the bound events to point to their new LSA
-			FBlueprintEditorUtils::FixLevelScriptActorBindings(LevelScriptActor, InBlueprint);
-		}		
+	if (!InBlueprint->bIsRegeneratingOnLoad &&
+		ensureMsgf(InBlueprint == LevelScriptBlueprint, TEXT("Level ('%s') received OnLevelScriptBlueprintChanged   for the wrong Blueprint ('%s')."), LevelScriptBlueprint ? *LevelScriptBlueprint->GetPathName() : TEXT("NULL"), *InBlueprint->GetPathName()))
+	{	    
+		RegenerateLevelScriptActor();
 	}
 }	
 
@@ -3045,4 +3002,80 @@ void ULevel::SetPartitionSubLevel(ULevel* SubLevel)
 	SubLevel->OwnerLevelPartition = Cast<UObject>(&*LevelPartition);
 }
 
-#endif // #if WITH_EDITORONLY_DATA
\ No newline at end of file
+#endif // #if WITH_EDITORONLY_DATA
+
+#if WITH_EDITOR
+void ULevel::RepairLevelScript()
+{
+	// Ensure bound events point to the level script actor
+	if (LevelScriptActor)
+	{
+		if (ULevelScriptBlueprint* LevelBlueprint = Cast<ULevelScriptBlueprint>(LevelScriptActor->GetClass()->ClassGeneratedBy))
+		{
+			FBlueprintEditorUtils::FixLevelScriptActorBindings(LevelScriptActor, LevelBlueprint);
+		}
+	}
+
+	// Verify that the LSA class matches the desired level script blueprint class
+	if (LevelScriptActor && LevelScriptBlueprint)
+	{
+		UClass* SpawnClass = (LevelScriptBlueprint->GeneratedClass) ? LevelScriptBlueprint->GeneratedClass : LevelScriptBlueprint->SkeletonGeneratedClass;
+		UClass* LSAClass = LevelScriptActor->GetClass();
+
+		if (SpawnClass && LSAClass != SpawnClass)
+		{
+			RegenerateLevelScriptActor();
+		}
+	}
+
+	// Catch the edge case where we have a level blueprint but have never created the LevelScriptActor based on it.
+	//    This could happen if a new level is saved before the level blueprint is compiled.
+	if (!LevelScriptActor && LevelScriptBlueprint && !LevelScriptBlueprint->bIsRegeneratingOnLoad)
+	{
+		RegenerateLevelScriptActor();
+	}	
+}
+
+void ULevel::RegenerateLevelScriptActor()
+{
+	bool bResetDebugObject = false;
+
+	check(LevelScriptBlueprint);
+	UClass* SpawnClass = (LevelScriptBlueprint->GeneratedClass) ? LevelScriptBlueprint->GeneratedClass : LevelScriptBlueprint->SkeletonGeneratedClass;
+	check(SpawnClass);
+
+	// Get rid of the old LevelScriptActor
+	if (LevelScriptActor)
+	{
+		// Clear the current debug object and indicate that it needs to be reset (below).
+		if (LevelScriptBlueprint->GetObjectBeingDebugged() == LevelScriptActor)
+		{
+			bResetDebugObject = true;
+			LevelScriptBlueprint->SetObjectBeingDebugged(nullptr);
+		}
+
+		LevelScriptActor->MarkPendingKill();
+		LevelScriptActor = nullptr;
+	}
+
+	check(OwningWorld);
+	// Create the new one
+	FActorSpawnParameters SpawnInfo;
+	SpawnInfo.OverrideLevel = this;
+	LevelScriptActor = OwningWorld->SpawnActor<ALevelScriptActor>(SpawnClass, SpawnInfo);
+
+	if (LevelScriptActor)
+	{
+		// Reset the current debug object to the new instance if it was previously set to the old instance.
+		if (bResetDebugObject)
+		{
+			LevelScriptBlueprint->SetObjectBeingDebugged(LevelScriptActor);
+		}
+
+		LevelScriptActor->ClearFlags(RF_Transactional);
+		check(LevelScriptActor->GetLevel() == this);
+		// Finally, fixup all the bound events to point to their new LSA
+		FBlueprintEditorUtils::FixLevelScriptActorBindings(LevelScriptActor, LevelScriptBlueprint);
+	}
+}
+#endif // WITH_EDITOR
