Switched RayTracingGroupId to uint32

diff --git a/Engine/Source/Runtime/Engine/Classes/Components/PrimitiveComponent.h b/Engine/Source/Runtime/Engine/Classes/Components/PrimitiveComponent.h
--- a/Engine/Source/Runtime/Engine/Classes/Components/PrimitiveComponent.h
+++ b/Engine/Source/Runtime/Engine/Classes/Components/PrimitiveComponent.h
@@ -630,10 +630,10 @@ public:
 
 	/**
 	 * Defines run-time groups of components. For example allows to assemble multiple parts of a building at runtime.
-	 * -1 means that component doesn't belong to any group.
+	 * 0 means that component doesn't belong to any group.
 	 */
 	UPROPERTY(EditAnywhere, AdvancedDisplay, BlueprintReadOnly, Category = RayTracing)
-	int32 RayTracingGroupId;
+	uint32 RayTracingGroupId;
 
 	/**
 	 * Defines how quickly it should be culled. For example buildings should have a low priority, but small dressing should have a high priority.
diff --git a/Engine/Source/Runtime/Engine/Private/Components/PrimitiveComponent.cpp b/Engine/Source/Runtime/Engine/Private/Components/PrimitiveComponent.cpp
--- a/Engine/Source/Runtime/Engine/Private/Components/PrimitiveComponent.cpp
+++ b/Engine/Source/Runtime/Engine/Private/Components/PrimitiveComponent.cpp
@@ -326,7 +326,7 @@ UPrimitiveComponent::UPrimitiveComponent(const FObjectInitializer& ObjectInitial
 	ComponentId.PrimIDValue = NextComponentId.Increment();
 	CustomDepthStencilValue = 0;
 	CustomDepthStencilWriteMask = ERendererStencilMask::ERSM_Default;
-	RayTracingGroupId = -1;
+	RayTracingGroupId = 0;
 	RayTracingGroupCullingPriority = ERayTracingGroupCullingPriority::CP_4_DEFAULT;
 
 	LDMaxDrawDistance = 0.f;
diff --git a/Engine/Source/Runtime/Engine/Public/ISMPartition/ISMComponentDescriptor.h b/Engine/Source/Runtime/Engine/Public/ISMPartition/ISMComponentDescriptor.h
--- a/Engine/Source/Runtime/Engine/Public/ISMPartition/ISMComponentDescriptor.h
+++ b/Engine/Source/Runtime/Engine/Public/ISMPartition/ISMComponentDescriptor.h
@@ -77,7 +77,7 @@ public:
 	FLightingChannels LightingChannels;
 
 	UPROPERTY()
-	int32 RayTracingGroupId;
+	uint32 RayTracingGroupId;
 
 	UPROPERTY()
 	ERayTracingGroupCullingPriority RayTracingGroupCullingPriority;
diff --git a/Engine/Source/Runtime/Engine/Public/PrimitiveSceneProxy.h b/Engine/Source/Runtime/Engine/Public/PrimitiveSceneProxy.h
--- a/Engine/Source/Runtime/Engine/Public/PrimitiveSceneProxy.h
+++ b/Engine/Source/Runtime/Engine/Public/PrimitiveSceneProxy.h
@@ -1089,7 +1089,7 @@ private:
 	uint8 LightingChannelMask;
 
 	// Run-time groups of proxies
-	int32 RayTracingGroupId;
+	uint32 RayTracingGroupId;
 	uint8 RayTracingGroupCullingPriority;
 
 protected:
