Delete XML-based XGE shader compiler.

- Also add a way to disable all distributed controllers (r.ShaderCompiler.AllowDistributedCompilation=0 or -NoRemoteShaderCompile) and made the naming more generic.

diff --git a/Engine/Source/Runtime/Engine/Private/ShaderCompiler/Experimental/ShaderCompileFASTBuild.cpp b/Engine/Source/Runtime/Engine/Private/ShaderCompiler/Experimental/ShaderCompileFASTBuild.cpp
--- a/Engine/Source/Runtime/Engine/Private/ShaderCompiler/Experimental/ShaderCompileFASTBuild.cpp
+++ b/Engine/Source/Runtime/Engine/Private/ShaderCompiler/Experimental/ShaderCompileFASTBuild.cpp
@@ -688,7 +688,7 @@ int32 FShaderCompileFASTBuildThreadRunnable::CompilingLoop()
 		{
 			const EShaderCompileJobPriority Priority = (EShaderCompileJobPriority)PriorityIndex;
 			const int32 MinBatchSize = (Priority == EShaderCompileJobPriority::Low) ? 1 : FASTBuildShaderCompilerVariables::MinBatchSize;
-			const int32 NumJobs = Manager->AllJobs.GetPendingJobs(EShaderCompilerWorkerType::XGE, Priority, MinBatchSize, INT32_MAX, JobQueue);
+			const int32 NumJobs = Manager->AllJobs.GetPendingJobs(EShaderCompilerWorkerType::Distributed, Priority, MinBatchSize, INT32_MAX, JobQueue);
 			if (NumJobs > 0)
 			{
 				UE_LOG(LogShaderCompilers, Display, TEXT("Started %d 'FASTBuild' shader compile jobs with '%s' priority"),
diff --git a/Engine/Source/Runtime/Engine/Private/ShaderCompiler/ShaderCompiler.cpp b/Engine/Source/Runtime/Engine/Private/ShaderCompiler/ShaderCompiler.cpp
--- a/Engine/Source/Runtime/Engine/Private/ShaderCompiler/ShaderCompiler.cpp
+++ b/Engine/Source/Runtime/Engine/Private/ShaderCompiler/ShaderCompiler.cpp
@@ -116,6 +116,14 @@ static FAutoConsoleVariableRef CVarShaderCompilerCacheStatsPrintoutInterval(
 	ECVF_Default
 );
 
+int32 GShaderCompilerAllowDistributedCompilation = 1;
+static FAutoConsoleVariableRef CVarShaderCompilerAllowDistributedCompilation(
+	TEXT("r.ShaderCompiler.AllowDistributedCompilation"),
+	GShaderCompilerAllowDistributedCompilation,
+	TEXT("If 0, only local (spawned by the engine) ShaderCompileWorkers will be used. If 1, SCWs will be distributed using one of several possible backends (XGE, FASTBuild, SN-DBS)"),
+	ECVF_Default
+);
+
 /** Helper functions for logging more debug info */
 namespace ShaderCompiler
 {
@@ -169,6 +177,18 @@ namespace ShaderCompiler
 	{
 		return GShaderCompilerJobCache != 0;
 	}
+
+	bool IsRemoteCompilingAllowed()
+	{
+		// commandline switches override the CVars
+		static bool bDisabledFromCommandline = FParse::Param(FCommandLine::Get(), TEXT("NoRemoteShaderCompile"));
+		if (bDisabledFromCommandline)
+		{
+			return false;
+		}
+
+		return GShaderCompilerAllowDistributedCompilation != 0;
+	}
 }
 
 // The Id of 0 is reserved for global shaders
@@ -300,7 +320,7 @@ void FShaderCompileJobCollection::InternalSetPriority(FShaderCommonCompileJob* J
 		Job->PendingPriority = InPriority;
 	}
 	else if (!Job->bFinalized &&
-		Job->CurrentWorker == EShaderCompilerWorkerType::XGE &&
+		Job->CurrentWorker == EShaderCompilerWorkerType::Distributed &&
 		InPriority == EShaderCompileJobPriority::ForceLocal)
 	{
 		FShaderCommonCompileJob* NewJob = CloneJob(Job);
@@ -2991,7 +3011,7 @@ FShaderCompilingManager::FShaderCompilingManager() :
 
 	TUniquePtr<FShaderCompileThreadRunnableBase> RemoteCompileThread;
 #if PLATFORM_WINDOWS
-	const bool bCanUseRemoteCompiling = bAllowCompilingThroughWorkers && AllTargetPlatformSupportsRemoteShaderCompiling();
+	const bool bCanUseRemoteCompiling = bAllowCompilingThroughWorkers && ShaderCompiler::IsRemoteCompilingAllowed() && AllTargetPlatformSupportsRemoteShaderCompiling();
 	BuildDistributionController = bCanUseRemoteCompiling ? FindRemoteCompilerController() : nullptr;
 	
 	if (BuildDistributionController)
@@ -3002,7 +3022,7 @@ FShaderCompilingManager::FShaderCompilingManager() :
 	else
 #endif // PLATFORM_WINDOWS
 #if PLATFORM_DESKTOP
-	if (bAllowCompilingThroughWorkers && FShaderCompileFASTBuildThreadRunnable::IsSupported())
+	if (bAllowCompilingThroughWorkers && ShaderCompiler::IsRemoteCompilingAllowed() && FShaderCompileFASTBuildThreadRunnable::IsSupported())
 	{
 		UE_LOG(LogShaderCompilers, Display, TEXT("Using FASTBuild Shader Compiler."));
 		RemoteCompileThread = MakeUnique<FShaderCompileFASTBuildThreadRunnable>(this);
@@ -3014,6 +3034,8 @@ FShaderCompilingManager::FShaderCompilingManager() :
 	TUniquePtr<FShaderCompileThreadRunnableBase> LocalThread = MakeUnique<FShaderCompileThreadRunnable>(this);
 	if (RemoteCompileThread)
 	{
+		checkf(ShaderCompiler::IsRemoteCompilingAllowed(), TEXT("We have a remote compiling thread without the remote compilation being allowed"));
+
 		// Only force-local jobs are guaranteed to stay on the local machine. Going wide with High priority jobs is important for the startup times,
 		// since special materials use High priority. Possibly the partition by priority is too rigid in general.
 		RemoteCompileThread->SetPriorityRange(EShaderCompileJobPriority::Low, EShaderCompileJobPriority::High);
diff --git a/Engine/Source/Runtime/Engine/Private/ShaderCompiler/ShaderCompilerDistributed.cpp b/Engine/Source/Runtime/Engine/Private/ShaderCompiler/ShaderCompilerDistributed.cpp
--- a/Engine/Source/Runtime/Engine/Private/ShaderCompiler/ShaderCompilerDistributed.cpp
+++ b/Engine/Source/Runtime/Engine/Private/ShaderCompiler/ShaderCompilerDistributed.cpp
@@ -6,15 +6,21 @@
 
 namespace DistributedShaderCompilerVariables
 {
-	//TODO: Remove XGE from the name after ensuring that we are not breaking existing configs by doing so.
+	//TODO: Remove the XGE doublet
 	int32 MinBatchSize = 20;
 	FAutoConsoleVariableRef CVarXGEShaderCompileMinBatchSize(
         TEXT("r.XGEShaderCompile.MinBatchSize"),
         MinBatchSize,
-        TEXT("Minimum number of shaders to compile with XGE.\n")
-        TEXT("Smaller number of shaders will compile locally."),
+        TEXT("This CVar is deprecated, please use r.ShaderCompiler.DistributedMinBatchSize"),
         ECVF_Default);
 
+	FAutoConsoleVariableRef CVarDistributedMinBatchSize(
+		TEXT("r.ShaderCompiler.DistributedMinBatchSize"),
+		MinBatchSize,
+		TEXT("Minimum number of shaders to compile with a distributed controller.\n")
+		TEXT("Smaller number of shaders will compile locally."),
+		ECVF_Default);
+
 	static int32 GDistributedControllerTimeout = 5 * 60;	// 
 	static FAutoConsoleVariableRef CVarDistributedControllerTimeout(
 		TEXT("r.ShaderCompiler.DistributedControllerTimeout"),
@@ -174,10 +180,10 @@ int32 FShaderCompileDistributedThreadRunnable_Interface::CompilingLoop()
 			// Grab as many jobs from the job queue as we can
 			const EShaderCompileJobPriority Priority = (EShaderCompileJobPriority)PriorityIndex;
 			const int32 MinBatchSize = (Priority == EShaderCompileJobPriority::Low) ? 1 : DistributedShaderCompilerVariables::MinBatchSize;
-			const int32 NumJobs = Manager->AllJobs.GetPendingJobs(EShaderCompilerWorkerType::XGE, Priority, MinBatchSize, INT32_MAX, PendingJobs);
+			const int32 NumJobs = Manager->AllJobs.GetPendingJobs(EShaderCompilerWorkerType::Distributed, Priority, MinBatchSize, INT32_MAX, PendingJobs);
 			if (NumJobs > 0)
 			{
-				UE_LOG(LogShaderCompilers, Display, TEXT("Started %d 'XGE' shader compile jobs with '%s' priority"),
+				UE_LOG(LogShaderCompilers, Verbose, TEXT("Started %d 'Distributed' shader compile jobs with '%s' priority"),
 					NumJobs,
 					ShaderCompileJobPriorityToString((EShaderCompileJobPriority)PriorityIndex));
 			}
@@ -331,7 +337,7 @@ int32 FShaderCompileDistributedThreadRunnable_Interface::CompilingLoop()
 			if (bOutputFileReadFailed)
 			{
 				// Reading result from XGE job failed, so recompile shaders in current job batch locally
-				UE_LOG(LogShaderCompilers, Log, TEXT("Rescheduling shader compilation to run locally after XGE job failed: %s"), *Task->OutputFilePath);
+				UE_LOG(LogShaderCompilers, Log, TEXT("Rescheduling shader compilation to run locally after distributed job failed: %s"), *Task->OutputFilePath);
 
 				for (FShaderCommonCompileJobPtr Job : Task->ShaderJobs)
 				{
diff --git a/Engine/Source/Runtime/Engine/Private/ShaderCompiler/ShaderCompilerXGE.cpp b/Engine/Source/Runtime/Engine/Private/ShaderCompiler/ShaderCompilerXGE.cpp
deleted file mode 100644
diff --git a/Engine/Source/Runtime/Engine/Public/ShaderCompiler.h b/Engine/Source/Runtime/Engine/Public/ShaderCompiler.h
--- a/Engine/Source/Runtime/Engine/Public/ShaderCompiler.h
+++ b/Engine/Source/Runtime/Engine/Public/ShaderCompiler.h
@@ -43,7 +43,7 @@ enum class EShaderCompilerWorkerType : uint8
 {
 	None,
 	LocalThread,
-	XGE,
+	Distributed,
 };
 
 enum class EShaderCompileJobType : uint8
@@ -650,99 +650,6 @@ namespace FShaderCompileUtilities
 
 }
 
-#if PLATFORM_WINDOWS // XGE shader compilation is only supported on Windows.
-
-class FShaderCompileXGEThreadRunnable_XmlInterface : public FShaderCompileThreadRunnableBase
-{
-private:
-	/** The handle referring to the XGE console process, if a build is in progress. */
-	FProcHandle BuildProcessHandle;
-	
-	/** Process ID of the XGE console, if a build is in progress. */
-	uint32 BuildProcessID;
-
-	/**
-	 * A map of directory paths to shader jobs contained within that directory.
-	 * One entry per XGE task.
-	 */
-	class FShaderBatch
-	{
-		TArray<FShaderCommonCompileJobPtr> Jobs;
-		bool bTransferFileWritten;
-
-	public:
-		const FString& DirectoryBase;
-		const FString& InputFileName;
-		const FString& SuccessFileName;
-		const FString& OutputFileName;
-
-		int32 BatchIndex;
-		int32 DirectoryIndex;
-
-		FString WorkingDirectory;
-		FString OutputFileNameAndPath;
-		FString SuccessFileNameAndPath;
-		FString InputFileNameAndPath;
-		
-		FShaderBatch(const FString& InDirectoryBase, const FString& InInputFileName, const FString& InSuccessFileName, const FString& InOutputFileName, int32 InDirectoryIndex, int32 InBatchIndex)
-			: bTransferFileWritten(false)
-			, DirectoryBase(InDirectoryBase)
-			, InputFileName(InInputFileName)
-			, SuccessFileName(InSuccessFileName)
-			, OutputFileName(InOutputFileName)
-		{
-			SetIndices(InDirectoryIndex, InBatchIndex);
-		}
-
-		void SetIndices(int32 InDirectoryIndex, int32 InBatchIndex);
-
-		void CleanUpFiles(bool keepInputFile);
-
-		inline int32 NumJobs()
-		{
-			return Jobs.Num();
-		}
-		inline const TArray<FShaderCommonCompileJobPtr>& GetJobs() const
-		{
-			return Jobs;
-		}
-
-		void AddJob(FShaderCommonCompileJobPtr Job);
-		
-		void WriteTransferFile();
-	};
-	TArray<FShaderBatch*> ShaderBatchesInFlight;
-	TArray<FShaderBatch*> ShaderBatchesFull;
-	TSparseArray<FShaderBatch*> ShaderBatchesIncomplete;
-
-	/** The full path to the two working directories for XGE shader builds. */
-	const FString XGEWorkingDirectory;
-	uint32 XGEDirectoryIndex;
-
-	uint64 LastAddTime;
-	uint64 StartTime;
-	int32 BatchIndexToCreate;
-	int32 BatchIndexToFill;
-
-	FDateTime ScriptFileCreationTime;
-
-	void PostCompletedJobsForBatch(FShaderBatch* Batch);
-
-	void GatherResultsFromXGE();
-
-public:
-	/** Initialization constructor. */
-	FShaderCompileXGEThreadRunnable_XmlInterface(class FShaderCompilingManager* InManager);
-	virtual ~FShaderCompileXGEThreadRunnable_XmlInterface();
-
-	/** Main work loop. */
-	virtual int32 CompilingLoop() override;
-
-	static bool IsSupported();
-};
-
-#endif // PLATFORM_WINDOWS
-
 class FShaderCompileDistributedThreadRunnable_Interface : public FShaderCompileThreadRunnableBase
 {
 	uint32 NumDispatchedJobs;
