Add render capture console command for skin cache.
Add direct access to skin cache buffers.
The direct access is only needed temporarily while we boot up the anim deformer system. It ignores transitions for now.

diff --git a/Engine/Source/Runtime/Engine/Private/GPUSkinCache.cpp b/Engine/Source/Runtime/Engine/Private/GPUSkinCache.cpp
--- a/Engine/Source/Runtime/Engine/Private/GPUSkinCache.cpp
+++ b/Engine/Source/Runtime/Engine/Private/GPUSkinCache.cpp
@@ -15,6 +15,7 @@ GPUSkinCache.cpp: Performs skinning on a compute shader into a buffer to avoid v
 #include "ClearQuad.h"
 #include "Shader.h"
 #include "MeshMaterialShader.h"
+#include "RenderCaptureInterface.h"
 #include "RenderGraphResources.h"
 #include "Algo/Unique.h"
 #include "HAL/IConsoleManager.h"
@@ -146,6 +147,12 @@ FAutoConsoleVariableRef CVarGPUSkinCachePrintMemorySummary(
 	ECVF_RenderThreadSafe
 );
 
+int32 GNumDispatchesToCapture = 0;
+static FAutoConsoleVariableRef CVarGPUSkinCacheNumDispatchesToCapture(
+	TEXT("r.SkinCache.Capture"),
+	GNumDispatchesToCapture,
+	TEXT("Trigger a render capture for the next skin cache dispatches."));
+
 static int32 GGPUSkinCacheFlushCounter = 0;
 
 #if RHI_RAYTRACING
@@ -503,6 +510,8 @@ public:
 	}
 #endif // RHI_RAYTRACING
 
+	TArray<FSectionDispatchData>& GetDispatchData() { return DispatchData; }
+
 protected:
 	FGPUSkinCache::FRWBuffersAllocation* PositionAllocation;
 	FGPUSkinCache* SkinCache;
@@ -1307,11 +1316,15 @@ void FGPUSkinCache::GetBufferUAVs(const TArray<FSkinCacheRWBuffer*>& InBuffers,
 
 void FGPUSkinCache::DoDispatch(FRHICommandListImmediate& RHICmdList)
 {
-	SCOPED_GPU_STAT(RHICmdList, GPUSkinCache);
-
 	int32 BatchCount = BatchDispatches.Num();
 	INC_DWORD_STAT_BY(STAT_GPUSkinCache_TotalNumChunks, BatchCount);
 
+	bool bCapture = BatchCount > 0 && GNumDispatchesToCapture > 0;
+	RenderCaptureInterface::FScopedCapture RenderCapture(bCapture, &RHICmdList, TEXT("GPUSkinCache"));
+	GNumDispatchesToCapture -= bCapture ? 1 : 0;
+
+	SCOPED_GPU_STAT(RHICmdList, GPUSkinCache);
+
 	TArray<FSkinCacheRWBuffer*> BuffersToTransitionForSkinning;
 	BuffersToTransitionForSkinning.Reserve(BatchCount * 2);
 	{
@@ -1413,6 +1426,9 @@ void FGPUSkinCache::DoDispatch(FRHICommandListImmediate& RHICmdList)
 
 void FGPUSkinCache::DoDispatch(FRHICommandListImmediate& RHICmdList, FGPUSkinCacheEntry* SkinCacheEntry, int32 Section, int32 RevisionNumber)
 {
+	RenderCaptureInterface::FScopedCapture RenderCapture(GNumDispatchesToCapture > 0, &RHICmdList, TEXT("GPUSkinCache"));
+	GNumDispatchesToCapture = FMath::Max(GNumDispatchesToCapture - 1, 0);
+
 	SCOPED_GPU_STAT(RHICmdList, GPUSkinCache);
 
 	INC_DWORD_STAT(STAT_GPUSkinCache_TotalNumChunks);
@@ -1425,7 +1441,7 @@ void FGPUSkinCache::DoDispatch(FRHICommandListImmediate& RHICmdList, FGPUSkinCac
 	GetBufferUAVs(BuffersToTransitionForSkinning, SkinningBuffersToOverlap);
 	RHICmdList.BeginUAVOverlap(SkinningBuffersToOverlap);
 	{
-		DispatchUpdateSkinning(RHICmdList, SkinCacheEntry, Section, RevisionNumber);
+	DispatchUpdateSkinning(RHICmdList, SkinCacheEntry, Section, RevisionNumber);
 	}
 	RHICmdList.EndUAVOverlap(SkinningBuffersToOverlap);
 
@@ -2209,6 +2225,48 @@ void FGPUSkinCache::InvalidateAllEntries()
 	SET_MEMORY_STAT(STAT_GPUSkinCache_TangentsIntermediateMemUsed, 0);
 }
 
+FRWBuffer* FGPUSkinCache::GetPositionBuffer(uint32 ComponentId, uint32 SectionIndex) const
+{
+	for (FGPUSkinCacheEntry* Entry : Entries)
+	{
+		if (Entry && Entry->GPUSkin && Entry->GPUSkin->GetComponentId() == ComponentId)
+		{
+			FGPUSkinCacheEntry::FSectionDispatchData& DispatchData = Entry->GetDispatchData()[SectionIndex];
+			FSkinCacheRWBuffer* SkinCacheRWBuffer = DispatchData.PreviousPositionBuffer != nullptr ? DispatchData.PreviousPositionBuffer : DispatchData.PositionBuffer;
+			return SkinCacheRWBuffer != nullptr ? &SkinCacheRWBuffer->Buffer : nullptr;
+		}
+	}
+
+	return nullptr;
+}
+
+FRWBuffer* FGPUSkinCache::GetTangentBuffer(uint32 ComponentId, uint32 SectionIndex) const
+{
+	for (FGPUSkinCacheEntry* Entry : Entries)
+	{
+		if (Entry && Entry->GPUSkin && Entry->GPUSkin->GetComponentId() == ComponentId)
+		{
+			FGPUSkinCacheEntry::FSectionDispatchData& DispatchData = Entry->GetDispatchData()[SectionIndex];
+			FSkinCacheRWBuffer* SkinCacheRWBuffer = DispatchData.GetTangentRWBuffer();
+			return SkinCacheRWBuffer != nullptr ? &SkinCacheRWBuffer->Buffer : nullptr;
+		}
+	}
+
+	return nullptr;
+}
+
+FRHIShaderResourceView* FGPUSkinCache::GetBoneBuffer(uint32 ComponentId, uint32 SectionIndex) const
+{
+	for (FGPUSkinCacheEntry* Entry : Entries)
+	{
+		FGPUSkinCacheEntry::FSectionDispatchData& DispatchData = Entry->GetDispatchData()[SectionIndex];
+		FGPUBaseSkinVertexFactory::FShaderDataType& ShaderData = DispatchData.SourceVertexFactory->GetShaderData();
+		return ShaderData.GetBoneBufferForReading(false).VertexBufferSRV;
+	}
+
+	return nullptr;
+}
+
 FCachedGeometry FGPUSkinCache::GetCachedGeometry(uint32 ComponentId) const
 {
 	FCachedGeometry Out;
diff --git a/Engine/Source/Runtime/Engine/Public/GPUSkinCache.h b/Engine/Source/Runtime/Engine/Public/GPUSkinCache.h
--- a/Engine/Source/Runtime/Engine/Public/GPUSkinCache.h
+++ b/Engine/Source/Runtime/Engine/Public/GPUSkinCache.h
@@ -414,6 +414,10 @@ public:
 
 	ENGINE_API void TransitionAllToReadable(FRHICommandList& RHICmdList);
 
+	ENGINE_API FRWBuffer* GetPositionBuffer(uint32 ComponentId, uint32 SectionIndex) const;
+	ENGINE_API FRWBuffer* GetTangentBuffer(uint32 ComponentId, uint32 SectionIndex) const;
+	ENGINE_API FRHIShaderResourceView* GetBoneBuffer(uint32 ComponentId, uint32 SectionIndex) const;
+
 #if RHI_RAYTRACING
 	void AddRayTracingGeometryToUpdate(FRayTracingGeometry* InRayTracingGeometry, EAccelerationStructureBuildMode InBuildMode);
 	ENGINE_API void CommitRayTracingGeometryUpdates(FRHICommandListImmediate& RHICmdList);
