Chaos: Fixed SweepGeom returning garbage distance for sweep hits. Fallout from LWC merge.

diff --git a/Engine/Source/Runtime/Engine/Private/PhysicsEngine/Experimental/PhysInterface_Chaos.cpp b/Engine/Source/Runtime/Engine/Private/PhysicsEngine/Experimental/PhysInterface_Chaos.cpp
--- a/Engine/Source/Runtime/Engine/Private/PhysicsEngine/Experimental/PhysInterface_Chaos.cpp
+++ b/Engine/Source/Runtime/Engine/Private/PhysicsEngine/Experimental/PhysInterface_Chaos.cpp
@@ -1142,7 +1142,7 @@ bool FPhysInterface_Chaos::Sweep_Geom(FHitResult& OutHit, const FBodyInstance* I
 								Hit.Actor = ShapeRef.ActorRef ? ShapeRef.ActorRef->GetParticle_LowLevel() : nullptr;
 								Hit.WorldPosition = WorldPosition;
 								Hit.WorldNormal = WorldNormal;
-								Hit.Distance = (float)Hit.Distance; // we should eventually have the Hit structure to use a Chaos::FReal equivalent instead
+								Hit.Distance = (float)Distance; // we should eventually have the Hit structure to use a Chaos::FReal equivalent instead
 								Hit.FaceIndex = FaceIdx;
 								if (!HadInitialOverlap(Hit))
 								{
