Fix broken rendering of cloth when running with LWC on
- changed FClothSimulData to use FVector3f array types
- adjusted the various API to take the FVector3f arrays
- changed GPUSkinCache to also use FMatrix44f instead of FMatrix in various places

diff --git a/Engine/Source/Runtime/ClothingSystemRuntimeInterface/Public/ClothingSystemRuntimeTypes.h b/Engine/Source/Runtime/ClothingSystemRuntimeInterface/Public/ClothingSystemRuntimeTypes.h
--- a/Engine/Source/Runtime/ClothingSystemRuntimeInterface/Public/ClothingSystemRuntimeTypes.h
+++ b/Engine/Source/Runtime/ClothingSystemRuntimeInterface/Public/ClothingSystemRuntimeTypes.h
@@ -17,10 +17,10 @@ struct FClothSimulData
 	}
 
 	// Positions of the simulation mesh particles
-	TArray<FVector> Positions;
+	TArray<FVector3f> Positions;
 
 	// Normals at the simulation mesh particles
-	TArray<FVector> Normals;
+	TArray<FVector3f> Normals;
 
 	// Transform applied per position/normal element when loaded
 	FTransform Transform;
diff --git a/Engine/Source/Runtime/Engine/Private/GPUSkinCache.cpp b/Engine/Source/Runtime/Engine/Private/GPUSkinCache.cpp
--- a/Engine/Source/Runtime/Engine/Private/GPUSkinCache.cpp
+++ b/Engine/Source/Runtime/Engine/Private/GPUSkinCache.cpp
@@ -1482,7 +1482,7 @@ void FGPUSkinCache::ProcessEntry(
 	const FMorphVertexBuffer* MorphVertexBuffer,
 	const FSkeletalMeshVertexClothBuffer* ClothVertexBuffer, 
 	const FClothSimulData* SimData,
-	const FMatrix& ClothLocalToWorld, 
+	const FMatrix44f& ClothLocalToWorld, 
 	float ClothBlendWeight, 
 	uint32 RevisionNumber, 
 	int32 Section, 
diff --git a/Engine/Source/Runtime/Engine/Private/GPUSkinVertexFactory.cpp b/Engine/Source/Runtime/Engine/Private/GPUSkinVertexFactory.cpp
--- a/Engine/Source/Runtime/Engine/Private/GPUSkinVertexFactory.cpp
+++ b/Engine/Source/Runtime/Engine/Private/GPUSkinVertexFactory.cpp
@@ -1044,8 +1044,8 @@ IMPLEMENT_TYPE_LAYOUT(TGPUSkinAPEXClothVertexFactoryShaderParameters);
 	TGPUSkinAPEXClothVertexFactory::ClothShaderType
 -----------------------------------------------------------------------------*/
 
-bool FGPUBaseSkinAPEXClothVertexFactory::ClothShaderType::UpdateClothSimulData(FRHICommandListImmediate& RHICmdList, const TArray<FVector>& InSimulPositions,
-	const TArray<FVector>& InSimulNormals, uint32 FrameNumberToPrepare, ERHIFeatureLevel::Type FeatureLevel)
+bool FGPUBaseSkinAPEXClothVertexFactory::ClothShaderType::UpdateClothSimulData(FRHICommandListImmediate& RHICmdList, const TArray<FVector3f>& InSimulPositions,
+	const TArray<FVector3f>& InSimulNormals, uint32 FrameNumberToPrepare, ERHIFeatureLevel::Type FeatureLevel)
 {
 	QUICK_SCOPE_CYCLE_COUNTER(STAT_FGPUBaseSkinAPEXClothVertexFactory_UpdateClothSimulData);
 
diff --git a/Engine/Source/Runtime/Engine/Private/SkeletalRenderGPUSkin.cpp b/Engine/Source/Runtime/Engine/Private/SkeletalRenderGPUSkin.cpp
--- a/Engine/Source/Runtime/Engine/Private/SkeletalRenderGPUSkin.cpp
+++ b/Engine/Source/Runtime/Engine/Private/SkeletalRenderGPUSkin.cpp
@@ -647,9 +647,9 @@ void FSkeletalMeshObjectGPUSkin::ProcessUpdatedDynamicData(FGPUSkinCache* GPUSki
 			if (bUseSkinCache)
 			{
 				// This takes the cloth positions from cloth space into world space
-				FMatrix ClothLocalToWorld = bClothFactory ? VertexFactoryData.ClothVertexFactories[SectionIdx]->GetClothShaderData().GetClothLocalToWorldForWriting(FrameNumberToPrepare) : FMatrix::Identity;
+				FMatrix44f ClothLocalToWorld = bClothFactory ? VertexFactoryData.ClothVertexFactories[SectionIdx]->GetClothShaderData().GetClothLocalToWorldForWriting(FrameNumberToPrepare) : FMatrix44f::Identity;
 				// Matrices are transposed in UE meaning matrix multiples need to happen in reverse ((AB)x = b becomes xTBTAT = b).
-				FMatrix LocalToCloth = DynamicData->ClothObjectLocalToWorld * ClothLocalToWorld.Inverse();
+				FMatrix44f LocalToCloth = DynamicData->ClothObjectLocalToWorld * ClothLocalToWorld.Inverse();
 
 				GPUSkinCache->ProcessEntry(
 					RHICmdList, 
diff --git a/Engine/Source/Runtime/Engine/Public/GPUSkinCache.h b/Engine/Source/Runtime/Engine/Public/GPUSkinCache.h
--- a/Engine/Source/Runtime/Engine/Public/GPUSkinCache.h
+++ b/Engine/Source/Runtime/Engine/Public/GPUSkinCache.h
@@ -164,7 +164,7 @@ public:
 		const FMorphVertexBuffer* MorphVertexBuffer, 
 		const FSkeletalMeshVertexClothBuffer* ClothVertexBuffer, 
 		const FClothSimulData* SimData,
-		const FMatrix& ClothLocalToWorld,
+		const FMatrix44f& ClothLocalToWorld,
 		float ClothBlendWeight, 
 		uint32 RevisionNumber, 
 		int32 Section, 
diff --git a/Engine/Source/Runtime/Engine/Public/GPUSkinVertexFactory.h b/Engine/Source/Runtime/Engine/Public/GPUSkinVertexFactory.h
--- a/Engine/Source/Runtime/Engine/Public/GPUSkinVertexFactory.h
+++ b/Engine/Source/Runtime/Engine/Public/GPUSkinVertexFactory.h
@@ -659,7 +659,7 @@ public:
 			Reset();
 		}
 
-		bool UpdateClothSimulData(FRHICommandListImmediate& RHICmdList, const TArray<FVector>& InSimulPositions, const TArray<FVector>& InSimulNormals, uint32 FrameNumber, ERHIFeatureLevel::Type FeatureLevel);
+		bool UpdateClothSimulData(FRHICommandListImmediate& RHICmdList, const TArray<FVector3f>& InSimulPositions, const TArray<FVector3f>& InSimulNormals, uint32 FrameNumber, ERHIFeatureLevel::Type FeatureLevel);
 
 		void ReleaseClothSimulData()
 		{
@@ -716,7 +716,7 @@ public:
 			return ClothSimulPositionNormalBuffer[Index];
 		}
 		
-		FMatrix& GetClothLocalToWorldForWriting(uint32 FrameNumber)
+		FMatrix44f& GetClothLocalToWorldForWriting(uint32 FrameNumber)
 		{
 			uint32 Index = GetOldestIndex(FrameNumber);
 			Index = (BufferFrameNumber[0] == FrameNumber) ? 0 : Index;
@@ -725,7 +725,7 @@ public:
 			return ClothLocalToWorld[Index];
 		}
 
-		const FMatrix& GetClothLocalToWorldForReading(bool bPrevious, uint32 FrameNumber) const
+		const FMatrix44f& GetClothLocalToWorldForReading(bool bPrevious, uint32 FrameNumber) const
 		{
 			int32 Index = GetMostRecentIndex(FrameNumber);
 
@@ -754,7 +754,7 @@ public:
 		/**
 		 * Matrix to apply to positions/normals
 		 */
-		FMatrix ClothLocalToWorld[2];
+		FMatrix44f ClothLocalToWorld[2];
 
 		// @return 0 / 1, index into ClothSimulPositionNormalBuffer[]
 		uint32 GetMostRecentIndex(uint32 FrameNumber) const
@@ -819,8 +819,8 @@ public:
 			BufferFrameNumber[0] = -1;
 			BufferFrameNumber[1] = -1;
 
-			ClothLocalToWorld[0] = FMatrix::Identity;
-			ClothLocalToWorld[1] = FMatrix::Identity;
+			ClothLocalToWorld[0] = FMatrix44f::Identity;
+			ClothLocalToWorld[1] = FMatrix44f::Identity;
 		}
 	};
 
