- Optimised GetPreferredLoadedLevelToAddToWorld (only when needed and sorting of streaming cells only done once per frame)
- Refactored World Partition streaming cell state management and moved logic to UWorldPartitionStreamingPolicy base class

diff --git a/Engine/Source/Runtime/Engine/Classes/Engine/World.h b/Engine/Source/Runtime/Engine/Classes/Engine/World.h
--- a/Engine/Source/Runtime/Engine/Classes/Engine/World.h
+++ b/Engine/Source/Runtime/Engine/Classes/Engine/World.h
@@ -1004,7 +1004,7 @@ public:
 	ULevel* GetCurrentLevelPendingVisibility() const { return CurrentLevelPendingVisibility; }
 
 	/** Returns the best candidate level, if any, to become the CurrentLevelPendingVisibility */
-	ULevel* GetNextPreferredLevelPendingVisibility() { return NextPreferredLevelPendingVisibility; }
+	ULevel* CalculateNextPreferredLevelPendingVisibility() const;
 
 	/** Returns the level, if any, in the process of being made invisible */
 	ULevel* GetCurrentLevelPendingInvisibility() const { return CurrentLevelPendingInvisibility; }
@@ -1063,10 +1063,6 @@ private:
 	UPROPERTY(Transient)
 	TObjectPtr<class ULevel>								CurrentLevelPendingVisibility;
 
-	/** Pointer to the next best candidate level to become the CurrentLevelPendingVisibility. */
-	UPROPERTY(Transient)
-	TObjectPtr<class ULevel>								NextPreferredLevelPendingVisibility;
-
 	/** Pointer to the current level in the queue to be made invisible, NULL if none are pending. */
 	UPROPERTY(Transient)
 	TObjectPtr<class ULevel>								CurrentLevelPendingInvisibility;
@@ -3455,9 +3451,6 @@ private:
 	/** Updates world's required streaming levels */
 	void InternalUpdateStreamingState();
 
-	/** Determine next making visible best candidate to AddToWorld */
-	void UpdateNextPreferredLevelPendingVisibility();
-
 #if WITH_EDITOR
 public:
 	void SetPlayInEditorInitialNetMode(ENetMode InNetMode)
diff --git a/Engine/Source/Runtime/Engine/Private/LevelStreaming.cpp b/Engine/Source/Runtime/Engine/Private/LevelStreaming.cpp
--- a/Engine/Source/Runtime/Engine/Private/LevelStreaming.cpp
+++ b/Engine/Source/Runtime/Engine/Private/LevelStreaming.cpp
@@ -588,7 +588,7 @@ void ULevelStreaming::UpdateStreamingState(bool& bOutUpdateAgain, bool& bOutRede
 				bOutUpdateAgain = true;
 				bOutRedetermineTarget = true;
 			}
-			else if (!World->GetNextPreferredLevelPendingVisibility() || World->GetNextPreferredLevelPendingVisibility() == LoadedLevel)
+			else
 			{
 				World->AddToWorld(LoadedLevel, LevelTransform, !bShouldBlockOnLoad);
 
diff --git a/Engine/Source/Runtime/Engine/Private/World.cpp b/Engine/Source/Runtime/Engine/Private/World.cpp
--- a/Engine/Source/Runtime/Engine/Private/World.cpp
+++ b/Engine/Source/Runtime/Engine/Private/World.cpp
@@ -380,7 +380,6 @@ UWorld::UWorld( const FObjectInitializer& ObjectInitializer )
 ,	NextTravelType(TRAVEL_Relative)
 ,	CleanupWorldTag(0)
 {
-	NextPreferredLevelPendingVisibility = nullptr;
 	TimerManager = new FTimerManager();
 #if WITH_EDITOR
 	SetPlayInEditorInitialNetMode(ENetMode::NM_Standalone);
@@ -2554,6 +2553,17 @@ void UWorld::AddToWorld( ULevel* Level, const FTransform& LevelTransform, bool b
 	bool bPerformedLastStep	= false;
 	const bool bIsGameWorld = IsGameWorld();
 
+	if (bExecuteNextStep)
+	{
+		ULevel* NextPreferredLevelPendingVisibility = CalculateNextPreferredLevelPendingVisibility();
+		if (NextPreferredLevelPendingVisibility != nullptr && NextPreferredLevelPendingVisibility != Level)
+		{
+			// Don't add to world if the level is not the preferred one
+			check(CurrentLevelPendingVisibility == nullptr);
+			bExecuteNextStep = false;
+		}
+	}
+
 	// Don't make this level visible if it's currently being made invisible
 	if( bExecuteNextStep && CurrentLevelPendingVisibility == NULL && CurrentLevelPendingInvisibility != Level )
 	{
@@ -2561,7 +2571,6 @@ void UWorld::AddToWorld( ULevel* Level, const FTransform& LevelTransform, bool b
 		
 		// Mark level as being the one in process of being made visible.
 		CurrentLevelPendingVisibility = Level;
-		NextPreferredLevelPendingVisibility = nullptr;
 
 		// Add to the UWorld's array of levels, which causes it to be rendered et al.
 		Levels.AddUnique( Level );
@@ -2720,7 +2729,7 @@ void UWorld::AddToWorld( ULevel* Level, const FTransform& LevelTransform, bool b
 			}
 
 			const float PreventNextStepTimeLimit = 0.0; // We will always run route actor initialize in its own frame if we are using a time limit
-			bExecuteNextStep = (!bConsiderTimeLimit || !IsTimeLimitExceeded( TEXT("initializing network actors"), StartTime, Level, PreventNextStepTimeLimit )); 
+			bExecuteNextStep = (!bConsiderTimeLimit || !IsTimeLimitExceeded( TEXT("initializing network actors"), StartTime, Level, PreventNextStepTimeLimit ));
 		}
 
 		// Route various initialization functions and set volumes.
@@ -3607,13 +3616,14 @@ void UWorld::InternalUpdateStreamingState()
 	}
 }
 
-void UWorld::UpdateNextPreferredLevelPendingVisibility()
+ULevel* UWorld::CalculateNextPreferredLevelPendingVisibility() const
 {
 	if (CurrentLevelPendingVisibility == nullptr)
 	{
 		const UWorldPartition* WorldPartition = GetWorldPartition();
-		NextPreferredLevelPendingVisibility = WorldPartition ? WorldPartition->GetPreferredLoadedLevelToAddToWorld() : nullptr;
+		return WorldPartition ? WorldPartition->GetPreferredLoadedLevelToAddToWorld() : nullptr;
 	}
+	return nullptr;
 }
 
 void UWorld::UpdateLevelStreaming()
@@ -3637,8 +3647,6 @@ void UWorld::UpdateLevelStreaming()
 	{
 		if (ULevelStreaming* StreamingLevel = StreamingLevelsToConsider.GetStreamingLevels()[Index])
 		{
-			UpdateNextPreferredLevelPendingVisibility();
-
 			bool bUpdateAgain = true;
 			bool bShouldContinueToConsider = true;
 			while (bUpdateAgain && bShouldContinueToConsider)
diff --git a/Engine/Source/Runtime/Engine/Private/WorldPartition/WorldPartitionLevelStreamingPolicy.cpp b/Engine/Source/Runtime/Engine/Private/WorldPartition/WorldPartitionLevelStreamingPolicy.cpp
--- a/Engine/Source/Runtime/Engine/Private/WorldPartition/WorldPartitionLevelStreamingPolicy.cpp
+++ b/Engine/Source/Runtime/Engine/Private/WorldPartition/WorldPartitionLevelStreamingPolicy.cpp
@@ -20,12 +20,6 @@
 #include "Misc/PackageName.h"
 #endif
 
-static int32 GMaxLoadingLevelStreamingCells = 4;
-static FAutoConsoleVariableRef CMaxLoadingLevelStreamingCells(
-	TEXT("wp.Runtime.MaxLoadingLevelStreamingCells"),
-	GMaxLoadingLevelStreamingCells,
-	TEXT("Used to limit the number of concurrent loading world partition streaming cells."));
-
 int32 UWorldPartitionLevelStreamingPolicy::GetCellLoadingCount() const
 {
 	int32 CellLoadingCount = 0;
@@ -55,147 +49,6 @@ void UWorldPartitionLevelStreamingPolicy::ForEachActiveRuntimeCell(TFunctionRef<
 	}
 }
 
-int32 UWorldPartitionLevelStreamingPolicy::GetMaxCellsToLoad() const
-{
-	// This policy limits the number of concurrent loading streaming cells, except if match hasn't started
-	UWorld* World = WorldPartition->GetWorld();
-	return World->bMatchStarted ? (GMaxLoadingLevelStreamingCells - GetCellLoadingCount()) : MAX_int32;
-}
-
-void UWorldPartitionLevelStreamingPolicy::SetTargetStateForCells(EWorldPartitionRuntimeCellState TargetState, const TSet<const UWorldPartitionRuntimeCell*>& Cells)
-{
-	switch (TargetState)
-	{
-	case EWorldPartitionRuntimeCellState::Unloaded:
-		SetCellsStateToUnloaded(Cells);
-		break;
-	case EWorldPartitionRuntimeCellState::Loaded:
-		SetCellsStateToLoaded(Cells);
-		break;
-	case EWorldPartitionRuntimeCellState::Activated:
-		SetCellsStateToActivated(Cells);
-		break;
-	default:
-		check(false);
-	}
-}
-
-void UWorldPartitionLevelStreamingPolicy::SetCellsStateToLoaded(const TSet<const UWorldPartitionRuntimeCell*>& ToLoadCells)
-{
-	TRACE_CPUPROFILER_EVENT_SCOPE(UWorldPartitionLevelStreamingPolicy::SetCellsStateToLoaded);
-
-	int32 MaxCellsToLoad = GetMaxCellsToLoad();
-	
-	// Sort cells based on importance
-	TArray<const UWorldPartitionRuntimeCell*, TInlineAllocator<256>> SortedCells;
-	WorldPartition->RuntimeHash->SortStreamingCellsByImportance(ToLoadCells, StreamingSources, SortedCells);
-
-	// Trigger cell loading. Depending on actual state of cell limit loading.
-	for (const UWorldPartitionRuntimeCell* Cell : SortedCells)
-	{
-		UE_LOG(LogWorldPartition, Verbose, TEXT("UWorldPartitionLevelStreamingPolicy::LoadCells %s"), *Cell->GetName());
-		const UWorldPartitionRuntimeLevelStreamingCell* LevelCell = CastChecked<const UWorldPartitionRuntimeLevelStreamingCell>(Cell);
-		if (ActivatedCells.Contains(Cell))
-		{
-			LevelCell->Deactivate();
-			ActivatedCells.Remove(Cell);
-			LoadedCells.Add(Cell);
-		}
-		else if (MaxCellsToLoad > 0)
-		{
-			LevelCell->Load();
-			LoadedCells.Add(Cell);
-			if (!Cell->IsAlwaysLoaded())
-			{
-				--MaxCellsToLoad;
-			}
-		}
-	}
-}
-
-void UWorldPartitionLevelStreamingPolicy::SetCellsStateToActivated(const TSet<const UWorldPartitionRuntimeCell*>& ToActivateCells)
-{
-	TRACE_CPUPROFILER_EVENT_SCOPE(UWorldPartitionLevelStreamingPolicy::SetCellsStateToActivated);
-
-	int32 MaxCellsToLoad = GetMaxCellsToLoad();
-
-	// Sort cells based on importance
-	TArray<const UWorldPartitionRuntimeCell*, TInlineAllocator<256>> SortedCells;
-	WorldPartition->RuntimeHash->SortStreamingCellsByImportance(ToActivateCells, StreamingSources, SortedCells);
-
-	// Trigger cell activation. Depending on actual state of cell limit loading.
-	for (const UWorldPartitionRuntimeCell* Cell : SortedCells)
-	{
-		UE_LOG(LogWorldPartition, Verbose, TEXT("UWorldPartitionLevelStreamingPolicy::ActivateCells %s"), *Cell->GetName());
-		const UWorldPartitionRuntimeLevelStreamingCell* LevelCell = CastChecked<const UWorldPartitionRuntimeLevelStreamingCell>(Cell);
-		if (LoadedCells.Contains(LevelCell))
-		{
-			LoadedCells.Remove(Cell);
-			ActivatedCells.Add(Cell);
-			LevelCell->Activate();
-		}
-		else if (MaxCellsToLoad > 0)
-		{
-			if (!Cell->IsAlwaysLoaded())
-			{
-				--MaxCellsToLoad;
-			}
-			ActivatedCells.Add(Cell);
-			LevelCell->Activate();
-		}		
-	}
-}
-
-void UWorldPartitionLevelStreamingPolicy::SetCellsStateToUnloaded(const TSet<const UWorldPartitionRuntimeCell*>& ToUnloadCells)
-{
-	TRACE_CPUPROFILER_EVENT_SCOPE(UWorldPartitionLevelStreamingPolicy::SetCellsStateToUnloaded);
-
-	for (const UWorldPartitionRuntimeCell* Cell : ToUnloadCells)
-	{
-		UE_LOG(LogWorldPartition, Verbose, TEXT("UWorldPartitionLevelStreamingPolicy::UnloadCells %s"), *Cell->GetName());
-		const UWorldPartitionRuntimeLevelStreamingCell* LevelCell = CastChecked<const UWorldPartitionRuntimeLevelStreamingCell>(Cell);
-		LevelCell->Unload();
-		ActivatedCells.Remove(Cell);
-		LoadedCells.Remove(Cell);
-	}
-}
-
-ULevel* UWorldPartitionLevelStreamingPolicy::GetPreferredLoadedLevelToAddToWorld() const
-{
-	TRACE_CPUPROFILER_EVENT_SCOPE(UWorldPartitionLevelStreamingPolicy::GetPreferredLoadedLevelToAddToWorld);
-
-	check(WorldPartition->IsInitialized());
-	UWorld* World = WorldPartition->GetWorld();
-	if (World->GetNetMode() == NM_DedicatedServer)
-	{
-		return nullptr;
-	}
-
-	// Sort loaded cells based on importance (only those with a streaming level in MakingVisible state)
-	TArray<const UWorldPartitionRuntimeCell*, TInlineAllocator<256>> SortedMakingVisibileCells;
-	TSet<const UWorldPartitionRuntimeCell*> MakingVisibileCells;
-	MakingVisibileCells.Reserve(ActivatedCells.Num());
-	for (const UWorldPartitionRuntimeCell* LoadedCell : ActivatedCells)
-	{
-		if (const UWorldPartitionRuntimeLevelStreamingCell* LevelSreamingCell = Cast<const UWorldPartitionRuntimeLevelStreamingCell>(LoadedCell))
-		{
-			ULevelStreaming* LevelStreaming = LevelSreamingCell->GetLevelStreaming();
-			if (LevelStreaming && LevelStreaming->GetLoadedLevel() && (LevelStreaming->GetCurrentState() == ULevelStreaming::ECurrentState::MakingVisible))
-			{
-				MakingVisibileCells.Add(LoadedCell);
-			}
-		}
-	}
-	WorldPartition->RuntimeHash->SortStreamingCellsByImportance(MakingVisibileCells, StreamingSources, SortedMakingVisibileCells);
-	return SortedMakingVisibileCells.Num() > 0 ? CastChecked<UWorldPartitionRuntimeLevelStreamingCell>(SortedMakingVisibileCells[0])->GetLevelStreaming()->GetLoadedLevel() : nullptr;
-}
-
-EWorldPartitionRuntimeCellState UWorldPartitionLevelStreamingPolicy::GetCurrentStateForCell(const UWorldPartitionRuntimeCell* Cell) const
-{
-	const UWorldPartitionRuntimeLevelStreamingCell* LevelCell = CastChecked<const UWorldPartitionRuntimeLevelStreamingCell>(Cell);
-	return LevelCell->GetCurrentState();
-}
-
 #if WITH_EDITOR
 
 FString UWorldPartitionLevelStreamingPolicy::GetCellPackagePath(const FName& InCellName, const UWorld* InWorld)
diff --git a/Engine/Source/Runtime/Engine/Private/WorldPartition/WorldPartitionRuntimeLevelStreamingCell.cpp b/Engine/Source/Runtime/Engine/Private/WorldPartition/WorldPartitionRuntimeLevelStreamingCell.cpp
--- a/Engine/Source/Runtime/Engine/Private/WorldPartition/WorldPartitionRuntimeLevelStreamingCell.cpp
+++ b/Engine/Source/Runtime/Engine/Private/WorldPartition/WorldPartitionRuntimeLevelStreamingCell.cpp
@@ -258,6 +258,23 @@ void UWorldPartitionRuntimeLevelStreamingCell::Activate() const
 	}
 }
 
+bool UWorldPartitionRuntimeLevelStreamingCell::IsAddedToWorld() const
+{
+	return LevelStreaming && LevelStreaming->GetLoadedLevel() && LevelStreaming->GetLoadedLevel()->bIsVisible;
+}
+
+bool UWorldPartitionRuntimeLevelStreamingCell::CanAddToWorld() const
+{
+	return LevelStreaming &&
+		   LevelStreaming->GetLoadedLevel() &&
+		   (LevelStreaming->GetCurrentState() == ULevelStreaming::ECurrentState::MakingVisible);
+}
+
+ULevel* UWorldPartitionRuntimeLevelStreamingCell::GetLevel() const 
+{
+	return LevelStreaming ? LevelStreaming->GetLoadedLevel() : nullptr;
+}
+
 void UWorldPartitionRuntimeLevelStreamingCell::Unload() const
 {
 #if WITH_EDITOR
diff --git a/Engine/Source/Runtime/Engine/Private/WorldPartition/WorldPartitionStreamingPolicy.cpp b/Engine/Source/Runtime/Engine/Private/WorldPartition/WorldPartitionStreamingPolicy.cpp
--- a/Engine/Source/Runtime/Engine/Private/WorldPartition/WorldPartitionStreamingPolicy.cpp
+++ b/Engine/Source/Runtime/Engine/Private/WorldPartition/WorldPartitionStreamingPolicy.cpp
@@ -27,9 +27,17 @@
 #include "LevelEditorViewport.h"
 #endif
 
+static int32 GMaxLoadingStreamingCells = 4;
+static FAutoConsoleVariableRef CMaxLoadingStreamingCells(
+	TEXT("wp.Runtime.MaxLoadingStreamingCells"),
+	GMaxLoadingStreamingCells,
+	TEXT("Used to limit the number of concurrent loading world partition streaming cells."));
+
 UWorldPartitionStreamingPolicy::UWorldPartitionStreamingPolicy(const FObjectInitializer& ObjectInitializer)
 : Super(ObjectInitializer) 
 , bIsServerLoadingDone(false)
+, UpdateStreamingStateEpoch(0)
+, SortedAddToWorldCellsEpoch(INT_MIN)
 {
 	if (!IsTemplate())
 	{
@@ -118,6 +126,9 @@ void UWorldPartitionStreamingPolicy::UpdateStreamingState()
 
 	UWorld* World = WorldPartition->GetWorld();
 	check(World && World->IsGameWorld());
+
+	// Increment/dirty UpdateStreamingStateEpoch
+	++UpdateStreamingStateEpoch;
 	
 	// Update streaming sources
 	UpdateStreamingSources();
@@ -222,6 +233,146 @@ void UWorldPartitionStreamingPolicy::UpdateStreamingState()
 	}
 }
 
+int32 UWorldPartitionStreamingPolicy::GetMaxCellsToLoad() const
+{
+	// This policy limits the number of concurrent loading streaming cells, except if match hasn't started
+	UWorld* World = WorldPartition->GetWorld();
+	return World->bMatchStarted ? (GMaxLoadingStreamingCells - GetCellLoadingCount()) : MAX_int32;
+}
+
+void UWorldPartitionStreamingPolicy::SetTargetStateForCells(EWorldPartitionRuntimeCellState TargetState, const TSet<const UWorldPartitionRuntimeCell*>& Cells)
+{
+	switch (TargetState)
+	{
+	case EWorldPartitionRuntimeCellState::Unloaded:
+		SetCellsStateToUnloaded(Cells);
+		break;
+	case EWorldPartitionRuntimeCellState::Loaded:
+		SetCellsStateToLoaded(Cells);
+		break;
+	case EWorldPartitionRuntimeCellState::Activated:
+		SetCellsStateToActivated(Cells);
+		break;
+	default:
+		check(false);
+	}
+}
+
+void UWorldPartitionStreamingPolicy::SetCellsStateToLoaded(const TSet<const UWorldPartitionRuntimeCell*>& ToLoadCells)
+{
+	TRACE_CPUPROFILER_EVENT_SCOPE(UWorldPartitionStreamingPolicy::SetCellsStateToLoaded);
+
+	int32 MaxCellsToLoad = GetMaxCellsToLoad();
+
+	// Sort cells based on importance
+	TArray<const UWorldPartitionRuntimeCell*, TInlineAllocator<256>> SortedCells;
+	WorldPartition->RuntimeHash->SortStreamingCellsByImportance(ToLoadCells, StreamingSources, SortedCells);
+
+	// Trigger cell loading. Depending on actual state of cell limit loading.
+	for (const UWorldPartitionRuntimeCell* Cell : SortedCells)
+	{
+		UE_LOG(LogWorldPartition, Verbose, TEXT("UWorldPartitionStreamingPolicy::LoadCells %s"), *Cell->GetName());
+		if (ActivatedCells.Contains(Cell))
+		{
+			Cell->Deactivate();
+			ActivatedCells.Remove(Cell);
+			LoadedCells.Add(Cell);
+		}
+		else if (MaxCellsToLoad > 0)
+		{
+			Cell->Load();
+			LoadedCells.Add(Cell);
+			if (!Cell->IsAlwaysLoaded())
+			{
+				--MaxCellsToLoad;
+			}
+		}
+	}
+}
+
+void UWorldPartitionStreamingPolicy::SetCellsStateToActivated(const TSet<const UWorldPartitionRuntimeCell*>& ToActivateCells)
+{
+	TRACE_CPUPROFILER_EVENT_SCOPE(UWorldPartitionStreamingPolicy::SetCellsStateToActivated);
+
+	int32 MaxCellsToLoad = GetMaxCellsToLoad();
+
+	// Sort cells based on importance
+	TArray<const UWorldPartitionRuntimeCell*, TInlineAllocator<256>> SortedCells;
+	WorldPartition->RuntimeHash->SortStreamingCellsByImportance(ToActivateCells, StreamingSources, SortedCells);
+
+	// Trigger cell activation. Depending on actual state of cell limit loading.
+	for (const UWorldPartitionRuntimeCell* Cell : SortedCells)
+	{
+		UE_LOG(LogWorldPartition, Verbose, TEXT("UWorldPartitionStreamingPolicy::ActivateCells %s"), *Cell->GetName());
+		if (LoadedCells.Contains(Cell))
+		{
+			LoadedCells.Remove(Cell);
+			ActivatedCells.Add(Cell);
+			Cell->Activate();
+		}
+		else if (MaxCellsToLoad > 0)
+		{
+			if (!Cell->IsAlwaysLoaded())
+			{
+				--MaxCellsToLoad;
+			}
+			ActivatedCells.Add(Cell);
+			Cell->Activate();
+		}
+	}
+}
+
+void UWorldPartitionStreamingPolicy::SetCellsStateToUnloaded(const TSet<const UWorldPartitionRuntimeCell*>& ToUnloadCells)
+{
+	TRACE_CPUPROFILER_EVENT_SCOPE(UWorldPartitionStreamingPolicy::SetCellsStateToUnloaded);
+
+	for (const UWorldPartitionRuntimeCell* Cell : ToUnloadCells)
+	{
+		UE_LOG(LogWorldPartition, Verbose, TEXT("UWorldPartitionStreamingPolicy::UnloadCells %s"), *Cell->GetName());
+		Cell->Unload();
+		ActivatedCells.Remove(Cell);
+		LoadedCells.Remove(Cell);
+	}
+}
+
+ULevel* UWorldPartitionStreamingPolicy::GetPreferredLoadedLevelToAddToWorld() const
+{
+	TRACE_CPUPROFILER_EVENT_SCOPE(UWorldPartitionStreamingPolicy::GetPreferredLoadedLevelToAddToWorld);
+
+	check(WorldPartition->IsInitialized());
+	UWorld* World = WorldPartition->GetWorld();
+	if (World->GetNetMode() == NM_DedicatedServer)
+	{
+		return nullptr;
+	}
+
+	// Sort activated cells based on importance only if CellState changed
+	if (SortedAddToWorldCellsEpoch != UpdateStreamingStateEpoch)
+	{
+		SortedAddToWorldCellsEpoch = UpdateStreamingStateEpoch;
+		TSet<const UWorldPartitionRuntimeCell*> AddToWorldCells;
+		for (const UWorldPartitionRuntimeCell* ActivatedCell : ActivatedCells)
+		{
+			if (!ActivatedCell->IsAddedToWorld())
+			{
+				AddToWorldCells.Add(ActivatedCell);
+			}
+		}
+		WorldPartition->RuntimeHash->SortStreamingCellsByImportance(AddToWorldCells, StreamingSources, SortedAddToWorldCells);
+	}
+
+	// Find first cell with a streaming level in MakingVisible state
+	for (const UWorldPartitionRuntimeCell* Cell : SortedAddToWorldCells)
+	{
+		if (Cell->CanAddToWorld())
+		{
+			check(!Cell->IsAddedToWorld());
+			return Cell->GetLevel();
+		}
+	}
+	return nullptr;
+}
+
 bool UWorldPartitionStreamingPolicy::IsStreamingCompleted(EWorldPartitionRuntimeCellState QueryState, const TArray<FWorldPartitionStreamingQuerySource>& QuerySources, bool bExactState) const
 {
 	const UDataLayerSubsystem* DataLayerSubsystem = WorldPartition->GetWorld()->GetSubsystem<UDataLayerSubsystem>();
@@ -234,7 +385,7 @@ bool UWorldPartitionStreamingPolicy::IsStreamingCompleted(EWorldPartitionRuntime
 
 		for (const UWorldPartitionRuntimeCell* Cell : Cells)
 		{
-			EWorldPartitionRuntimeCellState CellState = GetCurrentStateForCell(Cell);
+			EWorldPartitionRuntimeCellState CellState = Cell->GetCurrentState();
 			if (CellState != QueryState)
 			{
 				bool bSkipCell = false;
diff --git a/Engine/Source/Runtime/Engine/Public/WorldPartition/WorldPartitionLevelStreamingPolicy.h b/Engine/Source/Runtime/Engine/Public/WorldPartition/WorldPartitionLevelStreamingPolicy.h
--- a/Engine/Source/Runtime/Engine/Public/WorldPartition/WorldPartitionLevelStreamingPolicy.h
+++ b/Engine/Source/Runtime/Engine/Public/WorldPartition/WorldPartitionLevelStreamingPolicy.h
@@ -22,9 +22,6 @@ class UWorldPartitionLevelStreamingPolicy : public UWorldPartitionStreamingPolic
 	GENERATED_BODY()
 
 public:
-	virtual void SetTargetStateForCells(EWorldPartitionRuntimeCellState TargetState, const TSet<const UWorldPartitionRuntimeCell*>& Cells) override;
-	virtual EWorldPartitionRuntimeCellState GetCurrentStateForCell(const UWorldPartitionRuntimeCell* Cell) const override;
-	virtual class ULevel* GetPreferredLoadedLevelToAddToWorld() const override;
 	virtual void DrawRuntimeCellsDetails(class UCanvas* Canvas, FVector2D& Offset) override;
 	virtual void DrawStreamingStatusLegend(UCanvas* Canvas, FVector2D& Offset) override;
 
@@ -39,15 +36,11 @@ public:
 	virtual UObject* GetSubObject(const TCHAR* SubObjectPath) override;
 
 protected:
+	virtual int32 GetCellLoadingCount() const override;
+
 	void ForEachActiveRuntimeCell(TFunctionRef<void(const UWorldPartitionRuntimeCell*)> Func) const;
 
 private:
-	void SetCellsStateToLoaded(const TSet<const UWorldPartitionRuntimeCell*>& ToLoadCells);
-	void SetCellsStateToActivated(const TSet<const UWorldPartitionRuntimeCell*>& ToActivateCells);
-	void SetCellsStateToUnloaded(const TSet<const UWorldPartitionRuntimeCell*>& ToUnloadCells);
-	int32 GetCellLoadingCount() const;
-	int32 GetMaxCellsToLoad() const;
-
 #if WITH_EDITORONLY_DATA
 	UPROPERTY()
 	TMap<FName, FName> ActorToCellRemapping;
diff --git a/Engine/Source/Runtime/Engine/Public/WorldPartition/WorldPartitionRuntimeCell.h b/Engine/Source/Runtime/Engine/Public/WorldPartition/WorldPartitionRuntimeCell.h
--- a/Engine/Source/Runtime/Engine/Public/WorldPartition/WorldPartitionRuntimeCell.h
+++ b/Engine/Source/Runtime/Engine/Public/WorldPartition/WorldPartitionRuntimeCell.h
@@ -111,12 +111,18 @@ class UWorldPartitionRuntimeCell : public UObject
 {
 	GENERATED_UCLASS_BODY()
 
+	virtual void Load() const PURE_VIRTUAL(UWorldPartitionRuntimeCell::Load,);
+	virtual void Unload() const PURE_VIRTUAL(UWorldPartitionRuntimeCell::Unload,);
+	virtual void Activate() const PURE_VIRTUAL(UWorldPartitionRuntimeCell::Activate,);
+	virtual void Deactivate() const PURE_VIRTUAL(UWorldPartitionRuntimeCell::Deactivate,);
+	virtual bool IsAddedToWorld() const PURE_VIRTUAL(UWorldPartitionRuntimeCell::IsAddedToWorld, return false;);
+	virtual bool CanAddToWorld() const PURE_VIRTUAL(UWorldPartitionRuntimeCell::CanAddToWorld, return false;);
+	virtual ULevel* GetLevel() const PURE_VIRTUAL(UWorldPartitionRuntimeCell::GetLevel, return nullptr;);
+	virtual EWorldPartitionRuntimeCellState GetCurrentState() const PURE_VIRTUAL(UWorldPartitionRuntimeCell::GetCurrentState, return EWorldPartitionRuntimeCellState::Unloaded;);
 	virtual FLinearColor GetDebugColor() const { return FLinearColor::Black; }
 	virtual bool IsAlwaysLoaded() const { return bIsAlwaysLoaded; }
 	virtual void SetIsAlwaysLoaded(bool bInIsAlwaysLoaded) { bIsAlwaysLoaded = bInIsAlwaysLoaded; }
 	virtual void SetPriority(int32 InPriority) { Priority = InPriority; }
-	bool HasDataLayers() const { return !DataLayers.IsEmpty(); }
-	const TArray<FName>& GetDataLayers() const { return DataLayers; }
 	virtual EStreamingStatus GetStreamingStatus() const { return LEVEL_Unloaded; }
 	virtual bool IsLoading() const { return false; }
 	virtual const FString& GetDebugName() const { return DebugName; }
@@ -128,6 +134,9 @@ class UWorldPartitionRuntimeCell : public UObject
 
 	static void DirtyStreamingSourceCacheEpoch() { ++UWorldPartitionRuntimeCell::StreamingSourceCacheEpoch; }
 
+	bool HasDataLayers() const { return !DataLayers.IsEmpty(); }
+	const TArray<FName>& GetDataLayers() const { return DataLayers; }
+
 #if WITH_EDITOR
 	void SetDataLayers(const TArray<const UDataLayer*>& InDataLayers);
 	void SetDebugInfo(FIntVector InCoords, FName InGridName);
diff --git a/Engine/Source/Runtime/Engine/Public/WorldPartition/WorldPartitionRuntimeLevelStreamingCell.h b/Engine/Source/Runtime/Engine/Public/WorldPartition/WorldPartitionRuntimeLevelStreamingCell.h
--- a/Engine/Source/Runtime/Engine/Public/WorldPartition/WorldPartitionRuntimeLevelStreamingCell.h
+++ b/Engine/Source/Runtime/Engine/Public/WorldPartition/WorldPartitionRuntimeLevelStreamingCell.h
@@ -10,28 +10,36 @@ class UWorldPartitionRuntimeLevelStreamingCell : public UWorldPartitionRuntimeSp
 {
 	GENERATED_UCLASS_BODY()
 
-	void Load() const;
-	void Unload() const;
-	void Activate() const;
-	void Deactivate() const;
-	int32 GetStreamingPriority() const { return Level; }
-	class UWorldPartitionLevelStreamingDynamic* GetLevelStreaming() const;
+	//~Begin UWorldPartitionRuntimeCell Interface
+	virtual void Load() const override;
+	virtual void Unload() const override;
+	virtual void Activate() const override;
+	virtual void Deactivate() const override;
+	virtual bool IsAddedToWorld() const override;
+	virtual bool CanAddToWorld() const override;
+	virtual ULevel* GetLevel() const override;
+	virtual EWorldPartitionRuntimeCellState GetCurrentState() const override;
 	virtual FLinearColor GetDebugColor() const override;
+	virtual void SetIsAlwaysLoaded(bool bInIsAlwaysLoaded) override;
 	virtual EStreamingStatus GetStreamingStatus() const override;
 	virtual bool IsLoading() const override;
-	EWorldPartitionRuntimeCellState GetCurrentState() const;
+	//~End UWorldPartitionRuntimeCell Interface
+
+	int32 GetStreamingPriority() const { return Level; }
+	class UWorldPartitionLevelStreamingDynamic* GetLevelStreaming() const;
+	
 
 #if WITH_EDITOR
+	//~Begin UWorldPartitionRuntimeCell Interface
 	virtual void AddActorToCell(const FWorldPartitionActorDescView& ActorDescView, uint64 InContainerID, const FTransform& InContainerTransform, const UActorDescContainer* InContainer) override;
 	virtual int32 GetActorCount() const override;
-	const TArray<FWorldPartitionRuntimeCellObjectMapping>& GetPackages() const { return Packages; }
-
 	// Cook methods
 	virtual bool PopulateGeneratedPackageForCook(UPackage* InPackage) override;
 	virtual FString GetPackageNameToCreate() const override;
-#endif
+	//~End UWorldPartitionRuntimeCell Interface
 
-	virtual void SetIsAlwaysLoaded(bool bInIsAlwaysLoaded) override;
+	const TArray<FWorldPartitionRuntimeCellObjectMapping>& GetPackages() const { return Packages; }
+#endif
 
 private:
 	UFUNCTION()
diff --git a/Engine/Source/Runtime/Engine/Public/WorldPartition/WorldPartitionStreamingPolicy.h b/Engine/Source/Runtime/Engine/Public/WorldPartition/WorldPartitionStreamingPolicy.h
--- a/Engine/Source/Runtime/Engine/Public/WorldPartition/WorldPartitionStreamingPolicy.h
+++ b/Engine/Source/Runtime/Engine/Public/WorldPartition/WorldPartitionStreamingPolicy.h
@@ -23,9 +23,8 @@ class UWorldPartitionStreamingPolicy : public UObject
 
 public:
 	virtual void UpdateStreamingState();
-	virtual void SetTargetStateForCells(EWorldPartitionRuntimeCellState TargetState, const TSet<const UWorldPartitionRuntimeCell*>& Cells) PURE_VIRTUAL(UWorldPartitionStreamingPolicy::SetTargetStateForCells, );
-	virtual EWorldPartitionRuntimeCellState GetCurrentStateForCell(const UWorldPartitionRuntimeCell* Cell) const PURE_VIRTUAL(UWorldPartitionStreamingPolicy::GetCurrentStateForCell, return EWorldPartitionRuntimeCellState::Unloaded; );
-	virtual class ULevel* GetPreferredLoadedLevelToAddToWorld() const { return nullptr; }
+	virtual void SetTargetStateForCells(EWorldPartitionRuntimeCellState TargetState, const TSet<const UWorldPartitionRuntimeCell*>& Cells);
+	virtual class ULevel* GetPreferredLoadedLevelToAddToWorld() const;
 	virtual FVector2D GetDrawRuntimeHash2DDesiredFootprint(const FVector2D& CanvasSize);
 	virtual void DrawRuntimeHash2D(class UCanvas* Canvas, const FVector2D& PartitionCanvasSize, FVector2D& Offset);
 	virtual void DrawRuntimeHash3D();
@@ -48,11 +47,20 @@ public:
 	const TArray<FWorldPartitionStreamingSource>& GetStreamingSources() const { return StreamingSources; }
 
 protected:
-	void UpdateStreamingSources();
+	virtual void SetCellsStateToLoaded(const TSet<const UWorldPartitionRuntimeCell*>& ToLoadCells);
+	virtual void SetCellsStateToActivated(const TSet<const UWorldPartitionRuntimeCell*>& ToActivateCells);
+	virtual void SetCellsStateToUnloaded(const TSet<const UWorldPartitionRuntimeCell*>& ToUnloadCells);
+	virtual int32 GetCellLoadingCount() const { return 0; }
+	virtual int32 GetMaxCellsToLoad() const;
+	virtual void UpdateStreamingSources();
 
 	bool bIsServerLoadingDone;
 	const UWorldPartition* WorldPartition;
 	TSet<const UWorldPartitionRuntimeCell*> LoadedCells;
 	TSet<const UWorldPartitionRuntimeCell*> ActivatedCells;
 	TArray<FWorldPartitionStreamingSource> StreamingSources;
+	
+	int32 UpdateStreamingStateEpoch;
+	mutable int32 SortedAddToWorldCellsEpoch;
+	mutable TArray<const UWorldPartitionRuntimeCell*, TInlineAllocator<256>> SortedAddToWorldCells;
 };
\ No newline at end of file
