Consistency improvements for the H/ISM component API

 - AddInstance and AddInstances can now take the transform in either local or world space, as defined by the bWorldSpace argument.
   - This matches the API of GetInstanceTransform and UpdateInstanceTransform, and means that AddInstanceWorldSpace can be deprecated.
 - RemoveInstances is now available on the base ISM component, rather than being HISM specific.
   - This allows code operating on ISM components via the base API to still take advantage of HISM batched remove optimizations.

diff --git a/Engine/Source/Editor/UnrealEd/Private/Elements/SMInstance/SMInstanceElementEditorWorldInterface.cpp b/Engine/Source/Editor/UnrealEd/Private/Elements/SMInstance/SMInstanceElementEditorWorldInterface.cpp
--- a/Engine/Source/Editor/UnrealEd/Private/Elements/SMInstance/SMInstanceElementEditorWorldInterface.cpp
+++ b/Engine/Source/Editor/UnrealEd/Private/Elements/SMInstance/SMInstanceElementEditorWorldInterface.cpp
@@ -52,13 +52,11 @@ bool USMInstanceElementEditorWorldInterface::DeleteElements(TArrayView<const FTy
 		{
 			BatchedInstancesToDeletePair.Key->Modify();
 
-			// Sort the instance indices in descending order to avoid invalidating resolved indices mid-delete
-			BatchedInstancesToDeletePair.Value.Sort(TGreater<int32>());
 			for (const int32 InstanceIndex : BatchedInstancesToDeletePair.Value)
 			{
 				InSelectionSet->DeselectElement(UEngineElementsLibrary::AcquireEditorSMInstanceElementHandle(FSMInstanceId{ BatchedInstancesToDeletePair.Key, InstanceIndex }), FTypedElementSelectionOptions());
-				bDidDelete |= BatchedInstancesToDeletePair.Key->RemoveInstance(InstanceIndex);
 			}
+			bDidDelete |= BatchedInstancesToDeletePair.Key->RemoveInstances(BatchedInstancesToDeletePair.Value);
 		}
 
 		return bDidDelete;
diff --git a/Engine/Source/Runtime/Engine/Classes/Components/HierarchicalInstancedStaticMeshComponent.h b/Engine/Source/Runtime/Engine/Classes/Components/HierarchicalInstancedStaticMeshComponent.h
--- a/Engine/Source/Runtime/Engine/Classes/Components/HierarchicalInstancedStaticMeshComponent.h
+++ b/Engine/Source/Runtime/Engine/Classes/Components/HierarchicalInstancedStaticMeshComponent.h
@@ -219,9 +219,10 @@ public:
 #endif
 
 	// UInstancedStaticMesh interface
-	virtual int32 AddInstance(const FTransform& InstanceTransform) override;
-	virtual TArray<int32> AddInstances(const TArray<FTransform>& InstanceTransforms, bool bShouldReturnIndices) override;
+	virtual int32 AddInstance(const FTransform& InstanceTransform, bool bWorldSpace = false) override;
+	virtual TArray<int32> AddInstances(const TArray<FTransform>& InstanceTransforms, bool bShouldReturnIndices, bool bWorldSpace = false) override;
 	virtual bool RemoveInstance(int32 InstanceIndex) override;
+	virtual bool RemoveInstances(const TArray<int32>& InstancesToRemove) override;
 	virtual bool UpdateInstanceTransform(int32 InstanceIndex, const FTransform& NewInstanceTransform, bool bWorldSpace, bool bMarkRenderStateDirty = false, bool bTeleport = false) override;
 	virtual bool SetCustomDataValue(int32 InstanceIndex, int32 CustomDataIndex, float CustomDataValue, bool bMarkRenderStateDirty = false) override;
 	virtual bool SetCustomData(int32 InstanceIndex, const TArray<float>& InCustomData, bool bMarkRenderStateDirty = false) override;
@@ -235,10 +236,6 @@ public:
 	virtual void PreAllocateInstancesMemory(int32 AddedInstanceCount) override;
 	virtual bool SupportsRemoveSwap() const override { return true; }
 
-	/** Removes all the instances with indices specified in the InstancesToRemove array. Returns true on success. */
-	UFUNCTION(BlueprintCallable, Category = "Components|InstancedStaticMesh")
-	bool RemoveInstances(const TArray<int32>& InstancesToRemove);
-
 	/** Get the number of instances that overlap a given sphere */
 	int32 GetOverlappingSphereCount(const FSphere& Sphere) const;
 	/** Get the number of instances that overlap a given box */
diff --git a/Engine/Source/Runtime/Engine/Classes/Components/InstancedStaticMeshComponent.h b/Engine/Source/Runtime/Engine/Classes/Components/InstancedStaticMeshComponent.h
--- a/Engine/Source/Runtime/Engine/Classes/Components/InstancedStaticMeshComponent.h
+++ b/Engine/Source/Runtime/Engine/Classes/Components/InstancedStaticMeshComponent.h
@@ -175,17 +175,21 @@ class ENGINE_API UInstancedStaticMeshComponent : public UStaticMeshComponent
 	/** Returns the render instance buffer index. */
 	FORCEINLINE int32 GetRenderIndex(int32 InInstanceIndex) const { return InstanceReorderTable.IsValidIndex(InInstanceIndex) ? InstanceReorderTable[InInstanceIndex] : InInstanceIndex; }
 
-	/** Add an instance to this component. Transform is given in local space of this component. */
+	/** Add an instance to this component. Transform is given in local space of this component unless bWorldSpace is set. */
 	UFUNCTION(BlueprintCallable, Category="Components|InstancedStaticMesh")
-	virtual int32 AddInstance(const FTransform& InstanceTransform);
+	virtual int32 AddInstance(const FTransform& InstanceTransform, bool bWorldSpace = false);
 
-	/** Add multiple instances to this component. Transform is given in local space of this component. */
+	/** Add multiple instances to this component. Transform is given in local space of this component unless bWorldSpace is set. */
 	UFUNCTION(BlueprintCallable, Category="Components|InstancedStaticMesh")
-	virtual TArray<int32> AddInstances(const TArray<FTransform>& InstanceTransforms, bool bShouldReturnIndices);
+	virtual TArray<int32> AddInstances(const TArray<FTransform>& InstanceTransforms, bool bShouldReturnIndices, bool bWorldSpace = false);
 
 	/** Add an instance to this component. Transform is given in world space. */
 	UFUNCTION(BlueprintCallable, Category = "Components|InstancedStaticMesh")
-	int32 AddInstanceWorldSpace(const FTransform& WorldTransform);
+	int32 AddInstanceWorldSpace(const FTransform& WorldTransform)
+	{
+		// TODO: Deprecate this shim
+		return AddInstance(WorldTransform, /*bWorldSpace*/true);
+	}
 
 	/** Update custom data for specific instance */
 	UFUNCTION(BlueprintCallable, Category = "Components|InstancedStaticMesh")
@@ -265,10 +269,14 @@ class ENGINE_API UInstancedStaticMeshComponent : public UStaticMeshComponent
 
 	virtual bool BatchUpdateInstancesData(int32 StartInstanceIndex, int32 NumInstances, FInstancedStaticMeshInstanceData* StartInstanceData, bool bMarkRenderStateDirty = false, bool bTeleport = false);
 
-	/** Remove the instance specified. Returns True on success. Note that this will leave the array in order, but may shrink it. */
+	/** Remove the instance specified. Returns True on success. */
 	UFUNCTION(BlueprintCallable, Category = "Components|InstancedStaticMesh")
 	virtual bool RemoveInstance(int32 InstanceIndex);
 
+	/** Remove the instances specified. Returns True on success. */
+	UFUNCTION(BlueprintCallable, Category = "Components|InstancedStaticMesh")
+	virtual bool RemoveInstances(const TArray<int32>& InstancesToRemove);
+
 	/** Clear all instances being rendered by this component. */
 	UFUNCTION(BlueprintCallable, Category="Components|InstancedStaticMesh")
 	virtual void ClearInstances();
@@ -410,10 +418,10 @@ protected:
 	virtual bool SupportsPartialNavigationUpdate() const { return false; }
 
 	/** Internal version of AddInstance */
-	int32 AddInstanceInternal(int32 InstanceIndex, FInstancedStaticMeshInstanceData* InNewInstanceData, const FTransform& InstanceTransform);
+	int32 AddInstanceInternal(int32 InstanceIndex, FInstancedStaticMeshInstanceData* InNewInstanceData, const FTransform& InstanceTransform, bool bWorldSpace);
 
 	/** Internal implementation of AddInstances */
-	TArray<int32> AddInstancesInternal(int32 Count, const TArray<FTransform>& InstanceTransforms, bool bShouldReturnIndices);
+	TArray<int32> AddInstancesInternal(const TArray<FTransform>& InstanceTransforms, bool bShouldReturnIndices, bool bWorldSpace);
 
 	/** Internal version of RemoveInstance */	
 	bool RemoveInstanceInternal(int32 InstanceIndex, bool InstanceAlreadyRemoved);
@@ -469,7 +477,7 @@ struct HInstancedStaticMeshInstance : public HHitProxy
 
 	virtual EMouseCursor::Type GetMouseCursor() override
 	{
-		return EMouseCursor::CardinalCross;
+		return EMouseCursor::Crosshairs;
 	}
 };
 
diff --git a/Engine/Source/Runtime/Engine/Private/HierarchicalInstancedStaticMesh.cpp b/Engine/Source/Runtime/Engine/Private/HierarchicalInstancedStaticMesh.cpp
--- a/Engine/Source/Runtime/Engine/Private/HierarchicalInstancedStaticMesh.cpp
+++ b/Engine/Source/Runtime/Engine/Private/HierarchicalInstancedStaticMesh.cpp
@@ -2491,11 +2491,11 @@ void UHierarchicalInstancedStaticMeshComponent::PreAllocateInstancesMemory(int32
 	UnbuiltInstanceBoundsList.Reserve(UnbuiltInstanceBoundsList.Num() + AddedInstanceCount);
 }
 
-int32 UHierarchicalInstancedStaticMeshComponent::AddInstance(const FTransform& InstanceTransform)
+int32 UHierarchicalInstancedStaticMeshComponent::AddInstance(const FTransform& InstanceTransform, bool bWorldSpace)
 {
 	SCOPE_CYCLE_COUNTER(STAT_HISMCAddInstance);
 
-	int32 InstanceIndex = UInstancedStaticMeshComponent::AddInstance(InstanceTransform);
+	int32 InstanceIndex = UInstancedStaticMeshComponent::AddInstance(InstanceTransform, bWorldSpace);
 
 	// The tree will be fully rebuilt once the static mesh compilation is finished, no need for incremental update in that case.
 	if (InstanceIndex != INDEX_NONE && GetStaticMesh() && !GetStaticMesh()->IsCompiling() && GetStaticMesh()->HasValidRenderData())
@@ -2514,9 +2514,9 @@ int32 UHierarchicalInstancedStaticMeshComponent::AddInstance(const FTransform& I
 			++InstanceCountToRender;
 		}
 
-		InstanceUpdateCmdBuffer.AddInstance(InstanceTransform.ToMatrixWithScale());
+		InstanceUpdateCmdBuffer.AddInstance(PerInstanceSMData[InstanceIndex].Transform);
 
-		const FBox NewInstanceBounds = GetStaticMesh()->GetBounds().GetBox().TransformBy(InstanceTransform);
+		const FBox NewInstanceBounds = GetStaticMesh()->GetBounds().GetBox().TransformBy(PerInstanceSMData[InstanceIndex].Transform);
 		UnbuiltInstanceBounds += NewInstanceBounds;
 		UnbuiltInstanceBoundsList.Add(NewInstanceBounds);
 
@@ -2529,13 +2529,11 @@ int32 UHierarchicalInstancedStaticMeshComponent::AddInstance(const FTransform& I
 	return InstanceIndex;
 }
 
-TArray<int32> UHierarchicalInstancedStaticMeshComponent::AddInstances(const TArray<FTransform>& InstanceTransforms, bool bShouldReturnIndices)
+TArray<int32> UHierarchicalInstancedStaticMeshComponent::AddInstances(const TArray<FTransform>& InstanceTransforms, bool bShouldReturnIndices, bool bWorldSpace)
 {
 	SCOPE_CYCLE_COUNTER(STAT_HISMCAddInstances);
 
-	int32 BaseIndex = PerInstanceSMData.Num();
-
-	TArray<int32> InstanceIndices = UInstancedStaticMeshComponent::AddInstances(InstanceTransforms, true);
+	TArray<int32> InstanceIndices = UInstancedStaticMeshComponent::AddInstances(InstanceTransforms, true, bWorldSpace);
 
 	// The tree will be fully rebuilt once the static mesh compilation is finished, no need for incremental update in that case.
 	if (InstanceIndices.Num() > 0 && GetStaticMesh() && !GetStaticMesh()->IsCompiling() && GetStaticMesh()->HasValidRenderData())
@@ -2548,19 +2546,15 @@ TArray<int32> UHierarchicalInstancedStaticMeshComponent::AddInstances(const TArr
 		InstanceReorderTable.Reserve(InstanceReorderTable.Num() + Count);
 		UnbuiltInstanceBoundsList.Reserve(UnbuiltInstanceBoundsList.Num() + Count);
 
-		int32 TransformIndexOffset = BaseIndex;
-
 		const int32 InitialBufferOffset = InstanceCountToRender - InstanceReorderTable.Num();
 
 		for (const int32 InstanceIndex : InstanceIndices)
 		{
-			TransformIndexOffset = InstanceIndex - BaseIndex;
-
 			InstanceReorderTable.Add(InitialBufferOffset + InstanceIndex);
 
-			InstanceUpdateCmdBuffer.AddInstance(InstanceTransforms[TransformIndexOffset].ToMatrixWithScale());
+			InstanceUpdateCmdBuffer.AddInstance(PerInstanceSMData[InstanceIndex].Transform);
 
-			const FBox NewInstanceBounds = GetStaticMesh()->GetBounds().GetBox().TransformBy(InstanceTransforms[TransformIndexOffset]);
+			const FBox NewInstanceBounds = GetStaticMesh()->GetBounds().GetBox().TransformBy(PerInstanceSMData[InstanceIndex].Transform);
 			UnbuiltInstanceBounds += NewInstanceBounds;
 			UnbuiltInstanceBoundsList.Add(NewInstanceBounds);
 		}
diff --git a/Engine/Source/Runtime/Engine/Private/InstancedStaticMesh.cpp b/Engine/Source/Runtime/Engine/Private/InstancedStaticMesh.cpp
--- a/Engine/Source/Runtime/Engine/Private/InstancedStaticMesh.cpp
+++ b/Engine/Source/Runtime/Engine/Private/InstancedStaticMesh.cpp
@@ -2705,7 +2705,7 @@ void UInstancedStaticMeshComponent::PreAllocateInstancesMemory(int32 AddedInstan
 	PerInstanceSMCustomData.Reserve(PerInstanceSMCustomData.Num() + AddedInstanceCount * NumCustomDataFloats);
 }
 
-int32 UInstancedStaticMeshComponent::AddInstanceInternal(int32 InstanceIndex, FInstancedStaticMeshInstanceData* InNewInstanceData, const FTransform& InstanceTransform)
+int32 UInstancedStaticMeshComponent::AddInstanceInternal(int32 InstanceIndex, FInstancedStaticMeshInstanceData* InNewInstanceData, const FTransform& InstanceTransform, bool bWorldSpace)
 {
 	FInstancedStaticMeshInstanceData* NewInstanceData = InNewInstanceData;
 
@@ -2714,7 +2714,8 @@ int32 UInstancedStaticMeshComponent::AddInstanceInternal(int32 InstanceIndex, FI
 		NewInstanceData = &PerInstanceSMData.AddDefaulted_GetRef();
 	}
 
-	SetupNewInstanceData(*NewInstanceData, InstanceIndex, InstanceTransform);
+	const FTransform LocalTransform = bWorldSpace ? InstanceTransform.GetRelativeTransform(GetComponentTransform()) : InstanceTransform;
+	SetupNewInstanceData(*NewInstanceData, InstanceIndex, LocalTransform);
 
 	// Add custom data to instance
 	PerInstanceSMCustomData.AddZeroed(NumCustomDataFloats);
@@ -2740,13 +2741,15 @@ int32 UInstancedStaticMeshComponent::AddInstanceInternal(int32 InstanceIndex, FI
 	return InstanceIndex;
 }
 
-int32 UInstancedStaticMeshComponent::AddInstance(const FTransform& InstanceTransform)
+int32 UInstancedStaticMeshComponent::AddInstance(const FTransform& InstanceTransform, bool bWorldSpace)
 {
-	return AddInstanceInternal(PerInstanceSMData.Num(), nullptr, InstanceTransform);
+	return AddInstanceInternal(PerInstanceSMData.Num(), nullptr, InstanceTransform, bWorldSpace);
 }
 
-TArray<int32> UInstancedStaticMeshComponent::AddInstancesInternal(int32 Count, const TArray<FTransform>& InstanceTransforms, bool bShouldReturnIndices)
+TArray<int32> UInstancedStaticMeshComponent::AddInstancesInternal(const TArray<FTransform>& InstanceTransforms, bool bShouldReturnIndices, bool bWorldSpace)
 {
+	const int32 Count = InstanceTransforms.Num();
+
 	TArray<int32> NewInstanceIndices;
 
 	if (bShouldReturnIndices)
@@ -2762,11 +2765,12 @@ TArray<int32> UInstancedStaticMeshComponent::AddInstancesInternal(int32 Count, c
 	SelectedInstances.Add(false, Count);
 #endif
 
-	for (int32 i = 0; i < Count; ++i)
+	for (const FTransform& InstanceTransform : InstanceTransforms)
 	{
 		FInstancedStaticMeshInstanceData& NewInstanceData = PerInstanceSMData.AddDefaulted_GetRef();
 
-		SetupNewInstanceData(NewInstanceData, InstanceIndex, InstanceTransforms[i]);
+		const FTransform LocalTransform = bWorldSpace ? InstanceTransform.GetRelativeTransform(GetComponentTransform()) : InstanceTransform;
+		SetupNewInstanceData(NewInstanceData, InstanceIndex, LocalTransform);
 
 		if (bShouldReturnIndices)
 		{
@@ -2800,16 +2804,9 @@ TArray<int32> UInstancedStaticMeshComponent::AddInstancesInternal(int32 Count, c
 	return NewInstanceIndices;
 }
 
-TArray<int32> UInstancedStaticMeshComponent::AddInstances(const TArray<FTransform>& InstanceTransforms, bool bShouldReturnIndices)
+TArray<int32> UInstancedStaticMeshComponent::AddInstances(const TArray<FTransform>& InstanceTransforms, bool bShouldReturnIndices, bool bWorldSpace)
 {
-	return AddInstancesInternal(InstanceTransforms.Num(), InstanceTransforms, bShouldReturnIndices);
-}
-
-int32 UInstancedStaticMeshComponent::AddInstanceWorldSpace(const FTransform& WorldTransform)
- {
-	// Transform from world space to local space
-	FTransform RelativeTM = WorldTransform.GetRelativeTransform(GetComponentTransform());
-	return AddInstance(RelativeTM);
+	return AddInstancesInternal(InstanceTransforms, bShouldReturnIndices, bWorldSpace);
 }
 
 // Per Instance Custom Data - Updating custom data for specific instance
@@ -2927,6 +2924,25 @@ bool UInstancedStaticMeshComponent::RemoveInstance(int32 InstanceIndex)
 	return RemoveInstanceInternal(InstanceIndex, false);
 }
 
+bool UInstancedStaticMeshComponent::RemoveInstances(const TArray<int32>& InstancesToRemove)
+{
+	// Sort so Remove doesn't alter the indices of items still to remove
+	TArray<int32> SortedInstancesToRemove = InstancesToRemove;
+	SortedInstancesToRemove.Sort(TGreater<int32>());
+
+	if (!PerInstanceSMData.IsValidIndex(SortedInstancesToRemove[0]) || !PerInstanceSMData.IsValidIndex(SortedInstancesToRemove.Last()))
+	{
+		return false;
+	}
+
+	for (const int32 InstanceIndex : SortedInstancesToRemove)
+	{
+		RemoveInstanceInternal(InstanceIndex, false);
+	}
+
+	return true;
+}
+
 bool UInstancedStaticMeshComponent::GetInstanceTransform(int32 InstanceIndex, FTransform& OutInstanceTransform, bool bWorldSpace) const
 {
 	if (!PerInstanceSMData.IsValidIndex(InstanceIndex))
@@ -3738,7 +3754,7 @@ void UInstancedStaticMeshComponent::PostEditChangeChainProperty(FPropertyChanged
 				int32 AddedAtIndex = PropertyChangedEvent.GetArrayIndex(PropertyChangedEvent.Property->GetFName().ToString());
 				check(AddedAtIndex != INDEX_NONE);
 
-				AddInstanceInternal(AddedAtIndex, &PerInstanceSMData[AddedAtIndex], PropertyChangedEvent.ChangeType == EPropertyChangeType::ArrayAdd ? FTransform::Identity : FTransform(PerInstanceSMData[AddedAtIndex].Transform));
+				AddInstanceInternal(AddedAtIndex, &PerInstanceSMData[AddedAtIndex], PropertyChangedEvent.ChangeType == EPropertyChangeType::ArrayAdd ? FTransform::Identity : FTransform(PerInstanceSMData[AddedAtIndex].Transform), /*bWorldSpace*/false);
 
 				// added via the property editor, so we will want to interactively work with instances
 				bHasPerInstanceHitProxies = true;
