LWI fixes

 - Fixed non-PCH compilation.
 - Fixed ALightWeightInstanceStaticMeshManager::RemoveInstance failing to remove the ISM instance.

diff --git a/Engine/Source/Runtime/Engine/Classes/GameFramework/LightWeightInstanceStaticMeshManager.h b/Engine/Source/Runtime/Engine/Classes/GameFramework/LightWeightInstanceStaticMeshManager.h
--- a/Engine/Source/Runtime/Engine/Classes/GameFramework/LightWeightInstanceStaticMeshManager.h
+++ b/Engine/Source/Runtime/Engine/Classes/GameFramework/LightWeightInstanceStaticMeshManager.h
@@ -3,25 +3,13 @@
 #pragma once
 
 #include "CoreMinimal.h"
-#include "UObject/ObjectMacros.h"
-#include "UObject/UObjectBaseUtility.h"
-#include "UObject/Object.h"
-#include "Engine/EngineTypes.h"
-#include "Engine/EngineBaseTypes.h"
-#include "Kismet/BlueprintFunctionLibrary.h"
-#include "Templates/SharedPointer.h"
+#include "Engine/StaticMesh.h"
 #include "GameFramework/LightWeightInstanceManager.h"
 
-#include "Actor.h"
-
 #include "LightWeightInstanceStaticMeshManager.generated.h"
 
 DECLARE_DYNAMIC_DELEGATE_RetVal_OneParam(FActorInstanceHandle, FOnActorReady, FActorInstanceHandle, InHandle);
 
-
-
-
-
 UCLASS(BlueprintType, Blueprintable)
 class ENGINE_API ALightWeightInstanceStaticMeshManager : public ALightWeightInstanceManager
 {
diff --git a/Engine/Source/Runtime/Engine/Private/LightWeightInstanceStaticMeshManager.cpp b/Engine/Source/Runtime/Engine/Private/LightWeightInstanceStaticMeshManager.cpp
--- a/Engine/Source/Runtime/Engine/Private/LightWeightInstanceStaticMeshManager.cpp
+++ b/Engine/Source/Runtime/Engine/Private/LightWeightInstanceStaticMeshManager.cpp
@@ -2,7 +2,6 @@
 
 #include "GameFramework/LightWeightInstanceStaticMeshManager.h"
 #include "Components/HierarchicalInstancedStaticMeshComponent.h"
-#include "Engine/StaticMesh.h"
 #include "Net/UnrealNetwork.h"
 
 ALightWeightInstanceStaticMeshManager::ALightWeightInstanceStaticMeshManager(const FObjectInitializer& ObjectInitializer)
@@ -76,9 +75,8 @@ void ALightWeightInstanceStaticMeshManager::AddNewInstanceAt(FLWIData* InitData,
 
 void ALightWeightInstanceStaticMeshManager::RemoveInstance(const int32 Index)
 {
-	Super::RemoveInstance(Index);
-
 	RemoveInstanceFromRendering(Index);
+	Super::RemoveInstance(Index);
 }
 
 void ALightWeightInstanceStaticMeshManager::RemoveInstanceFromRendering(int32 DataIndex)
