Fixed uninitialized TObjectPtr in UStruct

#rnx

diff --git a/Engine/Source/Runtime/Engine/Classes/Kismet/KismetRenderingLibrary.h b/Engine/Source/Runtime/Engine/Classes/Kismet/KismetRenderingLibrary.h
--- a/Engine/Source/Runtime/Engine/Classes/Kismet/KismetRenderingLibrary.h
+++ b/Engine/Source/Runtime/Engine/Classes/Kismet/KismetRenderingLibrary.h
@@ -23,7 +23,7 @@ struct FDrawToRenderTargetContext
 	GENERATED_USTRUCT_BODY()
 
 	UPROPERTY()
-	TObjectPtr<UTextureRenderTarget2D> RenderTarget;
+	TObjectPtr<UTextureRenderTarget2D> RenderTarget = nullptr;
 
 #if WANTS_DRAW_MESH_EVENTS
 	FDrawEvent* DrawEvent = nullptr;
