Removing overly aggressive level script repair case causing a failed assertion when copying maps that I accidentally included as part of larger change in CL 16308301.

diff --git a/Engine/Source/Runtime/Engine/Private/Level.cpp b/Engine/Source/Runtime/Engine/Private/Level.cpp
--- a/Engine/Source/Runtime/Engine/Private/Level.cpp
+++ b/Engine/Source/Runtime/Engine/Private/Level.cpp
@@ -3016,18 +3016,6 @@ void ULevel::RepairLevelScript()
 		}
 	}
 
-	// Verify that the LSA class matches the desired level script blueprint class
-	if (LevelScriptActor && LevelScriptBlueprint)
-	{
-		UClass* SpawnClass = (LevelScriptBlueprint->GeneratedClass) ? LevelScriptBlueprint->GeneratedClass : LevelScriptBlueprint->SkeletonGeneratedClass;
-		UClass* LSAClass = LevelScriptActor->GetClass();
-
-		if (SpawnClass && LSAClass != SpawnClass)
-		{
-			RegenerateLevelScriptActor();
-		}
-	}
-
 	// Catch the edge case where we have a level blueprint but have never created the LevelScriptActor based on it.
 	//    This could happen if a new level is saved before the level blueprint is compiled.
 	if (!LevelScriptActor && LevelScriptBlueprint && !LevelScriptBlueprint->bIsRegeneratingOnLoad)
