Add the ability to export texture build actions for use when testing the standalone texture build workers.

diff --git a/Engine/Source/Developer/TextureFormat/Private/TextureFormatManagerModule.cpp b/Engine/Source/Developer/TextureFormat/Private/TextureFormatManagerModule.cpp
--- a/Engine/Source/Developer/TextureFormat/Private/TextureFormatManagerModule.cpp
+++ b/Engine/Source/Developer/TextureFormat/Private/TextureFormatManagerModule.cpp
@@ -31,7 +31,9 @@ public:
 	}
 
 	/** Destructor. */
-	virtual ~FTextureFormatManagerModule()
+	virtual ~FTextureFormatManagerModule() = default;
+
+	virtual void ShutdownModule()
 	{
 		FModuleManager::Get().OnModulesChanged().RemoveAll(this);
 	}
@@ -40,12 +42,12 @@ public:
 	virtual const TArray<const ITextureFormat*>& GetTextureFormats() override
 	{
 		static bool bInitialized = false;
-		static TArray<const ITextureFormat*> Results;
 
 		if (!bInitialized || bForceCacheUpdate)
 		{
 			bInitialized = true;
-			Results.Empty(Results.Num());
+			TextureFormats.Empty(TextureFormats.Num());
+			TextureFormatMetadata.Empty(TextureFormatMetadata.Num());
 
 			TArray<FName> Modules;
 
@@ -66,19 +68,30 @@ public:
 						ITextureFormat* Format = Module->GetTextureFormat();
 						if (Format != nullptr)
 						{
-							Results.Add(Format);
+							TextureFormats.Add(Format);
+							FTextureFormatMetadata& NewMeta = TextureFormatMetadata.AddDefaulted_GetRef();
+							NewMeta.ModuleName = Modules[Index];
+							NewMeta.Module = Module;
 						}
 					}
 				}
 			}
 		}
 
-		return Results;
+		return TextureFormats;
 	}
 
 	virtual const ITextureFormat* FindTextureFormat(FName Name) override
 	{
-		const TArray<const ITextureFormat*>& TextureFormats = GetTextureFormats();
+		FName ModuleNameUnused;
+		ITextureFormatModule* ModuleUnused;
+		return FindTextureFormatAndModule(Name, ModuleNameUnused, ModuleUnused);
+	}
+	
+	virtual const class ITextureFormat* FindTextureFormatAndModule(FName Name, FName& OutModuleName, ITextureFormatModule*& OutModule) override
+	{
+		// Called to ensure the arrays are populated
+		GetTextureFormats();
 
 		for (int32 Index = 0; Index < TextureFormats.Num(); Index++)
 		{
@@ -90,6 +103,9 @@ public:
 			{
 				if (Formats[FormatIndex] == Name)
 				{
+					const FTextureFormatMetadata& FoundMeta = TextureFormatMetadata[Index];
+					OutModuleName = FoundMeta.ModuleName;
+					OutModule = FoundMeta.Module;
 					return TextureFormats[Index];
 				}
 			}
@@ -118,6 +134,15 @@ private:
 
 	FName ModuleName;
 
+	TArray<const ITextureFormat*> TextureFormats;
+
+	struct FTextureFormatMetadata
+	{
+		FName ModuleName;
+		ITextureFormatModule* Module;
+	};
+	TArray<FTextureFormatMetadata> TextureFormatMetadata;
+
 	// Flag to force reinitialization of all cached data. This is needed to have up-to-date caches
 	// in case of a module reload of a TextureFormat-Module.
 	bool bForceCacheUpdate;
diff --git a/Engine/Source/Developer/TextureFormat/Public/Interfaces/ITextureFormatManagerModule.h b/Engine/Source/Developer/TextureFormat/Public/Interfaces/ITextureFormatManagerModule.h
--- a/Engine/Source/Developer/TextureFormat/Public/Interfaces/ITextureFormatManagerModule.h
+++ b/Engine/Source/Developer/TextureFormat/Public/Interfaces/ITextureFormatManagerModule.h
@@ -21,6 +21,16 @@ public:
 	 */
 	virtual const class ITextureFormat* FindTextureFormat( FName Name ) = 0;
 
+	/**
+	 * Finds a texture format with the specified name and provides information about the module it came from.
+	 *
+	 * @param Name Name of the format to find.
+	 * @param OutModuleName Name of the module that the found format came from, or unmodified if not found.
+	 * @param OutModule Interface of the module that the found format came from, or unmodified if not found.
+	 * @return The texture format, or nullptr if not found.
+	 */
+	virtual const class ITextureFormat* FindTextureFormatAndModule( FName Name, FName& OutModuleName, class ITextureFormatModule*& OutModule ) = 0;
+
 	/**
 	 * Returns the list of all ITextureFormats that were located in DLLs.
 	 *
diff --git a/Engine/Source/Runtime/Engine/Classes/Engine/Texture.h b/Engine/Source/Runtime/Engine/Classes/Engine/Texture.h
--- a/Engine/Source/Runtime/Engine/Classes/Engine/Texture.h
+++ b/Engine/Source/Runtime/Engine/Classes/Engine/Texture.h
@@ -31,27 +31,6 @@ class ITargetPlatform;
 class UAssetUserData;
 struct FPropertyChangedEvent;
 
-// This needs to be mirrored in EditorFactories.cpp.
-UENUM()
-enum TextureCompressionSettings
-{
-	TC_Default					UMETA(DisplayName="Default (DXT1/5, BC1/3 on DX11)"),
-	TC_Normalmap				UMETA(DisplayName="Normalmap (DXT5, BC5 on DX11)"),
-	TC_Masks					UMETA(DisplayName="Masks (no sRGB)"),
-	TC_Grayscale				UMETA(DisplayName="Grayscale (R8, RGB8 sRGB)"),
-	TC_Displacementmap			UMETA(DisplayName="Displacementmap (8/16bit)"),
-	TC_VectorDisplacementmap	UMETA(DisplayName="VectorDisplacementmap (RGBA8)"),
-	TC_HDR						UMETA(DisplayName="HDR (RGB, no sRGB)"),
-	TC_EditorIcon				UMETA(DisplayName="UserInterface2D (RGBA)"),
-	TC_Alpha					UMETA(DisplayName="Alpha (no sRGB, BC4 on DX11)"),
-	TC_DistanceFieldFont		UMETA(DisplayName="DistanceFieldFont (R8)"),
-	TC_HDR_Compressed			UMETA(DisplayName="HDRCompressed (RGB, BC6H, DX11)"),
-	TC_BC7						UMETA(DisplayName="BC7 (DX11, optional A)"),
-	TC_HalfFloat				UMETA(DisplayName="Half Float (R16F)"),
-	TC_ReflectionCapture		UMETA(Hidden),
-	TC_MAX,
-};
-
 UENUM()
 enum TextureFilter
 {
@@ -93,36 +72,6 @@ enum ETextureSourceArtType
 	TSAT_MAX,
 };
 
-UENUM()
-enum ETextureSourceFormat
-{
-	TSF_Invalid,
-	TSF_G8,
-	TSF_BGRA8,
-	TSF_BGRE8,
-	TSF_RGBA16,
-	TSF_RGBA16F,
-
-	//@todo: Deprecated!
-	TSF_RGBA8,
-	//@todo: Deprecated!
-	TSF_RGBE8,
-
-	TSF_G16,
-
-	TSF_MAX
-};
-
-UENUM()
-enum ETextureSourceCompressionFormat
-{
-	TSCF_None	UMETA(DisplayName = "ZLib"),
-	TSCF_PNG	UMETA(DisplayName = "PNG"),
-	TSCF_JPEG	UMETA(DisplayName = "JPEG"),
-
-	TSCF_MAX
-};
-
 UENUM()
 enum ETextureCompressionQuality
 {
@@ -310,6 +259,9 @@ struct FTextureSource
 	/** Returns the compression format of the source data in string format for use with the UI. */
 	FString GetSourceCompressionAsString() const;
 
+	/** Returns the compression format of the source data in enum format. */
+	ETextureSourceCompressionFormat GetSourceCompression() const;
+
 	/** Trivial accessors. These will only give values for Block0 so may not be correct for UDIM/multi-block textures, use GetBlock() for this case. */
 	FORCEINLINE FGuid GetId() const { return Id; }
 	FORCEINLINE int32 GetSizeX() const { return SizeX; }
@@ -338,6 +290,8 @@ struct FTextureSource
 #endif //UE_USE_VIRTUALBULKDATA
 	}
 
+	ENGINE_API void OperateOnLoadedBulkData(TFunctionRef<void (const FSharedBuffer& BulkDataBuffer)> Operation);
+
 	FORCEINLINE bool LoadBulkDataWithFileReader()
 	{
 		// TODO: When we remove UE_USE_VIRTUALBULKDATA and go all in, places that call this
@@ -675,6 +629,10 @@ public:
 	/** Serialization. */
 	void Serialize(FArchive& Ar, class UTexture* Owner);
 
+#if WITH_EDITORONLY_DATA
+	void SerializeWithConditionalBulkData(FArchive& Ar, class UTexture* Owner);
+#endif // WITH_EDITORONLY_DATA
+
 	/** 
 	 * Serialization for cooked builds.
 	 *
diff --git a/Engine/Source/Runtime/Engine/Classes/Engine/TextureDefines.h b/Engine/Source/Runtime/Engine/Classes/Engine/TextureDefines.h
--- a/Engine/Source/Runtime/Engine/Classes/Engine/TextureDefines.h
+++ b/Engine/Source/Runtime/Engine/Classes/Engine/TextureDefines.h
@@ -224,7 +224,6 @@ enum ETextureLossyCompressionAmount
 	TLCA_Highest		UMETA(DisplayName = "Highest (Worst image quality, smallest filesize)"),
 };
 
-
 UENUM()
 enum ECompositeTextureMode
 {
@@ -243,3 +242,54 @@ enum ECompositeTextureMode
 	// should be added at the bottom; reordering or removing entries will require changing the GUID
 	// in the texture compressor DDC key
 };
+
+UENUM()
+enum ETextureSourceCompressionFormat
+{
+	TSCF_None	UMETA(DisplayName = "ZLib"),
+	TSCF_PNG	UMETA(DisplayName = "PNG"),
+	TSCF_JPEG	UMETA(DisplayName = "JPEG"),
+
+	TSCF_MAX
+};
+
+UENUM()
+enum ETextureSourceFormat
+{
+	TSF_Invalid,
+	TSF_G8,
+	TSF_BGRA8,
+	TSF_BGRE8,
+	TSF_RGBA16,
+	TSF_RGBA16F,
+
+	//@todo: Deprecated!
+	TSF_RGBA8,
+	//@todo: Deprecated!
+	TSF_RGBE8,
+
+	TSF_G16,
+
+	TSF_MAX
+};
+
+// This needs to be mirrored in EditorFactories.cpp.
+UENUM()
+enum TextureCompressionSettings
+{
+	TC_Default					UMETA(DisplayName = "Default (DXT1/5, BC1/3 on DX11)"),
+	TC_Normalmap				UMETA(DisplayName = "Normalmap (DXT5, BC5 on DX11)"),
+	TC_Masks					UMETA(DisplayName = "Masks (no sRGB)"),
+	TC_Grayscale				UMETA(DisplayName = "Grayscale (R8, RGB8 sRGB)"),
+	TC_Displacementmap			UMETA(DisplayName = "Displacementmap (8/16bit)"),
+	TC_VectorDisplacementmap	UMETA(DisplayName = "VectorDisplacementmap (RGBA8)"),
+	TC_HDR						UMETA(DisplayName = "HDR (RGB, no sRGB)"),
+	TC_EditorIcon				UMETA(DisplayName = "UserInterface2D (RGBA)"),
+	TC_Alpha					UMETA(DisplayName = "Alpha (no sRGB, BC4 on DX11)"),
+	TC_DistanceFieldFont		UMETA(DisplayName = "DistanceFieldFont (R8)"),
+	TC_HDR_Compressed			UMETA(DisplayName = "HDRCompressed (RGB, BC6H, DX11)"),
+	TC_BC7						UMETA(DisplayName = "BC7 (DX11, optional A)"),
+	TC_HalfFloat				UMETA(DisplayName = "Half Float (R16F)"),
+	TC_ReflectionCapture		UMETA(Hidden),
+	TC_MAX,
+};
diff --git a/Engine/Source/Runtime/Engine/Engine.Build.cs b/Engine/Source/Runtime/Engine/Engine.Build.cs
--- a/Engine/Source/Runtime/Engine/Engine.Build.cs
+++ b/Engine/Source/Runtime/Engine/Engine.Build.cs
@@ -125,6 +125,7 @@ public class Engine : ModuleRules
 				"CrunchCompression",
 				"IntelISPC",
 				"TraceLog",
+				"TextureFormat",
 			}
 		);
 
diff --git a/Engine/Source/Runtime/Engine/Private/Texture.cpp b/Engine/Source/Runtime/Engine/Private/Texture.cpp
--- a/Engine/Source/Runtime/Engine/Private/Texture.cpp
+++ b/Engine/Source/Runtime/Engine/Private/Texture.cpp
@@ -1594,16 +1594,21 @@ FString FTextureSource::GetIdString() const
 	return GuidString;
 }
 
-FString FTextureSource::GetSourceCompressionAsString() const
+ETextureSourceCompressionFormat FTextureSource::GetSourceCompression() const
 {
 	// Until we deprecate bPNGCompressed it might not be 100% in sync with CompressionFormat
 	// so if it is set we should use that rather than the enum.
 	if (bPNGCompressed)
 	{
-		return StaticEnum<ETextureSourceCompressionFormat>()->GetDisplayNameTextByValue(ETextureSourceCompressionFormat::TSCF_PNG).ToString();
+		return ETextureSourceCompressionFormat::TSCF_PNG;
 	}
 
-	return StaticEnum<ETextureSourceCompressionFormat>()->GetDisplayNameTextByValue(CompressionFormat).ToString();
+	return CompressionFormat;
+}
+
+FString FTextureSource::GetSourceCompressionAsString() const
+{
+	return StaticEnum<ETextureSourceCompressionFormat>()->GetDisplayNameTextByValue(GetSourceCompression()).ToString();
 }
 
 FTextureSource::FMipAllocation FTextureSource::TryDecompressPngData(IImageWrapperModule* ImageWrapperModule) UE_VBD_CONST
@@ -1850,6 +1855,27 @@ void FTextureSource::UseHashAsGuid()
 	}
 }
 
+void FTextureSource::OperateOnLoadedBulkData(TFunctionRef<void(const FSharedBuffer& BulkDataBuffer)> Operation)
+{
+	checkf(LockState == ELockState::None, TEXT("OperateOnLoadedBulkData shouldn't be called in-between LockMip/UnlockMip"));
+
+#if WITH_EDITOR
+	FReadScopeLock BulkDataExclusiveScope(BulkDataLock.Get());
+#endif
+
+#if UE_USE_VIRTUALBULKDATA
+	FSharedBuffer Payload = BulkData.GetPayload().Get();
+#else
+	uint8* BulkDataPtr = (uint8*)BulkData.Lock(LOCK_READ_ONLY);
+	const int64 BulkDataSize = BulkData.GetBulkDataSize();
+	FSharedBuffer Payload = FSharedBuffer::MakeView(BulkDataPtr, BulkDataSize);
+#endif //UE_USE_VIRTUALBULKDATA
+	Operation(Payload);
+#if !UE_USE_VIRTUALBULKDATA
+	BulkData.Unlock();
+#endif //!UE_USE_VIRTUALBULKDATA
+}
+
 void FTextureSource::SetId(const FGuid& InId, bool bInGuidIsHash)
 {
 	Id = InId;
diff --git a/Engine/Source/Runtime/Engine/Private/Texture2D.cpp b/Engine/Source/Runtime/Engine/Private/Texture2D.cpp
--- a/Engine/Source/Runtime/Engine/Private/Texture2D.cpp
+++ b/Engine/Source/Runtime/Engine/Private/Texture2D.cpp
@@ -133,13 +133,21 @@ extern bool TrackTextureEvent( FStreamingRenderAsset* StreamingTexture, UStreama
 	FTexture2DMipMap
 -----------------------------------------------------------------------------*/
 
-void FTexture2DMipMap::Serialize(FArchive& Ar, UObject* Owner, int32 MipIdx)
+void FTexture2DMipMap::SerializeCommon(FArchive& Ar, UObject* Owner, int32 MipIdx, bool bIgnoreBulkDataForStreamingMips)
 {
 	bool bCooked = Ar.IsCooking();
 	Ar << bCooked;
 
 #if WITH_EDITORONLY_DATA
-	BulkData.Serialize(Ar, Owner, MipIdx, false, FileRegionType);
+	if (bIgnoreBulkDataForStreamingMips && !DerivedDataKey.IsEmpty())
+	{
+		FByteBulkData EmptyData;
+		EmptyData.Serialize(Ar, Owner, MipIdx, false, FileRegionType);
+	}
+	else
+	{
+		BulkData.Serialize(Ar, Owner, MipIdx, false, FileRegionType);
+	}
 #else
 	BulkData.Serialize(Ar, Owner, MipIdx, false);
 #endif
@@ -157,7 +165,17 @@ void FTexture2DMipMap::Serialize(FArchive& Ar, UObject* Owner, int32 MipIdx)
 #endif // #if WITH_EDITORONLY_DATA
 }
 
+void FTexture2DMipMap::Serialize(FArchive& Ar, UObject* Owner, int32 MipIdx)
+{
+	SerializeCommon(Ar, Owner, MipIdx, false);
+}
+
 #if WITH_EDITORONLY_DATA
+void FTexture2DMipMap::SerializeWithConditionalBulkData(FArchive& Ar, UObject* Owner, int32 MipIdx)
+{
+	SerializeCommon(Ar, Owner, MipIdx, true);
+}
+
 uint32 FTexture2DMipMap::StoreInDerivedDataCache(const FString& InDerivedDataKey, const FStringView& TextureName, bool bReplaceExistingDDC)
 {
 	int32 BulkDataSizeInBytes = BulkData.GetBulkDataSize();
diff --git a/Engine/Source/Runtime/Engine/Private/TextureDerivedData.cpp b/Engine/Source/Runtime/Engine/Private/TextureDerivedData.cpp
--- a/Engine/Source/Runtime/Engine/Private/TextureDerivedData.cpp
+++ b/Engine/Source/Runtime/Engine/Private/TextureDerivedData.cpp
@@ -273,7 +273,7 @@ static void SerializeForKey(FArchive& Ar, const FTextureBuildSettings& Settings)
  * @param KeySuffix - The key suffix.
  * @param OutKey - The full derived data key.
  */
-static void GetTextureDerivedDataKeyFromSuffix(const FString& KeySuffix, FString& OutKey)
+void GetTextureDerivedDataKeyFromSuffix(const FString& KeySuffix, FString& OutKey)
 {
 	OutKey = FDerivedDataCacheInterface::BuildCacheKey(
 		TEXT("TEXTURE"),
@@ -288,7 +288,7 @@ static void GetTextureDerivedDataKeyFromSuffix(const FString& KeySuffix, FString
  * @param MipIndex - The mip index.
  * @param OutKey - The full derived data key for the mip.
  */
-static void GetTextureDerivedMipKey(
+void GetTextureDerivedMipKey(
 	int32 MipIndex,
 	const FTexture2DMipMap& Mip,
 	const FString& KeySuffix,
@@ -1452,14 +1452,29 @@ bool FTexturePlatformData::AreDerivedVTChunksAvailable() const
 }
 #endif // #if WITH_EDITOR
 
+// Transient flags used to control behavior of platform data serialization
+enum class EPlatformDataSerializationFlags : uint8
+{
+	None = 0,
+	Cooked = 1<<0,
+	Streamable = 1<<1,
+
+#if WITH_EDITORONLY_DATA
+	IgnoreBulkDataForStreamingMips = 1<<2, // These are transient, so this can be changed if more runtime flags are needed above
+#endif // WITH_EDITORONLY_DATA
+};
+ENUM_CLASS_FLAGS(EPlatformDataSerializationFlags);
+
 static void SerializePlatformData(
 	FArchive& Ar,
 	FTexturePlatformData* PlatformData,
 	UTexture* Texture,
-	bool bCooked,
-	bool bStreamable
+	EPlatformDataSerializationFlags Flags
 )
 {
+	const bool bCooked = (Flags & EPlatformDataSerializationFlags::Cooked) == EPlatformDataSerializationFlags::Cooked;
+	const bool bStreamable = (Flags & EPlatformDataSerializationFlags::Streamable) == EPlatformDataSerializationFlags::Streamable;
+
 	DECLARE_SCOPE_CYCLE_COUNTER(TEXT("SerializePlatformData"), STAT_Texture_SerializePlatformData, STATGROUP_LoadTime);
 
 	UEnum* PixelFormatEnum = UTexture::GetPixelFormatEnum();
@@ -1613,9 +1628,22 @@ static void SerializePlatformData(
 			PlatformData->Mips.Add(new FTexture2DMipMap());
 		}
 	}
-	for (int32 MipIndex = 0; MipIndex < NumMips; ++MipIndex)
+
+#if WITH_EDITORONLY_DATA
+	if ((Flags & EPlatformDataSerializationFlags::IgnoreBulkDataForStreamingMips) == EPlatformDataSerializationFlags::IgnoreBulkDataForStreamingMips)
+	{
+		for (int32 MipIndex = 0; MipIndex < NumMips; ++MipIndex)
+		{
+			PlatformData->Mips[FirstMipToSerialize + MipIndex].SerializeWithConditionalBulkData(Ar, Texture, MipIndex);
+		}
+	}
+	else
+#endif // WITH_EDITORONLY_DATA
 	{
-		PlatformData->Mips[FirstMipToSerialize + MipIndex].Serialize(Ar, Texture, MipIndex);
+		for (int32 MipIndex = 0; MipIndex < NumMips; ++MipIndex)
+		{
+			PlatformData->Mips[FirstMipToSerialize + MipIndex].Serialize(Ar, Texture, MipIndex);
+		}
 	}
 
 	Ar << bIsVirtual;
@@ -1653,14 +1681,24 @@ static void SerializePlatformData(
 
 void FTexturePlatformData::Serialize(FArchive& Ar, UTexture* Owner)
 {
-	const bool bCooking = false;
-	const bool bStreamable = false;
-	SerializePlatformData(Ar, this, Owner, bCooking, bStreamable);
+	SerializePlatformData(Ar, this, Owner, EPlatformDataSerializationFlags::None);
 }
 
+#if WITH_EDITORONLY_DATA
+void FTexturePlatformData::SerializeWithConditionalBulkData(FArchive& Ar, UTexture* Owner)
+{
+	SerializePlatformData(Ar, this, Owner, EPlatformDataSerializationFlags::IgnoreBulkDataForStreamingMips);
+}
+#endif // WITH_EDITORONLY_DATA
+
 void FTexturePlatformData::SerializeCooked(FArchive& Ar, UTexture* Owner, bool bStreamable)
 {
-	SerializePlatformData(Ar, this, Owner, true, bStreamable);
+	EPlatformDataSerializationFlags Flags = EPlatformDataSerializationFlags::Cooked;
+	if (bStreamable)
+	{
+		Flags |= EPlatformDataSerializationFlags::Streamable;
+	}
+	SerializePlatformData(Ar, this, Owner, Flags);
 	if (Ar.IsLoading())
 	{
 		// Patch up Size as due to mips being stripped out during cooking it could be wrong.
diff --git a/Engine/Source/Runtime/Engine/Private/TextureDerivedDataBuildExporter.cpp b/Engine/Source/Runtime/Engine/Private/TextureDerivedDataBuildExporter.cpp
new file mode 100644
--- /dev/null
+++ b/Engine/Source/Runtime/Engine/Private/TextureDerivedDataBuildExporter.cpp
@@ -0,0 +1,456 @@
+// Copyright Epic Games, Inc. All Rights Reserved.
+
+#include "TextureDerivedDataBuildExporter.h"
+
+#if WITH_EDITOR
+
+#include "DerivedDataPayload.h"
+#include "HAL/FileManager.h"
+#include "HAL/IConsoleManager.h"
+#include "HAL/PlatformFileManager.h"
+#include "Interfaces/ITextureFormat.h"
+#include "Interfaces/ITextureFormatManagerModule.h"
+#include "Misc/FileHelper.h"
+#include "Misc/StringBuilder.h"
+#include "Serialization/CompactBinaryPackage.h"
+#include "Serialization/CompactBinaryWriter.h"
+#include "TextureDerivedDataTask.h"
+#include "TextureFormatManager.h"
+
+
+void GetTextureDerivedDataKeyFromSuffix(const FString& KeySuffix, FString& OutKey);
+void GetTextureDerivedMipKey(int32 MipIndex, const FTexture2DMipMap& Mip, const FString& KeySuffix, FString& OutKey);
+
+static bool ShortenKey(const TCHAR* CacheKey, FString& Result)
+{
+	const int32 MaxKeyLength = 120;
+	Result = FString(CacheKey);
+	if (Result.Len() <= MaxKeyLength)
+	{
+		return false;
+	}
+
+	FSHA1 HashState;
+	int32 Length = Result.Len();
+	HashState.Update((const uint8*)&Length, sizeof(int32));
+
+	auto ResultSrc = StringCast<UCS2CHAR>(*Result);
+
+	// This is pretty redundant. Incorporating the CRC of the name into the hash
+	// which also ends up computing SHA1 of the name is not really going to make 
+	// any meaningful difference to the strength of the key so it's just a waste
+	// of CPU. However it's difficult to get rid of without invalidating the DDC
+	// contents so here we are.
+	const uint32 CRCofPayload(FCrc::MemCrc32(ResultSrc.Get(), Length * sizeof(UCS2CHAR)));
+	HashState.Update((const uint8*)&CRCofPayload, sizeof(uint32));
+
+	HashState.Update((const uint8*)ResultSrc.Get(), Length * sizeof(UCS2CHAR));
+
+	HashState.Final();
+	uint8 Hash[FSHA1::DigestSize];
+	HashState.GetHash(Hash);
+	const FString HashString = BytesToHex(Hash, FSHA1::DigestSize);
+
+	int32 HashStringSize = HashString.Len();
+	int32 OriginalPart = MaxKeyLength - HashStringSize - 2;
+	Result = Result.Left(OriginalPart) + TEXT("__") + HashString;
+	check(Result.Len() == MaxKeyLength && Result.Len() > 0);
+	return true;
+}
+
+static void WriteCbField(FCbWriter& Writer, const FAnsiStringView& Name, const FColor& Color)
+{
+	// Choosing the big endian ordering
+	Writer.BeginArray(Name);
+	Writer.AddInteger(Color.A);
+	Writer.AddInteger(Color.R);
+	Writer.AddInteger(Color.G);
+	Writer.AddInteger(Color.B);
+	Writer.EndArray();
+}
+
+static void WriteCbField(FCbWriter& Writer, const FAnsiStringView& Name, const FVector4& Vec4)
+{
+	Writer.BeginArray(Name);
+	Writer.AddFloat(Vec4.X);
+	Writer.AddFloat(Vec4.Y);
+	Writer.AddFloat(Vec4.Z);
+	Writer.AddFloat(Vec4.W);
+	Writer.EndArray();
+}
+
+static void WriteCbField(FCbWriter& Writer, const FAnsiStringView& Name, const FIntPoint& IntPoint)
+{
+	Writer.BeginArray(Name);
+	Writer.AddInteger(IntPoint.X);
+	Writer.AddInteger(IntPoint.Y);
+	Writer.EndArray();
+}
+
+static void WriteBuildSettingsToCompactBinary(FCbWriter& Writer, const FAnsiStringView& Name, const FTextureBuildSettings& BuildSettings)
+{
+	Writer.BeginObject(Name);
+
+	Writer.BeginObject("ColorAdjustment");
+	Writer.AddFloat("AdjustBrightness", BuildSettings.ColorAdjustment.AdjustBrightness);
+	Writer.AddFloat("AdjustBrightnessCurve", BuildSettings.ColorAdjustment.AdjustBrightnessCurve);
+	Writer.AddFloat("AdjustSaturation", BuildSettings.ColorAdjustment.AdjustSaturation);
+	Writer.AddFloat("AdjustVibrance", BuildSettings.ColorAdjustment.AdjustVibrance);
+	Writer.AddFloat("AdjustRGBCurve", BuildSettings.ColorAdjustment.AdjustRGBCurve);
+	Writer.AddFloat("AdjustHue", BuildSettings.ColorAdjustment.AdjustHue);
+	Writer.AddFloat("AdjustMinAlpha", BuildSettings.ColorAdjustment.AdjustMinAlpha);
+	Writer.AddFloat("AdjustMaxAlpha", BuildSettings.ColorAdjustment.AdjustMaxAlpha);
+	Writer.EndObject();
+
+	WriteCbField(Writer, "AlphaCoverageThresholds", BuildSettings.AlphaCoverageThresholds);
+
+	Writer.AddFloat("MipSharpening", BuildSettings.MipSharpening);
+	Writer.AddInteger("DiffuseConvolveMipLevel", BuildSettings.DiffuseConvolveMipLevel);
+	Writer.AddInteger("SharpenMipKernelSize", BuildSettings.SharpenMipKernelSize);
+	Writer.AddInteger("MaxTextureResolution", BuildSettings.MaxTextureResolution);
+	Writer.AddString("TextureFormatName", WriteToString<64>(BuildSettings.TextureFormatName));
+	Writer.AddBool("bHDRSource", BuildSettings.bHDRSource);
+	Writer.AddInteger("MipGenSettings", BuildSettings.MipGenSettings);
+	Writer.AddBool("bCubemap", BuildSettings.bCubemap);
+	Writer.AddBool("bTextureArray", BuildSettings.bTextureArray);
+	Writer.AddBool("bVolume", BuildSettings.bVolume);
+	Writer.AddBool("bLongLatSource", BuildSettings.bLongLatSource);
+	Writer.AddBool("bSRGB", BuildSettings.bSRGB);
+	Writer.AddBool("bUseLegacyGamma", BuildSettings.bUseLegacyGamma);
+	Writer.AddBool("bPreserveBorder", BuildSettings.bPreserveBorder);
+	Writer.AddBool("bForceAlphaChannel", BuildSettings.bForceAlphaChannel);
+	Writer.AddBool("bDitherMipMapAlpha", BuildSettings.bDitherMipMapAlpha);
+	Writer.AddBool("bComputeBokehAlpha", BuildSettings.bComputeBokehAlpha);
+	Writer.AddBool("bReplicateRed", BuildSettings.bReplicateRed);
+	Writer.AddBool("bReplicateAlpha", BuildSettings.bReplicateAlpha);
+	Writer.AddBool("bDownsampleWithAverage", BuildSettings.bDownsampleWithAverage);
+	Writer.AddBool("bSharpenWithoutColorShift", BuildSettings.bSharpenWithoutColorShift);
+	Writer.AddBool("bBorderColorBlack", BuildSettings.bBorderColorBlack);
+	Writer.AddBool("bFlipGreenChannel", BuildSettings.bFlipGreenChannel);
+	Writer.AddBool("bApplyYCoCgBlockScale", BuildSettings.bApplyYCoCgBlockScale);
+	Writer.AddBool("bApplyKernelToTopMip", BuildSettings.bApplyKernelToTopMip);
+	Writer.AddBool("bRenormalizeTopMip", BuildSettings.bRenormalizeTopMip);
+	Writer.AddInteger("CompositeTextureMode", BuildSettings.CompositeTextureMode);
+	Writer.AddFloat("CompositePower", BuildSettings.CompositePower);
+	Writer.AddInteger("LODBias", BuildSettings.LODBias);
+	Writer.AddInteger("LODBiasWithCinematicMips", BuildSettings.LODBiasWithCinematicMips);
+
+	WriteCbField(Writer, "TopMipSize", BuildSettings.TopMipSize);
+
+	Writer.AddInteger("VolumeSizeZ", BuildSettings.VolumeSizeZ);
+	Writer.AddInteger("ArraySlices", BuildSettings.ArraySlices);
+	Writer.AddBool("bStreamable", BuildSettings.bStreamable);
+	Writer.AddBool("bVirtualStreamable", BuildSettings.bVirtualStreamable);
+	Writer.AddBool("bChromaKeyTexture", BuildSettings.bChromaKeyTexture);
+	Writer.AddInteger("PowerOfTwoMode", BuildSettings.PowerOfTwoMode);
+	WriteCbField(Writer, "PaddingColor", BuildSettings.PaddingColor);
+	WriteCbField(Writer, "ChromaKeyColor", BuildSettings.ChromaKeyColor);
+	Writer.AddFloat("ChromaKeyThreshold", BuildSettings.ChromaKeyThreshold);
+	Writer.AddInteger("CompressionQuality", BuildSettings.CompressionQuality);
+	Writer.AddInteger("LossyCompressionAmount", BuildSettings.LossyCompressionAmount);
+	Writer.AddFloat("Downscale", BuildSettings.Downscale);
+	Writer.AddInteger("DownscaleOptions", BuildSettings.DownscaleOptions);
+	Writer.AddInteger("VirtualAddressingModeX", BuildSettings.VirtualAddressingModeX);
+	Writer.AddInteger("VirtualAddressingModeY", BuildSettings.VirtualAddressingModeY);
+	Writer.AddInteger("VirtualTextureTileSize", BuildSettings.VirtualTextureTileSize);
+	Writer.AddInteger("VirtualTextureBorderSize", BuildSettings.VirtualTextureBorderSize);
+	Writer.AddBool("bVirtualTextureEnableCompressZlib", BuildSettings.bVirtualTextureEnableCompressZlib);
+	Writer.AddBool("bVirtualTextureEnableCompressCrunch", BuildSettings.bVirtualTextureEnableCompressCrunch);
+	Writer.AddBool("bHasEditorOnlyData", BuildSettings.bHasEditorOnlyData);
+
+	Writer.EndObject();
+}
+
+static void WriteOutputSettingsToCompactBinary(FCbWriter& Writer, const FAnsiStringView& Name, int32 NumInlineMips, const FString& KeySuffix)
+{
+	Writer.BeginObject(Name);
+	Writer.AddInteger("NumInlineMips", NumInlineMips);
+	
+	FString MipDerivedDataKey;
+	FTexture2DMipMap DummyMip;
+	DummyMip.SizeX = 0;
+	DummyMip.SizeY = 0;
+	GetTextureDerivedMipKey(0, DummyMip, KeySuffix, MipDerivedDataKey);
+	int32 PrefixEndIndex = MipDerivedDataKey.Find(TEXT("_MIP0_"), ESearchCase::CaseSensitive);
+	check(PrefixEndIndex != -1);
+	MipDerivedDataKey.LeftInline(PrefixEndIndex);
+	check(!MipDerivedDataKey.IsEmpty());
+	Writer.AddString("MipKeyPrefix",*MipDerivedDataKey);
+
+	Writer.EndObject();
+}
+
+static void WriteTextureSourceToCompactBinary(FCbWriter& Writer, const FAnsiStringView& Name, const FTextureSource& TextureSource, EGammaSpace GammaSpace)
+{
+	Writer.BeginObject(Name);
+
+	Writer.AddString("Input", TextureSource.GetIdString());
+	Writer.AddInteger("CompressionFormat", TextureSource.GetSourceCompression());
+	Writer.AddInteger("SourceFormat", TextureSource.GetFormat());
+	Writer.AddInteger("GammaSpace", (int32)GammaSpace);
+	Writer.AddInteger("NumSlices", TextureSource.GetNumSlices());
+	Writer.AddInteger("SizeX", TextureSource.GetSizeX());
+	Writer.AddInteger("SizeY", TextureSource.GetSizeY());
+	Writer.BeginArray("Mips");
+	const int32 NumMips = TextureSource.GetNumMips();
+	int64 Offset = 0;
+	for (int32 MipIndex = 0; MipIndex < NumMips; ++MipIndex)
+	{
+		Writer.BeginObject();
+		Writer.AddInteger("Offset", Offset);
+		const int64 MipSize = TextureSource.CalcMipSize(MipIndex);
+		Writer.AddInteger("Size", MipSize);
+		Offset += MipSize;
+		Writer.EndObject();
+	}
+	Writer.EndArray();
+
+	Writer.EndObject();
+}
+
+static FIoHash HashAndWriteToCompressedBufferFile(const FString& Directory, const void* InData, uint64 InDataSize)
+{
+	FIoHash DataHash = FIoHash::HashBuffer(InData, InDataSize);
+	TStringBuilder<128> DataHashStringBuilder;
+	DataHashStringBuilder << DataHash;
+
+	FCompressedBuffer CompressedBufferContents = FCompressedBuffer::Compress(NAME_None, FSharedBuffer::MakeView(InData, InDataSize));
+	if (TUniquePtr<FArchive> FileAr{IFileManager::Get().CreateFileWriter(*(Directory / *DataHashStringBuilder), FILEWRITE_NoReplaceExisting)})
+	{
+		*FileAr << CompressedBufferContents;
+	}
+	return DataHash;
+}
+
+static FIoHash ExportTextureBulkDataAttachment(const FString& ExportRoot, FTextureSource& TextureSource)
+{
+	const FString BuildInputPath = ExportRoot / TEXT("Inputs");
+
+	IPlatformFile& PlatformFile = FPlatformFileManager::Get().GetPlatformFile();
+	bool bExporting = PlatformFile.CreateDirectoryTree(*BuildInputPath);
+
+	if (!bExporting)
+	{
+		return FIoHash();
+	}
+
+	FIoHash BulkDataHash;
+	TextureSource.OperateOnLoadedBulkData([&BuildInputPath, &BulkDataHash] (const FSharedBuffer& BulkDataBuffer) 
+	{
+		BulkDataHash = HashAndWriteToCompressedBufferFile(BuildInputPath, BulkDataBuffer.GetData(), BulkDataBuffer.GetSize());
+	});
+
+	return BulkDataHash;
+}
+
+void FTextureDerivedDataBuildExporter::Init(const FString& InKeySuffix)
+{
+	static const bool bExportsEnabled = FParse::Param(FCommandLine::Get(), TEXT("ExportTextureBuilds"));
+	bEnabled = bExportsEnabled;
+	if (!bEnabled)
+	{
+		return;
+	}
+
+	KeySuffix = InKeySuffix;
+
+	FString DerivedDataKeyLong;
+	FString DerivedDataKey;
+	GetTextureDerivedDataKeyFromSuffix(KeySuffix, DerivedDataKeyLong);
+	ShortenKey(*DerivedDataKeyLong, DerivedDataKey);
+	ExportRoot = FPaths::ProjectSavedDir() / TEXT("TextureBuildActions") / DerivedDataKey;
+}
+
+void FTextureDerivedDataBuildExporter::ExportTextureSourceBulkData(FTextureSource& TextureSource)
+{
+	if (bEnabled)
+	{
+		ExportedTextureBulkDataHash = ExportTextureBulkDataAttachment(ExportRoot, TextureSource);
+	}
+}
+
+void FTextureDerivedDataBuildExporter::ExportCompositeTextureSourceBulkData(FTextureSource& TextureSource)
+{
+	if (bEnabled)
+	{
+		ExportedCompositeTextureBulkDataHash = ExportTextureBulkDataAttachment(ExportRoot, TextureSource);
+	}
+}
+
+void FTextureDerivedDataBuildExporter::ExportTextureBuild(const UTexture& Texture, const FTextureBuildSettings& BuildSettings, int32 LayerIndex, int32 NumInlineMips)
+{
+	if (!bEnabled)
+	{
+		return;
+	}
+	
+	IPlatformFile& PlatformFile = FPlatformFileManager::Get().GetPlatformFile();
+	bool bExporting = PlatformFile.CreateDirectoryTree(*ExportRoot);
+	
+	const ITextureFormat* TextureFormat = nullptr;
+	FName TextureFormatModuleName;
+	ITextureFormatModule* TextureFormatModule = nullptr;
+
+	ITextureFormatManagerModule* TFM = GetTextureFormatManager();
+	if (TFM)
+	{
+		TextureFormat = TFM->FindTextureFormatAndModule(BuildSettings.TextureFormatName, TextureFormatModuleName, TextureFormatModule);
+	}
+
+	if (TextureFormat == nullptr)
+	{
+		bExporting = false;
+	}
+
+	if (!bExporting)
+	{
+		return;
+	}
+
+	FCbWriter BuildWriter;
+	BuildWriter.BeginObject("BuildAction");
+
+	BuildWriter.BeginObject("Function");
+
+	// Texture format modules are inconsistent in their naming.  eg: TextureFormatUncompressed, PS5TextureFormat
+	// We attempt to  unify the naming here when specifying build function names.
+	TStringBuilder<64> FormatModuleNameBuilder;
+	FormatModuleNameBuilder << TextureFormatModuleName.ToString().Replace(TEXT("TextureFormat"), TEXT(""));
+	FormatModuleNameBuilder << TEXT("Texture");
+
+	BuildWriter.AddString("Name", FormatModuleNameBuilder);
+	BuildWriter.AddString("Version", "0"); // TODO: compute module version
+	BuildWriter.EndObject();
+
+	BuildWriter.BeginObject("Constants");
+	WriteBuildSettingsToCompactBinary(BuildWriter, "TextureBuildSettings", BuildSettings);
+	WriteOutputSettingsToCompactBinary(BuildWriter, "TextureOutputSettings", NumInlineMips, KeySuffix);
+
+	FTextureFormatSettings TextureFormatSettings;
+	Texture.GetLayerFormatSettings(LayerIndex, TextureFormatSettings);
+	EGammaSpace TextureGammaSpace = TextureFormatSettings.SRGB ? (Texture.bUseLegacyGamma ? EGammaSpace::Pow22 : EGammaSpace::sRGB) : EGammaSpace::Linear;
+	WriteTextureSourceToCompactBinary(BuildWriter, "TextureSource", Texture.Source, TextureGammaSpace);
+	if ((bool)Texture.CompositeTexture && !ExportedCompositeTextureBulkDataHash.IsZero())
+	{
+		FTextureFormatSettings CompositeTextureFormatSettings;
+		Texture.CompositeTexture->GetLayerFormatSettings(LayerIndex, CompositeTextureFormatSettings);
+		EGammaSpace CompositeTextureGammaSpace = CompositeTextureFormatSettings.SRGB ? (Texture.CompositeTexture->bUseLegacyGamma ? EGammaSpace::Pow22 : EGammaSpace::sRGB) : EGammaSpace::Linear;
+
+		WriteTextureSourceToCompactBinary(BuildWriter, "CompositeTextureSource", Texture.CompositeTexture->Source, CompositeTextureGammaSpace);
+	}
+	BuildWriter.EndObject();
+
+	BuildWriter.BeginObject("Inputs");
+	BuildWriter.AddBinaryAttachment(TCHAR_TO_UTF8(*Texture.Source.GetIdString()), ExportedTextureBulkDataHash);
+	if ((bool)Texture.CompositeTexture && !ExportedCompositeTextureBulkDataHash.IsZero())
+	{
+		BuildWriter.AddBinaryAttachment(TCHAR_TO_UTF8(*Texture.CompositeTexture->Source.GetIdString()), ExportedCompositeTextureBulkDataHash);
+	}
+	BuildWriter.EndObject();
+
+	BuildWriter.EndObject();
+	if (TUniquePtr<FArchive> Ar{IFileManager::Get().CreateFileWriter(*(ExportRoot / TEXT("build.uddba")))})
+	{
+		BuildWriter.Save(*Ar);
+	}
+}
+
+void FTextureDerivedDataBuildExporter::ExportTextureOutput(FTexturePlatformData& PlatformData, const FTextureBuildSettings& BuildSettings)
+{
+	if (!bEnabled)
+	{
+		return;
+	}
+	FString DerivedDataKeyLong;
+	GetTextureDerivedDataKeyFromSuffix(KeySuffix, DerivedDataKeyLong);
+
+	const bool bForceAllMipsToBeInlined = BuildSettings.bCubemap || (BuildSettings.bVolume && !BuildSettings.bStreamable) || (BuildSettings.bTextureArray && !BuildSettings.bStreamable);
+	const FString OutputPath = ExportRoot / TEXT("ReferenceOutputs");
+	TArray<TPair<FString,FString>> DDCReferences;
+
+	FCbWriter BuildWriter;
+	BuildWriter.BeginObject("BuildOutput");
+	BuildWriter.BeginArray("Payloads");
+
+	const int32 MipCount = PlatformData.Mips.Num();
+	const int32 FirstInlineMip = bForceAllMipsToBeInlined ? 0 : FMath::Max(0, MipCount - FMath::Max((int32)NUM_INLINE_DERIVED_MIPS, (int32)PlatformData.GetNumMipsInTail()));
+	for (int32 MipIndex = 0; MipIndex < FirstInlineMip; ++MipIndex)
+	{
+		FTexture2DMipMap& Mip = PlatformData.Mips[MipIndex];
+
+		int32 BulkDataSizeInBytes = Mip.BulkData.GetBulkDataSize();
+		check(BulkDataSizeInBytes > 0);
+
+		TArray<uint8> DerivedData;
+		FMemoryWriter Ar(DerivedData, /*bIsPersistent=*/ true);
+		Ar << BulkDataSizeInBytes;
+		{
+			void* BulkMipData = Mip.BulkData.Lock(LOCK_READ_ONLY);
+			Ar.Serialize(BulkMipData, BulkDataSizeInBytes);
+			Mip.BulkData.Unlock();
+		}
+
+		TStringBuilder<32> PayloadName;
+		PayloadName << "Mip" << MipIndex;
+
+		BuildWriter.BeginObject();
+		BuildWriter.AddObjectId("Id", FCbObjectId(UE::DerivedData::FPayloadId::FromName(*PayloadName).GetView()));
+		BuildWriter.AddInteger("RawSize", DerivedData.Num());
+		FIoHash DerivedDataHash = HashAndWriteToCompressedBufferFile(OutputPath, DerivedData.GetData(), DerivedData.Num());
+		BuildWriter.AddBinaryAttachment("RawHash", DerivedDataHash);
+		BuildWriter.EndObject();
+
+		check(Mip.DerivedDataKey.IsEmpty());
+
+		FString MipDerivedDataKeyLong;
+		FString MipDerivedDataKey;
+		GetTextureDerivedMipKey(MipIndex, Mip, KeySuffix, MipDerivedDataKeyLong);
+		ShortenKey(*MipDerivedDataKeyLong, MipDerivedDataKey);
+		DDCReferences.Emplace(PayloadName.ToString(), MipDerivedDataKey);
+		Mip.DerivedDataKey = MipDerivedDataKeyLong;
+	}
+
+	FTexturePlatformData StrippedPlatformData;
+
+	TArray<uint8> RawDerivedData;
+	FMemoryWriter Ar(RawDerivedData, /*bIsPersistent=*/ true);
+	PlatformData.SerializeWithConditionalBulkData(Ar, NULL);
+
+	for (int32 MipIndex = FirstInlineMip; MipIndex < MipCount; ++MipIndex)
+	{
+		PlatformData.Mips[MipIndex].DerivedDataKey.Empty();
+	}
+
+	BuildWriter.BeginObject();
+	BuildWriter.AddObjectId("Id", FCbObjectId(UE::DerivedData::FPayloadId::FromName("Texture").GetView()));
+	BuildWriter.AddInteger("RawSize", RawDerivedData.Num());
+	FIoHash DerivedDataHash = HashAndWriteToCompressedBufferFile(OutputPath, RawDerivedData.GetData(), RawDerivedData.Num());
+	BuildWriter.AddBinaryAttachment("RawHash", DerivedDataHash);
+	BuildWriter.EndObject();
+
+	FString DerivedDataKey;
+	ShortenKey(*DerivedDataKeyLong, DerivedDataKey);
+	DDCReferences.Emplace(TEXT("Texture"), DerivedDataKey);
+
+	BuildWriter.EndArray();
+
+	BuildWriter.EndObject();
+
+	if (TUniquePtr<FArchive> FileAr{IFileManager::Get().CreateFileWriter(*(ExportRoot / TEXT("ReferenceOutput.uddbo")))})
+	{
+		BuildWriter.Save(*FileAr);
+	}
+
+	TArray<FString> DDCRefStringArray;
+	FString Separator(TEXT(","));
+	for (TPair<FString, FString>& DDCReference : DDCReferences)
+	{
+		TStringBuilder<256> LineBuilder;
+		LineBuilder << DDCReference.Key << Separator << DDCReference.Value;
+		DDCRefStringArray.Add(LineBuilder.ToString());
+	}
+	FFileHelper::SaveStringArrayToFile(DDCRefStringArray, *(ExportRoot / TEXT("DDCReferences.txt")));
+}
+
+
+#endif // WITH_EDITOR
diff --git a/Engine/Source/Runtime/Engine/Private/TextureDerivedDataBuildExporter.h b/Engine/Source/Runtime/Engine/Private/TextureDerivedDataBuildExporter.h
new file mode 100644
--- /dev/null
+++ b/Engine/Source/Runtime/Engine/Private/TextureDerivedDataBuildExporter.h
@@ -0,0 +1,35 @@
+// Copyright Epic Games, Inc. All Rights Reserved.
+
+#pragma once
+
+#include "CoreMinimal.h"
+#include "IO/IoHash.h"
+
+#if WITH_EDITOR
+
+struct FTextureSource;
+struct FTextureBuildSettings;
+struct FTexturePlatformData;
+class UTexture;
+
+
+/** Collects input/output about the build of textures (i.e. compression) and exports the data to a file for testing of remote texture compression */
+class FTextureDerivedDataBuildExporter
+{
+public:
+	void Init(const FString& InKeySuffix);
+
+	void ExportTextureSourceBulkData(FTextureSource& TextureSource);
+	void ExportCompositeTextureSourceBulkData(FTextureSource& TextureSource);
+
+	void ExportTextureBuild(const UTexture& Texture, const FTextureBuildSettings& BuildSettings, int32 LayerIndex, int32 MaxInlineMips);
+	void ExportTextureOutput(FTexturePlatformData& PlatformData, const FTextureBuildSettings& BuildSettings);
+private:
+	FString KeySuffix;
+	FString ExportRoot;
+	FIoHash ExportedTextureBulkDataHash;
+	FIoHash ExportedCompositeTextureBulkDataHash;
+	bool bEnabled = false;
+};
+
+#endif // WITH_EDITOR
diff --git a/Engine/Source/Runtime/Engine/Private/TextureDerivedDataTask.cpp b/Engine/Source/Runtime/Engine/Private/TextureDerivedDataTask.cpp
--- a/Engine/Source/Runtime/Engine/Private/TextureDerivedDataTask.cpp
+++ b/Engine/Source/Runtime/Engine/Private/TextureDerivedDataTask.cpp
@@ -340,6 +340,8 @@ void FTextureCacheDerivedDataWorker::BuildTexture(bool bReplaceExistingDDC)
 
 		FOptTexturePlatformData OptData;
 
+		BuildExporter.ExportTextureBuild(Texture, BuildSettingsPerLayer[0], 0, NUM_INLINE_DERIVED_MIPS);
+
 		// Compress the texture.
 		TArray<FCompressedImage2D> CompressedMips;
 		if (Compressor->BuildTexture(TextureData.Blocks[0].MipsPerLayer[0],
@@ -385,6 +387,8 @@ void FTextureCacheDerivedDataWorker::BuildTexture(bool bReplaceExistingDDC)
 
 			DerivedData->SetOptData(OptData);
 
+			BuildExporter.ExportTextureOutput(*DerivedData, BuildSettingsPerLayer[0]);
+			
 			// Store it in the cache.
 			// @todo: This will remove the streaming bulk data, which we immediately reload below!
 			// Should ideally avoid this redundant work, but it only happens when we actually have 
@@ -448,6 +452,7 @@ FTextureCacheDerivedDataWorker::FTextureCacheDerivedDataWorker(
 	}
 	UTexture::GetPixelFormatEnum();
 	GetTextureDerivedDataKeySuffix(Texture, InSettingsPerLayer, KeySuffix);
+	BuildExporter.Init(KeySuffix);
 		
 	const bool bAllowAsyncBuild = (CacheFlags & ETextureCacheFlags::AllowAsyncBuild) != 0;
 	const bool bAllowAsyncLoading = (CacheFlags & ETextureCacheFlags::AllowAsyncLoading) != 0;
@@ -614,30 +619,22 @@ void FTextureCacheDerivedDataWorker::DoWork()
 	
 	if (!bSucceeded && bAllowAsyncBuild)
 	{
-		bool bHasTextureSourceMips = false;
 		if (TextureData.IsValid() && Texture.Source.IsBulkDataLoaded())
 		{
-			TextureData.GetSourceMips(Texture.Source, ImageWrapper);
-			bHasTextureSourceMips = true;
+			BuildExporter.ExportTextureSourceBulkData(Texture.Source);
 		}
-
-		bool bHasCompositeTextureSourceMips = false;
-		if (CompositeTextureData.IsValid() && Texture.CompositeTexture && Texture.CompositeTexture->Source.IsBulkDataLoaded())
+		else if (bAllowAsyncLoading)
 		{
-			CompositeTextureData.GetSourceMips(Texture.CompositeTexture->Source, ImageWrapper);
-			bHasCompositeTextureSourceMips = true;
+			BuildExporter.ExportTextureSourceBulkData(TextureData.AsyncSource);
 		}
 
-		if (bAllowAsyncLoading && !bHasTextureSourceMips)
+		if (CompositeTextureData.IsValid() && Texture.CompositeTexture && Texture.CompositeTexture->Source.IsBulkDataLoaded())
 		{
-			TextureData.GetAsyncSourceMips(ImageWrapper);
-			TextureData.AsyncSource.RemoveBulkData();
+			BuildExporter.ExportCompositeTextureSourceBulkData(Texture.CompositeTexture->Source);
 		}
-
-		if (bAllowAsyncLoading && !bHasCompositeTextureSourceMips)
+		else if (bAllowAsyncLoading)
 		{
-			CompositeTextureData.GetAsyncSourceMips(ImageWrapper);
-			CompositeTextureData.AsyncSource.RemoveBulkData();
+			BuildExporter.ExportCompositeTextureSourceBulkData(CompositeTextureData.AsyncSource);
 		}
 
 		if (TextureData.Blocks.Num() && TextureData.Blocks[0].MipsPerLayer.Num() && TextureData.Blocks[0].MipsPerLayer[0].Num() && 
diff --git a/Engine/Source/Runtime/Engine/Private/TextureDerivedDataTask.h b/Engine/Source/Runtime/Engine/Private/TextureDerivedDataTask.h
--- a/Engine/Source/Runtime/Engine/Private/TextureDerivedDataTask.h
+++ b/Engine/Source/Runtime/Engine/Private/TextureDerivedDataTask.h
@@ -15,6 +15,7 @@
 #include "TextureCompressorModule.h"
 #include "IImageWrapperModule.h"
 #include "Engine/Texture2D.h"
+#include "TextureDerivedDataBuildExporter.h"
 
 #endif // WITH_EDITOR
 
@@ -121,6 +122,8 @@ class FTextureCacheDerivedDataWorker : public FNonAbandonableTask
 	FTextureSourceData TextureData;
 	/** Source mip images of the composite texture (e.g. normal map for compute roughness). Not necessarily in RGBA32F, usually only top mip as other mips need to be generated first */
 	FTextureSourceData CompositeTextureData;
+	/** Exporter that can optionally write out build actions and reference outputs for the purpose of testing remote execution of builds */
+	FTextureDerivedDataBuildExporter BuildExporter;
 	/** Texture cache flags. */
 	uint32 CacheFlags;
 	/** Have many bytes were loaded from DDC or built (for telemetry) */
diff --git a/Engine/Source/Runtime/Engine/Public/TextureResource.h b/Engine/Source/Runtime/Engine/Public/TextureResource.h
--- a/Engine/Source/Runtime/Engine/Public/TextureResource.h
+++ b/Engine/Source/Runtime/Engine/Public/TextureResource.h
@@ -62,6 +62,8 @@ struct FTexture2DMipMap
 	ENGINE_API void Serialize(FArchive& Ar, UObject* Owner, int32 MipIndex);
 
 #if WITH_EDITORONLY_DATA
+	void SerializeWithConditionalBulkData(FArchive& Ar, UObject* Owner, int32 MipIndex);
+
 	/** Key if stored in the derived data cache. */
 	FString DerivedDataKey;
 
@@ -74,6 +76,8 @@ struct FTexture2DMipMap
 	 */
 	uint32 StoreInDerivedDataCache(const FString& InDerivedDataKey, const FStringView& TextureName, bool bReplaceExistingDDC);
 #endif // #if WITH_EDITORONLY_DATA
+private:
+	void SerializeCommon(FArchive& Ar, UObject* Owner, int32 MipIndex, bool bBulkDataConditionalOnKey);
 };
 
 /** 
