ComputeGraph graph based compilation.
More work to do but this is a minimal functional starting point.

diff --git a/Engine/Shaders/Private/ComputeKernel.usf b/Engine/Shaders/Private/ComputeKernel.usf
--- a/Engine/Shaders/Private/ComputeKernel.usf
+++ b/Engine/Shaders/Private/ComputeKernel.usf
@@ -4,16 +4,7 @@
 #include "/Engine/Private/Common.ush"
 #include "/Engine/Private/ComputeKernelCommon.ush"
 
-KERNEL_SRV(Buffer<float>, a);
-
-KERNEL_SRV(Buffer<float>, b);
-
-KERNEL_UAV(RWBuffer<float>, c);
-
 [numthreads(32, 1, 1)]
-KERNEL_ENTRY_POINT(LinearBlendSkinning)
+KERNEL_ENTRY_POINT(ComputeKernelNull)
 {
-	c[DTid.x] = a[DTid.x]
-	 + 
-	b[DTid.x]; 
 }
diff --git a/Engine/Shaders/Private/ComputeKernelCommon.ush b/Engine/Shaders/Private/ComputeKernelCommon.ush
--- a/Engine/Shaders/Private/ComputeKernelCommon.ush
+++ b/Engine/Shaders/Private/ComputeKernelCommon.ush
@@ -2,6 +2,8 @@
 
 #pragma once
 
+// Macro to mark a kernel entry point.
+// Only one should exist per kernel.
 
 #define KERNEL_ENTRY_POINT(EntryPointName) \
     void EntryPointName( \
@@ -10,6 +12,27 @@
         uint3 GTid: SV_GroupThreadID, \
         uint3 DTid : SV_DispatchThreadID)
 
-#define KERNEL_PARAM(Type, Name)    Type Name
-#define KERNEL_SRV(Type, Name)      Type Name
-#define KERNEL_UAV(Type, Name)      Type Name
+// Macros to directly expose parameters for a kernel.
+
+#define KERNEL_PARAM(Type, Name)    Type Name;
+
+// Macros to expose external functions for a kernel.
+// A kernel assumes these functions exist and are provided by its data providers.
+
+#define KERNEL_EXTERN_READ(Name, ...)
+#define KERNEL_EXTERN_WRITE(Name, ...)
+
+// Macro to disambiguate data interface code with a unique id.
+// This allows us to include more than one instance of the same data interface.
+// This should be used around all uniforms and all local (ie not DI_IMPL) functions.
+
+#define _DI_CAT_IMPL_(A, B) A##B
+#define _DI_CAT_(A, B) _DI_CAT_IMPL_(A, B)
+#define DI_LOCAL(Name) _DI_CAT_(DI_UID, Name)
+
+// Macros to expose functions implemented by a data provider. 
+// This makes a local unqiue function name and expands the function prototype.
+// We could just use DI_LOCAL, but these macro also gives the possiblity to extract functions by parsing the file.
+
+#define DI_IMPL_READ(Name, ReturnType, ...) ReturnType DI_LOCAL(Name)(__VA_ARGS__)
+#define DI_IMPL_WRITE(Name, ...) void DI_LOCAL(Name)(__VA_ARGS__)
diff --git a/Engine/Source/Developer/AssetTools/Private/AssetTools.cpp b/Engine/Source/Developer/AssetTools/Private/AssetTools.cpp
--- a/Engine/Source/Developer/AssetTools/Private/AssetTools.cpp
+++ b/Engine/Source/Developer/AssetTools/Private/AssetTools.cpp
@@ -114,6 +114,8 @@
 #include "AssetTypeActions/AssetTypeActions_VectorFieldStatic.h"
 #include "AssetTypeActions/AssetTypeActions_World.h"
 #include "AssetTypeActions/Experimental/AssetTypeActions_ChaosPhysicalMaterial.h"
+#include "ComputeFramework/ComputeKernelAssetActions.h"
+#include "ComputeFramework/ComputeKernelFromTextAssetActions.h"
 #include "SDiscoveringAssetsDialog.h"
 #include "AssetFixUpRedirectors.h"
 #include "ObjectTools.h"
@@ -234,6 +236,8 @@ UAssetToolsImpl::UAssetToolsImpl(const FObjectInitializer& ObjectInitializer)
 	RegisterAssetTypeActions(MakeShareable(new FAssetTypeActions_BlueprintGeneratedClass));
 	RegisterAssetTypeActions(MakeShareable(new FAssetTypeActions_CameraAnim));
 	RegisterAssetTypeActions(MakeShareable(new FAssetTypeActions_CanvasRenderTarget2D));
+	RegisterAssetTypeActions(MakeShareable(new FAssetTypeActions_ComputeKernel));
+	RegisterAssetTypeActions(MakeShareable(new FAssetTypeActions_ComputeKernelFromText));
 	RegisterAssetTypeActions(MakeShareable(new FAssetTypeActions_Curve));
 	RegisterAssetTypeActions(MakeShareable(new FAssetTypeActions_CurveFloat));
 	RegisterAssetTypeActions(MakeShareable(new FAssetTypeActions_CurveTable));
diff --git a/Engine/Source/Editor/UnrealEd/Private/ComputeFramework/ComputeKernelAssetActions.cpp b/Engine/Source/Editor/UnrealEd/Private/ComputeFramework/ComputeKernelAssetActions.cpp
new file mode 100644
--- /dev/null
+++ b/Engine/Source/Editor/UnrealEd/Private/ComputeFramework/ComputeKernelAssetActions.cpp
@@ -0,0 +1,31 @@
+// Copyright Epic Games, Inc. All Rights Reserved.
+
+#include "ComputeFramework/ComputeKernelAssetActions.h"
+
+#include "ComputeFramework/ComputeKernel.h"
+
+FAssetTypeActions_ComputeKernel::FAssetTypeActions_ComputeKernel(EAssetTypeCategories::Type InAssetCategoryBit)
+	: AssetCategoryBit(InAssetCategoryBit)
+{
+}
+
+FText FAssetTypeActions_ComputeKernel::GetName() const
+{
+	return NSLOCTEXT("ComputeFramework", "ComputeKernelInstanceName", "Compute Kernel Instance");
+}
+
+FColor FAssetTypeActions_ComputeKernel::GetTypeColor() const
+{
+	return FColor::Turquoise;
+}
+
+UClass* FAssetTypeActions_ComputeKernel::GetSupportedClass() const
+{
+	return UComputeKernel::StaticClass();
+}
+
+uint32 FAssetTypeActions_ComputeKernel::GetCategories()
+{
+	return AssetCategoryBit;
+}
+
diff --git a/Engine/Source/Editor/UnrealEd/Private/ComputeFramework/ComputeKernelFactory.cpp b/Engine/Source/Editor/UnrealEd/Private/ComputeFramework/ComputeKernelFactory.cpp
new file mode 100644
--- /dev/null
+++ b/Engine/Source/Editor/UnrealEd/Private/ComputeFramework/ComputeKernelFactory.cpp
@@ -0,0 +1,18 @@
+// Copyright Epic Games, Inc. All Rights Reserved.
+
+#include "ComputeFramework/ComputeKernelFactory.h"
+#include "ComputeFramework/ComputeKernel.h"
+
+UComputeKernelFactory::UComputeKernelFactory()
+{
+	SupportedClass = UComputeKernel::StaticClass();
+	bCreateNew = true;
+	bEditAfterNew = true;
+}
+
+UObject* UComputeKernelFactory::FactoryCreateNew(UClass* InClass, UObject* InParent, FName InName, EObjectFlags Flags, UObject* Context, FFeedbackContext* Warn)
+{
+	UComputeKernel* Kernel = NewObject<UComputeKernel>(InParent, InClass, InName, Flags);
+
+	return Kernel;
+}
diff --git a/Engine/Source/Editor/UnrealEd/Private/ComputeFramework/ComputeKernelFromTextAssetActions.cpp b/Engine/Source/Editor/UnrealEd/Private/ComputeFramework/ComputeKernelFromTextAssetActions.cpp
new file mode 100644
--- /dev/null
+++ b/Engine/Source/Editor/UnrealEd/Private/ComputeFramework/ComputeKernelFromTextAssetActions.cpp
@@ -0,0 +1,30 @@
+// Copyright Epic Games, Inc. All Rights Reserved.
+
+#include "ComputeFramework/ComputeKernelFromTextAssetActions.h"
+
+#include "ComputeFramework/ComputeKernelFromText.h"
+
+FAssetTypeActions_ComputeKernelFromText::FAssetTypeActions_ComputeKernelFromText(EAssetTypeCategories::Type InAssetCategoryBit)
+	: AssetCategoryBit(InAssetCategoryBit)
+{
+}
+
+FText FAssetTypeActions_ComputeKernelFromText::GetName() const
+{
+	return NSLOCTEXT("ComputeFramework", "ComputeKernelFromTextName", "Compute Kernel (Text)");
+}
+
+FColor FAssetTypeActions_ComputeKernelFromText::GetTypeColor() const
+{
+	return FColor::Turquoise;
+}
+
+UClass* FAssetTypeActions_ComputeKernelFromText::GetSupportedClass() const
+{
+	return UComputeKernelFromText::StaticClass();
+}
+
+uint32 FAssetTypeActions_ComputeKernelFromText::GetCategories()
+{
+	return AssetCategoryBit;
+}
diff --git a/Engine/Source/Editor/UnrealEd/Private/ComputeFramework/ComputeKernelFromTextFactory.cpp b/Engine/Source/Editor/UnrealEd/Private/ComputeFramework/ComputeKernelFromTextFactory.cpp
new file mode 100644
--- /dev/null
+++ b/Engine/Source/Editor/UnrealEd/Private/ComputeFramework/ComputeKernelFromTextFactory.cpp
@@ -0,0 +1,18 @@
+// Copyright Epic Games, Inc. All Rights Reserved.
+
+#include "ComputeFramework/ComputeKernelFromTextFactory.h"
+#include "ComputeFramework/ComputeKernelFromText.h"
+
+UComputeKernelFromTextFactory::UComputeKernelFromTextFactory()
+{
+	SupportedClass = UComputeKernelFromText::StaticClass();
+	bCreateNew = true;
+	bEditAfterNew = true;
+}
+
+UObject* UComputeKernelFromTextFactory::FactoryCreateNew(UClass* InClass, UObject* InParent, FName InName, EObjectFlags Flags, UObject* Context, FFeedbackContext* Warn)
+{
+	UComputeKernelFromText* Kernel = NewObject<UComputeKernelFromText>(InParent, InClass, InName, Flags);
+
+	return Kernel;
+}
diff --git a/Engine/Source/Editor/UnrealEd/Public/ComputeFramework/ComputeKernelAssetActions.h b/Engine/Source/Editor/UnrealEd/Public/ComputeFramework/ComputeKernelAssetActions.h
new file mode 100644
--- /dev/null
+++ b/Engine/Source/Editor/UnrealEd/Public/ComputeFramework/ComputeKernelAssetActions.h
@@ -0,0 +1,23 @@
+// Copyright Epic Games, Inc. All Rights Reserved.
+
+#pragma once
+
+#include "AssetTypeCategories.h"
+#include "AssetTypeActions_Base.h"
+
+class UNREALED_API FAssetTypeActions_ComputeKernel : public FAssetTypeActions_Base
+{
+public:
+	FAssetTypeActions_ComputeKernel(EAssetTypeCategories::Type InAssetCategoryBit = EAssetTypeCategories::Animation);
+
+protected:
+	//~ Begin FAssetTypeActions_Base Interface.
+	FText GetName() const override;
+	FColor GetTypeColor() const override;
+	UClass* GetSupportedClass() const override;
+	uint32 GetCategories() override;
+	//~ End FAssetTypeActions_Base Interface.
+
+private:
+	EAssetTypeCategories::Type AssetCategoryBit;
+};
diff --git a/Engine/Source/Editor/UnrealEd/Public/ComputeFramework/ComputeKernelFactory.h b/Engine/Source/Editor/UnrealEd/Public/ComputeFramework/ComputeKernelFactory.h
new file mode 100644
--- /dev/null
+++ b/Engine/Source/Editor/UnrealEd/Public/ComputeFramework/ComputeKernelFactory.h
@@ -0,0 +1,24 @@
+// Copyright Epic Games, Inc. All Rights Reserved.
+
+#pragma once
+
+#include "Factories/Factory.h"
+#include "ComputeKernelFactory.generated.h"
+
+UCLASS(hidecategories = Object)
+class UComputeKernelFactory : public UFactory
+{
+	GENERATED_BODY()
+
+public:
+	UComputeKernelFactory();
+
+	virtual UObject* FactoryCreateNew(
+		UClass* InClass, 
+		UObject* InParent, 
+		FName InName, 
+		EObjectFlags Flags, 
+		UObject* Context, 
+		FFeedbackContext* Warn
+		) override;
+};
diff --git a/Engine/Source/Editor/UnrealEd/Public/ComputeFramework/ComputeKernelFromTextAssetActions.h b/Engine/Source/Editor/UnrealEd/Public/ComputeFramework/ComputeKernelFromTextAssetActions.h
new file mode 100644
--- /dev/null
+++ b/Engine/Source/Editor/UnrealEd/Public/ComputeFramework/ComputeKernelFromTextAssetActions.h
@@ -0,0 +1,23 @@
+// Copyright Epic Games, Inc. All Rights Reserved.
+
+#pragma once
+
+#include "AssetTypeCategories.h"
+#include "AssetTypeActions_Base.h"
+
+class UNREALED_API FAssetTypeActions_ComputeKernelFromText : public FAssetTypeActions_Base
+{
+public:
+	FAssetTypeActions_ComputeKernelFromText(EAssetTypeCategories::Type InAssetCategoryBit = EAssetTypeCategories::Animation);
+
+protected:
+	//~ Begin FAssetTypeActions_Base Interface.
+	FText GetName() const override;
+	FColor GetTypeColor() const override;
+	UClass* GetSupportedClass() const override;
+	uint32 GetCategories() override;
+	//~ End FAssetTypeActions_Base Interface.
+
+private:
+	EAssetTypeCategories::Type AssetCategoryBit;
+};
diff --git a/Engine/Source/Editor/UnrealEd/Public/ComputeFramework/ComputeKernelFromTextFactory.h b/Engine/Source/Editor/UnrealEd/Public/ComputeFramework/ComputeKernelFromTextFactory.h
new file mode 100644
--- /dev/null
+++ b/Engine/Source/Editor/UnrealEd/Public/ComputeFramework/ComputeKernelFromTextFactory.h
@@ -0,0 +1,24 @@
+// Copyright Epic Games, Inc. All Rights Reserved.
+
+#pragma once
+
+#include "Factories/Factory.h"
+#include "ComputeKernelFromTextFactory.generated.h"
+
+UCLASS(hidecategories = Object)
+class UComputeKernelFromTextFactory : public UFactory
+{
+	GENERATED_BODY()
+
+public:
+	UComputeKernelFromTextFactory();
+
+	virtual UObject* FactoryCreateNew(
+		UClass* InClass, 
+		UObject* InParent, 
+		FName InName, 
+		EObjectFlags Flags, 
+		UObject* Context, 
+		FFeedbackContext* Warn
+		) override;
+};
diff --git a/Engine/Source/Runtime/Engine/Private/ComputeFramework/ComputeFramework.cpp b/Engine/Source/Runtime/Engine/Private/ComputeFramework/ComputeFramework.cpp
--- a/Engine/Source/Runtime/Engine/Private/ComputeFramework/ComputeFramework.cpp
+++ b/Engine/Source/Runtime/Engine/Private/ComputeFramework/ComputeFramework.cpp
@@ -1,15 +1,8 @@
 // Copyright Epic Games, Inc. All Rights Reserved.
 
 #include "ComputeFramework/ComputeFramework.h"
-#include "ComputeFramework/ComputeKernel.h"
-#include "ComputeFramework/ComputeGraph.h"
-#include "RenderGraphBuilder.h"
-#include "RenderGraphUtils.h"
-#include "RHIDefinitions.h"
-#include "ComputeFramework/ComputeKernelShader.h"
-#include "ProfilingDebugging/RealtimeGPUProfiler.h"
 
-DECLARE_GPU_STAT_NAMED(ComputeFramework_ExecuteBatches, TEXT("ComputeFramework::ExecuteBatches"));
+DEFINE_LOG_CATEGORY(LogComputeFramework);
 
 static int32 GComputeFrameworkMode = 1;
 static FAutoConsoleVariableRef CVarComputeFrameworkMode(
@@ -27,186 +20,3 @@ bool SupportsComputeFramework(ERHIFeatureLevel::Type FeatureLevel, EShaderPlatfo
 		&& FeatureLevel >= ERHIFeatureLevel::SM5
 		&& FDataDrivenShaderPlatformInfo::GetSupportsComputeFramework(ShaderPlatform);
 }
-
-void FComputeGraph::Initialize(UComputeGraph* ComputeGraph)
-{
-	for (auto& Invocation : ComputeGraph->GetShaderInvocationList())
-	{
-		FKernelInvocation KernelInvocation = {
-			Invocation.ComputeKernel->GetFName(),
-			FName("InvocationName"),
-			FIntVector(1, 1, 1),
-			Invocation.ComputeKernel->GetResource()
-			};
-
-		KernelInvocations.Emplace(KernelInvocation);
-	}
-}
-
-void FComputeFramework::EnqueueForExecution(const FComputeGraph* ComputeGraph)
-{
-	//ExternalBuffers.Append(ComputeGraph.ExternalBuffers);
-	//TransientBuffers.Append(ComputeGraph.TransientBuffers);
-
-	for (auto& Invocation : ComputeGraph->KernelInvocations)
-	{
-		FShaderInvocation ShaderInvocation = {
-			Invocation.KernelName,
-			Invocation.InvocationName,
-			Invocation.DispatchDim,
-			Invocation.Kernel->GetShaderParamMetadata(),
-			Invocation.Kernel->GetShader()
-			};
-
-		ComputeShaders.Emplace(ShaderInvocation);
-	}
-}
-
-struct FComputeExecutionBuffer
-{
-	FComputeExecutionBuffer() = default;
-
-	FComputeExecutionBuffer(FName InName, FRDGBufferRef InBufferRef)
-		: Name(InName)
-		, BufferRef(InBufferRef)
-	{
-	}
-
-	FName Name;
-	FRDGBufferRef BufferRef;
-};
-
-void FComputeFramework::ExecuteBatches(
-	FRHICommandListImmediate& RHICmdList,
-	ERHIFeatureLevel::Type FeatureLevel
-	)
-{
-	if (ComputeShaders.IsEmpty())
-	{
-		return;
-	}
-
-	FMemMark MemStackMark(FMemStack::Get());
-	{
-		SCOPED_DRAW_EVENTF(RHICmdList, ComputeFramework_ExecuteBatches, TEXT("ComputeFramework::ExecuteBatches"));
-		SCOPED_GPU_STAT(RHICmdList, ComputeFramework_ExecuteBatches);
-
-		FRDGBuilder GraphBuilder(RHICmdList);
-		FGlobalShaderMap* ShaderMap = GetGlobalShaderMap(FeatureLevel);
-
-		TArray<FComputeExecutionBuffer> ExecutionBuffers;
-
-		for (auto& ExtRes : ExternalBuffers)
-		{
-			TCHAR ResName[FComputeResource::MAX_NAME_LENGTH];
-			ExtRes.Name.ToString(ResName);
-
-			ExecutionBuffers.Emplace(
-				ExtRes.Name,
-				GraphBuilder.RegisterExternalBuffer(ExtRes.Buffer, ResName)
-				);
-		}
-
-		for (auto& TransRes : TransientBuffers)
-		{
-			TCHAR ResName[FComputeResource::MAX_NAME_LENGTH];
-			TransRes.Name.ToString(ResName);
-
-			FRDGBufferDesc BufferDesc;
-		
-			EComputeExecutionBufferType BufferType = (EComputeExecutionBufferType)TransRes.BufferType;
-			switch (BufferType)
-			{
-			case EComputeExecutionBufferType::Buffer:
-				BufferDesc = FRDGBufferDesc::CreateBufferDesc(TransRes.BytesPerElement, TransRes.ElementCount);
-				break;
-
-			case EComputeExecutionBufferType::StructuredBuffer:
-				BufferDesc = FRDGBufferDesc::CreateStructuredDesc(TransRes.BytesPerElement, TransRes.ElementCount);
-				break;
-
-			case EComputeExecutionBufferType::ByteAddressBuffer:
-				BufferDesc = FRDGBufferDesc::CreateByteAddressDesc(TransRes.BytesPerElement * TransRes.ElementCount);
-				break;
-
-			default:
-				check(!"Unsupported buffer type");
-				break;
-			}
-
-			ExecutionBuffers.Emplace(
-				TransRes.Name,
-				GraphBuilder.CreateBuffer(BufferDesc, ResName)
-				);
-		}
-
-		for (auto& Compute : ComputeShaders)
-		{
-			FRDGBufferDesc BufferDesc = FRDGBufferDesc::CreateBufferDesc(4, 32);
-
-			void* RawBuffer = GraphBuilder.Alloc(Compute.ShaderParamMetadata->GetSize(), SHADER_PARAMETER_STRUCT_ALIGNMENT);
-			FMemory::Memzero(RawBuffer, Compute.ShaderParamMetadata->GetSize());
-
-			uint8* ParamBuffer = static_cast<uint8*>(RawBuffer);
-			TArray<FShaderParametersMetadata::FMember> ParamMembers = Compute.ShaderParamMetadata->GetMembers();
-			for (auto& Member : ParamMembers)
-			{
-				switch (Member.GetBaseType())
-				{
-				case EUniformBufferBaseType::UBMT_INT32:
-				case EUniformBufferBaseType::UBMT_UINT32:
-				case EUniformBufferBaseType::UBMT_FLOAT32:
-					// Fill is CBuffer data
-					break;
-
-				case EUniformBufferBaseType::UBMT_RDG_BUFFER_SRV:
-					{
-						FRDGBufferRef InputBuffer = GraphBuilder.CreateBuffer(BufferDesc, Member.GetName());
-
-						FComputeShaderUtils::ClearUAV(
-							GraphBuilder,
-							ShaderMap,
-							GraphBuilder.CreateUAV(InputBuffer, PF_A32B32G32R32F),
-							FVector4(0.0f)
-							);
-
-						*reinterpret_cast<FRDGBufferSRVRef*>(&ParamBuffer[Member.GetOffset()]) = GraphBuilder.CreateSRV(InputBuffer, PF_R32_FLOAT);
-					}
-					break;
-
-				case EUniformBufferBaseType::UBMT_RDG_BUFFER_UAV:
-					{						
-						FRDGBufferRef OutputBuffer = GraphBuilder.CreateBuffer(BufferDesc, Member.GetName());
-
-						*reinterpret_cast<FRDGBufferUAVRef*>(&ParamBuffer[Member.GetOffset()]) = GraphBuilder.CreateUAV(OutputBuffer, PF_R32_FLOAT);
-					}
-					break;
-
-				default:
-					check(!"Unsupported type");
-					break;
-				};
-			}
-
-			TCHAR KernelName[FComputeKernelInvocation::MAX_NAME_LENGTH];
-			Compute.KernelName.ToString(KernelName);
-
-			TCHAR InvocationName[FComputeKernelInvocation::MAX_NAME_LENGTH];
-			Compute.InvocationName.ToString(InvocationName);
-
-			FComputeShaderUtils::AddPass(
-				GraphBuilder,
-				RDG_EVENT_NAME("Compute[%s]: %s", KernelName, InvocationName),
-				ERDGPassFlags::Compute | ERDGPassFlags::NeverCull,
-				Compute.Shader,
-				Compute.ShaderParamMetadata,
-				static_cast<FComputeKernelShader::FParameters*>(RawBuffer),
-				Compute.DispatchDim
-				);		
-		}
-
-		GraphBuilder.Execute();
-
-		ComputeShaders.Reset();
-	}
-}
diff --git a/Engine/Source/Runtime/Engine/Private/ComputeFramework/ComputeGraph.cpp b/Engine/Source/Runtime/Engine/Private/ComputeFramework/ComputeGraph.cpp
--- a/Engine/Source/Runtime/Engine/Private/ComputeFramework/ComputeGraph.cpp
+++ b/Engine/Source/Runtime/Engine/Private/ComputeFramework/ComputeGraph.cpp
@@ -1,36 +1,276 @@
 // Copyright Epic Games, Inc. All Rights Reserved.
 
 #include "ComputeFramework/ComputeGraph.h"
+
+#include "ComputeFramework/ComputeDataInterface.h"
+#include "ComputeFramework/ComputeFramework.h"
 #include "ComputeFramework/ComputeKernel.h"
-#include "UObject/UObjectGlobals.h"
+#include "ComputeFramework/ComputeKernelShared.h"
+#include "ComputeFramework/ComputeKernelSource.h"
+#include "ShaderParameterMetadataBuilder.h"
 
+UComputeGraph::UComputeGraph(const FObjectInitializer& ObjectInitializer)
+	: UObject(ObjectInitializer)
+{
+}
 
-// #TODO_ZABIR: TEMP
-#include "ComputeFramework/ComputeKernelFromText.h"
+UComputeGraph::UComputeGraph(FVTableHelper& Helper)
+	: Super(Helper)
+{
+}
 
+UComputeGraph::~UComputeGraph() = default;
 
-FComputeKernelInvocationHandle UComputeGraph::AddKernel(UComputeKernelSource* KernelSource)
+void UComputeGraph::PostLoad()
 {
-	return {};
+	Super::PostLoad();
+
+#if WITH_EDITOR
+	// PostLoad our kernel dependencies before compiling.
+	for (UComputeKernel* Kernel : KernelInvocations)
+	{
+		if (Kernel != nullptr)
+		{
+			Kernel->ConditionalPostLoad();
+		}
+	}
+
+	if (ValidateGraph())
+	{
+		CacheResourceShadersForRendering();
+	}
+#endif
 }
 
-bool UComputeGraph::RemoveKernel(FComputeKernelInvocationHandle KernelInvocation)
+bool UComputeGraph::ValidateGraph(FString* OutErrors)
 {
-	return false;
+	// todo[CF]:
+	// Check same number of kernel in/outs as edges.
+	// Check each edge connects matching function types.
+	// Check graph is DAG
+	return true;
 }
 
-void UComputeGraph::PostLoad()
+TCHAR const* UComputeGraph::GetDataInterfaceUID(int32 DataInterfaceIndex)
 {
-	Super::PostLoad();
+	// Use static TChunkedArray so that data is never released/reallocated.
+	static TChunkedArray<FString, 512> UIDStore;
+	
+	if (DataInterfaceIndex >= UIDStore.Num())
+	{
+		UIDStore.Add(DataInterfaceIndex + 1 - UIDStore.Num());
+	}
+	if (UIDStore[DataInterfaceIndex].IsEmpty())
+	{
+		UIDStore[DataInterfaceIndex] = FString::Printf(TEXT("DI%03d"), DataInterfaceIndex);
+	}
+	return *UIDStore[DataInterfaceIndex];
+}
+
+#if WITH_EDITOR
+
+namespace
+{
+	/** Add HLSL code to implement an external function. */
+	void GetFunctionShimHLSL(FShaderFunctionDefinition const& FnImpl, FShaderFunctionDefinition const& FnWrap, TCHAR const* UID, FString& InOutHLSL)
+	{
+		const bool bHasReturn = FnWrap.bHasReturnType;
+		const int32 NumParams = FnWrap.ParamTypes.Num();
+
+		TStringBuilder<512> StringBuilder;
+
+		StringBuilder.Append(bHasReturn ? *FnWrap.ParamTypes[0].TypeDeclaration : TEXT("void"));
+		StringBuilder.Append(TEXT(" "));
+		StringBuilder.Append(*FnWrap.Name);
+		StringBuilder.Append(TEXT("("));
+		
+		for (int32 ParameterIndex = bHasReturn ? 1 : 0; ParameterIndex < NumParams; ++ParameterIndex)
+		{
+			StringBuilder.Append(*FnWrap.ParamTypes[ParameterIndex].TypeDeclaration);
+			StringBuilder.Appendf(TEXT(" P%d"), ParameterIndex);
+			StringBuilder.Append((ParameterIndex < NumParams - 1) ? TEXT(", ") : TEXT(""));
+		}
+
+		StringBuilder.Append(TEXT(") { "));
+		StringBuilder.Append(bHasReturn ? TEXT("return ") : TEXT(""));
+		StringBuilder.Append(UID).Append(TEXT("_")).Append(*FnImpl.Name);
+		StringBuilder.Append(TEXT("("));
+
+		for (int32 ParameterIndex = bHasReturn ? 1 : 0; ParameterIndex < NumParams; ++ParameterIndex)
+		{
+			StringBuilder.Appendf(TEXT("P%d"), ParameterIndex);
+			StringBuilder.Append((ParameterIndex < NumParams - 1) ? TEXT(", ") : TEXT(""));
+		}
+
+		StringBuilder.Append(TEXT("); }\n"));
+
+		InOutHLSL += StringBuilder.ToString();
+	}
+}
+
+void UComputeGraph::CacheResourceShadersForRendering()
+{
+	const uint32 CompilationFlags = uint32(EComputeKernelCompilationFlags::ApplyCompletedShaderMapForRendering) | uint32(EComputeKernelCompilationFlags::Force);
+	CacheResourceShadersForRendering(CompilationFlags);
+}
+
+void UComputeGraph::CacheResourceShadersForRendering(uint32 CompilationFlags)
+{
+	// We expect the graph to be validated before attempting to compile.
+	if (FApp::CanEverRender() && ensure(ValidateGraph()))
+	{
+		KernelResources.SetNum(KernelInvocations.Num());
+		for (int32 KernelIndex = 0; KernelIndex < KernelInvocations.Num(); ++KernelIndex)
+		{
+			UComputeKernel* Kernel = KernelInvocations[KernelIndex];
+			UComputeKernelSource* KernelSource = Kernel == nullptr ? nullptr : Kernel->KernelSource;
+
+			if (KernelSource == nullptr)
+			{
+				if (KernelResources[KernelIndex].IsValid())
+				{
+					KernelResources[KernelIndex]->Invalidate();
+					KernelResources[KernelIndex] = nullptr;
+				}
+
+				continue;
+			}
+
+			if (!KernelResources[KernelIndex].IsValid())
+			{
+				KernelResources[KernelIndex] = MakeUnique<FComputeKernelResource>();
+			}
+
+			FComputeKernelResource* KernelResource = KernelResources[KernelIndex].Get();
+
+			TArray<int32> RelevantEdgeIndices;
+			TArray<int32> DataProviderIndices;
+			for (int32 GraphEdgeIndex = 0; GraphEdgeIndex < GraphEdges.Num(); ++GraphEdgeIndex)
+			{
+				if (GraphEdges[GraphEdgeIndex].KernelIndex == KernelIndex)
+				{
+					RelevantEdgeIndices.Add(GraphEdgeIndex);
+					DataProviderIndices.AddUnique(GraphEdges[GraphEdgeIndex].DataInterfaceIndex);
+				}
+			}
+
+			// Collect data interface shader code.
+			FString HLSL;
+			for (int32 DataProviderIndex : DataProviderIndices)
+			{
+				UComputeDataInterface* DataInterface = DataInterfaces[DataProviderIndex];
+				
+				// Add a unique prefix to generate unique names in the data interface shader code.
+				TCHAR const* UID = GetDataInterfaceUID(DataProviderIndex);
+				HLSL += FString::Printf(TEXT("#define DI_UID %s_\n"), UID);
+				DataInterface->GetHLSL(HLSL);
+				HLSL += TEXT("#undef DI_UID\n");
+			}
+
+			// Bind every external kernel function to the associated data input function.
+			for (int32 GraphEdgeIndex : RelevantEdgeIndices)
+			{
+				if (GraphEdges[GraphEdgeIndex].bKernelInput)
+				{
+					TArray<FShaderFunctionDefinition> DataProviderFunctions;
+					DataInterfaces[GraphEdges[GraphEdgeIndex].DataInterfaceIndex]->GetSupportedInputs(DataProviderFunctions);
+					FShaderFunctionDefinition& DataProviderFunction = DataProviderFunctions[GraphEdges[GraphEdgeIndex].DataInterfaceBindingIndex];
+					FShaderFunctionDefinition& KernelFunction = KernelSource->ExternalInputs[GraphEdges[GraphEdgeIndex].KernelBindingIndex];
+					TCHAR const* UID = GetDataInterfaceUID(GraphEdges[GraphEdgeIndex].DataInterfaceIndex);
+					GetFunctionShimHLSL(DataProviderFunction, KernelFunction, UID, HLSL);
+				}
+				else
+				{
+					TArray<FShaderFunctionDefinition> DataProviderFunctions;
+					DataInterfaces[GraphEdges[GraphEdgeIndex].DataInterfaceIndex]->GetSupportedOutputs(DataProviderFunctions);
+					FShaderFunctionDefinition& DataProviderFunction = DataProviderFunctions[GraphEdges[GraphEdgeIndex].DataInterfaceBindingIndex];
+					FShaderFunctionDefinition& KernelFunction = KernelSource->ExternalOutputs[GraphEdges[GraphEdgeIndex].KernelBindingIndex];
+					TCHAR const* UID = GetDataInterfaceUID(GraphEdges[GraphEdgeIndex].DataInterfaceIndex);
+					GetFunctionShimHLSL(DataProviderFunction, KernelFunction, UID, HLSL);
+				}
+			}
+
+			HLSL += KernelSource->GetSource();
+
+			// Collect and build shader parameter metadata.
+			FShaderParametersMetadataBuilder Builder;
+			KernelSource->GetShaderParameters(Builder);
 
-	UComputeKernel* Kernel = LoadObject<UComputeKernel>(
-		nullptr, 
-		TEXT("/Game/Developers/ZabirHoque/CF/LBS_Kernel.LBS_Kernel"),
-		nullptr,
-		LOAD_None, 
-		nullptr
+			for (int32 DataProviderIndex : DataProviderIndices)
+			{
+				UComputeDataInterface* DataInterface = DataInterfaces[DataProviderIndex];
+				TCHAR const* UID = GetDataInterfaceUID(DataProviderIndex);
+				DataInterface->GetShaderParameters(UID, Builder);
+			}
+
+			FShaderParametersMetadata* ShaderMetadata = Builder.Build(FShaderParametersMetadata::EUseCase::ShaderParameterStruct, *GetName());
+
+			// Now we have all the information that the KernelResource will need for compilation.
+			KernelResource->SetupResource(
+				GMaxRHIFeatureLevel,
+				GetName(),
+				KernelSource->GetEntryPoint(),
+				MoveTemp(HLSL),
+				KernelSource->GetSourceCodeHash(),
+				ShaderMetadata);
+
+			const ERHIFeatureLevel::Type CacheFeatureLevel = GMaxRHIFeatureLevel;
+			const EShaderPlatform ShaderPlatform = GShaderPlatformForFeatureLevel[CacheFeatureLevel];
+
+			CacheShadersForResource(ShaderPlatform, nullptr, CompilationFlags | uint32(EComputeKernelCompilationFlags::Force), KernelResource);
+		}
+	}
+}
+
+void UComputeGraph::CacheShadersForResource(
+	EShaderPlatform ShaderPlatform,
+	const ITargetPlatform* TargetPlatform,
+	uint32 CompilationFlags,
+	FComputeKernelResource* KernelResource)
+{
+	bool bCooking = (CompilationFlags & uint32(EComputeKernelCompilationFlags::IsCooking)) != 0;
+
+	const bool bIsDefault = (KernelResource->GetKernelFlags() & uint32(EComputeKernelFlags::IsDefaultKernel)) != 0;
+	if (!GIsEditor || GIsAutomationTesting || bIsDefault || bCooking)
+	{
+		CompilationFlags |= uint32(EComputeKernelCompilationFlags::Synchronous);
+	}
+
+	const bool bIsSuccess = KernelResource->CacheShaders(
+		ShaderPlatform,
+		TargetPlatform,
+		CompilationFlags & uint32(EComputeKernelCompilationFlags::ApplyCompletedShaderMapForRendering),
+		CompilationFlags & uint32(EComputeKernelCompilationFlags::Synchronous)
+	);
+
+	if (!bIsSuccess)
+	{
+		if (bIsDefault)
+		{
+			UE_LOG(
+				LogComputeFramework,
+				Fatal,
+				TEXT("Failed to compile default FComputeKernelResource [%s] for platform [%s]!"),
+				*KernelResource->GetFriendlyName(),
+				*LegacyShaderPlatformToShaderFormat(ShaderPlatform).ToString()
+			);
+		}
+
+		UE_LOG(
+			LogComputeFramework,
+			Warning,
+			TEXT("Failed to compile FComputeKernelResource [%s] for platform [%s]."),
+			*KernelResource->GetFriendlyName(),
+			*LegacyShaderPlatformToShaderFormat(ShaderPlatform).ToString()
 		);
 
-	KernelInvocations.Reset();
-	KernelInvocations.Emplace(Kernel, 1);
+		auto& CompilationErrors = KernelResource->GetCompileErrors();
+		uint32 ErrorCount = CompilationErrors.Num();
+		for (uint32 i = 0; i < ErrorCount; ++i)
+		{
+			UE_LOG(LogComputeFramework, Warning, TEXT("      [Error] - %s"), *CompilationErrors[i]);
+		}
+	}
 }
+
+#endif // WITH_EDITOR
diff --git a/Engine/Source/Runtime/Engine/Private/ComputeFramework/ComputeGraphComponent.cpp b/Engine/Source/Runtime/Engine/Private/ComputeFramework/ComputeGraphComponent.cpp
--- a/Engine/Source/Runtime/Engine/Private/ComputeFramework/ComputeGraphComponent.cpp
+++ b/Engine/Source/Runtime/Engine/Private/ComputeFramework/ComputeGraphComponent.cpp
@@ -1,20 +1,31 @@
 // Copyright Epic Games, Inc. All Rights Reserved.
 
 #include "ComputeFramework/ComputeGraphComponent.h"
+
+#include "ComputeFramework/ComputeDataInterface.h"
+#include "ComputeFramework/ComputeDataProvider.h"
 #include "ComputeFramework/ComputeGraph.h"
-#include "ComputeFramework/ComputeFramework.h"
+#include "ComputeFramework/ComputeGraphScheduler.h"
 #include "SceneInterface.h"
 
-
 UComputeGraphComponent::UComputeGraphComponent()
 {
 	PrimaryComponentTick.bCanEverTick = true;
 	PrimaryComponentTick.bStartWithTickEnabled = false;
 }
 
+void UComputeGraphComponent::SetDataProvider(int32 Index, UComputeDataProvider* DataProvider)
+{
+	if (DataProviders.Num() <= Index)
+	{
+		DataProviders.SetNumZeroed(Index + 1);
+	}
+	DataProviders[Index] = DataProvider;
+}
+
 void UComputeGraphComponent::QueueExecute()
 {
-	if (!ComputeGraph || !GetScene()->GetComputeFramework())
+	if (ComputeGraph == nullptr || GetScene() == nullptr || GetScene()->GetComputeGraphScheduler() == nullptr)
 	{
 		return;
 	}
@@ -24,19 +35,41 @@ void UComputeGraphComponent::QueueExecute()
 
 void UComputeGraphComponent::SendRenderDynamicData_Concurrent()
 {
-	check(ComputeGraph);
 	Super::SendRenderDynamicData_Concurrent();
 	
-	FComputeFramework* ComputeFramework = GetScene()->GetComputeFramework();
-	check(ComputeFramework);
+	if (!ensure(ComputeGraph))
+	{
+		return;
+	}
+
+	FSceneInterface* Scene = GetScene();
+	FComputeGraphScheduler* ComputeGraphScheduler = Scene != nullptr ? Scene->GetComputeGraphScheduler() : nullptr;
+	if (!ensure(ComputeGraphScheduler))
+	{
+		return;
+	}
 
-	FComputeGraph* ComputeGraphProxy = new FComputeGraph();
+	FComputeGraphProxy* ComputeGraphProxy = new FComputeGraphProxy();
 	ComputeGraphProxy->Initialize(ComputeGraph);
 
+	TArray<FComputeDataProviderRenderProxy*> ComputeDataProviderProxies;
+	for (UComputeDataProvider* DataProvider : DataProviders)
+	{
+		// Add null provider slots because we want to maintain consistent array indices.
+		FComputeDataProviderRenderProxy* ProviderProxy = DataProvider != nullptr ? DataProvider->GetRenderProxy() : nullptr;
+		ComputeDataProviderProxies.Add(ProviderProxy);
+	}
+
 	ENQUEUE_RENDER_COMMAND(ComputeFrameworkEnqueueExecutionCommand)(
-		[ComputeFramework, ComputeGraphProxy](FRHICommandListImmediate& RHICmdList)
+		[ComputeGraphScheduler, ComputeGraphProxy, ProviderProxies = MoveTemp(ComputeDataProviderProxies)](FRHICommandListImmediate& RHICmdList)
 		{
-			ComputeFramework->EnqueueForExecution(ComputeGraphProxy);
+			ComputeGraphScheduler->EnqueueForExecution(ComputeGraphProxy, ProviderProxies);
+			
+			for (FComputeDataProviderRenderProxy* ProviderProxy : ProviderProxies)
+			{
+				delete ProviderProxy;
+			}
+
 			delete ComputeGraphProxy;
 		});
 }
diff --git a/Engine/Source/Runtime/Engine/Private/ComputeFramework/ComputeGraphScheduler.cpp b/Engine/Source/Runtime/Engine/Private/ComputeFramework/ComputeGraphScheduler.cpp
new file mode 100644
--- /dev/null
+++ b/Engine/Source/Runtime/Engine/Private/ComputeFramework/ComputeGraphScheduler.cpp
@@ -0,0 +1,254 @@
+// Copyright Epic Games, Inc. All Rights Reserved.
+
+#include "ComputeFramework/ComputeGraphScheduler.h"
+
+#include "ComputeFramework/ComputeKernel.h"
+#include "ComputeFramework/ComputeKernelShader.h"
+#include "ComputeFramework/ComputeGraph.h"
+#include "ProfilingDebugging/RealtimeGPUProfiler.h"
+#include "RenderGraphBuilder.h"
+#include "RenderGraphUtils.h"
+#include "RHIDefinitions.h"
+
+DECLARE_GPU_STAT_NAMED(ComputeFramework_ExecuteBatches, TEXT("ComputeFramework::ExecuteBatches"));
+
+void FComputeGraphProxy::Initialize(UComputeGraph* ComputeGraph)
+{
+	for (int32 KernelIndex = 0; KernelIndex < ComputeGraph->GetKernelInvocations().Num(); ++KernelIndex)
+	{
+		UComputeKernel* Kernel = ComputeGraph->GetKernelInvocations()[KernelIndex];
+		FComputeKernelResource* KernelResource = ComputeGraph->GetKernelResources()[KernelIndex].Get();
+
+		if (Kernel != nullptr && KernelResource != nullptr)
+		{
+			FKernelInvocation KernelInvocation = {
+				Kernel->GetFName(),
+				FName("InvocationName"),
+				FIntVector(32, 1, 1), // todo: read from kernel
+				KernelResource };
+
+			KernelInvocations.Emplace(KernelInvocation);
+		}
+	}
+}
+
+void FComputeGraphScheduler::EnqueueForExecution(const FComputeGraphProxy* ComputeGraph, TArrayView<FComputeDataProviderRenderProxy* const> ComputeDataProviders)
+{
+	// todo[CF]: Allocate a specific data provider per kernel to drive the number of invocations?
+	int32 FirstProvider = -1;
+	for (int32 ProviderIndex = 0; ProviderIndex < ComputeDataProviders.Num(); ++ProviderIndex)
+	{
+		if (ComputeDataProviders[ProviderIndex] != nullptr)
+		{
+			FirstProvider = ProviderIndex;
+			break;
+		}
+	}
+
+	const int32 SubInvocationCount = FirstProvider != -1 ? ComputeDataProviders[FirstProvider]->GetInvocationCount() : 1;
+
+	for (FComputeGraphProxy::FKernelInvocation const& Invocation : ComputeGraph->KernelInvocations)
+	{
+		// todo[CF]: If you hit this then shader compilation might not have happened yet.
+		if (ensure(Invocation.Kernel->GetShader().IsValid()))
+		{
+			for (int32 SubInvocationIndex = 0; SubInvocationIndex < SubInvocationCount; ++ SubInvocationIndex)
+			{
+				// todo[CF]: dispatch dimension logic needs to be way more involved
+				const FIntVector DispatchDim = FirstProvider != -1 ? ComputeDataProviders[FirstProvider]->GetDispatchDim(SubInvocationIndex, Invocation.GroupDim) : FIntVector(1, 1, 1);
+
+				FShaderInvocation ShaderInvocation = {
+					Invocation.KernelName,
+					Invocation.InvocationName,
+					DispatchDim,
+					Invocation.Kernel->GetShaderParamMetadata(),
+					Invocation.Kernel->GetShader(),
+					Bindings.Num() + SubInvocationIndex
+					};
+
+				ComputeShaders.Emplace(ShaderInvocation);
+			}
+		}
+	}
+
+	// todo[CF]: Move GetBindings() calls to ExecuteBatches() so that RDG buffers can be registered and bound.
+	Bindings.AddDefaulted(SubInvocationCount);
+	for (int32 SubInvocationIndex = 0; SubInvocationIndex < SubInvocationCount; ++SubInvocationIndex)
+	{
+		for (int32 ProviderIndex = 0; ProviderIndex < ComputeDataProviders.Num(); ++ProviderIndex)
+		{
+			if (ComputeDataProviders[ProviderIndex] != nullptr)
+			{
+				TCHAR const* UID = UComputeGraph::GetDataInterfaceUID(ProviderIndex);
+				ComputeDataProviders[ProviderIndex]->GetBindings(SubInvocationIndex, UID, Bindings[SubInvocationIndex]);
+			}
+		}
+	}
+}
+
+struct FComputeExecutionBuffer
+{
+	FComputeExecutionBuffer() = default;
+
+	FComputeExecutionBuffer(FName InName, FRDGBufferRef InBufferRef)
+		: Name(InName)
+		, BufferRef(InBufferRef)
+	{
+	}
+
+	FName Name;
+	FRDGBufferRef BufferRef;
+};
+
+void FComputeGraphScheduler::ExecuteBatches(
+	FRHICommandListImmediate& RHICmdList,
+	ERHIFeatureLevel::Type FeatureLevel
+	)
+{
+	if (ComputeShaders.IsEmpty())
+	{
+		return;
+	}
+
+	FMemMark MemStackMark(FMemStack::Get());
+	{
+		SCOPED_DRAW_EVENTF(RHICmdList, ComputeFramework_ExecuteBatches, TEXT("ComputeFramework::ExecuteBatches"));
+		SCOPED_GPU_STAT(RHICmdList, ComputeFramework_ExecuteBatches);
+
+		FRDGBuilder GraphBuilder(RHICmdList);
+		FGlobalShaderMap* ShaderMap = GetGlobalShaderMap(FeatureLevel);
+
+		TArray<FComputeExecutionBuffer> ExecutionBuffers;
+
+		// todo[CF]: Implement RDG allocations in the data providers. Probably need to move GetBindings() into the RDG block here?
+		/*
+		for (auto& ExtRes : ExternalBuffers)
+		{
+			TCHAR ResName[FComputeResource::MAX_NAME_LENGTH];
+			ExtRes.Name.ToString(ResName);
+
+			ExecutionBuffers.Emplace(
+				ExtRes.Name,
+				GraphBuilder.RegisterExternalBuffer(ExtRes.Buffer, ResName)
+				);
+		}
+
+		for (auto& TransRes : TransientBuffers)
+		{
+			TCHAR ResName[FComputeResource::MAX_NAME_LENGTH];
+			TransRes.Name.ToString(ResName);
+
+			FRDGBufferDesc BufferDesc;
+		
+			EComputeExecutionBufferType BufferType = (EComputeExecutionBufferType)TransRes.BufferType;
+			switch (BufferType)
+			{
+			case EComputeExecutionBufferType::Buffer:
+				BufferDesc = FRDGBufferDesc::CreateBufferDesc(TransRes.BytesPerElement, TransRes.ElementCount);
+				break;
+
+			case EComputeExecutionBufferType::StructuredBuffer:
+				BufferDesc = FRDGBufferDesc::CreateStructuredDesc(TransRes.BytesPerElement, TransRes.ElementCount);
+				break;
+
+			case EComputeExecutionBufferType::ByteAddressBuffer:
+				BufferDesc = FRDGBufferDesc::CreateByteAddressDesc(TransRes.BytesPerElement * TransRes.ElementCount);
+				break;
+
+			default:
+				check(!"Unsupported buffer type");
+				break;
+			}
+
+			ExecutionBuffers.Emplace(
+				TransRes.Name,
+				GraphBuilder.CreateBuffer(BufferDesc, ResName)
+				);
+		}
+		*/
+
+		for (int32 KernelIndex = 0; KernelIndex < ComputeShaders.Num(); KernelIndex++)
+		{
+			FShaderInvocation& Compute = ComputeShaders[KernelIndex];
+			FComputeDataProviderRenderProxy::FBindings& Binding = Bindings[Compute.BindingsIndex];
+
+			FRDGBufferDesc BufferDesc = FRDGBufferDesc::CreateBufferDesc(4, 32);
+
+			void* RawBuffer = GraphBuilder.Alloc(Compute.ShaderParamMetadata->GetSize(), SHADER_PARAMETER_STRUCT_ALIGNMENT);
+			FMemory::Memzero(RawBuffer, Compute.ShaderParamMetadata->GetSize());
+
+			bool bBindFailed = false;
+			uint8* ParamBuffer = static_cast<uint8*>(RawBuffer);
+			TArray<FShaderParametersMetadata::FMember> ParamMembers = Compute.ShaderParamMetadata->GetMembers();
+			for (auto& Member : ParamMembers)
+			{
+				switch (Member.GetBaseType())
+				{
+				case EUniformBufferBaseType::UBMT_INT32:
+					{
+						int32* ParamValue = Binding.ParamsInt.Find(Member.GetName());
+						*reinterpret_cast<int32*>(&ParamBuffer[Member.GetOffset()]) = ParamValue != nullptr ? *ParamValue : 0;
+					}
+					break;
+
+				case EUniformBufferBaseType::UBMT_UINT32:
+					{
+						uint32* ParamValue = Binding.ParamsUint.Find(Member.GetName());
+						*reinterpret_cast<uint32*>(&ParamBuffer[Member.GetOffset()]) = ParamValue != nullptr ? *ParamValue : 0;
+					}
+				break;
+
+				case EUniformBufferBaseType::UBMT_FLOAT32:
+					{
+						float* ParamValue = Binding.ParamsFloat.Find(Member.GetName());
+						*reinterpret_cast<float*>(&ParamBuffer[Member.GetOffset()]) = ParamValue != nullptr ? *ParamValue : 0;
+					}
+					break;
+
+				case EUniformBufferBaseType::UBMT_NESTED_STRUCT:
+					{
+						TArray<uint8>* ParamValue = Binding.Structs.Find(Member.GetName());
+						if (ParamValue != nullptr)
+						{
+							FMemory::Memcpy(&ParamBuffer[Member.GetOffset()], ParamValue->GetData(), Member.GetStructMetadata()->GetSize());
+						}
+						else
+						{
+							bBindFailed = true;
+						}
+					}
+					break;
+
+				default:
+					check(!"Unsupported type");
+					bBindFailed = true;
+					break;
+				};
+			}
+
+			if (ensure(!bBindFailed))
+			{
+				TCHAR KernelName[128];
+				Compute.KernelName.ToString(KernelName);
+
+				TCHAR InvocationName[128];
+				Compute.InvocationName.ToString(InvocationName);
+
+				FComputeShaderUtils::AddPass(
+					GraphBuilder,
+					RDG_EVENT_NAME("Compute[%s]: %s", KernelName, InvocationName),
+					ERDGPassFlags::Compute | ERDGPassFlags::NeverCull,
+					Compute.Shader,
+					Compute.ShaderParamMetadata,
+					static_cast<FComputeKernelShader::FParameters*>(RawBuffer),
+					Compute.DispatchDim
+					);		
+			}
+		}
+
+		GraphBuilder.Execute();
+
+		ComputeShaders.Reset();
+		Bindings.Reset();
+	}
+}
diff --git a/Engine/Source/Runtime/Engine/Private/ComputeFramework/ComputeKernel.cpp b/Engine/Source/Runtime/Engine/Private/ComputeFramework/ComputeKernel.cpp
--- a/Engine/Source/Runtime/Engine/Private/ComputeFramework/ComputeKernel.cpp
+++ b/Engine/Source/Runtime/Engine/Private/ComputeFramework/ComputeKernel.cpp
@@ -1,24 +1,20 @@
 // Copyright Epic Games, Inc. All Rights Reserved.
 
 #include "ComputeFramework/ComputeKernel.h"
-#include "ComputeFramework/ComputeKernelPermutationSet.h"
 #include "ComputeFramework/ComputeKernelSource.h"
+#include "ComputeFramework/ShaderParamTypeDefinition.h"
 
-DEFINE_LOG_CATEGORY(ComputeKernel);
+#if WITH_EDITOR
 
 void UComputeKernel::PostLoad()
 {
 	Super::PostLoad();
-
-#if WITH_EDITOR
-	if (FApp::CanEverRender())
+	if (KernelSource != nullptr)
 	{
-		CacheResourceShadersForRendering();
+		KernelSource->ConditionalPostLoad();
 	}
-#endif
 }
 
-#if WITH_EDITOR
 void UComputeKernel::PostEditChangeChainProperty(FPropertyChangedChainEvent& PropertyChangedEvent)
 {
 	Super::PostEditChangeChainProperty(PropertyChangedEvent);
@@ -29,8 +25,9 @@ void UComputeKernel::PostEditChangeChainProperty(FPropertyChangedChainEvent& Pro
 		return;
 	}
 
-	FName ModifiedPropName = ModifiedProperty->GetFName();
+	bool bNotifyGraphs = false;
 
+	FName ModifiedPropName = ModifiedProperty->GetFName();
 	if (ModifiedPropName == GET_MEMBER_NAME_CHECKED(UComputeKernel, KernelSource))
 	{
 		if (KernelSource)
@@ -43,98 +40,19 @@ void UComputeKernel::PostEditChangeChainProperty(FPropertyChangedChainEvent& Pro
 			PermutationSetOverrides = FComputeKernelPermutationSet();
 			DefinitionsSetOverrides = FComputeKernelDefinitionsSet();
 		}
-		
-		CacheResourceShadersForRendering();
+
+		bNotifyGraphs = true;
 	}
 	else if (ModifiedPropName == GET_MEMBER_NAME_CHECKED(UComputeKernel, PermutationSetOverrides) ||
 		ModifiedPropName == GET_MEMBER_NAME_CHECKED(UComputeKernel, DefinitionsSetOverrides))
 	{
-		CacheResourceShadersForRendering();
+		bNotifyGraphs = true;
 	}
-}
 
-void UComputeKernel::CacheResourceShadersForRendering(
-	uint32 CompilationFlags /*= uint32(EComputeKernelCompilationFlags::ApplyCompletedShaderMapForRendering)*/
-	)
-{
-	if (!KernelSource)
+	if (bNotifyGraphs)
 	{
-		if (KernelResource)
-		{
-			KernelResource->Invalidate();
-			KernelResource = nullptr;
-		}
-
-		return;
+		// todo[CF]: Notify graphs to update when kernel changes.
 	}
-
-	if (!KernelResource)
-	{
-		KernelResource = MakeUnique<FComputeKernelResource>();
-	}
-
-	KernelResource = MakeUnique<FComputeKernelResource>();
-	KernelResource->SetupResource(
-		GMaxRHIFeatureLevel, 
-		KernelSource,  
-		GetName()
-		);
-
-	ERHIFeatureLevel::Type CacheFeatureLevel = GMaxRHIFeatureLevel;
-	const EShaderPlatform ShaderPlatform = GShaderPlatformForFeatureLevel[CacheFeatureLevel];
-
-	UComputeKernel::CacheShadersForResource(ShaderPlatform, nullptr, CompilationFlags | uint32(EComputeKernelCompilationFlags::Force), KernelResource.Get());
 }
 
-void UComputeKernel::CacheShadersForResource(
-	EShaderPlatform ShaderPlatform,
-	const ITargetPlatform* TargetPlatform,
-	uint32 CompilationFlags,
-	FComputeKernelResource* KernelResource
-	)
-{
-	bool bCooking = (CompilationFlags & uint32(EComputeKernelCompilationFlags::IsCooking)) != 0;
-	
-	const bool bIsDefault = (KernelResource->GetKernelFlags() & uint32(EComputeKernelFlags::IsDefaultKernel)) != 0;
-	if (!GIsEditor || GIsAutomationTesting || bIsDefault || bCooking)
-	{
-		CompilationFlags |= uint32(EComputeKernelCompilationFlags::Synchronous);
-	}
-
-	const bool bIsSuccess = KernelResource->CacheShaders(
-		ShaderPlatform, 
-		TargetPlatform, 
-		CompilationFlags & uint32(EComputeKernelCompilationFlags::ApplyCompletedShaderMapForRendering),
-		CompilationFlags & uint32(EComputeKernelCompilationFlags::Synchronous)
-		);
-
-	if (!bIsSuccess)
-	{
-		if (bIsDefault)
-		{
-			UE_LOG(
-				ComputeKernel, 
-				Fatal, 
-				TEXT("Failed to compile default FComputeKernelResource [%s] for platform [%s]!"),
-				*KernelResource->GetFriendlyName(),
-				*LegacyShaderPlatformToShaderFormat(ShaderPlatform).ToString()
-				);
-		}
-
-		UE_LOG(
-			ComputeKernel, 
-			Warning, 
-			TEXT("Failed to compile FComputeKernelResource [%s] for platform [%s]."),
-			*KernelResource->GetFriendlyName(),
-			*LegacyShaderPlatformToShaderFormat(ShaderPlatform).ToString()
-			);
-
-		auto& CompilationErrors = KernelResource->GetCompileErrors();
-		uint32 ErrorCount = CompilationErrors.Num();
-		for (uint32 i = 0; i < ErrorCount; ++i)
-		{
-			UE_LOG(ComputeKernel, Warning, TEXT("      [Error] - %s"), *CompilationErrors[i]);
-		}
-	}
-}
-#endif
+#endif // WITH_EDITOR
diff --git a/Engine/Source/Runtime/Engine/Private/ComputeFramework/ComputeKernelDerivedDataVersion.h b/Engine/Source/Runtime/Engine/Private/ComputeFramework/ComputeKernelDerivedDataVersion.h
--- a/Engine/Source/Runtime/Engine/Private/ComputeFramework/ComputeKernelDerivedDataVersion.h
+++ b/Engine/Source/Runtime/Engine/Private/ComputeFramework/ComputeKernelDerivedDataVersion.h
@@ -10,4 +10,4 @@ ComputeKernelShaderDerivedDataVersion.h: Derived data version for ComputeKernel
 
 // In case of merge conflicts with DDC versions, you MUST generate a new GUID and set this new
 // guid as version
-#define COMPUTEKERNEL_DERIVEDDATA_VER		TEXT("EDA5D304-9967-482C-9C75-80A5F6BEC61C")
+#define COMPUTEKERNEL_DERIVEDDATA_VER		TEXT("BF0A9C9A-1BD3-4FC8-A74B-156FCD9AEF74")
diff --git a/Engine/Source/Runtime/Engine/Private/ComputeFramework/ComputeKernelFromText.cpp b/Engine/Source/Runtime/Engine/Private/ComputeFramework/ComputeKernelFromText.cpp
--- a/Engine/Source/Runtime/Engine/Private/ComputeFramework/ComputeKernelFromText.cpp
+++ b/Engine/Source/Runtime/Engine/Private/ComputeFramework/ComputeKernelFromText.cpp
@@ -2,25 +2,25 @@
 
 #include "ComputeFramework/ComputeKernelFromText.h"
 #include "ComputeFramework/ComputeKernel.h"
+#include "ComputeFramework/ComputeFramework.h"
+#include "Internationalization/Regex.h"
 #include "Misc/FileHelper.h"
+#include "Misc/Paths.h"
 #include "Misc/ScopeExit.h"
 
-#include "Internationalization/Regex.h"
-
 UComputeKernelFromText::UComputeKernelFromText()
 {
 	UniqueId = FGuid::NewGuid();
 }
 
+#if WITH_EDITOR
+
 void UComputeKernelFromText::PostLoad()
 {
 	Super::PostLoad();
-#if WITH_EDITOR
 	ReparseKernelSourceText();
-#endif
 }
 
-#if WITH_EDITOR
 void UComputeKernelFromText::PreEditChange(FEditPropertyChain& PropertyAboutToChange)
 {
 	Super::PreEditChange(PropertyAboutToChange);
@@ -47,12 +47,6 @@ void UComputeKernelFromText::PostEditChangeChainProperty(FPropertyChangedChainEv
 
 void UComputeKernelFromText::ReparseKernelSourceText()
 {
-	ON_SCOPE_EXIT
-	{
-		//FriendlyName = FPaths::GetCleanFilename(SourceFile.FilePath);
-		//Name = (FriendlyName.IsEmpty()) ? TEXT("Unset"_SV) : FriendlyName;
-	};
-
 	if (SourceFile.FilePath.IsEmpty())
 	{
 		EntryPointName = FString();
@@ -60,8 +54,6 @@ void UComputeKernelFromText::ReparseKernelSourceText()
 		PermutationSet = FComputeKernelPermutationSet();
 		DefinitionsSet = FComputeKernelDefinitionsSet();
 		InputParams.Empty();
-		InputSRVs.Empty();
-		Outputs.Empty();
 
 		return;
 	}
@@ -71,7 +63,7 @@ void UComputeKernelFromText::ReparseKernelSourceText()
 	IPlatformFile& PlatformFileSystem = IPlatformFile::GetPlatformPhysical();
 	if (!PlatformFileSystem.FileExists(*FullKernelPath))
 	{
-		UE_LOG(ComputeKernel, Error, TEXT("Unable to find kernel file \"%s\""), *FullKernelPath);
+		UE_LOG(LogComputeFramework, Error, TEXT("Unable to find kernel file \"%s\""), *FullKernelPath);
 
 		SourceFile = PrevSourceFile;
 		return;
@@ -79,7 +71,7 @@ void UComputeKernelFromText::ReparseKernelSourceText()
 
 	if (!FFileHelper::LoadFileToString(KernelSourceText, &PlatformFileSystem, *FullKernelPath))
 	{
-		UE_LOG(ComputeKernel, Error, TEXT("Unable to read kernel file \"%s\""), *FullKernelPath);
+		UE_LOG(LogComputeFramework, Error, TEXT("Unable to read kernel file \"%s\""), *FullKernelPath);
 
 		SourceFile = PrevSourceFile;
 		return;
@@ -150,8 +142,8 @@ void UComputeKernelFromText::ReparseKernelSourceText()
 
 			case EShaderFundamentalDimensionType::Matrix:
 				FIntVector2 MtxDim = FShaderParamTypeDefinition::ParseMatrixDimension(DimensionType);
-				Param.MatrixColumnCount	= MtxDim.X;
-				Param.MatrixRowCount	= MtxDim.Y;
+				Param.MatrixRowCount = MtxDim.X;
+				Param.MatrixColumnCount = MtxDim.Y;
 				break;
 			}
 
@@ -162,107 +154,121 @@ void UComputeKernelFromText::ReparseKernelSourceText()
 		InputParams = MoveTemp(Params);
 	}
 
-	const FRegexPattern KernelInputSRVPattern(TEXT(R"(KERNEL_SRV\(\s*(ByteAddress|Structured|Buffer|Texture1D|Texture2D|Texture3D|TextureCube)\s*<\s*(?:SNORM|UNORM)?\s*(int|uint|float)((?:[1-4])|(?:[1-4]x[1-4]))?\s*>\s*,\s*([a-zA-Z_]\w*)\s*\))"));
+	const FRegexPattern KernelReadExternPattern(TEXT(R"(KERNEL_EXTERN_READ\(\s*([a-zA-Z_]\w*)((\s*,\s*(?:bool|int|uint|float)(?:(?:[1-4]x[1-4])|(?:[1-4])|))+)\s*\))"));
 	{
-		TArray<FShaderParamTypeDefinition> SRVs;
+		TArray<FShaderFunctionDefinition> ExternalReadFunctions;
 
-		FRegexMatcher Matcher(KernelInputSRVPattern, KernelSourceText);
+		FRegexMatcher Matcher(KernelReadExternPattern, KernelSourceText);
 		while (Matcher.FindNext())
 		{
-			FShaderParamTypeDefinition Param = {};
+			FShaderFunctionDefinition FunctionDesc;
+			FunctionDesc.Name = Matcher.GetCaptureGroup(1);
+			FunctionDesc.bHasReturnType = true;
+
+			FString AllParameters = Matcher.GetCaptureGroup(2);
+			const FRegexPattern ParameterPattern(TEXT(R"((bool|int|uint|float|half)((?:[1-4]x[1-4])|(?:[1-4])|))"));
+			FRegexMatcher ParameterMatcher(ParameterPattern, *AllParameters);
+			
+			while (ParameterMatcher.FindNext())
+			{
+				FShaderParamTypeDefinition Param = {};
 
-			FString ResourceType	= Matcher.GetCaptureGroup(1);
-			FString FundamentalType = Matcher.GetCaptureGroup(2);
-			FString DimensionType	= Matcher.GetCaptureGroup(3);
-			FString ParamName		= Matcher.GetCaptureGroup(4);
+				FString FundamentalType = ParameterMatcher.GetCaptureGroup(1);
+				FString DimensionType = ParameterMatcher.GetCaptureGroup(2);
 
-			Param.Name				= MoveTemp(ParamName);
-			Param.FundamentalType	= FShaderParamTypeDefinition::ParseFundamental(FundamentalType);
-			Param.DimType			= FShaderParamTypeDefinition::ParseDimension(DimensionType);
-			Param.BindingType		= EShaderParamBindingType::ReadOnlyResource;
-			Param.ResourceType		= FShaderParamTypeDefinition::ParseResource(ResourceType);
+				Param.FundamentalType = FShaderParamTypeDefinition::ParseFundamental(FundamentalType);
+				Param.DimType = FShaderParamTypeDefinition::ParseDimension(DimensionType);
 
-			switch (Param.DimType)
-			{
-			case EShaderFundamentalDimensionType::Scalar:
-				Param.VectorDimension = 0;
-				break;
+				switch (Param.DimType)
+				{
+				case EShaderFundamentalDimensionType::Scalar:
+					Param.VectorDimension = 0;
+					break;
 
-			case EShaderFundamentalDimensionType::Vector:
-				Param.VectorDimension = FShaderParamTypeDefinition::ParseVectorDimension(DimensionType);
-				break;
+				case EShaderFundamentalDimensionType::Vector:
+					Param.VectorDimension = FShaderParamTypeDefinition::ParseVectorDimension(DimensionType);
+					break;
 
-			case EShaderFundamentalDimensionType::Matrix:
-				FIntVector2 MtxDim = FShaderParamTypeDefinition::ParseMatrixDimension(DimensionType);
-				Param.MatrixColumnCount = MtxDim.X;
-				Param.MatrixRowCount	= MtxDim.Y;
-				break;
+				case EShaderFundamentalDimensionType::Matrix:
+					FIntVector2 MtxDim = FShaderParamTypeDefinition::ParseMatrixDimension(DimensionType);
+					Param.MatrixRowCount = MtxDim.X;
+					Param.MatrixColumnCount = MtxDim.Y;
+					break;
+				}
+
+				Param.ResetTypeDeclaration();
+				FunctionDesc.ParamTypes.Emplace(MoveTemp(Param));
 			}
 
-			Param.ResetTypeDeclaration();
-			SRVs.Emplace(MoveTemp(Param));
+			ExternalReadFunctions.Emplace(MoveTemp(FunctionDesc));
 		}
 
-		InputSRVs = MoveTemp(SRVs);
+		ExternalInputs = MoveTemp(ExternalReadFunctions);
 	}
 
-	const FRegexPattern KernelOutputPattern(TEXT(R"(KERNEL_UAV\(\s*RW(ByteAddress|Structured|Buffer|Texture1D|Texture2D|Texture3D|TextureCube)\s*<\s*(?:SNORM|UNORM)?\s*(int|uint|float)((?:[1-4])|(?:[1-4]x[1-4]))?\s*>\s*,\s*([a-zA-Z_]\w*)\s*\))"));
+	const FRegexPattern KernelWriteExternPattern(TEXT(R"(KERNEL_EXTERN_WRITE\(\s*([a-zA-Z_]\w*)((\s*,\s*(?:bool|int|uint|float)(?:(?:[1-4]x[1-4])|(?:[1-4])|))+)\s*\))"));
 	{
-		TArray<FShaderParamTypeDefinition> OutputParams;
+		TArray<FShaderFunctionDefinition> ExternalWriteFunctions;
 
-		FRegexMatcher Matcher(KernelOutputPattern, KernelSourceText);
+		FRegexMatcher Matcher(KernelWriteExternPattern, KernelSourceText);
 		while (Matcher.FindNext())
 		{
-			FShaderParamTypeDefinition Param = {};
-
-			FString ResourceType	= Matcher.GetCaptureGroup(1);
-			FString FundamentalType = Matcher.GetCaptureGroup(2);
-			FString DimensionType	= Matcher.GetCaptureGroup(3);
-			FString ParamName		= Matcher.GetCaptureGroup(4);
+			FShaderFunctionDefinition FunctionDesc;
+			FunctionDesc.Name = Matcher.GetCaptureGroup(1);
+			FunctionDesc.bHasReturnType = false;
 
-			Param.Name				= MoveTemp(ParamName);
-			Param.FundamentalType	= FShaderParamTypeDefinition::ParseFundamental(FundamentalType);
-			Param.DimType			= FShaderParamTypeDefinition::ParseDimension(DimensionType);
-			Param.BindingType		= EShaderParamBindingType::ReadWriteResource;
-			Param.ResourceType		= FShaderParamTypeDefinition::ParseResource(ResourceType);
+			FString AllParameters = Matcher.GetCaptureGroup(2);
+			const FRegexPattern ParameterPattern(TEXT(R"((bool|int|uint|float|half)((?:[1-4]x[1-4])|(?:[1-4])|))"));
+			FRegexMatcher ParameterMatcher(ParameterPattern, *AllParameters);
 
-			switch (Param.DimType)
+			while (ParameterMatcher.FindNext())
 			{
-			case EShaderFundamentalDimensionType::Scalar:
-				Param.VectorDimension = 0;
-				break;
+				FShaderParamTypeDefinition Param = {};
 
-			case EShaderFundamentalDimensionType::Vector:
-				Param.VectorDimension = FShaderParamTypeDefinition::ParseVectorDimension(DimensionType);
-				break;
+				FString FundamentalType = ParameterMatcher.GetCaptureGroup(1);
+				FString DimensionType = ParameterMatcher.GetCaptureGroup(2);
 
-			case EShaderFundamentalDimensionType::Matrix:
-				FIntVector2 MtxDim = FShaderParamTypeDefinition::ParseMatrixDimension(DimensionType);
-				Param.MatrixColumnCount = MtxDim.X;
-				Param.MatrixRowCount	= MtxDim.Y;
-				break;
+				Param.FundamentalType = FShaderParamTypeDefinition::ParseFundamental(FundamentalType);
+				Param.DimType = FShaderParamTypeDefinition::ParseDimension(DimensionType);
+
+				switch (Param.DimType)
+				{
+				case EShaderFundamentalDimensionType::Scalar:
+					Param.VectorDimension = 0;
+					break;
+
+				case EShaderFundamentalDimensionType::Vector:
+					Param.VectorDimension = FShaderParamTypeDefinition::ParseVectorDimension(DimensionType);
+					break;
+
+				case EShaderFundamentalDimensionType::Matrix:
+					FIntVector2 MtxDim = FShaderParamTypeDefinition::ParseMatrixDimension(DimensionType);
+					Param.MatrixRowCount = MtxDim.X;
+					Param.MatrixColumnCount = MtxDim.Y;
+					break;
+				}
+
+				Param.ResetTypeDeclaration();
+				FunctionDesc.ParamTypes.Emplace(MoveTemp(Param));
 			}
 
-			Param.ResetTypeDeclaration();
-			OutputParams.Emplace(MoveTemp(Param));
+			ExternalWriteFunctions.Emplace(MoveTemp(FunctionDesc));
 		}
 
-		Outputs = MoveTemp(OutputParams);
+		ExternalOutputs = MoveTemp(ExternalWriteFunctions);
 	}
 
-	/*
 	uint64 NewHash = 0;
 
-	NewHash = Hash(NewHash, UniqueId);
-	NewHash = Hash(NewHash, *SourceFile.FilePath);
-	NewHash = Hash(NewHash, *KernelSource);
+	NewHash = FCrc::TypeCrc32(UniqueId, NewHash);
+	NewHash = FCrc::TypeCrc32(*SourceFile.FilePath, NewHash);
+	NewHash = FCrc::TypeCrc32(*KernelSourceText, NewHash);
 
 	if (SourceHash != NewHash)
 	{
 		SourceHash = NewHash;
-		// Trigger recompile
+		// todo[CF]: Nofify graphs for recompilation
 	}
-	*/
 }
 
 #endif
diff --git a/Engine/Source/Runtime/Engine/Private/ComputeFramework/ComputeKernelShader.cpp b/Engine/Source/Runtime/Engine/Private/ComputeFramework/ComputeKernelShader.cpp
--- a/Engine/Source/Runtime/Engine/Private/ComputeFramework/ComputeKernelShader.cpp
+++ b/Engine/Source/Runtime/Engine/Private/ComputeFramework/ComputeKernelShader.cpp
@@ -2,25 +2,11 @@
 
 #include "ComputeFramework/ComputeKernelShader.h"
 
-#include "Engine/VolumeTexture.h"
-#include "ComputeKernelDerivedDataVersion.h"
-#include "ComputeFramework/ComputeKernelShaderCompilationManager.h"
-#include "ComputeFramework/ComputeKernelShared.h"
-#include "ProfilingDebugging/CookStats.h"
-#include "ProfilingDebugging/DiagnosticTable.h"
-#include "Serialization/MemoryReader.h"
-#include "Serialization/MemoryWriter.h"
-#include "ShaderCompiler.h"
-#include "ShaderParameterUtils.h"
-#include "Stats/StatsMisc.h"
-#include "UObject/UObjectGlobals.h"
+#include "ComputeFramework/ComputeKernelShaderType.h"
 
 IMPLEMENT_SHADER_TYPE(, FComputeKernelShader, TEXT("/Engine/Private/ComputeKernel.usf"), TEXT("__"), SF_Compute)
 
-
-FComputeKernelShader::FComputeKernelShader(
-	const FComputeKernelShaderType::CompiledShaderInitializerType& Initializer
-	)
+FComputeKernelShader::FComputeKernelShader(const FComputeKernelShaderType::CompiledShaderInitializerType& Initializer)
 	: FShader(Initializer)
 {
 	const FShaderParametersMetadata& ShaderParametersMetadata = 
diff --git a/Engine/Source/Runtime/Engine/Private/ComputeFramework/ComputeKernelShader.h b/Engine/Source/Runtime/Engine/Private/ComputeFramework/ComputeKernelShader.h
new file mode 100644
--- /dev/null
+++ b/Engine/Source/Runtime/Engine/Private/ComputeFramework/ComputeKernelShader.h
@@ -0,0 +1,28 @@
+// Copyright Epic Games, Inc. All Rights Reserved.
+
+#pragma once
+
+#include "GlobalShader.h"
+#include "ComputeKernelShaderType.h"
+#include "ComputeKernelShared.h"
+#include "SceneView.h"
+#include "Shader.h"
+#include "ShaderParameters.h"
+
+class FTextureResource;
+class UClass;
+
+class FComputeKernelShader : public FShader
+{
+	DECLARE_SHADER_TYPE(FComputeKernelShader, ComputeKernel);
+
+public:
+	FComputeKernelShader() = default;
+	FComputeKernelShader(
+		const FComputeKernelShaderType::CompiledShaderInitializerType & Initializer
+		);
+
+	MS_ALIGN(SHADER_PARAMETER_STRUCT_ALIGNMENT) struct FParameters
+	{
+	};
+};
diff --git a/Engine/Source/Runtime/Engine/Private/ComputeFramework/ComputeKernelShaderCompilationManager.h b/Engine/Source/Runtime/Engine/Private/ComputeFramework/ComputeKernelShaderCompilationManager.h
new file mode 100644
--- /dev/null
+++ b/Engine/Source/Runtime/Engine/Private/ComputeFramework/ComputeKernelShaderCompilationManager.h
@@ -0,0 +1,125 @@
+// Copyright Epic Games, Inc. All Rights Reserved.
+
+#pragma once
+
+#include "ComputeKernelShared.h"
+#include "ShaderCompiler.h"
+
+#if WITH_EDITOR
+#include "TickableEditorObject.h"
+#endif
+
+
+/** Information tracked for each shader compile worker process instance. */
+struct FComputeKernelShaderCompileWorkerInfo
+{
+	/** Process handle of the worker app once launched.  Invalid handle means no process. */
+	FProcHandle WorkerProcess;
+
+	/** Tracks whether tasks have been issued to the worker. */
+	bool bIssuedTasksToWorker;
+
+	/** Whether the worker has been launched for this set of tasks. */
+	bool bLaunchedWorker;
+
+	/** Tracks whether all tasks issued to the worker have been received. */
+	bool bComplete;
+
+	/** Time at which the worker started the most recent batch of tasks. */
+	double StartTime;
+
+	/** Jobs that this worker is responsible for compiling. */
+	TArray<FShaderCommonCompileJobPtr> QueuedJobs;
+
+	FComputeKernelShaderCompileWorkerInfo() :
+		bIssuedTasksToWorker(false),
+		bLaunchedWorker(false),
+		bComplete(false),
+		StartTime(0)
+	{
+	}
+
+	// warning: not virtual
+	~FComputeKernelShaderCompileWorkerInfo()
+	{
+		if (WorkerProcess.IsValid())
+		{
+			FPlatformProcess::TerminateProc(WorkerProcess);
+			FPlatformProcess::CloseProc(WorkerProcess);
+		}
+	}
+};
+
+
+/** Results for a single compiled shader map. */
+struct FComputeKernelShaderMapCompileResults
+{
+	FComputeKernelShaderMapCompileResults() :
+		NumJobsQueued(0),
+		bAllJobsSucceeded(true),
+		bRecreateComponentRenderStateOnCompletion(false)
+	{}
+
+	int32 NumJobsQueued;
+	bool bAllJobsSucceeded;
+	bool bRecreateComponentRenderStateOnCompletion;
+	TArray<FShaderCommonCompileJobPtr> FinishedJobs;
+};
+
+
+/** Results for a single compiled and finalized shader map. */
+struct FComputeKernelShaderMapFinalizeResults : public FComputeKernelShaderMapCompileResults
+{
+	/** Tracks finalization progress on this shader map. */
+	int32 FinalizeJobIndex;
+
+	FComputeKernelShaderMapFinalizeResults(const FComputeKernelShaderMapCompileResults& InCompileResults) :
+		FComputeKernelShaderMapCompileResults(InCompileResults),
+		FinalizeJobIndex(0)
+	{}
+};
+
+
+#if WITH_EDITOR
+// handles finished shader compile jobs, applying of the shaders to their config asset, and some error handling
+class FComputeKernelShaderCompilationManager : public FTickableEditorObject
+{
+public:
+	FComputeKernelShaderCompilationManager();
+	~FComputeKernelShaderCompilationManager();
+
+	virtual ETickableTickType GetTickableTickType() const override
+	{
+		return ETickableTickType::Always;
+	}
+
+	virtual TStatId GetStatId() const
+	{
+		RETURN_QUICK_DECLARE_CYCLE_STAT(FComputeKernelShaderCompilationManager, STATGROUP_Tickables);
+	}
+
+	void Tick(float DeltaSeconds) override;
+	void AddJobs(TArray<FShaderCommonCompileJobPtr> InNewJobs);
+	void ProcessAsyncResults();
+
+	void FinishCompilation(const TCHAR* InKernelName, const TArray<int32>& ShaderMapIdsToFinishCompiling);
+
+private:
+	void ProcessCompiledComputeKernelShaderMaps(TMap<int32, FComputeKernelShaderMapFinalizeResults>& CompiledShaderMaps, float TimeBudget);
+	void RunCompileJobs();
+
+	void InitWorkerInfo();
+
+	TArray<FShaderCommonCompileJobPtr> JobQueue;
+
+	/** Map from shader map Id to the compile results for that map, used to gather compiled results. */
+	TMap<int32, FComputeKernelShaderMapCompileResults> ComputeKernelShaderMapJobs;
+
+	/** Map from shader map id to results being finalized.  Used to track shader finalizations over multiple frames. */
+	TMap<int32, FComputeKernelShaderMapFinalizeResults> PendingFinalizeComputeKernelShaderMaps;
+
+	TArray<struct FComputeKernelShaderCompileWorkerInfo*> WorkerInfos;
+};
+
+extern FComputeKernelShaderCompilationManager GComputeKernelShaderCompilationManager;
+#endif
diff --git a/Engine/Source/Runtime/Engine/Private/ComputeFramework/ComputeKernelShaderType.h b/Engine/Source/Runtime/Engine/Private/ComputeFramework/ComputeKernelShaderType.h
new file mode 100644
--- /dev/null
+++ b/Engine/Source/Runtime/Engine/Private/ComputeFramework/ComputeKernelShaderType.h
@@ -0,0 +1,137 @@
+// Copyright Epic Games, Inc. All Rights Reserved.
+
+#pragma once
+
+#include "CoreMinimal.h"
+#include "Engine/EngineTypes.h"
+#include "GlobalShader.h"
+#include "Shader.h"
+#include "ShaderCompiler.h"
+
+class FShaderCommonCompileJob;
+class FShaderCompileJob;
+class FComputeKernelResource;
+
+
+struct FComputeKernelShaderPermutationParameters : public FShaderPermutationParameters
+{
+	FComputeKernelShaderPermutationParameters(EShaderPlatform InPlatform)
+		: FShaderPermutationParameters(InPlatform)
+	{
+	}
+};
+
+class FComputeKernelShaderType : public FShaderType
+{
+public:
+	struct FParameters : public FShaderType::FParameters
+	{
+		const FShaderParametersMetadata& ShaderParamMetadata;
+
+		FParameters(
+			const FShaderParametersMetadata& InShaderParamMetadata
+			)
+			: ShaderParamMetadata(InShaderParamMetadata)
+		{
+		}
+	};
+
+	struct CompiledShaderInitializerType : FGlobalShaderType::CompiledShaderInitializerType
+	{
+		const FString DebugDescription;
+
+		CompiledShaderInitializerType(
+			const FShaderType* InType,
+			const FParameters* InParameters,
+			int32 InPermutationId,
+			const FShaderCompilerOutput& CompilerOutput,
+			const FSHAHash& InComputeKernelShaderMapHash,
+			const FString& InDebugDescription
+			)
+			: FGlobalShaderType::CompiledShaderInitializerType(
+				InType,
+				InParameters,
+				InPermutationId,
+				CompilerOutput, 
+				InComputeKernelShaderMapHash,
+				nullptr,
+				nullptr
+				)
+			, DebugDescription(InDebugDescription)
+		{
+		}
+	};
+
+	FComputeKernelShaderType(
+		FTypeLayoutDesc& InTypeLayout,
+		const TCHAR* InName,
+		const TCHAR* InSourceFilename,
+		const TCHAR* InFunctionName,
+		uint32 InFrequency,					// ugly - ignored for but needed for IMPLEMENT_SHADER_TYPE macro magic
+		int32 InTotalPermutationCount,
+		ConstructSerializedType InConstructSerializedRef,
+		ConstructCompiledType InConstructCompiledRef,
+		ModifyCompilationEnvironmentType InModifyCompilationEnvironmentRef,
+		ShouldCompilePermutationType InShouldCompilePermutationRef,
+		ValidateCompiledResultType InValidateCompiledResultRef,
+		uint32 InTypeSize,
+		const FShaderParametersMetadata* InRootParametersMetadata = nullptr
+		):
+		FShaderType(
+			EShaderTypeForDynamicCast::ComputeKernel, 
+			InTypeLayout, 
+			InName, 
+			InSourceFilename, 
+			InFunctionName, 
+			SF_Compute, 
+			InTotalPermutationCount,
+			InConstructSerializedRef,
+			InConstructCompiledRef,
+			InModifyCompilationEnvironmentRef,
+			InShouldCompilePermutationRef,
+			InValidateCompiledResultRef,
+			InTypeSize,
+			InRootParametersMetadata
+			)
+	{
+	}
+
+	void BeginCompileShader(
+		uint32 ShaderMapId,
+		const FComputeKernelResource* InKernel,
+		FSharedShaderCompilerEnvironment* InCompilationEnvironment,
+		EShaderPlatform Platform,
+		TArray<FShaderCommonCompileJobPtr>& InOutNewJobs,
+		FShaderTarget Target
+		);
+
+	/**
+	 * Either creates a new instance of this type or returns an equivalent existing shader.
+	 */
+	FShader* FinishCompileShader(
+		const FSHAHash& InComputeKernelShaderMapHash,
+		const FShaderCompileJob& CurrentJob,
+		const FString& InDebugDescription
+		) const;
+
+	/**
+	 * Checks if the shader type should be cached.
+	 * @return True if this shader type should be cached.
+	 */
+	bool ShouldCache(EShaderPlatform InPlatform, const FComputeKernelResource* Kernel) const
+	{
+		return ShouldCompilePermutation(FComputeKernelShaderPermutationParameters(InPlatform));
+	}
+
+
+protected:
+	/**
+	 * Sets up the environment used to compile an instance of this shader type.
+	 * @param InPlatform - Platform to compile for.
+	 * @param OutEnvironment - The shader compile environment that the function modifies.
+	 */
+	void SetupCompileEnvironment(EShaderPlatform InPlatform, const FComputeKernelResource* Kernel, FShaderCompilerEnvironment& OutEnvironment) const
+	{
+		ModifyCompilationEnvironment(FComputeKernelShaderPermutationParameters(InPlatform), OutEnvironment);
+	}
+};
diff --git a/Engine/Source/Runtime/Engine/Private/ComputeFramework/ComputeKernelShared.cpp b/Engine/Source/Runtime/Engine/Private/ComputeFramework/ComputeKernelShared.cpp
--- a/Engine/Source/Runtime/Engine/Private/ComputeFramework/ComputeKernelShared.cpp
+++ b/Engine/Source/Runtime/Engine/Private/ComputeFramework/ComputeKernelShared.cpp
@@ -182,7 +182,8 @@ void FComputeKernelResource::SerializeShaderMap(FArchive& Ar)
 				TRefCountPtr<FComputeKernelShaderMap> LoadedShaderMap = new FComputeKernelShaderMap();
 				bool bSuccessfullyLoaded = LoadedShaderMap->Serialize(Ar);
 
-				// Toss the loaded shader data if this is a server only instance (@todo - don't cook it in the first place) or if it's for a different RHI than the current one
+				// Toss the loaded shader data if this is a server only instance or if it's for a different RHI than the current one.
+				// todo[CF] Don't cook it in the first place
 				if (bSuccessfullyLoaded && FApp::CanEverRender())
 				{
 					GameThreadShaderMap = RenderingThreadShaderMap = LoadedShaderMap;
@@ -194,58 +195,19 @@ void FComputeKernelResource::SerializeShaderMap(FArchive& Ar)
 
 void FComputeKernelResource::SetupResource(
 	ERHIFeatureLevel::Type InFeatureLevel,
-	const UComputeKernelSource* Source,
-	FString InFriendlyName
+	FString const& InFriendlyName,
+	FString const& InShaderEntryPoint,
+	FString InShaderSource,
+	uint64 InShaderCodeHash,
+	FShaderParametersMetadata* InShaderMetadata
 	)
 {
 	FeatureLevel = InFeatureLevel;
-	ShaderCodeHash = Source->GetSourceHashCode();
-	ShaderEntryPoint = Source->GetEntryPoint();
-	ShaderSource = Source->GetSource();
-	FriendlyName = MoveTemp(InFriendlyName);
-
-	FShaderParametersMetadataBuilder Builder;
-
-	for (auto& Input : Source->InputParams)
-	{
-		ensureAlways(Input.DimType == EShaderFundamentalDimensionType::Scalar);
-
-		switch (Input.FundamentalType)
-		{
-			case EShaderFundamentalType::Bool:
-				Builder.AddParam<bool>(*Input.Name);
-				break;
-
-			case EShaderFundamentalType::Int:
-				Builder.AddParam<int32>(*Input.Name);
-				break;
-
-			case EShaderFundamentalType::Uint:
-				Builder.AddParam<uint32>(*Input.Name);
-				break;
-
-			case EShaderFundamentalType::Float:
-				Builder.AddParam<float>(*Input.Name);
-				break;
-		}
-	}
-
-	for(auto& Input : Source->InputSRVs)
-	{
-		Builder.AddRDGBufferSRV(*Input.Name, *Input.TypeDeclaration);
-	}
-
-	for (auto& Output : Source->Outputs)
-	{
-		Builder.AddRDGBufferUAV(*Output.Name, *Output.TypeDeclaration);
-	}
-
-	ShaderMetadata.Reset(
-		Builder.Build(
-			FShaderParametersMetadata::EUseCase::ShaderParameterStruct, 
-			*FriendlyName
-			)
-		);
+	FriendlyName = InFriendlyName;
+	ShaderEntryPoint = InShaderEntryPoint;
+	ShaderSource = MoveTemp(InShaderSource);
+	ShaderCodeHash = InShaderCodeHash;
+	ShaderMetadata.Reset(InShaderMetadata);
 }
 
 void FComputeKernelResource::SetRenderingThreadShaderMap(FComputeKernelShaderMap* InShaderMap)
diff --git a/Engine/Source/Runtime/Engine/Private/ComputeFramework/ComputeKernelShared.h b/Engine/Source/Runtime/Engine/Private/ComputeFramework/ComputeKernelShared.h
new file mode 100644
--- /dev/null
+++ b/Engine/Source/Runtime/Engine/Private/ComputeFramework/ComputeKernelShared.h
@@ -0,0 +1,526 @@
+// Copyright Epic Games, Inc. All Rights Reserved.
+
+#pragma once
+
+#include "CoreMinimal.h"
+#include "Engine/EngineTypes.h"
+#include "Misc/Guid.h"
+#include "Misc/Optional.h"
+#include "Misc/SecureHash.h"
+#include "RenderResource.h"
+#include "RenderingThread.h"
+#include "RHI.h"
+#include "SceneTypes.h"
+#include "Shader.h"
+#include "StaticParameterSet.h"
+#include "Templates/RefCounting.h"
+#include "ShaderCompiler.h"
+
+
+class FComputeKernelResource;
+class FComputeKernelShaderMap;
+class FComputeKernelShader;
+class FComputeKernelShaderMapId;
+class UComputeKernelSource;
+class FShaderParametersMetadata;
+
+
+/** Stores outputs from the  kernel compile that need to be saved. */
+class FComputeKernelCompilationOutput
+{
+	DECLARE_TYPE_LAYOUT(FComputeKernelCompilationOutput, NonVirtual);
+public:
+	FComputeKernelCompilationOutput()
+	{}
+};
+
+/** Contains all the information needed to uniquely identify a FComputeKernelShaderMapID. */
+class FComputeKernelShaderMapId
+{
+	DECLARE_TYPE_LAYOUT(FComputeKernelShaderMapId, NonVirtual);
+public:
+	/** Feature level that the shader map is going to be compiled for.  */
+	LAYOUT_FIELD(ERHIFeatureLevel::Type, FeatureLevel);
+
+	/**
+	 * The hash computed from the shader code generated by the library.
+	 */
+	LAYOUT_FIELD(uint64, ShaderCodeHash);
+	
+	/** Shader types of shaders that are inlined in this shader map in the DDC. */
+	LAYOUT_FIELD(TMemoryImageArray<FShaderTypeDependency>, ShaderTypeDependencies);
+
+	/*
+	* Type layout parameters of the memory image
+	*/
+	LAYOUT_FIELD(FPlatformTypeLayoutParameters, LayoutParams);
+
+public:
+	FComputeKernelShaderMapId()
+		: FeatureLevel(ERHIFeatureLevel::SM5)
+	{ }
+
+	~FComputeKernelShaderMapId() = default;
+
+	void SetShaderDependencies(const TArray<FShaderType*>& InShaderTypes, EShaderPlatform InShaderPlatform);
+
+	//void Serialize(FArchive& Ar);
+
+	friend uint32 GetTypeHash(const FComputeKernelShaderMapId& Ref)
+	{
+		return GetTypeHash(Ref.ShaderCodeHash);
+	}
+
+	SIZE_T GetSizeBytes() const
+	{
+		return sizeof(*this) + ShaderTypeDependencies.GetAllocatedSize();
+	}
+
+	/** Hashes the kernel specific part of this shader map Id. */
+	void GetComputeKernelHash(FSHAHash& OutHash) const;
+
+	/**
+	 * Tests this set against another for equality, disregarding override settings.
+	 *
+	 * @param InReferenceSet	The set to compare against
+	 * @return					true if the sets are equal
+	 */
+	bool operator==(const FComputeKernelShaderMapId& InReferenceSet) const;
+
+	bool operator!=(const FComputeKernelShaderMapId& InReferenceSet) const
+	{
+		return !(*this == InReferenceSet);
+	}
+
+	/** Appends string representations of this Id to a key string. */
+	void AppendKeyString(FString& OutKeyString) const;
+
+	/** Returns true if the requested shader type is a dependency of this shader map Id. */
+	bool ContainsShaderType(const FShaderType* ShaderType) const;
+};
+
+class FComputeKernelShaderMapContent : public FShaderMapContent
+{
+	using Super = FShaderMapContent;
+	friend class FComputeKernelShaderMap;
+	DECLARE_TYPE_LAYOUT(FComputeKernelShaderMapContent, NonVirtual);
+private:
+	explicit FComputeKernelShaderMapContent(EShaderPlatform InPlatform) : Super(InPlatform) {}
+
+	/** The kernel's user friendly name */
+	LAYOUT_FIELD(FMemoryImageString, FriendlyName);
+
+	/** The static parameter set that this shader map was compiled with */
+	LAYOUT_FIELD(FComputeKernelShaderMapId, ShaderMapId);
+
+	/** Shader compilation output */
+	LAYOUT_FIELD(FComputeKernelCompilationOutput, CompilationOutput);
+};
+
+
+/**
+ * All the shader permutations for a kernel.
+ */
+class FComputeKernelShaderMap : public TShaderMap<FComputeKernelShaderMapContent, FShaderMapPointerTable>, public FDeferredCleanupInterface
+{
+public:
+	using Super = TShaderMap<FComputeKernelShaderMapContent, FShaderMapPointerTable>;
+
+	/**
+	 * Finds the shader map for a kernel.
+	 * @param InShaderMapId - Id to find
+	 * @param Platform - The platform to lookup for
+	 * @return nullptr if no cached shader map was found.
+	 */
+	static FComputeKernelShaderMap* FindId(const FComputeKernelShaderMapId& InShaderMapId, EShaderPlatform InPlatform);
+
+	// ShaderMap interface
+	template<typename ShaderType> TShaderRef<ShaderType> GetShader() const { return TShaderRef<ShaderType>(GetContent()->GetShader<ShaderType>(), *this); }
+	TShaderRef<FShader> GetShader(FShaderType* ShaderType) const { return TShaderRef<FShader>(GetContent()->GetShader(ShaderType), *this); }
+
+	/**
+	 * Attempts to load the shader map for the given kernel from the Derived Data Cache.
+	 * If InOutShaderMap is valid, attempts to load the individual missing shaders instead.
+	 */
+	static void LoadFromDerivedDataCache(const FComputeKernelResource* InKernel, const FComputeKernelShaderMapId& InShaderMapId, EShaderPlatform InPlatform, TRefCountPtr<FComputeKernelShaderMap>& InOutShaderMap);
+
+	FComputeKernelShaderMap();
+	~FComputeKernelShaderMap();
+
+	/**
+	 * Compiles the shaders for a kernel and caches them in this shader map.
+	 * @param InKernel - The kernel to compile shaders for.
+	 * @param InShaderMapId - the set of static parameters to compile for
+	 * @param InPlatform - The platform to compile to
+	 */
+	void Compile(
+		FComputeKernelResource* InKernel,
+		const FComputeKernelShaderMapId& InShaderMapId,
+		TRefCountPtr<FSharedShaderCompilerEnvironment> InCompilationEnvironment,
+		const FComputeKernelCompilationOutput& InCompilationOutput,
+		EShaderPlatform InPlatform,
+		bool bSynchronousCompile,
+		bool bApplyCompletedShaderMapForRendering
+		);
+
+	/** Sorts the incoming compiled jobs into the appropriate shader maps, and finalizes this shader map so that it can be used for rendering. */
+	bool ProcessCompilationResults(const TArray<FShaderCommonCompileJobPtr>& InCompilationResults, int32& InOutResultIndex, float& InOutTimeBudget);
+
+	/**
+	 * Checks whether the shader map is missing any shader types necessary for the given kernel.
+	 * @param InKernel - The ComputeKernel which is checked.
+	 * @return True if the shader map has all of the shader types necessary.
+	 */
+	bool IsComplete(const FComputeKernelResource* InKernel, bool bSilent);
+
+	/** Attempts to load missing shaders from memory. */
+	void LoadMissingShadersFromMemory(const FComputeKernelResource* InKernel);
+
+	/**
+	 * Checks to see if the shader map is already being compiled for another kernel, and if so
+	 * adds the specified kernel to the list to be applied to once the compile finishes.
+	 * @param InKernel - The ComputeKernel we also wish to apply the compiled shader map to.
+	 * @return True if the shader map was being compiled and we added it to the list to be applied.
+	 */
+	bool TryToAddToExistingCompilationTask(FComputeKernelResource* InKernel);
+
+	/** Builds a list of the shaders in a shader map. */
+	 void GetShaderList(TMap<FShaderId, TShaderRef<FShader>>& OutShaders) const;
+
+	/** Registers a ComputeKernel shader map in the global map so it can be used. */
+	void Register(EShaderPlatform InShaderPlatform);
+
+	// Reference counting.
+	void AddRef();
+	void Release();
+
+	/**
+	 * Removes all entries in the cache with exceptions based on a shader type
+	 * @param ShaderType - The shader type to flush
+	 */
+	void FlushShadersByShaderType(const FShaderType* ShaderType);
+
+	/** Removes a kernel from ComputeKernelShaderMapsBeingCompiled. */
+	static void RemovePending(FComputeKernelResource* InKernel);
+
+	/** Finds a shader map currently being compiled that was enqueued for the given kernel. */
+	static const FComputeKernelShaderMap* GetShaderMapBeingCompiled(const FComputeKernelResource* InKernel);
+
+	/** Serializes the shader map. */
+	bool Serialize(FArchive& Ar, bool bInlineShaderResources = true);
+
+	/** Saves this shader map to the derived data cache. */
+	void SaveToDerivedDataCache();
+
+	// Accessors.
+	const FComputeKernelShaderMapId& GetShaderMapId() const	{ return GetContent()->ShaderMapId; }
+	EShaderPlatform GetShaderPlatform() const				{ return GetContent()->GetShaderPlatform(); }
+	const FMemoryImageString& GetFriendlyName() const		{ return GetContent()->FriendlyName; }
+	uint32 GetCompilingId() const							{ return CompilingId; }
+	bool IsCompilationFinalized() const						{ return bCompilationFinalized; }
+	bool CompiledSuccessfully() const						{ return bCompiledSuccessfully; }
+
+	bool IsValid() const
+	{
+		return bCompilationFinalized && bCompiledSuccessfully && !bDeletedThroughDeferredCleanup;
+	}
+
+	int32 GetNumRefs() const { return NumRefs; }
+	uint32 GetCompilingId()  { return CompilingId; }
+	static TMap<TRefCountPtr<FComputeKernelShaderMap>, TArray<FComputeKernelResource*> > &GetInFlightShaderMaps() { return ComputeKernelShaderMapsBeingCompiled; }
+
+	void SetCompiledSuccessfully(bool bSuccess) { bCompiledSuccessfully = bSuccess; }
+private:
+
+	/**
+	 * A global map from a kernel's ID and static switch set to any shader map cached for that kernel.
+	 * Note: this does not necessarily contain all kernel shader maps in memory.  Shader maps with the same key can evict each other.
+	 * No ref counting needed as these are removed on destruction of the shader map.
+	 */
+	static TMap<FComputeKernelShaderMapId, FComputeKernelShaderMap*> GIdToComputeKernelShaderMap[SP_NumPlatforms];
+
+	/**
+	 * All kernel shader maps in memory.
+	 * No ref counting needed as these are removed on destruction of the shader map.
+	 */
+	static TArray<FComputeKernelShaderMap*> AllComputeKernelShaderMaps;
+
+	/** Next value for CompilingId. */
+	static uint32 NextCompilingId;
+
+	/** Tracks resources and their shader maps that need to be compiled but whose compilation is being deferred. */
+	static TMap<TRefCountPtr<FComputeKernelShaderMap>, TArray<FComputeKernelResource*> > ComputeKernelShaderMapsBeingCompiled;
+
+	/** Uniquely identifies this shader map during compilation, needed for deferred compilation where shaders from multiple shader maps are compiled together. */
+	uint32 CompilingId;
+
+	mutable int32 NumRefs;
+
+	/** Used to catch errors where the shader map is deleted directly. */
+	bool bDeletedThroughDeferredCleanup;
+
+	/** Indicates whether this shader map has been registered in GIdToComputeKernelShaderMap */
+	uint32 bRegistered : 1;
+
+	/**
+	 * Indicates whether this shader map has had ProcessCompilationResults called after Compile.
+	 * The shader map must not be used on the rendering thread unless bCompilationFinalized is true.
+	 */
+	uint32 bCompilationFinalized : 1;
+
+	uint32 bCompiledSuccessfully : 1;
+
+	/** Indicates whether the shader map should be stored in the shader cache. */
+	uint32 bIsPersistent : 1;
+
+	uint32 bHasFrozenContent : 1;
+
+	FShader* ProcessCompilationResultsForSingleJob(FShaderCompileJob& InSingleJob, const FSHAHash& InShaderMapHash);
+
+	bool IsComputeKernelShaderComplete(const FComputeKernelResource* InKernel, const FComputeKernelShaderType* InShaderType, bool bSilent);
+
+	friend class FShaderCompilingManager;
+};
+
+
+/**
+ * FComputeKernelResource represents a UComputeKernel to the shader compilation process
+ */
+class FComputeKernelResource
+{
+public:
+
+	/**
+	 * Minimal initialization constructor.
+	 */
+	FComputeKernelResource()
+		: bContainsInlineShaders(false)
+		, bLoadedCookedShaderMapId(false)
+	{}
+
+	virtual ~FComputeKernelResource();
+
+	/**
+	 * Caches the shaders for this kernel with no static parameters on the given platform.
+	 * This is used by UComputeKernel
+	 */
+	bool CacheShaders(EShaderPlatform InPlatform, const ITargetPlatform* TargetPlatform, bool bApplyCompletedShaderMapForRendering, bool bSynchronous);
+	bool CacheShaders(const FComputeKernelShaderMapId& InShaderMapId, EShaderPlatform InPlatform, bool bApplyCompletedShaderMapForRendering, bool bSynchronous);
+
+	/**
+	 * Should the shader for this kernel with the given platform and shader type combination be compiled
+	 *
+	 * @param InPlatform		The platform currently being compiled for
+	 * @param InShaderType		Which shader is being compiled
+	 *
+	 * @return true if the shader should be compiled
+	 */
+	virtual bool ShouldCache(EShaderPlatform InPlatform, const FShaderType* InShaderType) const;
+
+	void SerializeShaderMap(FArchive& Ar);
+
+	/** Releases this kernel's shader map.  Must only be called on ComputeKernels not exposed to the rendering thread! */
+	void ReleaseShaderMap();
+
+	/** Discards loaded shader maps if the application can't render */
+	void DiscardShaderMap();
+
+	void GetDependentShaderTypes(EShaderPlatform InPlatform, TArray<FShaderType*>& OutShaderTypes) const;
+	virtual void GetShaderMapId(EShaderPlatform InPlatform, const ITargetPlatform* TargetPlatform, FComputeKernelShaderMapId& OutId) const;
+
+	void Invalidate();
+
+	/**
+	 * Should shaders compiled for this kernel be saved to disk?
+	 */
+	virtual bool IsPersistent() const { return true; }
+
+	/**
+	 * Called when compilation finishes, after the GameThreadShaderMap is set and the render command to set the RenderThreadShaderMap is queued
+	 */
+	virtual void NotifyCompilationFinished();
+
+	/**
+	 * Cancels all outstanding compilation jobs
+	 */
+	void CancelCompilation();
+
+	/**
+	 * Blocks until compilation has completed. Returns immediately if a compilation is not outstanding.
+	 */
+	void FinishCompilation();
+
+	/**
+	 * Checks if the compilation for this shader is finished
+	 *
+	 * @return returns true if compilation is complete false otherwise
+	 */
+	bool IsCompilationFinished() const;
+
+	// Accessors.
+	const TArray<FString>& GetCompileErrors() const { return CompileErrors; }
+	void SetCompileErrors(const TArray<FString>& InCompileErrors) { CompileErrors = InCompileErrors; }
+
+	ERHIFeatureLevel::Type GetFeatureLevel() const { return FeatureLevel; }
+
+	FComputeKernelShaderMap* GetGameThreadShaderMap() const
+	{
+		checkSlow(IsInGameThread() || IsInAsyncLoadingThread());
+		return GameThreadShaderMap;
+	}
+
+	/** Note: SetRenderingThreadShaderMap must also be called with the same value, but from the rendering thread. */
+	void SetGameThreadShaderMap(FComputeKernelShaderMap* InShaderMap)
+	{
+		checkSlow(IsInGameThread() || IsInAsyncLoadingThread());
+		GameThreadShaderMap = InShaderMap;
+	}
+
+	/** Note: SetGameThreadShaderMap must also be called with the same value, but from the game thread. */
+	void SetRenderingThreadShaderMap(FComputeKernelShaderMap* InShaderMap);
+
+	void AddCompileId(uint32 InIdentifier) 
+	{
+		OutstandingCompileShaderMapIds.Add(InIdentifier);
+	}
+
+	void SetInlineShaderMap(FComputeKernelShaderMap* InShaderMap)
+	{
+		checkSlow(IsInGameThread() || IsInAsyncLoadingThread());
+		GameThreadShaderMap = InShaderMap;
+		bContainsInlineShaders = true;
+		bLoadedCookedShaderMapId = true;
+		CookedShaderMapId = InShaderMap->GetShaderMapId();
+	}
+
+	void RemoveOutstandingCompileId(const int32 InOldOutstandingCompileShaderMapId);
+
+	/**
+	* Get generated or written source code for the shader
+	* @param OutSource - generated source code
+	* @return - true on Success
+	*/
+	bool GetHLSLSource(FString& OutSource) const
+	{
+		OutSource = ShaderSource;
+		return true;
+	};
+
+	const FString& GetHLSLSource() const
+	{
+		return ShaderSource;
+	}
+
+	const FString& GetEntryPoint() const
+	{
+		return ShaderEntryPoint;
+	}
+
+	const FString& GetFriendlyName()	const { return FriendlyName; }
+
+	void SetupResource(
+		ERHIFeatureLevel::Type InFeatureLevel,
+		FString const& InFriendlyName,
+		FString const& InShaderEntryPoint,
+		FString InShaderSource,
+		uint64 InShaderCodeHash,
+		FShaderParametersMetadata* InShaderMetadata
+	);
+
+	void SetCompileErrors(TArray<FString> &InErrors)
+	{
+		CompileErrors = InErrors;
+	}
+
+	const FShaderParametersMetadata* GetShaderParamMetadata() const
+	{
+		return ShaderMetadata.Get();
+	}
+	
+	TShaderRef<FComputeKernelShader> GetShader() const;
+	TShaderRef<FComputeKernelShader> GetShaderGameThread() const;
+	
+	bool IsSame(const FComputeKernelShaderMapId& InId) const;
+
+	uint32 GetKernelFlags() const { return KernelFlags; }
+
+protected:
+
+	/**
+	 * Fills the passed array with IDs of shader maps unfinished compilation jobs.
+	 */
+	void GetShaderMapIDsWithUnfinishedCompilation(TArray<int32>& OutShaderMapIds);
+
+private:
+	TUniquePtr<FShaderParametersMetadata> ShaderMetadata;
+
+	TArray<FString> CompileErrors;
+
+	/** 
+	 * Game thread tracked shader map, which is ref counted and manages shader map lifetime. 
+	 * The shader map uses deferred deletion so that the rendering thread has a chance to process a release command when the shader map is no longer referenced.
+	 * Code that sets this is responsible for updating RenderingThreadShaderMap in a thread safe way.
+	 * During an async compile, this will be NULL and will not contain the actual shader map until compilation is complete.
+	 */
+	TRefCountPtr<FComputeKernelShaderMap> GameThreadShaderMap = nullptr;
+
+	/** 
+	 * Shader map for this FComputeKernelResource which is accessible by the rendering thread. 
+	 * This must be updated along with GameThreadShaderMap, but on the rendering thread.
+	 */
+	FComputeKernelShaderMap* RenderingThreadShaderMap = nullptr;
+
+	FString ShaderEntryPoint;
+
+	FString ShaderSource;
+
+	/** 
+	 * Hash computed from shader code. If multiple kernels give
+	 * the same output, same shader resource will be shared 
+	 */
+	uint64 ShaderCodeHash = 0;
+
+	uint32 KernelFlags = 0;
+
+	/** 
+	 * Contains the compiling id of this shader map when it is being compiled asynchronously. 
+	 * This can be used to access the shader map during async compiling, since GameThreadShaderMap will not have been set yet.
+	 */
+	TArray<int32, TInlineAllocator<1> > OutstandingCompileShaderMapIds;
+
+	/** Feature level that this kernel is representing. */
+	ERHIFeatureLevel::Type FeatureLevel = ERHIFeatureLevel::SM5;
+
+	/**
+	 * Whether this kernel was loaded with shaders inlined.
+	 * If true, GameThreadShaderMap will contain a reference to the inlined shader map between Serialize and PostLoad.
+	 */
+
+	uint32 bContainsInlineShaders : 1;
+	uint32 bLoadedCookedShaderMapId : 1;
+	FComputeKernelShaderMapId CookedShaderMapId;
+
+	/**
+	 * Compiles this kernel for InPlatform, storing the result in OutShaderMap if the compile was synchronous
+	 */
+	bool BeginCompileShaderMap(
+		const FComputeKernelShaderMapId& InShaderMapId,
+		EShaderPlatform InPlatform, 
+		TRefCountPtr<FComputeKernelShaderMap>& OutShaderMap, 
+		bool bApplyCompletedShaderMapForRendering,
+		bool bSynchronous = false);
+
+	/** Populates OutEnvironment with defines needed to compile shaders for this kernel. */
+	void SetupShaderCompilationEnvironment(
+		EShaderPlatform InPlatform,
+		FShaderCompilerEnvironment& OutEnvironment
+		) const;
+
+
+	FString FriendlyName;
+
+	friend class FComputeKernelShaderMap;
+	friend class FShaderCompilingManager;
+};
+
+
diff --git a/Engine/Source/Runtime/Engine/Private/ComputeFramework/ComputeKernelSource.cpp b/Engine/Source/Runtime/Engine/Private/ComputeFramework/ComputeKernelSource.cpp
--- a/Engine/Source/Runtime/Engine/Private/ComputeFramework/ComputeKernelSource.cpp
+++ b/Engine/Source/Runtime/Engine/Private/ComputeFramework/ComputeKernelSource.cpp
@@ -1,3 +1,32 @@
 // Copyright Epic Games, Inc. All Rights Reserved.
 
 #include "ComputeFramework/ComputeKernelSource.h"
+
+#include "ShaderParameterMetadataBuilder.h"
+
+void UComputeKernelSource::GetShaderParameters(class FShaderParametersMetadataBuilder& OutBuilder) const
+{
+	for (auto& Input : InputParams)
+	{
+		ensureAlways(Input.DimType == EShaderFundamentalDimensionType::Scalar);
+
+		switch (Input.FundamentalType)
+		{
+		case EShaderFundamentalType::Bool:
+			OutBuilder.AddParam<bool>(*Input.Name);
+			break;
+
+		case EShaderFundamentalType::Int:
+			OutBuilder.AddParam<int32>(*Input.Name);
+			break;
+
+		case EShaderFundamentalType::Uint:
+			OutBuilder.AddParam<uint32>(*Input.Name);
+			break;
+
+		case EShaderFundamentalType::Float:
+			OutBuilder.AddParam<float>(*Input.Name);
+			break;
+		}
+	}
+}
diff --git a/Engine/Source/Runtime/Engine/Private/ComputeFramework/ShaderParamTypeDefinition.cpp b/Engine/Source/Runtime/Engine/Private/ComputeFramework/ShaderParamTypeDefinition.cpp
--- a/Engine/Source/Runtime/Engine/Private/ComputeFramework/ShaderParamTypeDefinition.cpp
+++ b/Engine/Source/Runtime/Engine/Private/ComputeFramework/ShaderParamTypeDefinition.cpp
@@ -2,11 +2,10 @@
 
 #include "ComputeFramework/ShaderParamTypeDefinition.h"
 
+#include "Algo/Find.h"
+#include "Misc/DefaultValueHelper.h"
 #include "Serialization/Archive.h"
 #include "Templates/TypeHash.h"
-#include "Misc/DefaultValueHelper.h"
-#include "Algo/Find.h"
-
 
 // Storage for shared shader value types.
 static uint32 GetTypeHash(const FShaderValueTypeHandle& InTypeHandle)
@@ -33,11 +32,8 @@ struct HandleKeyFuncs : BaseKeyFuncs<FShaderValueTypeHandle,FShaderValueTypeHand
 	}
 };
 
-
 static TSet<FShaderValueTypeHandle, HandleKeyFuncs> GloballyKnownValueTypes;
 
-
-
 FShaderValueTypeHandle FShaderValueType::Get(EShaderFundamentalType InType)
 {
 	if (InType == EShaderFundamentalType::Struct)
@@ -52,7 +48,6 @@ FShaderValueTypeHandle FShaderValueType::Get(EShaderFundamentalType InType)
 	return GetOrCreate(MoveTemp(ValueType));
 }
 
-
 FShaderValueTypeHandle FShaderValueType::Get(EShaderFundamentalType InType, int32 InElemCount)
 {
 	if (InType == EShaderFundamentalType::Struct || InElemCount < 1 || InElemCount > 4)
@@ -68,7 +63,6 @@ FShaderValueTypeHandle FShaderValueType::Get(EShaderFundamentalType InType, int3
 	return GetOrCreate(MoveTemp(ValueType));
 }
 
-
 FShaderValueTypeHandle FShaderValueType::Get(EShaderFundamentalType InType, int32 InRowCount, int32 InColumnCount)
 {
 	if (InType == EShaderFundamentalType::Struct || 
@@ -87,7 +81,6 @@ FShaderValueTypeHandle FShaderValueType::Get(EShaderFundamentalType InType, int3
 	return GetOrCreate(MoveTemp(ValueType));
 }
 
-
 FShaderValueTypeHandle FShaderValueType::Get(
 	FName InName, 
 	std::initializer_list<FStructElement> InStructElements)
@@ -128,7 +121,6 @@ FShaderValueTypeHandle FShaderValueType::Get(
 	return GetOrCreate(MoveTemp(ValueType));
 }
 
-
 FShaderValueTypeHandle FShaderValueType::GetOrCreate(FShaderValueType&& InValueType)
 {
 	FShaderValueTypeHandle Handle;
@@ -145,7 +137,6 @@ FShaderValueTypeHandle FShaderValueType::GetOrCreate(FShaderValueType&& InValueT
 	return Handle;
 }
 
-
 bool FShaderValueType::operator==(const FShaderValueType& InOtherType) const
 {
 	if (Type != InOtherType.Type)
@@ -222,7 +213,6 @@ uint32 GetTypeHash(const FShaderValueType& InShaderValueType)
 	return Hash;
 }
 
-
 FString FShaderValueType::ToString() const
 {
 	// FIXME: Cache on create?
@@ -258,7 +248,6 @@ FString FShaderValueType::ToString() const
 	return BaseName;
 }
 
-
 FString FShaderValueType::GetTypeDeclaration() const
 {
 	// FIXME: Cache on create?
@@ -323,7 +312,6 @@ FArchive& operator<<(FArchive& InArchive, FShaderValueTypeHandle& InHandle)
 	return InArchive;
 }
 
-
 FArchive& operator<<(FArchive& InArchive, FShaderValueType::FStructElement& InElement)
 {
 	InArchive << InElement.Name;
@@ -331,7 +319,6 @@ FArchive& operator<<(FArchive& InArchive, FShaderValueType::FStructElement& InEl
 	return InArchive;
 }
 
-
 using FFundamentalStingPair = TPair<EShaderFundamentalType, FString>;
 using FResourceStingPair = TPair<EShaderResourceType, FString>;
 
@@ -418,8 +405,7 @@ EShaderResourceType FShaderParamTypeDefinition::ParseResource(const FString& Str
 		}
 	}
 
-	check(!"Unknown Type");
-	return EShaderResourceType::Buffer;
+	return EShaderResourceType::None;
 }
 
 void FShaderParamTypeDefinition::ResetTypeDeclaration(
@@ -431,7 +417,8 @@ void FShaderParamTypeDefinition::ResetTypeDeclaration(
 		TypeDecl.Append(TEXT("RW"));
 	}
 
-	if (BindingType != EShaderParamBindingType::ConstantParameter)
+	const bool bIsResourceType = BindingType == EShaderParamBindingType::ReadOnlyResource || BindingType == EShaderParamBindingType::ReadWriteResource;
+	if (bIsResourceType)
 	{
 		auto* foundItem = Algo::FindByPredicate(
 			ResTypeStringMap, 
@@ -468,7 +455,7 @@ void FShaderParamTypeDefinition::ResetTypeDeclaration(
 		break;
 	};
 
-	if (BindingType != EShaderParamBindingType::ConstantParameter)
+	if (bIsResourceType)
 	{
 		TypeDecl.AppendChar(TEXT('>'));
 	}
diff --git a/Engine/Source/Runtime/Engine/Private/LevelTick.cpp b/Engine/Source/Runtime/Engine/Private/LevelTick.cpp
--- a/Engine/Source/Runtime/Engine/Private/LevelTick.cpp
+++ b/Engine/Source/Runtime/Engine/Private/LevelTick.cpp
@@ -65,7 +65,7 @@
 #include "ProfilingDebugging/CsvProfiler.h"
 #include "ProfilingDebugging/RealtimeGPUProfiler.h"
 #include "GPUSkinCache.h"
-#include "ComputeFramework/ComputeFramework.h"
+#include "ComputeFramework/ComputeGraphScheduler.h"
 
 #if WITH_EDITOR
 	#include "Editor.h"
@@ -977,13 +977,13 @@ struct FSendAllEndOfFrameUpdates
 		if (InScene != nullptr)
 		{
 			GPUSkinCache = InScene->GetGPUSkinCache();
-			ComputeFramework = InScene->GetComputeFramework();
+			ComputeGraphScheduler = InScene->GetComputeGraphScheduler();
 			FeatureLevel = InScene->GetFeatureLevel();
 		}
 	}
 	
 	FGPUSkinCache* GPUSkinCache = nullptr;
-	FComputeFramework* ComputeFramework = nullptr;
+	FComputeGraphScheduler* ComputeGraphScheduler = nullptr;
 	ERHIFeatureLevel::Type FeatureLevel = ERHIFeatureLevel::Num;
 
 #if WANTS_DRAW_MESH_EVENTS
@@ -1011,7 +1011,7 @@ DECLARE_GPU_STAT(ComputeFrameworkExecuteBatches);
 void EndSendEndOfFrameUpdatesDrawEvent(FSendAllEndOfFrameUpdates& SendAllEndOfFrameUpdates)
 {
 	ENQUEUE_RENDER_COMMAND(EndDrawEventCommand)(
-		[GPUSkinCache = SendAllEndOfFrameUpdates.GPUSkinCache, ComputeFramework = SendAllEndOfFrameUpdates.ComputeFramework, FeatureLevel = SendAllEndOfFrameUpdates.FeatureLevel](FRHICommandListImmediate& RHICmdList)
+		[GPUSkinCache = SendAllEndOfFrameUpdates.GPUSkinCache, ComputeGraphScheduler = SendAllEndOfFrameUpdates.ComputeGraphScheduler, FeatureLevel = SendAllEndOfFrameUpdates.FeatureLevel](FRHICommandListImmediate& RHICmdList)
 		{
 			SCOPED_GPU_STAT(RHICmdList, EndOfFrameUpdates);
 
@@ -1032,10 +1032,10 @@ void EndSendEndOfFrameUpdatesDrawEvent(FSendAllEndOfFrameUpdates& SendAllEndOfFr
 #endif // RHI_RAYTRACING
 			}
 
-			if (ComputeFramework != nullptr)
+			if (ComputeGraphScheduler != nullptr)
 			{
 				SCOPED_GPU_STAT(RHICmdList, ComputeFrameworkExecuteBatches);
-				ComputeFramework->ExecuteBatches(RHICmdList, FeatureLevel);
+				ComputeGraphScheduler->ExecuteBatches(RHICmdList, FeatureLevel);
 			}
 		});
 
diff --git a/Engine/Source/Runtime/Engine/Public/ComputeFramework/ComputeDataInterface.h b/Engine/Source/Runtime/Engine/Public/ComputeFramework/ComputeDataInterface.h
new file mode 100644
--- /dev/null
+++ b/Engine/Source/Runtime/Engine/Public/ComputeFramework/ComputeDataInterface.h
@@ -0,0 +1,36 @@
+// Copyright Epic Games, Inc. All Rights Reserved.
+
+#pragma once
+
+#include "CoreMinimal.h"
+#include "ComputeDataInterface.generated.h"
+
+struct FComputeKernelPermutationSet;
+struct FShaderCompilerEnvironment;
+struct FShaderFunctionDefinition;
+class FShaderParametersMetadataBuilder;
+
+/** 
+ * Compute Data Interface required to compile a Compute Graph. 
+ * Compute Kernels require Data Interfaces to fulfill their external functions.
+ * Compute Data Interfaces define how Compute Data Providers will actually marshal data in and out of Kernels.
+ */
+UCLASS(Abstract, Const)
+class ENGINE_API UComputeDataInterface : public UObject
+{
+	GENERATED_BODY()
+
+public:
+	/** Gather permutations from the data interface. Any connected kernel will include these in its total compiled permutations. */
+	virtual void GetPermutations(FComputeKernelPermutationSet& OutPermutationSet) const {}
+	/** Get the data interface functions available to fulfill external inputs of a kernel. */
+	virtual void GetSupportedInputs(TArray<FShaderFunctionDefinition>& OutFunctions) const {}
+	/** Get the data interface functions available to fulfill external outputs of a kernel. */
+	virtual void GetSupportedOutputs(TArray<FShaderFunctionDefinition>& OutFunctions) const {}
+	/** Gather the shader metadata exposed by the data provider payload. */
+	virtual void GetShaderParameters(TCHAR const* UID, FShaderParametersMetadataBuilder& OutBuilder) const {}
+	/** Gather the shader code for this data provider. */
+	virtual void GetHLSL(FString& OutHLSL) const {}
+	/** Gather modifications to the compilation environment always required when including this data provider. */
+	virtual void ModifyCompilationEnvironment(FShaderCompilerEnvironment& OutEnvironment) const {}
+};
diff --git a/Engine/Source/Runtime/Engine/Public/ComputeFramework/ComputeDataProvider.h b/Engine/Source/Runtime/Engine/Public/ComputeFramework/ComputeDataProvider.h
new file mode 100644
--- /dev/null
+++ b/Engine/Source/Runtime/Engine/Public/ComputeFramework/ComputeDataProvider.h
@@ -0,0 +1,70 @@
+// Copyright Epic Games, Inc. All Rights Reserved.
+
+#pragma once
+
+#include "CoreMinimal.h"
+#include "ComputeDataProvider.generated.h"
+
+class FComputeDataProviderRenderProxy;
+struct FComputeKernelPermutationSet;
+class FRDGBuffer;
+class FRDGTexture;
+class FRHIShaderResourceView;
+class FRHIUnorderedAccessView;
+
+/**
+ * Compute Framework Data Provider.
+ * A concrete instance of this is responsible for supplying data declared by a UComputeDataInterface.
+ * One of these must be created for each UComputeDataInterface object in an instance of a Compute Graph.
+ */
+UCLASS(Abstract)
+class ENGINE_API UComputeDataProvider : public UObject
+{
+	GENERATED_BODY()
+
+public:
+	/**
+	 * Get an associated render thread proxy object.
+	 * Currently these are created and destroyed per frame by the UComputeGraphComponent.
+	 * todo[CF]: Don't destroy FComputeDataProviderRenderProxy objects every frame?
+	 */
+	virtual FComputeDataProviderRenderProxy* GetRenderProxy() { return nullptr; }
+};
+
+/**
+ * Compute Framework Data Provider Proxy. 
+ * A concrete instance of this is created by the UComputeDataProvider and used for the render thread gathering of data for a Compute Kernel. 
+ */
+class ENGINE_API FComputeDataProviderRenderProxy
+{
+public:
+	virtual ~FComputeDataProviderRenderProxy() {}
+
+	/** Called on render thread to determine how many dispatches are required to do all work on the associated data provider. */
+	virtual int32 GetInvocationCount() const { return 1; }
+
+	/** Called on render thread to determine dispatch dimension required to do all work on the associated data provider. */
+	virtual FIntVector GetDispatchDim(int32 InvocationIndex, FIntVector GroupDim) const { return FIntVector(1, 1, 1); }
+
+	/** 
+	 * Get the shader permutations required for this data provider. 
+	 * All potential data permutations should already have been registered by the associated data interface to ensure that the compiled permutation exists. 
+	 */
+	virtual void GetPermutations(int32 InvocationIndex, FComputeKernelPermutationSet& OutPermutationSet) const {}
+
+	/** 
+	 * The name-value shader bindings that are collected from data providers.
+	 * All names should already have been registered by the associated data interface. 
+	 * todo[CF]: Investigate if this can be replaced with something more efficient. We know that data providers should only need to expose structs and could fill and return them directly.
+	 */
+	struct FBindings
+	{
+		TMap< FString, int32 > ParamsInt;
+		TMap< FString, uint32 > ParamsUint;
+		TMap< FString, float > ParamsFloat;
+		TMap< FString, TArray<uint8> > Structs;
+	};
+
+	/** Gather the shader bindings for the data provider. */
+	virtual void GetBindings(int32 InvocationIndex, TCHAR const* UID, FBindings& OutBindings) const {}
+};
diff --git a/Engine/Source/Runtime/Engine/Public/ComputeFramework/ComputeFramework.h b/Engine/Source/Runtime/Engine/Public/ComputeFramework/ComputeFramework.h
--- a/Engine/Source/Runtime/Engine/Public/ComputeFramework/ComputeFramework.h
+++ b/Engine/Source/Runtime/Engine/Public/ComputeFramework/ComputeFramework.h
@@ -2,80 +2,10 @@
 
 #pragma once
 
-#include "Templates/RefCounting.h"
+#include "CoreMinimal.h"
 #include "RHIDefinitions.h"
-#include "Shader.h"
 
-class FRDGPooledBuffer;
-class FComputeKernelResource;
-class FComputeKernelShader;
-
-class UComputeGraph;
+DECLARE_LOG_CATEGORY_EXTERN(LogComputeFramework, Log, All);
 
+/** Returns true if ComputeFramework is supported. */
 extern ENGINE_API bool SupportsComputeFramework(ERHIFeatureLevel::Type FeatureLevel, EShaderPlatform ShaderPlatform);
-
-enum class EComputeExecutionBufferType
-{
-	Buffer,
-	IndexBuffer,
-	StructuredBuffer,
-	ByteAddressBuffer,
-};
-
-struct FComputeExecutionBufferDesc
-{
-	FName Name;
-	uint32 BufferType : 2;
-	uint32 Flags : 30;
-	uint32 BytesPerElement = 0;
-	uint32 ElementCount = 0;
-};
-
-struct FComputeExecutionExternalBufferDesc
-{
-	FName Name;
-	TRefCountPtr<FRDGPooledBuffer> Buffer;
-};
-
-class FComputeGraph
-{
-public:
-	struct FKernelInvocation
-	{
-		FName KernelName;
-		FName InvocationName;
-		FIntVector DispatchDim;
-		FComputeKernelResource* Kernel = nullptr;
-	};
-
-	void Initialize(UComputeGraph* ComputeGraph);
-
-	TArray<FKernelInvocation> KernelInvocations;
-};
-
-class FComputeFramework
-{
-public:
-	struct FShaderInvocation
-	{
-		FName KernelName;
-		FName InvocationName;
-		FIntVector DispatchDim;
-		const FShaderParametersMetadata* ShaderParamMetadata = nullptr;
-		TShaderRef<FComputeKernelShader> Shader;
-	};
-
-	void EnqueueForExecution(
-		const FComputeGraph* ComputeGraph
-		);
-
-	void ExecuteBatches(
-		FRHICommandListImmediate& RHICmdList,
-		ERHIFeatureLevel::Type FeatureLevel
-		);
-
-private:
-	TArray<FComputeExecutionExternalBufferDesc> ExternalBuffers;
-	TArray<FComputeExecutionBufferDesc> TransientBuffers;
-	TArray<FShaderInvocation> ComputeShaders;
-};
diff --git a/Engine/Source/Runtime/Engine/Public/ComputeFramework/ComputeGraph.h b/Engine/Source/Runtime/Engine/Public/ComputeFramework/ComputeGraph.h
--- a/Engine/Source/Runtime/Engine/Public/ComputeFramework/ComputeGraph.h
+++ b/Engine/Source/Runtime/Engine/Public/ComputeFramework/ComputeGraph.h
@@ -2,188 +2,114 @@
 
 #pragma once
 
-#include "PixelFormat.h"
+#include "CoreMinimal.h"
+#include "RHIDefinitions.h"
 #include "ComputeGraph.generated.h"
 
-class UComputeKernelSource;
+class FComputeKernelResource;
+class UComputeDataInterface;
 class UComputeKernel;
 
-/* Identifies a specific kernel invocation within a compute graph. */
-struct FComputeKernelInvocationHandle
+/** Compute Kernel compilation flags. */
+enum class EComputeKernelCompilationFlags
 {
-	uint16 InvocationIdx = 0xFFFF;
-	uint16 GenerationIdx = 0xFFFF;
-};
+	None = 0,
 
-/* Identifies a specific resource within a compute graph. */
-struct FComputeResourceHandle
-{
-	uint16 ResourceIdx = 0xFFFF;
-	uint16 GenerationIdx = 0xFFFF;
-};
+	/* Force recompilation even if kernel is not dirty and/or DDC data is available. */
+	Force = 1 << 0,
 
-/* Identifies a specific binding within a compute graph. */
-struct FComputeBindingHandle
-{
-	uint16 BindingIdx = 0xFFFF;
-	uint16 GenerationIdx = 0xFFFF;
-};
+	/* Compile the shader while blocking the main thread. */
+	Synchronous = 1 << 1,
 
-UENUM()
-enum class EComputeGraphResourceType : uint8
-{
-	Buffer,
-	Texture1D,
-	Texture2D,
-	Texture3D,
-};
+	/* Replaces all instances of the shader with the newly compiled version. */
+	ApplyCompletedShaderMapForRendering = 1 << 2,
 
-/* Defines creation format and dimensions of a resource. */
-struct FComputeResourceDesc
-{
-	EPixelFormat ResourceFormat = EPixelFormat::PF_Unknown;
-	uint16 X = 0;
-	uint16 Y = 0;	// Only used for Tex2D and Tex3D
-	uint16 Z = 0;	// Only used for Tex3D
-	EComputeGraphResourceType Type = EComputeGraphResourceType::Buffer;
-	uint8 ArrayCount = 0;
-	uint8 MipLevels = 0;
-	uint8 Flags = 0;
-	uint8 SampleCount = 0;
-	uint8 SampleQuality = 0;
+	IsCooking = 1 << 3,
 };
 
-/*
- * Defines view into an externally allocated resouse and valid uses.
- * Less params are necessary than FComputeResourceDesc since only
- * bindings/types are validated with external resources.
+/** 
+ * Description of a single edge in a UComputeGraph. 
+ * todo[CF]: Consider better storage for graph data structure that is easier to interrogate efficiently.
  */
-struct FComputeResourceExternalDesc
-{
-	EPixelFormat ResourceFormat = EPixelFormat::PF_Unknown;
-	EComputeGraphResourceType Type = EComputeGraphResourceType::Buffer;
-	uint8 Flags = 0;
-};
-
-USTRUCT()
-struct FComputeKernelInvocation
-{
-	GENERATED_BODY()
-
-	static constexpr uint32 MAX_NAME_LENGTH = 32;
-
-	UPROPERTY()
-	TObjectPtr<UComputeKernel> ComputeKernel = nullptr;
-	
-	UPROPERTY()
-	uint16 GenerationIdx = 0xFFFF;
-
-	FComputeKernelInvocation() = default;
-	FComputeKernelInvocation(UComputeKernel* InComputeKernel, uint16 InGenerationIdx)
-		: ComputeKernel(InComputeKernel)
-		, GenerationIdx(InGenerationIdx)
-	{}
-};
-
-USTRUCT()
-struct FComputeResource
-{
-	GENERATED_BODY()
-
-	static constexpr uint32 MAX_NAME_LENGTH = 32;
-
-	UPROPERTY()
-	uint16 GenerationIdx = 0xFFFF;
-};
-
-/* Identifies a specific kernel invocation AND parameter slot within a compute graph. */
-struct FComputeKernelInvocationBindPoint
-{
-	FComputeKernelInvocationHandle KernelInvocation;
-	FString BindPointName;
-};
-
-/* Define is the binding data flows from resource into kernel or kernel into resource. */
-enum class EComputeGraphBindingType : uint8
-{
-	Input,	// Resource       ->   Kernel Input
-	Output, // Kernel Output  ->   Resource
-};
-
-/* Identifies a specific data flow between a kernel and resource with the compute graph. */
 USTRUCT()
-struct FKernelBinding
+struct FComputeGraphEdge
 {
 	GENERATED_BODY()
 
-	EComputeGraphBindingType BindingType = EComputeGraphBindingType::Input;
-	FComputeKernelInvocationBindPoint InvocationBindPoint;
-	FComputeResourceHandle Resource;
+	bool bKernelInput;
+	int32 KernelIndex;
+	int32 KernelBindingIndex;
+	int32 DataInterfaceIndex;
+	int32 DataInterfaceBindingIndex;
 };
 
-/*
- * Graph representing the logical execution and data flow topology as a DAG. The underlying
- * scheduler will extract parallel/serial execution from this DAG description.
+/** 
+ * Class representing a Compute Graph.
+ * This holds the basic topology of the graph and is responsible for linking Kernels with Data Interfaces and compiling the resulting shader code.
+ * Multiple Compute Graph asset types can derive from this to specialize the graph creation process. 
+ * For example the Animation Deformer system provides a UI for creating UComputeGraph assets.
  */
 UCLASS()
 class ENGINE_API UComputeGraph : public UObject
 {
 	GENERATED_BODY()
 
-public:
-	//===================================
-	// Add/Remove "nodes" in the graph
-	//===================================
+protected:
+	/** Kernels in the graph. */
+	UPROPERTY()
+	TArray< TObjectPtr<UComputeKernel> > KernelInvocations;
 
-	/* May fail and return nullptr. */
-	FComputeKernelInvocationHandle AddKernel(UComputeKernelSource* KernelSource);
+	/** Data interfaces in the graph. */
+	UPROPERTY()
+	TArray< TObjectPtr<UComputeDataInterface> > DataInterfaces;
 
-	/* Returns true if kernel successfully removed. When a kernel is removed all associated bindings are also removed. */
-	bool RemoveKernel(FComputeKernelInvocationHandle KernelInvocation);
+	/** Edges in the graph between kernels and data interfaces. */
+	UPROPERTY()
+	TArray<FComputeGraphEdge> GraphEdges;
 
-	/*
-	 * Compute resource are parameters/buffers/textures that hold data during the execution of the 
-	 * compute graph. They can either be transient - their lifetime does not out live the execution 
-	 * of the graph, or external - the resource out lives the graph and acts as the mechanism to 
-	 * bring data in or move data out of the compute graph.
+	/** 
+	 * Kernel resources that hold the compiled shader resources.
+	 * This is stored with the same indexing as the KernelInvocations array. 
 	 */
-	FComputeResourceHandle AddResource(FString Name, FComputeResourceDesc Desc);
-	FComputeResourceHandle AddResource(FString Name, FComputeResourceExternalDesc Desc);
-
-	/* Returns true if resource successfully removed. When a resource is removed all associated bindings are also removed. */
-	bool RemoveResource(FComputeResourceHandle Resource);
-
-
-
-	//===================================
-	// Add/Remove edges in the graph
-	//===================================
-
-	/* Creates a binding, aka data flow, from src to dest. */
-	FComputeBindingHandle AddBinding(
-		EComputeGraphBindingType BindingType,
-		FComputeKernelInvocationBindPoint InvocationBindPoint,
-		FComputeResourceHandle Resource
-		);
-
-	void RemoveBinding(FComputeBindingHandle Binding);
-
-	TArray<FComputeKernelInvocation> GetShaderInvocationList() { return KernelInvocations; }
+	TArray< TUniquePtr<FComputeKernelResource> > KernelResources;
 
+public:
+	UComputeGraph(const FObjectInitializer& ObjectInitializer);
+	UComputeGraph(FVTableHelper& Helper);
+	virtual ~UComputeGraph();
+	
+	/** Get the kernel instances in the graph. Note that it is valid for the array to have holes in it. */
+	TArray< TObjectPtr<UComputeKernel> > const& GetKernelInvocations() const { return KernelInvocations; }
+	/** Get the kernel resources in the graph. This array is in sync with the one returned by GetKernelInvocations(). */
+	TArray< TUniquePtr<FComputeKernelResource> > const& GetKernelResources() const { return KernelResources; }
+
+	/** Returns true if graph is valid. A valid graph should be guaranteed to compile, assuming the underlying shader code is well formed. */
+	bool ValidateGraph(FString* OutErrors = nullptr);
+
+	/**
+	 * Get unique data interface id.
+	 * This is just a string containing the index of the data interface in UComputeGraph::DataInterfaces.
+	 * It is used as a prefix to disambiguate shader code etc.
+	 * This function permanently allocates the UID on first use so that returned TCHAR pointers can be held by structures with long lifetimes.
+	 */
+	static TCHAR const* GetDataInterfaceUID(int32 DataInterfaceIndex);
 
-	// TEMP
+protected:
+	//~ Begin UObject Interface.
 	void PostLoad() override;
-
-private:
-	/* Represents all the invocation nodes of the graph. */
-	UPROPERTY()
-	TArray<FComputeKernelInvocation> KernelInvocations;
-
-	/* Represents all the data nodes of the graph */
-	UPROPERTY()
-	TArray<FComputeResource> Resources;
-
-	/* Represents all the edges of the graph. */
-	UPROPERTY()
-	TArray<FKernelBinding> 	Bindings;
+	//~ End UObject Interface.
+
+#if WITH_EDITOR
+	/** Triggers compilation of all kernels in the graph. */
+	void CacheResourceShadersForRendering();
+	void CacheResourceShadersForRendering(uint32 CompilationFlags);
+
+	/** Trigger compilation of a specific kernel. */
+	static void CacheShadersForResource(
+		EShaderPlatform ShaderPlatform,
+		const ITargetPlatform* TargetPlatform,
+		uint32 CompilationFlags,
+		FComputeKernelResource* Kernel
+	);
+#endif
 };
diff --git a/Engine/Source/Runtime/Engine/Public/ComputeFramework/ComputeGraphComponent.h b/Engine/Source/Runtime/Engine/Public/ComputeFramework/ComputeGraphComponent.h
--- a/Engine/Source/Runtime/Engine/Public/ComputeFramework/ComputeGraphComponent.h
+++ b/Engine/Source/Runtime/Engine/Public/ComputeFramework/ComputeGraphComponent.h
@@ -2,16 +2,16 @@
 
 #pragma once
 
-#include "CoreMinimal.h"
-#include "UObject/ObjectMacros.h"
-#include "UObject/UObjectGlobals.h"
 #include "Components/ActorComponent.h"
-
 #include "ComputeGraphComponent.generated.h"
 
+class UComputeDataProvider;
 class UComputeGraph;
 
-
+/** 
+ * Component which holds an instance of a specific context for a UComputeGraph.
+ * This object binds a graph to its data providers, and queues the execution. 
+ */
 UCLASS(Blueprintable, meta = (BlueprintSpawnableComponent))
 class UComputeGraphComponent : public UActorComponent
 {
@@ -20,19 +20,28 @@ class UComputeGraphComponent : public UActorComponent
 public:
 	UComputeGraphComponent();
 
-	UPROPERTY(EditAnywhere, Category = "Compute")
+	/** The Compute Graph asset. */
+	UPROPERTY(EditAnywhere, BlueprintReadWrite, Category = "Compute")
 	TObjectPtr<UComputeGraph> ComputeGraph = nullptr;
 
+	/** The bound Data Provider objects. */
+	UPROPERTY(Transient)
+	TArray< TObjectPtr<UComputeDataProvider> > DataProviders;
+
+	/** 
+	 * Set a Data Provider object to be available on the next graph execution. 
+	 * todo[CF]: Can we automate the setup of data providers from information in the graph and other metadata?
+	 */
+	UFUNCTION(BlueprintCallable, Category = "Compute")
+	void SetDataProvider(int32 Index, UComputeDataProvider* DataProvider);
+
+	/** Queue the graph for execution at the next render update. */
 	UFUNCTION(BlueprintCallable, Category = "Compute")
 	void QueueExecute();
 
 protected:
 	//~ Begin UActorComponent Interface
-	bool ShouldCreateRenderState() const override
-	{
-		return true;
-	}
-
 	void SendRenderDynamicData_Concurrent() override;
+	bool ShouldCreateRenderState() const override {	return true; }
 	//~ End UActorComponent Interface
 };
diff --git a/Engine/Source/Runtime/Engine/Public/ComputeFramework/ComputeGraphScheduler.h b/Engine/Source/Runtime/Engine/Public/ComputeFramework/ComputeGraphScheduler.h
new file mode 100644
--- /dev/null
+++ b/Engine/Source/Runtime/Engine/Public/ComputeFramework/ComputeGraphScheduler.h
@@ -0,0 +1,79 @@
+// Copyright Epic Games, Inc. All Rights Reserved.
+
+#pragma once
+
+#include "CoreMinimal.h"
+#include "ComputeFramework/ComputeDataProvider.h"
+#include "RHIDefinitions.h"
+#include "Shader.h"
+
+class FComputeDataProviderRenderProxy;
+class FComputeKernelResource;
+class FComputeKernelShader;
+class UComputeGraph;
+
+/** 
+ * Render thread proxy object for a UComputeGraph. 
+ * FComputeGraphProxy objects are created every frame for each graph that is scheduled for execution.
+ */
+class FComputeGraphProxy
+{
+public:
+	/** Called on the game thread to set up the data required by the render thread. */
+	void Initialize(UComputeGraph* ComputeGraph);
+
+	/** Description for each kernel in the graph. */
+	struct FKernelInvocation
+	{
+		FName KernelName;
+		FName InvocationName;
+		FIntVector GroupDim;
+		FComputeKernelResource* Kernel = nullptr;
+	};
+
+	TArray<FKernelInvocation> KernelInvocations;
+};
+
+/** 
+ * Class that manages the scheduling of Compute Graph work.
+ * Work can be enqueued on the render thread for the execution at the next call to ExecuteBatches().
+ */
+class FComputeGraphScheduler
+{
+public:
+	/** Enqueue a compute graph for execution. */
+	void EnqueueForExecution(
+		const FComputeGraphProxy* ComputeGraph,
+		TArrayView<FComputeDataProviderRenderProxy* const> ComputeDataProviders
+		);
+
+	/** Submit enqueued compute graph work. */
+	void ExecuteBatches(
+		FRHICommandListImmediate& RHICmdList,
+		ERHIFeatureLevel::Type FeatureLevel
+		);
+
+private:
+	/** Description of each dispatch that is enqueued. */
+	struct FShaderInvocation
+	{
+		FName KernelName;
+		FName InvocationName;
+		FIntVector DispatchDim;
+		const FShaderParametersMetadata* ShaderParamMetadata = nullptr;
+		TShaderRef<FComputeKernelShader> Shader;
+		int32 BindingsIndex = 0;
+	};
+
+	/** 
+	 * Shader invocations to dispatch. Stored in required order for dispatch. 
+	 * todo[CF]: We probably need more context for dispatching work with minimal overhead. For example we would like to overlap UAVs on any skin cache writing.
+	 */
+	TArray<FShaderInvocation> ComputeShaders;
+
+	/** 
+	 * Parameter and resource binding information for queued dispatches.
+	 * todo[CF]: Replace this with something more efficient. Hopefully using the existing StructMetadata classes.
+	 */
+	TArray<FComputeDataProviderRenderProxy::FBindings> Bindings;
+};
diff --git a/Engine/Source/Runtime/Engine/Public/ComputeFramework/ComputeKernel.h b/Engine/Source/Runtime/Engine/Public/ComputeFramework/ComputeKernel.h
--- a/Engine/Source/Runtime/Engine/Public/ComputeFramework/ComputeKernel.h
+++ b/Engine/Source/Runtime/Engine/Public/ComputeFramework/ComputeKernel.h
@@ -3,34 +3,12 @@
 #pragma once
 
 #include "CoreMinimal.h"
-#include "ComputeFramework/ShaderParamTypeDefinition.h"
 #include "ComputeFramework/ComputeKernelPermutationSet.h"
-#include "ComputeFramework/ComputeKernelShared.h"
-
 #include "ComputeKernel.generated.h"
 
-class FComputeKernelResource;
 class UComputeKernelSource;
 
-DECLARE_LOG_CATEGORY_EXTERN(ComputeKernel, Log, All);
-
-/* Describes the size and shape (threads) of a kernel invocation. */
-USTRUCT(BlueprintType)
-struct FKernelInvocationDimension
-{
-	GENERATED_BODY()
-
-	UPROPERTY(EditAnywhere, Category = "Kernel")
-	uint16 X = 0;
-
-	UPROPERTY(EditAnywhere, Category = "Kernel")
-	uint16 Y = 0;
-
-	UPROPERTY(EditAnywhere, Category = "Kernel")
-	uint16 Z = 0;
-};
-
-/* Flags that convey kernel behavior to aid compilation/optimizations. */
+/** Flags that convey kernel behavior to aid compilation/optimizations. */
 UENUM(meta = (Bitflags, UseEnumValuesAsMaskValuesInEditor = "true"))
 enum class EComputeKernelFlags : uint32
 {
@@ -50,23 +28,7 @@ enum class EComputeKernelFlags : uint32
 };
 ENUM_CLASS_FLAGS(EComputeKernelFlags);
 
-enum class EComputeKernelCompilationFlags
-{
-	None		= 0,
-	
-	/* Force recompilation even if kernel is not dirty and/or DDC data is available. */
-	Force		= 1 << 0,
-
-	/* Compile the shader while blocking the main thread. */
-	Synchronous = 1 << 1,
-
-	/* Replaces all instances of the shader with the newly compiled version. */
-	ApplyCompletedShaderMapForRendering = 1 << 2,
-
-	IsCooking = 1 << 3,
-};
-
-/* Base class representing a kernel that will be run as a shader on the GPU. */
+/** Base class representing a kernel that will be run as a shader on the GPU. */
 UCLASS(hidecategories = (Object))
 class ENGINE_API UComputeKernel : public UObject
 {
@@ -74,29 +36,24 @@ class ENGINE_API UComputeKernel : public UObject
 
 public:
 	/* 
+	 * The compute kernel source asset.
 	 * A kernel's source may be authored by different mechanisms; e.g. HLSL text, VPL graph, ML Meta Lang, etc
-	 * This abstracts the source and compilation process.
 	 */
 	UPROPERTY(EditAnywhere, AssetRegistrySearchable, meta = (ShowOnlyInnerProperties), Category = "Kernel")
 	TObjectPtr<UComputeKernelSource> KernelSource = nullptr;
 
-	/* Specifying certain memory access flags allows for optimizations such as kernel fusing. */
+	/** Specifying certain memory access flags allows for optimizations such as kernel fusing. */
 	UPROPERTY(VisibleAnywhere, BlueprintReadOnly, meta = (Bitmask, BitmaskEnum = "EComputeKernelFlags"), Category = "Kernel")
 	int32 KernelFlags = 0;
 
+	/** Permutation overrides on the kernel. */
 	UPROPERTY(EditDefaultsOnly, meta = (ShowOnlyInnerProperties), Category = "Kernel")
 	FComputeKernelPermutationSet PermutationSetOverrides;
 
+	/** Shader compilation environment overrides on the kernel. */
 	UPROPERTY(EditDefaultsOnly, meta = (ShowOnlyInnerProperties), Category = "Kernel")
 	FComputeKernelDefinitionsSet DefinitionsSetOverrides;
 
-	/*
-	 * The minimum number of invocations (threads) launched by the kernel definition.
-	 * Generally aim for multiple of 64 total. i.e. X * Y * Z = 64
-	 */
-	UPROPERTY(VisibleAnywhere, Category = "Kernel")
-	FKernelInvocationDimension GroupSizeDim;
-
 #if WITH_EDITOR
 	DECLARE_EVENT_OneParam(UComputeKernel, FShaderResetEvent, const UComputeKernel*);
 
@@ -107,44 +64,11 @@ public:
 	FShaderResetEvent ShaderResetSignal;
 #endif
 
-	void SetKernelSource(UComputeKernelSource* KernelSource);
-
-	/*
-	 * Implemented by derived classes to specify how many invocations of the kernel are necessary 
-	 * given the specific data payloads the kernel will operate on. 
-	 */
-	//virtual FKernelInvocationDimension CalculateInvocationDimension(const TArray<FBindPointDataDesc>& BindingCtx) const = 0;
-
-	//=========================================================================
-	// UObject interface
-	void PostLoad(
-		) override;
-
-#if WITH_EDITOR
-	void PostEditChangeChainProperty(
-		FPropertyChangedChainEvent& PropertyChangedEvent
-		) override;
-#endif
-	//=========================================================================
-
-	FComputeKernelResource* GetResource() { return KernelResource.Get(); }
-
-private:
+protected:
 #if WITH_EDITOR
-	/* Entry point to initiate any [re]compilation necessary due to changes. */
-	void CacheResourceShadersForRendering(
-		uint32 CompilationFlags = uint32(EComputeKernelCompilationFlags::ApplyCompletedShaderMapForRendering) | uint32(EComputeKernelCompilationFlags::Force)
-		);
-
-	/* Given a FKernelResouce, its associated shader is compiled - used for rendering and cooking. */
-	static void CacheShadersForResource(
-		EShaderPlatform ShaderPlatform,
-		const ITargetPlatform* TargetPlatform,
-		uint32 CompilationFlags,
-		FComputeKernelResource* Kernel
-		);
+	//~ Begin UObject Interface.
+	void PostLoad() override;
+	void PostEditChangeChainProperty(FPropertyChangedChainEvent& PropertyChangedEvent) override;
+	//~ End UObject Interface.
 #endif
-
-	/* The shader resource encapsulating the kernel. Akin to FMaterialResource */
-	TUniquePtr<FComputeKernelResource> KernelResource = nullptr;
 };
diff --git a/Engine/Source/Runtime/Engine/Public/ComputeFramework/ComputeKernelFromText.h b/Engine/Source/Runtime/Engine/Public/ComputeFramework/ComputeKernelFromText.h
--- a/Engine/Source/Runtime/Engine/Public/ComputeFramework/ComputeKernelFromText.h
+++ b/Engine/Source/Runtime/Engine/Public/ComputeFramework/ComputeKernelFromText.h
@@ -4,11 +4,10 @@
 
 #include "ComputeKernelSource.h"
 #include "Engine/EngineTypes.h"
-
 #include "ComputeKernelFromText.generated.h"
 
-/* 
- * Responsible for loading HLSL text and parsing the options available.
+/**
+ * Class responsible for loading HLSL text and parsing the options available.
  */
 UCLASS(BlueprintType)
 class ENGINE_API UComputeKernelFromText : public UComputeKernelSource
@@ -18,20 +17,24 @@ class ENGINE_API UComputeKernelFromText : public UComputeKernelSource
 public:
 	UComputeKernelFromText();
 
-	/* Filepath to the source file containing the kernel entry points. */
+	/** Filepath to the source file containing the kernel entry points and all options for parsing. */
 	UPROPERTY(EditDefaultsOnly, AssetRegistrySearchable, meta = (ContentDir, RelativeToGameContentDir, FilePathFilter = "Unreal Shader File (*.usf)|*.usf"), Category = "Kernel")
 	FFilePath SourceFile;
 
+	/** Kernel entry point. */
 	UPROPERTY(VisibleAnywhere, AssetRegistrySearchable, Category = "Kernel")
 	FString EntryPointName;
 
+	/** A unique id for the asset. */
 	UPROPERTY()
 	FGuid UniqueId;
 
+	/** Stored hash for the kernel source. */
 	UPROPERTY()
 	uint64 SourceHash = 0;
 
-
+protected:
+	//~ Begin UComputeKernelSource Interface.
 	FString GetEntryPoint() const override
 	{
 		return EntryPointName;
@@ -46,23 +49,26 @@ public:
 #endif
 	}
 
-	uint64 GetSourceHashCode() const override
+	uint64 GetSourceCodeHash() const override
 	{
 		return SourceHash;
 	}
-
-	void PostLoad() override;
+	//~ End UComputeKernelSource Interface.
 
 #if WITH_EDITOR
-	FString KernelSourceText;
-	FFilePath PrevSourceFile;
-
-	virtual void PreEditChange(FEditPropertyChain& PropertyAboutToChange) override;
-	virtual void PostEditChangeChainProperty(FPropertyChangedChainEvent& PropertyChangedEvent) override;
+	//~ Begin UObject Interface.
+	void PostLoad() override;
+	void PreEditChange(FEditPropertyChain& PropertyAboutToChange) override;
+	void PostEditChangeChainProperty(FPropertyChangedChainEvent& PropertyChangedEvent) override;
+	//~ End UObject Interface.
 #endif
 
 private:
 #if WITH_EDITOR
+	/** Parse the kernel source to get the kernel external functions and other data. */
 	void ReparseKernelSourceText();
+
+	FString KernelSourceText;
+	FFilePath PrevSourceFile;
 #endif
 };
diff --git a/Engine/Source/Runtime/Engine/Public/ComputeFramework/ComputeKernelShader.h b/Engine/Source/Runtime/Engine/Public/ComputeFramework/ComputeKernelShader.h
deleted file mode 100644
diff --git a/Engine/Source/Runtime/Engine/Public/ComputeFramework/ComputeKernelShaderCompilationManager.h b/Engine/Source/Runtime/Engine/Public/ComputeFramework/ComputeKernelShaderCompilationManager.h
deleted file mode 100644
diff --git a/Engine/Source/Runtime/Engine/Public/ComputeFramework/ComputeKernelShaderType.h b/Engine/Source/Runtime/Engine/Public/ComputeFramework/ComputeKernelShaderType.h
deleted file mode 100644
diff --git a/Engine/Source/Runtime/Engine/Public/ComputeFramework/ComputeKernelShared.h b/Engine/Source/Runtime/Engine/Public/ComputeFramework/ComputeKernelShared.h
deleted file mode 100644
diff --git a/Engine/Source/Runtime/Engine/Public/ComputeFramework/ComputeKernelSource.h b/Engine/Source/Runtime/Engine/Public/ComputeFramework/ComputeKernelSource.h
--- a/Engine/Source/Runtime/Engine/Public/ComputeFramework/ComputeKernelSource.h
+++ b/Engine/Source/Runtime/Engine/Public/ComputeFramework/ComputeKernelSource.h
@@ -4,35 +4,45 @@
 
 #include "ComputeFramework/ComputeKernelPermutationSet.h"
 #include "ComputeFramework/ShaderParamTypeDefinition.h"
-
 #include "ComputeKernelSource.generated.h"
 
-class FComputeKernelResource;
-
+/** 
+ * Class representing the source for a UComputeKernel 
+ * We derive from this for each authoring mechanism. (HLSL text, VPL graph, ML Meta Lang, etc.)
+ */
 UCLASS(Abstract, EditInlineNew)
 class UComputeKernelSource : public UObject
 {
 	GENERATED_BODY()
 
 public:
+	/** Base permutations exposed by the kernel. These will be extended by further permutations declared in any linked data providers. */
 	UPROPERTY(VisibleAnywhere, meta = (ShowOnlyInnerProperties), Category = "Kernel")
 	FComputeKernelPermutationSet PermutationSet;
 
+	/** Base environment defines for kernel compilation. These will be extended by further defines declared in any linked data providers. */
 	UPROPERTY(VisibleAnywhere, meta = (ShowOnlyInnerProperties), Category = "Kernel")
 	FComputeKernelDefinitionsSet DefinitionsSet;
 
-	/* Named input types to the kernel. */
+	/* Named input parameters for the kernel. */
 	UPROPERTY(VisibleAnywhere, EditFixedSize, Category = "Parameters")
 	TArray<FShaderParamTypeDefinition> InputParams;
 
-	UPROPERTY(VisibleAnywhere, EditFixedSize, Category = "Parameters")
-	TArray<FShaderParamTypeDefinition> InputSRVs;
+	/* Named external inputs for the kernel. These must be fulfilled by linked data providers. */
+	UPROPERTY(VisibleAnywhere, EditFixedSize, Category = "External")
+	TArray<FShaderFunctionDefinition> ExternalInputs;
 
-	/* Named output types to the kernel. */
-	UPROPERTY(VisibleAnywhere, EditFixedSize, Category = "Parameters")
-	TArray<FShaderParamTypeDefinition> Outputs;
+	/* Named external outputs for the kernel. These must be fulfilled by linked data providers. */
+	UPROPERTY(VisibleAnywhere, EditFixedSize, Category = "External")
+	TArray<FShaderFunctionDefinition> ExternalOutputs;
+
+	/** Register the input parameters with a shader metadata builder. */
+	void GetShaderParameters(class FShaderParametersMetadataBuilder& OutBuilder) const;
 
+	/** Get kernel entry point name. */
 	virtual FString GetEntryPoint() const PURE_VIRTUAL(UComputeKernelSource::GetEntryPoint, return {};);
+	/** Get kernel source code. */
 	virtual FString GetSource() const PURE_VIRTUAL(UComputeKernelSource::GetSource, return {};);
-	virtual uint64 GetSourceHashCode() const PURE_VIRTUAL(UComputeKernelSource::GetSourceHashCode, return 0;);
+	/** Get a hash of the kernel source code. */
+	virtual uint64 GetSourceCodeHash() const PURE_VIRTUAL(UComputeKernelSource::GetSourceCodeHash, return 0;);
 };
diff --git a/Engine/Source/Runtime/Engine/Public/ComputeFramework/ShaderParamTypeDefinition.h b/Engine/Source/Runtime/Engine/Public/ComputeFramework/ShaderParamTypeDefinition.h
--- a/Engine/Source/Runtime/Engine/Public/ComputeFramework/ShaderParamTypeDefinition.h
+++ b/Engine/Source/Runtime/Engine/Public/ComputeFramework/ShaderParamTypeDefinition.h
@@ -4,14 +4,11 @@
 
 #include "CoreMinimal.h"
 #include "UObject/NameTypes.h"
-
 #include "ShaderParamTypeDefinition.generated.h"
 
-
 class FArchive;
 struct FShaderValueType;
 
-
 /* The base types of data that shaders can consume/expose */
 UENUM()
 enum class EShaderFundamentalType : uint8
@@ -40,14 +37,17 @@ enum class EShaderFundamentalDimensionType : uint8
 UENUM()
 enum class EShaderParamBindingType : uint8
 {
+	None,
 	ConstantParameter,
 	ReadOnlyResource,  // SRV, treated as Inputs
 	ReadWriteResource, // UAV, treated as Outputs
 };
 
+/*  */
 UENUM()
 enum class EShaderResourceType : uint8
 {
+	None,
 	Texture1D,
 	Texture2D,
 	Texture3D,
@@ -57,6 +57,7 @@ enum class EShaderResourceType : uint8
 	ByteAddressBuffer,
 };
 
+/*  */
 USTRUCT()
 struct FShaderValueTypeHandle
 {
@@ -86,7 +87,7 @@ struct FShaderValueTypeHandle
 
 };
 
-
+/*  */
 USTRUCT()
 struct ENGINE_API FShaderValueType
 {
@@ -112,8 +113,6 @@ struct ENGINE_API FShaderValueType
 		FShaderValueTypeHandle Type;
 	};
 
-
-
 	/** Returns a scalar value type. If the fundamental type given is invalid for scalar values 
 	  * (e.g. struct), then this function returns a nullptr. 
 	  */
@@ -180,7 +179,6 @@ private:
 */
 uint32 GetTypeHash(const FShaderValueType& InShaderValueType);
 
-
 /* Fully describes the name and type of a parameter a shader exposes. */
 USTRUCT()
 struct FShaderParamTypeDefinition
@@ -240,8 +238,8 @@ public:
 
 		struct
 		{
-			uint8 MatrixColumnCount : 2;
-			uint8 MatrixRowCount : 2;
+			uint8 MatrixRowCount : 3;
+			uint8 MatrixColumnCount : 3;
 		};
 	};
 
@@ -295,3 +293,23 @@ public:
 	void ResetTypeDeclaration(
 		);
 };
+
+/* Describes a shader function signature. */
+USTRUCT()
+struct FShaderFunctionDefinition
+{
+	GENERATED_BODY()
+
+public:
+	/** Function name. */
+	UPROPERTY(VisibleAnywhere, Category = "Kernel")
+	FString	Name;
+
+	/** Function parameter types. */
+	UPROPERTY(VisibleAnywhere, Category = "Kernel")
+	TArray<FShaderParamTypeDefinition> ParamTypes;
+
+	/** If set to true then the first element of ParamTypes is the return type. Otherwise return type is void. */
+	UPROPERTY()
+	bool bHasReturnType;
+};
diff --git a/Engine/Source/Runtime/Engine/Public/SceneInterface.h b/Engine/Source/Runtime/Engine/Public/SceneInterface.h
--- a/Engine/Source/Runtime/Engine/Public/SceneInterface.h
+++ b/Engine/Source/Runtime/Engine/Public/SceneInterface.h
@@ -33,7 +33,7 @@ class UStaticMeshComponent;
 class UTextureCube;
 class FViewInfo;
 class FSceneRenderer;
-class FComputeFramework;
+class FComputeGraphScheduler;
 class FInstanceCullingManager;
 struct FHairStrandsInstance;
 
@@ -451,7 +451,7 @@ public:
 		return nullptr;
 	}
 
-	virtual FComputeFramework* GetComputeFramework()
+	virtual FComputeGraphScheduler* GetComputeGraphScheduler()
 	{
 		return nullptr;
 	}
diff --git a/Engine/Source/Runtime/RenderCore/Private/ShaderParameterMetadata.cpp b/Engine/Source/Runtime/RenderCore/Private/ShaderParameterMetadata.cpp
--- a/Engine/Source/Runtime/RenderCore/Private/ShaderParameterMetadata.cpp
+++ b/Engine/Source/Runtime/RenderCore/Private/ShaderParameterMetadata.cpp
@@ -934,6 +934,54 @@ FString FShaderParametersMetadata::GetFullMemberCodeName(uint16 MemberOffset) co
 	return MemberName;
 }
 
+void FShaderParametersMetadataBuilder::AddBufferSRV(
+	const TCHAR* Name,
+	const TCHAR* ShaderType,
+	EShaderPrecisionModifier::Type Precision /* = EShaderPrecisionModifier::Float */
+)
+{
+	NextMemberOffset = Align(NextMemberOffset, SHADER_PARAMETER_POINTER_ALIGNMENT);
+
+	new(Members) FShaderParametersMetadata::FMember(
+		Name,
+		ShaderType,
+		__LINE__,
+		NextMemberOffset,
+		UBMT_SRV,
+		Precision,
+		TShaderResourceParameterTypeInfo<FRHIShaderResourceView*>::NumRows,
+		TShaderResourceParameterTypeInfo<FRHIShaderResourceView*>::NumColumns,
+		TShaderResourceParameterTypeInfo<FRHIShaderResourceView*>::NumElements,
+		TShaderResourceParameterTypeInfo<FRHIShaderResourceView*>::GetStructMetadata()
+	);
+
+	NextMemberOffset += SHADER_PARAMETER_POINTER_ALIGNMENT;
+}
+
+void FShaderParametersMetadataBuilder::AddBufferUAV(
+	const TCHAR* Name,
+	const TCHAR* ShaderType,
+	EShaderPrecisionModifier::Type Precision /* = EShaderPrecisionModifier::Float */
+)
+{
+	NextMemberOffset = Align(NextMemberOffset, SHADER_PARAMETER_POINTER_ALIGNMENT);
+
+	new(Members) FShaderParametersMetadata::FMember(
+		Name,
+		ShaderType,
+		__LINE__,
+		NextMemberOffset,
+		UBMT_UAV,
+		Precision,
+		TShaderResourceParameterTypeInfo<FRHIUnorderedAccessView*>::NumRows,
+		TShaderResourceParameterTypeInfo<FRHIUnorderedAccessView*>::NumColumns,
+		TShaderResourceParameterTypeInfo<FRHIUnorderedAccessView*>::NumElements,
+		TShaderResourceParameterTypeInfo<FRHIUnorderedAccessView*>::GetStructMetadata()
+	);
+
+	NextMemberOffset += SHADER_PARAMETER_POINTER_ALIGNMENT;
+}
+
 void FShaderParametersMetadataBuilder::AddRDGBufferSRV(
 	const TCHAR* Name,
 	const TCHAR* ShaderType,
@@ -949,11 +997,11 @@ void FShaderParametersMetadataBuilder::AddRDGBufferSRV(
 		NextMemberOffset,
 		UBMT_RDG_BUFFER_SRV,
 		Precision,
-		TShaderResourceParameterTypeInfo<FRDGBufferSRV*>::NumRows,
-		TShaderResourceParameterTypeInfo<FRDGBufferSRV*>::NumColumns,
-		TShaderResourceParameterTypeInfo<FRDGBufferSRV*>::NumElements,
-		TShaderResourceParameterTypeInfo<FRDGBufferSRV*>::GetStructMetadata()
-		);
+ 		TShaderResourceParameterTypeInfo<FRDGBufferSRV*>::NumRows,
+ 		TShaderResourceParameterTypeInfo<FRDGBufferSRV*>::NumColumns,
+ 		TShaderResourceParameterTypeInfo<FRDGBufferSRV*>::NumElements,
+ 		TShaderResourceParameterTypeInfo<FRDGBufferSRV*>::GetStructMetadata()
+	);
 
 	NextMemberOffset += SHADER_PARAMETER_POINTER_ALIGNMENT;
 }
@@ -973,10 +1021,10 @@ void FShaderParametersMetadataBuilder::AddRDGBufferUAV(
 		NextMemberOffset,
 		UBMT_RDG_BUFFER_UAV,
 		Precision,
-		TShaderResourceParameterTypeInfo<FRDGBufferUAV*>::NumRows,
-		TShaderResourceParameterTypeInfo<FRDGBufferUAV*>::NumColumns,
-		TShaderResourceParameterTypeInfo<FRDGBufferUAV*>::NumElements,
-		TShaderResourceParameterTypeInfo<FRDGBufferUAV*>::GetStructMetadata()
+ 		TShaderResourceParameterTypeInfo<FRDGBufferUAV*>::NumRows,
+ 		TShaderResourceParameterTypeInfo<FRDGBufferUAV*>::NumColumns,
+ 		TShaderResourceParameterTypeInfo<FRDGBufferUAV*>::NumElements,
+ 		TShaderResourceParameterTypeInfo<FRDGBufferUAV*>::GetStructMetadata()
 		);
 
 	NextMemberOffset += SHADER_PARAMETER_POINTER_ALIGNMENT;
diff --git a/Engine/Source/Runtime/RenderCore/Public/ShaderParameterMetadataBuilder.h b/Engine/Source/Runtime/RenderCore/Public/ShaderParameterMetadataBuilder.h
--- a/Engine/Source/Runtime/RenderCore/Public/ShaderParameterMetadataBuilder.h
+++ b/Engine/Source/Runtime/RenderCore/Public/ShaderParameterMetadataBuilder.h
@@ -38,6 +38,44 @@ public:
 		NextMemberOffset += sizeof(typename TParamTypeInfo::TAlignedType);
 	}
 
+	template<typename T>
+	void AddNestedStruct(
+		const TCHAR* Name,
+		EShaderPrecisionModifier::Type Precision = EShaderPrecisionModifier::Float
+		)
+	{
+		using TParamTypeInfo = TShaderParameterStructTypeInfo<T>;
+
+		NextMemberOffset = Align(NextMemberOffset, TParamTypeInfo::Alignment);
+
+		new(Members) FShaderParametersMetadata::FMember(
+			Name,
+			TEXT(""),
+			__LINE__,
+			NextMemberOffset,
+			UBMT_NESTED_STRUCT,
+			Precision,
+			TParamTypeInfo::NumRows,
+			TParamTypeInfo::NumColumns,
+			TParamTypeInfo::NumElements,
+			TParamTypeInfo::GetStructMetadata()
+		);
+
+		NextMemberOffset += sizeof(typename TParamTypeInfo::TAlignedType);
+	}
+
+	void AddBufferSRV(
+		const TCHAR* Name,
+		const TCHAR* ShaderType,
+		EShaderPrecisionModifier::Type Precision = EShaderPrecisionModifier::Float
+		);
+
+	void AddBufferUAV(
+		const TCHAR* Name,
+		const TCHAR* ShaderType,
+		EShaderPrecisionModifier::Type Precision = EShaderPrecisionModifier::Float
+		);
+
 	void AddRDGBufferSRV(
 		const TCHAR* Name,
 		const TCHAR* ShaderType,
diff --git a/Engine/Source/Runtime/Renderer/Private/RendererScene.cpp b/Engine/Source/Runtime/Renderer/Private/RendererScene.cpp
--- a/Engine/Source/Runtime/Renderer/Private/RendererScene.cpp
+++ b/Engine/Source/Runtime/Renderer/Private/RendererScene.cpp
@@ -56,6 +56,7 @@
 #include "Engine/StaticMesh.h"
 #include "GPUSkinCache.h"
 #include "ComputeFramework/ComputeFramework.h"
+#include "ComputeFramework/ComputeGraphScheduler.h"
 #include "DynamicShadowMapChannelBindingHelper.h"
 #include "GPUScene.h"
 #include "HAL/LowLevelMemTracker.h"
@@ -1083,7 +1084,7 @@ FScene::FScene(UWorld* InWorld, bool bInRequiresHitProxies, bool bInIsEditorScen
 
 	if (SupportsComputeFramework(InFeatureLevel, GetFeatureLevelShaderPlatform(InFeatureLevel)))
 	{
-		ComputeFramework = new FComputeFramework();
+		ComputeGraphScheduler = new FComputeGraphScheduler();
 	}
 
 #if RHI_RAYTRACING
diff --git a/Engine/Source/Runtime/Renderer/Private/ScenePrivate.h b/Engine/Source/Runtime/Renderer/Private/ScenePrivate.h
--- a/Engine/Source/Runtime/Renderer/Private/ScenePrivate.h
+++ b/Engine/Source/Runtime/Renderer/Private/ScenePrivate.h
@@ -88,7 +88,7 @@ class FRHIGPUTextureReadback;
 class FRuntimeVirtualTextureSceneProxy;
 class FLumenSceneData;
 class FVirtualShadowMapArrayCacheManager;
-class FComputeFramework;
+class FComputeGraphScheduler;
 struct FHairStrandsInstance;
 struct FPathTracingConfig;
 
@@ -2869,7 +2869,7 @@ public:
 	class FGPUSkinCache* GPUSkinCache;
 
 	/* Root system that handles runtime execution of Compute Framework tasks */
-	FComputeFramework* ComputeFramework = nullptr;
+	FComputeGraphScheduler* ComputeGraphScheduler = nullptr;
 
 	/** Uniform buffers for parameter collections with the corresponding Ids. */
 	TMap<FGuid, FUniformBufferRHIRef> ParameterCollections;
@@ -3086,9 +3086,9 @@ public:
 		return GPUSkinCache;
 	}
 
-	FComputeFramework* GetComputeFramework() override
+	FComputeGraphScheduler* GetComputeGraphScheduler() override
 	{
-		return ComputeFramework;
+		return ComputeGraphScheduler;
 	}
 
 #if RHI_RAYTRACING
