Modified WorldPartitionCookPackageSplitter to fully prepare runtime streaming cells in PreSaveGeneratorPackage so that TryPopulatePackage only fills the generated package without altering the cell.
This was necessary for CookOnTheFly to work as the Generator package (in our case the world package) can be saved prior to any of its generated packages (uses the deferred populate mode).

diff --git a/Engine/Source/Runtime/Engine/Private/WorldPartition/WorldPartition.cpp b/Engine/Source/Runtime/Engine/Private/WorldPartition/WorldPartition.cpp
--- a/Engine/Source/Runtime/Engine/Private/WorldPartition/WorldPartition.cpp
+++ b/Engine/Source/Runtime/Engine/Private/WorldPartition/WorldPartition.cpp
@@ -1203,13 +1203,16 @@ bool UWorldPartition::PopulateGeneratedPackageForCook(UPackage* InPackage, const
 	return RuntimeHash->PopulateGeneratedPackageForCook(InPackage, InPackageRelativePath);
 }
 
-void UWorldPartition::FinalizeGeneratedPackageForCook()
+bool UWorldPartition::FinalizeGeneratorPackageForCook(const TArray<ICookPackageSplitter::FGeneratedPackageForPreSave>& InGeneratedPackages)
 {
 	check(RuntimeHash);
-	RuntimeHash->FinalizeGeneratedPackageForCook();
-
-	// Apply remapping of Persistent Level's SoftObjectPaths
-	FWorldPartitionLevelHelper::RemapLevelSoftObjectPaths(World->PersistentLevel, this);
+	if (RuntimeHash->FinalizeGeneratorPackageForCook(InGeneratedPackages))
+	{
+		// Apply remapping of Persistent Level's SoftObjectPaths
+		FWorldPartitionLevelHelper::RemapLevelSoftObjectPaths(World->PersistentLevel, this);
+		return true;
+	}
+	return false;
 }
 
 void UWorldPartition::SavePerUserSettings()
diff --git a/Engine/Source/Runtime/Engine/Private/WorldPartition/WorldPartitionCookPackageSplitter.cpp b/Engine/Source/Runtime/Engine/Private/WorldPartition/WorldPartitionCookPackageSplitter.cpp
--- a/Engine/Source/Runtime/Engine/Private/WorldPartition/WorldPartitionCookPackageSplitter.cpp
+++ b/Engine/Source/Runtime/Engine/Private/WorldPartition/WorldPartitionCookPackageSplitter.cpp
@@ -75,11 +75,11 @@ bool FWorldPartitionCookPackageSplitter::TryPopulatePackage(const UPackage* Owne
 }
 
 void FWorldPartitionCookPackageSplitter::PreSaveGeneratorPackage(UPackage* OwnerPackage, UObject* OwnerObject,
-	const TArray<FGeneratedPackageForPreSave>& PlaceholderPackages)
+	const TArray<FGeneratedPackageForPreSave>& GeneratedPackages)
 {
 	UWorld* PartitionedWorld = ValidateDataObject(OwnerObject);
 	UWorldPartition* WorldPartition = PartitionedWorld->PersistentLevel->GetWorldPartition();
-	WorldPartition->FinalizeGeneratedPackageForCook();
+	WorldPartition->FinalizeGeneratorPackageForCook(GeneratedPackages);
 }
 
 #endif
\ No newline at end of file
diff --git a/Engine/Source/Runtime/Engine/Private/WorldPartition/WorldPartitionCookPackageSplitter.h b/Engine/Source/Runtime/Engine/Private/WorldPartition/WorldPartitionCookPackageSplitter.h
--- a/Engine/Source/Runtime/Engine/Private/WorldPartition/WorldPartitionCookPackageSplitter.h
+++ b/Engine/Source/Runtime/Engine/Private/WorldPartition/WorldPartitionCookPackageSplitter.h
@@ -18,7 +18,7 @@ public:
 	virtual bool TryPopulatePackage(const UPackage* OwnerPackage, const UObject* OwnerObject,
 		const ICookPackageSplitter::FGeneratedPackageForPopulate& GeneratedPackage) override;
 	virtual void PreSaveGeneratorPackage(UPackage* OwnerPackage, UObject* OwnerObject,
-		const TArray<ICookPackageSplitter::FGeneratedPackageForPreSave>& PlaceholderPackages) override;
+		const TArray<ICookPackageSplitter::FGeneratedPackageForPreSave>& GeneratedPackages) override;
 	//~ End of ICookPackageSplitter
 
 private:
diff --git a/Engine/Source/Runtime/Engine/Private/WorldPartition/WorldPartitionRuntimeLevelStreamingCell.cpp b/Engine/Source/Runtime/Engine/Private/WorldPartition/WorldPartitionRuntimeLevelStreamingCell.cpp
--- a/Engine/Source/Runtime/Engine/Private/WorldPartition/WorldPartitionRuntimeLevelStreamingCell.cpp
+++ b/Engine/Source/Runtime/Engine/Private/WorldPartition/WorldPartitionRuntimeLevelStreamingCell.cpp
@@ -147,6 +147,25 @@ void UWorldPartitionRuntimeLevelStreamingCell::MoveAlwaysLoadedContentToPersiste
 	}
 }
 
+// Do all necessary work to prepare cell object for cook.
+bool UWorldPartitionRuntimeLevelStreamingCell::PrepareCellForCook(UPackage* InPackage)
+{
+	// LevelStreaming could already be created
+	if (!LevelStreaming && GetActorCount() > 0)
+	{
+		if (!InPackage)
+		{
+			return false;
+		}
+
+		ULevelStreaming* NewLevelStreaming = CreateLevelStreaming(InPackage->GetName());
+		LevelStreaming = Cast<UWorldPartitionLevelStreamingDynamic>(NewLevelStreaming);
+		check(LevelStreaming);
+	}
+
+	return true;
+}
+
 bool UWorldPartitionRuntimeLevelStreamingCell::PopulateGeneratedPackageForCook(UPackage* InPackage)
 {
 	check(!IsAlwaysLoaded());
@@ -157,18 +176,22 @@ bool UWorldPartitionRuntimeLevelStreamingCell::PopulateGeneratedPackageForCook(U
 
 	if (GetActorCount() > 0)
 	{
-		UWorldPartition* WorldPartition = GetOuterUWorldPartition();
-		UWorld* OuterWorld = WorldPartition->GetTypedOuter<UWorld>();
-		ULevelStreaming* NewLevelStreaming = CreateLevelStreaming(InPackage->GetName());
-		check(NewLevelStreaming)
+		// When cook splitter doesn't use deferred populate, cell needs to be prepared here.
+		if (!PrepareCellForCook(InPackage))
+		{
+			return false;
+		}
 
 		// Load cell Actors
 		LoadActorsForCook();
-		
-		LevelStreaming = Cast<UWorldPartitionLevelStreamingDynamic>(NewLevelStreaming);
-		ULevel* NewLevel = FWorldPartitionLevelHelper::CreateEmptyLevelForRuntimeCell(OuterWorld, NewLevelStreaming->GetWorldAsset().ToString(), InPackage);
+
+		// Create a level and move these actors in it
+		UWorldPartition* WorldPartition = GetOuterUWorldPartition();
+		UWorld* OuterWorld = WorldPartition->GetTypedOuter<UWorld>();
+		ULevel* NewLevel = FWorldPartitionLevelHelper::CreateEmptyLevelForRuntimeCell(OuterWorld, LevelStreaming->GetWorldAsset().ToString(), InPackage);
 		check(NewLevel->GetPackage() == InPackage);
 		FWorldPartitionLevelHelper::MoveExternalActorsToLevel(Packages, NewLevel);
+
 		// Remap Level's SoftObjectPaths
 		FWorldPartitionLevelHelper::RemapLevelSoftObjectPaths(NewLevel, WorldPartition);
 	}
diff --git a/Engine/Source/Runtime/Engine/Private/WorldPartition/WorldPartitionRuntimeSpatialHash.cpp b/Engine/Source/Runtime/Engine/Private/WorldPartition/WorldPartitionRuntimeSpatialHash.cpp
--- a/Engine/Source/Runtime/Engine/Private/WorldPartition/WorldPartitionRuntimeSpatialHash.cpp
+++ b/Engine/Source/Runtime/Engine/Private/WorldPartition/WorldPartitionRuntimeSpatialHash.cpp
@@ -36,6 +36,7 @@
 
 #if WITH_EDITOR
 #include "Editor/EditorEngine.h"
+#include "CookPackageSplitter.h"
 #include "Misc/Parse.h"
 
 #include "Algo/Find.h"
@@ -935,7 +936,7 @@ bool UWorldPartitionRuntimeSpatialHash::CreateStreamingGrid(const FSpatialHashRu
 							check(!PackageRelativePath.IsEmpty());
 							OutPackagesToGenerate->Add(PackageRelativePath);
 
-							// Map relative package to StreamingCell for PopulateGeneratedPackageForCook/FinalizeGeneratedPackageForCook
+							// Map relative package to StreamingCell for PopulateGeneratedPackageForCook/FinalizeGeneratorPackageForCook
 							PackagesToGenerateForCook.Add(PackageRelativePath, StreamingCell);
 						}
 					}
@@ -960,16 +961,18 @@ bool UWorldPartitionRuntimeSpatialHash::PopulateGeneratedPackageForCook(UPackage
 	return false;
 }
 
-void UWorldPartitionRuntimeSpatialHash::FinalizeGeneratedPackageForCook()
+bool UWorldPartitionRuntimeSpatialHash::FinalizeGeneratorPackageForCook(const TArray<ICookPackageSplitter::FGeneratedPackageForPreSave>& InGeneratedPackages)
 {
-	for (const auto& Package : PackagesToGenerateForCook)
+	for (const ICookPackageSplitter::FGeneratedPackageForPreSave& GeneratedPackage : InGeneratedPackages)
 	{
-		UWorldPartitionRuntimeCell* Cell = Package.Value;
-		if (ensure(Cell))
+		UWorldPartitionRuntimeCell** MatchingCell = PackagesToGenerateForCook.Find(GeneratedPackage.RelativePath);
+		UWorldPartitionRuntimeCell* Cell = MatchingCell ? *MatchingCell : nullptr;
+		if (!Cell || !Cell->PrepareCellForCook(GeneratedPackage.Package))
 		{
-			Cell->FinalizeGeneratedPackageForCook();
+			return false;
 		}
 	}
+	return true;
 }
 
 void UWorldPartitionRuntimeSpatialHash::FlushStreaming()
diff --git a/Engine/Source/Runtime/Engine/Public/WorldPartition/WorldPartition.h b/Engine/Source/Runtime/Engine/Public/WorldPartition/WorldPartition.h
--- a/Engine/Source/Runtime/Engine/Public/WorldPartition/WorldPartition.h
+++ b/Engine/Source/Runtime/Engine/Public/WorldPartition/WorldPartition.h
@@ -13,6 +13,7 @@
 
 #if WITH_EDITOR
 #include "PackageSourceControlHelper.h"
+#include "CookPackageSplitter.h"
 #endif
 
 #include "WorldPartition.generated.h"
@@ -144,7 +145,7 @@ public:
 
 	// Cook Methods
 	bool PopulateGeneratedPackageForCook(UPackage* InPackage, const FString& InPackageRelativePath);
-	void FinalizeGeneratedPackageForCook();
+	bool FinalizeGeneratorPackageForCook(const TArray<ICookPackageSplitter::FGeneratedPackageForPreSave>& InGeneratedPackages);
 
 	FBox GetWorldBounds() const;
 	FBox GetEditorWorldBounds() const;
diff --git a/Engine/Source/Runtime/Engine/Public/WorldPartition/WorldPartitionRuntimeCell.h b/Engine/Source/Runtime/Engine/Public/WorldPartition/WorldPartitionRuntimeCell.h
--- a/Engine/Source/Runtime/Engine/Public/WorldPartition/WorldPartitionRuntimeCell.h
+++ b/Engine/Source/Runtime/Engine/Public/WorldPartition/WorldPartitionRuntimeCell.h
@@ -145,9 +145,9 @@ class UWorldPartitionRuntimeCell : public UObject
 	virtual int32 GetActorCount() const PURE_VIRTUAL(UWorldPartitionRuntimeCell::GetActorCount, return 0;);
 
 	// Cook methods
+	virtual bool PrepareCellForCook(UPackage* InPackage) { return false; }
 	virtual bool PopulateGeneratedPackageForCook(UPackage* InPackage) PURE_VIRTUAL(UWorldPartitionRuntimeCell::PopulateGeneratedPackageForCook, return false;);
 	virtual void MoveAlwaysLoadedContentToPersistentLevel() PURE_VIRTUAL(UWorldPartitionRuntimeCell::MoveAlwaysLoadedContentToPersistentLevel);
-	virtual void FinalizeGeneratedPackageForCook() {}
 	virtual FString GetPackageNameToCreate() const PURE_VIRTUAL(UWorldPartitionRuntimeCell::GetPackageNameToCreate, return FString(""););
 #endif
 
diff --git a/Engine/Source/Runtime/Engine/Public/WorldPartition/WorldPartitionRuntimeHash.h b/Engine/Source/Runtime/Engine/Public/WorldPartition/WorldPartitionRuntimeHash.h
--- a/Engine/Source/Runtime/Engine/Public/WorldPartition/WorldPartitionRuntimeHash.h
+++ b/Engine/Source/Runtime/Engine/Public/WorldPartition/WorldPartitionRuntimeHash.h
@@ -10,6 +10,9 @@
 #include "WorldPartition/ActorDescList.h"
 #include "WorldPartition/WorldPartitionHandle.h"
 #include "WorldPartition/WorldPartitionActorDescViewProxy.h"
+#if WITH_EDITOR
+#include "CookPackageSplitter.h"
+#endif
 #include "WorldPartitionRuntimeHash.generated.h"
 
 class UWorldPartitionRuntimeCell;
@@ -27,7 +30,7 @@ class ENGINE_API UWorldPartitionRuntimeHash : public UObject
 	virtual void SetDefaultValues() {}
 	virtual void ImportFromWorldComposition(class UWorldComposition* WorldComposition) {}
 	virtual bool PopulateGeneratedPackageForCook(UPackage* InPackage, const FString& InPackageRelativePath) { return false; }
-	virtual void FinalizeGeneratedPackageForCook() {}
+	virtual bool FinalizeGeneratorPackageForCook(const TArray<ICookPackageSplitter::FGeneratedPackageForPreSave>& InGeneratedPackages) { return false; }
 	virtual void FlushStreaming() {}
 	virtual bool GenerateHLOD(ISourceControlHelper* SourceControlHelper, bool bCreateActorsOnly) { return false; }
 	virtual bool GenerateNavigationData() { return false; }
diff --git a/Engine/Source/Runtime/Engine/Public/WorldPartition/WorldPartitionRuntimeLevelStreamingCell.h b/Engine/Source/Runtime/Engine/Public/WorldPartition/WorldPartitionRuntimeLevelStreamingCell.h
--- a/Engine/Source/Runtime/Engine/Public/WorldPartition/WorldPartitionRuntimeLevelStreamingCell.h
+++ b/Engine/Source/Runtime/Engine/Public/WorldPartition/WorldPartitionRuntimeLevelStreamingCell.h
@@ -34,6 +34,7 @@ class UWorldPartitionRuntimeLevelStreamingCell : public UWorldPartitionRuntimeSp
 	virtual void AddActorToCell(const FWorldPartitionActorDescView& ActorDescView, uint64 InContainerID, const FTransform& InContainerTransform, const UActorDescContainer* InContainer) override;
 	virtual int32 GetActorCount() const override;
 	// Cook methods
+	virtual bool PrepareCellForCook(UPackage* InPackage) override;
 	virtual bool PopulateGeneratedPackageForCook(UPackage* InPackage) override;
 	virtual FString GetPackageNameToCreate() const override;
 	//~End UWorldPartitionRuntimeCell Interface
diff --git a/Engine/Source/Runtime/Engine/Public/WorldPartition/WorldPartitionRuntimeSpatialHash.h b/Engine/Source/Runtime/Engine/Public/WorldPartition/WorldPartitionRuntimeSpatialHash.h
--- a/Engine/Source/Runtime/Engine/Public/WorldPartition/WorldPartitionRuntimeSpatialHash.h
+++ b/Engine/Source/Runtime/Engine/Public/WorldPartition/WorldPartitionRuntimeSpatialHash.h
@@ -203,7 +203,7 @@ public:
 	virtual void SetDefaultValues() override;
 	virtual void ImportFromWorldComposition(class UWorldComposition* WorldComposition) override;
 	virtual bool PopulateGeneratedPackageForCook(UPackage* InPackage, const FString& InPackageRelativePath) override;
-	virtual void FinalizeGeneratedPackageForCook() override;
+	virtual bool FinalizeGeneratorPackageForCook(const TArray<ICookPackageSplitter::FGeneratedPackageForPreSave>& InGeneratedPackages) override;
 	virtual void FlushStreaming() override;
 	virtual bool GenerateHLOD(ISourceControlHelper* SourceControlHelper, bool bCreateActorsOnly) override;
 	virtual bool GenerateNavigationData() override;
