Restore bulk data loading logic in TextureDerivedDataTask while avoiding exporting empty composite texture bulk data when exporting texture builds.

diff --git a/Engine/Source/Runtime/Engine/Private/TextureDerivedDataTask.cpp b/Engine/Source/Runtime/Engine/Private/TextureDerivedDataTask.cpp
--- a/Engine/Source/Runtime/Engine/Private/TextureDerivedDataTask.cpp
+++ b/Engine/Source/Runtime/Engine/Private/TextureDerivedDataTask.cpp
@@ -619,22 +619,37 @@ void FTextureCacheDerivedDataWorker::DoWork()
 	
 	if (!bSucceeded && bAllowAsyncBuild)
 	{
+		bool bHasTextureSourceMips = false;
 		if (TextureData.IsValid() && Texture.Source.IsBulkDataLoaded())
 		{
+			TextureData.GetSourceMips(Texture.Source, ImageWrapper);
 			BuildExporter.ExportTextureSourceBulkData(Texture.Source);
-		}
-		else if (bAllowAsyncLoading)
-		{
-			BuildExporter.ExportTextureSourceBulkData(TextureData.AsyncSource);
+			bHasTextureSourceMips = true;
 		}
 
+		bool bHasCompositeTextureSourceMips = false;
 		if (CompositeTextureData.IsValid() && Texture.CompositeTexture && Texture.CompositeTexture->Source.IsBulkDataLoaded())
 		{
+			CompositeTextureData.GetSourceMips(Texture.CompositeTexture->Source, ImageWrapper);
 			BuildExporter.ExportCompositeTextureSourceBulkData(Texture.CompositeTexture->Source);
+			bHasCompositeTextureSourceMips = true;
 		}
-		else if (bAllowAsyncLoading)
+
+		if (bAllowAsyncLoading && !bHasTextureSourceMips)
 		{
-			BuildExporter.ExportCompositeTextureSourceBulkData(CompositeTextureData.AsyncSource);
+			TextureData.GetAsyncSourceMips(ImageWrapper);
+			BuildExporter.ExportTextureSourceBulkData(TextureData.AsyncSource);
+			TextureData.AsyncSource.RemoveBulkData();
+		}
+
+		if (bAllowAsyncLoading && !bHasCompositeTextureSourceMips)
+		{
+			CompositeTextureData.GetAsyncSourceMips(ImageWrapper);
+			if ((bool)Texture.CompositeTexture)
+			{
+				BuildExporter.ExportCompositeTextureSourceBulkData(CompositeTextureData.AsyncSource);
+			}
+			CompositeTextureData.AsyncSource.RemoveBulkData();
 		}
 
 		if (TextureData.Blocks.Num() && TextureData.Blocks[0].MipsPerLayer.Num() && TextureData.Blocks[0].MipsPerLayer[0].Num() && 
