Fixed missing includes reported in non-unity builds in EngineAnalyticsSessionSummary.cpp.

diff --git a/Engine/Source/Runtime/Engine/Private/Analytics/EngineAnalyticsSessionSummary.cpp b/Engine/Source/Runtime/Engine/Private/Analytics/EngineAnalyticsSessionSummary.cpp
--- a/Engine/Source/Runtime/Engine/Private/Analytics/EngineAnalyticsSessionSummary.cpp
+++ b/Engine/Source/Runtime/Engine/Private/Analytics/EngineAnalyticsSessionSummary.cpp
@@ -4,6 +4,7 @@
 #include "AnalyticsSessionSummaryManager.h"
 #include "EngineGlobals.h"
 #include "GeneralProjectSettings.h"
+#include "Engine/Engine.h"
 #include "GenericPlatform/GenericPlatformCrashContext.h"
 #include "HAL/ExceptionHandling.h"
 #include "HAL/FileManager.h"
@@ -12,8 +13,13 @@
 #include "Interfaces/IAnalyticsPropertyStore.h"
 #include "Interfaces/IPluginManager.h"
 #include "Interfaces/IProjectManager.h"
+#include "Misc/App.h"
+#include "Misc/CommandLine.h"
+#include "Misc/CoreDelegates.h"
 #include "Misc/EngineVersion.h"
 #include "UserActivityTracking.h"
+#include "RHI.h"
+#include "DynamicRHI.h"
 
 extern ENGINE_API float GAverageFPS;
 extern CORE_API bool GIsGPUCrashed;
