Corrected typos following HLOD doc review

diff --git a/Engine/Source/Runtime/Engine/Classes/Engine/MeshMerging.h b/Engine/Source/Runtime/Engine/Classes/Engine/MeshMerging.h
--- a/Engine/Source/Runtime/Engine/Classes/Engine/MeshMerging.h
+++ b/Engine/Source/Runtime/Engine/Classes/Engine/MeshMerging.h
@@ -214,7 +214,7 @@ struct FMeshProxySettings
 	int32 ScreenSize;
 
 	/** Override when converting multiple meshes for proxy LOD merging. Warning, large geometry with small sampling has very high memory costs*/
-	UPROPERTY(BlueprintReadWrite, EditAnywhere, AdvancedDisplay, Category = ProxySettings, meta = (EditCondition = "bOverrideVoxelSize", ClampMin = "0.1", DisplayName = "Overide Spatial Sampling Distance"))
+	UPROPERTY(BlueprintReadWrite, EditAnywhere, AdvancedDisplay, Category = ProxySettings, meta = (EditCondition = "bOverrideVoxelSize", ClampMin = "0.1", DisplayName = "Override Spatial Sampling Distance"))
 	float VoxelSize;
 
 	/** Material simplification */
@@ -260,7 +260,7 @@ struct FMeshProxySettings
 	float HardAngleThreshold;
 
 	/** Lightmap resolution */
-	UPROPERTY(BlueprintReadWrite, EditAnywhere, Category = ProxySettings, meta = (ClampMin = 32, ClampMax = 4096, EditCondition = "!bComputeLightMapResolution", DisplayAfter="NormalCalculationMethod"))
+	UPROPERTY(BlueprintReadWrite, EditAnywhere, Category = ProxySettings, meta = (ClampMin = 32, ClampMax = 4096, EditCondition = "!bComputeLightMapResolution", DisplayAfter="NormalCalculationMethod", DisplayName="Lightmap Resolution"))
 	int32 LightMapResolution;
 
 	/** Controls the method used to calculate the normal for the simplified geometry */
@@ -288,7 +288,7 @@ struct FMeshProxySettings
 	uint8 bUseHardAngleThreshold:1;
 
 	/** If ticked will compute the lightmap resolution by summing the dimensions for each mesh included for merging */
-	UPROPERTY(BlueprintReadWrite, EditAnywhere, Category = ProxySettings)
+	UPROPERTY(BlueprintReadWrite, EditAnywhere, Category = ProxySettings, meta = (DisplayName="Compute Lightmap Resolution"))
 	uint8 bComputeLightMapResolution:1;
 
 	/** Whether Simplygon should recalculate normals, otherwise the normals channel will be sampled from the original mesh */
@@ -450,7 +450,7 @@ struct FMeshMergingSettings
 	GENERATED_USTRUCT_BODY()
 
 	/** The lightmap resolution used both for generating lightmap UV coordinates, and also set on the generated static mesh */
-	UPROPERTY(EditAnywhere, AdvancedDisplay, BlueprintReadWrite, Category = MeshSettings, meta=(ClampMax = 4096, EditCondition = "!bComputedLightMapResolution", DisplayAfter="bGenerateLightMapUV"))
+	UPROPERTY(EditAnywhere, AdvancedDisplay, BlueprintReadWrite, Category = MeshSettings, meta=(ClampMax = 4096, EditCondition = "!bComputedLightMapResolution", DisplayAfter="bGenerateLightMapUV", DisplayName="Target Lightmap Resolution"))
 	int32 TargetLightMapResolution;
 
 	/** Whether to output the specified UV channels into the merged mesh (only if the source meshes contain valid UVs for the specified channel) */
@@ -474,11 +474,11 @@ struct FMeshMergingSettings
 	EMeshLODSelectionType LODSelectionType;
 
 	/** Whether to generate lightmap UVs for a merged mesh*/
-	UPROPERTY(EditAnywhere, AdvancedDisplay, BlueprintReadWrite, Category = MeshSettings)
+	UPROPERTY(EditAnywhere, AdvancedDisplay, BlueprintReadWrite, Category = MeshSettings, meta=(DisplayName="Generate Lightmap UV"))
 	uint8 bGenerateLightMapUV:1;
 
 	/** Whether or not the lightmap resolution should be computed by summing the lightmap resolutions for the input Mesh Components */
-	UPROPERTY(EditAnywhere, AdvancedDisplay, BlueprintReadWrite, Category = MeshSettings)
+	UPROPERTY(EditAnywhere, AdvancedDisplay, BlueprintReadWrite, Category = MeshSettings, meta=(DisplayName="Computed Lightmap Resolution"))
 	uint8 bComputedLightMapResolution:1;
 
 	/** Whether merged mesh should have pivot at world origin, or at first merged component otherwise */
