Deprecated AddInstanceWorldSpace in favor of AddInstance or AddInstances with bWorldSpace set to true

diff --git a/Engine/Plugins/Experimental/Water/Source/Editor/Private/CausticsGeneratorActor.cpp b/Engine/Plugins/Experimental/Water/Source/Editor/Private/CausticsGeneratorActor.cpp
--- a/Engine/Plugins/Experimental/Water/Source/Editor/Private/CausticsGeneratorActor.cpp
+++ b/Engine/Plugins/Experimental/Water/Source/Editor/Private/CausticsGeneratorActor.cpp
@@ -99,7 +99,7 @@ void ACausticsGeneratorActor::SpawnCausticParticleGrid(UHierarchicalInstancedSta
 			PolyCenterVector += HalfVector * CellSizeF;
 			PolyCenterVector += ActorLocation;
 
-			HISMC->AddInstanceWorldSpace(FTransform(FRotator::ZeroRotator, PolyCenterVector, FVector(CellSizeVector.X, CellSizeVector.Y, 0.0001)));
+			HISMC->AddInstance(FTransform(FRotator::ZeroRotator, PolyCenterVector, FVector(CellSizeVector.X, CellSizeVector.Y, 0.0001)), /*bWorldSpace*/true);
 		}
 	}
 }
\ No newline at end of file
diff --git a/Engine/Source/Developer/WorldPartitionHLODUtilities/Private/HLODBuilderInstancing.cpp b/Engine/Source/Developer/WorldPartitionHLODUtilities/Private/HLODBuilderInstancing.cpp
--- a/Engine/Source/Developer/WorldPartitionHLODUtilities/Private/HLODBuilderInstancing.cpp
+++ b/Engine/Source/Developer/WorldPartitionHLODUtilities/Private/HLODBuilderInstancing.cpp
@@ -91,6 +91,7 @@ TArray<UPrimitiveComponent*> FHLODBuilder_Instancing::CreateComponents(AWorldPar
 		DisableCollisions(Component);
 
 		// Add all instances
+		TArray<FTransform> InstanceTransforms;
 		for (UPrimitiveComponent* SMC : EntryComponents)
 		{
 			// If we have an ISMC, retrieve all instances
@@ -100,14 +101,15 @@ TArray<UPrimitiveComponent*> FHLODBuilder_Instancing::CreateComponents(AWorldPar
 				{
 					FTransform InstanceTransform;
 					InstancedStaticMeshComponent->GetInstanceTransform(InstanceIdx, InstanceTransform, true);
-					Component->AddInstanceWorldSpace(InstanceTransform);
+					InstanceTransforms.Add(InstanceTransform);
 				}
 			}
 			else
 			{
-				Component->AddInstanceWorldSpace(SMC->GetComponentTransform());
+				InstanceTransforms.Add(SMC->GetComponentTransform());
 			}
 		}
+		Component->AddInstances(InstanceTransforms, /*bWorldSpace*/true);
 
 		Components.Add(Component);
 	};
diff --git a/Engine/Source/Runtime/Engine/Classes/Components/InstancedStaticMeshComponent.h b/Engine/Source/Runtime/Engine/Classes/Components/InstancedStaticMeshComponent.h
--- a/Engine/Source/Runtime/Engine/Classes/Components/InstancedStaticMeshComponent.h
+++ b/Engine/Source/Runtime/Engine/Classes/Components/InstancedStaticMeshComponent.h
@@ -184,10 +184,10 @@ class ENGINE_API UInstancedStaticMeshComponent : public UStaticMeshComponent
 	virtual TArray<int32> AddInstances(const TArray<FTransform>& InstanceTransforms, bool bShouldReturnIndices, bool bWorldSpace = false);
 
 	/** Add an instance to this component. Transform is given in world space. */
-	UFUNCTION(BlueprintCallable, Category = "Components|InstancedStaticMesh")
+	UE_DEPRECATED(5.0, "Use AddInstance or AddInstances with bWorldSpace set to true.")
+	UFUNCTION(BlueprintCallable, Category = "Components|InstancedStaticMesh", meta=(DeprecatedFunction, DeprecationMessage="Use 'Add Instance' or 'Add Instances' with 'World Space' set to true."))
 	int32 AddInstanceWorldSpace(const FTransform& WorldTransform)
 	{
-		// TODO: Deprecate this shim
 		return AddInstance(WorldTransform, /*bWorldSpace*/true);
 	}
 
diff --git a/Engine/Source/Runtime/Engine/Private/ISMPartition/ISMPartitionActor.cpp b/Engine/Source/Runtime/Engine/Private/ISMPartition/ISMPartitionActor.cpp
--- a/Engine/Source/Runtime/Engine/Private/ISMPartition/ISMPartitionActor.cpp
+++ b/Engine/Source/Runtime/Engine/Private/ISMPartition/ISMPartitionActor.cpp
@@ -299,7 +299,7 @@ void AISMPartitionActor::AddInstanceToComponent(FISMComponentData& ComponentData
 {
 	check(ComponentData.Component);
 	ModifyComponent(ComponentData);
-	ComponentData.Component->AddInstanceWorldSpace(WorldTransform);
+	ComponentData.Component->AddInstance(WorldTransform, /*bWorldSpace*/true);
 	InvalidateComponentLightingCache(ComponentData);
 }
 
diff --git a/Engine/Source/Runtime/Engine/Private/LODActor.cpp b/Engine/Source/Runtime/Engine/Private/LODActor.cpp
--- a/Engine/Source/Runtime/Engine/Private/LODActor.cpp
+++ b/Engine/Source/Runtime/Engine/Private/LODActor.cpp
@@ -1032,10 +1032,7 @@ void ALODActor::AddInstances(const UStaticMesh* InStaticMesh, const UMaterialInt
 	check(InTransforms.Num() > 0);
 
 	UInstancedStaticMeshComponent* Component = GetOrCreateISMComponent(FHLODInstancingKey(InStaticMesh, InMaterial));
-	for (const FTransform& Transform : InTransforms)
-	{
-		Component->AddInstanceWorldSpace(Transform);
-	}
+	Component->AddInstances(InTransforms, /*bWorldSpace*/true);
 
 	// Ensure parenting is up to date and take into account the newly created component.
 	UpdateSubActorLODParents();
diff --git a/Engine/Source/Runtime/Engine/Private/LevelInstance/Packed/LevelInstanceISMPacker.cpp b/Engine/Source/Runtime/Engine/Private/LevelInstance/Packed/LevelInstanceISMPacker.cpp
--- a/Engine/Source/Runtime/Engine/Private/LevelInstance/Packed/LevelInstanceISMPacker.cpp
+++ b/Engine/Source/Runtime/Engine/Private/LevelInstance/Packed/LevelInstanceISMPacker.cpp
@@ -51,6 +51,7 @@ void FLevelInstanceISMPacker::PackActors(FPackedLevelInstanceBuilderContext& InC
 
 	ISMCluster->ISMDescriptor.InitComponent(PackComponent);
 
+	TArray<FTransform> InstanceTransforms;
 	for (UActorComponent* Component : InComponents)
 	{
 		// If we have a ISM we need to add all instances
@@ -62,16 +63,17 @@ void FLevelInstanceISMPacker::PackActors(FPackedLevelInstanceBuilderContext& InC
 
 				if (ensure(ISMComponent->GetInstanceTransform(InstanceIndex, InstanceTransform, /*bWorldSpace=*/ true)))
 				{
-					PackComponent->AddInstanceWorldSpace(InstanceTransform);
+					InstanceTransforms.Add(InstanceTransform);
 				}
 			}
 		}
 		else // other subclasses are processed like regular UStaticMeshComponent
 		{
 			UStaticMeshComponent* StaticMeshComponent = CastChecked<UStaticMeshComponent>(Component);
-			PackComponent->AddInstanceWorldSpace(StaticMeshComponent->GetComponentTransform());
+			InstanceTransforms.Add(StaticMeshComponent->GetComponentTransform());
 		}
 	}
+	PackComponent->AddInstances(InstanceTransforms, /*bWorldSpace*/true);
 
 	FTransform NewWorldTransform = ActorTransform * CurrentPivotOffsetInverse * FTransform(InContext.GetPivotOffset());
 
diff --git a/Engine/Source/Runtime/Engine/Private/LightWeightInstanceStaticMeshManager.cpp b/Engine/Source/Runtime/Engine/Private/LightWeightInstanceStaticMeshManager.cpp
--- a/Engine/Source/Runtime/Engine/Private/LightWeightInstanceStaticMeshManager.cpp
+++ b/Engine/Source/Runtime/Engine/Private/LightWeightInstanceStaticMeshManager.cpp
@@ -69,7 +69,7 @@ void ALightWeightInstanceStaticMeshManager::AddNewInstanceAt(FLWIData* InitData,
 	// Update the HISMC
 	if (InstancedStaticMeshComponent)
 	{
-		ensure(RenderingIdx == InstancedStaticMeshComponent->AddInstanceWorldSpace(InitData->Transform));
+		ensure(RenderingIdx == InstancedStaticMeshComponent->AddInstance(InitData->Transform, /*bWorldSpace*/true));
 	}
 }
 
@@ -129,7 +129,7 @@ void ALightWeightInstanceStaticMeshManager::OnRep_Transforms()
 
 	if (InstancedStaticMeshComponent)
 	{
-		InstancedStaticMeshComponent->AddInstanceWorldSpace(InstanceTransforms.Last());
+		InstancedStaticMeshComponent->AddInstance(InstanceTransforms.Last(), /*bWorldSpace*/true);
 	}
 }
 
diff --git a/Engine/Source/Runtime/Foliage/Private/InstancedFoliage.cpp b/Engine/Source/Runtime/Foliage/Private/InstancedFoliage.cpp
--- a/Engine/Source/Runtime/Foliage/Private/InstancedFoliage.cpp
+++ b/Engine/Source/Runtime/Foliage/Private/InstancedFoliage.cpp
@@ -1064,7 +1064,7 @@ void FFoliageStaticMesh::PreAddInstances(const UFoliageType* FoliageType, int32
 void FFoliageStaticMesh::AddInstance(const FFoliageInstance& NewInstance)
 {
 	check(Component);
-	Component->AddInstanceWorldSpace(NewInstance.GetInstanceWorldTransform());
+	Component->AddInstance(NewInstance.GetInstanceWorldTransform(), /*bWorldSpace*/true);
 	bInvalidateLightingCache = true;
 }
 
@@ -1652,10 +1652,9 @@ void FFoliageStaticMesh::Reapply(const UFoliageType* FoliageType)
 
 		Component->bAutoRebuildTreeOnInstanceChanges = false;
 
-		for (auto& Instance : Info->Instances)
-		{
-			Component->AddInstanceWorldSpace(Instance.GetInstanceWorldTransform());
-		}
+		TArray<FTransform> InstanceTransforms;
+		Algo::Transform(Info->Instances, InstanceTransforms, &FFoliageInstance::GetInstanceWorldTransform);
+		Component->AddInstances(InstanceTransforms, /*bWorldSpace*/true);
 
 		Component->bAutoRebuildTreeOnInstanceChanges = true;
 		Component->BuildTreeIfOutdated(true, true);
