Avoid expensive post GC check by world partition when there is an active gameplay world.

diff --git a/Engine/Source/Runtime/Engine/Private/WorldPartition/WorldPartition.cpp b/Engine/Source/Runtime/Engine/Private/WorldPartition/WorldPartition.cpp
--- a/Engine/Source/Runtime/Engine/Private/WorldPartition/WorldPartition.cpp
+++ b/Engine/Source/Runtime/Engine/Private/WorldPartition/WorldPartition.cpp
@@ -250,6 +250,18 @@ UWorldPartition::UWorldPartition(const FObjectInitializer& ObjectInitializer)
 #if WITH_EDITOR
 void UWorldPartition::OnGCPostReachabilityAnalysis()
 {
+	const TIndirectArray<FWorldContext>& WorldContextList = GEngine->GetWorldContexts();
+
+	// Avoid running this process while a game world is live
+	for (const FWorldContext& WorldContext : WorldContextList)
+	{
+		if (WorldContext.World() != nullptr && WorldContext.World()->IsGameWorld())
+		{
+			return;
+		}
+	}
+	
+
 	for (FRawObjectIterator It; It; ++It)
 	{
 		if (AActor* Actor = Cast<AActor>(static_cast<UObject*>(It->Object)))
