Do not perform a full GC flush in the editor when there is a gameplay level running.

diff --git a/Engine/Source/Runtime/CoreUObject/Private/UObject/GarbageCollection.cpp b/Engine/Source/Runtime/CoreUObject/Private/UObject/GarbageCollection.cpp
--- a/Engine/Source/Runtime/CoreUObject/Private/UObject/GarbageCollection.cpp
+++ b/Engine/Source/Runtime/CoreUObject/Private/UObject/GarbageCollection.cpp
@@ -2077,7 +2077,7 @@ void CollectGarbageInternal(EObjectFlags KeepFlags, bool bPerformFullPurge)
 		GObjPurgeIsRequired = true;
 
 		// Perform a full purge by not using a time limit for the incremental purge. The Editor always does a full purge.
-		if (bPerformFullPurge || GIsEditor)
+		if (bPerformFullPurge)
 		{
 			IncrementalPurgeGarbage(false);
 		}
diff --git a/Engine/Source/Runtime/Engine/Private/UnrealEngine.cpp b/Engine/Source/Runtime/Engine/Private/UnrealEngine.cpp
--- a/Engine/Source/Runtime/Engine/Private/UnrealEngine.cpp
+++ b/Engine/Source/Runtime/Engine/Private/UnrealEngine.cpp
@@ -1481,8 +1481,19 @@ void UEngine::PerformGarbageCollectionAndCleanupActors()
 	// to block on loading the remaining data.
 	if (!IsAsyncLoading())
 	{
+		bool bForcePurge = true;
+		for (FWorldContext& Context : WorldList)
+		{
+			UWorld* World = Context.World();
+			if (World->IsGameWorld())
+			{
+				bForcePurge = false;
+				break;
+			}
+		}
+
 		// Perform housekeeping.
-		if (TryCollectGarbage(GARBAGE_COLLECTION_KEEPFLAGS, false))
+		if (TryCollectGarbage(GARBAGE_COLLECTION_KEEPFLAGS, bForcePurge))
 		{
 			ForEachObjectOfClass(UWorld::StaticClass(), [](UObject* World)
 			{
