CIS static analysis fix

diff --git a/Engine/Source/Runtime/Engine/Private/ComputeFramework/ComputeGraph.cpp b/Engine/Source/Runtime/Engine/Private/ComputeFramework/ComputeGraph.cpp
--- a/Engine/Source/Runtime/Engine/Private/ComputeFramework/ComputeGraph.cpp
+++ b/Engine/Source/Runtime/Engine/Private/ComputeFramework/ComputeGraph.cpp
@@ -123,9 +123,7 @@ void UComputeGraph::CacheResourceShadersForRendering(uint32 CompilationFlags)
 		for (int32 KernelIndex = 0; KernelIndex < KernelInvocations.Num(); ++KernelIndex)
 		{
 			UComputeKernel* Kernel = KernelInvocations[KernelIndex];
-			UComputeKernelSource* KernelSource = Kernel == nullptr ? nullptr : Kernel->KernelSource;
-
-			if (KernelSource == nullptr)
+			if (Kernel == nullptr || Kernel->KernelSource == nullptr)
 			{
 				if (KernelResources[KernelIndex].IsValid())
 				{
@@ -141,6 +139,7 @@ void UComputeGraph::CacheResourceShadersForRendering(uint32 CompilationFlags)
 				KernelResources[KernelIndex] = MakeUnique<FComputeKernelResource>();
 			}
 
+			UComputeKernelSource* KernelSource = Kernel->KernelSource;
 			FComputeKernelResource* KernelResource = KernelResources[KernelIndex].Get();
 
 			TArray<int32> RelevantEdgeIndices;
