Remove UMaterialExpressionReturnMaterialAttributes
- Instead, use control flow materials use the regular root node, with an exec input stacked on top
- Execution begin expression remains
- Will simplify graph processing, as now we only have a single end-point

diff --git a/Engine/Source/Editor/GraphEditor/Private/MaterialNodes/SGraphNodeMaterialResult.cpp b/Engine/Source/Editor/GraphEditor/Private/MaterialNodes/SGraphNodeMaterialResult.cpp
--- a/Engine/Source/Editor/GraphEditor/Private/MaterialNodes/SGraphNodeMaterialResult.cpp
+++ b/Engine/Source/Editor/GraphEditor/Private/MaterialNodes/SGraphNodeMaterialResult.cpp
@@ -24,34 +24,40 @@ void SGraphNodeMaterialResult::Construct(const FArguments& InArgs, UMaterialGrap
 void SGraphNodeMaterialResult::CreatePinWidgets()
 {
 	// Create Pin widgets for each of the pins.
-	for( int32 PinIndex=0; PinIndex < GraphNode->Pins.Num(); ++PinIndex )
-	{
-		UEdGraphPin* CurPin = GraphNode->Pins[PinIndex];
+	UMaterialGraphNode_Base* MaterialGraphNode = Cast<UMaterialGraphNode_Base>(GraphNode);
+	UMaterialGraph* MaterialGraph = CastChecked<UMaterialGraph>(GraphNode->GetGraph());
 
-		bool bHideNoConnectionPins = false;
-		
-		if (OwnerGraphPanelPtr.IsValid())
-		{
-			bHideNoConnectionPins = OwnerGraphPanelPtr.Pin()->GetPinVisibility() == SGraphEditor::Pin_HideNoConnection;
-		}
+	bool bHideNoConnectionPins = false;
+	if (OwnerGraphPanelPtr.IsValid())
+	{
+		bHideNoConnectionPins = OwnerGraphPanelPtr.Pin()->GetPinVisibility() == SGraphEditor::Pin_HideNoConnection;
+	}
 
+	for(const auto& It : MaterialGraphNode->PinInfoMap)
+	{
+		UEdGraphPin* CurPin = It.Key;
+		const FMaterialGraphPinInfo& PinInfo = It.Value;
 		const bool bPinHasConections = CurPin->LinkedTo.Num() > 0;
 
-		//const bool bPinDesiresToBeHidden = CurPin->bHidden || (bHideNoConnectionPins && !bPinHasConections);
-
-		UMaterialGraph* MaterialGraph = CastChecked<UMaterialGraph>(GraphNode->GetGraph());
-
-		check(PinIndex < MaterialGraph->MaterialInputs.Num());
-
-		const bool bPinDesiresToBeHidden = !MaterialGraph->MaterialInputs[PinIndex].IsVisiblePin(MaterialGraph->Material) || (bHideNoConnectionPins && !bPinHasConections);
+		bool bPinDesiresToBeHidden = bHideNoConnectionPins && !bPinHasConections;
+		if (PinInfo.PinType == EMaterialGraphPinType::Data)
+		{
+			if (!MaterialGraph->MaterialInputs[PinInfo.Index].IsVisiblePin(MaterialGraph->Material))
+			{
+				bPinDesiresToBeHidden = true;
+			}
+		}
 
 		if (!bPinDesiresToBeHidden)
 		{
 			TSharedPtr<SGraphPin> NewPin = CreatePinWidget(CurPin);
 			check(NewPin.IsValid());
 
-			TSharedPtr<SToolTip> ToolTipWidget = IDocumentation::Get()->CreateToolTip(MaterialGraph->MaterialInputs[PinIndex].GetToolTip(), nullptr, FString( TEXT("") ), FString( TEXT("") ) );
-			NewPin->SetToolTip( ToolTipWidget.ToSharedRef() );
+			if (PinInfo.PinType == EMaterialGraphPinType::Data)
+			{
+				TSharedPtr<SToolTip> ToolTipWidget = IDocumentation::Get()->CreateToolTip(MaterialGraph->MaterialInputs[PinInfo.Index].GetToolTip(), nullptr, FString(TEXT("")), FString(TEXT("")));
+				NewPin->SetToolTip(ToolTipWidget.ToSharedRef());
+			}
 
 			this->AddPin(NewPin.ToSharedRef());
 		}
diff --git a/Engine/Source/Editor/MaterialEditor/Private/MaterialExpressionClasses.cpp b/Engine/Source/Editor/MaterialEditor/Private/MaterialExpressionClasses.cpp
--- a/Engine/Source/Editor/MaterialEditor/Private/MaterialExpressionClasses.cpp
+++ b/Engine/Source/Editor/MaterialEditor/Private/MaterialExpressionClasses.cpp
@@ -18,12 +18,12 @@
 #include "Materials/MaterialExpressionMaterialLayerOutput.h"
 #include "Materials/MaterialExpressionNamedReroute.h"
 #include "Materials/MaterialExpressionExecBegin.h"
+#include "Materials/MaterialExpressionExecEnd.h"
 
 #include "Materials/MaterialExpressionIfThenElse.h"
 #include "Materials/MaterialExpressionForLoop.h"
 #include "Materials/MaterialExpressionGetLocal.h"
 #include "Materials/MaterialExpressionSetLocal.h"
-#include "Materials/MaterialExpressionReturnMaterialAttributes.h"
 
 ///////////////////////////////////////////////////////////////////////////////////////////////////////////////////////
 //
@@ -118,7 +118,8 @@ void MaterialExpressionClasses::InitMaterialExpressionClasses()
 						if ( Class != UMaterialExpressionComment::StaticClass() 
 							&& Class != UMaterialExpressionParameter::StaticClass()
 							&& Class != UMaterialExpressionNamedRerouteUsage::StaticClass()
-							&& Class != UMaterialExpressionExecBegin::StaticClass() )
+							&& Class != UMaterialExpressionExecBegin::StaticClass()
+							&& Class != UMaterialExpressionExecEnd::StaticClass() )
 						{
 							FMaterialExpression MaterialExpression;
 							// Trim the material expression name and add it to the list used for filtering.
diff --git a/Engine/Source/Editor/UnrealEd/Private/MaterialGraph.cpp b/Engine/Source/Editor/UnrealEd/Private/MaterialGraph.cpp
--- a/Engine/Source/Editor/UnrealEd/Private/MaterialGraph.cpp
+++ b/Engine/Source/Editor/UnrealEd/Private/MaterialGraph.cpp
@@ -17,6 +17,8 @@
 #include "Materials/MaterialExpressionCustomOutput.h"
 #include "Materials/MaterialExpressionReroute.h"
 #include "Materials/MaterialExpressionNamedReroute.h"
+#include "Materials/MaterialExpressionExecBegin.h"
+#include "Materials/MaterialExpressionExecEnd.h"
 
 #include "MaterialGraphNode_Knot.h"
 
@@ -51,6 +53,30 @@ void UMaterialGraph::RebuildGraph()
 	RebuildGraphInternal(SubgraphExpressionMap, SubgraphCommentMap);
 }
 
+template<typename NodeType>
+static UMaterialGraphNode* InitExpressionNewNode(UMaterialGraph* Graph, UMaterialExpression* Expression, bool bUserInvoked)
+{
+	UMaterialGraphNode* NewNode = nullptr;
+
+	FGraphNodeCreator<NodeType> NodeCreator(*Graph);
+	if (bUserInvoked)
+	{
+		NewNode = NodeCreator.CreateUserInvokedNode();
+	}
+	else
+	{
+		NewNode = NodeCreator.CreateNode(false);
+	}
+	NewNode->MaterialExpression = Expression;
+	NewNode->RealtimeDelegate = Graph->RealtimeDelegate;
+	NewNode->MaterialDirtyDelegate = Graph->MaterialDirtyDelegate;
+	Expression->GraphNode = NewNode;
+	Expression->SubgraphExpression = Graph->SubgraphExpression;
+	NodeCreator.Finalize();
+
+	return NewNode;
+}
+
 void UMaterialGraph::RebuildGraphInternal(const TMap<UMaterialExpression*, TArray<UMaterialExpression*>>& SubgraphExpressionMap, const TMap<UMaterialExpression*, TArray<UMaterialExpressionComment*>>& SubgraphCommentMap)
 {
 	Modify();
@@ -93,11 +119,14 @@ void UMaterialGraph::RebuildGraphInternal(const TMap<UMaterialExpression*, TArra
 		//^^^ New material properties go above here. ^^^^
 		MaterialInputs.Add(FMaterialInputInfo(LOCTEXT("MaterialAttributes", "Material Attributes"), MP_MaterialAttributes, LOCTEXT("MaterialAttributesToolTip", "Material Attributes")));
 
-		// No root node if material is using exec pin
-		// Attribute assignments should be done using explicit nodes in this case
-		if (!Material->IsCompiledWithExecutionFlow())
+		if (Material->IsCompiledWithExecutionFlow())
+		{
+			check(Material->ExpressionExecBegin);
+			InitExpressionNewNode<UMaterialGraphNode>(this, Material->ExpressionExecBegin, false);
+		}
+		
+		// Add Root Node
 		{
-			// Add Root Node
 			FGraphNodeCreator<UMaterialGraphNode_Root> NodeCreator(*this);
 			RootNode = NodeCreator.CreateNode();
 			RootNode->Material = Material;
@@ -155,33 +184,13 @@ void UMaterialGraph::RebuildGraphInternal(const TMap<UMaterialExpression*, TArra
 	LinkGraphNodesFromMaterial();
 }
 
-template<typename NodeType>
-static UMaterialGraphNode* InitExpressionNewNode(UMaterialGraph* Graph, UMaterialExpression* Expression, bool bUserInvoked)
-{
-	UMaterialGraphNode* NewNode = nullptr;
-
-	FGraphNodeCreator<NodeType> NodeCreator(*Graph);
-	if (bUserInvoked)
-	{
-		NewNode = NodeCreator.CreateUserInvokedNode();
-	}
-	else
-	{
-		NewNode = NodeCreator.CreateNode(false);
-	}
-	NewNode->MaterialExpression = Expression;
-	NewNode->RealtimeDelegate = Graph->RealtimeDelegate;
-	NewNode->MaterialDirtyDelegate = Graph->MaterialDirtyDelegate;
-	Expression->GraphNode = NewNode;
-	Expression->SubgraphExpression = Graph->SubgraphExpression;
-	NodeCreator.Finalize();
-
-	return NewNode;
-}
-
 UMaterialGraphNode* UMaterialGraph::AddExpression(UMaterialExpression* Expression, bool bUserInvoked)
 {
-	if (Expression)
+	// Node for UMaterialExpressionExecBegin is explicitly placed if needed
+	// We don't created any node for UMaterialExpressionExecEnd, it's handled as part of the root node
+	if (Expression &&
+		!Expression->IsA(UMaterialExpressionExecBegin::StaticClass()) &&
+		!Expression->IsA(UMaterialExpressionExecEnd::StaticClass()))
 	{
 		Modify();
 
@@ -371,7 +380,12 @@ void UMaterialGraph::LinkGraphNodesFromMaterial()
 				FExpressionExecOutput* ExecOutput = ExecOutputs[PinInfo.Index].Output;
 				if (ExecOutput->Expression)
 				{
-					if (UMaterialGraphNode* GraphNode = Cast<UMaterialGraphNode>(ExecOutput->Expression->GraphNode))
+					if (ExecOutput->Expression == Material->ExpressionExecEnd)
+					{
+						// Exec end point is the root node
+						Pin->MakeLinkTo(RootNode->GetExecInputPin());
+					}
+					else if (UMaterialGraphNode* GraphNode = Cast<UMaterialGraphNode>(ExecOutput->Expression->GraphNode))
 					{
 						Pin->MakeLinkTo(GraphNode->GetExecInputPin());
 					}
@@ -395,18 +409,18 @@ void UMaterialGraph::LinkMaterialExpressionsFromGraph() const
 			Material->Modify();
 			Material->EditorX = RootNode->NodePosX;
 			Material->EditorY = RootNode->NodePosY;
-			check(RootNode->Pins.Num() == MaterialInputs.Num());
-			for (int32 PinIndex = 0; PinIndex < RootNode->Pins.Num() && PinIndex < MaterialInputs.Num(); ++PinIndex)
+			check(RootNode->InputPins.Num() == MaterialInputs.Num());
+			for (int32 PinIndex = 0; PinIndex < RootNode->InputPins.Num() && PinIndex < MaterialInputs.Num(); ++PinIndex)
 			{
 				FExpressionInput& MaterialInput = MaterialInputs[PinIndex].GetExpressionInput(Material);
 
-				if (RootNode->Pins[PinIndex]->LinkedTo.Num() > 0)
+				if (RootNode->InputPins[PinIndex]->LinkedTo.Num() > 0)
 				{
-					UMaterialGraphNode* ConnectedNode = CastChecked<UMaterialGraphNode>(RootNode->Pins[PinIndex]->LinkedTo[0]->GetOwningNode());
+					UMaterialGraphNode* ConnectedNode = CastChecked<UMaterialGraphNode>(RootNode->InputPins[PinIndex]->LinkedTo[0]->GetOwningNode());
 					// Work out the index of the connected pin
 					for (int32 OutPinIndex = 0; OutPinIndex < ConnectedNode->OutputPins.Num(); ++OutPinIndex)
 					{
-						if (ConnectedNode->OutputPins[OutPinIndex] == RootNode->Pins[PinIndex]->LinkedTo[0])
+						if (ConnectedNode->OutputPins[OutPinIndex] == RootNode->InputPins[PinIndex]->LinkedTo[0])
 						{
 							if (!ConnectedNode->MaterialExpression->IsExpressionConnected(&MaterialInput, OutPinIndex))
 							{
@@ -501,20 +515,31 @@ void UMaterialGraph::LinkMaterialExpressionsFromGraph() const
 
 							if (Pin->LinkedTo.Num() > 0)
 							{
-								UMaterialGraphNode* ConnectedNode = CastChecked<UMaterialGraphNode>(Pin->LinkedTo[0]->GetOwningNode());
-
-								if (ExpressionOutput &&
-									ExpressionOutput->Expression != ConnectedNode->MaterialExpression &&
-									ConnectedNode->MaterialExpression->HasExecInput())
+								if (Pin->LinkedTo[0]->GetOwningNode() == RootNode)
 								{
 									if (!bModifiedExpression)
 									{
 										bModifiedExpression = true;
 										Expression->Modify();
 									}
+									ExpressionOutput->Expression = Material->ExpressionExecEnd;
+								}
+								else
+								{
+									UMaterialGraphNode* ConnectedNode = CastChecked<UMaterialGraphNode>(Pin->LinkedTo[0]->GetOwningNode());
+									if (ExpressionOutput &&
+										ExpressionOutput->Expression != ConnectedNode->MaterialExpression &&
+										ConnectedNode->MaterialExpression->HasExecInput())
+									{
+										if (!bModifiedExpression)
+										{
+											bModifiedExpression = true;
+											Expression->Modify();
+										}
 
-									ConnectedNode->MaterialExpression->Modify();
-									ExpressionOutput->Expression = ConnectedNode->MaterialExpression;
+										ConnectedNode->MaterialExpression->Modify();
+										ExpressionOutput->Expression = ConnectedNode->MaterialExpression;
+									}
 								}
 							}
 							else if (ExpressionOutput && ExpressionOutput->Expression)
@@ -568,12 +593,10 @@ bool UMaterialGraph::IsInputActive(UEdGraphPin* GraphPin) const
 {
 	if (Material && RootNode)
 	{
-		for (int32 Index = 0; Index < RootNode->Pins.Num(); ++Index)
+		int32 Index = INDEX_NONE;
+		if (RootNode->InputPins.Find(GraphPin, Index))
 		{
-			if (RootNode->Pins[Index] == GraphPin)
-			{
-				return Material->IsPropertyActiveInEditor(MaterialInputs[Index].GetProperty());
-			}
+			return Material->IsPropertyActiveInEditor(MaterialInputs[Index].GetProperty());
 		}
 	}
 	return true;
@@ -589,8 +612,6 @@ void UMaterialGraph::GetUnusedExpressions(TArray<UEdGraphNode*>& UnusedNodes) co
 	{
 		for (int32 Index = 0; Index < RootNode->InputPins.Num(); ++Index)
 		{
-			check(Index < MaterialInputs.Num());
-			
 			if (MaterialInputs[Index].IsVisiblePin(Material)
 				&& RootNode->InputPins[Index]->LinkedTo.Num() > 0 && RootNode->InputPins[Index]->LinkedTo[0])
 			{
diff --git a/Engine/Source/Editor/UnrealEd/Private/MaterialGraphNode.cpp b/Engine/Source/Editor/UnrealEd/Private/MaterialGraphNode.cpp
--- a/Engine/Source/Editor/UnrealEd/Private/MaterialGraphNode.cpp
+++ b/Engine/Source/Editor/UnrealEd/Private/MaterialGraphNode.cpp
@@ -38,7 +38,6 @@
 #include "Materials/MaterialExpressionNamedReroute.h"
 #include "Materials/MaterialExpressionReroute.h"
 #include "Materials/MaterialExpressionCurveAtlasRowParameter.h"
-#include "Materials/MaterialExpressionReturnMaterialAttributes.h"
 #include "Materials/MaterialExpressionExecBegin.h"
 #include "Materials/MaterialExpressionStrata.h"
 #include "MaterialEditorUtilities.h"
@@ -279,14 +278,6 @@ FLinearColor UMaterialGraphNode::GetNodeTitleColor() const
 		// Previously FColor(255, 155, 0);
 		return Settings->ResultNodeTitleColor;
 	}
-	else if (MaterialExpression->IsA(UMaterialExpressionReturnMaterialAttributes::StaticClass()))
-	{
-		return Settings->ResultNodeTitleColor;
-	}
-	else if (MaterialExpression->IsA(UMaterialExpressionExecBegin::StaticClass()))
-	{
-		return Settings->FunctionTerminatorNodeTitleColor;
-	}
 	else if (MaterialExpression->HasExecInput())
 	{
 		return Settings->ExecutionPinTypeColor;
diff --git a/Engine/Source/Editor/UnrealEd/Private/MaterialGraphNode_Root.cpp b/Engine/Source/Editor/UnrealEd/Private/MaterialGraphNode_Root.cpp
--- a/Engine/Source/Editor/UnrealEd/Private/MaterialGraphNode_Root.cpp
+++ b/Engine/Source/Editor/UnrealEd/Private/MaterialGraphNode_Root.cpp
@@ -47,31 +47,44 @@ void UMaterialGraphNode_Root::PostPlacedNewNode()
 
 uint32 UMaterialGraphNode_Root::GetPinMaterialType(const UEdGraphPin* Pin, const FMaterialGraphPinInfo& PinInfo) const
 {
-	check(PinInfo.PinType == EMaterialGraphPinType::Data);
-	const UMaterialGraph* MaterialGraph = CastChecked<UMaterialGraph>(GetGraph());
-	const FMaterialInputInfo& MaterialInput = MaterialGraph->MaterialInputs[PinInfo.Index];
-	EMaterialProperty Property = MaterialInput.GetProperty();
-
-	uint32 MaterialType = 0u;
-	if (Property == MP_MaterialAttributes)
-	{
-		MaterialType = MCT_MaterialAttributes;
-	}
-	else if (Property == MP_FrontMaterial)
-	{
-		MaterialType = MCT_Strata;
-	}
-	else
+	if (PinInfo.PinType == EMaterialGraphPinType::Data)
 	{
-		MaterialType = FMaterialAttributeDefinitionMap::GetValueType(Property);
+		const UMaterialGraph* MaterialGraph = CastChecked<UMaterialGraph>(GetGraph());
+		const FMaterialInputInfo& MaterialInput = MaterialGraph->MaterialInputs[PinInfo.Index];
+		EMaterialProperty Property = MaterialInput.GetProperty();
+
+		uint32 MaterialType = 0u;
+		if (Property == MP_MaterialAttributes)
+		{
+			MaterialType = MCT_MaterialAttributes;
+		}
+		else if (Property == MP_FrontMaterial)
+		{
+			MaterialType = MCT_Strata;
+		}
+		else
+		{
+			MaterialType = FMaterialAttributeDefinitionMap::GetValueType(Property);
+		}
+		return MaterialType;
 	}
-	return MaterialType;
+	return MCT_Execution;
 }
 
 void UMaterialGraphNode_Root::CreateInputPins()
 {
 	UMaterialGraph* MaterialGraph = CastChecked<UMaterialGraph>(GetGraph());
 
+	if (Material->IsCompiledWithExecutionFlow())
+	{
+		// Create the execution pin
+		UEdGraphPin* NewPin = CreatePin(EGPD_Input, UMaterialGraphSchema::PC_Exec, NAME_None, NAME_None);
+		// Makes sure pin has a name for lookup purposes but user will never see it
+		NewPin->PinName = CreateUniquePinName(TEXT("Input"));
+		NewPin->PinFriendlyName = LOCTEXT("Space", " ");
+		RegisterPin(NewPin, EMaterialGraphPinType::Exec, 0);
+	}
+
 	for (int32 Index = 0; Index < MaterialGraph->MaterialInputs.Num(); ++Index)
 	{
 		const FMaterialInputInfo& MaterialInput = MaterialGraph->MaterialInputs[Index];
diff --git a/Engine/Source/Runtime/Engine/Classes/Materials/Material.h b/Engine/Source/Runtime/Engine/Classes/Materials/Material.h
--- a/Engine/Source/Runtime/Engine/Classes/Materials/Material.h
+++ b/Engine/Source/Runtime/Engine/Classes/Materials/Material.h
@@ -811,6 +811,9 @@ public:
 	/** The execution begin expression, if material is using an exec wire */
 	UPROPERTY()
 	class UMaterialExpressionExecBegin* ExpressionExecBegin;
+	
+	UPROPERTY()
+	class UMaterialExpressionExecEnd* ExpressionExecEnd;
 
 	/** Controls where this parameter group is displayed in a material instance parameter list.  The lower the number the higher up in the parameter list. */
 	UPROPERTY(EditAnywhere, EditFixedSize, Category = "Group Sorting")
@@ -1209,6 +1212,8 @@ private:
 	  * from the Map Check message log.
 	  */
 	void FixupMaterialUsageAfterLoad();
+
+	void CreateExecutionFlowExpressions();
 #endif
 
 public:
diff --git a/Engine/Source/Runtime/Engine/Classes/Materials/MaterialExpressionExecEnd.h b/Engine/Source/Runtime/Engine/Classes/Materials/MaterialExpressionExecEnd.h
new file mode 100644
--- /dev/null
+++ b/Engine/Source/Runtime/Engine/Classes/Materials/MaterialExpressionExecEnd.h
@@ -0,0 +1,24 @@
+// Copyright Epic Games, Inc. All Rights Reserved.
+#pragma once
+
+#include "CoreMinimal.h"
+#include "UObject/ObjectMacros.h"
+#include "MaterialExpressionIO.h"
+#include "Materials/MaterialExpression.h"
+#include "MaterialExpressionExecEnd.generated.h"
+
+UCLASS(MinimalAPI, collapsecategories, hidecategories = Object, meta = (MaterialControlFlow))
+class UMaterialExpressionExecEnd : public UMaterialExpression
+{
+	GENERATED_UCLASS_BODY()
+
+	//~ Begin UMaterialExpression Interface
+#if WITH_EDITOR
+	virtual int32 Compile(class FMaterialCompiler* Compiler, int32 OutputIndex) override;
+	virtual bool CanUserDeleteExpression() const override { return false; }
+	virtual bool HasExecInput() override { return true; }
+
+	virtual EMaterialGenerateHLSLStatus GenerateHLSLStatement(FMaterialHLSLGenerator& Generator, UE::HLSLTree::FScope& Scope, UE::HLSLTree::FStatement*& OutStatement) override;
+#endif
+	//~ End UMaterialExpression Interface
+};
diff --git a/Engine/Source/Runtime/Engine/Classes/Materials/MaterialExpressionReturnMaterialAttributes.h b/Engine/Source/Runtime/Engine/Classes/Materials/MaterialExpressionReturnMaterialAttributes.h
deleted file mode 100644
diff --git a/Engine/Source/Runtime/Engine/Private/HLSLTree/HLSLTree.cpp b/Engine/Source/Runtime/Engine/Private/HLSLTree/HLSLTree.cpp
--- a/Engine/Source/Runtime/Engine/Private/HLSLTree/HLSLTree.cpp
+++ b/Engine/Source/Runtime/Engine/Private/HLSLTree/HLSLTree.cpp
@@ -237,73 +237,80 @@ const UE::HLSLTree::FEmitValue* UE::HLSLTree::FEmitContext::AcquireValue(UE::HLS
 
 const UE::HLSLTree::FEmitValue* UE::HLSLTree::FEmitContext::AcquireValue(UE::HLSLTree::FExpression* Expression)
 {
-	check(Expression);
+	if (!Expression)
+	{
+		return nullptr;
+	}
 
-	FDeclarationEntry*& Entry = FunctionStack.Last().DeclarationMap.FindOrAdd(Expression);
+	FFunctionStackEntry& FunctionStackEntry = FunctionStack.Last();
+	FDeclarationEntry** FoundEntry = FunctionStackEntry.DeclarationMap.Find(Expression);
+	FDeclarationEntry* Entry = FoundEntry ? *FoundEntry : nullptr;
 	if (!Entry)
 	{
 		FLocalHLSLCodeWriter LocalWriter;
 		Shader::FPreshaderData LocalPreshader;
 		FExpressionEmitResult EmitResult(LocalWriter, LocalPreshader);
-		Expression->EmitCode(*this, EmitResult);
-
-		check(EmitResult.EvaluationType != EExpressionEvaluationType::None);
-		check(EmitResult.Type != Shader::EValueType::Void);
-
-		Entry = new(*Allocator) FDeclarationEntry();
-		Entry->Value.ExpressionType = EmitResult.Type;
-		Entry->Value.EvaluationType = EmitResult.EvaluationType;
-		if (EmitResult.EvaluationType == EExpressionEvaluationType::Constant)
-		{
-			// Evaluate the constant preshader and store its value
-			FMaterialRenderContext RenderContext(nullptr, *Material, nullptr);
-			LocalPreshader.Evaluate(nullptr, RenderContext, Entry->Value.ConstantValue);
-		}
-		else if (EmitResult.EvaluationType == EExpressionEvaluationType::Preshader)
+		if (Expression->EmitCode(*this, EmitResult))
 		{
-			// Non-constant preshader, store it
-			Shader::FPreshaderData* Preshader = new Shader::FPreshaderData(MoveTemp(LocalPreshader));
-			TempPreshaders.Add(Preshader); // TODO - dedupe? store more efficiently?
-			Entry->Value.Preshader = Preshader;
-		}
-		else
-		{
-			check(EmitResult.EvaluationType == EExpressionEvaluationType::Shader);
-			if (EmitResult.bInline)
+			check(EmitResult.EvaluationType != EExpressionEvaluationType::None);
+			check(EmitResult.Type != Shader::EValueType::Void);
+
+			Entry = new(*Allocator) FDeclarationEntry();
+			FunctionStackEntry.DeclarationMap.Add(Expression, Entry);
+			Entry->Value.ExpressionType = EmitResult.Type;
+			Entry->Value.EvaluationType = EmitResult.EvaluationType;
+			if (EmitResult.EvaluationType == EExpressionEvaluationType::Constant)
 			{
-				Entry->Value.Code = AllocateString(*Allocator, LocalWriter.GetStringBuilder());
+				// Evaluate the constant preshader and store its value
+				FMaterialRenderContext RenderContext(nullptr, *Material, nullptr);
+				LocalPreshader.Evaluate(nullptr, RenderContext, Entry->Value.ConstantValue);
+			}
+			else if (EmitResult.EvaluationType == EExpressionEvaluationType::Preshader)
+			{
+				// Non-constant preshader, store it
+				Shader::FPreshaderData* Preshader = new Shader::FPreshaderData(MoveTemp(LocalPreshader));
+				TempPreshaders.Add(Preshader); // TODO - dedupe? store more efficiently?
+				Entry->Value.Preshader = Preshader;
 			}
 			else
 			{
-				const int32 ScopeIndex = FindScopeIndex(Expression->ParentScope);
-				check(ScopeIndex != INDEX_NONE);
-		
-				// Check to see if we've already generated code for an equivalent expression in either this scope, or any outer visible scope
-				const FSHAHash Hash = LocalWriter.GetCodeHash();
-				const TCHAR* Declaration = nullptr;
-				for (int32 CheckScopeIndex = ScopeIndex; CheckScopeIndex >= 0; --CheckScopeIndex)
+				check(EmitResult.EvaluationType == EExpressionEvaluationType::Shader);
+				if (EmitResult.bInline)
 				{
-					const FScopeEntry& CheckScopeEntry = ScopeStack[CheckScopeIndex];
-					const TCHAR** FoundDeclaration = CheckScopeEntry.ExpressionMap->Find(Hash);
-					if (FoundDeclaration)
-					{
-						// Re-use results from previous matching expression
-						Declaration = *FoundDeclaration;
-						break;
-					}
+					Entry->Value.Code = AllocateString(*Allocator, LocalWriter.GetStringBuilder());
 				}
-				
-				if (!Declaration)
+				else
 				{
-					const Shader::FValueTypeDescription& TypeDesc = Shader::GetValueTypeDescription(EmitResult.Type);
-					Declaration = AllocateStringf(*Allocator, TEXT("Local%d"), NumExpressionLocals++);
-					ScopeStack[ScopeIndex].ExpressionCodeWriter->WriteLinef(TEXT("const %s %s = %s;"),
-						TypeDesc.Name,
-						Declaration,
-						LocalWriter.GetStringBuilder().ToString());
-					ScopeStack[ScopeIndex].ExpressionMap->Add(Hash, Declaration);
+					const int32 ScopeIndex = FindScopeIndex(Expression->ParentScope);
+					check(ScopeIndex != INDEX_NONE);
+
+					// Check to see if we've already generated code for an equivalent expression in either this scope, or any outer visible scope
+					const FSHAHash Hash = LocalWriter.GetCodeHash();
+					const TCHAR* Declaration = nullptr;
+					for (int32 CheckScopeIndex = ScopeIndex; CheckScopeIndex >= 0; --CheckScopeIndex)
+					{
+						const FScopeEntry& CheckScopeEntry = ScopeStack[CheckScopeIndex];
+						const TCHAR** FoundDeclaration = CheckScopeEntry.ExpressionMap->Find(Hash);
+						if (FoundDeclaration)
+						{
+							// Re-use results from previous matching expression
+							Declaration = *FoundDeclaration;
+							break;
+						}
+					}
+
+					if (!Declaration)
+					{
+						const Shader::FValueTypeDescription& TypeDesc = Shader::GetValueTypeDescription(EmitResult.Type);
+						Declaration = AllocateStringf(*Allocator, TEXT("Local%d"), NumExpressionLocals++);
+						ScopeStack[ScopeIndex].ExpressionCodeWriter->WriteLinef(TEXT("const %s %s = %s;"),
+							TypeDesc.Name,
+							Declaration,
+							LocalWriter.GetStringBuilder().ToString());
+						ScopeStack[ScopeIndex].ExpressionMap->Add(Hash, Declaration);
+					}
+					Entry->Value.Code = Declaration;
 				}
-				Entry->Value.Code = Declaration;
 			}
 		}
 	}
@@ -508,7 +515,7 @@ UE::HLSLTree::ENodeVisitResult UE::HLSLTree::FScope::Visit(UE::HLSLTree::FNodeVi
 	return Result;
 }
 
-void UE::HLSLTree::FScope::EmitHLSL(FEmitContext& Context, FCodeWriter& OutWriter) const
+bool UE::HLSLTree::FScope::EmitHLSL(FEmitContext& Context, FCodeWriter& OutWriter) const
 {
 	const int32 ScopeIndentLevel = OutWriter.IndentLevel + 1;
 
@@ -527,9 +534,14 @@ void UE::HLSLTree::FScope::EmitHLSL(FEmitContext& Context, FCodeWriter& OutWrite
 	FLocalHLSLCodeWriter StatementCodeWriter;
 	StatementCodeWriter.IndentLevel = ScopeIndentLevel;
 	const FStatement* Statement = FirstStatement;
+	bool bResult = true;
 	while (Statement)
 	{
-		Statement->EmitHLSL(Context, StatementCodeWriter);
+		if (!Statement->EmitHLSL(Context, StatementCodeWriter))
+		{
+			bResult = false;
+			break;
+		}
 
 		// First write any expressions needed by the statement, then the statement itself
 		OutWriter.Append(DeclarationCodeWriter);
@@ -543,18 +555,25 @@ void UE::HLSLTree::FScope::EmitHLSL(FEmitContext& Context, FCodeWriter& OutWrite
 
 	OutWriter.WriteLine(TEXT("}"));
 	Context.ScopeStack.Pop(false);
+
+	return bResult;
 }
 
-void UE::HLSLTree::FScope::EmitUnscopedHLSL(FEmitContext& Context, FCodeWriter& OutWriter) const
+bool UE::HLSLTree::FScope::EmitUnscopedHLSL(FEmitContext& Context, FCodeWriter& OutWriter) const
 {
 	FCodeWriter& DeclarationCodeWriter = *Context.ScopeStack.Last().ExpressionCodeWriter;
 
 	FLocalHLSLCodeWriter StatementCodeWriter;
 	StatementCodeWriter.IndentLevel = OutWriter.IndentLevel;
 	const FStatement* Statement = FirstStatement;
+	bool bResult = true;
 	while (Statement)
 	{
-		Statement->EmitHLSL(Context, StatementCodeWriter);
+		if (!Statement->EmitHLSL(Context, StatementCodeWriter))
+		{
+			bResult = false;
+			break;
+		}
 
 		// First write any expressions needed by the statement, then the statement itself
 		OutWriter.Append(DeclarationCodeWriter);
@@ -565,6 +584,7 @@ void UE::HLSLTree::FScope::EmitUnscopedHLSL(FEmitContext& Context, FCodeWriter&
 
 		Statement = Statement->NextStatement;
 	}
+	return bResult;
 }
 
 void UE::HLSLTree::FScope::AddDeclaration(UE::HLSLTree::FLocalDeclaration* Declaration)
@@ -740,9 +760,9 @@ UE::HLSLTree::FTree* UE::HLSLTree::FTree::Create(FMemStackBase& Allocator)
 	return Tree;
 }
 
-void UE::HLSLTree::FTree::EmitHLSL(UE::HLSLTree::FEmitContext& Context, UE::HLSLTree::FCodeWriter& Writer) const
+bool UE::HLSLTree::FTree::EmitHLSL(UE::HLSLTree::FEmitContext& Context, UE::HLSLTree::FCodeWriter& Writer) const
 {
-	RootScope->EmitHLSL(Context, Writer);
+	return RootScope->EmitHLSL(Context, Writer);
 }
 
 UE::HLSLTree::FScope* UE::HLSLTree::FTree::NewScope(UE::HLSLTree::FScope& Scope)
diff --git a/Engine/Source/Runtime/Engine/Private/HLSLTree/HLSLTreeCommon.cpp b/Engine/Source/Runtime/Engine/Private/HLSLTree/HLSLTreeCommon.cpp
--- a/Engine/Source/Runtime/Engine/Private/HLSLTree/HLSLTreeCommon.cpp
+++ b/Engine/Source/Runtime/Engine/Private/HLSLTree/HLSLTreeCommon.cpp
@@ -59,21 +59,28 @@ UE::HLSLTree::FSwizzleParameters UE::HLSLTree::MakeSwizzleMask(bool bInR, bool b
 }
 
 
-void UE::HLSLTree::FExpressionConstant::EmitCode(FEmitContext& Context, FExpressionEmitResult& OutResult) const
+bool UE::HLSLTree::FExpressionConstant::EmitCode(FEmitContext& Context, FExpressionEmitResult& OutResult) const
 {
 	OutResult.EvaluationType = EExpressionEvaluationType::Constant;
 	OutResult.Type = Value.GetType();
 	OutResult.Preshader.WriteOpcode(Shader::EPreshaderOpcode::Constant);
 	OutResult.Preshader.Write(Value);
+	return true;
 }
 
-void UE::HLSLTree::FExpressionLocalVariable::EmitCode(FEmitContext& Context, FExpressionEmitResult& OutResult) const
+bool UE::HLSLTree::FExpressionLocalVariable::EmitCode(FEmitContext& Context, FExpressionEmitResult& OutResult) const
 {
 	const FEmitValue* DeclarationValue = Context.AcquireValue(Declaration);
+	if (!DeclarationValue)
+	{
+		return false;
+	}
+
 	OutResult.ForwardValue(Context, DeclarationValue);
+	return true;
 }
 
-void UE::HLSLTree::FExpressionParameter::EmitCode(FEmitContext& Context, FExpressionEmitResult& OutResult) const
+bool UE::HLSLTree::FExpressionParameter::EmitCode(FEmitContext& Context, FExpressionEmitResult& OutResult) const
 {
 	OutResult.Type = Declaration->DefaultValue.GetType();
 
@@ -106,9 +113,10 @@ void UE::HLSLTree::FExpressionParameter::EmitCode(FEmitContext& Context, FExpres
 		OutResult.EvaluationType = EExpressionEvaluationType::Preshader;
 		OutResult.Preshader.WriteOpcode(Shader::EPreshaderOpcode::VectorParameter).Write((uint16)ParameterIndex);
 	}
+	return true;
 }
 
-void UE::HLSLTree::FExpressionExternalInput::EmitCode(FEmitContext& Context, FExpressionEmitResult& OutResult) const
+bool UE::HLSLTree::FExpressionExternalInput::EmitCode(FEmitContext& Context, FExpressionEmitResult& OutResult) const
 {
 	const int32 TypeIndex = (int32)InputType;
 	const int32 TexCoordIndex = TypeIndex - (int32)UE::HLSLTree::EExternalInputType::TexCoord0;
@@ -125,10 +133,17 @@ void UE::HLSLTree::FExpressionExternalInput::EmitCode(FEmitContext& Context, FEx
 	{
 		checkNoEntry();
 	}
+	return true;
 }
 
-void UE::HLSLTree::FExpressionTextureSample::EmitCode(FEmitContext& Context, FExpressionEmitResult& OutResult) const
+bool UE::HLSLTree::FExpressionTextureSample::EmitCode(FEmitContext& Context, FExpressionEmitResult& OutResult) const
 {
+	const FEmitValue* TexCoordValue = Context.AcquireValue(TexCoordExpression);
+	if (!TexCoordExpression)
+	{
+		return false;
+	}
+
 	const UE::HLSLTree::FTextureDescription& Desc = Declaration->Description;
 	const EMaterialValueType MaterialValueType = Desc.Texture->GetMaterialType();
 	EMaterialTextureParameterType TextureType = EMaterialTextureParameterType::Count;
@@ -253,35 +268,45 @@ void UE::HLSLTree::FExpressionTextureSample::EmitCode(FEmitContext& Context, FEx
 		}
 	}
 
-	const FEmitValue* TexCoordValue = Context.AcquireValue(TexCoordExpression);
 	OutResult.EvaluationType = EExpressionEvaluationType::Shader;
 	OutResult.Type = Shader::EValueType::Float4;
 	OutResult.Writer.Writef(TEXT("%s(%s(%s, %s, %s))"), SamplerTypeFunction, SampleFunctionName, *TextureName, *SamplerStateCode, Context.GetCode(TexCoordValue));
+	return true;
 }
 
-void UE::HLSLTree::FExpressionDefaultMaterialAttributes::EmitCode(FEmitContext& Context, FExpressionEmitResult& OutResult) const
+bool UE::HLSLTree::FExpressionDefaultMaterialAttributes::EmitCode(FEmitContext& Context, FExpressionEmitResult& OutResult) const
 {
 	OutResult.EvaluationType = EExpressionEvaluationType::Shader;
 	OutResult.Type = Shader::EValueType::MaterialAttributes;
 	OutResult.bInline = true;
 	OutResult.Writer.Write(TEXT("DefaultMaterialAttributes"));
+	return true;
 }
 
-void UE::HLSLTree::FExpressionSetMaterialAttribute::EmitCode(FEmitContext& Context, FExpressionEmitResult& OutResult) const
+bool UE::HLSLTree::FExpressionSetMaterialAttribute::EmitCode(FEmitContext& Context, FExpressionEmitResult& OutResult) const
 {
-	const FString PropertyName = FMaterialAttributeDefinitionMap::GetAttributeName(AttributeID);
-
 	const FEmitValue* AttributresValue = Context.AcquireValue(AttributesExpression);
 	const FEmitValue* Value = Context.AcquireValue(ValueExpression);
+	if (!AttributresValue || !Value)
+	{
+		return false;
+	}
 
+	const FString PropertyName = FMaterialAttributeDefinitionMap::GetAttributeName(AttributeID);
 	OutResult.EvaluationType = EExpressionEvaluationType::Shader;
 	OutResult.Type = Shader::EValueType::MaterialAttributes;
 	OutResult.Writer.Writef(TEXT("FMaterialAttributes_Set%s(%s, %s)"), *PropertyName, Context.GetCode(AttributresValue), Context.GetCode(Value));
+	return true;
 }
 
-void UE::HLSLTree::FExpressionSelect::EmitCode(FEmitContext& Context, FExpressionEmitResult& OutResult) const
+bool UE::HLSLTree::FExpressionSelect::EmitCode(FEmitContext& Context, FExpressionEmitResult& OutResult) const
 {
 	const FEmitValue* ConditionValue = Context.AcquireValue(ConditionExpression);
+	if (!ConditionValue)
+	{
+		return false;
+	}
+
 	if (ConditionValue->GetEvaluationType() == EExpressionEvaluationType::Constant)
 	{
 		const bool bCondition = ConditionValue->GetConstantValue().AsBoolScalar();
@@ -307,12 +332,17 @@ void UE::HLSLTree::FExpressionSelect::EmitCode(FEmitContext& Context, FExpressio
 			Context.GetCode(TrueValue),
 			Context.GetCode(FalseValue));
 	}
+	return true;
 }
 
-void UE::HLSLTree::FExpressionBinaryOp::EmitCode(FEmitContext& Context, FExpressionEmitResult& OutResult) const
+bool UE::HLSLTree::FExpressionBinaryOp::EmitCode(FEmitContext& Context, FExpressionEmitResult& OutResult) const
 {
 	const FEmitValue* LhsValue = Context.AcquireValue(Lhs);
 	const FEmitValue* RhsValue = Context.AcquireValue(Rhs);
+	if (!LhsValue || !RhsValue)
+	{
+		return false;
+	}
 
 	FString ErrorMessage;
 
@@ -335,10 +365,17 @@ void UE::HLSLTree::FExpressionBinaryOp::EmitCode(FEmitContext& Context, FExpress
 		Context.AppendPreshader(RhsValue, OutResult.Preshader);
 		OutResult.Preshader.WriteOpcode(Opcode);
 	}
+	return true;
 }
 
-void UE::HLSLTree::FExpressionSwizzle::EmitCode(FEmitContext& Context, FExpressionEmitResult& OutResult) const
+bool UE::HLSLTree::FExpressionSwizzle::EmitCode(FEmitContext& Context, FExpressionEmitResult& OutResult) const
 {
+	const FEmitValue* InputValue = Context.AcquireValue(Input);
+	if (!InputValue)
+	{
+		return false;
+	}
+
 	static const TCHAR ComponentName[] = { 'x', 'y', 'z', 'w' };
 	TCHAR Swizzle[5] = TEXT("");
 	for (int32 i = 0; i < Parameters.NumComponents; ++i)
@@ -348,8 +385,6 @@ void UE::HLSLTree::FExpressionSwizzle::EmitCode(FEmitContext& Context, FExpressi
 		Swizzle[i] = ComponentName[ComponentIndex];
 	}
 
-	const FEmitValue* InputValue = Context.AcquireValue(Input);
-
 	OutResult.bInline = true;
 	OutResult.EvaluationType = InputValue->GetEvaluationType();
 	OutResult.Type = Shader::MakeValueType(InputValue->GetExpressionType(), Parameters.NumComponents);
@@ -367,25 +402,30 @@ void UE::HLSLTree::FExpressionSwizzle::EmitCode(FEmitContext& Context, FExpressi
 			.Write((uint8)Parameters.ComponentIndex[2])
 			.Write((uint8)Parameters.ComponentIndex[3]);
 	}
+	return true;
 }
 
-void UE::HLSLTree::FExpressionCast::EmitCode(FEmitContext& Context, FExpressionEmitResult& OutResult) const
+bool UE::HLSLTree::FExpressionCast::EmitCode(FEmitContext& Context, FExpressionEmitResult& OutResult) const
 {
 	const FEmitValue* InputValue = Context.AcquireValue(Input);
-	const Shader::EValueType InputType = InputValue->GetExpressionType();
-	const TCHAR* InputCode = Context.GetCode(InputValue);
-
-	OutResult.bInline = true;
-	OutResult.Type = Type;
-	OutResult.EvaluationType = EExpressionEvaluationType::Shader; // TODO - preshader
+	if (!InputValue)
+	{
+		return false;
+	}
 
+	const Shader::EValueType InputType = InputValue->GetExpressionType();
 	if (InputType == Type)
 	{
 		// No cast required (nop)
-		OutResult.Writer.Writef(TEXT("%s"), InputCode);
-		return;
+		OutResult.ForwardValue(Context, InputValue);
+		return true;
 	}
 
+	OutResult.bInline = true;
+	OutResult.Type = Type;
+	OutResult.EvaluationType = EExpressionEvaluationType::Shader; // TODO - preshader
+
+	const TCHAR* InputCode = Context.GetCode(InputValue);
 	const Shader::FValueTypeDescription OutputTypeDesc = Shader::GetValueTypeDescription(Type);
 	const Shader::FValueTypeDescription InputTypeDesc = Shader::GetValueTypeDescription(InputType);
 	if (InputTypeDesc.NumComponents == OutputTypeDesc.NumComponents)
@@ -430,17 +470,19 @@ void UE::HLSLTree::FExpressionCast::EmitCode(FEmitContext& Context, FExpressionE
 			OutResult.Writer.Write(TEXT(")"));
 		}
 	}
+	return true;
 }
 
-void UE::HLSLTree::FExpressionReflectionVector::EmitCode(FEmitContext& Context, FExpressionEmitResult& OutResult) const
+bool UE::HLSLTree::FExpressionReflectionVector::EmitCode(FEmitContext& Context, FExpressionEmitResult& OutResult) const
 {
 	OutResult.bInline = true;
 	OutResult.Type = Shader::EValueType::Float3;
 	OutResult.EvaluationType = EExpressionEvaluationType::Shader;
 	OutResult.Writer.Write(TEXT("Parameters.ReflectionVector"));
+	return true;
 }
 
-void UE::HLSLTree::FExpressionFunctionInput::EmitCode(FEmitContext& Context, FExpressionEmitResult& OutResult) const
+bool UE::HLSLTree::FExpressionFunctionInput::EmitCode(FEmitContext& Context, FExpressionEmitResult& OutResult) const
 {
 	const FEmitContext::FFunctionStackEntry& StackEntry = Context.FunctionStack.Last();
 
@@ -448,30 +490,54 @@ void UE::HLSLTree::FExpressionFunctionInput::EmitCode(FEmitContext& Context, FEx
 	FExpression* InputExpression = StackEntry.FunctionCall->Inputs[InputIndex];
 	const FEmitValue* Value = Context.AcquireValue(InputExpression);
 	OutResult.ForwardValue(Context, Value);
+	return true;
 }
 
-void UE::HLSLTree::FExpressionFunctionOutput::EmitCode(FEmitContext& Context, FExpressionEmitResult& OutResult) const
+bool UE::HLSLTree::FExpressionFunctionOutput::EmitCode(FEmitContext& Context, FExpressionEmitResult& OutResult) const
 {
 	const FEmitValue* Value = Context.AcquireValue(FunctionCall, OutputIndex);
+	if (!Value)
+	{
+		return false;
+	}
+
 	OutResult.ForwardValue(Context, Value);
+	return true;
 }
 
-void UE::HLSLTree::FStatementReturn::EmitHLSL(FEmitContext& Context, FCodeWriter& OutWriter) const
+bool UE::HLSLTree::FStatementReturn::EmitHLSL(FEmitContext& Context, FCodeWriter& OutWriter) const
 {
 	const FEmitValue* Value = Context.AcquireValue(Expression);
+	if (!Value)
+	{
+		return false;
+	}
+
 	OutWriter.WriteLinef(TEXT("return %s;"), Context.GetCode(Value));
+	return true;
 }
 
-void UE::HLSLTree::FStatementSetLocalVariable::EmitHLSL(FEmitContext& Context, FCodeWriter& OutWriter) const
+bool UE::HLSLTree::FStatementSetLocalVariable::EmitHLSL(FEmitContext& Context, FCodeWriter& OutWriter) const
 {
 	const FEmitValue* DeclarationValue = Context.AcquireValue(Declaration);
 	const FEmitValue* ExpressionValue = Context.AcquireValue(Expression);
+	if (!DeclarationValue || !ExpressionValue)
+	{
+		return false;
+	}
+
 	OutWriter.WriteLinef(TEXT("%s = %s;"), Context.GetCode(DeclarationValue), Context.GetCode(ExpressionValue));
+	return true;
 }
 
-void UE::HLSLTree::FStatementIf::EmitHLSL(FEmitContext& Context, FCodeWriter& OutWriter) const
+bool UE::HLSLTree::FStatementIf::EmitHLSL(FEmitContext& Context, FCodeWriter& OutWriter) const
 {
 	const FEmitValue* ConditionValue = Context.AcquireValue(ConditionExpression);
+	if (!ConditionValue)
+	{
+		return false;
+	}
+
 	OutWriter.WriteLinef(TEXT("if (%s)"), Context.GetCode(ConditionValue));
 	ThenScope->EmitHLSL(Context, OutWriter);
 	if (ElseScope)
@@ -479,9 +545,10 @@ void UE::HLSLTree::FStatementIf::EmitHLSL(FEmitContext& Context, FCodeWriter& Ou
 		OutWriter.WriteLine(TEXT("else"));
 		ElseScope->EmitHLSL(Context, OutWriter);
 	}
+	return true;
 }
 
-void UE::HLSLTree::FStatementFor::EmitHLSL(FEmitContext& Context, FCodeWriter& OutWriter) const
+bool UE::HLSLTree::FStatementFor::EmitHLSL(FEmitContext& Context, FCodeWriter& OutWriter) const
 {
-	
+	return false;
 }
diff --git a/Engine/Source/Runtime/Engine/Private/Materials/Material.cpp b/Engine/Source/Runtime/Engine/Private/Materials/Material.cpp
--- a/Engine/Source/Runtime/Engine/Private/Materials/Material.cpp
+++ b/Engine/Source/Runtime/Engine/Private/Materials/Material.cpp
@@ -47,6 +47,7 @@
 #include "Materials/MaterialExpressionShadingModel.h"
 #include "Materials/MaterialExpressionTransform.h"
 #include "Materials/MaterialExpressionExecBegin.h"
+#include "Materials/MaterialExpressionExecEnd.h"
 #include "Materials/MaterialFunction.h"
 #include "Materials/MaterialFunctionInstance.h"
 #include "Materials/MaterialExpressionMaterialFunctionCall.h"
@@ -4132,6 +4133,8 @@ void UMaterial::PostLoad()
 	}
 
 #if WITH_EDITOR
+	// Create exec flow expressions, if needed
+	CreateExecutionFlowExpressions();
 	if (GIsEditor)
 	{
 		// Clean up any removed material expression classes	
@@ -4143,7 +4146,7 @@ void UMaterial::PostLoad()
 			FlushResourceShaderMaps();
 		}
 	}
-#endif
+#endif // WITH_EDITOR
 
 	if (!StateId.IsValid())
 	{
@@ -4560,6 +4563,26 @@ bool UMaterial::CanEditChange(const FProperty* InProperty) const
 	return true;
 }
 
+void UMaterial::CreateExecutionFlowExpressions()
+{
+	if (IsCompiledWithExecutionFlow())
+	{
+		if (!ExpressionExecBegin)
+		{
+			ExpressionExecBegin = NewObject<UMaterialExpressionExecBegin>(this);
+			ExpressionExecBegin->Material = this;
+			Expressions.Add(ExpressionExecBegin);
+		}
+
+		if (!ExpressionExecEnd)
+		{
+			ExpressionExecEnd = NewObject<UMaterialExpressionExecEnd>(this);
+			ExpressionExecEnd->Material = this;
+			Expressions.Add(ExpressionExecEnd);
+		}
+	}
+}
+
 void UMaterial::PreEditChange(FProperty* PropertyThatChanged)
 {
 	Super::PreEditChange(PropertyThatChanged);
@@ -4636,19 +4659,7 @@ void UMaterial::PostEditChangePropertyInternal(FPropertyChangedEvent& PropertyCh
 
 	if (PropertyChangedEvent.GetPropertyName() == GET_MEMBER_NAME_CHECKED(UMaterial, bEnableExecWire))
 	{
-		if (IsCompiledWithExecutionFlow())
-		{
-			check(!ExpressionExecBegin);
-			ExpressionExecBegin = NewObject<UMaterialExpressionExecBegin>(this);
-			ExpressionExecBegin->Material = this;
-			Expressions.Add(ExpressionExecBegin);
-		}
-		else
-		{
-			check(ExpressionExecBegin);
-			verify(Expressions.Remove(ExpressionExecBegin) == 1);
-			ExpressionExecBegin = nullptr;
-		}
+		CreateExecutionFlowExpressions();
 	}
 
 	TranslucencyDirectionalLightingIntensity = FMath::Clamp(TranslucencyDirectionalLightingIntensity, .1f, 10.0f);
diff --git a/Engine/Source/Runtime/Engine/Private/Materials/MaterialExpressionHLSL.cpp b/Engine/Source/Runtime/Engine/Private/Materials/MaterialExpressionHLSL.cpp
--- a/Engine/Source/Runtime/Engine/Private/Materials/MaterialExpressionHLSL.cpp
+++ b/Engine/Source/Runtime/Engine/Private/Materials/MaterialExpressionHLSL.cpp
@@ -10,7 +10,7 @@
 
 #include "Materials/MaterialExpression.h"
 #include "Materials/MaterialExpressionExecBegin.h"
-#include "Materials/MaterialExpressionReturnMaterialAttributes.h"
+#include "Materials/MaterialExpressionExecEnd.h"
 #include "Materials/MaterialExpressionConstant.h"
 #include "Materials/MaterialExpressionConstant2Vector.h"
 #include "Materials/MaterialExpressionConstant3Vector.h"
@@ -317,6 +317,12 @@ EMaterialGenerateHLSLStatus UMaterialExpressionExecBegin::GenerateHLSLStatement(
 	return EMaterialGenerateHLSLStatus::Success;
 }
 
+EMaterialGenerateHLSLStatus UMaterialExpressionExecEnd::GenerateHLSLStatement(FMaterialHLSLGenerator& Generator, UE::HLSLTree::FScope& Scope, UE::HLSLTree::FStatement*& OutStatement)
+{
+	OutStatement = Generator.NewResult(Scope);
+	return EMaterialGenerateHLSLStatus::Success;
+}
+
 EMaterialGenerateHLSLStatus UMaterialExpressionSetLocal::GenerateHLSLStatement(FMaterialHLSLGenerator& Generator, UE::HLSLTree::FScope& Scope, UE::HLSLTree::FStatement*& OutStatement)
 {
 	UE::HLSLTree::FExpression* ValueExpression = Value.AcquireHLSLExpressionWithCast(Generator, Scope, UE::Shader::EValueType::Float3);
@@ -341,21 +347,6 @@ EMaterialGenerateHLSLStatus UMaterialExpressionSetLocal::GenerateHLSLStatement(F
 	return EMaterialGenerateHLSLStatus::Success;
 }
 
-EMaterialGenerateHLSLStatus UMaterialExpressionReturnMaterialAttributes::GenerateHLSLStatement(FMaterialHLSLGenerator& Generator, UE::HLSLTree::FScope& Scope, UE::HLSLTree::FStatement*& OutStatement)
-{
-	UE::HLSLTree::FExpression* AttributesExpression = MaterialAttributes.AcquireHLSLExpression(Generator, Scope);
-	if (!AttributesExpression)
-	{
-		return Generator.Error(TEXT("Missing attribute connection"));
-	}
-
-	UE::HLSLTree::FStatementReturn* ReturnStatement = Generator.GetTree().NewStatement<UE::HLSLTree::FStatementReturn>(Scope);
-	ReturnStatement->Expression = AttributesExpression;
-
-	OutStatement = ReturnStatement;
-	return EMaterialGenerateHLSLStatus::Success;
-}
-
 EMaterialGenerateHLSLStatus UMaterialExpressionIfThenElse::GenerateHLSLStatement(FMaterialHLSLGenerator& Generator, UE::HLSLTree::FScope& Scope, UE::HLSLTree::FStatement*& OutStatement)
 {
 	UE::HLSLTree::FExpression* ConditionExpression = Condition.AcquireHLSLExpression(Generator, Scope);
diff --git a/Engine/Source/Runtime/Engine/Private/Materials/MaterialExpressions.cpp b/Engine/Source/Runtime/Engine/Private/Materials/MaterialExpressions.cpp
--- a/Engine/Source/Runtime/Engine/Private/Materials/MaterialExpressions.cpp
+++ b/Engine/Source/Runtime/Engine/Private/Materials/MaterialExpressions.cpp
@@ -254,7 +254,7 @@
 #include "Materials/MaterialUniformExpressions.h"
 #include "Materials/MaterialExpressionSamplePhysicsField.h"
 #include "Materials/MaterialExpressionExecBegin.h"
-#include "Materials/MaterialExpressionReturnMaterialAttributes.h"
+#include "Materials/MaterialExpressionExecEnd.h"
 #include "Materials/MaterialExpressionIfThenElse.h"
 #include "Materials/MaterialExpressionForLoop.h"
 #include "Materials/MaterialExpressionSetLocal.h"
@@ -21351,58 +21351,20 @@ int32 UMaterialExpressionExecBegin::Compile(class FMaterialCompiler* Compiler, i
 }
 #endif
 
-UMaterialExpressionReturnMaterialAttributes::UMaterialExpressionReturnMaterialAttributes(const FObjectInitializer& ObjectInitializer)
+UMaterialExpressionExecEnd::UMaterialExpressionExecEnd(const FObjectInitializer& ObjectInitializer)
 	: Super(ObjectInitializer)
 {
-	struct FConstructorStatics
-	{
-		FText NAME_Execution;
-		FConstructorStatics() : NAME_Execution(LOCTEXT("Execution", "Execution")) { }
-	};
-	static FConstructorStatics ConstructorStatics;
 #if WITH_EDITORONLY_DATA
-	MenuCategories.Add(ConstructorStatics.NAME_Execution);
 	Outputs.Reset();
 	bHidePreviewWindow = true;
 #endif
 }
 
 #if WITH_EDITOR
-void UMaterialExpressionReturnMaterialAttributes::GetCaption(TArray<FString>& OutCaptions) const
-{
-	OutCaptions.Add(TEXT("Return Material Attributes"));
-}
-
-int32 UMaterialExpressionReturnMaterialAttributes::Compile(class FMaterialCompiler* Compiler, int32 OutputIndex)
-{
-	check(OutputIndex == CompileExecutionOutputIndex);
-	const int32 AttributesIndex = MaterialAttributes.Compile(Compiler);
-	return Compiler->ReturnMaterialAttributes(AttributesIndex);
-}
-
-const TArray<FExpressionInput*> UMaterialExpressionReturnMaterialAttributes::GetInputs()
-{
-	TArray<FExpressionInput*> Result;
-	Result.Add(&MaterialAttributes);
-	return Result;
-}
-
-FExpressionInput* UMaterialExpressionReturnMaterialAttributes::GetInput(int32 InputIndex)
+int32 UMaterialExpressionExecEnd::Compile(class FMaterialCompiler* Compiler, int32 OutputIndex)
 {
-	switch (InputIndex)
-	{
-	case 0: return &MaterialAttributes;
-	default: checkNoEntry(); return nullptr;
-	}
-}
-
-uint32 UMaterialExpressionReturnMaterialAttributes::GetInputType(int32 InputIndex)
-{
-	switch (InputIndex)
-	{
-	case 0: return MCT_MaterialAttributes;
-	default: checkNoEntry(); return 0u;
-	}
+	const int32 Attributes = Material->MaterialAttributes.Compile(Compiler);
+	return Compiler->ReturnMaterialAttributes(Attributes);
 }
 #endif
 
diff --git a/Engine/Source/Runtime/Engine/Private/Materials/MaterialHLSLGenerator.cpp b/Engine/Source/Runtime/Engine/Private/Materials/MaterialHLSLGenerator.cpp
--- a/Engine/Source/Runtime/Engine/Private/Materials/MaterialHLSLGenerator.cpp
+++ b/Engine/Source/Runtime/Engine/Private/Materials/MaterialHLSLGenerator.cpp
@@ -16,10 +16,22 @@
 #include "HLSLTree/HLSLTreeCommon.h"
 #include "Containers/LazyPrintf.h"
 
-FMaterialHLSLGenerator::FMaterialHLSLGenerator(const FMaterialCompileTargetParameters& InCompileTarget, UE::HLSLTree::FTree& InOutTree)
-	: CompileTarget(InCompileTarget), HLSLTree(&InOutTree)
+FMaterialHLSLGenerator::FMaterialHLSLGenerator(UMaterial* InTargetMaterial, const FMaterialCompileTargetParameters& InCompileTarget, UE::HLSLTree::FTree& InOutTree)
+	: CompileTarget(InCompileTarget)
+	, TargetMaterial(InTargetMaterial)
+	, TargetMaterialFunction(nullptr)
+	, HLSLTree(&InOutTree)
+	, bGeneratedResult(false)
+{
+}
+
+FMaterialHLSLGenerator::FMaterialHLSLGenerator(UMaterialFunctionInterface* InTargetMaterialFunction, const FMaterialCompileTargetParameters& InCompileTarget, UE::HLSLTree::FTree& InOutTree)
+	: CompileTarget(InCompileTarget)
+	, TargetMaterial(nullptr)
+	, TargetMaterialFunction(InTargetMaterialFunction)
+	, HLSLTree(&InOutTree)
+	, bGeneratedResult(false)
 {
-	
 }
 
 void FMaterialHLSLGenerator::AcquireErrors(TArray<FString>& OutCompileErrors, TArray<UMaterialExpression*>& OutErrorExpressions)
@@ -67,6 +79,109 @@ EMaterialGenerateHLSLStatus FMaterialHLSLGenerator::Error(const FString& Message
 	return EMaterialGenerateHLSLStatus::Error;
 }
 
+static UE::HLSLTree::FExpression* CompileMaterialInput(FMaterialHLSLGenerator& Generator,
+	UE::HLSLTree::FScope& Scope,
+	EMaterialProperty InputProperty,
+	UMaterial* Material,
+	UE::HLSLTree::FExpression* AttributesExpression)
+{
+	UE::HLSLTree::FExpression* Expression = nullptr;
+	if (Material->IsPropertyActive(InputProperty))
+	{
+		FMaterialInputDescription InputDescription;
+		if (Material->GetExpressionInputDescription(InputProperty, InputDescription))
+		{
+			if (InputDescription.bUseConstant)
+			{
+				UE::Shader::FValue DefaultValue = FMaterialAttributeDefinitionMap::GetDefaultValue(InputProperty);
+				DefaultValue.NumComponents = UE::Shader::GetValueTypeDescription(InputDescription.Type).NumComponents;
+				if (InputDescription.ConstantValue != DefaultValue)
+				{
+					Expression = Generator.NewConstant(Scope, InputDescription.ConstantValue);
+				}
+			}
+			else
+			{
+				check(InputDescription.Input);
+				if (InputProperty >= MP_CustomizedUVs0 && InputProperty <= MP_CustomizedUVs7)
+				{
+					const int32 TexCoordIndex = (int32)InputProperty - MP_CustomizedUVs0;
+					if (TexCoordIndex < Material->NumCustomizedUVs)
+					{
+						Expression = InputDescription.Input->AcquireHLSLExpressionWithCast(Generator, Scope, InputDescription.Type);
+					}
+					if (!Expression)
+					{
+						Expression = Generator.NewTexCoord(Scope, TexCoordIndex);
+					}
+				}
+				else
+				{
+					Expression = InputDescription.Input->AcquireHLSLExpressionWithCast(Generator, Scope, InputDescription.Type);
+				}
+			}
+		}
+	}
+
+	if (Expression)
+	{
+		UE::HLSLTree::FExpressionSetMaterialAttribute* SetAttributeExpression = Generator.GetTree().NewExpression<UE::HLSLTree::FExpressionSetMaterialAttribute>(Scope);
+		SetAttributeExpression->AttributeID = FMaterialAttributeDefinitionMap::GetID(InputProperty);
+		SetAttributeExpression->AttributesExpression = AttributesExpression;
+		SetAttributeExpression->ValueExpression = Expression;
+		return SetAttributeExpression;
+	}
+
+	return AttributesExpression;
+}
+
+UE::HLSLTree::FStatement* FMaterialHLSLGenerator::NewResult(UE::HLSLTree::FScope& Scope)
+{
+	UE::HLSLTree::FStatement* Result = nullptr;
+	if (bGeneratedResult)
+	{
+		Error(TEXT("Multiple connections to execution output"));
+	}
+	else
+	{
+		if (TargetMaterial)
+		{
+			UE::HLSLTree::FExpression* AttributesExpression = nullptr;
+			if (TargetMaterial->bUseMaterialAttributes)
+			{
+				FMaterialInputDescription InputDescription;
+				if (TargetMaterial->GetExpressionInputDescription(MP_MaterialAttributes, InputDescription))
+				{
+					check(InputDescription.Type == UE::Shader::EValueType::MaterialAttributes);
+					AttributesExpression = InputDescription.Input->AcquireHLSLExpression(*this, Scope);
+				}
+			}
+			else
+			{
+				AttributesExpression = HLSLTree->NewExpression<UE::HLSLTree::FExpressionDefaultMaterialAttributes>(Scope);
+				for (int32 PropertyIndex = 0; PropertyIndex < MP_MAX; ++PropertyIndex)
+				{
+					const EMaterialProperty Property = (EMaterialProperty)PropertyIndex;
+					AttributesExpression = CompileMaterialInput(*this, Scope, Property, TargetMaterial, AttributesExpression);
+				}
+			}
+
+			if (AttributesExpression)
+			{
+				UE::HLSLTree::FStatementReturn* ReturnStatement = HLSLTree->NewStatement<UE::HLSLTree::FStatementReturn>(Scope);
+				ReturnStatement->Expression = AttributesExpression;
+				Result = ReturnStatement;
+			}
+		}
+		else
+		{
+			check(false);
+		}
+		bGeneratedResult = true;
+	}
+	return Result;
+}
+
 UE::HLSLTree::FExpressionConstant* FMaterialHLSLGenerator::NewConstant(UE::HLSLTree::FScope& Scope, const UE::Shader::FValue& Value)
 {
 	UE::HLSLTree::FExpressionConstant* Expression = HLSLTree->NewExpression<UE::HLSLTree::FExpressionConstant>(Scope, Value);
@@ -259,7 +374,7 @@ UE::HLSLTree::FStatement* FMaterialHLSLGenerator::AcquireStatement(UE::HLSLTree:
 	else
 	{
 		Statement = *PrevStatement;
-		if (!Scope.TryMoveStatement(Statement))
+		if (Statement && !Scope.TryMoveStatement(Statement))
 		{
 			// Could not move existing statement to the given scope
 			Error(TEXT("Invalid control flow"));
diff --git a/Engine/Source/Runtime/Engine/Private/Materials/MaterialHLSLTree.cpp b/Engine/Source/Runtime/Engine/Private/Materials/MaterialHLSLTree.cpp
--- a/Engine/Source/Runtime/Engine/Private/Materials/MaterialHLSLTree.cpp
+++ b/Engine/Source/Runtime/Engine/Private/Materials/MaterialHLSLTree.cpp
@@ -16,69 +16,13 @@ FMaterialHLSLTree::FMaterialHLSLTree()
 {
 }
 
-static UE::HLSLTree::FExpression* CompileMaterialInput(FMaterialHLSLGenerator& Generator,
-	UE::HLSLTree::FScope& Scope,
-	EMaterialProperty InputProperty,
-	UMaterial* Material,
-	UE::HLSLTree::FExpression* AttributesExpression)
-{
-	UE::HLSLTree::FExpression* Expression = nullptr;
-	if (Material->IsPropertyActive(InputProperty))
-	{
-		FMaterialInputDescription InputDescription;
-		if (Material->GetExpressionInputDescription(InputProperty, InputDescription))
-		{
-			if (InputDescription.bUseConstant)
-			{
-				UE::Shader::FValue DefaultValue = FMaterialAttributeDefinitionMap::GetDefaultValue(InputProperty);
-				DefaultValue.NumComponents = UE::Shader::GetValueTypeDescription(InputDescription.Type).NumComponents;
-				if (InputDescription.ConstantValue != DefaultValue)
-				{
-					Expression = Generator.NewConstant(Scope, InputDescription.ConstantValue);
-				}
-			}
-			else
-			{
-				check(InputDescription.Input);
-				if (InputProperty >= MP_CustomizedUVs0 && InputProperty <= MP_CustomizedUVs7)
-				{
-					const int32 TexCoordIndex = (int32)InputProperty - MP_CustomizedUVs0;
-					if (TexCoordIndex < Material->NumCustomizedUVs)
-					{
-						Expression = InputDescription.Input->AcquireHLSLExpressionWithCast(Generator, Scope, InputDescription.Type);
-					}
-					if (!Expression)
-					{
-						Expression = Generator.NewTexCoord(Scope, TexCoordIndex);
-					}
-				}
-				else
-				{
-					Expression = InputDescription.Input->AcquireHLSLExpressionWithCast(Generator, Scope, InputDescription.Type);
-				}
-			}
-		}
-	}
-
-	if (Expression)
-	{
-		UE::HLSLTree::FExpressionSetMaterialAttribute* SetAttributeExpression = Generator.GetTree().NewExpression<UE::HLSLTree::FExpressionSetMaterialAttribute>(Scope);
-		SetAttributeExpression->AttributeID = FMaterialAttributeDefinitionMap::GetID(InputProperty);
-		SetAttributeExpression->AttributesExpression = AttributesExpression;
-		SetAttributeExpression->ValueExpression = Expression;
-		return SetAttributeExpression;
-	}
-
-	return AttributesExpression;
-}
-
 void FMaterialHLSLTree::InitializeForMaterial(const FMaterialCompileTargetParameters& InCompilerTarget, FMaterial& InOutMaterial)
 {
 	UMaterialInterface* MaterialInterface = InOutMaterial.GetMaterialInterface();
 	UMaterial* BaseMaterial = MaterialInterface->GetMaterial();
 
 	HLSLTree = UE::HLSLTree::FTree::Create(Allocator);
-	FMaterialHLSLGenerator Generator(InCompilerTarget, *HLSLTree);
+	FMaterialHLSLGenerator Generator(BaseMaterial, InCompilerTarget, *HLSLTree);
 	UE::HLSLTree::FScope& RootScope = HLSLTree->GetRootScope();
 
 	if (BaseMaterial->IsCompiledWithExecutionFlow())
@@ -88,15 +32,7 @@ void FMaterialHLSLTree::InitializeForMaterial(const FMaterialCompileTargetParame
 	}
 	else
 	{
-		UE::HLSLTree::FExpression* AttributesExpression = Generator.GetTree().NewExpression<UE::HLSLTree::FExpressionDefaultMaterialAttributes>(RootScope);
-		for (int32 PropertyIndex = 0; PropertyIndex < MP_MAX; ++PropertyIndex)
-		{
-			const EMaterialProperty Property = (EMaterialProperty)PropertyIndex;
-			AttributesExpression = CompileMaterialInput(Generator, RootScope, Property, BaseMaterial, AttributesExpression);
-		}
-
-		UE::HLSLTree::FStatementReturn* ReturnStatement = Generator.GetTree().NewStatement<UE::HLSLTree::FStatementReturn>(RootScope);
-		ReturnStatement->Expression = AttributesExpression;
+		Generator.NewResult(RootScope);
 	}
 
 	Generator.AcquireErrors(InOutMaterial.CompileErrors, InOutMaterial.ErrorExpressions);
@@ -115,7 +51,7 @@ void FMaterialHLSLTree::InitializeForFunction(const FMaterialCompileTargetParame
 	NumFunctionOutputs = FunctionOutputs.Num();
 	FunctionOutputExpressions = New<UE::HLSLTree::FExpression*>(Allocator, NumFunctionOutputs);
 
-	FMaterialHLSLGenerator Generator(InCompilerTarget, *HLSLTree);
+	FMaterialHLSLGenerator Generator(InOutFunction, InCompilerTarget, *HLSLTree);
 	for (int32 InputIndex = 0; InputIndex < FunctionInputs.Num(); ++InputIndex)
 	{
 		const FFunctionExpressionInput& Input = FunctionInputs[InputIndex];
diff --git a/Engine/Source/Runtime/Engine/Private/Materials/MaterialShared.cpp b/Engine/Source/Runtime/Engine/Private/Materials/MaterialShared.cpp
--- a/Engine/Source/Runtime/Engine/Private/Materials/MaterialShared.cpp
+++ b/Engine/Source/Runtime/Engine/Private/Materials/MaterialShared.cpp
@@ -282,7 +282,7 @@ UE::HLSLTree::FExpression* FExpressionInput::AcquireHLSLExpression(FMaterialHLSL
 UE::HLSLTree::FExpression* FExpressionInput::AcquireHLSLExpressionWithCast(FMaterialHLSLGenerator& Generator, UE::HLSLTree::FScope& Scope, UE::Shader::EValueType Type) const
 {
 	UE::HLSLTree::FExpression* Result = AcquireHLSLExpression(Generator, Scope);
-	return Generator.NewCast(Scope, Type, Result);
+	return Result ? Generator.NewCast(Scope, Type, Result) : nullptr;
 }
 
 UE::HLSLTree::FTextureParameterDeclaration* FExpressionInput::AcquireHLSLTexture(FMaterialHLSLGenerator& Generator, UE::HLSLTree::FScope& Scope) const
diff --git a/Engine/Source/Runtime/Engine/Public/HLSLTree/HLSLTree.h b/Engine/Source/Runtime/Engine/Public/HLSLTree/HLSLTree.h
--- a/Engine/Source/Runtime/Engine/Public/HLSLTree/HLSLTree.h
+++ b/Engine/Source/Runtime/Engine/Public/HLSLTree/HLSLTree.h
@@ -228,7 +228,7 @@ class FStatement : public FNode
 {
 public:
 	/** Emits HLSL code for the statement. The generated code should include any required semi-colons and newlines */
-	virtual void EmitHLSL(FEmitContext& Context, FCodeWriter& Writer) const = 0;
+	virtual bool EmitHLSL(FEmitContext& Context, FCodeWriter& Writer) const = 0;
 
 	virtual ENodeVisitResult Visit(FNodeVisitor& Visitor) override;
 
@@ -249,7 +249,7 @@ public:
 	virtual ENodeVisitResult Visit(FNodeVisitor& Visitor) override;
 
 	/** Emits code for the given expression, either HLSL code or preshader bytecode */
-	virtual void EmitCode(FEmitContext& Context, FExpressionEmitResult& OutResult) const = 0;
+	virtual bool EmitCode(FEmitContext& Context, FExpressionEmitResult& OutResult) const = 0;
 };
 
 /**
@@ -323,9 +323,9 @@ class FScope final : public FNode
 {
 public:
 	virtual ENodeVisitResult Visit(FNodeVisitor& Visitor) override;
-	void EmitHLSL(FEmitContext& Context, FCodeWriter& OutWriter) const;
+	bool EmitHLSL(FEmitContext& Context, FCodeWriter& OutWriter) const;
 
-	void EmitUnscopedHLSL(FEmitContext& Context, FCodeWriter& OutWriter) const;
+	bool EmitUnscopedHLSL(FEmitContext& Context, FCodeWriter& OutWriter) const;
 
 	void AddDeclaration(FLocalDeclaration* Declaration);
 	void AddExpression(FExpression* Expression);
@@ -356,7 +356,7 @@ class FTree
 public:
 	static FTree* Create(FMemStackBase& Allocator);
 
-	void EmitHLSL(FEmitContext& Context, FCodeWriter& Writer) const;
+	bool EmitHLSL(FEmitContext& Context, FCodeWriter& Writer) const;
 
 	FScope& GetRootScope() const { return *RootScope; }
 
diff --git a/Engine/Source/Runtime/Engine/Public/HLSLTree/HLSLTreeCommon.h b/Engine/Source/Runtime/Engine/Public/HLSLTree/HLSLTreeCommon.h
--- a/Engine/Source/Runtime/Engine/Public/HLSLTree/HLSLTreeCommon.h
+++ b/Engine/Source/Runtime/Engine/Public/HLSLTree/HLSLTreeCommon.h
@@ -26,7 +26,7 @@ public:
 
 	Shader::FValue Value;
 
-	virtual void EmitCode(FEmitContext& Context, FExpressionEmitResult& OutResult) const override;
+	virtual bool EmitCode(FEmitContext& Context, FExpressionEmitResult& OutResult) const override;
 };
 
 class FExpressionLocalVariable : public FExpression
@@ -46,7 +46,7 @@ public:
 		return Result;
 	}
 
-	virtual void EmitCode(FEmitContext& Context, FExpressionEmitResult& OutResult) const override;
+	virtual bool EmitCode(FEmitContext& Context, FExpressionEmitResult& OutResult) const override;
 };
 
 class FExpressionParameter : public FExpression
@@ -68,7 +68,7 @@ public:
 		return Result;
 	}
 
-	virtual void EmitCode(FEmitContext& Context, FExpressionEmitResult& OutResult) const override;
+	virtual bool EmitCode(FEmitContext& Context, FExpressionEmitResult& OutResult) const override;
 };
 
 enum class EExternalInputType
@@ -99,7 +99,7 @@ public:
 
 	EExternalInputType InputType;
 
-	virtual void EmitCode(FEmitContext& Context, FExpressionEmitResult& OutResult) const override;
+	virtual bool EmitCode(FEmitContext& Context, FExpressionEmitResult& OutResult) const override;
 };
 
 class FExpressionTextureSample : public FExpression
@@ -128,7 +128,7 @@ public:
 		return Result;
 	}
 
-	virtual void EmitCode(FEmitContext& Context, FExpressionEmitResult& OutResult) const override;
+	virtual bool EmitCode(FEmitContext& Context, FExpressionEmitResult& OutResult) const override;
 };
 
 class FExpressionDefaultMaterialAttributes : public FExpression
@@ -136,7 +136,7 @@ class FExpressionDefaultMaterialAttributes : public FExpression
 public:
 	FExpressionDefaultMaterialAttributes() {}
 
-	virtual void EmitCode(FEmitContext& Context, FExpressionEmitResult& OutResult) const override;
+	virtual bool EmitCode(FEmitContext& Context, FExpressionEmitResult& OutResult) const override;
 };
 
 class FExpressionSetMaterialAttribute : public FExpression
@@ -159,7 +159,7 @@ public:
 		return Result;
 	}
 
-	virtual void EmitCode(FEmitContext& Context, FExpressionEmitResult& OutResult) const override;
+	virtual bool EmitCode(FEmitContext& Context, FExpressionEmitResult& OutResult) const override;
 };
 
 class FExpressionSelect : public FExpression
@@ -187,7 +187,7 @@ public:
 		return Result;
 	}
 
-	virtual void EmitCode(FEmitContext& Context, FExpressionEmitResult& OutResult) const override;
+	virtual bool EmitCode(FEmitContext& Context, FExpressionEmitResult& OutResult) const override;
 };
 
 class FExpressionBinaryOp : public FExpression
@@ -214,7 +214,7 @@ public:
 		return Result;
 	}
 
-	virtual void EmitCode(FEmitContext& Context, FExpressionEmitResult& OutResult) const override;
+	virtual bool EmitCode(FEmitContext& Context, FExpressionEmitResult& OutResult) const override;
 };
 
 struct FSwizzleParameters
@@ -248,7 +248,7 @@ public:
 		return Result;
 	}
 
-	virtual void EmitCode(FEmitContext& Context, FExpressionEmitResult& OutResult) const override;
+	virtual bool EmitCode(FEmitContext& Context, FExpressionEmitResult& OutResult) const override;
 };
 
 class FExpressionCast : public FExpression
@@ -274,13 +274,13 @@ public:
 		return Result;
 	}
 
-	virtual void EmitCode(FEmitContext& Context, FExpressionEmitResult& OutResult) const override;
+	virtual bool EmitCode(FEmitContext& Context, FExpressionEmitResult& OutResult) const override;
 };
 
 class FExpressionReflectionVector : public FExpression
 {
 public:
-	virtual void EmitCode(FEmitContext& Context, FExpressionEmitResult& OutResult) const override;
+	virtual bool EmitCode(FEmitContext& Context, FExpressionEmitResult& OutResult) const override;
 };
 
 class FExpressionFunctionInput : public FExpression
@@ -290,7 +290,7 @@ public:
 		: Name(InName), Type(InType), InputIndex(InIndex)
 	{}
 
-	virtual void EmitCode(FEmitContext& Context, FExpressionEmitResult& OutResult) const override;
+	virtual bool EmitCode(FEmitContext& Context, FExpressionEmitResult& OutResult) const override;
 
 	FName Name;
 	Shader::EValueType Type;
@@ -320,7 +320,7 @@ public:
 		return Result;
 	}
 
-	virtual void EmitCode(FEmitContext& Context, FExpressionEmitResult& OutResult) const override;
+	virtual bool EmitCode(FEmitContext& Context, FExpressionEmitResult& OutResult) const override;
 };
 
 class FStatementReturn : public FStatement
@@ -338,7 +338,7 @@ public:
 		return Result;
 	}
 
-	virtual void EmitHLSL(FEmitContext& Context, FCodeWriter& Writer) const override;
+	virtual bool EmitHLSL(FEmitContext& Context, FCodeWriter& Writer) const override;
 };
 
 class FStatementSetLocalVariable : public FStatement
@@ -358,7 +358,7 @@ public:
 		return Result;
 	}
 
-	virtual void EmitHLSL(FEmitContext& Context, FCodeWriter& Writer) const override;
+	virtual bool EmitHLSL(FEmitContext& Context, FCodeWriter& Writer) const override;
 };
 
 class FStatementIf : public FStatement
@@ -380,7 +380,7 @@ public:
 		return Result;
 	}
 
-	virtual void EmitHLSL(FEmitContext& Context, FCodeWriter& Writer) const override;
+	virtual bool EmitHLSL(FEmitContext& Context, FCodeWriter& Writer) const override;
 };
 
 class FStatementFor : public FStatement
@@ -402,7 +402,7 @@ public:
 		return Result;
 	}
 
-	virtual void EmitHLSL(FEmitContext& Context, FCodeWriter& Writer) const override;
+	virtual bool EmitHLSL(FEmitContext& Context, FCodeWriter& Writer) const override;
 };
 
 }
diff --git a/Engine/Source/Runtime/Engine/Public/MaterialHLSLGenerator.h b/Engine/Source/Runtime/Engine/Public/MaterialHLSLGenerator.h
--- a/Engine/Source/Runtime/Engine/Public/MaterialHLSLGenerator.h
+++ b/Engine/Source/Runtime/Engine/Public/MaterialHLSLGenerator.h
@@ -15,6 +15,7 @@ class FMaterial;
 class FMaterialCompilationOutput;
 struct FSharedShaderCompilerEnvironment;
 struct FMaterialCompileTargetParameters;
+class UMaterial;
 class UMaterialFunctionInterface;
 class UMaterialExpression;
 class UMaterialExpressionFunctionInput;
@@ -42,7 +43,9 @@ class FFunctionCall;
 class FMaterialHLSLGenerator
 {
 public:
-	FMaterialHLSLGenerator(const FMaterialCompileTargetParameters& InCompilerTarget,
+	FMaterialHLSLGenerator(UMaterial* InTargetMaterial, const FMaterialCompileTargetParameters& InCompilerTarget,
+		UE::HLSLTree::FTree& InOutTree);
+	FMaterialHLSLGenerator(UMaterialFunctionInterface* InTargetMaterialFunction, const FMaterialCompileTargetParameters& InCompilerTarget,
 		UE::HLSLTree::FTree& InOutTree);
 
 	const FMaterialCompileTargetParameters& GetCompileTarget() const { return CompileTarget; }
@@ -64,6 +67,8 @@ public:
 
 	UE::HLSLTree::FTree& GetTree() const { return *HLSLTree; }
 
+	UE::HLSLTree::FStatement* NewResult(UE::HLSLTree::FScope& Scope);
+
 	UE::HLSLTree::FExpressionConstant* NewConstant(UE::HLSLTree::FScope& Scope, const UE::Shader::FValue& Value);
 	UE::HLSLTree::FExpressionExternalInput* NewTexCoord(UE::HLSLTree::FScope& Scope, int32 Index);
 	UE::HLSLTree::FExpressionSwizzle* NewSwizzle(UE::HLSLTree::FScope& Scope, const UE::HLSLTree::FSwizzleParameters& Params, UE::HLSLTree::FExpression* Input);
@@ -127,6 +132,9 @@ private:
 	};
 
 	const FMaterialCompileTargetParameters& CompileTarget;
+	UMaterial* TargetMaterial;
+	UMaterialFunctionInterface* TargetMaterialFunction;
+
 	UE::HLSLTree::FTree* HLSLTree;
 	TArray<FExpressionKey> ExpressionStack;
 	TArray<FString> CompileErrors;
@@ -138,6 +146,7 @@ private:
 	TMap<FFunctionCallKey, UE::HLSLTree::FFunctionCall*> FunctionCallMap;
 	TMap<FExpressionKey, UE::HLSLTree::FExpression*> ExpressionMap;
 	TMap<UMaterialExpression*, UE::HLSLTree::FStatement*> StatementMap;
+	bool bGeneratedResult;
 };
 
 #endif // WITH_EDITOR
