Change default LogAudioConcurrency verbosity to Log

diff --git a/Engine/Source/Runtime/Engine/Classes/Sound/SoundConcurrency.h b/Engine/Source/Runtime/Engine/Classes/Sound/SoundConcurrency.h
--- a/Engine/Source/Runtime/Engine/Classes/Sound/SoundConcurrency.h
+++ b/Engine/Source/Runtime/Engine/Classes/Sound/SoundConcurrency.h
@@ -24,7 +24,7 @@ using FSoundOwnerObjectID = uint32;
 /** Sound instance (USoundBase) object ID. */
 using FSoundObjectID = uint32;
 
-ENGINE_API DECLARE_LOG_CATEGORY_EXTERN(LogAudioConcurrency, Verbose, All);
+ENGINE_API DECLARE_LOG_CATEGORY_EXTERN(LogAudioConcurrency, Log, All);
 
 UENUM()
 namespace EMaxConcurrentResolutionRule
