Undo workaround from CL16330409 and avoid depending on TextureFormat module in non-editor builds.

diff --git a/Engine/Source/Runtime/Engine/Engine.Build.cs b/Engine/Source/Runtime/Engine/Engine.Build.cs
--- a/Engine/Source/Runtime/Engine/Engine.Build.cs
+++ b/Engine/Source/Runtime/Engine/Engine.Build.cs
@@ -125,7 +125,6 @@ public class Engine : ModuleRules
 				"CrunchCompression",
 				"IntelISPC",
 				"TraceLog",
-				"TextureFormat",
 			}
 		);
 
@@ -148,8 +147,13 @@ public class Engine : ModuleRules
 		// to prevent "causes WARNING: Non-editor build cannot depend on non-redistributable modules."
 		if (Target.Type == TargetType.Editor)
 		{
-			// for now we depend on this
-			PrivateDependencyModuleNames.Add("RawMesh");
+			// for now we depend on these
+			PrivateDependencyModuleNames.AddRange(
+				new string[] {
+					"RawMesh",
+					"TextureFormat",
+				}
+			);
 		}
 
 		PrivateIncludePathModuleNames.AddRange(
