LevelInstance: WorldOrigin Pivot type

diff --git a/Engine/Source/Runtime/Engine/Private/LevelInstance/LevelInstanceEditorPivotActor.cpp b/Engine/Source/Runtime/Engine/Private/LevelInstance/LevelInstanceEditorPivotActor.cpp
--- a/Engine/Source/Runtime/Engine/Private/LevelInstance/LevelInstanceEditorPivotActor.cpp
+++ b/Engine/Source/Runtime/Engine/Private/LevelInstance/LevelInstanceEditorPivotActor.cpp
@@ -88,6 +88,10 @@ void ALevelInstancePivot::SetPivot(ELevelInstancePivotType PivotType, AActor* Pi
 		}
 		SetActorLocation(Location);
 	}
+	else if (PivotType == ELevelInstancePivotType::WorldOrigin)
+	{
+		SetActorLocation(FVector(0.f, 0.f, 0.f));
+	}
 	else // unsupported
 	{
 		check(0);
diff --git a/Engine/Source/Runtime/Engine/Private/LevelInstance/LevelInstanceSubsystem.cpp b/Engine/Source/Runtime/Engine/Private/LevelInstance/LevelInstanceSubsystem.cpp
--- a/Engine/Source/Runtime/Engine/Private/LevelInstance/LevelInstanceSubsystem.cpp
+++ b/Engine/Source/Runtime/Engine/Private/LevelInstance/LevelInstanceSubsystem.cpp
@@ -733,6 +733,10 @@ ALevelInstance* ULevelInstanceSubsystem::CreateLevelInstanceFrom(const TArray<AA
 		check(CreationParams.PivotActor);
 		LevelInstanceLocation = CreationParams.PivotActor->GetActorLocation();
 	}
+	else if (CreationParams.PivotType == ELevelInstancePivotType::WorldOrigin)
+	{
+		LevelInstanceLocation = FVector(0.f, 0.f, 0.f);
+	}
 	else
 	{
 		LevelInstanceLocation = ActorLocationBox.GetCenter();
diff --git a/Engine/Source/Runtime/Engine/Public/LevelInstance/LevelInstanceTypes.h b/Engine/Source/Runtime/Engine/Public/LevelInstance/LevelInstanceTypes.h
--- a/Engine/Source/Runtime/Engine/Public/LevelInstance/LevelInstanceTypes.h
+++ b/Engine/Source/Runtime/Engine/Public/LevelInstance/LevelInstanceTypes.h
@@ -50,7 +50,8 @@ enum class ELevelInstancePivotType : uint8
 {
 	CenterMinZ,
 	Center,
-	Actor
+	Actor,
+	WorldOrigin
 };
 
 USTRUCT()
