Added some profile tags.

diff --git a/Engine/Source/Runtime/Engine/Private/WorldPartition/WorldPartitionLevelStreamingDynamic.cpp b/Engine/Source/Runtime/Engine/Private/WorldPartition/WorldPartitionLevelStreamingDynamic.cpp
--- a/Engine/Source/Runtime/Engine/Private/WorldPartition/WorldPartitionLevelStreamingDynamic.cpp
+++ b/Engine/Source/Runtime/Engine/Private/WorldPartition/WorldPartitionLevelStreamingDynamic.cpp
@@ -306,6 +306,8 @@ bool UWorldPartitionLevelStreamingDynamic::IssueLoadRequests()
 
 void UWorldPartitionLevelStreamingDynamic::FinalizeRuntimeLevel()
 {
+	TRACE_CPUPROFILER_EVENT_SCOPE(UWorldPartitionLevelStreamingDynamic::FinalizeRuntimeLevel);
+
 	check(!HasLoadedLevel());
 	check(RuntimeLevel);
 	check(!bLoadRequestInProgress);
@@ -322,17 +324,23 @@ void UWorldPartitionLevelStreamingDynamic::FinalizeRuntimeLevel()
 		int32 PIEInstanceID = GetPackage()->PIEInstanceID;
 		check(PIEInstanceID != INDEX_NONE);
 
-		// PIE Fixup LazyObjectPtrs
-		FTemporaryPlayInEditorIDOverride SetPlayInEditorID(PIEInstanceID);
-		FFixupLazyObjectPtrForPIEArchive FixupLazyPointersAr;
-		FixupLazyPointersAr << RuntimeLevel;
+		{
+			TRACE_CPUPROFILER_EVENT_SCOPE(FixupLazyObjectPtr);
+			// PIE Fixup LazyObjectPtrs
+			FTemporaryPlayInEditorIDOverride SetPlayInEditorID(PIEInstanceID);
+			FFixupLazyObjectPtrForPIEArchive FixupLazyPointersAr;
+			FixupLazyPointersAr << RuntimeLevel;
+		}
 
-		// PIE Fixup SoftObjectPaths
-		RuntimeLevel->FixupForPIE(PIEInstanceID, [&](int32 InPIEInstanceID, FSoftObjectPath& ObjectPath)
 		{
-			// Remap Runtime Level's SoftObjectPath before each PIE Fixup to avoid doing 2 passes of serialization
-			OuterWorldPartition->RemapSoftObjectPath(ObjectPath);
-		});
+			TRACE_CPUPROFILER_EVENT_SCOPE(FixupSoftObjectPath);
+			// PIE Fixup SoftObjectPaths
+			RuntimeLevel->FixupForPIE(PIEInstanceID, [&](int32 InPIEInstanceID, FSoftObjectPath& ObjectPath)
+			{
+				// Remap Runtime Level's SoftObjectPath before each PIE Fixup to avoid doing 2 passes of serialization
+				OuterWorldPartition->RemapSoftObjectPath(ObjectPath);
+			});
+		}
 	}
 	else if (OuterWorld->IsGameWorld())
 	{
diff --git a/Engine/Source/Runtime/Engine/Private/WorldPartition/WorldPartitionLevelStreamingPolicy.cpp b/Engine/Source/Runtime/Engine/Private/WorldPartition/WorldPartitionLevelStreamingPolicy.cpp
--- a/Engine/Source/Runtime/Engine/Private/WorldPartition/WorldPartitionLevelStreamingPolicy.cpp
+++ b/Engine/Source/Runtime/Engine/Private/WorldPartition/WorldPartitionLevelStreamingPolicy.cpp
@@ -155,6 +155,8 @@ void UWorldPartitionLevelStreamingPolicy::RemapSoftObjectPath(FSoftObjectPath& O
 
 UObject* UWorldPartitionLevelStreamingPolicy::GetSubObject(const TCHAR* SubObjectPath)
 {
+	TRACE_CPUPROFILER_EVENT_SCOPE(UWorldPartitionLevelStreamingPolicy::GetSubObject);
+
 	// Support for subobjects such as Actor.Component
 	FString SubObjectName;
 	FString SubObjectContext;	
