Disable distributed shader compilation via FASTBuild until we handle the ShaderAutogen dependency.

#rnx

diff --git a/Engine/Source/Runtime/Engine/Private/ShaderCompiler/Experimental/ShaderCompileFASTBuild.cpp b/Engine/Source/Runtime/Engine/Private/ShaderCompiler/Experimental/ShaderCompileFASTBuild.cpp
--- a/Engine/Source/Runtime/Engine/Private/ShaderCompiler/Experimental/ShaderCompileFASTBuild.cpp
+++ b/Engine/Source/Runtime/Engine/Private/ShaderCompiler/Experimental/ShaderCompileFASTBuild.cpp
@@ -10,7 +10,8 @@
 
 namespace FASTBuildShaderCompilerVariables
 {
-	int32 Enabled = 1;
+	// Disabled until the ShaderAutogen dependency issue is fixed.
+	int32 Enabled = 0;
 	FAutoConsoleVariableRef CVarFASTBuildShaderCompile(
 		TEXT("r.FASTBuildShaderCompile"),
 		Enabled,
@@ -22,7 +23,7 @@ namespace FASTBuildShaderCompilerVariables
 	int32 SendPDB = 0;
 	FAutoConsoleVariableRef CVarFASTBuildSendPDB(
 		TEXT("r.FASTBuildShaderSendPDB"),
-		Enabled,
+		SendPDB,
 		TEXT("Enable when distributed shader compiler workers crash.\n")
 		TEXT("0: Do not send along debug information in FASTBuild. \n")
 		TEXT("1: Send along debug information in FASTBuild."),
