Reenable bulk serialization of UInstancedStaticMeshComponent::PerInstanceSMData in non LWC enabled builds, as the impact on serialization is significant.

diff --git a/Engine/Source/Runtime/Engine/Private/InstancedStaticMesh.cpp b/Engine/Source/Runtime/Engine/Private/InstancedStaticMesh.cpp
--- a/Engine/Source/Runtime/Engine/Private/InstancedStaticMesh.cpp
+++ b/Engine/Source/Runtime/Engine/Private/InstancedStaticMesh.cpp
@@ -2677,8 +2677,7 @@ void UInstancedStaticMeshComponent::Serialize(FArchive& Ar)
 	else
 #endif //WITH_EDITOR
 	{
-		// LWC_TODO: Perf pessimization. PerInstance data contains an FMatrix so can't be bulk serialized
-		PerInstanceSMData.BulkSerialize(Ar, true);
+		PerInstanceSMData.BulkSerialize(Ar, !UE_LARGE_WORLD_COORDINATES_DISABLED); // LWC_TODO: Need to force per element with LWC-on PerInstance data contains an FMatrix
 	}
 
 	if (!Ar.IsLoading() || Ar.CustomVer(FRenderingObjectVersion::GUID) >= FRenderingObjectVersion::PerInstanceCustomData)
