Add support for class redirectors in UActorDescContainer.

diff --git a/Engine/Source/Runtime/Engine/Private/WorldPartition/ActorDescContainer.cpp b/Engine/Source/Runtime/Engine/Private/WorldPartition/ActorDescContainer.cpp
--- a/Engine/Source/Runtime/Engine/Private/WorldPartition/ActorDescContainer.cpp
+++ b/Engine/Source/Runtime/Engine/Private/WorldPartition/ActorDescContainer.cpp
@@ -8,6 +8,7 @@
 #include "WorldPartition/WorldPartitionHandle.h"
 #include "Misc/Base64.h"
 #include "UObject/ObjectSaveContext.h"
+#include "UObject/LinkerLoad.h"
 #endif
 
 UActorDescContainer::UActorDescContainer(const FObjectInitializer& ObjectInitializer)
@@ -46,16 +47,23 @@ void UActorDescContainer::Initialize(UWorld* InWorld, FName InPackageName)
 
 	auto GetActorDescriptor = [this](const FAssetData& InAssetData) -> TUniquePtr<FWorldPartitionActorDesc>
 	{
-		FString ActorClassName;
+		FName ActorClassName;
 		static FName NAME_ActorMetaDataClass(TEXT("ActorMetaDataClass"));
 		if (InAssetData.GetTagValue(NAME_ActorMetaDataClass, ActorClassName))
 		{
+			// Look for a class redirectors
+			FName NewActorClassName = FLinkerLoad::FindNewNameForClass(ActorClassName, false);
+			if (NewActorClassName != NAME_None)
+			{
+				ActorClassName = NewActorClassName;
+			}
+
 			FString ActorMetaDataStr;
 			static FName NAME_ActorMetaData(TEXT("ActorMetaData"));
 			if (InAssetData.GetTagValue(NAME_ActorMetaData, ActorMetaDataStr))
 			{
 				bool bIsValidClass = true;
-				UClass* ActorClass = FindObject<UClass>(ANY_PACKAGE, *ActorClassName, true);
+				UClass* ActorClass = FindObject<UClass>(ANY_PACKAGE, *ActorClassName.ToString(), true);
 
 				if (!ActorClass)
 				{
