Adding a "Preview Blend Mode" on material functions that allows you to modifying the blend mode of the preview material. This is saved off on the material function so that artists can pick what they'd prefer the default preview behavior to be.

Contributed by [REDACTED_PERSON]

diff --git a/Engine/Source/Editor/MaterialEditor/Private/MaterialEditor.cpp b/Engine/Source/Editor/MaterialEditor/Private/MaterialEditor.cpp
--- a/Engine/Source/Editor/MaterialEditor/Private/MaterialEditor.cpp
+++ b/Engine/Source/Editor/MaterialEditor/Private/MaterialEditor.cpp
@@ -2331,6 +2331,11 @@ void FMaterialEditor::UpdatePreviewMaterial( bool bForce )
 		// So that RebuildMaterialFunctionInfo will see all the nested material functions that may need to be updated
 		ExpressionPreviewMaterial->Expressions = Material->Expressions;
 
+		if (MaterialFunction)
+		{
+			ExpressionPreviewMaterial->BlendMode = MaterialFunction->PreviewBlendMode;
+		}
+
 		FMaterialUpdateContext UpdateContext(FMaterialUpdateContext::EOptions::SyncWithRenderingThread);
 		UpdateContext.AddMaterial(ExpressionPreviewMaterial);
 
diff --git a/Engine/Source/Runtime/Engine/Classes/Materials/MaterialFunction.h b/Engine/Source/Runtime/Engine/Classes/Materials/MaterialFunction.h
--- a/Engine/Source/Runtime/Engine/Classes/Materials/MaterialFunction.h
+++ b/Engine/Source/Runtime/Engine/Classes/Materials/MaterialFunction.h
@@ -75,6 +75,9 @@ class UMaterialFunction : public UMaterialFunctionInterface
 	UPROPERTY()
 	TArray<TObjectPtr<class UMaterialExpressionMaterialFunctionCall>> DependentFunctionExpressionCandidates;
 
+	/** Determines the blend mode when previewing a material function. */
+	UPROPERTY(EditAnywhere, Category = Preview, AssetRegistrySearchable)
+	TEnumAsByte<enum EBlendMode> PreviewBlendMode = BLEND_Opaque;
 private:
 	/** Transient flag used to track re-entrance in recursive functions like IsDependent. */
 	UPROPERTY(transient)
