Opt out Encoded TextureCube out of package if the platform doesn't need EncodedData

diff --git a/Engine/Config/BaseEngine.ini b/Engine/Config/BaseEngine.ini
--- a/Engine/Config/BaseEngine.ini
+++ b/Engine/Config/BaseEngine.ini
@@ -864,6 +864,8 @@ NearClipPlane=10.0
 +PropertyRedirects=(OldName="CameraShakeSourceComponent.bAutoPlay",NewName="bAutoStart")
 +PropertyRedirects=(OldName="FieldSystemComponent.IsChaosField",NewName="bIsChaosField")
 +PropertyRedirects=(OldName="FieldSystemComponent.IsGlobalField",NewName="bIsWorldField")
+
++EnumRedirects=(OldName="TextureCompressionSettings",ValueChanges=(("TC_ReflectionCapture","TC_EncodedReflectionCapture")))
 +EnumRedirects=(OldName="EEnvQueryTrace",ValueChanges=(("Geometry","GeometryByChannel")))
 
 ;5.0
diff --git a/Engine/Source/Developer/Android/AndroidTargetPlatform/Private/AndroidTargetPlatform.cpp b/Engine/Source/Developer/Android/AndroidTargetPlatform/Private/AndroidTargetPlatform.cpp
--- a/Engine/Source/Developer/Android/AndroidTargetPlatform/Private/AndroidTargetPlatform.cpp
+++ b/Engine/Source/Developer/Android/AndroidTargetPlatform/Private/AndroidTargetPlatform.cpp
@@ -460,7 +460,7 @@ void FAndroidTargetPlatform::GetTextureFormats( const UTexture* InTexture, TArra
 			{
 				FormatPerLayer[LayerIndex] = AndroidTexFormat::NameBGRA8;
 			}
-			else if (LayerFormatSettings.CompressionSettings == TC_ReflectionCapture && !LayerFormatSettings.CompressionNone)
+			else if (LayerFormatSettings.CompressionSettings == TC_EncodedReflectionCapture && !LayerFormatSettings.CompressionNone)
 			{
 				FormatPerLayer[LayerIndex] = AndroidTexFormat::NameETC2_RGBA;
 			}
diff --git a/Engine/Source/Developer/Android/AndroidTargetPlatform/Private/AndroidTargetPlatform.h b/Engine/Source/Developer/Android/AndroidTargetPlatform/Private/AndroidTargetPlatform.h
--- a/Engine/Source/Developer/Android/AndroidTargetPlatform/Private/AndroidTargetPlatform.h
+++ b/Engine/Source/Developer/Android/AndroidTargetPlatform/Private/AndroidTargetPlatform.h
@@ -383,7 +383,7 @@ public:
 				{
 					FTextureFormatSettings FormatSettings;
 					Cube->GetDefaultFormatSettings(FormatSettings);
-					if (FormatSettings.CompressionSettings == TC_ReflectionCapture && !FormatSettings.CompressionNone)
+					if (FormatSettings.CompressionSettings == TC_EncodedReflectionCapture && !FormatSettings.CompressionNone)
 					{
 						TextureFormatName = FName(TEXT("ETC2_RGBA"));
 					}
diff --git a/Engine/Source/Developer/IOS/IOSTargetPlatform/Private/IOSTargetPlatform.cpp b/Engine/Source/Developer/IOS/IOSTargetPlatform/Private/IOSTargetPlatform.cpp
--- a/Engine/Source/Developer/IOS/IOSTargetPlatform/Private/IOSTargetPlatform.cpp
+++ b/Engine/Source/Developer/IOS/IOSTargetPlatform/Private/IOSTargetPlatform.cpp
@@ -611,7 +611,7 @@ void FIOSTargetPlatform::GetTextureFormats( const UTexture* Texture, TArray< TAr
 			{
 				FTextureFormatSettings FormatSettings;
 				Cube->GetDefaultFormatSettings(FormatSettings);
-				if (FormatSettings.CompressionSettings == TC_ReflectionCapture && !FormatSettings.CompressionNone)
+				if (FormatSettings.CompressionSettings == TC_EncodedReflectionCapture && !FormatSettings.CompressionNone)
 				{
 					TextureFormatName = FName(TEXT("ETC2_RGBA"));
 				}
diff --git a/Engine/Source/Runtime/Engine/Classes/Engine/TextureCube.h b/Engine/Source/Runtime/Engine/Classes/Engine/TextureCube.h
--- a/Engine/Source/Runtime/Engine/Classes/Engine/TextureCube.h
+++ b/Engine/Source/Runtime/Engine/Classes/Engine/TextureCube.h
@@ -41,6 +41,7 @@ public:
 	ENGINE_API virtual void GetAssetRegistryTags(TArray<FAssetRegistryTag>& OutTags) const override;
 	ENGINE_API virtual FString GetDesc() override;
 	ENGINE_API virtual void GetResourceSizeEx(FResourceSizeEx& CumulativeResourceSize) override;
+	ENGINE_API virtual bool NeedsLoadForTargetPlatform(const ITargetPlatform* TargetPlatform) const override;
 	//~ End UObject Interface.
 
 	ENGINE_API int32 GetSizeX() const;
diff --git a/Engine/Source/Runtime/Engine/Classes/Engine/TextureDefines.h b/Engine/Source/Runtime/Engine/Classes/Engine/TextureDefines.h
--- a/Engine/Source/Runtime/Engine/Classes/Engine/TextureDefines.h
+++ b/Engine/Source/Runtime/Engine/Classes/Engine/TextureDefines.h
@@ -290,6 +290,6 @@ enum TextureCompressionSettings
 	TC_HDR_Compressed			UMETA(DisplayName = "HDRCompressed (RGB, BC6H, DX11)"),
 	TC_BC7						UMETA(DisplayName = "BC7 (DX11, optional A)"),
 	TC_HalfFloat				UMETA(DisplayName = "Half Float (R16F)"),
-	TC_ReflectionCapture		UMETA(Hidden),
+	TC_EncodedReflectionCapture		UMETA(Hidden),
 	TC_MAX,
 };
diff --git a/Engine/Source/Runtime/Engine/Private/Components/ReflectionCaptureComponent.cpp b/Engine/Source/Runtime/Engine/Private/Components/ReflectionCaptureComponent.cpp
--- a/Engine/Source/Runtime/Engine/Private/Components/ReflectionCaptureComponent.cpp
+++ b/Engine/Source/Runtime/Engine/Private/Components/ReflectionCaptureComponent.cpp
@@ -675,7 +675,7 @@ void GenerateEncodedHDRTextureCube(UMapBuildDataRegistry* Registry, FReflectionC
 #if WITH_EDITOR
 	UTextureFactory* TextureFactory = NewObject<UTextureFactory>();
 
-	TextureFactory->CompressionSettings = TC_ReflectionCapture;
+	TextureFactory->CompressionSettings = TC_EncodedReflectionCapture;
 	UTextureCube* TextureCube = TextureFactory->CreateTextureCube(Registry, FName(TextureName), RF_Public);
 
 	if (TextureCube)
@@ -702,7 +702,7 @@ void GenerateEncodedHDRTextureCube(UMapBuildDataRegistry* Registry, FReflectionC
 			bIsCompressed = CaptureComponent->MobileReflectionCompression == EMobileReflectionCompression::Default ? bIsReflectionCaptureCompressionProjectSetting : CaptureComponent->MobileReflectionCompression == EMobileReflectionCompression::On;
 		}
 
-		TextureCube->CompressionSettings = TC_ReflectionCapture;
+		TextureCube->CompressionSettings = TC_EncodedReflectionCapture;
 		TextureCube->CompressionNone = !bIsCompressed;
 		TextureCube->CompressionQuality = TCQ_Highest;
 		TextureCube->Filter = TF_Trilinear;
diff --git a/Engine/Source/Runtime/Engine/Private/MapBuildData.cpp b/Engine/Source/Runtime/Engine/Private/MapBuildData.cpp
--- a/Engine/Source/Runtime/Engine/Private/MapBuildData.cpp
+++ b/Engine/Source/Runtime/Engine/Private/MapBuildData.cpp
@@ -302,7 +302,7 @@ FArchive& operator<<(FArchive& Ar, FReflectionCaptureMapBuildData& ReflectionCap
 
 	TArray<FName> Formats;
 
-	if (Ar.IsSaving() && Ar.IsCooking())
+	if (Ar.IsCooking())
 	{
 		// Get all the reflection capture formats that the target platform wants
 		Ar.CookingTarget()->GetReflectionCaptureFormats(Formats);
@@ -320,15 +320,7 @@ FArchive& operator<<(FArchive& Ar, FReflectionCaptureMapBuildData& ReflectionCap
 
 	if (Ar.CustomVer(FMobileObjectVersion::GUID) >= FMobileObjectVersion::StoreReflectionCaptureCompressedMobile)
 	{
-		if (Ar.IsCooking() && !Formats.Contains(EncodedHDR))
-		{
-			UTextureCube* StrippedData = NULL;
-			Ar << StrippedData;
-		}
-		else
-		{
-			Ar << ReflectionCaptureMapBuildData.EncodedCaptureData;
-		}
+		Ar << ReflectionCaptureMapBuildData.EncodedCaptureData;
 	}
 	else
 	{
diff --git a/Engine/Source/Runtime/Engine/Private/Texture.cpp b/Engine/Source/Runtime/Engine/Private/Texture.cpp
--- a/Engine/Source/Runtime/Engine/Private/Texture.cpp
+++ b/Engine/Source/Runtime/Engine/Private/Texture.cpp
@@ -2005,7 +2005,7 @@ FName GetDefaultTextureFormatName( const ITargetPlatform* TargetPlatform, const
 		|| (Texture->LODGroup == TEXTUREGROUP_Bokeh)
 		|| (Texture->LODGroup == TEXTUREGROUP_IESLightProfile)
 		|| (Texture->GetMaterialType() == MCT_VolumeTexture && !bSupportCompressedVolumeTexture)
-		|| FormatSettings.CompressionSettings == TC_ReflectionCapture;
+		|| FormatSettings.CompressionSettings == TC_EncodedReflectionCapture;
 
 	if (!bNoCompression && Texture->PowerOfTwoMode == ETexturePowerOfTwoSetting::None)
 	{
diff --git a/Engine/Source/Runtime/Engine/Private/TextureCube.cpp b/Engine/Source/Runtime/Engine/Private/TextureCube.cpp
--- a/Engine/Source/Runtime/Engine/Private/TextureCube.cpp
+++ b/Engine/Source/Runtime/Engine/Private/TextureCube.cpp
@@ -10,6 +10,7 @@
 #include "EngineUtils.h"
 #include "DeviceProfiles/DeviceProfile.h"
 #include "DeviceProfiles/DeviceProfileManager.h"
+#include "Interfaces/ITargetPlatform.h"
 #include "TextureCompiler.h"
 #include "Misc/ScopedSlowTask.h"
 #include "UObject/StrongObjectPtr.h"
@@ -617,6 +618,21 @@ void UTextureCube::GetResourceSizeEx(FResourceSizeEx& CumulativeResourceSize)
 	}
 }
 
+bool UTextureCube::NeedsLoadForTargetPlatform(const ITargetPlatform* TargetPlatform) const
+{
+	// TC_EncodedReflectionCapture is used only for EncodedCapture for Mobile and is hidden for the user.
+	if (CompressionSettings == TC_EncodedReflectionCapture)
+	{
+		const static FName EncodedHDR(TEXT("EncodedHDR"));
+		TArray<FName> Formats;
+
+		TargetPlatform->GetReflectionCaptureFormats(Formats);
+
+		return Formats.Contains(EncodedHDR);
+	}
+	return true;
+}
+
 #if WITH_EDITOR
 bool UTextureCube::IsDefaultTexture() const
 {
