Reduced allocations in UWorldPartitionStreamingPolicy::UpdateStreamingState when calling GetStreamingCells.

diff --git a/Engine/Source/Runtime/Engine/Private/WorldPartition/WorldPartitionRuntimeHash.cpp b/Engine/Source/Runtime/Engine/Private/WorldPartition/WorldPartitionRuntimeHash.cpp
--- a/Engine/Source/Runtime/Engine/Private/WorldPartition/WorldPartitionRuntimeHash.cpp
+++ b/Engine/Source/Runtime/Engine/Private/WorldPartition/WorldPartitionRuntimeHash.cpp
@@ -278,4 +278,10 @@ void UWorldPartitionRuntimeHash::SortStreamingCellsByImportance(const TSet<const
 	});
 }
 
+void UWorldPartitionRuntimeHash::FStreamingSourceCells::AddCell(const UWorldPartitionRuntimeCell* InCell, const FWorldPartitionStreamingSource& InSource)
+{
+	InCell->CacheStreamingSourceInfo(InSource);
+	Cells.Add(InCell);
+}
+
 #undef LOCTEXT_NAMESPACE
\ No newline at end of file
diff --git a/Engine/Source/Runtime/Engine/Private/WorldPartition/WorldPartitionStreamingPolicy.cpp b/Engine/Source/Runtime/Engine/Private/WorldPartition/WorldPartitionStreamingPolicy.cpp
--- a/Engine/Source/Runtime/Engine/Private/WorldPartition/WorldPartitionStreamingPolicy.cpp
+++ b/Engine/Source/Runtime/Engine/Private/WorldPartition/WorldPartitionStreamingPolicy.cpp
@@ -6,7 +6,6 @@
 
 #include "WorldPartition/WorldPartitionStreamingPolicy.h"
 #include "WorldPartition/WorldPartitionRuntimeCell.h"
-#include "WorldPartition/WorldPartitionRuntimeHash.h"
 #include "WorldPartition/WorldPartitionStreamingSource.h"
 #include "WorldPartition/WorldPartition.h"
 #include "WorldPartition/DataLayer/DataLayerSubsystem.h"
@@ -161,16 +160,16 @@ void UWorldPartitionStreamingPolicy::UpdateStreamingState()
 			return;
 		}
 
-		UWorldPartitionRuntimeHash::FStreamingSourceCells ActivateCells;
-		UWorldPartitionRuntimeHash::FStreamingSourceCells LoadCells;
-		TSet<const UWorldPartitionRuntimeCell*>& LoadStreamingCells = LoadCells.GetCells();
-		TSet<const UWorldPartitionRuntimeCell*>& ActivateStreamingCells = ActivateCells.GetCells();
+		TSet<const UWorldPartitionRuntimeCell*>& ActivateStreamingCells = FrameActivateCells.GetCells();
+		TSet<const UWorldPartitionRuntimeCell*>& LoadStreamingCells = FrameLoadCells.GetCells();
+		check(ActivateStreamingCells.IsEmpty());
+		check(LoadStreamingCells.IsEmpty());
 
 		// When uninitializing, UpdateStreamingState is called, but we don't want any cells to be loaded
 		if (WorldPartition->IsInitialized())
 		{
 			UWorldPartitionRuntimeCell::DirtyStreamingSourceCacheEpoch();
-			WorldPartition->RuntimeHash->GetStreamingCells(StreamingSources, ActivateCells, LoadCells);
+			WorldPartition->RuntimeHash->GetStreamingCells(StreamingSources, FrameActivateCells, FrameLoadCells);
 
 			// Activation superseeds Loading
 			LoadStreamingCells = LoadStreamingCells.Difference(ActivateStreamingCells);
@@ -230,6 +229,10 @@ void UWorldPartitionStreamingPolicy::UpdateStreamingState()
 		{
 			SetTargetStateForCells(EWorldPartitionRuntimeCellState::Loaded, ToLoadCells);
 		}
+
+		// Reset frame StreamingSourceCells (optimization to avoid reallocation at every call to UpdateStreamingState)
+		FrameActivateCells.Reset();
+		FrameLoadCells.Reset();
 	}
 }
 
diff --git a/Engine/Source/Runtime/Engine/Public/WorldPartition/WorldPartitionRuntimeHash.h b/Engine/Source/Runtime/Engine/Public/WorldPartition/WorldPartitionRuntimeHash.h
--- a/Engine/Source/Runtime/Engine/Public/WorldPartition/WorldPartitionRuntimeHash.h
+++ b/Engine/Source/Runtime/Engine/Public/WorldPartition/WorldPartitionRuntimeHash.h
@@ -5,7 +5,6 @@
 #include "UObject/ObjectMacros.h"
 #include "UObject/Object.h"
 #include "WorldPartition.h"
-#include "WorldPartitionStreamingPolicy.h"
 #include "WorldPartitionActorDescView.h"
 #include "WorldPartition/ActorDescList.h"
 #include "WorldPartition/WorldPartitionHandle.h"
@@ -47,12 +46,8 @@ class ENGINE_API UWorldPartitionRuntimeHash : public UObject
 	class FStreamingSourceCells
 	{
 	public:
-		void AddCell(const UWorldPartitionRuntimeCell* InCell, const FWorldPartitionStreamingSource& InSource)
-		{
-			InCell->CacheStreamingSourceInfo(InSource);
-			Cells.Add(InCell);
-		}
-
+		void AddCell(const UWorldPartitionRuntimeCell* InCell, const FWorldPartitionStreamingSource& InSource);
+		void Reset() { Cells.Reset(); }
 		int32 Num() const { return Cells.Num(); }
 		TSet<const UWorldPartitionRuntimeCell*>& GetCells() { return Cells; }
 
diff --git a/Engine/Source/Runtime/Engine/Public/WorldPartition/WorldPartitionStreamingPolicy.h b/Engine/Source/Runtime/Engine/Public/WorldPartition/WorldPartitionStreamingPolicy.h
--- a/Engine/Source/Runtime/Engine/Public/WorldPartition/WorldPartitionStreamingPolicy.h
+++ b/Engine/Source/Runtime/Engine/Public/WorldPartition/WorldPartitionStreamingPolicy.h
@@ -11,6 +11,7 @@
 
 #include "Containers/Set.h"
 #include "WorldPartition/WorldPartitionRuntimeCell.h"
+#include "WorldPartition/WorldPartitionRuntimeHash.h"
 #include "WorldPartition/WorldPartitionStreamingSource.h"
 #include "WorldPartitionStreamingPolicy.generated.h"
 
@@ -59,6 +60,9 @@ protected:
 	TSet<const UWorldPartitionRuntimeCell*> LoadedCells;
 	TSet<const UWorldPartitionRuntimeCell*> ActivatedCells;
 	TArray<FWorldPartitionStreamingSource> StreamingSources;
+
+	UWorldPartitionRuntimeHash::FStreamingSourceCells FrameActivateCells;
+	UWorldPartitionRuntimeHash::FStreamingSourceCells FrameLoadCells;
 	
 	int32 UpdateStreamingStateEpoch;
 	mutable int32 SortedAddToWorldCellsEpoch;
