Activate the deprecation of direct access to the Resource member of UTexture in favor of the GetResource() accessor.

#rnx

diff --git a/Engine/Source/Runtime/Engine/Classes/Engine/Texture.h b/Engine/Source/Runtime/Engine/Classes/Engine/Texture.h
--- a/Engine/Source/Runtime/Engine/Classes/Engine/Texture.h
+++ b/Engine/Source/Runtime/Engine/Classes/Engine/Texture.h
@@ -997,9 +997,7 @@ private:
 	class FTextureResource* PrivateResourceRenderThread;
 
 public:
-	// The deprecation will be enabled further along the dev cycle and fixed accordingly.
-	// This is to avoid merge conflicts with other branches that fixing this deprecation might cause.
-	//UE_DEPRECATED(5.00, "Use GetResource() / SetResource() accessors instead.")
+	UE_DEPRECATED(5.00, "Use GetResource() / SetResource() accessors instead.")
 	TFieldPtrAccessor<FTextureResource> Resource;
 
 	/** Set texture's resource, can be NULL */
