Fixed broken Strata a translator code preventing more than one vertical layer.

diff --git a/Engine/Source/Runtime/Engine/Private/Materials/HLSLMaterialTranslator.cpp b/Engine/Source/Runtime/Engine/Private/Materials/HLSLMaterialTranslator.cpp
--- a/Engine/Source/Runtime/Engine/Private/Materials/HLSLMaterialTranslator.cpp
+++ b/Engine/Source/Runtime/Engine/Private/Materials/HLSLMaterialTranslator.cpp
@@ -1644,12 +1644,6 @@ void FHLSLMaterialTranslator::GetMaterialEnvironment(EShaderPlatform InPlatform,
 	// if duals source blending (colored transmittance) is not supported on a platform, it will fall back to standard alpha blending (grey scale transmittance)
 	OutEnvironment.SetDefine(TEXT("DUAL_SOURCE_COLOR_BLENDING_ENABLED"), bMaterialRequestsDualSourceBlending && Material->IsDualBlendingEnabled(Platform) ? TEXT("1") : TEXT("0"));
 
-	// Translate() is called before getting here so we can create related define
-	if (StrataMaterialAnalysis.RequestedBSDFCount > 0)
-	{
-		OutEnvironment.SetDefine(TEXT("STRATA_CLAMPED_LAYER_COUNT"), StrataMaterialAnalysis.ClampedLayerCount);
-	}
-
 	if (bMaterialIsStrata)
 	{
 		StrataPixelNormalInitializerValues += FString::Printf(TEXT("\n\n\n\t// Strata normal and tangent\n"));
@@ -1762,7 +1756,12 @@ void FHLSLMaterialTranslator::GetMaterialEnvironment(EShaderPlatform InPlatform,
 			static const auto CVarStrataBytePerPixel = IConsoleManager::Get().FindTConsoleVariableDataInt(TEXT("r.Strata.BytesPerPixel"));
 			check(CVarStrataBytePerPixel);
 			const uint32 StrataBytePerPixel = CVarStrataBytePerPixel ? CVarStrataBytePerPixel->GetValueOnAnyThread() : 0;
-			StrataMaterialAnalysis = StrataCompilationInfoMaterialAnalysis(this, StrataCompilationInfo, StrataBytePerPixel);
+			FStrataMaterialAnalysisResult StrataMaterialAnalysis = StrataCompilationInfoMaterialAnalysis(this, StrataCompilationInfo, StrataBytePerPixel);
+
+			if (StrataMaterialAnalysis.RequestedBSDFCount > 0)
+			{
+				OutEnvironment.SetDefine(TEXT("STRATA_CLAMPED_LAYER_COUNT"), StrataMaterialAnalysis.ClampedLayerCount);
+			}
 
 			FString StrataMaterialDescription;
 
diff --git a/Engine/Source/Runtime/Engine/Private/Materials/HLSLMaterialTranslator.h b/Engine/Source/Runtime/Engine/Private/Materials/HLSLMaterialTranslator.h
--- a/Engine/Source/Runtime/Engine/Private/Materials/HLSLMaterialTranslator.h
+++ b/Engine/Source/Runtime/Engine/Private/Materials/HLSLMaterialTranslator.h
@@ -412,8 +412,6 @@ protected:
 	/** Tracks shared normal & tangent used by strata materials, mapping a normal code chunk hash to a SharedMaterialInfo.
 	 * A normal code chunk hash can point to multiple shared info in case it is paired with different tangents. */
 	TMultiMap<uint64, FStrataSharedNormalInfo> CodeChunkToStrataSharedNormal;
-	/** The final material topology and memory analysis that can be feedback to the user in the editor UI */
-	FStrataMaterialAnalysisResult StrataMaterialAnalysis;
 
 	/** Tracks the total number of vt samples in the shader. */
 	uint32 NumVtSamples;
