When cooking, don't allow actors in dynamically loaded data layers to be loaded by WorldPartition initialization.
This is necessary for Cook on the fly to work as world's package is saved before any generated packages and PersistentLevel must not contain any always loaded actors that are part of a dynamically loaded data layer.

diff --git a/Engine/Source/Runtime/Engine/Private/WorldPartition/WorldPartition.cpp b/Engine/Source/Runtime/Engine/Private/WorldPartition/WorldPartition.cpp
--- a/Engine/Source/Runtime/Engine/Private/WorldPartition/WorldPartition.cpp
+++ b/Engine/Source/Runtime/Engine/Private/WorldPartition/WorldPartition.cpp
@@ -433,7 +433,7 @@ void UWorldPartition::Initialize(UWorld* InWorld, const FTransform& InTransform)
 
 	if (bEditorOnly)
 	{
-		// Make sure to preload only AWorldDataLayers actor first (ShouldActorBeLoaded requires it)
+		// Make sure to preload only AWorldDataLayers actor first (ShouldActorBeLoadedByEditorCells requires it)
 		FWorldPartitionReference WorldDataLayersActor;
 		for (UWorldPartitionEditorCell::FActorHandle& ActorHandle : EditorHash->GetAlwaysLoadedCell()->Actors)
 		{
@@ -801,7 +801,7 @@ bool UWorldPartition::UpdateEditorCells(TFunctionRef<bool(TArray<UWorldPartition
 	{
 		for (const UWorldPartitionEditorCell::FActorReference& ActorDesc : Cell->LoadedActors)
 		{
-			if (!bIsCellShouldBeLoaded || !ShouldActorBeLoaded(*ActorDesc))
+			if (!bIsCellShouldBeLoaded || !ShouldActorBeLoadedByEditorCells(*ActorDesc))
 			{
 				UnloadCount.FindOrAdd(*ActorDesc, 0)++;
 			}
@@ -873,29 +873,38 @@ bool UWorldPartition::UpdateEditorCells(TFunctionRef<bool(TArray<UWorldPartition
 	return true;
 }
 
-bool UWorldPartition::ShouldActorBeLoaded(const FWorldPartitionActorDesc* ActorDesc) const
+bool UWorldPartition::ShouldActorBeLoadedByEditorCells(const FWorldPartitionActorDesc* ActorDesc) const
 {
-	// No filtering when running cook commandlet
-	if (IsRunningCookCommandlet())
-	{
-		return true;
-	}
-
 	if (const AWorldDataLayers* WorldDataLayers = GetWorld()->GetWorldDataLayers())
 	{
-		uint32 NumValidLayers = 0;
-		for (const FName& DataLayerName : ActorDesc->GetDataLayers())
+		if (IsRunningCookCommandlet())
 		{
-			if (const UDataLayer* DataLayer = WorldDataLayers->GetDataLayerFromName(DataLayerName))
+			// When running cook commandlet, dont allow loading of actors with dynamically loaded data layers
+			for (const FName& DataLayerName : ActorDesc->GetDataLayers())
 			{
-				if (DataLayer->IsDynamicallyLoadedInEditor())
+				const UDataLayer* DataLayer = WorldDataLayers->GetDataLayerFromName(DataLayerName);
+				if (DataLayer && DataLayer->IsDynamicallyLoaded())
 				{
-					return true;
+					return false;
+				}
+			}
+		}
+		else
+		{
+			uint32 NumValidLayers = 0;
+			for (const FName& DataLayerName : ActorDesc->GetDataLayers())
+			{
+				if (const UDataLayer* DataLayer = WorldDataLayers->GetDataLayerFromName(DataLayerName))
+				{
+					if (DataLayer->IsDynamicallyLoadedInEditor())
+					{
+						return true;
+					}
+					NumValidLayers++;
 				}
-				NumValidLayers++;
 			}
+			return !NumValidLayers;
 		}
-		return !NumValidLayers;
 	}
 
 	return true;
@@ -930,7 +939,7 @@ void UWorldPartition::UpdateLoadingEditorCell(UWorldPartitionEditorCell* Cell, b
 			check(Actor || !ActorHandle->GetHardRefCount());
 
 			// Filter actor against DataLayers
-			if (ShouldActorBeLoaded(*ActorHandle))
+			if (ShouldActorBeLoadedByEditorCells(*ActorHandle))
 			{
 				Cell->LoadedActors.Add(UWorldPartitionEditorCell::FActorReference(ActorHandle.Source, ActorHandle.Handle));
 			}
diff --git a/Engine/Source/Runtime/Engine/Public/WorldPartition/WorldPartition.h b/Engine/Source/Runtime/Engine/Public/WorldPartition/WorldPartition.h
--- a/Engine/Source/Runtime/Engine/Public/WorldPartition/WorldPartition.h
+++ b/Engine/Source/Runtime/Engine/Public/WorldPartition/WorldPartition.h
@@ -238,7 +238,7 @@ private:
 	void UpdateLoadingEditorCell(UWorldPartitionEditorCell* Cell, bool bShouldBeLoaded);
 	void HashActorDesc(FWorldPartitionActorDesc* ActorDesc);
 	void UnhashActorDesc(FWorldPartitionActorDesc* ActorDesc);
-	bool ShouldActorBeLoaded(const FWorldPartitionActorDesc* ActorDesc) const;
+	bool ShouldActorBeLoadedByEditorCells(const FWorldPartitionActorDesc* ActorDesc) const;
 	bool UpdateEditorCells(TFunctionRef<bool(TArray<UWorldPartitionEditorCell*>&)> GetCellsToProcess, bool bIsCellShouldBeLoaded);
 
 	void ApplyActorTransform(AActor* Actor, const FTransform& InTransform);
