Remove concept of 'linked scopes', instead track number of exec input connections to each expression, and use that to process if/else scopes
- This way we don't need to move statements around, can generate both if/else scope correctly before switching back to parent scope
- Will make it easier to track locals within the current scope

diff --git a/Engine/Source/Editor/MaterialEditor/Private/MaterialEditorUtilities.cpp b/Engine/Source/Editor/MaterialEditor/Private/MaterialEditorUtilities.cpp
--- a/Engine/Source/Editor/MaterialEditor/Private/MaterialEditorUtilities.cpp
+++ b/Engine/Source/Editor/MaterialEditor/Private/MaterialEditorUtilities.cpp
@@ -597,7 +597,7 @@ void FMaterialEditorUtilities::GetVisibleMaterialParametersFromExpression(
 			MaterialExpressionKey.Expression->GetExecOutputs(ExpressionExecOutputs);
 			for (const FExpressionExecOutputEntry& Entry : ExpressionExecOutputs)
 			{
-				GetVisibleMaterialParametersFromExpression(FMaterialExpressionKey(Entry.Output->Expression, INDEX_NONE), MaterialInstance, VisibleExpressions, FunctionStack);
+				GetVisibleMaterialParametersFromExpression(FMaterialExpressionKey(Entry.Output->GetExpression(), INDEX_NONE), MaterialInstance, VisibleExpressions, FunctionStack);
 			}
 		}
 	}
diff --git a/Engine/Source/Editor/UnrealEd/Private/MaterialGraph.cpp b/Engine/Source/Editor/UnrealEd/Private/MaterialGraph.cpp
--- a/Engine/Source/Editor/UnrealEd/Private/MaterialGraph.cpp
+++ b/Engine/Source/Editor/UnrealEd/Private/MaterialGraph.cpp
@@ -378,14 +378,15 @@ void UMaterialGraph::LinkGraphNodesFromMaterial()
 			else if (Pin->Direction == EGPD_Output && PinInfo.PinType == EMaterialGraphPinType::Exec)
 			{
 				FExpressionExecOutput* ExecOutput = ExecOutputs[PinInfo.Index].Output;
-				if (ExecOutput->Expression)
+				UMaterialExpression* ConnectedExpression = ExecOutput->GetExpression();
+				if (ConnectedExpression)
 				{
-					if (ExecOutput->Expression == Material->ExpressionExecEnd)
+					if (ConnectedExpression == Material->ExpressionExecEnd)
 					{
 						// Exec end point is the root node
 						Pin->MakeLinkTo(RootNode->GetExecInputPin());
 					}
-					else if (UMaterialGraphNode* GraphNode = Cast<UMaterialGraphNode>(ExecOutput->Expression->GraphNode))
+					else if (UMaterialGraphNode* GraphNode = Cast<UMaterialGraphNode>(ConnectedExpression->GraphNode))
 					{
 						Pin->MakeLinkTo(GraphNode->GetExecInputPin());
 					}
@@ -522,13 +523,13 @@ void UMaterialGraph::LinkMaterialExpressionsFromGraph() const
 										bModifiedExpression = true;
 										Expression->Modify();
 									}
-									ExpressionOutput->Expression = Material->ExpressionExecEnd;
+									ExpressionOutput->Connect(Material->ExpressionExecEnd);
 								}
 								else
 								{
 									UMaterialGraphNode* ConnectedNode = CastChecked<UMaterialGraphNode>(Pin->LinkedTo[0]->GetOwningNode());
 									if (ExpressionOutput &&
-										ExpressionOutput->Expression != ConnectedNode->MaterialExpression &&
+										ExpressionOutput->GetExpression() != ConnectedNode->MaterialExpression &&
 										ConnectedNode->MaterialExpression->HasExecInput())
 									{
 										if (!bModifiedExpression)
@@ -538,18 +539,18 @@ void UMaterialGraph::LinkMaterialExpressionsFromGraph() const
 										}
 
 										ConnectedNode->MaterialExpression->Modify();
-										ExpressionOutput->Expression = ConnectedNode->MaterialExpression;
+										ExpressionOutput->Connect(ConnectedNode->MaterialExpression);
 									}
 								}
 							}
-							else if (ExpressionOutput && ExpressionOutput->Expression)
+							else if (ExpressionOutput && ExpressionOutput->GetExpression())
 							{
 								if (!bModifiedExpression)
 								{
 									bModifiedExpression = true;
 									Expression->Modify();
 								}
-								ExpressionOutput->Expression = NULL;
+								ExpressionOutput->Connect(nullptr);
 							}
 						}
 					}
diff --git a/Engine/Source/Runtime/Engine/Classes/Materials/MaterialExpression.h b/Engine/Source/Runtime/Engine/Classes/Materials/MaterialExpression.h
--- a/Engine/Source/Runtime/Engine/Classes/Materials/MaterialExpression.h
+++ b/Engine/Source/Runtime/Engine/Classes/Materials/MaterialExpression.h
@@ -116,18 +116,21 @@ USTRUCT()
 struct FExpressionExecOutput
 {
 	GENERATED_BODY()
-
+public:
 #if WITH_EDITOR
 	ENGINE_API int32 Compile(class FMaterialCompiler* Compiler) const;
 
+	ENGINE_API void Connect(class UMaterialExpression* InExpression);
+
+	inline class UMaterialExpression* GetExpression() const { return Expression; }
+
 	/** Returns the statement for the expression connected to this input */ 
-	ENGINE_API UE::HLSLTree::FStatement* AcquireHLSLStatement(FMaterialHLSLGenerator& Generator, UE::HLSLTree::FScope& Scope) const;
+	ENGINE_API bool GenerateHLSLStatements(FMaterialHLSLGenerator& Generator, UE::HLSLTree::FScope& Scope) const;
 
 	/** Creates a new scope, and populates it with the expression connected to this input */
-	ENGINE_API UE::HLSLTree::FScope* NewScopeWithStatement(FMaterialHLSLGenerator& Generator, UE::HLSLTree::FScope& Scope) const;
-	ENGINE_API UE::HLSLTree::FScope* NewLinkedScopeWithStatement(FMaterialHLSLGenerator& Generator, UE::HLSLTree::FScope& Scope) const;
+	ENGINE_API UE::HLSLTree::FScope* NewScopeWithStatements(FMaterialHLSLGenerator& Generator, UE::HLSLTree::FScope& Scope) const;
 #endif // WITH_EDITOR
-
+private:
 	UPROPERTY()
 	class UMaterialExpression* Expression = nullptr;
 };
@@ -151,6 +154,10 @@ class ENGINE_API UMaterialExpression : public UObject
 
 	static constexpr int32 CompileExecutionOutputIndex = -2;
 
+#if WITH_EDITOR
+	static void InitializeNumExecutionInputs(TArrayView<UMaterialExpression*> Expressions);
+#endif
+
 #if WITH_EDITORONLY_DATA
 	UPROPERTY()
 	int32 MaterialExpressionEditorX;
@@ -193,6 +200,9 @@ class ENGINE_API UMaterialExpression : public UObject
 	UPROPERTY(EditAnywhere, Category=MaterialExpression, meta=(MultiLine=true, DisplayAfter = "SortPriority"))
 	FString Desc;
 
+	/** Number of expressions connected to this expression's execution input */
+	int32 NumExecutionInputs;
+
 	/** Set to true by RecursiveUpdateRealtimePreview() if the expression's preview needs to be updated in realtime in the material editor. */
 	UPROPERTY()
 	uint32 bRealtimePreview:1;
@@ -289,7 +299,7 @@ class ENGINE_API UMaterialExpression : public UObject
 	 * For example, a for-loop expression might generate a statement for the execution input, but also generate an expression to access the loop index
 	 * These methods replace the Compile() method; once we switch over to the new system, Compile() will be removed
 	 */
-	virtual EMaterialGenerateHLSLStatus GenerateHLSLStatement(FMaterialHLSLGenerator& Generator, UE::HLSLTree::FScope& Scope, UE::HLSLTree::FStatement*& OutStatement);
+	virtual EMaterialGenerateHLSLStatus GenerateHLSLStatements(FMaterialHLSLGenerator& Generator, UE::HLSLTree::FScope& Scope);
 	virtual EMaterialGenerateHLSLStatus GenerateHLSLExpression(FMaterialHLSLGenerator& Generator, UE::HLSLTree::FScope& Scope, int32 OutputIndex, UE::HLSLTree::FExpression* &OutExpression);
 	virtual EMaterialGenerateHLSLStatus GenerateHLSLTexture(FMaterialHLSLGenerator& Generator, UE::HLSLTree::FScope& Scope, int32 OutputIndex, UE::HLSLTree::FTextureParameterDeclaration*& OutTexture);
 
diff --git a/Engine/Source/Runtime/Engine/Classes/Materials/MaterialExpressionExecBegin.h b/Engine/Source/Runtime/Engine/Classes/Materials/MaterialExpressionExecBegin.h
--- a/Engine/Source/Runtime/Engine/Classes/Materials/MaterialExpressionExecBegin.h
+++ b/Engine/Source/Runtime/Engine/Classes/Materials/MaterialExpressionExecBegin.h
@@ -23,7 +23,7 @@ class UMaterialExpressionExecBegin : public UMaterialExpression
 	virtual void GetCaption(TArray<FString>& OutCaptions) const override;
 	virtual bool CanUserDeleteExpression() const override { return false; }
 
-	virtual EMaterialGenerateHLSLStatus GenerateHLSLStatement(FMaterialHLSLGenerator& Generator, UE::HLSLTree::FScope& Scope, UE::HLSLTree::FStatement*& OutStatement) override;
+	virtual EMaterialGenerateHLSLStatus GenerateHLSLStatements(FMaterialHLSLGenerator& Generator, UE::HLSLTree::FScope& Scope) override;
 #endif
 	//~ End UMaterialExpression Interface
 };
diff --git a/Engine/Source/Runtime/Engine/Classes/Materials/MaterialExpressionExecEnd.h b/Engine/Source/Runtime/Engine/Classes/Materials/MaterialExpressionExecEnd.h
--- a/Engine/Source/Runtime/Engine/Classes/Materials/MaterialExpressionExecEnd.h
+++ b/Engine/Source/Runtime/Engine/Classes/Materials/MaterialExpressionExecEnd.h
@@ -18,7 +18,7 @@ class UMaterialExpressionExecEnd : public UMaterialExpression
 	virtual bool CanUserDeleteExpression() const override { return false; }
 	virtual bool HasExecInput() override { return true; }
 
-	virtual EMaterialGenerateHLSLStatus GenerateHLSLStatement(FMaterialHLSLGenerator& Generator, UE::HLSLTree::FScope& Scope, UE::HLSLTree::FStatement*& OutStatement) override;
+	virtual EMaterialGenerateHLSLStatus GenerateHLSLStatements(FMaterialHLSLGenerator& Generator, UE::HLSLTree::FScope& Scope) override;
 #endif
 	//~ End UMaterialExpression Interface
 };
diff --git a/Engine/Source/Runtime/Engine/Classes/Materials/MaterialExpressionForLoop.h b/Engine/Source/Runtime/Engine/Classes/Materials/MaterialExpressionForLoop.h
--- a/Engine/Source/Runtime/Engine/Classes/Materials/MaterialExpressionForLoop.h
+++ b/Engine/Source/Runtime/Engine/Classes/Materials/MaterialExpressionForLoop.h
@@ -37,7 +37,8 @@ class UMaterialExpressionForLoop : public UMaterialExpression
 	virtual uint32 GetOutputType(int32 InputIndex) override;
 	virtual bool HasExecInput() override { return true; }
 
-	virtual EMaterialGenerateHLSLStatus GenerateHLSLStatement(FMaterialHLSLGenerator& Generator, UE::HLSLTree::FScope& Scope, UE::HLSLTree::FStatement*& OutStatement) override;
+	virtual EMaterialGenerateHLSLStatus GenerateHLSLExpression(FMaterialHLSLGenerator& Generator, UE::HLSLTree::FScope& Scope, int32 OutputIndex, UE::HLSLTree::FExpression*& OutExpression) override;
+	virtual EMaterialGenerateHLSLStatus GenerateHLSLStatements(FMaterialHLSLGenerator& Generator, UE::HLSLTree::FScope& Scope) override;
 #endif
 	//~ End UMaterialExpression Interface
 };
diff --git a/Engine/Source/Runtime/Engine/Classes/Materials/MaterialExpressionIfThenElse.h b/Engine/Source/Runtime/Engine/Classes/Materials/MaterialExpressionIfThenElse.h
--- a/Engine/Source/Runtime/Engine/Classes/Materials/MaterialExpressionIfThenElse.h
+++ b/Engine/Source/Runtime/Engine/Classes/Materials/MaterialExpressionIfThenElse.h
@@ -30,7 +30,7 @@ class UMaterialExpressionIfThenElse : public UMaterialExpression
 	virtual uint32 GetInputType(int32 InputIndex) override;
 	virtual bool HasExecInput() override { return true; }
 
-	virtual EMaterialGenerateHLSLStatus GenerateHLSLStatement(FMaterialHLSLGenerator& Generator, UE::HLSLTree::FScope& Scope, UE::HLSLTree::FStatement*& OutStatement) override;
+	virtual EMaterialGenerateHLSLStatus GenerateHLSLStatements(FMaterialHLSLGenerator& Generator, UE::HLSLTree::FScope& Scope) override;
 #endif
 	//~ End UMaterialExpression Interface
 };
diff --git a/Engine/Source/Runtime/Engine/Classes/Materials/MaterialExpressionSetLocal.h b/Engine/Source/Runtime/Engine/Classes/Materials/MaterialExpressionSetLocal.h
--- a/Engine/Source/Runtime/Engine/Classes/Materials/MaterialExpressionSetLocal.h
+++ b/Engine/Source/Runtime/Engine/Classes/Materials/MaterialExpressionSetLocal.h
@@ -29,7 +29,7 @@ class UMaterialExpressionSetLocal : public UMaterialExpression
 	virtual void GetCaption(TArray<FString>& OutCaptions) const override;
 	virtual uint32 GetInputType(int32 InputIndex) override;
 	virtual bool HasExecInput() override { return true; }
-	virtual EMaterialGenerateHLSLStatus GenerateHLSLStatement(FMaterialHLSLGenerator& Generator, UE::HLSLTree::FScope& Scope, UE::HLSLTree::FStatement*& OutStatement) override;
+	virtual EMaterialGenerateHLSLStatus GenerateHLSLStatements(FMaterialHLSLGenerator& Generator, UE::HLSLTree::FScope& Scope) override;
 #endif
 	//~ End UMaterialExpression Interface
 };
diff --git a/Engine/Source/Runtime/Engine/Private/HLSLTree/HLSLTree.cpp b/Engine/Source/Runtime/Engine/Private/HLSLTree/HLSLTree.cpp
--- a/Engine/Source/Runtime/Engine/Private/HLSLTree/HLSLTree.cpp
+++ b/Engine/Source/Runtime/Engine/Private/HLSLTree/HLSLTree.cpp
@@ -698,60 +698,6 @@ void UE::HLSLTree::FScope::AddStatement(UE::HLSLTree::FStatement* Statement)
 	}
 }
 
-bool UE::HLSLTree::FScope::TryMoveStatement(UE::HLSLTree::FStatement* Statement)
-{
-	bool bResult = false;
-	FScope* PrevScope = Statement->ParentScope;
-	if (!PrevScope)
-	{
-		// Statement does not currently belong to any scope, just add it
-		AddStatement(Statement);
-		bResult = true;
-	}
-	else if (LinkedScope && ParentScope == PrevScope)
-	{
-		// Already moved
-		bResult = true;
-	}
-	else if(LinkedScope == PrevScope)
-	{
-		// Statement belongs to our linked scope, so move it into our parent scope
-		check(ParentScope->FirstStatement);
-		check(ParentScope->LastStatement);
-
-		// First need to unlink the statement from its previous scope
-		FStatement* PrevStatement = PrevScope->FirstStatement;
-		bool bFoundInPrevScope = false;
-		while (PrevStatement)
-		{
-			FStatement* NextStatement = PrevStatement->NextStatement;
-			if (NextStatement == Statement)
-			{
-				PrevScope->LastStatement = PrevStatement;
-				PrevStatement->NextStatement = nullptr;
-				bFoundInPrevScope = true;
-				break;
-			}
-			PrevStatement = NextStatement;
-		}
-		check(bFoundInPrevScope);
-
-		// Now link it into the new scope
-		FNodeVisitor_MoveToScope Visitor(ParentScope);
-		ParentScope->LastStatement->NextStatement = Statement;
-		FStatement* StatementToLink = Statement;
-		while (StatementToLink)
-		{
-			StatementToLink->ParentScope = ParentScope;
-			ParentScope->LastStatement = StatementToLink;
-			Visitor.VisitNode(StatementToLink);
-			StatementToLink = StatementToLink->NextStatement;
-		}
-		bResult = true;
-	}
-	return bResult;
-}
-
 UE::HLSLTree::FTree* UE::HLSLTree::FTree::Create(FMemStackBase& Allocator)
 {
 	FTree* Tree = new(Allocator) FTree();
@@ -773,15 +719,6 @@ UE::HLSLTree::FScope* UE::HLSLTree::FTree::NewScope(UE::HLSLTree::FScope& Scope)
 	return NewScope;
 }
 
-UE::HLSLTree::FScope* UE::HLSLTree::FTree::NewLinkedScope(UE::HLSLTree::FScope& Scope)
-{
-	FScope* NewScope = NewNode<FScope>();
-	NewScope->ParentScope = Scope.ParentScope;
-	NewScope->LinkedScope = &Scope;
-	NewScope->NestedLevel = Scope.NestedLevel;
-	return NewScope;
-}
-
 UE::HLSLTree::FLocalDeclaration* UE::HLSLTree::FTree::NewLocalDeclaration(UE::HLSLTree::FScope& Scope, Shader::EValueType Type, const FName& Name)
 {
 	FLocalDeclaration* Declaration = NewNode<FLocalDeclaration>(Name, Type);
diff --git a/Engine/Source/Runtime/Engine/Private/HLSLTree/HLSLTreeCommon.cpp b/Engine/Source/Runtime/Engine/Private/HLSLTree/HLSLTreeCommon.cpp
--- a/Engine/Source/Runtime/Engine/Private/HLSLTree/HLSLTreeCommon.cpp
+++ b/Engine/Source/Runtime/Engine/Private/HLSLTree/HLSLTreeCommon.cpp
@@ -550,5 +550,22 @@ bool UE::HLSLTree::FStatementIf::EmitHLSL(FEmitContext& Context, FCodeWriter& Ou
 
 bool UE::HLSLTree::FStatementFor::EmitHLSL(FEmitContext& Context, FCodeWriter& OutWriter) const
 {
-	return false;
+	const FEmitValue* DeclarationValue = Context.AcquireValue(LoopControlDeclaration);
+	const FEmitValue* StartExpressionValue = Context.AcquireValue(StartExpression);
+	const FEmitValue* EndExpressionValue = Context.AcquireValue(StartExpression);
+	if (!DeclarationValue || !StartExpressionValue || !EndExpressionValue)
+	{
+		return false;
+	}
+
+	const TCHAR* DeclarationCode = Context.GetCode(DeclarationValue);
+
+	OutWriter.WriteLinef(TEXT("for (%s = %s; %s < %s; %s++"),
+		DeclarationCode,
+		Context.GetCode(StartExpressionValue),
+		DeclarationCode,
+		Context.GetCode(EndExpressionValue),
+		DeclarationCode);
+	LoopScope->EmitHLSL(Context, OutWriter);
+	return true;
 }
diff --git a/Engine/Source/Runtime/Engine/Private/Materials/MaterialExpressionHLSL.cpp b/Engine/Source/Runtime/Engine/Private/Materials/MaterialExpressionHLSL.cpp
--- a/Engine/Source/Runtime/Engine/Private/Materials/MaterialExpressionHLSL.cpp
+++ b/Engine/Source/Runtime/Engine/Private/Materials/MaterialExpressionHLSL.cpp
@@ -40,7 +40,7 @@ EMaterialGenerateHLSLStatus UMaterialExpression::GenerateHLSLExpression(FMateria
 	return Generator.Error(TEXT("Node does not support expressions"));
 }
 
-EMaterialGenerateHLSLStatus UMaterialExpression::GenerateHLSLStatement(FMaterialHLSLGenerator& Generator, UE::HLSLTree::FScope& Scope, UE::HLSLTree::FStatement*& OutStatement)
+EMaterialGenerateHLSLStatus UMaterialExpression::GenerateHLSLStatements(FMaterialHLSLGenerator& Generator, UE::HLSLTree::FScope& Scope)
 {
 	return Generator.Error(TEXT("Node does not support statements"));
 }
@@ -311,19 +311,19 @@ EMaterialGenerateHLSLStatus UMaterialExpressionMaterialFunctionCall::GenerateHLS
 	return EMaterialGenerateHLSLStatus::Success;
 }
 
-EMaterialGenerateHLSLStatus UMaterialExpressionExecBegin::GenerateHLSLStatement(FMaterialHLSLGenerator& Generator, UE::HLSLTree::FScope& Scope, UE::HLSLTree::FStatement*& OutStatement)
+EMaterialGenerateHLSLStatus UMaterialExpressionExecBegin::GenerateHLSLStatements(FMaterialHLSLGenerator& Generator, UE::HLSLTree::FScope& Scope)
 {
-	OutStatement = Exec.AcquireHLSLStatement(Generator, Scope);
+	Exec.GenerateHLSLStatements(Generator, Scope);
 	return EMaterialGenerateHLSLStatus::Success;
 }
 
-EMaterialGenerateHLSLStatus UMaterialExpressionExecEnd::GenerateHLSLStatement(FMaterialHLSLGenerator& Generator, UE::HLSLTree::FScope& Scope, UE::HLSLTree::FStatement*& OutStatement)
+EMaterialGenerateHLSLStatus UMaterialExpressionExecEnd::GenerateHLSLStatements(FMaterialHLSLGenerator& Generator, UE::HLSLTree::FScope& Scope)
 {
-	OutStatement = Generator.NewResult(Scope);
+	Generator.GenerateResult(Scope);
 	return EMaterialGenerateHLSLStatus::Success;
 }
 
-EMaterialGenerateHLSLStatus UMaterialExpressionSetLocal::GenerateHLSLStatement(FMaterialHLSLGenerator& Generator, UE::HLSLTree::FScope& Scope, UE::HLSLTree::FStatement*& OutStatement)
+EMaterialGenerateHLSLStatus UMaterialExpressionSetLocal::GenerateHLSLStatements(FMaterialHLSLGenerator& Generator, UE::HLSLTree::FScope& Scope)
 {
 	UE::HLSLTree::FExpression* ValueExpression = Value.AcquireHLSLExpressionWithCast(Generator, Scope, UE::Shader::EValueType::Float3);
 	if (!ValueExpression)
@@ -341,13 +341,12 @@ EMaterialGenerateHLSLStatus UMaterialExpressionSetLocal::GenerateHLSLStatement(F
 	Statement->Declaration = Declaration;
 	Statement->Expression = ValueExpression;
 
-	Exec.AcquireHLSLStatement(Generator, Scope);
+	Exec.GenerateHLSLStatements(Generator, Scope);
 
-	OutStatement = Statement;
 	return EMaterialGenerateHLSLStatus::Success;
 }
 
-EMaterialGenerateHLSLStatus UMaterialExpressionIfThenElse::GenerateHLSLStatement(FMaterialHLSLGenerator& Generator, UE::HLSLTree::FScope& Scope, UE::HLSLTree::FStatement*& OutStatement)
+EMaterialGenerateHLSLStatus UMaterialExpressionIfThenElse::GenerateHLSLStatements(FMaterialHLSLGenerator& Generator, UE::HLSLTree::FScope& Scope)
 {
 	UE::HLSLTree::FExpression* ConditionExpression = Condition.AcquireHLSLExpression(Generator, Scope);
 	if (!ConditionExpression)
@@ -355,7 +354,7 @@ EMaterialGenerateHLSLStatus UMaterialExpressionIfThenElse::GenerateHLSLStatement
 		return Generator.Error(TEXT("Missing condition connection"));
 	}
 
-	UE::HLSLTree::FScope* ThenScope = Then.NewScopeWithStatement(Generator, Scope);
+	UE::HLSLTree::FScope* ThenScope = Then.NewScopeWithStatements(Generator, Scope);
 	if (!ThenScope)
 	{
 		return Generator.Error(TEXT("Missing Then connection"));
@@ -364,40 +363,47 @@ EMaterialGenerateHLSLStatus UMaterialExpressionIfThenElse::GenerateHLSLStatement
 	UE::HLSLTree::FStatementIf* IfStatement = Generator.GetTree().NewStatement<UE::HLSLTree::FStatementIf>(Scope);
 	IfStatement->ConditionExpression = ConditionExpression;
 	IfStatement->ThenScope = ThenScope;
-	IfStatement->ElseScope = Else.NewLinkedScopeWithStatement(Generator, *ThenScope);
+	IfStatement->ElseScope = Else.NewScopeWithStatements(Generator, Scope);
 
-	OutStatement = IfStatement;
 	return EMaterialGenerateHLSLStatus::Success;
 }
 
-EMaterialGenerateHLSLStatus UMaterialExpressionForLoop::GenerateHLSLStatement(FMaterialHLSLGenerator& Generator, UE::HLSLTree::FScope& Scope, UE::HLSLTree::FStatement*& OutStatement)
+EMaterialGenerateHLSLStatus UMaterialExpressionForLoop::GenerateHLSLExpression(FMaterialHLSLGenerator& Generator, UE::HLSLTree::FScope& Scope, int32 OutputIndex, UE::HLSLTree::FExpression*& OutExpression)
 {
-	UE::HLSLTree::FExpression* StartExpression = StartIndex.AcquireHLSLExpression(Generator, Scope);
+	return EMaterialGenerateHLSLStatus::Success;
+}
+
+EMaterialGenerateHLSLStatus UMaterialExpressionForLoop::GenerateHLSLStatements(FMaterialHLSLGenerator& Generator, UE::HLSLTree::FScope& Scope)
+{
+	if (!LoopBody.GetExpression())
+	{
+		return Generator.Error(TEXT("Missing LoopBody connection"));
+	}
+
+	UE::HLSLTree::FExpression* StartExpression = StartIndex.AcquireHLSLExpressionWithCast(Generator, Scope, UE::Shader::EValueType::Int1);
 	if (!StartExpression)
 	{
 		return Generator.Error(TEXT("Missing StartIndex connection"));
 	}
 
-	UE::HLSLTree::FExpression* EndExpression = EndIndex.AcquireHLSLExpression(Generator, Scope);
+	UE::HLSLTree::FExpression* EndExpression = EndIndex.AcquireHLSLExpressionWithCast(Generator, Scope, UE::Shader::EValueType::Int1);
 	if (!EndExpression)
 	{
 		return Generator.Error(TEXT("Missing EndIndex connection"));
 	}
 
-	UE::HLSLTree::FScope* LoopScope = LoopBody.NewScopeWithStatement(Generator, Scope);
-	if (!LoopScope)
-	{
-		return Generator.Error(TEXT("Missing LoopBody connection"));
-	}
+	UE::HLSLTree::FLocalDeclaration* LoopControlDeclaration = Generator.GetTree().NewLocalDeclaration(Scope, UE::Shader::EValueType::Int1, TEXT("ForLoopControl"));
+	UE::HLSLTree::FScope* LoopScope = Generator.GetTree().NewScope(Scope);
+	LoopBody.GenerateHLSLStatements(Generator, *LoopScope);
 
 	UE::HLSLTree::FStatementFor* ForStatement = Generator.GetTree().NewStatement<UE::HLSLTree::FStatementFor>(Scope);
+	ForStatement->LoopControlDeclaration = LoopControlDeclaration;
 	ForStatement->StartExpression = StartExpression;
 	ForStatement->EndExpression = EndExpression;
 	ForStatement->LoopScope = LoopScope;
 
-	Completed.AcquireHLSLStatement(Generator, Scope);
+	Completed.GenerateHLSLStatements(Generator, Scope);
 
-	OutStatement = ForStatement;
 	return EMaterialGenerateHLSLStatus::Success;
 }
 
diff --git a/Engine/Source/Runtime/Engine/Private/Materials/MaterialExpressions.cpp b/Engine/Source/Runtime/Engine/Private/Materials/MaterialExpressions.cpp
--- a/Engine/Source/Runtime/Engine/Private/Materials/MaterialExpressions.cpp
+++ b/Engine/Source/Runtime/Engine/Private/Materials/MaterialExpressions.cpp
@@ -639,6 +639,32 @@ int32 CompileStrataBlendFunction(FMaterialCompiler* Compiler, const int32 A, con
 	return OutputCodeChunk;
 }
 
+#if WITH_EDITOR
+void UMaterialExpression::InitializeNumExecutionInputs(TArrayView<UMaterialExpression*> Expressions)
+{
+	for (UMaterialExpression* Expression : Expressions)
+	{
+		Expression->NumExecutionInputs = 0;
+	}
+
+	TArray<FExpressionExecOutputEntry> ExecOutputs;
+	for (UMaterialExpression* Expression : Expressions)
+	{
+		ExecOutputs.Reset();
+		Expression->GetExecOutputs(ExecOutputs);
+		for (const FExpressionExecOutputEntry& Output : ExecOutputs)
+		{
+			UMaterialExpression* ConnectedExpression = Output.Output->GetExpression();
+			if (ConnectedExpression)
+			{
+				check(ConnectedExpression->HasExecInput());
+				ConnectedExpression->NumExecutionInputs++;
+			}
+		}
+	}
+}
+#endif // WITH_EDITOR
+
 UMaterialExpression::UMaterialExpression(const FObjectInitializer& ObjectInitializer)
 	: Super(ObjectInitializer)
 #if WITH_EDITORONLY_DATA
@@ -649,6 +675,8 @@ UMaterialExpression::UMaterialExpression(const FObjectInitializer& ObjectInitial
 #if WITH_EDITORONLY_DATA
 	Outputs.Add(FExpressionOutput(TEXT("")));
 
+	NumExecutionInputs = 0;
+
 	bShowInputs = true;
 	bShowOutputs = true;
 	bCollapsed = true;
@@ -832,6 +860,20 @@ void UMaterialExpression::PostLoad()
 	UpdateParameterGuid(false, false);
 	
 	UpdateMaterialExpressionGuid(false, false);
+
+#if WITH_EDITORONLY_DATA
+	TArray<FExpressionExecOutputEntry> ExecOutputs;
+	GetExecOutputs(ExecOutputs);
+	for (const FExpressionExecOutputEntry& Output : ExecOutputs)
+	{
+		UMaterialExpression* ConnectedExpression = Output.Output->GetExpression();
+		if (ConnectedExpression)
+		{
+			check(ConnectedExpression->HasExecInput());
+			ConnectedExpression->NumExecutionInputs++;
+		}
+	}
+#endif // WITH_EDITORONLY_DATA
 }
 
 void UMaterialExpression::PostDuplicate(bool bDuplicateForPIE)
@@ -21328,6 +21370,7 @@ UMaterialExpressionExecBegin::UMaterialExpressionExecBegin(const FObjectInitiali
 #if WITH_EDITORONLY_DATA
 	Outputs.Reset();
 	bHidePreviewWindow = true;
+	NumExecutionInputs = 1; // ExecBegin has 1 implicit input (since it's the start of the graph)
 #endif
 }
 
@@ -21340,7 +21383,7 @@ void UMaterialExpressionExecBegin::GetCaption(TArray<FString>& OutCaptions) cons
 int32 UMaterialExpressionExecBegin::Compile(class FMaterialCompiler* Compiler, int32 OutputIndex)
 {
 	check(OutputIndex == CompileExecutionOutputIndex);
-	if (Exec.Expression)
+	if (Exec.GetExpression())
 	{
 		const int32 RootScope = Compiler->BeginScope();
 		Exec.Compile(Compiler);
@@ -21389,7 +21432,7 @@ UMaterialExpressionIfThenElse::UMaterialExpressionIfThenElse(const FObjectInitia
 int32 UMaterialExpressionIfThenElse::Compile(class FMaterialCompiler* Compiler, int32 OutputIndex)
 {
 	check(OutputIndex == CompileExecutionOutputIndex);
-	if (!Then.Expression)
+	if (!Then.GetExpression())
 	{
 		return Compiler->Error(TEXT("Then pin must be connected"));
 	}
@@ -21405,7 +21448,7 @@ int32 UMaterialExpressionIfThenElse::Compile(class FMaterialCompiler* Compiler,
 	Compiler->EndScope();
 
 	int32 ElseScope = INDEX_NONE;
-	if (Else.Expression)
+	if (Else.GetExpression())
 	{
 		ElseScope = Compiler->BeginScope_Else();
 		Else.Compile(Compiler);
@@ -21453,7 +21496,7 @@ int32 UMaterialExpressionForLoop::Compile(class FMaterialCompiler* Compiler, int
 {
 	if (OutputIndex == CompileExecutionOutputIndex)
 	{
-		if (!LoopBody.Expression)
+		if (!LoopBody.GetExpression())
 		{
 			return Compiler->Error(TEXT("Loop pin must be connected"));
 		}
diff --git a/Engine/Source/Runtime/Engine/Private/Materials/MaterialHLSLGenerator.cpp b/Engine/Source/Runtime/Engine/Private/Materials/MaterialHLSLGenerator.cpp
--- a/Engine/Source/Runtime/Engine/Private/Materials/MaterialHLSLGenerator.cpp
+++ b/Engine/Source/Runtime/Engine/Private/Materials/MaterialHLSLGenerator.cpp
@@ -40,6 +40,28 @@ void FMaterialHLSLGenerator::AcquireErrors(TArray<FString>& OutCompileErrors, TA
 	OutErrorExpressions = MoveTemp(ErrorExpressions);
 }
 
+bool FMaterialHLSLGenerator::Finalize()
+{
+	if (!bGeneratedResult)
+	{
+		Error(TEXT("Missing connection to material output"));
+		return false;
+	}
+
+	for (const auto& It : StatementMap)
+	{
+		const UMaterialExpression* Expression = It.Key;
+		const FStatementEntry& Entry = It.Value;
+		if (Entry.NumInputs != Expression->NumExecutionInputs)
+		{
+			Error(TEXT("Invalid number of input connections"));
+			return false;
+		}
+	}
+
+	return true;
+}
+
 EMaterialGenerateHLSLStatus FMaterialHLSLGenerator::Error(const FString& Message)
 {
 	UMaterialExpression* ExpressionToError = nullptr;
@@ -135,9 +157,9 @@ static UE::HLSLTree::FExpression* CompileMaterialInput(FMaterialHLSLGenerator& G
 	return AttributesExpression;
 }
 
-UE::HLSLTree::FStatement* FMaterialHLSLGenerator::NewResult(UE::HLSLTree::FScope& Scope)
+bool FMaterialHLSLGenerator::GenerateResult(UE::HLSLTree::FScope& Scope)
 {
-	UE::HLSLTree::FStatement* Result = nullptr;
+	bool bResult = false;
 	if (bGeneratedResult)
 	{
 		Error(TEXT("Multiple connections to execution output"));
@@ -170,7 +192,7 @@ UE::HLSLTree::FStatement* FMaterialHLSLGenerator::NewResult(UE::HLSLTree::FScope
 			{
 				UE::HLSLTree::FStatementReturn* ReturnStatement = HLSLTree->NewStatement<UE::HLSLTree::FStatementReturn>(Scope);
 				ReturnStatement->Expression = AttributesExpression;
-				Result = ReturnStatement;
+				bResult = true;
 			}
 		}
 		else
@@ -179,7 +201,7 @@ UE::HLSLTree::FStatement* FMaterialHLSLGenerator::NewResult(UE::HLSLTree::FScope
 		}
 		bGeneratedResult = true;
 	}
-	return Result;
+	return bResult;
 }
 
 UE::HLSLTree::FExpressionConstant* FMaterialHLSLGenerator::NewConstant(UE::HLSLTree::FScope& Scope, const UE::Shader::FValue& Value)
@@ -359,30 +381,28 @@ UE::HLSLTree::FTextureParameterDeclaration* FMaterialHLSLGenerator::AcquireTextu
 	return TextureDeclaration;
 }
 
-UE::HLSLTree::FStatement* FMaterialHLSLGenerator::AcquireStatement(UE::HLSLTree::FScope& Scope, UMaterialExpression* MaterialExpression)
+bool FMaterialHLSLGenerator::GenerateStatements(UE::HLSLTree::FScope& Scope, UMaterialExpression* MaterialExpression)
 {
-	UE::HLSLTree::FStatement** PrevStatement = StatementMap.Find(MaterialExpression);
-	UE::HLSLTree::FStatement* Statement = nullptr;
-	if (!PrevStatement)
+	FStatementEntry& Entry = StatementMap.FindOrAdd(MaterialExpression);
+	Entry.NumInputs++;
+	check(Entry.NumInputs <= MaterialExpression->NumExecutionInputs);
+
+	bool bResult = true;
+	if (Entry.NumInputs == MaterialExpression->NumExecutionInputs)
 	{
+		UE::HLSLTree::FScope* ScopeToUse = &Scope;
+		if (MaterialExpression->NumExecutionInputs > 1u)
+		{
+			ScopeToUse = ScopeToUse->ParentScope;
+		}
+
 		const FExpressionKey Key(MaterialExpression);
 		ExpressionStack.Add(Key);
-		const EMaterialGenerateHLSLStatus Status = MaterialExpression->GenerateHLSLStatement(*this, Scope, Statement);
+		const EMaterialGenerateHLSLStatus Status = MaterialExpression->GenerateHLSLStatements(*this, *ScopeToUse);
 		verify(ExpressionStack.Pop() == Key);
-		StatementMap.Add(MaterialExpression, Statement);
-	}
-	else
-	{
-		Statement = *PrevStatement;
-		if (Statement && !Scope.TryMoveStatement(Statement))
-		{
-			// Could not move existing statement to the given scope
-			Error(TEXT("Invalid control flow"));
-			return nullptr;
-		}
 	}
 
-	return Statement;
+	return bResult;
 }
 
 #endif // WITH_EDITOR
diff --git a/Engine/Source/Runtime/Engine/Private/Materials/MaterialHLSLTree.cpp b/Engine/Source/Runtime/Engine/Private/Materials/MaterialHLSLTree.cpp
--- a/Engine/Source/Runtime/Engine/Private/Materials/MaterialHLSLTree.cpp
+++ b/Engine/Source/Runtime/Engine/Private/Materials/MaterialHLSLTree.cpp
@@ -16,7 +16,7 @@ FMaterialHLSLTree::FMaterialHLSLTree()
 {
 }
 
-void FMaterialHLSLTree::InitializeForMaterial(const FMaterialCompileTargetParameters& InCompilerTarget, FMaterial& InOutMaterial)
+bool FMaterialHLSLTree::InitializeForMaterial(const FMaterialCompileTargetParameters& InCompilerTarget, FMaterial& InOutMaterial)
 {
 	UMaterialInterface* MaterialInterface = InOutMaterial.GetMaterialInterface();
 	UMaterial* BaseMaterial = MaterialInterface->GetMaterial();
@@ -25,20 +25,28 @@ void FMaterialHLSLTree::InitializeForMaterial(const FMaterialCompileTargetParame
 	FMaterialHLSLGenerator Generator(BaseMaterial, InCompilerTarget, *HLSLTree);
 	UE::HLSLTree::FScope& RootScope = HLSLTree->GetRootScope();
 
+	bool bResult = false;
 	if (BaseMaterial->IsCompiledWithExecutionFlow())
 	{
 		UMaterialExpression* BaseExpression = BaseMaterial->ExpressionExecBegin;
-		Generator.AcquireStatement(RootScope, BaseExpression);
+		bResult = Generator.GenerateStatements(RootScope, BaseExpression);
 	}
 	else
 	{
-		Generator.NewResult(RootScope);
+		bResult = Generator.GenerateResult(RootScope);
+	}
+
+	if (!Generator.Finalize())
+	{
+		bResult = false;
 	}
 
 	Generator.AcquireErrors(InOutMaterial.CompileErrors, InOutMaterial.ErrorExpressions);
+
+	return bResult;
 }
 
-void FMaterialHLSLTree::InitializeForFunction(const FMaterialCompileTargetParameters& InCompilerTarget, UMaterialFunctionInterface* InOutFunction)
+bool FMaterialHLSLTree::InitializeForFunction(const FMaterialCompileTargetParameters& InCompilerTarget, UMaterialFunctionInterface* InOutFunction)
 {
 	TArray<FFunctionExpressionInput> FunctionInputs;
 	TArray<FFunctionExpressionOutput> FunctionOutputs;
@@ -64,6 +72,8 @@ void FMaterialHLSLTree::InitializeForFunction(const FMaterialCompileTargetParame
 		UMaterialExpressionFunctionOutput* ExpressionOutput = Output.ExpressionOutput;
 		FunctionOutputExpressions[OutputIndex] = ExpressionOutput->A.AcquireHLSLExpression(Generator, RootScope);
 	}
+
+	return Generator.Finalize();
 }
 
 UE::HLSLTree::FFunctionCall* FMaterialHLSLTree::GenerateFunctionCall(FMaterialHLSLGenerator& Generator, UE::HLSLTree::FScope& Scope, TArrayView<UE::HLSLTree::FExpression*> Inputs) const
diff --git a/Engine/Source/Runtime/Engine/Private/Materials/MaterialShared.cpp b/Engine/Source/Runtime/Engine/Private/Materials/MaterialShared.cpp
--- a/Engine/Source/Runtime/Engine/Private/Materials/MaterialShared.cpp
+++ b/Engine/Source/Runtime/Engine/Private/Materials/MaterialShared.cpp
@@ -325,39 +325,45 @@ int32 FExpressionExecOutput::Compile(class FMaterialCompiler* Compiler) const
 	return INDEX_NONE;
 }
 
-UE::HLSLTree::FStatement* FExpressionExecOutput::AcquireHLSLStatement(FMaterialHLSLGenerator& Generator, UE::HLSLTree::FScope& Scope) const
+void FExpressionExecOutput::Connect(UMaterialExpression* InExpression)
 {
-	UE::HLSLTree::FStatement* Result = nullptr;
-	if (Expression)
+	check(!InExpression || InExpression->HasExecInput());
+	if (InExpression != Expression)
 	{
-		Expression->ValidateState();
-		Result = Generator.AcquireStatement(Scope, Expression);
-	}
+		if (Expression)
+		{
+			check(Expression->NumExecutionInputs > 0);
+			Expression->NumExecutionInputs--;
+		}
 
-	return Result;
+		Expression = InExpression;
+		if (InExpression)
+		{
+			InExpression->NumExecutionInputs++;
+		}
+	}
 }
 
-UE::HLSLTree::FScope* FExpressionExecOutput::NewScopeWithStatement(FMaterialHLSLGenerator& Generator, UE::HLSLTree::FScope& Scope) const
+bool FExpressionExecOutput::GenerateHLSLStatements(FMaterialHLSLGenerator& Generator, UE::HLSLTree::FScope& Scope) const
 {
-	UE::HLSLTree::FScope* Result = nullptr;
+	bool bResult = false;
 	if (Expression)
 	{
 		Expression->ValidateState();
-		Result = Generator.GetTree().NewScope(Scope); // Create a new scope for the statement
-		Generator.AcquireStatement(*Result, Expression);
+		bResult = Generator.GenerateStatements(Scope, Expression);
 	}
 
-	return Result;
+	return bResult;
 }
 
-UE::HLSLTree::FScope* FExpressionExecOutput::NewLinkedScopeWithStatement(FMaterialHLSLGenerator& Generator, UE::HLSLTree::FScope& Scope) const
+UE::HLSLTree::FScope* FExpressionExecOutput::NewScopeWithStatements(FMaterialHLSLGenerator& Generator, UE::HLSLTree::FScope& Scope) const
 {
 	UE::HLSLTree::FScope* Result = nullptr;
 	if (Expression)
 	{
 		Expression->ValidateState();
-		Result = Generator.GetTree().NewLinkedScope(Scope); // Create a new scope for the statement
-		Generator.AcquireStatement(*Result, Expression);
+		Result = Generator.GetTree().NewScope(Scope); // Create a new scope for the statements
+		Generator.GenerateStatements(*Result, Expression);
 	}
 
 	return Result;
diff --git a/Engine/Source/Runtime/Engine/Public/HLSLTree/HLSLTree.h b/Engine/Source/Runtime/Engine/Public/HLSLTree/HLSLTree.h
--- a/Engine/Source/Runtime/Engine/Public/HLSLTree/HLSLTree.h
+++ b/Engine/Source/Runtime/Engine/Public/HLSLTree/HLSLTree.h
@@ -334,15 +334,12 @@ public:
 	void UseDeclaration(FLocalDeclaration* Declaration);
 	void UseFunctionCall(FFunctionCall* FunctionCall);
 	void UseExpression(FExpression* Expression);
-	bool TryMoveStatement(FStatement* Statement);
-
 private:
 	friend class FTree;
 	friend class FNodeVisitor_MoveToScope;
 
 	void UseNode(FNode* Node);
 
-	FScope* LinkedScope = nullptr;
 	FStatement* FirstStatement = nullptr;
 	FStatement* LastStatement = nullptr;
 	int32 NestedLevel = 0;
@@ -377,19 +374,6 @@ public:
 	}
 
 	FScope* NewScope(FScope& Scope);
-
-	/**
-	 * 2 scopes may be linked if they are both logically part of the same control structure (an 'if' scope and an 'else' scope are linked for example)
-	 * Linked scopes always share the same parent scope
-	 * Attempting to move a statement from a scope to its linked scope will instead move that statement into the parent scope
-	 * When translating node graphs, a chain of statements may be generated when translating an 'if' scope for example
-	 * Then when translating the 'else' scope, the flow may eventually hit a statement that was previously translated from the 'if' scope
-	 * At this point, that statement connected to both 'if' and 'else' scope should become the next statement after the if/else block,
-	 * which means it should be moved to the parent scope (the scope that contains the linked if/else scopes)
-	 * It's not clear if this pattern will be useful outside if/else block
-	 */
-	FScope* NewLinkedScope(FScope& Scope);
-	
 	FLocalDeclaration* NewLocalDeclaration(FScope& Scope, Shader::EValueType Type, const FName& Name);
 	FParameterDeclaration* NewParameterDeclaration(FScope& Scope, const FName& Name, const Shader::FValue& DefaultValue);
 	FTextureParameterDeclaration* NewTextureParameterDeclaration(FScope& Scope, const FName& Name, const FTextureDescription& DefaultValue);
diff --git a/Engine/Source/Runtime/Engine/Public/HLSLTree/HLSLTreeCommon.h b/Engine/Source/Runtime/Engine/Public/HLSLTree/HLSLTreeCommon.h
--- a/Engine/Source/Runtime/Engine/Public/HLSLTree/HLSLTreeCommon.h
+++ b/Engine/Source/Runtime/Engine/Public/HLSLTree/HLSLTreeCommon.h
@@ -386,6 +386,7 @@ public:
 class FStatementFor : public FStatement
 {
 public:
+	FLocalDeclaration* LoopControlDeclaration;
 	FExpression* StartExpression;
 	FExpression* EndExpression;
 	FScope* LoopScope;
@@ -395,6 +396,7 @@ public:
 		const ENodeVisitResult Result = FStatement::Visit(Visitor);
 		if (ShouldVisitDependentNodes(Result))
 		{
+			Visitor.VisitNode(LoopControlDeclaration);
 			Visitor.VisitNode(StartExpression);
 			Visitor.VisitNode(EndExpression);
 			Visitor.VisitNode(LoopScope);
diff --git a/Engine/Source/Runtime/Engine/Public/MaterialHLSLGenerator.h b/Engine/Source/Runtime/Engine/Public/MaterialHLSLGenerator.h
--- a/Engine/Source/Runtime/Engine/Public/MaterialHLSLGenerator.h
+++ b/Engine/Source/Runtime/Engine/Public/MaterialHLSLGenerator.h
@@ -50,9 +50,7 @@ public:
 
 	const FMaterialCompileTargetParameters& GetCompileTarget() const { return CompileTarget; }
 
-	//bool GenerateHLSL(FMaterial& InMaterial,
-	//	FMaterialCompilationOutput& OutCompilationOutput,
-	//	TRefCountPtr<FSharedShaderCompilerEnvironment>& OutMaterialEnvironment);
+	bool Finalize();
 
 	/** Retrieve the compile errors from the generator */
 	void AcquireErrors(TArray<FString>& OutCompileErrors, TArray<UMaterialExpression*>& OutErrorExpressions);
@@ -67,7 +65,7 @@ public:
 
 	UE::HLSLTree::FTree& GetTree() const { return *HLSLTree; }
 
-	UE::HLSLTree::FStatement* NewResult(UE::HLSLTree::FScope& Scope);
+	bool GenerateResult(UE::HLSLTree::FScope& Scope);
 
 	UE::HLSLTree::FExpressionConstant* NewConstant(UE::HLSLTree::FScope& Scope, const UE::Shader::FValue& Value);
 	UE::HLSLTree::FExpressionExternalInput* NewTexCoord(UE::HLSLTree::FScope& Scope, int32 Index);
@@ -94,7 +92,7 @@ public:
 	 */
 	UE::HLSLTree::FExpression* AcquireExpression(UE::HLSLTree::FScope& Scope, UMaterialExpression* MaterialExpression, int32 OutputIndex);
 	UE::HLSLTree::FTextureParameterDeclaration* AcquireTextureDeclaration(UE::HLSLTree::FScope& Scope, UMaterialExpression* MaterialExpression, int32 OutputIndex);
-	UE::HLSLTree::FStatement* AcquireStatement(UE::HLSLTree::FScope& Scope, UMaterialExpression* MaterialExpression);
+	bool GenerateStatements(UE::HLSLTree::FScope& Scope, UMaterialExpression* MaterialExpression);
 
 private:
 	struct FExpressionKey
@@ -131,6 +129,11 @@ private:
 		}
 	};
 
+	struct FStatementEntry
+	{
+		int32 NumInputs = 0;
+	};
+
 	const FMaterialCompileTargetParameters& CompileTarget;
 	UMaterial* TargetMaterial;
 	UMaterialFunctionInterface* TargetMaterialFunction;
@@ -145,7 +148,7 @@ private:
 	TMap<FName, UE::HLSLTree::FTextureParameterDeclaration*> TextureParameterDeclarationMap;
 	TMap<FFunctionCallKey, UE::HLSLTree::FFunctionCall*> FunctionCallMap;
 	TMap<FExpressionKey, UE::HLSLTree::FExpression*> ExpressionMap;
-	TMap<UMaterialExpression*, UE::HLSLTree::FStatement*> StatementMap;
+	TMap<UMaterialExpression*, FStatementEntry> StatementMap;
 	bool bGeneratedResult;
 };
 
diff --git a/Engine/Source/Runtime/Engine/Public/MaterialHLSLTree.h b/Engine/Source/Runtime/Engine/Public/MaterialHLSLTree.h
--- a/Engine/Source/Runtime/Engine/Public/MaterialHLSLTree.h
+++ b/Engine/Source/Runtime/Engine/Public/MaterialHLSLTree.h
@@ -27,9 +27,9 @@ class FMaterialHLSLTree
 public:
 	FMaterialHLSLTree();
 
-	void InitializeForMaterial(const FMaterialCompileTargetParameters& InCompilerTarget, FMaterial& InOutMaterial);
+	bool InitializeForMaterial(const FMaterialCompileTargetParameters& InCompilerTarget, FMaterial& InOutMaterial);
 
-	void InitializeForFunction(const FMaterialCompileTargetParameters& InCompilerTarget, UMaterialFunctionInterface* InOutFunction);
+	bool InitializeForFunction(const FMaterialCompileTargetParameters& InCompilerTarget, UMaterialFunctionInterface* InOutFunction);
 
 	UE::HLSLTree::FTree& GetTree() const { return *HLSLTree; }
 	int32 GetNumFunctionInputs() const { return NumFunctionInputs; }
