Expose id->object lookup in object trace system

diff --git a/Engine/Source/Runtime/Engine/Private/ObjectTrace.cpp b/Engine/Source/Runtime/Engine/Private/ObjectTrace.cpp
--- a/Engine/Source/Runtime/Engine/Private/ObjectTrace.cpp
+++ b/Engine/Source/Runtime/Engine/Private/ObjectTrace.cpp
@@ -76,10 +76,21 @@ struct FTracedObjectAnnotation
 	{
 		return bTraced == false && Id == 0;
 	}
+
+
+	bool operator == (const FTracedObjectAnnotation& other) const
+	{
+		return Id == other.Id;
+	}
 };
 
+int32 GetTypeHash(const FTracedObjectAnnotation& Annotation)
+{
+	return GetTypeHash(Annotation.Id);
+}
+
 // Object annotations used for tracing
-FUObjectAnnotationSparse<FTracedObjectAnnotation, true> GObjectTraceAnnotations;
+FUObjectAnnotationSparseSearchable<FTracedObjectAnnotation, true> GObjectTraceAnnotations;
 
 // Handle used to hook to world tick
 static FDelegateHandle WorldTickStartHandle;
@@ -151,6 +162,14 @@ uint64 FObjectTrace::GetObjectId(const UObject* InObject)
 	return Id | (OuterId << 32);
 }
 
+UObject* FObjectTrace::GetObjectFromId(uint64 Id)
+{
+	FTracedObjectAnnotation FindAnnotation;
+	// Id used for annotation map doesn't include the parent id in the upper bits, so zero those first
+	FindAnnotation.Id = Id & 0x00000000FFFFFFFFll;
+	return GObjectTraceAnnotations.Find(FindAnnotation);
+}
+
 void FObjectTrace::ResetWorldElapsedTime(const UWorld* World)
 {
 	if(UObjectTraceWorldSubsystem* WorldSubsystem = UWorld::GetSubsystem<UObjectTraceWorldSubsystem>(World))
diff --git a/Engine/Source/Runtime/Engine/Public/ObjectTrace.h b/Engine/Source/Runtime/Engine/Public/ObjectTrace.h
--- a/Engine/Source/Runtime/Engine/Public/ObjectTrace.h
+++ b/Engine/Source/Runtime/Engine/Public/ObjectTrace.h
@@ -60,6 +60,9 @@ struct FObjectTrace
 	/** Helper function to get an object ID from a UObject */
 	ENGINE_API static uint64 GetObjectId(const UObject* InObject);
 
+	/** Helper function to get the UObject from an ObjectId, if the UObject still exists */
+	ENGINE_API static UObject* GetObjectFromId(uint64 id);
+
 	/** Helper function to get an object's world's tick counter */
 	ENGINE_API static uint16 GetObjectWorldTickCounter(const UObject* InObject);
 
