Add spaces to overridden display names, it doesn't add them automatically
Move some functions out of Utilities into their own categories like Transformation, this improves the UX for the node picker and Utilities is for Core functionality
Simplify and combine some redundant categories
Clean up functions like IsValid to specify the type and attempt to unify case for math functions

diff --git a/Engine/Plugins/Experimental/VirtualProductionUtilities/Source/VPUtilities/Public/VPFullScreenUserWidgetActor.h b/Engine/Plugins/Experimental/VirtualProductionUtilities/Source/VPUtilities/Public/VPFullScreenUserWidgetActor.h
--- a/Engine/Plugins/Experimental/VirtualProductionUtilities/Source/VPUtilities/Public/VPFullScreenUserWidgetActor.h
+++ b/Engine/Plugins/Experimental/VirtualProductionUtilities/Source/VPUtilities/Public/VPFullScreenUserWidgetActor.h
@@ -12,7 +12,7 @@ class UVPFullScreenUserWidget;
 /**
  * Widgets are first rendered to a render target, then that render target is displayed in the world.
  */
-UCLASS(Blueprintable, ClassGroup = "UserInterface", HideCategories = (Actor, Input, Movement, Collision, Rendering, "Utilities|Transformation", LOD), ShowCategories = ("Input|MouseInput", "Input|TouchInput"))
+UCLASS(Blueprintable, ClassGroup = "UserInterface", HideCategories = (Actor, Input, Movement, Collision, Rendering, Transformation, LOD), ShowCategories = ("Input|MouseInput", "Input|TouchInput"))
 class VPUTILITIES_API AFullScreenUserWidgetActor : public AInfo
 {
 	GENERATED_BODY()
diff --git a/Engine/Source/Developer/FunctionalTesting/Classes/FunctionalTestLevelScript.h b/Engine/Source/Developer/FunctionalTesting/Classes/FunctionalTestLevelScript.h
--- a/Engine/Source/Developer/FunctionalTesting/Classes/FunctionalTestLevelScript.h
+++ b/Engine/Source/Developer/FunctionalTesting/Classes/FunctionalTestLevelScript.h
@@ -6,7 +6,7 @@
 #include "Engine/LevelScriptActor.h"
 #include "FunctionalTestLevelScript.generated.h"
 
-UCLASS(notplaceable, meta=(KismetHideOverrides = "ReceiveAnyDamage,ReceivePointDamage,ReceiveRadialDamage,ReceiveActorBeginOverlap,ReceiveActorEndOverlap,ReceiveHit,ReceiveDestroyed,ReceiveActorBeginCursorOver,ReceiveActorEndCursorOver,ReceiveActorOnClicked,ReceiveActorOnReleased,ReceiveActorOnInputTouchBegin,ReceiveActorOnInputTouchEnd,ReceiveActorOnInputTouchEnter,ReceiveActorOnInputTouchLeave"), HideCategories=(Collision,Rendering,"Utilities|Transformation"))
+UCLASS(notplaceable, meta=(KismetHideOverrides = "ReceiveAnyDamage,ReceivePointDamage,ReceiveRadialDamage,ReceiveActorBeginOverlap,ReceiveActorEndOverlap,ReceiveHit,ReceiveDestroyed,ReceiveActorBeginCursorOver,ReceiveActorEndCursorOver,ReceiveActorOnClicked,ReceiveActorOnReleased,ReceiveActorOnInputTouchBegin,ReceiveActorOnInputTouchEnd,ReceiveActorOnInputTouchEnter,ReceiveActorOnInputTouchLeave"), HideCategories=(Collision,Rendering,Transformation))
 class FUNCTIONALTESTING_API AFunctionalTestLevelScript : public ALevelScriptActor
 {
 	GENERATED_UCLASS_BODY()
diff --git a/Engine/Source/Runtime/AssetRegistry/Public/AssetRegistry/AssetRegistryHelpers.h b/Engine/Source/Runtime/AssetRegistry/Public/AssetRegistry/AssetRegistryHelpers.h
--- a/Engine/Source/Runtime/AssetRegistry/Public/AssetRegistry/AssetRegistryHelpers.h
+++ b/Engine/Source/Runtime/AssetRegistry/Public/AssetRegistry/AssetRegistryHelpers.h
@@ -39,7 +39,7 @@ public:
 	static FAssetData CreateAssetData(const UObject* InAsset, bool bAllowBlueprintClass = false);
 
 	/** Checks to see if this AssetData refers to an asset or is NULL */
-	UFUNCTION(BlueprintPure, Category = "Asset Data", meta=(ScriptMethod))
+	UFUNCTION(BlueprintPure, Category = "Asset Data", meta=(DisplayName = "Is Valid Asset Data", ScriptMethod))
 	static bool IsValid(const FAssetData& InAssetData);
 
 	/** Returns true if this is the primary asset in a package, true for maps and assets but false for secondary objects like class redirectors */
diff --git a/Engine/Source/Runtime/Engine/Classes/Atmosphere/AtmosphericFog.h b/Engine/Source/Runtime/Engine/Classes/Atmosphere/AtmosphericFog.h
--- a/Engine/Source/Runtime/Engine/Classes/Atmosphere/AtmosphericFog.h
+++ b/Engine/Source/Runtime/Engine/Classes/Atmosphere/AtmosphericFog.h
@@ -12,7 +12,7 @@
  *	A placeable fog actor that simulates atmospheric light scattering  
  *	@see https://docs.unrealengine.com/latest/INT/Engine/Actors/FogEffects/AtmosphericFog/index.html
  */
-UCLASS(showcategories=(Movement, Rendering, "Utilities|Transformation", "Input|MouseInput", "Input|TouchInput"), ClassGroup=Fog, hidecategories=(Info,Object,Input), MinimalAPI)
+UCLASS(showcategories=(Movement, Rendering, Transformation, "Input|MouseInput", "Input|TouchInput"), ClassGroup=Fog, hidecategories=(Info,Object,Input), MinimalAPI)
 class UE_DEPRECATED(4.26, "Please use the SkyAtmosphere actor instead.") AAtmosphericFog : public AInfo
 {
 	GENERATED_UCLASS_BODY()
diff --git a/Engine/Source/Runtime/Engine/Classes/Components/ActorComponent.h b/Engine/Source/Runtime/Engine/Classes/Components/ActorComponent.h
--- a/Engine/Source/Runtime/Engine/Classes/Components/ActorComponent.h
+++ b/Engine/Source/Runtime/Engine/Classes/Components/ActorComponent.h
@@ -432,7 +432,7 @@ public:
 	virtual void SetAutoActivate(bool bNewAutoActivate);
 
 	/** Sets whether this component can tick when paused. */
-	UFUNCTION(BlueprintCallable, Category="Utilities")
+	UFUNCTION(BlueprintCallable, Category="Components|Tick")
 	void SetTickableWhenPaused(bool bTickableWhenPaused);
 
 	/** Create any physics engine information for this component */
@@ -681,7 +681,7 @@ public:
 	 * 
 	 * @param	bEnabled - Whether it should be enabled or not
 	 */
-	UFUNCTION(BlueprintCallable, Category="Utilities")
+	UFUNCTION(BlueprintCallable, Category="Components|Tick")
 	virtual void SetComponentTickEnabled(bool bEnabled);
 
 	/** 
@@ -693,25 +693,25 @@ public:
 	/** 
 	 * Returns whether this component has tick enabled or not
 	 */
-	UFUNCTION(BlueprintCallable, Category="Utilities")
+	UFUNCTION(BlueprintCallable, Category="Components|Tick")
 	virtual bool IsComponentTickEnabled() const;
 
 	/** 
 	* Sets the tick interval for this component's primary tick function. Does not enable the tick interval. Takes effect on next tick.
 	* @param TickInterval	The duration between ticks for this component's primary tick function
 	*/
-	UFUNCTION(BlueprintCallable, Category="Utilities")
+	UFUNCTION(BlueprintCallable, Category="Components|Tick")
 	void SetComponentTickInterval(float TickInterval);
 
 	/**
 	* Sets the tick interval for this component's primary tick function. Does not enable the tick interval. Takes effect imediately.
 	* @param TickInterval	The duration between ticks for this component's primary tick function
 	*/
-	UFUNCTION(BlueprintCallable, Category = "Utilities")
+	UFUNCTION(BlueprintCallable, Category="Components|Tick")
 	void SetComponentTickIntervalAndCooldown(float TickInterval);
 
 	/** Returns the tick interval for this component's primary tick function, which is the frequency in seconds at which it will be executed */
-	UFUNCTION(BlueprintCallable, Category="Utilities")
+	UFUNCTION(BlueprintCallable, Category="Components|Tick")
 	float GetComponentTickInterval() const;
 
 	/**
@@ -931,30 +931,30 @@ public:
 	/**
 	 * Unregister and mark for pending kill a component.  This may not be used to destroy a component that is owned by an actor unless the owning actor is calling the function.
 	 */
-	UFUNCTION(BlueprintCallable, Category="Components", meta=(Keywords = "Delete", HidePin="Object", DefaultToSelf="Object", DisplayName = "DestroyComponent", ScriptName = "DestroyComponent"))
+	UFUNCTION(BlueprintCallable, Category="Components", meta=(Keywords = "Delete", HidePin="Object", DefaultToSelf="Object", DisplayName = "Destroy Component", ScriptName = "DestroyComponent"))
 	void K2_DestroyComponent(UObject* Object);
 
 	/** Unregisters and immediately re-registers component. */
 	void ReregisterComponent();
 
 	/** Changes the ticking group for this component */
-	UFUNCTION(BlueprintCallable, Category="Utilities", meta=(Keywords = "dependency"))
+	UFUNCTION(BlueprintCallable, Category="Components|Tick", meta=(Keywords = "dependency"))
 	void SetTickGroup(ETickingGroup NewTickGroup);
 
 	/** Make this component tick after PrerequisiteActor */
-	UFUNCTION(BlueprintCallable, Category="Utilities", meta=(Keywords = "dependency"))
+	UFUNCTION(BlueprintCallable, Category="Components|Tick", meta=(Keywords = "dependency"))
 	virtual void AddTickPrerequisiteActor(AActor* PrerequisiteActor);
 
 	/** Make this component tick after PrerequisiteComponent. */
-	UFUNCTION(BlueprintCallable, Category="Utilities", meta=(Keywords = "dependency"))
+	UFUNCTION(BlueprintCallable, Category="Components|Tick", meta=(Keywords = "dependency"))
 	virtual void AddTickPrerequisiteComponent(UActorComponent* PrerequisiteComponent);
 
 	/** Remove tick dependency on PrerequisiteActor. */
-	UFUNCTION(BlueprintCallable, Category="Utilities", meta=(Keywords = "dependency"))
+	UFUNCTION(BlueprintCallable, Category="Components|Tick", meta=(Keywords = "dependency"))
 	virtual void RemoveTickPrerequisiteActor(AActor* PrerequisiteActor);
 
 	/** Remove tick dependency on PrerequisiteComponent. */
-	UFUNCTION(BlueprintCallable, Category="Utilities", meta=(Keywords = "dependency"))
+	UFUNCTION(BlueprintCallable, Category="Components|Tick", meta=(Keywords = "dependency"))
 	virtual void RemoveTickPrerequisiteComponent(UActorComponent* PrerequisiteComponent);
 
 	/** Event called every frame if tick is enabled */
diff --git a/Engine/Source/Runtime/Engine/Classes/Components/MeshComponent.h b/Engine/Source/Runtime/Engine/Classes/Components/MeshComponent.h
--- a/Engine/Source/Runtime/Engine/Classes/Components/MeshComponent.h
+++ b/Engine/Source/Runtime/Engine/Classes/Components/MeshComponent.h
@@ -27,16 +27,16 @@ public:
 	UPROPERTY(EditAnywhere, AdvancedDisplay, Category=Rendering, Meta=(ToolTip="Material overrides."))
 	TArray<TObjectPtr<class UMaterialInterface>> OverrideMaterials;
 	
-	UFUNCTION(BlueprintCallable, Category="Components|Mesh")
+	UFUNCTION(BlueprintCallable, Category="Rendering|Material")
 	virtual TArray<class UMaterialInterface*> GetMaterials() const;
 
-	UFUNCTION(BlueprintCallable, Category = "Components|Mesh")
+	UFUNCTION(BlueprintCallable, Category = "Rendering|Material")
 	virtual int32 GetMaterialIndex(FName MaterialSlotName) const;
 
-	UFUNCTION(BlueprintCallable, Category = "Components|Mesh")
+	UFUNCTION(BlueprintCallable, Category = "Rendering|Material")
 	virtual TArray<FName> GetMaterialSlotNames() const;
 
-	UFUNCTION(BlueprintCallable, Category = "Components|Mesh")
+	UFUNCTION(BlueprintCallable, Category = "Rendering|Material")
 	virtual bool IsMaterialSlotNameValid(FName MaterialSlotName) const;
 
 	/** Returns override Materials count */
@@ -95,7 +95,7 @@ public:
 	 */
 	virtual void RegisterLODStreamingCallback(FLODStreamingCallback&& Callback, int32 LODIdx, float TimeoutSecs, bool bOnStreamIn);
 
-	/** Get the material info for texture stremaing. Return whether the data is valid or not. */
+	/** Get the material info for texture streaming. Return whether the data is valid or not. */
 	virtual bool GetMaterialStreamingData(int32 MaterialIndex, FPrimitiveMaterialInfo& MaterialData) const { return false; }
 
 	/** Generate streaming data for all materials. */
diff --git a/Engine/Source/Runtime/Engine/Classes/Components/PrimitiveComponent.h b/Engine/Source/Runtime/Engine/Classes/Components/PrimitiveComponent.h
--- a/Engine/Source/Runtime/Engine/Classes/Components/PrimitiveComponent.h
+++ b/Engine/Source/Runtime/Engine/Classes/Components/PrimitiveComponent.h
@@ -798,7 +798,7 @@ public:
 	 * @param Tolerance  How many seconds ago the actor last render time can be and still count as having been "recently" rendered.
 	 * @return Whether this actor was recently rendered.
 	 */
-	UFUNCTION(Category = "Rendering", BlueprintCallable, meta=(DisplayName="WasComponentRecentlyRendered", Keywords="scene visible"))
+	UFUNCTION(Category = "Rendering", BlueprintCallable, meta=(DisplayName="Was Component Recently Rendered", Keywords="scene visible"))
 	bool WasRecentlyRendered(float Tolerance = 0.2) const;
 
 	void SetLastRenderTime(float InLastRenderTime);
@@ -826,7 +826,7 @@ public:
 	/**
 	 * Returns the list of actors we currently ignore when moving.
 	 */
-	UFUNCTION(BlueprintCallable, meta=(DisplayName="GetMoveIgnoreActors", UnsafeDuringActorConstruction="true"), Category = "Collision")
+	UFUNCTION(BlueprintCallable, meta=(DisplayName="Get Move Ignore Actors", UnsafeDuringActorConstruction="true"), Category = "Collision")
 	TArray<AActor*> CopyArrayOfMoveIgnoreActors();
 
 	/**
@@ -861,7 +861,7 @@ public:
 	/**
 	* Returns the list of actors we currently ignore when moving.
 	*/
-	UFUNCTION(BlueprintCallable, meta=(DisplayName="GetMoveIgnoreComponents", UnsafeDuringActorConstruction="true"), Category = "Collision")
+	UFUNCTION(BlueprintCallable, meta=(DisplayName="Get Move Ignore Components", UnsafeDuringActorConstruction="true"), Category = "Collision")
 	TArray<UPrimitiveComponent*> CopyArrayOfMoveIgnoreComponents();
 
 	/**
@@ -1219,7 +1219,7 @@ public:
 	 *	@param	SectionIndex	Section of the mesh that the face belongs to
 	 *	@return					Material applied to section that the hit face belongs to
 	 */
-	UFUNCTION(BlueprintPure, Category = "Components|Mesh")
+	UFUNCTION(BlueprintPure, Category = "Rendering|Material")
 	virtual UMaterialInterface* GetMaterialFromCollisionFaceIndex(int32 FaceIndex, int32& SectionIndex) const;
 
 	/** Returns the slope override struct for this component. */
diff --git a/Engine/Source/Runtime/Engine/Classes/Components/SceneComponent.h b/Engine/Source/Runtime/Engine/Classes/Components/SceneComponent.h
--- a/Engine/Source/Runtime/Engine/Classes/Components/SceneComponent.h
+++ b/Engine/Source/Runtime/Engine/Classes/Components/SceneComponent.h
@@ -366,7 +366,7 @@ public:
 	 *							If false, physics velocity is updated based on the change in position (affecting ragdoll parts).
 	 *							If CCD is on and not teleporting, this will affect objects along the entire sweep volume.
 	 */
-	UFUNCTION(BlueprintCallable, Category="Utilities|Transformation", meta=(DisplayName="SetRelativeLocation", ScriptName="SetRelativeLocation"))
+	UFUNCTION(BlueprintCallable, Category="Transformation", meta=(DisplayName="Set Relative Location", ScriptName="SetRelativeLocation"))
 	void K2_SetRelativeLocation(FVector NewLocation, bool bSweep, FHitResult& SweepHitResult, bool bTeleport);
 	void SetRelativeLocation(FVector NewLocation, bool bSweep=false, FHitResult* OutSweepHitResult=nullptr, ETeleportType Teleport = ETeleportType::None);
 
@@ -379,7 +379,7 @@ public:
 	 *							If true, physics velocity for this object is unchanged (so ragdoll parts are not affected by change in location).
 	 *							If false, physics velocity is updated based on the change in position (affecting ragdoll parts).
 	 */
-	UFUNCTION(BlueprintCallable, Category="Utilities|Transformation", meta=(DisplayName="SetRelativeRotation", ScriptName="SetRelativeRotation", AdvancedDisplay="bSweep,SweepHitResult,bTeleport"))
+	UFUNCTION(BlueprintCallable, Category="Transformation", meta=(DisplayName="Set Relative Rotation", ScriptName="SetRelativeRotation", AdvancedDisplay="bSweep,SweepHitResult,bTeleport"))
 	void K2_SetRelativeRotation(FRotator NewRotation, bool bSweep, FHitResult& SweepHitResult, bool bTeleport);
 	void SetRelativeRotation(FRotator NewRotation, bool bSweep=false, FHitResult* OutSweepHitResult=nullptr, ETeleportType Teleport = ETeleportType::None);
 	void SetRelativeRotation(const FQuat& NewRotation, bool bSweep=false, FHitResult* OutSweepHitResult=nullptr, ETeleportType Teleport = ETeleportType::None);
@@ -393,20 +393,20 @@ public:
 	 *							If true, physics velocity for this object is unchanged (so ragdoll parts are not affected by change in location).
 	 *							If false, physics velocity is updated based on the change in position (affecting ragdoll parts).
 	 */
-	UFUNCTION(BlueprintCallable, Category="Utilities|Transformation", meta=(DisplayName="SetRelativeTransform", ScriptName="SetRelativeTransform"))
+	UFUNCTION(BlueprintCallable, Category="Transformation", meta=(DisplayName="Set Relative Transform", ScriptName="SetRelativeTransform"))
 	void K2_SetRelativeTransform(const FTransform& NewTransform, bool bSweep, FHitResult& SweepHitResult, bool bTeleport);
 	void SetRelativeTransform(const FTransform& NewTransform, bool bSweep=false, FHitResult* OutSweepHitResult=nullptr, ETeleportType Teleport = ETeleportType::None);
 
 	/** Returns the transform of the component relative to its parent */
-	UFUNCTION(BlueprintCallable, Category="Utilities|Transformation")
+	UFUNCTION(BlueprintCallable, Category="Transformation")
 	FTransform GetRelativeTransform() const;
 
 	/** Reset the transform of the component relative to its parent. Sets relative location to zero, relative rotation to no rotation, and Scale to 1. */
-	UFUNCTION(BlueprintCallable, Category="Utilities|Transformation")
+	UFUNCTION(BlueprintCallable, Category="Transformation")
 	void ResetRelativeTransform();
 
 	/** Set the non-uniform scale of the component relative to its parent */
-	UFUNCTION(BlueprintCallable, Category="Utilities|Transformation")
+	UFUNCTION(BlueprintCallable, Category="Transformation")
 	void SetRelativeScale3D(FVector NewScale3D);
 
 	/**
@@ -420,7 +420,7 @@ public:
 	 *							If false, physics velocity is updated based on the change in position (affecting ragdoll parts).
 	 *							If CCD is on and not teleporting, this will affect objects along the entire sweep volume.
 	 */
-	UFUNCTION(BlueprintCallable, Category="Utilities|Transformation", meta=(DisplayName="AddRelativeLocation", ScriptName="AddRelativeLocation"))
+	UFUNCTION(BlueprintCallable, Category="Transformation", meta=(DisplayName="Add Relative Location", ScriptName="AddRelativeLocation"))
 	void K2_AddRelativeLocation(FVector DeltaLocation, bool bSweep, FHitResult& SweepHitResult, bool bTeleport);
 	void AddRelativeLocation(FVector DeltaLocation, bool bSweep=false, FHitResult* OutSweepHitResult=nullptr, ETeleportType Teleport = ETeleportType::None);
 
@@ -433,7 +433,7 @@ public:
 	 *							If true, physics velocity for this object is unchanged (so ragdoll parts are not affected by change in location).
 	 *							If false, physics velocity is updated based on the change in position (affecting ragdoll parts).
 	 */
-	UFUNCTION(BlueprintCallable, Category="Utilities|Transformation", meta=(DisplayName="AddRelativeRotation", ScriptName="AddRelativeRotation", AdvancedDisplay="bSweep,SweepHitResult,bTeleport"))
+	UFUNCTION(BlueprintCallable, Category="Transformation", meta=(DisplayName="Add Relative Rotation", ScriptName="AddRelativeRotation", AdvancedDisplay="bSweep,SweepHitResult,bTeleport"))
 	void K2_AddRelativeRotation(FRotator DeltaRotation, bool bSweep, FHitResult& SweepHitResult, bool bTeleport);
 	void AddRelativeRotation(FRotator DeltaRotation, bool bSweep=false, FHitResult* OutSweepHitResult=nullptr, ETeleportType Teleport = ETeleportType::None);
 	void AddRelativeRotation(const FQuat& DeltaRotation, bool bSweep=false, FHitResult* OutSweepHitResult=nullptr, ETeleportType Teleport = ETeleportType::None);
@@ -449,7 +449,7 @@ public:
 	 *							If false, physics velocity is updated based on the change in position (affecting ragdoll parts).
 	 *							If CCD is on and not teleporting, this will affect objects along the entire sweep volume.
 	 */
-	UFUNCTION(BlueprintCallable, Category="Utilities|Transformation", meta=(DisplayName="AddLocalOffset", ScriptName="AddLocalOffset", Keywords="location position"))
+	UFUNCTION(BlueprintCallable, Category="Transformation", meta=(DisplayName="Add Local Offset", ScriptName="AddLocalOffset", Keywords="location position"))
 	void K2_AddLocalOffset(FVector DeltaLocation, bool bSweep, FHitResult& SweepHitResult, bool bTeleport);
 	void AddLocalOffset(FVector DeltaLocation, bool bSweep=false, FHitResult* OutSweepHitResult=nullptr, ETeleportType Teleport = ETeleportType::None);
 
@@ -462,7 +462,7 @@ public:
 	 *							If true, physics velocity for this object is unchanged (so ragdoll parts are not affected by change in location).
 	 *							If false, physics velocity is updated based on the change in position (affecting ragdoll parts).
 	 */
-	UFUNCTION(BlueprintCallable, Category="Utilities|Transformation", meta=(DisplayName="AddLocalRotation", ScriptName="AddLocalRotation", AdvancedDisplay="bSweep,SweepHitResult,bTeleport"))
+	UFUNCTION(BlueprintCallable, Category="Transformation", meta=(DisplayName="Add Local Rotation", ScriptName="AddLocalRotation", AdvancedDisplay="bSweep,SweepHitResult,bTeleport"))
 	void K2_AddLocalRotation(FRotator DeltaRotation, bool bSweep, FHitResult& SweepHitResult, bool bTeleport);
 	void AddLocalRotation(FRotator DeltaRotation, bool bSweep=false, FHitResult* OutSweepHitResult=nullptr, ETeleportType Teleport = ETeleportType::None);
 	void AddLocalRotation(const FQuat& DeltaRotation, bool bSweep=false, FHitResult* OutSweepHitResult=nullptr, ETeleportType Teleport = ETeleportType::None);
@@ -478,7 +478,7 @@ public:
 	 *							If false, physics velocity is updated based on the change in position (affecting ragdoll parts).
 	 *							If CCD is on and not teleporting, this will affect objects along the entire sweep volume.
 	 */
-	UFUNCTION(BlueprintCallable, Category="Utilities|Transformation", meta=(DisplayName="AddLocalTransform", ScriptName="AddLocalTransform"))
+	UFUNCTION(BlueprintCallable, Category="Transformation", meta=(DisplayName="Add Local Transform", ScriptName="AddLocalTransform"))
 	void K2_AddLocalTransform(const FTransform& DeltaTransform, bool bSweep, FHitResult& SweepHitResult, bool bTeleport);
 	void AddLocalTransform(const FTransform& DeltaTransform, bool bSweep=false, FHitResult* OutSweepHitResult=nullptr, ETeleportType Teleport = ETeleportType::None);
 
@@ -493,7 +493,7 @@ public:
 	 *							If false, physics velocity is updated based on the change in position (affecting ragdoll parts).
 	 *							If CCD is on and not teleporting, this will affect objects along the entire sweep volume.
 	 */
-	UFUNCTION(BlueprintCallable, Category="Utilities|Transformation", meta=(DisplayName="SetWorldLocation", ScriptName="SetWorldLocation"))
+	UFUNCTION(BlueprintCallable, Category="Transformation", meta=(DisplayName="Set World Location", ScriptName="SetWorldLocation"))
 	void K2_SetWorldLocation(FVector NewLocation, bool bSweep, FHitResult& SweepHitResult, bool bTeleport);
 	void SetWorldLocation(FVector NewLocation, bool bSweep=false, FHitResult* OutSweepHitResult=nullptr, ETeleportType Teleport = ETeleportType::None);
 
@@ -507,7 +507,7 @@ public:
 	 *							If false, physics velocity is updated based on the change in position (affecting ragdoll parts).
 	 *							If CCD is on and not teleporting, this will affect objects along the entire sweep volume.
 	 */
-	UFUNCTION(BlueprintCallable, Category="Utilities|Transformation", meta=(DisplayName="SetWorldRotation", ScriptName="SetWorldRotation", AdvancedDisplay="bSweep,SweepHitResult,bTeleport"))
+	UFUNCTION(BlueprintCallable, Category="Transformation", meta=(DisplayName="Set World Rotation", ScriptName="SetWorldRotation", AdvancedDisplay="bSweep,SweepHitResult,bTeleport"))
 	void K2_SetWorldRotation(FRotator NewRotation, bool bSweep, FHitResult& SweepHitResult, bool bTeleport);
 	void SetWorldRotation(FRotator NewRotation, bool bSweep=false, FHitResult* OutSweepHitResult=nullptr, ETeleportType Teleport = ETeleportType::None);
 	void SetWorldRotation(const FQuat& NewRotation, bool bSweep=false, FHitResult* OutSweepHitResult=nullptr, ETeleportType Teleport = ETeleportType::None);
@@ -516,7 +516,7 @@ public:
 	 * Set the relative scale of the component to put it at the supplied scale in world space.
 	 * @param NewScale		New scale in world space for this component.
 	 */
-	UFUNCTION(BlueprintCallable, Category="Utilities|Transformation")
+	UFUNCTION(BlueprintCallable, Category="Transformation")
 	void SetWorldScale3D(FVector NewScale);
 
 	/**
@@ -530,7 +530,7 @@ public:
 	 *							If false, physics velocity is updated based on the change in position (affecting ragdoll parts).
 	 *							If CCD is on and not teleporting, this will affect objects along the entire sweep volume.
 	 */
-	UFUNCTION(BlueprintCallable, Category="Utilities|Transformation", meta=(DisplayName="SetWorldTransform", ScriptName="SetWorldTransform"))
+	UFUNCTION(BlueprintCallable, Category="Transformation", meta=(DisplayName="Set World Transform", ScriptName="SetWorldTransform"))
 	void K2_SetWorldTransform(const FTransform& NewTransform, bool bSweep, FHitResult& SweepHitResult, bool bTeleport);
 	void SetWorldTransform(const FTransform& NewTransform, bool bSweep=false, FHitResult* OutSweepHitResult=nullptr, ETeleportType Teleport = ETeleportType::None);
 
@@ -545,7 +545,7 @@ public:
 	 *							If false, physics velocity is updated based on the change in position (affecting ragdoll parts).
 	 *							If CCD is on and not teleporting, this will affect objects along the entire sweep volume.
 	 */
-	UFUNCTION(BlueprintCallable, Category="Utilities|Transformation", meta=(DisplayName="AddWorldOffset", ScriptName="AddWorldOffset", Keywords="location position"))
+	UFUNCTION(BlueprintCallable, Category="Transformation", meta=(DisplayName="Add World Offset", ScriptName="AddWorldOffset", Keywords="location position"))
 	void K2_AddWorldOffset(FVector DeltaLocation, bool bSweep, FHitResult& SweepHitResult, bool bTeleport);
 	void AddWorldOffset(FVector DeltaLocation, bool bSweep=false, FHitResult* OutSweepHitResult=nullptr, ETeleportType Teleport = ETeleportType::None);
 
@@ -559,7 +559,7 @@ public:
 	 *							If false, physics velocity is updated based on the change in position (affecting ragdoll parts).
 	 *							If CCD is on and not teleporting, this will affect objects along the entire sweep volume.
 	 */
-	UFUNCTION(BlueprintCallable, Category="Utilities|Transformation", meta=(DisplayName="AddWorldRotation", ScriptName="AddWorldRotation", AdvancedDisplay="bSweep,SweepHitResult,bTeleport"))
+	UFUNCTION(BlueprintCallable, Category="Transformation", meta=(DisplayName="Add World Rotation", ScriptName="AddWorldRotation", AdvancedDisplay="bSweep,SweepHitResult,bTeleport"))
 	void K2_AddWorldRotation(FRotator DeltaRotation, bool bSweep, FHitResult& SweepHitResult, bool bTeleport);
 	void AddWorldRotation(FRotator DeltaRotation, bool bSweep=false, FHitResult* OutSweepHitResult=nullptr, ETeleportType Teleport = ETeleportType::None);
 	void AddWorldRotation(const FQuat& DeltaRotation, bool bSweep=false, FHitResult* OutSweepHitResult=nullptr, ETeleportType Teleport = ETeleportType::None);
@@ -575,7 +575,7 @@ public:
 	 *							If false, physics velocity is updated based on the change in position (affecting ragdoll parts).
 	 *							If CCD is on and not teleporting, this will affect objects along the entire sweep volume.
 	 */
-	UFUNCTION(BlueprintCallable, Category="Utilities|Transformation", meta=(DisplayName="AddWorldTransform", ScriptName="AddWorldTransform"))
+	UFUNCTION(BlueprintCallable, Category="Transformation", meta=(DisplayName="Add World Transform", ScriptName="AddWorldTransform"))
 	void K2_AddWorldTransform(const FTransform& DeltaTransform, bool bSweep, FHitResult& SweepHitResult, bool bTeleport);
 	void AddWorldTransform(const FTransform& DeltaTransform, bool bSweep=false, FHitResult* OutSweepHitResult=nullptr, ETeleportType Teleport = ETeleportType::None);
 
@@ -590,36 +590,36 @@ public:
 	 *							If false, physics velocity is updated based on the change in position (affecting ragdoll parts).
 	 *							If CCD is on and not teleporting, this will affect objects along the entire sweep volume.
 	 */
-	UFUNCTION(BlueprintCallable, Category="Utilities|Transformation", meta=(DisplayName="AddWorldTransformKeepScale", ScriptName="AddWorldTransformKeepScale"))
+	UFUNCTION(BlueprintCallable, Category="Transformation", meta=(DisplayName="Add World Transform Keep Scale", ScriptName="AddWorldTransformKeepScale"))
 	void K2_AddWorldTransformKeepScale(const FTransform& DeltaTransform, bool bSweep, FHitResult& SweepHitResult, bool bTeleport);
 	void AddWorldTransformKeepScale(const FTransform& DeltaTransform, bool bSweep=false, FHitResult* OutSweepHitResult=nullptr, ETeleportType Teleport = ETeleportType::None);
 
 	/** Return location of the component, in world space */
-	UFUNCTION(BlueprintCallable, meta=(DisplayName = "GetWorldLocation", ScriptName = "GetWorldLocation", Keywords = "position"), Category="Utilities|Transformation")
+	UFUNCTION(BlueprintCallable, meta=(DisplayName = "Get World Location", ScriptName = "GetWorldLocation", Keywords = "position"), Category="Transformation")
 	FVector K2_GetComponentLocation() const;
 
 	/** Returns rotation of the component, in world space. */
-	UFUNCTION(BlueprintCallable, meta=(DisplayName = "GetWorldRotation", ScriptName = "GetWorldRotation"), Category="Utilities|Transformation")
+	UFUNCTION(BlueprintCallable, meta=(DisplayName = "Get World Rotation", ScriptName = "GetWorldRotation"), Category="Transformation")
 	FRotator K2_GetComponentRotation() const;
 	
 	/** Returns scale of the component, in world space. */
-	UFUNCTION(BlueprintCallable, meta=(DisplayName = "GetWorldScale", ScriptName = "GetWorldScale"), Category="Utilities|Transformation")
+	UFUNCTION(BlueprintCallable, meta=(DisplayName = "Get World Scale", ScriptName = "GetWorldScale"), Category="Transformation")
 	FVector K2_GetComponentScale() const;
 
 	/** Get the current component-to-world transform for this component */
-	UFUNCTION(BlueprintCallable, meta=(DisplayName = "GetWorldTransform", ScriptName = "GetWorldTransform"), Category="Utilities|Transformation")
+	UFUNCTION(BlueprintCallable, meta=(DisplayName = "Get World Transform", ScriptName = "GetWorldTransform"), Category="Transformation")
 	FTransform K2_GetComponentToWorld() const;
 
 	/** Get the forward (X) unit direction vector from this component, in world space.  */
-	UFUNCTION(BlueprintCallable, Category="Utilities|Transformation")
+	UFUNCTION(BlueprintCallable, Category="Transformation")
 	FVector GetForwardVector() const;
 
 	/** Get the up (Z) unit direction vector from this component, in world space.  */
-	UFUNCTION(BlueprintCallable, Category="Utilities|Transformation")
+	UFUNCTION(BlueprintCallable, Category="Transformation")
 	FVector GetUpVector() const;
 
 	/** Get the right (Y) unit direction vector from this component, in world space.  */
-	UFUNCTION(BlueprintCallable, Category="Utilities|Transformation")
+	UFUNCTION(BlueprintCallable, Category="Transformation")
 	FVector GetRightVector() const;
 
 	/** Returns whether the specified body is currently using physics simulation */
@@ -634,11 +634,11 @@ public:
 	const TArray<TObjectPtr<USceneComponent>>& GetAttachChildren() const;
 
 	/** Get the SceneComponent we are attached to. */
-	UFUNCTION(BlueprintCallable, Category="Utilities|Transformation")
+	UFUNCTION(BlueprintCallable, Category="Transformation")
 	USceneComponent* GetAttachParent() const;
 
 	/** Get the socket we are attached to. */
-	UFUNCTION(BlueprintCallable, Category="Utilities|Transformation")
+	UFUNCTION(BlueprintCallable, Category="Transformation")
 	FName GetAttachSocketName() const;
 
 	/** Gets all attachment parent components up to and including the root component */
@@ -675,7 +675,7 @@ public:
 
 	/** DEPRECATED - Use AttachToComponent() instead */
 	UE_DEPRECATED(4.17, "This function is deprecated, please use AttachToComponent() instead.")
-	UFUNCTION(BlueprintCallable, Category = "Utilities|Transformation", meta = (DeprecationMessage = "OVERRIDE BAD MESSAGE", DisplayName = "AttachTo (Deprecated)", AttachType = "KeepRelativeOffset"))
+	UFUNCTION(BlueprintCallable, Category = "Transformation", meta = (DeprecationMessage = "OVERRIDE BAD MESSAGE", DisplayName = "AttachTo (Deprecated)", AttachType = "KeepRelativeOffset"))
 	bool K2_AttachTo(USceneComponent* InParent, FName InSocketName = NAME_None, EAttachLocation::Type AttachType = EAttachLocation::KeepRelativeOffset, bool bWeldSimulatedBodies = true);
 
 	/**
@@ -698,12 +698,12 @@ public:
 	* @param  bWeldSimulatedBodies		Whether to weld together simulated physics bodies.
 	* @return True if attachment is successful (or already attached to requested parent/socket), false if attachment is rejected and there is no change in AttachParent.
 	*/
-	UFUNCTION(BlueprintCallable, Category = "Utilities|Transformation", meta = (DisplayName = "AttachComponentToComponent", ScriptName = "AttachToComponent", bWeldSimulatedBodies=true))
+	UFUNCTION(BlueprintCallable, Category = "Transformation", meta = (DisplayName = "Attach Component To Component", ScriptName = "AttachToComponent", bWeldSimulatedBodies=true))
 	bool K2_AttachToComponent(USceneComponent* Parent, FName SocketName, EAttachmentRule LocationRule, EAttachmentRule RotationRule, EAttachmentRule ScaleRule, bool bWeldSimulatedBodies);
 
 	/** DEPRECATED - Use DetachFromComponent() instead */
 	UE_DEPRECATED(4.12, "This function is deprecated, please use DetachFromComponent() instead.")
-	UFUNCTION(BlueprintCallable, Category = "Utilities|Transformation", meta = (DisplayName = "DetachFromParent (Deprecated)"))
+	UFUNCTION(BlueprintCallable, Category = "Transformation", meta = (DisplayName = "DetachFromParent (Deprecated)"))
 	virtual void DetachFromParent(bool bMaintainWorldPosition = false, bool bCallModify = true);
 
 	/** 
@@ -713,7 +713,7 @@ public:
 	 * @param ScaleRule					How to handle scales when detaching.
 	 * @param bCallModify				If true, call Modify() on the component and the current attach parent component
 	 */
-	UFUNCTION(BlueprintCallable, meta = (DisplayName = "DetachFromComponent", ScriptName = "DetachFromComponent"), Category = "Utilities|Transformation")
+	UFUNCTION(BlueprintCallable, meta = (DisplayName = "Detach From Component", ScriptName = "DetachFromComponent"), Category = "Transformation")
 	void K2_DetachFromComponent(EDetachmentRule LocationRule = EDetachmentRule::KeepRelative, EDetachmentRule RotationRule = EDetachmentRule::KeepRelative, EDetachmentRule ScaleRule = EDetachmentRule::KeepRelative, bool bCallModify = true);
 
 	/** 
@@ -726,7 +726,7 @@ public:
 	 * Gets the names of all the sockets on the component.
 	 * @return Get the names of all the sockets on the component.
 	 */
-	UFUNCTION(BlueprintCallable, Category="Utilities|Transformation", meta=(Keywords="Bone"))
+	UFUNCTION(BlueprintCallable, Category="Transformation", meta=(Keywords="Bone"))
 	TArray<FName> GetAllSocketNames() const;
 
 	/** 
@@ -734,7 +734,7 @@ public:
 	 * @param InSocketName Name of the socket or the bone to get the transform 
 	 * @return Socket transform in world space if socket if found. Otherwise it will return component's transform in world space.
 	 */
-	UFUNCTION(BlueprintCallable, Category="Utilities|Transformation", meta=(Keywords="Bone"))
+	UFUNCTION(BlueprintCallable, Category="Transformation", meta=(Keywords="Bone"))
 	virtual FTransform GetSocketTransform(FName InSocketName, ERelativeTransformSpace TransformSpace = RTS_World) const;
 
 	/** 
@@ -742,7 +742,7 @@ public:
 	 * @param InSocketName Name of the socket or the bone to get the transform 
 	 * @return Socket transform in world space if socket if found. Otherwise it will return component's transform in world space.
 	 */
-	UFUNCTION(BlueprintCallable, Category="Utilities|Transformation", meta=(Keywords="Bone"))
+	UFUNCTION(BlueprintCallable, Category="Transformation", meta=(Keywords="Bone"))
 	virtual FVector GetSocketLocation(FName InSocketName) const;
 
 	/** 
@@ -750,7 +750,7 @@ public:
 	 * @param InSocketName Name of the socket or the bone to get the transform 
 	 * @return Socket transform in world space if socket if found. Otherwise it will return component's transform in world space.
 	 */
-	UFUNCTION(BlueprintCallable, Category="Utilities|Transformation", meta=(Keywords="Bone"))
+	UFUNCTION(BlueprintCallable, Category="Transformation", meta=(Keywords="Bone"))
 	virtual FRotator GetSocketRotation(FName InSocketName) const;
 
 	/** 
@@ -758,14 +758,14 @@ public:
 	 * @param InSocketName Name of the socket or the bone to get the transform 
 	 * @return Socket transform in world space if socket if found. Otherwise it will return component's transform in world space.
 	 */
-	UFUNCTION(BlueprintCallable, Category="Utilities|Transformation", meta=(Keywords="Bone", DeprecatedFunction, DeprecationMessage="Use GetSocketRotation instead, Quat is not fully supported in blueprints."))
+	UFUNCTION(BlueprintCallable, Category="Transformation", meta=(Keywords="Bone", DeprecatedFunction, DeprecationMessage="Use GetSocketRotation instead, Quat is not fully supported in blueprints."))
 	virtual FQuat GetSocketQuaternion(FName InSocketName) const;
 
 	/** 
 	 * Return true if socket with the given name exists
 	 * @param InSocketName Name of the socket or the bone to get the transform 
 	 */
-	UFUNCTION(BlueprintCallable, Category="Utilities|Transformation", meta=(Keywords="Bone"))
+	UFUNCTION(BlueprintCallable, Category="Transformation", meta=(Keywords="Bone"))
 	virtual bool DoesSocketExist(FName InSocketName) const;
 
 	/**
@@ -782,7 +782,7 @@ public:
 	 * Get velocity of the component: either ComponentVelocity, or the velocity of the physics body if simulating physics.
 	 * @return Velocity of the component
 	 */
-	UFUNCTION(BlueprintCallable, Category="Utilities|Transformation")
+	UFUNCTION(BlueprintCallable, Category="Transformation")
 	virtual FVector GetComponentVelocity() const;
 
 	/** Returns true if this component is visible in the current context */
@@ -1120,13 +1120,13 @@ public:
 	 *							If false, physics velocity is updated based on the change in position (affecting ragdoll parts).
 	 *							If CCD is on and not teleporting, this will affect objects along the entire sweep volume.
 	 */
-	UFUNCTION(BlueprintCallable, Category="Utilities|Transformation", meta=(DisplayName="SetRelativeLocationAndRotation", ScriptName="SetRelativeLocationAndRotation"))
+	UFUNCTION(BlueprintCallable, Category="Transformation", meta=(DisplayName="Set Relative Location And Rotation", ScriptName="SetRelativeLocationAndRotation"))
 	void K2_SetRelativeLocationAndRotation(FVector NewLocation, FRotator NewRotation, bool bSweep, FHitResult& SweepHitResult, bool bTeleport);
 	void SetRelativeLocationAndRotation(FVector NewLocation, FRotator NewRotation, bool bSweep=false, FHitResult* OutSweepHitResult=nullptr, ETeleportType Teleport = ETeleportType::None);
 	void SetRelativeLocationAndRotation(FVector NewLocation, const FQuat& NewRotation, bool bSweep=false, FHitResult* OutSweepHitResult=nullptr, ETeleportType Teleport = ETeleportType::None);
 
 	/** Set which parts of the relative transform should be relative to parent, and which should be relative to world */
-	UFUNCTION(BlueprintCallable, Category="Utilities|Transformation")
+	UFUNCTION(BlueprintCallable, Category="Transformation")
 	void SetAbsolute(bool bNewAbsoluteLocation = false, bool bNewAbsoluteRotation = false, bool bNewAbsoluteScale = false);
 
 	/**
@@ -1141,7 +1141,7 @@ public:
 	 *							If false, physics velocity is updated based on the change in position (affecting ragdoll parts).
 	 *							If CCD is on and not teleporting, this will affect objects along the entire sweep volume.
 	 */
-	UFUNCTION(BlueprintCallable, Category="Utilities|Transformation", meta=(DisplayName="SetWorldLocationAndRotation", ScriptName="SetWorldLocationAndRotation"))
+	UFUNCTION(BlueprintCallable, Category="Transformation", meta=(DisplayName="Set World Location And Rotation", ScriptName="SetWorldLocationAndRotation"))
 	void K2_SetWorldLocationAndRotation(FVector NewLocation, FRotator NewRotation, bool bSweep, FHitResult& SweepHitResult, bool bTeleport);
 	void SetWorldLocationAndRotation(FVector NewLocation, FRotator NewRotation, bool bSweep=false, FHitResult* OutSweepHitResult=nullptr, ETeleportType Teleport = ETeleportType::None);
 
@@ -1185,7 +1185,7 @@ public:
 	ECollisionResponse GetCollisionResponseToComponent(USceneComponent* OtherComponent) const;
 
 	/** Set how often this component is allowed to move during runtime. Causes a component re-register if the component is already registered */
-	UFUNCTION(BlueprintCallable, Category="Utilities|Transformation")
+	UFUNCTION(BlueprintCallable, Category="Transformation")
 	virtual void SetMobility(EComponentMobility::Type NewMobility);
 
 	/** Walks up the attachment chain from this SceneComponent and returns the SceneComponent at the top. If AttachParent is NULL, returns this. */
diff --git a/Engine/Source/Runtime/Engine/Classes/Components/SkeletalMeshComponent.h b/Engine/Source/Runtime/Engine/Classes/Components/SkeletalMeshComponent.h
--- a/Engine/Source/Runtime/Engine/Classes/Components/SkeletalMeshComponent.h
+++ b/Engine/Source/Runtime/Engine/Classes/Components/SkeletalMeshComponent.h
@@ -818,7 +818,7 @@ public:
 	FConstraintBrokenSignature OnConstraintBroken;
 
 	/** Class of the object responsible for  */
-	UPROPERTY(EditAnywhere, Category = ClothingSimulation)
+	UPROPERTY(EditAnywhere, Category = Clothing)
 	TSubclassOf<class UClothingSimulationFactory> ClothingSimulationFactory;
 
 	struct FPendingRadialForces
@@ -1104,46 +1104,46 @@ public:
 	 * The snapshot is taken at the current LOD, so if for example you took the snapshot at LOD1 
 	 * and then used it at LOD0 any bones not in LOD1 will use the reference pose 
 	 */
-	UFUNCTION(BlueprintCallable, Category = "Pose")
+	UFUNCTION(BlueprintCallable, Category = "Components|SkeletalMesh")
 	void SnapshotPose(UPARAM(ref) FPoseSnapshot& Snapshot);
 
 	/**
 	 * Get/Set the max distance scale of clothing mesh vertices
 	 */
-	UFUNCTION(BlueprintCallable, Category="Components|SkeletalMesh")
+	UFUNCTION(BlueprintCallable, Category="Clothing")
 	float GetClothMaxDistanceScale() const;
-	UFUNCTION(BlueprintCallable, Category="Components|SkeletalMesh")
+	UFUNCTION(BlueprintCallable, Category="Clothing")
 	void SetClothMaxDistanceScale(float Scale);
 
 	/** 
 	 * Used to indicate we should force 'teleport' during the next call to UpdateClothState, 
 	 * This will transform positions and velocities and thus keep the simulation state, just translate it to a new pose.
 	 */
-	UFUNCTION(BlueprintCallable, Category="Components|SkeletalMesh")
+	UFUNCTION(BlueprintCallable, Category="Clothing")
 	void ForceClothNextUpdateTeleport();
 	/** 
 	 * Used to indicate we should force 'teleport and reset' during the next call to UpdateClothState.
 	 * This can be used to reset it from a bad state or by a teleport where the old state is not important anymore.
 	 */
-	UFUNCTION(BlueprintCallable, Category="Components|SkeletalMesh")
+	UFUNCTION(BlueprintCallable, Category="Clothing")
 	void ForceClothNextUpdateTeleportAndReset();
 
 	/** Stops simulating clothing, but does not show clothing ref pose. Keeps the last known simulation state */
-	UFUNCTION(BlueprintCallable, Category="Components|SkeletalMesh", meta=(UnsafeDuringActorConstruction))
+	UFUNCTION(BlueprintCallable, Category="Clothing", meta=(UnsafeDuringActorConstruction))
 	void SuspendClothingSimulation();
 
 	/** Resumes a previously suspended clothing simulation, teleporting the clothing on the next tick */
-	UFUNCTION(BlueprintCallable, Category = "Components|SkeletalMesh", meta=(UnsafeDuringActorConstruction))
+	UFUNCTION(BlueprintCallable, Category = "Clothing", meta=(UnsafeDuringActorConstruction))
 	void ResumeClothingSimulation();
 
 	/** Gets whether or not the clothing simulation is currently suspended */
-	UFUNCTION(BlueprintCallable, Category = "Components|SkeletalMesh")
+	UFUNCTION(BlueprintCallable, Category = "Clothing")
 	bool IsClothingSimulationSuspended() const;
 
 	/**
 	 * Reset the teleport mode of a next update to 'Continuous'
 	 */
-	UFUNCTION(BlueprintCallable, Category="Components|SkeletalMesh")
+	UFUNCTION(BlueprintCallable, Category="Clothing")
 	void ResetClothTeleportMode();
 
 	/** 
@@ -1152,7 +1152,7 @@ public:
 	 * @Note This will FORCE any cloth actor on the master component to simulate in local space. Also
 	 * The meshes used in the components must be identical for the cloth to bind correctly
 	 */
-	UFUNCTION(BlueprintCallable, Category="Components|SkeletalMesh", meta=(UnsafeDuringActorConstruction="true"))
+	UFUNCTION(BlueprintCallable, Category="Clothing", meta=(UnsafeDuringActorConstruction="true"))
 	void BindClothToMasterPoseComponent();
 
 	/** 
@@ -1162,7 +1162,7 @@ public:
 	 * space, we will restore this setting. This will cause the master component to reset which may be
 	 * undesirable.
 	 */
-	UFUNCTION(BlueprintCallable, Category="Components|SkeletalMesh", meta=(UnsafeDuringActorConstruction="true"))
+	UFUNCTION(BlueprintCallable, Category="Clothing", meta=(UnsafeDuringActorConstruction="true"))
 	void UnbindClothFromMasterPoseComponent(bool bRestoreSimulationSpace = true);
 
 	/**
@@ -1368,7 +1368,7 @@ public:
 	const IClothingSimulationContext* GetClothingSimulationContext() const;
 
 	/** Get the current interactor for a clothing simulation, if the current simulation supports runtime interaction. */
-	UFUNCTION(BlueprintCallable, Category=ClothingSimulation)
+	UFUNCTION(BlueprintCallable, Category=Clothing)
 	UClothingSimulationInteractor* GetClothingSimulationInteractor() const;
 
 	/** Callback when the parallel clothing task finishes, copies needed data back to component for gamethread */
@@ -1678,7 +1678,7 @@ public:
 	 *	@param bEnableGravity	Whether gravity should be enabled or disabled.
 	 *	@param BoneName			The name of the bone to modify.
 	 */
-	UFUNCTION(BlueprintCallable, Category="Physics|Components")
+	UFUNCTION(BlueprintCallable, Category="Physics")
 	void SetEnableBodyGravity(bool bEnableGravity, FName BoneName);
 
 	/**
@@ -1689,7 +1689,7 @@ public:
 	 *	@param BoneName	The name of the bone to check.
 	 *	@return True if gravity is enabled on the bone.
 	 */
-	UFUNCTION(BlueprintCallable, Category="Physics|Components")
+	UFUNCTION(BlueprintCallable, Category="Physics")
 	bool IsBodyGravityEnabled(FName BoneName);
 
 	/**
@@ -1701,7 +1701,7 @@ public:
 	*	@param BoneName			The name of the top most bone.
 	*	@param bIncludeSelf		Whether the bone specified should be edited.
 	*/
-	UFUNCTION(BlueprintCallable, Category = "Physics|Components")
+	UFUNCTION(BlueprintCallable, Category = "Physics")
 	void SetEnableGravityOnAllBodiesBelow(bool bEnableGravity, FName BoneName, bool bIncludeSelf=true);
 
 	/** 
@@ -1719,7 +1719,7 @@ public:
 	 *  @param	ClosestPointOnPhysicsAsset	The data associated with the closest point (position, normal, etc...)
 	 *  @return	true if we found a closest point
 	 */
-	UFUNCTION(BlueprintCallable, Category="Components|SkeletalMesh", meta=(DisplayName="GetClosestPointOnPhysicsAsset", ScriptName="GetClosestPointOnPhysicsAsset", Keywords="closest point"))
+	UFUNCTION(BlueprintCallable, Category="Components|SkeletalMesh", meta=(DisplayName="Get Closest Point On Physics Asset", ScriptName="GetClosestPointOnPhysicsAsset", Keywords="closest point"))
 	bool K2_GetClosestPointOnPhysicsAsset(const FVector& WorldPosition, FVector& ClosestWorldPosition, FVector& Normal, FName& BoneName, float& Distance) const;
 
 	virtual bool LineTraceComponent( FHitResult& OutHit, const FVector Start, const FVector End, const FCollisionQueryParams& Params ) override;
@@ -1923,7 +1923,7 @@ public:
 	int32 FindRootBodyIndex() const;
 
 	/** Terminate physics on all bodies below the named bone, effectively disabling collision forever. If you terminate, you won't be able to re-init later. */
-	UFUNCTION(BlueprintCallable, Category="Components|SkeletalMesh")
+	UFUNCTION(BlueprintCallable, Category="Physics")
 	void TermBodiesBelow(FName ParentBoneName);
 
 	/** Find instance of the constraint that matches the name supplied. */
@@ -1945,7 +1945,7 @@ public:
 	void SetAllBodiesNotifyRigidBodyCollision(bool bNewNotifyRigidBodyCollision);
 
 	/** Set bSimulatePhysics to true for all bone bodies. Does not change the component bSimulatePhysics flag. */
-	UFUNCTION(BlueprintCallable, Category="Components|SkeletalMesh")
+	UFUNCTION(BlueprintCallable, Category="Physics")
 	void SetAllBodiesSimulatePhysics(bool bNewSimulate);
 
 	/** This is global set up for setting physics blend weight
@@ -1954,11 +1954,11 @@ public:
 	 * Also it will respect each body's setup, so if the body is fixed, it won't simulate. Vice versa
 	 * So if you'd like all bodies to change manually, do not use this function, but SetAllBodiesPhysicsBlendWeight
 	 */
-	UFUNCTION(BlueprintCallable, Category="Components|SkeletalMesh")
+	UFUNCTION(BlueprintCallable, Category="Physics")
 	void SetPhysicsBlendWeight(float PhysicsBlendWeight);
 
 	/** Disable physics blending of bones **/
-	UFUNCTION(BlueprintCallable, Category="Components|SkeletalMesh")
+	UFUNCTION(BlueprintCallable, Category="Physics")
 	void SetEnablePhysicsBlending(bool bNewBlendPhysics);
 
 	/** [WARNING: Chaos Only] 
@@ -1966,46 +1966,46 @@ public:
 	* Bodies will not be colliding or be part of the physics simulation. 
 	* This is different from SetAllBodiesBelowSimulatePhysics that changes bodies to Kinematic/simulated 	
 	*/
-	UFUNCTION(BlueprintCallable, Category = "Components|SkeletalMesh")
+	UFUNCTION(BlueprintCallable, Category = "Physics")
 	void SetAllBodiesBelowPhysicsDisabled(const FName& InBoneName, bool bDisabled, bool bIncludeSelf = true);
 
 	/** Set all of the bones below passed in bone to be simulated */
-	UFUNCTION(BlueprintCallable, Category="Components|SkeletalMesh")
+	UFUNCTION(BlueprintCallable, Category="Physics")
 	void SetAllBodiesBelowSimulatePhysics(const FName& InBoneName, bool bNewSimulate, bool bIncludeSelf = true );
 
 	/** Allows you to reset bodies Simulate state based on where bUsePhysics is set to true in the BodySetup. */
-	UFUNCTION(BlueprintCallable, Category="Components|SkeletalMesh")
+	UFUNCTION(BlueprintCallable, Category="Physics")
 	void ResetAllBodiesSimulatePhysics();
 
-	UFUNCTION(BlueprintCallable, Category="Components|SkeletalMesh")
+	UFUNCTION(BlueprintCallable, Category="Physics")
 	void SetAllBodiesPhysicsBlendWeight(float PhysicsBlendWeight, bool bSkipCustomPhysicsType = false );
 
 	/** Set all of the bones below passed in bone to be simulated */
-	UFUNCTION(BlueprintCallable, Category="Components|SkeletalMesh")
+	UFUNCTION(BlueprintCallable, Category="Physics")
 	void SetAllBodiesBelowPhysicsBlendWeight(const FName& InBoneName, float PhysicsBlendWeight, bool bSkipCustomPhysicsType = false, bool bIncludeSelf = true );
 
 	/** Accumulate AddPhysicsBlendWeight to physics blendweight for all of the bones below passed in bone to be simulated */
-	UFUNCTION(BlueprintCallable, Category="Components|SkeletalMesh")
+	UFUNCTION(BlueprintCallable, Category="Physics")
 	void AccumulateAllBodiesBelowPhysicsBlendWeight(const FName& InBoneName, float AddPhysicsBlendWeight, bool bSkipCustomPhysicsType = false );
 
 	/** Enable or Disable AngularPositionDrive. If motor is in SLERP mode it will be turned on if either EnableSwingDrive OR EnableTwistDrive are enabled. In Twist and Swing mode the twist and the swing can be controlled individually.*/
-	UFUNCTION(BlueprintCallable, Category = "Components|SkeletalMesh")
+	UFUNCTION(BlueprintCallable, Category = "Physics")
 	void SetAllMotorsAngularPositionDrive(bool bEnableSwingDrive, bool bEnableTwistDrive, bool bSkipCustomPhysicsType = false);
 
 	/** Enable or Disable AngularVelocityDrive. If motor is in SLERP mode it will be turned on if either EnableSwingDrive OR EnableTwistDrive are enabled. In Twist and Swing mode the twist and the swing can be controlled individually.*/
-	UFUNCTION(BlueprintCallable, Category = "Components|SkeletalMesh")
+	UFUNCTION(BlueprintCallable, Category = "Physics")
 	void SetAllMotorsAngularVelocityDrive(bool bEnableSwingDrive, bool bEnableTwistDrive, bool bSkipCustomPhysicsType = false);
 
 	/** Set Angular Drive motors params for all constraint instances */
-	UFUNCTION(BlueprintCallable, Category = "Components|SkeletalMesh")
+	UFUNCTION(BlueprintCallable, Category = "Physics")
 	void SetAllMotorsAngularDriveParams(float InSpring, float InDamping, float InForceLimit, bool bSkipCustomPhysicsType = false);
 
 	/** Sets the constraint profile properties (limits, motors, etc...) to match the constraint profile as defined in the physics asset. If profile name is not found the joint is set to use the default constraint profile.*/
-	UFUNCTION(BlueprintCallable, Category = "Physics|Components|PhysicsConstraint")
+	UFUNCTION(BlueprintCallable, Category = "Physics")
 	void SetConstraintProfile(FName JointName, FName ProfileName, bool bDefaultIfNotFound = false);
 
 	/** Sets the constraint profile properties (limits, motors, etc...) to match the constraint profile as defined in the physics asset for all constraints. If profile name is not found the joint is set to use the default constraint profile.*/
-	UFUNCTION(BlueprintCallable, Category = "Physics|Components|PhysicsConstraint")
+	UFUNCTION(BlueprintCallable, Category = "Physics")
 	void SetConstraintProfileForAll(FName ProfileName, bool bDefaultIfNotFound = false);
 
 	/** Enable or Disable AngularPositionDrive based on a list of bone names */
@@ -2051,7 +2051,7 @@ public:
 	 * @param	ConstraintIndex	Index of constraint to look for
 	 * @return	Constraint Joint Name
 	 */
-	UFUNCTION(BlueprintCallable, Category = "Physics|Components|PhysicsConstraint")
+	UFUNCTION(BlueprintCallable, Category = "Physics")
 	FName	FindConstraintBoneName(int32 ConstraintIndex);
 
 	/** 
@@ -2082,21 +2082,21 @@ public:
 	 * @param	HitLocation	location of the hit
 	 * @param	InBoneName	Name of bone to break constraint for
 	 */
-	UFUNCTION(BlueprintCallable, Category = "Components|SkeletalMesh", meta = (Keywords = "Constraint"))
+	UFUNCTION(BlueprintCallable, Category = "Physics", meta = (Keywords = "Constraint"))
 	void BreakConstraint(FVector Impulse, FVector HitLocation, FName InBoneName);
 
 	/** Gets a constraint by its name 
 	* @param ConstraintName		name of the constraint
-	* @param IncludesTerminated wether or not to return a terminated constraint
+	* @param IncludesTerminated whether or not to return a terminated constraint
 	* */
-	UFUNCTION(BlueprintCallable, Category = "Components|SkeletalMesh", meta = (Keywords = "Components|SkeletalMesh"))
+	UFUNCTION(BlueprintCallable, Category = "Physics", meta = (Keywords = "Components|SkeletalMesh"))
 	FConstraintInstanceAccessor GetConstraintByName(FName ConstraintName, bool bIncludesTerminated);
 
 	/** Gets all the constraints
-	* @param IncludesTerminated wether or not to return terminated constraints
+	* @param IncludesTerminated whether or not to return terminated constraints
 	* @param OutConstraints returned list of constraints matching the parameters
 	* */
-	UFUNCTION(BlueprintCallable, Category = "Components|SkeletalMesh", meta = (Keywords = "Components|SkeletalMesh"))
+	UFUNCTION(BlueprintCallable, Category = "Physics", meta = (Keywords = "Components|SkeletalMesh"))
 	void GetConstraints(bool bIncludesTerminated, TArray<FConstraintInstanceAccessor>& OutConstraints);
 
 	/** Sets the Angular Motion Ranges for a named constraint
@@ -2105,7 +2105,7 @@ public:
 	*  @param TwistLimitAngle	 Size of limit in degrees, 0 means locked, 180 means free
 	*  @param Swing2LimitAngle	 Size of limit in degrees, 0 means locked, 180 means free
 	*/
-	UFUNCTION(BlueprintCallable, Category = "Components|SkeletalMesh")
+	UFUNCTION(BlueprintCallable, Category = "Physics")
 	void  SetAngularLimits(FName InBoneName,float Swing1LimitAngle, float TwistLimitAngle, float Swing2LimitAngle);
 
 	/** Gets the current Angular state for a named bone constraint 
@@ -2114,7 +2114,7 @@ public:
 	*  @param TwistAngle  current angular state of the constraint
 	*  @param Swing2Angle current angular state of the constraint
 	*/
-	UFUNCTION(BlueprintCallable, Category = "Components|SkeletalMesh")
+	UFUNCTION(BlueprintCallable, Category = "Physics")
 	void GetCurrentJointAngles(FName InBoneName,float& Swing1Angle, float& TwistAngle, float& Swing2Angle) ;
 
 
diff --git a/Engine/Source/Runtime/Engine/Classes/Components/SkyAtmosphereComponent.h b/Engine/Source/Runtime/Engine/Classes/Components/SkyAtmosphereComponent.h
--- a/Engine/Source/Runtime/Engine/Classes/Components/SkyAtmosphereComponent.h
+++ b/Engine/Source/Runtime/Engine/Classes/Components/SkyAtmosphereComponent.h
@@ -258,7 +258,7 @@ private:
  * A placeable actor that represents a planet atmosphere material and simulates sky and light scattering within it.
  * @see https://docs.unrealengine.com/en-US/Engine/Actors/FogEffects/SkyAtmosphere/index.html
  */
-UCLASS(showcategories = (Movement, Rendering, "Utilities|Transformation", "Input|MouseInput", "Input|TouchInput"), ClassGroup = Fog, hidecategories = (Info, Object, Input), MinimalAPI)
+UCLASS(showcategories = (Movement, Rendering, Transformation, "Input|MouseInput", "Input|TouchInput"), ClassGroup = Fog, hidecategories = (Info, Object, Input), MinimalAPI)
 class ASkyAtmosphere : public AInfo
 {
 	GENERATED_UCLASS_BODY()
diff --git a/Engine/Source/Runtime/Engine/Classes/Components/VolumetricCloudComponent.h b/Engine/Source/Runtime/Engine/Classes/Components/VolumetricCloudComponent.h
--- a/Engine/Source/Runtime/Engine/Classes/Components/VolumetricCloudComponent.h
+++ b/Engine/Source/Runtime/Engine/Classes/Components/VolumetricCloudComponent.h
@@ -170,7 +170,7 @@ private:
  * A placeable actor that represents a participating media material around a planet, e.g. clouds.
  * @see TODO address to the documentation.
  */
-UCLASS(showcategories = (Movement, Rendering, "Utilities|Transformation", "Input|MouseInput", "Input|TouchInput"), ClassGroup = Fog, hidecategories = (Info, Object, Input), MinimalAPI)
+UCLASS(showcategories = (Movement, Rendering, Transformation, "Input|MouseInput", "Input|TouchInput"), ClassGroup = Fog, hidecategories = (Info, Object, Input), MinimalAPI)
 class AVolumetricCloud : public AInfo
 {
 	GENERATED_UCLASS_BODY()
diff --git a/Engine/Source/Runtime/Engine/Classes/Engine/ExponentialHeightFog.h b/Engine/Source/Runtime/Engine/Classes/Engine/ExponentialHeightFog.h
--- a/Engine/Source/Runtime/Engine/Classes/Engine/ExponentialHeightFog.h
+++ b/Engine/Source/Runtime/Engine/Classes/Engine/ExponentialHeightFog.h
@@ -10,7 +10,7 @@
 /**
  * Implements an Actor for exponential height fog.
  */
-UCLASS(showcategories=(Movement, Rendering, "Utilities|Transformation"), ClassGroup=Fog, MinimalAPI)
+UCLASS(showcategories=(Movement, Rendering, Transformation), ClassGroup=Fog, MinimalAPI)
 class AExponentialHeightFog
 	: public AInfo
 {
diff --git a/Engine/Source/Runtime/Engine/Classes/Engine/LevelScriptActor.h b/Engine/Source/Runtime/Engine/Classes/Engine/LevelScriptActor.h
--- a/Engine/Source/Runtime/Engine/Classes/Engine/LevelScriptActor.h
+++ b/Engine/Source/Runtime/Engine/Classes/Engine/LevelScriptActor.h
@@ -21,7 +21,7 @@
  * @see https://docs.unrealengine.com/latest/INT/Engine/Blueprints/index.html
  * @see UBlueprint
  */
-UCLASS(notplaceable, meta=(ChildCanTick, KismetHideOverrides = "ReceiveAnyDamage,ReceivePointDamage,ReceiveRadialDamage,ReceiveActorBeginOverlap,ReceiveActorEndOverlap,ReceiveHit,ReceiveDestroyed,ReceiveActorBeginCursorOver,ReceiveActorEndCursorOver,ReceiveActorOnClicked,ReceiveActorOnReleased,ReceiveActorOnInputTouchBegin,ReceiveActorOnInputTouchEnd,ReceiveActorOnInputTouchEnter,ReceiveActorOnInputTouchLeave"), HideCategories=(Collision,Rendering,"Utilities|Transformation"))
+UCLASS(notplaceable, meta=(ChildCanTick, KismetHideOverrides = "ReceiveAnyDamage,ReceivePointDamage,ReceiveRadialDamage,ReceiveActorBeginOverlap,ReceiveActorEndOverlap,ReceiveHit,ReceiveDestroyed,ReceiveActorBeginCursorOver,ReceiveActorEndCursorOver,ReceiveActorOnClicked,ReceiveActorOnReleased,ReceiveActorOnInputTouchBegin,ReceiveActorOnInputTouchEnd,ReceiveActorOnInputTouchEnter,ReceiveActorOnInputTouchLeave"), HideCategories=(Collision,Rendering,Transformation))
 class ENGINE_API ALevelScriptActor : public AActor
 {
 	GENERATED_UCLASS_BODY()
diff --git a/Engine/Source/Runtime/Engine/Classes/Engine/WindDirectionalSource.h b/Engine/Source/Runtime/Engine/Classes/Engine/WindDirectionalSource.h
--- a/Engine/Source/Runtime/Engine/Classes/Engine/WindDirectionalSource.h
+++ b/Engine/Source/Runtime/Engine/Classes/Engine/WindDirectionalSource.h
@@ -7,7 +7,7 @@
 #include "WindDirectionalSource.generated.h"
 
 /** Actor that provides a directional wind source. Only affects SpeedTree assets. */
-UCLASS(ClassGroup=Wind, showcategories=(Rendering, "Utilities|Transformation"))
+UCLASS(ClassGroup=Wind, showcategories=(Rendering, Transformation))
 class ENGINE_API AWindDirectionalSource : public AInfo
 {
 	GENERATED_UCLASS_BODY()
diff --git a/Engine/Source/Runtime/Engine/Classes/GameFramework/Actor.h b/Engine/Source/Runtime/Engine/Classes/GameFramework/Actor.h
--- a/Engine/Source/Runtime/Engine/Classes/GameFramework/Actor.h
+++ b/Engine/Source/Runtime/Engine/Classes/GameFramework/Actor.h
@@ -251,7 +251,7 @@ public:
 #endif
 
 	/** Networking - Server - TearOff this actor to stop replication to clients. Will set bTearOff to true. */
-	UFUNCTION(BlueprintCallable, Category=Replication)
+	UFUNCTION(BlueprintCallable, Category=Networking)
 	virtual void TearOff();
 
 	/**
@@ -506,14 +506,14 @@ public:
 	 * @param bInReplicates Whether this Actor replicates to network clients.
 	 * @see https://docs.unrealengine.com/latest/INT/Gameplay/Networking/Replication/
 	 */
-	UFUNCTION(BlueprintCallable, BlueprintAuthorityOnly, Category=Replication)
+	UFUNCTION(BlueprintCallable, BlueprintAuthorityOnly, Category=Networking)
 	void SetReplicates(bool bInReplicates);
 
 	/**
 	 * Set whether this actor's movement replicates to network clients.
 	 * @param bInReplicateMovement Whether this Actor's movement replicates to clients.
 	 */
-	UFUNCTION(BlueprintCallable, Category=Replication)
+	UFUNCTION(BlueprintCallable, Category=Networking)
 	virtual void SetReplicateMovement(bool bInReplicateMovement);
 
 	/** Sets whether or not this Actor is an autonomous proxy, which is an actor on a network client that is controlled by a user on that client. */
@@ -523,11 +523,11 @@ public:
 	void CopyRemoteRoleFrom(const AActor* CopyFromActor);
 
 	/** Returns how much control the local machine has over this actor. */
-	UFUNCTION(BlueprintCallable, Category=Replication)
+	UFUNCTION(BlueprintCallable, Category=Networking)
 	ENetRole GetLocalRole() const { return Role; }
 
 	/** Returns how much control the remote machine has over this actor. */
-	UFUNCTION(BlueprintCallable, Category=Replication)
+	UFUNCTION(BlueprintCallable, Category=Networking)
 	ENetRole GetRemoteRole() const;
 
 	/**
@@ -725,7 +725,7 @@ public:
 
 protected:
 	/** The component that defines the transform (location, rotation, scale) of this Actor in the world, all other components must be attached to this one somehow */
-	UPROPERTY(BlueprintGetter=K2_GetRootComponent, Category="Utilities|Transformation")
+	UPROPERTY(BlueprintGetter=K2_GetRootComponent, Category="Transformation")
 	TObjectPtr<USceneComponent> RootComponent;
 
 #if WITH_EDITORONLY_DATA
@@ -1095,7 +1095,7 @@ public:
 	 * Get the actor-to-world transform.
 	 * @return The transform that transforms from actor space to world space.
 	 */
-	UFUNCTION(BlueprintCallable, meta=(DisplayName = "GetActorTransform", ScriptName = "GetActorTransform"), Category="Utilities|Transformation")
+	UFUNCTION(BlueprintCallable, meta=(DisplayName = "Get Actor Transform", ScriptName = "GetActorTransform"), Category="Transformation")
 	const FTransform& GetTransform() const
 	{
 		return ActorToWorld();
@@ -1108,7 +1108,7 @@ public:
 	}
 
 	/** Returns the location of the RootComponent of this Actor */
-	UFUNCTION(BlueprintCallable, meta=(DisplayName = "GetActorLocation", ScriptName = "GetActorLocation", Keywords="position"), Category="Utilities|Transformation")
+	UFUNCTION(BlueprintCallable, meta=(DisplayName = "Get Actor Location", ScriptName = "GetActorLocation", Keywords="position"), Category="Transformation")
 	FVector K2_GetActorLocation() const;
 
 	/** 
@@ -1123,23 +1123,23 @@ public:
 	 * @param SweepHitResult	The hit result from the move if swept.
 	 * @return	Whether the location was successfully set (if not swept), or whether movement occurred at all (if swept).
 	 */
-	UFUNCTION(BlueprintCallable, meta=(DisplayName = "SetActorLocation", ScriptName = "SetActorLocation", Keywords="position"), Category="Utilities|Transformation")
+	UFUNCTION(BlueprintCallable, meta=(DisplayName = "Set Actor Location", ScriptName = "SetActorLocation", Keywords="position"), Category="Transformation")
 	bool K2_SetActorLocation(FVector NewLocation, bool bSweep, FHitResult& SweepHitResult, bool bTeleport);
 
 	/** Returns rotation of the RootComponent of this Actor. */
-	UFUNCTION(BlueprintCallable, meta=(DisplayName = "GetActorRotation", ScriptName = "GetActorRotation"), Category="Utilities|Transformation")
+	UFUNCTION(BlueprintCallable, meta=(DisplayName = "Get Actor Rotation", ScriptName = "GetActorRotation"), Category="Transformation")
 	FRotator K2_GetActorRotation() const;
 
 	/** Get the forward (X) vector (length 1.0) from this Actor, in world space.  */
-	UFUNCTION(BlueprintCallable, Category = "Utilities|Transformation")
+	UFUNCTION(BlueprintCallable, Category = "Transformation")
 	FVector GetActorForwardVector() const;
 
 	/** Get the up (Z) vector (length 1.0) from this Actor, in world space.  */
-	UFUNCTION(BlueprintCallable, Category = "Utilities|Transformation")
+	UFUNCTION(BlueprintCallable, Category = "Transformation")
 	FVector GetActorUpVector() const;
 
 	/** Get the right (Y) vector (length 1.0) from this Actor, in world space.  */
-	UFUNCTION(BlueprintCallable, Category = "Utilities|Transformation")
+	UFUNCTION(BlueprintCallable, Category = "Transformation")
 	FVector GetActorRightVector() const;
 
 	/**
@@ -1149,7 +1149,7 @@ public:
 	 * @param	BoxExtent					Set to half the actor's size in 3d space
 	 * @param	bIncludeFromChildActors		If true then recurse in to ChildActor components 
 	 */
-	UFUNCTION(BlueprintCallable, Category="Collision", meta=(DisplayName = "GetActorBounds"))
+	UFUNCTION(BlueprintCallable, Category="Collision", meta=(DisplayName = "Get Actor Bounds"))
 	virtual void GetActorBounds(bool bOnlyCollidingComponents, FVector& Origin, FVector& BoxExtent, bool bIncludeFromChildActors = false) const;
 
 	/** Returns the RootComponent of this Actor */
@@ -1157,7 +1157,7 @@ public:
 	USceneComponent* K2_GetRootComponent() const;
 
 	/** Returns velocity (in cm/s (Unreal Units/second) of the rootcomponent if it is either using physics or has an associated MovementComponent */
-	UFUNCTION(BlueprintCallable, Category="Utilities|Transformation")
+	UFUNCTION(BlueprintCallable, Category="Transformation")
 	virtual FVector GetVelocity() const;
 
 	/** 
@@ -1184,7 +1184,7 @@ public:
 	 *			If false, physics velocity is updated based on the change in position (affecting ragdoll parts).
 	 * @return	Whether the rotation was successfully set.
 	 */
-	UFUNCTION(BlueprintCallable, meta=(DisplayName = "SetActorRotation", ScriptName = "SetActorRotation"), Category="Utilities|Transformation")
+	UFUNCTION(BlueprintCallable, meta=(DisplayName = "Set Actor Rotation", ScriptName = "SetActorRotation"), Category="Transformation")
 	bool K2_SetActorRotation(FRotator NewRotation, bool bTeleportPhysics);
 	
 	/**
@@ -1213,7 +1213,7 @@ public:
 	 * @param SweepHitResult	The hit result from the move if swept.
 	 * @return	Whether the rotation was successfully set.
 	 */
-	UFUNCTION(BlueprintCallable, Category="Utilities|Transformation", meta=(DisplayName="SetActorLocationAndRotation", ScriptName="SetActorLocationAndRotation"))
+	UFUNCTION(BlueprintCallable, Category="Transformation", meta=(DisplayName="Set Actor Location And Rotation", ScriptName="SetActorLocationAndRotation"))
 	bool K2_SetActorLocationAndRotation(FVector NewLocation, FRotator NewRotation, bool bSweep, FHitResult& SweepHitResult, bool bTeleport);
 	
 	/** 
@@ -1234,39 +1234,39 @@ public:
 	bool SetActorLocationAndRotation(FVector NewLocation, const FQuat& NewRotation, bool bSweep=false, FHitResult* OutSweepHitResult=nullptr, ETeleportType Teleport = ETeleportType::None);
 
 	/** Set the Actor's world-space scale. */
-	UFUNCTION(BlueprintCallable, Category="Utilities|Transformation")
+	UFUNCTION(BlueprintCallable, Category="Transformation")
 	void SetActorScale3D(FVector NewScale3D);
 
 	/** Returns the Actor's world-space scale. */
-	UFUNCTION(BlueprintCallable, Category="Utilities|Orientation")
+	UFUNCTION(BlueprintCallable, Category="Transformation")
 	FVector GetActorScale3D() const;
 
 	/** Returns the distance from this Actor to OtherActor. */
-	UFUNCTION(BlueprintCallable, Category = "Utilities|Transformation")
+	UFUNCTION(BlueprintCallable, Category = "Transformation")
 	float GetDistanceTo(const AActor* OtherActor) const;
 
 	/** Returns the squared distance from this Actor to OtherActor. */
-	UFUNCTION(BlueprintCallable, Category = "Utilities|Transformation")
+	UFUNCTION(BlueprintCallable, Category = "Transformation")
 	float GetSquaredDistanceTo(const AActor* OtherActor) const;
 
 	/** Returns the distance from this Actor to OtherActor, ignoring Z. */
-	UFUNCTION(BlueprintCallable, Category = "Utilities|Transformation")
+	UFUNCTION(BlueprintCallable, Category = "Transformation")
 	float GetHorizontalDistanceTo(const AActor* OtherActor) const;
 
 	/** Returns the squared distance from this Actor to OtherActor, ignoring Z. */
-	UFUNCTION(BlueprintCallable, Category = "Utilities|Transformation")
+	UFUNCTION(BlueprintCallable, Category = "Transformation")
 	float GetSquaredHorizontalDistanceTo(const AActor* OtherActor) const;
 
 	/** Returns the distance from this Actor to OtherActor, ignoring XY. */
-	UFUNCTION(BlueprintCallable, Category = "Utilities|Transformation")
+	UFUNCTION(BlueprintCallable, Category = "Transformation")
 	float GetVerticalDistanceTo(const AActor* OtherActor) const;
 
 	/** Returns the dot product from this Actor to OtherActor. Returns -2.0 on failure. Returns 0.0 for coincidental actors. */
-	UFUNCTION(BlueprintCallable, Category = "Utilities|Transformation")
+	UFUNCTION(BlueprintCallable, Category = "Transformation")
 	float GetDotProductTo(const AActor* OtherActor) const;
 
 	/** Returns the dot product from this Actor to OtherActor, ignoring Z. Returns -2.0 on failure. Returns 0.0 for coincidental actors. */
-	UFUNCTION(BlueprintCallable, Category = "Utilities|Transformation")
+	UFUNCTION(BlueprintCallable, Category = "Transformation")
 	float GetHorizontalDotProductTo(const AActor* OtherActor) const;
 
 	/**
@@ -1281,7 +1281,7 @@ public:
 	 *							If CCD is on and not teleporting, this will affect objects along the entire swept volume.
 	 * @param SweepHitResult	The hit result from the move if swept.
 	 */
-	UFUNCTION(BlueprintCallable, Category="Utilities|Transformation", meta=(DisplayName="AddActorWorldOffset", ScriptName="AddActorWorldOffset", Keywords="location position"))
+	UFUNCTION(BlueprintCallable, Category="Transformation", meta=(DisplayName="Add Actor World Offset", ScriptName="AddActorWorldOffset", Keywords="location position"))
 	void K2_AddActorWorldOffset(FVector DeltaLocation, bool bSweep, FHitResult& SweepHitResult, bool bTeleport);
 	void AddActorWorldOffset(FVector DeltaLocation, bool bSweep=false, FHitResult* OutSweepHitResult=nullptr, ETeleportType Teleport = ETeleportType::None);
 
@@ -1296,18 +1296,18 @@ public:
 	 *							If CCD is on and not teleporting, this will affect objects along the entire swept volume.
 	 * @param SweepHitResult	The hit result from the move if swept.
 	 */
-	UFUNCTION(BlueprintCallable, Category="Utilities|Transformation", meta=(DisplayName="AddActorWorldRotation", ScriptName="AddActorWorldRotation", AdvancedDisplay="bSweep,SweepHitResult,bTeleport"))
+	UFUNCTION(BlueprintCallable, Category="Transformation", meta=(DisplayName="Add Actor World Rotation", ScriptName="AddActorWorldRotation", AdvancedDisplay="bSweep,SweepHitResult,bTeleport"))
 	void K2_AddActorWorldRotation(FRotator DeltaRotation, bool bSweep, FHitResult& SweepHitResult, bool bTeleport);
 	void AddActorWorldRotation(FRotator DeltaRotation, bool bSweep=false, FHitResult* OutSweepHitResult=nullptr, ETeleportType Teleport = ETeleportType::None);
 	void AddActorWorldRotation(const FQuat& DeltaRotation, bool bSweep=false, FHitResult* OutSweepHitResult=nullptr, ETeleportType Teleport = ETeleportType::None);
 
 	/** Adds a delta to the transform of this actor in world space. Ignores scale and sets it to (1,1,1). */
-	UFUNCTION(BlueprintCallable, Category="Utilities|Transformation", meta=(DisplayName="AddActorWorldTransform", ScriptName="AddActorWorldTransform"))
+	UFUNCTION(BlueprintCallable, Category="Transformation", meta=(DisplayName="Add Actor World Transform", ScriptName="AddActorWorldTransform"))
 	void K2_AddActorWorldTransform(const FTransform& DeltaTransform, bool bSweep, FHitResult& SweepHitResult, bool bTeleport);
 	void AddActorWorldTransform(const FTransform& DeltaTransform, bool bSweep=false, FHitResult* OutSweepHitResult=nullptr, ETeleportType Teleport = ETeleportType::None);
 
 	/** Adds a delta to the transform of this actor in world space. Scale is unchanged. */
-	UFUNCTION(BlueprintCallable, Category = "Utilities|Transformation", meta = (DisplayName = "AddActorWorldTransformKeepScale", ScriptName = "AddActorWorldTransformKeepScale"))
+	UFUNCTION(BlueprintCallable, Category = "Transformation", meta = (DisplayName = "Add Actor World Transform Keep Scale", ScriptName = "AddActorWorldTransformKeepScale"))
 	void K2_AddActorWorldTransformKeepScale(const FTransform& DeltaTransform, bool bSweep, FHitResult& SweepHitResult, bool bTeleport);
 	void AddActorWorldTransformKeepScale(const FTransform& DeltaTransform, bool bSweep = false, FHitResult* OutSweepHitResult = nullptr, ETeleportType Teleport = ETeleportType::None);
 
@@ -1321,7 +1321,7 @@ public:
 	 *							If false, physics velocity is updated based on the change in position (affecting ragdoll parts).
 	 *							If CCD is on and not teleporting, this will affect objects along the entire swept volume.
 	 */
-	UFUNCTION(BlueprintCallable, Category="Utilities|Transformation", meta=(DisplayName="SetActorTransform", ScriptName="SetActorTransform"))
+	UFUNCTION(BlueprintCallable, Category="Transformation", meta=(DisplayName="Set Actor Transform", ScriptName="SetActorTransform"))
 	bool K2_SetActorTransform(const FTransform& NewTransform, bool bSweep, FHitResult& SweepHitResult, bool bTeleport);
 	bool SetActorTransform(const FTransform& NewTransform, bool bSweep=false, FHitResult* OutSweepHitResult=nullptr, ETeleportType Teleport = ETeleportType::None);
 
@@ -1335,7 +1335,7 @@ public:
 	 *							If false, physics velocity is updated based on the change in position (affecting ragdoll parts).
 	 *							If CCD is on and not teleporting, this will affect objects along the entire swept volume.
 	 */
-	UFUNCTION(BlueprintCallable, Category="Utilities|Transformation", meta=(DisplayName="AddActorLocalOffset", ScriptName="AddActorLocalOffset", Keywords="location position"))
+	UFUNCTION(BlueprintCallable, Category="Transformation", meta=(DisplayName="Add Actor Local Offset", ScriptName="AddActorLocalOffset", Keywords="location position"))
 	void K2_AddActorLocalOffset(FVector DeltaLocation, bool bSweep, FHitResult& SweepHitResult, bool bTeleport);
 	void AddActorLocalOffset(FVector DeltaLocation, bool bSweep=false, FHitResult* OutSweepHitResult=nullptr, ETeleportType Teleport = ETeleportType::None);
 
@@ -1349,7 +1349,7 @@ public:
 	 *							If false, physics velocity is updated based on the change in position (affecting ragdoll parts).
 	 *							If CCD is on and not teleporting, this will affect objects along the entire swept volume.
 	 */
-	UFUNCTION(BlueprintCallable, Category="Utilities|Transformation", meta=(DisplayName="AddActorLocalRotation", ScriptName="AddActorLocalRotation", AdvancedDisplay="bSweep,SweepHitResult,bTeleport"))
+	UFUNCTION(BlueprintCallable, Category="Transformation", meta=(DisplayName="Add Actor Local Rotation", ScriptName="AddActorLocalRotation", AdvancedDisplay="bSweep,SweepHitResult,bTeleport"))
 	void K2_AddActorLocalRotation(FRotator DeltaRotation, bool bSweep, FHitResult& SweepHitResult, bool bTeleport);
 	void AddActorLocalRotation(FRotator DeltaRotation, bool bSweep=false, FHitResult* OutSweepHitResult=nullptr, ETeleportType Teleport = ETeleportType::None);
 	void AddActorLocalRotation(const FQuat& DeltaRotation, bool bSweep=false, FHitResult* OutSweepHitResult=nullptr, ETeleportType Teleport = ETeleportType::None);
@@ -1364,7 +1364,7 @@ public:
 	 *							If false, physics velocity is updated based on the change in position (affecting ragdoll parts).
 	 *							If CCD is on and not teleporting, this will affect objects along the entire swept volume.
 	 */
-	UFUNCTION(BlueprintCallable, Category="Utilities|Transformation", meta=(DisplayName="AddActorLocalTransform", ScriptName="AddActorLocalTransform"))
+	UFUNCTION(BlueprintCallable, Category="Transformation", meta=(DisplayName="Add Actor Local Transform", ScriptName="AddActorLocalTransform"))
 	void K2_AddActorLocalTransform(const FTransform& NewTransform, bool bSweep, FHitResult& SweepHitResult, bool bTeleport);
 	void AddActorLocalTransform(const FTransform& NewTransform, bool bSweep=false, FHitResult* OutSweepHitResult=nullptr, ETeleportType Teleport = ETeleportType::None);
 
@@ -1378,7 +1378,7 @@ public:
 	 *								If false, physics velocity is updated based on the change in position (affecting ragdoll parts).
 	 *								If CCD is on and not teleporting, this will affect objects along the entire swept volume.
 	 */
-	UFUNCTION(BlueprintCallable, Category="Utilities|Transformation", meta=(DisplayName="SetActorRelativeLocation", ScriptName="SetActorRelativeLocation"))
+	UFUNCTION(BlueprintCallable, Category="Transformation", meta=(DisplayName="Set Actor Relative Location", ScriptName="SetActorRelativeLocation"))
 	void K2_SetActorRelativeLocation(FVector NewRelativeLocation, bool bSweep, FHitResult& SweepHitResult, bool bTeleport);
 	void SetActorRelativeLocation(FVector NewRelativeLocation, bool bSweep=false, FHitResult* OutSweepHitResult=nullptr, ETeleportType Teleport = ETeleportType::None);
 
@@ -1392,7 +1392,7 @@ public:
 	 *								If false, physics velocity is updated based on the change in position (affecting ragdoll parts).
 	 *								If CCD is on and not teleporting, this will affect objects along the entire swept volume.
 	 */
-	UFUNCTION(BlueprintCallable, Category="Utilities|Transformation", meta=(DisplayName="SetActorRelativeRotation", ScriptName="SetActorRelativeRotation", AdvancedDisplay="bSweep,SweepHitResult,bTeleport"))
+	UFUNCTION(BlueprintCallable, Category="Transformation", meta=(DisplayName="Set Actor Relative Rotation", ScriptName="SetActorRelativeRotation", AdvancedDisplay="bSweep,SweepHitResult,bTeleport"))
 	void K2_SetActorRelativeRotation(FRotator NewRelativeRotation, bool bSweep, FHitResult& SweepHitResult, bool bTeleport);
 	void SetActorRelativeRotation(FRotator NewRelativeRotation, bool bSweep=false, FHitResult* OutSweepHitResult=nullptr, ETeleportType Teleport = ETeleportType::None);
 	void SetActorRelativeRotation(const FQuat& NewRelativeRotation, bool bSweep=false, FHitResult* OutSweepHitResult=nullptr, ETeleportType Teleport = ETeleportType::None);
@@ -1407,7 +1407,7 @@ public:
 	 *							If false, physics velocity is updated based on the change in position (affecting ragdoll parts).
 	 *							If CCD is on and not teleporting, this will affect objects along the entire swept volume.
 	 */
-	UFUNCTION(BlueprintCallable, Category="Utilities|Transformation", meta=(DisplayName="SetActorRelativeTransform", ScriptName="SetActorRelativeTransform"))
+	UFUNCTION(BlueprintCallable, Category="Transformation", meta=(DisplayName="Set Actor Relative Transform", ScriptName="SetActorRelativeTransform"))
 	void K2_SetActorRelativeTransform(const FTransform& NewRelativeTransform, bool bSweep, FHitResult& SweepHitResult, bool bTeleport);
 	void SetActorRelativeTransform(const FTransform& NewRelativeTransform, bool bSweep=false, FHitResult* OutSweepHitResult=nullptr, ETeleportType Teleport = ETeleportType::None);
 
@@ -1415,11 +1415,11 @@ public:
 	 * Set the actor's RootComponent to the specified relative scale 3d
 	 * @param NewRelativeScale	New scale to set the actor's RootComponent to
 	 */
-	UFUNCTION(BlueprintCallable, Category="Utilities|Transformation")
+	UFUNCTION(BlueprintCallable, Category="Transformation")
 	void SetActorRelativeScale3D(FVector NewRelativeScale);
 
 	/** Return the actor's relative scale 3d */
-	UFUNCTION(BlueprintCallable, Category="Utilities|Orientation")
+	UFUNCTION(BlueprintCallable, Category="Transformation")
 	FVector GetActorRelativeScale3D() const;
 
 	/**
@@ -1438,7 +1438,7 @@ public:
 	bool GetActorEnableCollision() const;
 
 	/** Destroy the actor */
-	UFUNCTION(BlueprintCallable, Category="Utilities", meta=(Keywords = "Delete", DisplayName = "DestroyActor", ScriptName = "DestroyActor"))
+	UFUNCTION(BlueprintCallable, Category="Actor", meta=(Keywords = "Delete", DisplayName = "Destroy Actor", ScriptName = "DestroyActor"))
 	virtual void K2_DestroyActor();
 
 	/** Returns whether this actor has network authority */
@@ -1497,7 +1497,7 @@ public:
 
 	/** DEPRECATED - Use AttachToComponent() instead */
 	UE_DEPRECATED(4.17, "Use AttachToComponent() instead.")
-	UFUNCTION(BlueprintCallable, meta = (DisplayName = "AttachRootComponentTo (Deprecated)", ScriptNoExport, AttachLocationType = "KeepRelativeOffset"), Category = "Utilities|Transformation")
+	UFUNCTION(BlueprintCallable, meta = (DisplayName = "AttachRootComponentTo (Deprecated)", ScriptNoExport, AttachLocationType = "KeepRelativeOffset"), Category = "Transformation")
 	void K2_AttachRootComponentTo(USceneComponent* InParent, FName InSocketName = NAME_None, EAttachLocation::Type AttachLocationType = EAttachLocation::KeepRelativeOffset, bool bWeldSimulatedBodies = true);
 
 	/**
@@ -1509,7 +1509,7 @@ public:
 	 * @param ScaleRule					How to handle scale when attaching.
 	 * @param bWeldSimulatedBodies		Whether to weld together simulated physics bodies.
 	 */
-	UFUNCTION(BlueprintCallable, meta = (DisplayName = "AttachActorToComponent", ScriptName = "AttachToComponent", bWeldSimulatedBodies = true), Category = "Utilities|Transformation")
+	UFUNCTION(BlueprintCallable, meta = (DisplayName = "Attach Actor To Component", ScriptName = "AttachToComponent", bWeldSimulatedBodies = true), Category = "Transformation")
 	void K2_AttachToComponent(USceneComponent* Parent, FName SocketName, EAttachmentRule LocationRule, EAttachmentRule RotationRule, EAttachmentRule ScaleRule, bool bWeldSimulatedBodies);
 
 	/**
@@ -1522,7 +1522,7 @@ public:
 
 	/** DEPRECATED - Use AttachToActor() instead */
 	UE_DEPRECATED(4.17, "Use AttachToActor() instead.")
-	UFUNCTION(BlueprintCallable, meta = (DisplayName = "AttachRootComponentToActor (Deprecated)", ScriptNoExport, AttachLocationType = "KeepRelativeOffset"), Category = "Utilities|Transformation")
+	UFUNCTION(BlueprintCallable, meta = (DisplayName = "AttachRootComponentToActor (Deprecated)", ScriptNoExport, AttachLocationType = "KeepRelativeOffset"), Category = "Transformation")
 	void K2_AttachRootComponentToActor(AActor* InParentActor, FName InSocketName = NAME_None, EAttachLocation::Type AttachLocationType = EAttachLocation::KeepRelativeOffset, bool bWeldSimulatedBodies = true);
 
 	/**
@@ -1542,12 +1542,12 @@ public:
 	 * @param ScaleRule					How to handle scale when attaching.
 	 * @param bWeldSimulatedBodies		Whether to weld together simulated physics bodies.
 	 */
-	UFUNCTION(BlueprintCallable, meta = (DisplayName = "AttachActorToActor", ScriptName = "AttachToActor", bWeldSimulatedBodies=true), Category = "Utilities|Transformation")
+	UFUNCTION(BlueprintCallable, meta = (DisplayName = "Attach Actor To Actor", ScriptName = "AttachToActor", bWeldSimulatedBodies=true), Category = "Transformation")
 	void K2_AttachToActor(AActor* ParentActor, FName SocketName, EAttachmentRule LocationRule, EAttachmentRule RotationRule, EAttachmentRule ScaleRule, bool bWeldSimulatedBodies);
 
 	/** DEPRECATED - Use DetachFromActor() instead */
 	UE_DEPRECATED(4.17, "Use DetachFromActor() instead")
-	UFUNCTION(BlueprintCallable, meta=(DisplayName = "DetachActorFromActor (Deprecated)", ScriptNoExport), Category="Utilities|Transformation")
+	UFUNCTION(BlueprintCallable, meta=(DisplayName = "DetachActorFromActor (Deprecated)", ScriptNoExport), Category="Transformation")
 	void DetachRootComponentFromParent(bool bMaintainWorldPosition = true);
 
 	/** 
@@ -1556,7 +1556,7 @@ public:
 	 * @param  RotationRule				How to handle rotation when detaching.
 	 * @param  ScaleRule				How to handle scale when detaching.
 	 */
-	UFUNCTION(BlueprintCallable, meta=(DisplayName = "DetachFromActor", ScriptName = "DetachFromActor"), Category="Utilities|Transformation")
+	UFUNCTION(BlueprintCallable, meta=(DisplayName = "Detach From Actor", ScriptName = "DetachFromActor"), Category="Transformation")
 	void K2_DetachFromActor(EDetachmentRule LocationRule = EDetachmentRule::KeepRelative, EDetachmentRule RotationRule = EDetachmentRule::KeepRelative, EDetachmentRule ScaleRule = EDetachmentRule::KeepRelative);
 
 	/** 
@@ -1573,7 +1573,7 @@ public:
 	void DetachAllSceneComponents(class USceneComponent* InParentComponent, const FDetachmentTransformRules& DetachmentRules);
 
 	/** See if this actor's Tags array contains the supplied name tag */
-	UFUNCTION(BlueprintCallable, Category="Utilities")
+	UFUNCTION(BlueprintCallable, Category="Actor")
 	bool ActorHasTag(FName Tag) const;
 
 
@@ -1584,7 +1584,7 @@ public:
 	 * Get ActorTimeDilation - this can be used for input control or speed control for slomo.
 	 * We don't want to scale input globally because input can be used for UI, which do not care for TimeDilation.
 	 */
-	UFUNCTION(BlueprintCallable, Category="Utilities|Time")
+	UFUNCTION(BlueprintCallable, Category="Actor")
 	float GetActorTimeDilation() const;
 
 	/**
@@ -1593,27 +1593,27 @@ public:
 	float GetActorTimeDilation(const UWorld& ActorWorld) const;
 
 	/** Make this actor tick after PrerequisiteActor. This only applies to this actor's tick function; dependencies for owned components must be set up separately if desired. */
-	UFUNCTION(BlueprintCallable, Category="Utilities", meta=(Keywords = "dependency"))
+	UFUNCTION(BlueprintCallable, Category="Actor|Tick", meta=(Keywords = "dependency"))
 	virtual void AddTickPrerequisiteActor(AActor* PrerequisiteActor);
 
 	/** Make this actor tick after PrerequisiteComponent. This only applies to this actor's tick function; dependencies for owned components must be set up separately if desired. */
-	UFUNCTION(BlueprintCallable, Category="Utilities", meta=(Keywords = "dependency"))
+	UFUNCTION(BlueprintCallable, Category="Actor|Tick", meta=(Keywords = "dependency"))
 	virtual void AddTickPrerequisiteComponent(UActorComponent* PrerequisiteComponent);
 
 	/** Remove tick dependency on PrerequisiteActor. */
-	UFUNCTION(BlueprintCallable, Category="Utilities", meta=(Keywords = "dependency"))
+	UFUNCTION(BlueprintCallable, Category="Actor|Tick", meta=(Keywords = "dependency"))
 	virtual void RemoveTickPrerequisiteActor(AActor* PrerequisiteActor);
 
 	/** Remove tick dependency on PrerequisiteComponent. */
-	UFUNCTION(BlueprintCallable, Category="Utilities", meta=(Keywords = "dependency"))
+	UFUNCTION(BlueprintCallable, Category="Actor|Tick", meta=(Keywords = "dependency"))
 	virtual void RemoveTickPrerequisiteComponent(UActorComponent* PrerequisiteComponent);
 	
 	/** Gets whether this actor can tick when paused. */
-	UFUNCTION(BlueprintCallable, Category="Utilities")
+	UFUNCTION(BlueprintCallable, Category="Actor|Tick")
 	bool GetTickableWhenPaused();
 
 	/** Sets whether this actor can tick when paused. */
-	UFUNCTION(BlueprintCallable, Category="Utilities")
+	UFUNCTION(BlueprintCallable, Category="Actor|Tick")
 	void SetTickableWhenPaused(bool bTickableWhenPaused);
 
 	/** The number of seconds (in game time) since this Actor was created, relative to Get Game Time In Seconds. */
@@ -1794,11 +1794,11 @@ public:
 	void ReceiveHit(class UPrimitiveComponent* MyComp, AActor* Other, class UPrimitiveComponent* OtherComp, bool bSelfMoved, FVector HitLocation, FVector HitNormal, FVector NormalImpulse, const FHitResult& Hit);
 
 	/** Set the lifespan of this actor. When it expires the object will be destroyed. If requested lifespan is 0, the timer is cleared and the actor will not be destroyed. */
-	UFUNCTION(BlueprintCallable, Category="Utilities", meta=(Keywords = "delete destroy"))
+	UFUNCTION(BlueprintCallable, Category="Actor", meta=(Keywords = "delete destroy"))
 	virtual void SetLifeSpan( float InLifespan );
 
 	/** Get the remaining lifespan of this actor. If zero is returned the actor lives forever. */
-	UFUNCTION(BlueprintCallable, Category="Utilities", meta=(Keywords = "delete destroy"))
+	UFUNCTION(BlueprintCallable, Category="Actor", meta=(Keywords = "delete destroy"))
 	virtual float GetLifeSpan() const;
 
 	/**
@@ -2321,22 +2321,22 @@ public:
 	 * This only modifies the tick function on actor itself
 	 * @param	bEnabled	Whether it should be enabled or not
 	 */
-	UFUNCTION(BlueprintCallable, Category="Utilities")
+	UFUNCTION(BlueprintCallable, Category="Actor|Tick")
 	void SetActorTickEnabled(bool bEnabled);
 
 	/**  Returns whether this actor has tick enabled or not	 */
-	UFUNCTION(BlueprintCallable, Category="Utilities")
+	UFUNCTION(BlueprintCallable, Category="Actor|Tick")
 	bool IsActorTickEnabled() const;
 
 	/** 
 	 * Sets the tick interval of this actor's primary tick function. Will not enable a disabled tick function. Takes effect on next tick. 
 	 * @param TickInterval	The rate at which this actor should be ticking
 	 */
-	UFUNCTION(BlueprintCallable, Category="Utilities")
+	UFUNCTION(BlueprintCallable, Category="Actor|Tick")
 	void SetActorTickInterval(float TickInterval);
 
 	/**  Returns the tick interval of this actor's primary tick function */
-	UFUNCTION(BlueprintCallable, Category="Utilities")
+	UFUNCTION(BlueprintCallable, Category="Actor|Tick")
 	float GetActorTickInterval() const;
 
 	/**
@@ -2739,7 +2739,7 @@ public:
 	 * @param DestRotation The target rotation at the destination
 	 * @return true if the actor has been successfully moved, or false if it couldn't fit.
 	 */
-	UFUNCTION(BlueprintCallable, meta=( DisplayName="Teleport", ScriptName="Teleport", Keywords = "Move Position" ), Category="Utilities|Transformation")
+	UFUNCTION(BlueprintCallable, meta=( DisplayName="Teleport", ScriptName="Teleport", Keywords = "Move Position" ), Category="Transformation")
 	bool K2_TeleportTo( FVector DestLocation, FRotator DestRotation );
 
 	/** Called from TeleportTo() when teleport succeeds */
@@ -2880,25 +2880,25 @@ public:
 	void CheckComponentInstanceName(const FName InName);
 
 	/** Walk up the attachment chain from RootComponent until we encounter a different actor, and return it. If we are not attached to a component in a different actor, returns nullptr */
-	UFUNCTION(BlueprintPure, Category = "Utilities")
+	UFUNCTION(BlueprintPure, Category = "Actor")
 	AActor* GetAttachParentActor() const;
 
 	/** Walk up the attachment chain from RootComponent until we encounter a different actor, and return the socket name in the component. If we are not attached to a component in a different actor, returns NAME_None */
-	UFUNCTION(BlueprintPure, Category = "Utilities")
+	UFUNCTION(BlueprintPure, Category = "Actor")
 	FName GetAttachParentSocketName() const;
 
 	/** Call a functor for Actors which are attached directly to a component in this actor. Functor should return true to carry on, false to abort. */
 	void ForEachAttachedActors(TFunctionRef<bool(class AActor*)> Functor) const;
 	
 	/** Find all Actors which are attached directly to a component in this actor */
-	UFUNCTION(BlueprintPure, Category = "Utilities")
+	UFUNCTION(BlueprintPure, Category = "Actor")
 	void GetAttachedActors(TArray<AActor*>& OutActors, bool bResetArray = true) const;
 
 	/**
 	 * Sets the ticking group for this actor.
 	 * @param NewTickGroup the new value to assign
 	 */
-	UFUNCTION(BlueprintCallable, Category="Utilities", meta=(Keywords = "dependency"))
+	UFUNCTION(BlueprintCallable, Category="Actor|Tick", meta=(Keywords = "dependency"))
 	void SetTickGroup(ETickingGroup NewTickGroup);
 
 	/** Called when this actor is explicitly being destroyed during gameplay or in the editor, not called during level streaming or gameplay ending */
@@ -3045,7 +3045,7 @@ public:
 	 * @param Tolerance  How many seconds ago the actor last render time can be and still count as having been "recently" rendered.
 	 * @return Whether this actor was recently rendered.
 	 */
-	UFUNCTION(Category="Rendering", BlueprintCallable, meta=(DisplayName="WasActorRecentlyRendered", Keywords="scene visible"))
+	UFUNCTION(Category="Rendering", BlueprintCallable, meta=(DisplayName="Was Actor Recently Rendered", Keywords="scene visible"))
 	bool WasRecentlyRendered(float Tolerance = 0.2) const;
 
 	/** Returns the most recent time any of this actor's components were rendered */
@@ -3161,7 +3161,7 @@ public:
 	 * Currently returns an array of UActorComponent which must be cast to the correct type.
 	 * This intended to only be used by blueprints. Use GetComponents() in C++.
 	 */
-	UFUNCTION(BlueprintCallable, Category = "Actor", meta = (ComponentClass = "ActorComponent", DisplayName = "GetComponentsByClass", ScriptName = "GetComponentsByClass", DeterminesOutputType = "ComponentClass"))
+	UFUNCTION(BlueprintCallable, Category = "Actor", meta = (ComponentClass = "ActorComponent", DisplayName = "Get Components By Class", ScriptName = "GetComponentsByClass", DeterminesOutputType = "ComponentClass"))
 	TArray<UActorComponent*> K2_GetComponentsByClass(TSubclassOf<UActorComponent> ComponentClass) const;
 
 	/** Gets all the components that inherit from the given class with a given tag. */
diff --git a/Engine/Source/Runtime/Engine/Classes/GameFramework/Controller.h b/Engine/Source/Runtime/Engine/Classes/GameFramework/Controller.h
--- a/Engine/Source/Runtime/Engine/Classes/GameFramework/Controller.h
+++ b/Engine/Source/Runtime/Engine/Classes/GameFramework/Controller.h
@@ -36,7 +36,7 @@ DECLARE_MULTICAST_DELEGATE_OneParam(FPawnChangedSignature, APawn* /*NewPawn*/);
  * 
  * @see https://docs.unrealengine.com/latest/INT/Gameplay/Framework/Controller/
  */
-UCLASS(abstract, notplaceable, NotBlueprintable, HideCategories=(Collision,Rendering,"Utilities|Transformation")) 
+UCLASS(abstract, notplaceable, NotBlueprintable, HideCategories=(Collision,Rendering,Transformation)) 
 class ENGINE_API AController : public AActor, public INavAgentInterface
 {
 	GENERATED_BODY()
diff --git a/Engine/Source/Runtime/Engine/Classes/GameFramework/Info.h b/Engine/Source/Runtime/Engine/Classes/GameFramework/Info.h
--- a/Engine/Source/Runtime/Engine/Classes/GameFramework/Info.h
+++ b/Engine/Source/Runtime/Engine/Classes/GameFramework/Info.h
@@ -16,7 +16,7 @@
  * Info is the base class of an Actor that isn't meant to have a physical representation in the world, used primarily
  * for "manager" type classes that hold settings data about the world, but might need to be an Actor for replication purposes.
  */
-UCLASS(abstract, hidecategories=(Input, Movement, Collision, Rendering, LOD, WorldPartition, DataLayers, "Utilities|Transformation"), showcategories=("Input|MouseInput", "Input|TouchInput"), MinimalAPI, NotBlueprintable)
+UCLASS(abstract, hidecategories=(Input, Movement, Collision, Rendering, LOD, WorldPartition, DataLayers, Transformation), showcategories=("Input|MouseInput", "Input|TouchInput"), MinimalAPI, NotBlueprintable)
 class AInfo : public AActor
 {
 	GENERATED_UCLASS_BODY()
diff --git a/Engine/Source/Runtime/Engine/Classes/GameFramework/PlayerController.h b/Engine/Source/Runtime/Engine/Classes/GameFramework/PlayerController.h
--- a/Engine/Source/Runtime/Engine/Classes/GameFramework/PlayerController.h
+++ b/Engine/Source/Runtime/Engine/Classes/GameFramework/PlayerController.h
@@ -628,18 +628,18 @@ public:
 	bool GetHitResultUnderFingerForObjects(ETouchIndex::Type FingerIndex, const  TArray<TEnumAsByte<EObjectTypeQuery> > & ObjectTypes, bool bTraceComplex, FHitResult& HitResult) const;
 
 	/** Convert current mouse 2D position to World Space 3D position and direction. Returns false if unable to determine value. **/
-	UFUNCTION(BlueprintCallable, Category="Game|Player", meta = (DisplayName = "ConvertMouseLocationToWorldSpace", Keywords = "deproject"))
+	UFUNCTION(BlueprintCallable, Category="Game|Player", meta = (DisplayName = "Convert Mouse Location To World Space", Keywords = "deproject"))
 	bool DeprojectMousePositionToWorld(FVector& WorldLocation, FVector& WorldDirection) const;
 
 	/** Convert 2D screen position to World Space 3D position and direction. Returns false if unable to determine value. **/
-	UFUNCTION(BlueprintCallable, Category="Game|Player", meta = (DisplayName = "ConvertScreenLocationToWorldSpace", Keywords = "deproject"))
+	UFUNCTION(BlueprintCallable, Category="Game|Player", meta = (DisplayName = "Convert Screen Location To World Space", Keywords = "deproject"))
 	bool DeprojectScreenPositionToWorld(float ScreenX, float ScreenY, FVector& WorldLocation, FVector& WorldDirection) const;
 
 	/**
 	 * Convert a World Space 3D position into a 2D Screen Space position.
 	 * @return true if the world coordinate was successfully projected to the screen.
 	 */
-	UFUNCTION(BlueprintCallable, Category="Game|Player", meta = (DisplayName = "ConvertWorldLocationToScreenLocation", Keywords = "project"))
+	UFUNCTION(BlueprintCallable, Category="Game|Player", meta = (DisplayName = "Convert World Location To Screen Location", Keywords = "project"))
 	bool ProjectWorldLocationToScreen(FVector WorldLocation, FVector2D& ScreenLocation, bool bPlayerViewportRelative = false) const;
 
 	/**
@@ -655,7 +655,7 @@ public:
 	virtual bool PostProcessWorldToScreen(FVector WorldLocation, FVector2D& ScreenLocation, bool bPlayerViewportRelative) const;
 
 	/** Positions the mouse cursor in screen space, in pixels. */
-	UFUNCTION( BlueprintCallable, Category="Game|Player", meta = (DisplayName = "SetMousePosition", Keywords = "mouse" ))
+	UFUNCTION( BlueprintCallable, Category="Game|Player", meta = (DisplayName = "Set Mouse Position", Keywords = "mouse" ))
 	void SetMouseLocation( const int X, const int Y );
 
 	/**
@@ -1035,7 +1035,7 @@ public:
 	 * @param Cursor - the cursor to set the widget for
 	 * @param CursorWidget - the widget to set the cursor to
 	 */
-	UFUNCTION(BlueprintCallable, Category="User Interface")
+	UFUNCTION(BlueprintCallable, Category="Game|Player")
 	void SetMouseCursorWidget(EMouseCursor::Type Cursor, class UUserWidget* CursorWidget);
 
 	/** Set the view target
@@ -1783,7 +1783,7 @@ public:
 	virtual void ViewAPlayer(int32 dir);
 
 	/** Returns true if this controller thinks it's able to restart. Called from GameModeBase::PlayerCanRestart */
-	UFUNCTION(BlueprintCallable, Category = "Game")
+	UFUNCTION(BlueprintCallable, Category = "Game|Player")
 	virtual bool CanRestartPlayer();
 
 	/**
diff --git a/Engine/Source/Runtime/Engine/Classes/Kismet/GameplayStatics.h b/Engine/Source/Runtime/Engine/Classes/Kismet/GameplayStatics.h
--- a/Engine/Source/Runtime/Engine/Classes/Kismet/GameplayStatics.h
+++ b/Engine/Source/Runtime/Engine/Classes/Kismet/GameplayStatics.h
@@ -70,7 +70,7 @@ class ENGINE_API UGameplayStatics : public UBlueprintFunctionLibrary
 	// --- Actor functions ------------------------------
 
 	/** Find the average location (centroid) of an array of Actors */
-	UFUNCTION(BlueprintCallable, Category="Utilities|Transformation")
+	UFUNCTION(BlueprintCallable, Category="Transformation")
 	static FVector GetActorArrayAverageLocation(const TArray<AActor*>& Actors);
 
 	/** Bind the bounds of an array of Actors */
@@ -83,7 +83,7 @@ class ENGINE_API UGameplayStatics : public UBlueprintFunctionLibrary
 	 *	@param	ActorClass	Class of Actor to find. Must be specified or result will be empty.
 	 *	@return				Actor of the specified class.
 	 */
-	UFUNCTION(BlueprintCallable, Category="Utilities", meta=(WorldContext="WorldContextObject", DeterminesOutputType="ActorClass"))
+	UFUNCTION(BlueprintCallable, Category="Actor", meta=(WorldContext="WorldContextObject", DeterminesOutputType="ActorClass"))
 	static class AActor* GetActorOfClass(const UObject* WorldContextObject, TSubclassOf<AActor> ActorClass);
 
 	/** 
@@ -92,7 +92,7 @@ class ENGINE_API UGameplayStatics : public UBlueprintFunctionLibrary
 	 *	@param	ActorClass	Class of Actor to find. Must be specified or result array will be empty.
 	 *	@param	OutActors	Output array of Actors of the specified class.
 	 */
-	UFUNCTION(BlueprintCallable, Category="Utilities",  meta=(WorldContext="WorldContextObject", DeterminesOutputType="ActorClass", DynamicOutputParam="OutActors"))
+	UFUNCTION(BlueprintCallable, Category="Actor",  meta=(WorldContext="WorldContextObject", DeterminesOutputType="ActorClass", DynamicOutputParam="OutActors"))
 	static void GetAllActorsOfClass(const UObject* WorldContextObject, TSubclassOf<AActor> ActorClass, TArray<AActor*>& OutActors);
 
 	/** 
@@ -101,7 +101,7 @@ class ENGINE_API UGameplayStatics : public UBlueprintFunctionLibrary
 	 *	@param	Interface	Interface to find. Must be specified or result array will be empty.
 	 *	@param	OutActors	Output array of Actors of the specified interface.
 	 */
-	UFUNCTION(BlueprintCallable, Category="Utilities",  meta=(WorldContext="WorldContextObject", DeterminesOutputType="Interface", DynamicOutputParam="OutActors"))
+	UFUNCTION(BlueprintCallable, Category="Actor",  meta=(WorldContext="WorldContextObject", DeterminesOutputType="Interface", DynamicOutputParam="OutActors"))
 	static void GetAllActorsWithInterface(const UObject* WorldContextObject, TSubclassOf<UInterface> Interface, TArray<AActor*>& OutActors);
 
 	/**
@@ -110,7 +110,7 @@ class ENGINE_API UGameplayStatics : public UBlueprintFunctionLibrary
 	 *	@param	Tag			Tag to find. Must be specified or result array will be empty.
 	 *	@param	OutActors	Output array of Actors of the specified tag.
 	 */
-	UFUNCTION(BlueprintCallable, Category="Utilities",  meta=(WorldContext="WorldContextObject"))
+	UFUNCTION(BlueprintCallable, Category="Actor",  meta=(WorldContext="WorldContextObject"))
 	static void GetAllActorsWithTag(const UObject* WorldContextObject, FName Tag, TArray<AActor*>& OutActors);
 
 	/**
@@ -120,7 +120,7 @@ class ENGINE_API UGameplayStatics : public UBlueprintFunctionLibrary
 	 *	@param	ActorClass	Class of Actor to find. Must be specified or result array will be empty.
 	 *	@param	OutActors	Output array of Actors of the specified tag.
 	 */
-	UFUNCTION(BlueprintCallable, Category = "Utilities", meta = (WorldContext="WorldContextObject", DeterminesOutputType="ActorClass", DynamicOutputParam="OutActors"))
+	UFUNCTION(BlueprintCallable, Category = "Actor", meta = (WorldContext="WorldContextObject", DeterminesOutputType="ActorClass", DynamicOutputParam="OutActors"))
 	static void GetAllActorsOfClassWithTag(const UObject* WorldContextObject, TSubclassOf<AActor> ActorClass, FName Tag, TArray<AActor*>& OutActors);
 
 	/**
@@ -130,7 +130,7 @@ class ENGINE_API UGameplayStatics : public UBlueprintFunctionLibrary
 	 *	@param	Distance	Distance from Origin to the returned Actor.
 	 *	@return				Nearest Actor.
 	 */
-	UFUNCTION(BlueprintPure, Category = "Utilities")
+	UFUNCTION(BlueprintPure, Category = "Actor")
 	static AActor* FindNearestActor(FVector Origin, const TArray<AActor*>& ActorsToCheck, float& Distance);
 
 	// --- Player functions ------------------------------
@@ -260,7 +260,7 @@ class ENGINE_API UGameplayStatics : public UBlueprintFunctionLibrary
 	static class AGameStateBase* GetGameState(const UObject* WorldContextObject);
 
 	/** Returns the class of a passed in Object, will always be valid if Object is not NULL */
-	UFUNCTION(BlueprintPure, meta=(DisplayName = "GetClass", DeterminesOutputType = "Object"), Category="Utilities")
+	UFUNCTION(BlueprintPure, meta=(DisplayName = "Get Class", DeterminesOutputType = "Object"), Category="Utilities")
 	static class UClass *GetObjectClass(const UObject *Object);
 
 	/**
@@ -323,13 +323,13 @@ class ENGINE_API UGameplayStatics : public UBlueprintFunctionLibrary
 	/**
 	 * Returns the current viewport mouse capture mode
 	 */
-	UFUNCTION(BlueprintPure, Category = "Utilities", meta = (WorldContext = "WorldContextObject"))
+	UFUNCTION(BlueprintPure, Category = "Viewport", meta = (WorldContext = "WorldContextObject"))
 	static EMouseCaptureMode GetViewportMouseCaptureMode(const UObject* WorldContextObject);
 
 	/**
 	 * Sets the current viewport mouse capture mode
 	 */
-	UFUNCTION(BlueprintCallable, Category = "Utilities", meta = (WorldContext = "WorldContextObject"))
+	UFUNCTION(BlueprintCallable, Category = "Viewport", meta = (WorldContext = "WorldContextObject"))
 	static void SetViewportMouseCaptureMode(const UObject* WorldContextObject, const EMouseCaptureMode MouseCaptureMode);
 
 	/** Hurt locally authoritative actors within the radius. Will only hit components that block the Visibility channel.
@@ -400,7 +400,7 @@ class ENGINE_API UGameplayStatics : public UBlueprintFunctionLibrary
 	 * @param Falloff - Affects falloff of effect as it nears OuterRadius
 	 * @param bOrientShakeTowardsEpicenter - Changes the rotation of shake to point towards epicenter instead of forward
 	 */
-	UFUNCTION(BlueprintCallable, Category="Game|Feedback", meta=(WorldContext="WorldContextObject", UnsafeDuringActorConstruction = "true"))
+	UFUNCTION(BlueprintCallable, Category="Camera", meta=(WorldContext="WorldContextObject", UnsafeDuringActorConstruction = "true"))
 	static void PlayWorldCameraShake(const UObject* WorldContextObject, TSubclassOf<class UCameraShakeBase> Shake, FVector Epicenter, float InnerRadius, float OuterRadius, float Falloff = 1.f, bool bOrientShakeTowardsEpicenter = false);
 
 	// --- Particle functions ------------------------------
@@ -415,7 +415,7 @@ class ENGINE_API UGameplayStatics : public UBlueprintFunctionLibrary
 	 * @param PoolingMethod - Method used for pooling this component. Defaults to none.
 	 * @param bAutoActivate - Whether the component will be automatically activated on creation.
 	 */
-	UFUNCTION(BlueprintCallable, Category="Effects|Components|ParticleSystem", meta=(Keywords = "particle system", WorldContext="WorldContextObject", UnsafeDuringActorConstruction = "true"))
+	UFUNCTION(BlueprintCallable, Category="Effects", meta=(Keywords = "particle system", WorldContext="WorldContextObject", UnsafeDuringActorConstruction = "true"))
 	static UParticleSystemComponent* SpawnEmitterAtLocation(const UObject* WorldContextObject, UParticleSystem* EmitterTemplate, FVector Location, FRotator Rotation = FRotator::ZeroRotator, FVector Scale = FVector(1.f), bool bAutoDestroy = true, EPSCPoolMethod PoolingMethod = EPSCPoolMethod::None, bool bAutoActivateSystem = true);
 
 	// Backwards compatible version of SpawnEmitterAttached for C++ without Scale
@@ -448,7 +448,7 @@ public:
 	 * @param PoolingMethod - Method used for pooling this component. Defaults to none.
 	 * @param bAutoActivate - Whether the component will be automatically activated on creation.
 	 */
-	UFUNCTION(BlueprintCallable, Category="Effects|Components|ParticleSystem", meta=(Keywords = "particle system", UnsafeDuringActorConstruction = "true"))
+	UFUNCTION(BlueprintCallable, Category="Effects", meta=(Keywords = "particle system", UnsafeDuringActorConstruction = "true"))
 	static UParticleSystemComponent* SpawnEmitterAttached(class UParticleSystem* EmitterTemplate, class USceneComponent* AttachToComponent, FName AttachPointName = NAME_None, FVector Location = FVector(ForceInit), FRotator Rotation = FRotator::ZeroRotator, FVector Scale = FVector(1.f), EAttachLocation::Type LocationType = EAttachLocation::KeepRelativeOffset, bool bAutoDestroy = true, EPSCPoolMethod PoolingMethod = EPSCPoolMethod::None, bool bAutoActivate=true);
 
 	// Backwards compatible version of SpawnEmitterAttached for C++ without Scale
@@ -891,7 +891,7 @@ public:
 	 * @param Rotation - rotation to place the decal in world space	
 	 * @param LifeSpan - destroy decal component after time runs out (0 = infinite)
 	 */
-	UFUNCTION(BlueprintCallable, BlueprintCosmetic, Category="Rendering|Components|Decal", meta=(WorldContext="WorldContextObject", UnsafeDuringActorConstruction = "true"))
+	UFUNCTION(BlueprintCallable, BlueprintCosmetic, Category="Rendering|Decal", meta=(WorldContext="WorldContextObject", UnsafeDuringActorConstruction = "true"))
 	static UDecalComponent* SpawnDecalAtLocation(const UObject* WorldContextObject, class UMaterialInterface* DecalMaterial, FVector DecalSize, FVector Location, FRotator Rotation = FRotator(-90, 0, 0), float LifeSpan = 0);
 
 	/** Spawns a decal attached to and following the specified component. Does not replicate.
@@ -904,7 +904,7 @@ public:
 	 * @param LocationType - Specifies whether Location is a relative offset or an absolute world position
 	 * @param LifeSpan - destroy decal component after time runs out (0 = infinite)
 	 */
-	UFUNCTION(BlueprintCallable, BlueprintCosmetic, Category="Rendering|Components|Decal", meta=(UnsafeDuringActorConstruction = "true"))
+	UFUNCTION(BlueprintCallable, BlueprintCosmetic, Category="Rendering|Decal", meta=(UnsafeDuringActorConstruction = "true"))
 	static UDecalComponent* SpawnDecalAttached(class UMaterialInterface* DecalMaterial, FVector DecalSize, class USceneComponent* AttachToComponent, FName AttachPointName = NAME_None, FVector Location = FVector(ForceInit), FRotator Rotation = FRotator::ZeroRotator, EAttachLocation::Type LocationType = EAttachLocation::KeepRelativeOffset, float LifeSpan = 0);
 
 	/** Extracts data from a HitResult. 
@@ -1134,7 +1134,7 @@ public:
 	 * @param bDrawDebug		When true, a debug arc is drawn (red for an invalid arc, green for a valid arc)
 	 * @return					Returns false if there is no valid solution or the valid solutions are blocked.  Returns true otherwise.
 	 */
-	UFUNCTION(BlueprintCallable, Category="Game|Components|ProjectileMovement", DisplayName="SuggestProjectileVelocity", meta=(WorldContext="WorldContextObject"))
+	UFUNCTION(BlueprintCallable, Category="Game", DisplayName="Suggest Projectile Velocity", meta=(WorldContext="WorldContextObject"))
 	static bool BlueprintSuggestProjectileVelocity(const UObject* WorldContextObject, FVector& TossVelocity, FVector StartLocation, FVector EndLocation, float LaunchSpeed, float OverrideGravityZ, ESuggestProjVelocityTraceOption::Type TraceOption, float CollisionRadius, bool bFavorHighArc, bool bDrawDebug);
 
 	/** Native version, has more options than the Blueprint version. */
@@ -1220,7 +1220,7 @@ public:
 	* @param OverrideGravityZ			Optional override of WorldGravityZ
 	* @param ArcParam					Change height of arc between 0.0-1.0 where 0.5 is the default medium arc, 0 is up, and 1 is directly toward EndPos.
 	*/
-	UFUNCTION(BlueprintCallable, Category = "Game", DisplayName = "SuggestProjectileVelocity Custom Arc", meta = (WorldContext = "WorldContextObject", AdvancedDisplay = "OverrideGravityZ, ArcParam"))
+	UFUNCTION(BlueprintCallable, Category = "Game", DisplayName = "Suggest Projectile Velocity Custom Arc", meta = (WorldContext = "WorldContextObject", AdvancedDisplay = "OverrideGravityZ, ArcParam"))
 	static bool SuggestProjectileVelocity_CustomArc(const UObject* WorldContextObject, FVector& OutLaunchVelocity, FVector StartPos, FVector EndPos, float OverrideGravityZ = 0, float ArcParam = 0.5f);
 
 	/** Returns world origin current location. */
@@ -1258,7 +1258,7 @@ public:
 	 * @param WorldPosition		(out) Corresponding 3D position in world space.
 	 * @param WorldDirection	(out) World space direction vector away from the camera at the given 2d point.
 	 */
-	UFUNCTION(BlueprintPure, Category = "Utilities", meta = (Keywords = "unproject"))
+	UFUNCTION(BlueprintPure, Category = "Camera", meta = (Keywords = "unproject"))
 	static bool DeprojectScreenToWorld(APlayerController const* Player, const FVector2D& ScreenPosition, FVector& WorldPosition, FVector& WorldDirection);
 
 	/** 
@@ -1268,7 +1268,7 @@ public:
 	 * @param ScreenPosition	(out) Corresponding 2D position in screen space
 	 * @param bPlayerViewportRelative	Should this be relative to the player viewport subregion (useful when using player attached widgets in split screen)
 	 */
-	UFUNCTION(BlueprintPure, Category = "Utilities")
+	UFUNCTION(BlueprintPure, Category = "Camera")
 	static bool ProjectWorldToScreen(APlayerController const* Player, const FVector& WorldPosition, FVector2D& ScreenPosition, bool bPlayerViewportRelative = false);
 
 	/**
@@ -1278,7 +1278,7 @@ public:
 	 * @param ProjectionMatrix		(out) Corresponding Projection Matrix
 	 * @param ViewProjectionMatrix	(out) Corresponding View x Projection Matrix
 	 */
-	UFUNCTION(BlueprintPure, Category = "Utilities")
+	UFUNCTION(BlueprintPure, Category = "Camera")
 	static void GetViewProjectionMatrix(FMinimalViewInfo DesiredView, FMatrix &ViewMatrix, FMatrix &ProjectionMatrix, FMatrix &ViewProjectionMatrix);
 
 	/**
@@ -1351,7 +1351,7 @@ public:
 	 * Try to make announcements concise and clear.
 	 * NOTE: Currently only supported on Win10, Mac, iOS
 	*/
-	UFUNCTION(BlueprintCallable, Category = "Utilities")
+	UFUNCTION(BlueprintCallable, Category = "Accessibility")
 	static void AnnounceAccessibleString(const FString& AnnouncementString);
 };
 
diff --git a/Engine/Source/Runtime/Engine/Classes/Kismet/KismetGuidLibrary.h b/Engine/Source/Runtime/Engine/Classes/Kismet/KismetGuidLibrary.h
--- a/Engine/Source/Runtime/Engine/Classes/Kismet/KismetGuidLibrary.h
+++ b/Engine/Source/Runtime/Engine/Classes/Kismet/KismetGuidLibrary.h
@@ -20,15 +20,15 @@ class ENGINE_API UKismetGuidLibrary
 	static bool EqualEqual_GuidGuid( const FGuid& A, const FGuid& B );
 	
 	/** Returns true if the values are not equal (A != B) */
-	UFUNCTION(BlueprintPure, meta=(DisplayName="NotEqual (Guid)", CompactNodeTitle="!=", BlueprintThreadSafe), Category="Guid")
+	UFUNCTION(BlueprintPure, meta=(DisplayName="Not Equal (Guid)", CompactNodeTitle="!=", BlueprintThreadSafe), Category="Guid")
 	static bool NotEqual_GuidGuid( const FGuid& A, const FGuid& B );
 
 	/** Checks whether the given GUID is valid */
-	UFUNCTION(BlueprintPure, meta=(DisplayName="IsValid (Guid)", CompactNodeTitle="isValid?", BlueprintThreadSafe), Category="Guid")
+	UFUNCTION(BlueprintPure, meta=(DisplayName="Is Valid Guid", CompactNodeTitle="isValid?", BlueprintThreadSafe), Category="Guid")
 	static bool IsValid_Guid( const FGuid& InGuid );
 
 	/** Invalidates the given GUID */
-	UFUNCTION(BlueprintCallable, meta=(DisplayName="Invalidate_Guid (Guid)", BlueprintThreadSafe), Category="Guid")
+	UFUNCTION(BlueprintCallable, meta=(DisplayName="Invalidate Guid", BlueprintThreadSafe), Category="Guid")
 	static void Invalidate_Guid( UPARAM(ref) FGuid& InGuid );
 
 	/** Returns a new unique GUID */
@@ -36,7 +36,7 @@ class ENGINE_API UKismetGuidLibrary
 	static FGuid NewGuid();
 	
 	/** Converts a GUID value to a string, in the form 'A-B-C-D' */
-	UFUNCTION(BlueprintPure, meta=(DisplayName="ToString (Guid)", CompactNodeTitle="->", ScriptMethod="ToString", BlueprintThreadSafe), Category="Guid")
+	UFUNCTION(BlueprintPure, meta=(DisplayName="To String (Guid)", CompactNodeTitle="->", ScriptMethod="ToString", BlueprintThreadSafe), Category="Guid")
 	static FString Conv_GuidToString( const FGuid& InGuid );
 
 	/** Converts a String of format EGuidFormats to a Guid. Returns Guid OutGuid, Returns bool Success */
diff --git a/Engine/Source/Runtime/Engine/Classes/Kismet/KismetInputLibrary.h b/Engine/Source/Runtime/Engine/Classes/Kismet/KismetInputLibrary.h
--- a/Engine/Source/Runtime/Engine/Classes/Kismet/KismetInputLibrary.h
+++ b/Engine/Source/Runtime/Engine/Classes/Kismet/KismetInputLibrary.h
@@ -27,7 +27,7 @@ class ENGINE_API UKismetInputLibrary : public UBlueprintFunctionLibrary
 	GENERATED_UCLASS_BODY()
 
 	/** Calibrate the tilt for the input device */
-	UFUNCTION(BlueprintCallable, Category="Input")
+	UFUNCTION(BlueprintCallable, Category="Input|MotionTracking")
 	static void CalibrateTilt();
 
 	/**
@@ -36,7 +36,7 @@ class ENGINE_API UKismetInputLibrary : public UBlueprintFunctionLibrary
 	 * @param B - The key to compare
 	 * Returns true if the key are equal, false otherwise
 	 */
-	UFUNCTION(BlueprintPure, meta=(DisplayName = "Equal (Key)", CompactNodeTitle = "=="), Category="Utilities|Key")
+	UFUNCTION(BlueprintPure, meta=(DisplayName = "Equal (Key)", CompactNodeTitle = "=="), Category="Input|Key")
 	static bool EqualEqual_KeyKey(FKey A, FKey B);
 
 	/**
@@ -45,183 +45,183 @@ class ENGINE_API UKismetInputLibrary : public UBlueprintFunctionLibrary
 	* @param B - The chord to compare
 	* Returns true if the chords are equal, false otherwise
 	*/
-	UFUNCTION( BlueprintPure, meta = ( DisplayName = "Equal (Input Chord)", CompactNodeTitle = "==" ), Category = "Utilities|Key" )
+	UFUNCTION( BlueprintPure, meta = ( DisplayName = "Equal (InputChord)", CompactNodeTitle = "==" ), Category = "Input|Key" )
 	static bool EqualEqual_InputChordInputChord( FInputChord A, FInputChord B );
 
 	/**
 	 * Returns true if the key is a modifier key: Ctrl, Command, Alt, Shift
 	 */
-	UFUNCTION(BlueprintPure, meta = (DisplayName = "Is Modifier Key"), Category="Utilities|Key")
+	UFUNCTION(BlueprintPure, meta = (DisplayName = "Is Modifier Key"), Category="Input|Key")
 	static bool Key_IsModifierKey(const FKey& Key);
 
 	/**
 	 * Returns true if the key is a gamepad button
 	 */
-	UFUNCTION(BlueprintPure, meta = (DisplayName = "Is Gamepad Key"), Category="Utilities|Key")
+	UFUNCTION(BlueprintPure, meta = (DisplayName = "Is Gamepad Key"), Category="Input|Key")
 	static bool Key_IsGamepadKey(const FKey& Key);
 
 	/**
 	 * Returns true if the key is a mouse button
 	 */
-	UFUNCTION(BlueprintPure, meta = (DisplayName = "Is Mouse Button"), Category="Utilities|Key")
+	UFUNCTION(BlueprintPure, meta = (DisplayName = "Is Mouse Button"), Category="Input|Key")
 	static bool Key_IsMouseButton(const FKey& Key);
 
 	/**
 	 * Returns true if the key is a keyboard button
 	 */
-	UFUNCTION(BlueprintPure, meta = (DisplayName = "Is Keyboard Key"), Category="Utilities|Key")
+	UFUNCTION(BlueprintPure, meta = (DisplayName = "Is Keyboard Key"), Category="Input|Key")
 	static bool Key_IsKeyboardKey(const FKey& Key);
 
 	/**
 	 * Returns true if the key is a vector axis
 	 * @note Deprecated. Use Is Axis 2D/3D instead.
 	 */
-	UFUNCTION(BlueprintPure, meta = (DisplayName = "Is Vector Axis", DeprecatedFunction, DeprecationMessage = "Use Is Axis 2D/3D instead."), Category = "Utilities|Key")
+	UFUNCTION(BlueprintPure, meta = (DisplayName = "Is Vector Axis", DeprecatedFunction, DeprecationMessage = "Use Is Axis 2D/3D instead."), Category = "Input|Key")
 	static bool Key_IsVectorAxis(const FKey& Key);
 
 	/**
 	 * Returns true if the key is a 1D (float) axis
 	 */
-	UFUNCTION(BlueprintPure, meta = (Keywords = "IsFloatAxis", DisplayName = "Is Axis 1D"), Category="Utilities|Key")
+	UFUNCTION(BlueprintPure, meta = (Keywords = "IsFloatAxis", DisplayName = "Is Axis 1D"), Category="Input|Key")
 	static bool Key_IsAxis1D(const FKey& Key);
 
 	/**
 	 * Returns true if the key is a 2D (vector) axis
 	 */
-	UFUNCTION(BlueprintPure, meta = (Keywords = "IsVectorAxis", DisplayName = "Is Axis 2D"), Category="Utilities|Key")
+	UFUNCTION(BlueprintPure, meta = (Keywords = "IsVectorAxis", DisplayName = "Is Axis 2D"), Category="Input|Key")
 	static bool Key_IsAxis2D(const FKey& Key);
 
 	/**
 	 * Returns true if the key is a 3D (vector) axis
 	 */
-	UFUNCTION(BlueprintPure, meta = (Keywords = "IsVectorAxis", DisplayName = "Is Axis 3D"), Category="Utilities|Key")
+	UFUNCTION(BlueprintPure, meta = (Keywords = "IsVectorAxis", DisplayName = "Is Axis 3D"), Category="Input|Key")
 	static bool Key_IsAxis3D(const FKey& Key);
 
 	/**
 	 * Returns true if the key is a 1D axis emulating a digital button press.
 	 */
-	UFUNCTION(BlueprintPure, meta = (Keywords = "IsFloatAxis", DisplayName = "Is Button Axis"), Category = "Utilities|Key")
+	UFUNCTION(BlueprintPure, meta = (Keywords = "IsFloatAxis", DisplayName = "Is Button Axis"), Category = "Input|Key")
 	static bool Key_IsButtonAxis(const FKey& Key);
 
 	/**
 	 * Returns true if the key is an analog axis
 	 */
-	UFUNCTION(BlueprintPure, meta = (Keywords = "IsFloatAxis, IsVectorAxis", DisplayName = "Is Analog"), Category = "Utilities|Key")
+	UFUNCTION(BlueprintPure, meta = (Keywords = "IsFloatAxis, IsVectorAxis", DisplayName = "Is Analog"), Category = "Input|Key")
 	static bool Key_IsAnalog(const FKey& Key);
 
 	/**
 	 * Returns true if the key is a digital button press
 	 */
-	UFUNCTION(BlueprintPure, meta = (Keywords = "IsFloatAxis, IsVectorAxis", DisplayName = "Is Digital"), Category = "Utilities|Key")
+	UFUNCTION(BlueprintPure, meta = (Keywords = "IsFloatAxis, IsVectorAxis", DisplayName = "Is Digital"), Category = "Input|Key")
 	static bool Key_IsDigital(const FKey& Key);
 
 	/**
 	 * Returns true if this is a valid key.
 	 */
-	UFUNCTION(BlueprintPure, meta = (DisplayName = "Is Valid"), Category = "Utilities|Key")
+	UFUNCTION(BlueprintPure, meta = (DisplayName = "Is Valid Key"), Category = "Input|Key")
 	static bool Key_IsValid(const FKey& Key);
 
-	UFUNCTION(BlueprintPure, meta = (DisplayName = "Get Key Navigation Action", DeprecatedFunction, DeprecationMessage = "Use Get Key Event Navigation Action instead"), Category = "Utilities|Key")
+	UFUNCTION(BlueprintPure, meta = (DisplayName = "Get Key Navigation Action", DeprecatedFunction, DeprecationMessage = "Use Get Key Event Navigation Action instead"), Category = "Input|Key")
 	static EUINavigationAction Key_GetNavigationAction(const FKey& InKey);
 
 	/** Returns the navigation action corresponding to this key, or Invalid if not found */
-	UFUNCTION(BlueprintPure, meta = (DisplayName = "Get Key Event Navigation Action"), Category = "Utilities|KeyEvent")
+	UFUNCTION(BlueprintPure, meta = (DisplayName = "Get Key Event Navigation Action"), Category = "Input|KeyEvent")
 	static EUINavigationAction Key_GetNavigationActionFromKey(const FKeyEvent& InKeyEvent);
 
 	/** Returns the navigation action corresponding to this key, or Invalid if not found */
-	UFUNCTION(BlueprintPure, meta = (DisplayName = "Get Key Event Navigation Direction"), Category = "Utilities|KeyEvent")
+	UFUNCTION(BlueprintPure, meta = (DisplayName = "Get Key Event Navigation Direction"), Category = "Input|KeyEvent")
 	static EUINavigation Key_GetNavigationDirectionFromKey(const FKeyEvent& InKeyEvent);
 
 	/** Returns the navigation action corresponding to this key, or Invalid if not found */
-	UFUNCTION(BlueprintPure, meta = (DisplayName = "Get Analog Event Navigation Direction"), Category = "Utilities|AnalogEvent")
+	UFUNCTION(BlueprintPure, meta = (DisplayName = "Get Analog Event Navigation Direction"), Category = "Input|AnalogEvent")
 	static EUINavigation Key_GetNavigationDirectionFromAnalog(const FAnalogInputEvent& InAnalogEvent);
 
 	/**
 	 * Returns the display name of the key.
 	 */
-	UFUNCTION(BlueprintPure, meta = (DisplayName = "Get Key Display Name"), Category="Utilities|Key")
+	UFUNCTION(BlueprintPure, meta = (DisplayName = "Get Key Display Name"), Category="Input|Key")
 	static FText Key_GetDisplayName(const FKey& Key);
 
 	/**
 	 * Returns whether or not this character is an auto-repeated keystroke
 	 */
-	UFUNCTION(BlueprintPure, meta=( DisplayName = "Is Repeat" ), Category="Utilities|InputEvent")
+	UFUNCTION(BlueprintPure, meta=( DisplayName = "Is Repeat" ), Category="Input|InputEvent")
 	static bool InputEvent_IsRepeat(const FInputEvent& Input);
 
 	/**
 	 * Returns true if either shift key was down when this event occurred
 	 */
-	UFUNCTION(BlueprintPure, meta=( DisplayName = "Is Shift Down" ), Category="Utilities|InputEvent")
+	UFUNCTION(BlueprintPure, meta=( DisplayName = "Is Shift Down" ), Category="Input|InputEvent")
 	static bool InputEvent_IsShiftDown(const FInputEvent& Input);
 
 	/**
 	 * Returns true if left shift key was down when this event occurred
 	 */
-	UFUNCTION(BlueprintPure, meta=( DisplayName = "Is Left Shift Down" ), Category="Utilities|InputEvent")
+	UFUNCTION(BlueprintPure, meta=( DisplayName = "Is Left Shift Down" ), Category="Input|InputEvent")
 	static bool InputEvent_IsLeftShiftDown(const FInputEvent& Input);
 
 	/**
 	 * Returns true if right shift key was down when this event occurred
 	 */
-	UFUNCTION(BlueprintPure, meta=( DisplayName = "Is Right Shift Down" ), Category="Utilities|InputEvent")
+	UFUNCTION(BlueprintPure, meta=( DisplayName = "Is Right Shift Down" ), Category="Input|InputEvent")
 	static bool InputEvent_IsRightShiftDown(const FInputEvent& Input);
 
 	/**
 	 * Returns true if either control key was down when this event occurred
 	 */
-	UFUNCTION(BlueprintPure, meta=( DisplayName = "Is Control Down" ), Category="Utilities|InputEvent")
+	UFUNCTION(BlueprintPure, meta=( DisplayName = "Is Control Down" ), Category="Input|InputEvent")
 	static bool InputEvent_IsControlDown(const FInputEvent& Input);
 
 	/**
 	 * Returns true if left control key was down when this event occurred
 	 */
-	UFUNCTION(BlueprintPure, meta=( DisplayName = "Is Left Control Down" ), Category="Utilities|InputEvent")
+	UFUNCTION(BlueprintPure, meta=( DisplayName = "Is Left Control Down" ), Category="Input|InputEvent")
 	static bool InputEvent_IsLeftControlDown(const FInputEvent& Input);
 
 	/**
 	 * Returns true if left control key was down when this event occurred
 	 */
-	UFUNCTION(BlueprintPure, meta=( DisplayName = "Is Right Control Down" ), Category="Utilities|InputEvent")
+	UFUNCTION(BlueprintPure, meta=( DisplayName = "Is Right Control Down" ), Category="Input|InputEvent")
 	static bool InputEvent_IsRightControlDown(const FInputEvent& Input);
 
 	/**
 	 * Returns true if either alt key was down when this event occurred
 	 */
-	UFUNCTION(BlueprintPure, meta=( DisplayName = "Is Alt Down" ), Category="Utilities|InputEvent")
+	UFUNCTION(BlueprintPure, meta=( DisplayName = "Is Alt Down" ), Category="Input|InputEvent")
 	static bool InputEvent_IsAltDown(const FInputEvent& Input);
 
 	/**
 	 * Returns true if left alt key was down when this event occurred
 	 */
-	UFUNCTION(BlueprintPure, meta=( DisplayName = "Is Left Alt Down" ), Category="Utilities|InputEvent")
+	UFUNCTION(BlueprintPure, meta=( DisplayName = "Is Left Alt Down" ), Category="Input|InputEvent")
 	static bool InputEvent_IsLeftAltDown(const FInputEvent& Input);
 
 	/**
 	 * Returns true if right alt key was down when this event occurred
 	 */
-	UFUNCTION(BlueprintPure, meta=( DisplayName = "Is Right Alt Down" ), Category="Utilities|InputEvent")
+	UFUNCTION(BlueprintPure, meta=( DisplayName = "Is Right Alt Down" ), Category="Input|InputEvent")
 	static bool InputEvent_IsRightAltDown(const FInputEvent& Input);
 
 	/**
 	 * Returns true if either command key was down when this event occurred
 	 */
-	UFUNCTION(BlueprintPure, meta=( DisplayName = "Is Command Down" ), Category="Utilities|InputEvent")
+	UFUNCTION(BlueprintPure, meta=( DisplayName = "Is Command Down" ), Category="Input|InputEvent")
 	static bool InputEvent_IsCommandDown(const FInputEvent& Input);
 
 	/**
 	 * Returns true if left command key was down when this event occurred
 	 */
-	UFUNCTION(BlueprintPure, meta=( DisplayName = "Is Left Command Down" ), Category="Utilities|InputEvent")
+	UFUNCTION(BlueprintPure, meta=( DisplayName = "Is Left Command Down" ), Category="Input|InputEvent")
 	static bool InputEvent_IsLeftCommandDown(const FInputEvent& Input);
 
 	/**
 	 * Returns true if right command key was down when this event occurred
 	 */
-	UFUNCTION(BlueprintPure, meta=( DisplayName = "Is Right Command Down" ), Category="Utilities|InputEvent")
+	UFUNCTION(BlueprintPure, meta=( DisplayName = "Is Right Command Down" ), Category="Input|InputEvent")
 	static bool InputEvent_IsRightCommandDown(const FInputEvent& Input);
 
 	/** @return The display name of the input chord */
-	UFUNCTION(BlueprintPure, meta = (DisplayName = "Get Input Chord Display Name"), Category = "Utilities|Key")
+	UFUNCTION(BlueprintPure, meta = (DisplayName = "Get Input Chord Display Name"), Category = "Input|Key")
 	static FText InputChord_GetDisplayName(const FInputChord& Key);
 
 	/**
@@ -229,60 +229,60 @@ class ENGINE_API UKismetInputLibrary : public UBlueprintFunctionLibrary
 	 *
 	 * @return  Key name
 	 */
-	UFUNCTION(BlueprintPure, Category="Utilities|KeyEvent")
+	UFUNCTION(BlueprintPure, Category="Input|KeyEvent")
 	static FKey GetKey(const FKeyEvent& Input);
 
-	UFUNCTION(BlueprintPure, Category = "Utilities|KeyEvent")
+	UFUNCTION(BlueprintPure, Category = "Input|KeyEvent")
 	static int32 GetUserIndex(const FKeyEvent& Input);
 
-	UFUNCTION(BlueprintPure, Category = "Utilities|FAnalogInputEvent")
+	UFUNCTION(BlueprintPure, Category = "Input|InputEvent")
 	static float GetAnalogValue(const FAnalogInputEvent& Input);
 
 	/** Returns The position of the cursor in screen space */
-	UFUNCTION(BlueprintPure, meta=( DisplayName = "Get Screen Space Position" ), Category="Utilities|PointerEvent")
+	UFUNCTION(BlueprintPure, meta=( DisplayName = "Get Screen Space Position" ), Category="Input|PointerEvent")
 	static FVector2D PointerEvent_GetScreenSpacePosition(const FPointerEvent& Input);
 
 	/** Returns the position of the cursor in screen space last time we handled an input event */
-	UFUNCTION(BlueprintPure, meta=( DisplayName = "Get Last Screen Space Position" ), Category="Utilities|PointerEvent")
+	UFUNCTION(BlueprintPure, meta=( DisplayName = "Get Last Screen Space Position" ), Category="Input|PointerEvent")
 	static FVector2D PointerEvent_GetLastScreenSpacePosition(const FPointerEvent& Input);
 
 	/** Returns the distance the mouse traveled since the last event was handled. */
-	UFUNCTION(BlueprintPure, meta=( DisplayName = "Get Cursor Delta" ), Category="Utilities|PointerEvent")
+	UFUNCTION(BlueprintPure, meta=( DisplayName = "Get Cursor Delta" ), Category="Input|PointerEvent")
 	static FVector2D PointerEvent_GetCursorDelta(const FPointerEvent& Input);
 
 	/** Mouse buttons that are currently pressed */
-	UFUNCTION(BlueprintPure, meta=( DisplayName = "Is Mouse Button Down" ), Category="Utilities|PointerEvent")
+	UFUNCTION(BlueprintPure, meta=( DisplayName = "Is Mouse Button Down" ), Category="Input|PointerEvent")
 	static bool PointerEvent_IsMouseButtonDown(const FPointerEvent& Input, FKey MouseButton);
 
 	/** Mouse button that caused this event to be raised (possibly FKey::Invalid) */
-	UFUNCTION(BlueprintPure, meta=( DisplayName = "Get Effecting Button" ), Category="Utilities|PointerEvent")
+	UFUNCTION(BlueprintPure, meta=( DisplayName = "Get Effecting Button" ), Category="Input|PointerEvent")
 	static FKey PointerEvent_GetEffectingButton(const FPointerEvent& Input);
 
 	/** How much did the mouse wheel turn since the last mouse event */
-	UFUNCTION(BlueprintPure, meta=( DisplayName = "Get Wheel Delta" ), Category="Utilities|PointerEvent")
+	UFUNCTION(BlueprintPure, meta=( DisplayName = "Get Wheel Delta" ), Category="Input|PointerEvent")
 	static float PointerEvent_GetWheelDelta(const FPointerEvent& Input);
 
 	/** Returns the index of the user that caused the event */
-	UFUNCTION(BlueprintPure, meta=( DisplayName = "Get User Index" ), Category="Utilities|PointerEvent")
+	UFUNCTION(BlueprintPure, meta=( DisplayName = "Get User Index" ), Category="Input|PointerEvent")
 	static int32 PointerEvent_GetUserIndex(const FPointerEvent& Input);
 
 	/** Returns the unique identifier of the pointer (e.g., finger index) */
-	UFUNCTION(BlueprintPure, meta=( DisplayName = "Get Pointer Index" ), Category="Utilities|PointerEvent")
+	UFUNCTION(BlueprintPure, meta=( DisplayName = "Get Pointer Index" ), Category="Input|PointerEvent")
 	static int32 PointerEvent_GetPointerIndex(const FPointerEvent& Input);
 
 	/** Returns the index of the touch pad that generated this event (for platforms with multiple touch pads per user) */
-	UFUNCTION(BlueprintPure, meta=( DisplayName = "Get Touchpad Index" ), Category="Utilities|PointerEvent")
+	UFUNCTION(BlueprintPure, meta=( DisplayName = "Get Touchpad Index" ), Category="Input|PointerEvent")
 	static int32 PointerEvent_GetTouchpadIndex(const FPointerEvent& Input);
 
 	/** Returns true if this event a result from a touch (as opposed to a mouse) */
-	UFUNCTION(BlueprintPure, meta=( DisplayName = "Is Touch Event" ), Category="Utilities|PointerEvent")
+	UFUNCTION(BlueprintPure, meta=( DisplayName = "Is Touch Event" ), Category="Input|PointerEvent")
 	static bool PointerEvent_IsTouchEvent(const FPointerEvent& Input);
 
 	/** Returns the type of touch gesture */
-	UFUNCTION(BlueprintPure, Category="Utilities|PointerEvent")
+	UFUNCTION(BlueprintPure, Category="Input|PointerEvent")
 	static ESlateGesture PointerEvent_GetGestureType(const FPointerEvent& Input);
 
 	/** Returns the change in gesture value since the last gesture event of the same type. */
-	UFUNCTION(BlueprintPure, meta=( DisplayName = "Get Gesture Delta" ), Category="Utilities|PointerEvent")
+	UFUNCTION(BlueprintPure, meta=( DisplayName = "Get Gesture Delta" ), Category="Input|PointerEvent")
 	static FVector2D PointerEvent_GetGestureDelta(const FPointerEvent& Input);
 };
diff --git a/Engine/Source/Runtime/Engine/Classes/Kismet/KismetMathLibrary.h b/Engine/Source/Runtime/Engine/Classes/Kismet/KismetMathLibrary.h
--- a/Engine/Source/Runtime/Engine/Classes/Kismet/KismetMathLibrary.h
+++ b/Engine/Source/Runtime/Engine/Classes/Kismet/KismetMathLibrary.h
@@ -184,11 +184,11 @@ class ENGINE_API UKismetMathLibrary : public UBlueprintFunctionLibrary
 	static bool Not_PreBool(bool A);
 
 	/** Returns true if the values are equal (A == B) */
-	UFUNCTION(BlueprintPure, meta=(DisplayName = "Equal Boolean", CompactNodeTitle = "==", Keywords = "== equal"), Category="Math|Boolean")
+	UFUNCTION(BlueprintPure, meta=(DisplayName = "Equal (Boolean)", CompactNodeTitle = "==", Keywords = "== equal"), Category="Math|Boolean")
 	static bool EqualEqual_BoolBool(bool A, bool B);
 
 	/** Returns true if the values are not equal (A != B) */
-	UFUNCTION(BlueprintPure, meta=(DisplayName = "NotEqual Boolean", CompactNodeTitle = "!=", Keywords = "!= not equal"), Category="Math|Boolean")
+	UFUNCTION(BlueprintPure, meta=(DisplayName = "Not Equal (Boolean)", CompactNodeTitle = "!=", Keywords = "!= not equal"), Category="Math|Boolean")
 	static bool NotEqual_BoolBool(bool A, bool B);
 
 	/** Returns the logical AND of two values (A AND B) */
@@ -264,7 +264,7 @@ class ENGINE_API UKismetMathLibrary : public UBlueprintFunctionLibrary
 	static bool EqualEqual_ByteByte(uint8 A, uint8 B);
 
 	/** Returns true if A is not equal to B (A != B) */
-	UFUNCTION(BlueprintPure, meta=(DisplayName = "NotEqual (Byte)", CompactNodeTitle = "!=", Keywords = "!= not equal"), Category="Math|Byte")
+	UFUNCTION(BlueprintPure, meta=(DisplayName = "Not Equal (Byte)", CompactNodeTitle = "!=", Keywords = "!= not equal"), Category="Math|Byte")
 	static bool NotEqual_ByteByte(uint8 A, uint8 B);
 
 	//
@@ -280,7 +280,7 @@ class ENGINE_API UKismetMathLibrary : public UBlueprintFunctionLibrary
 	static int32 Divide_IntInt(int32 A, int32 B = 1);
 
 	/** Modulo (A % B) */
-	UFUNCTION(BlueprintPure, meta=(DisplayName = "% (integer)", CompactNodeTitle = "%", Keywords = "% modulus"), Category="Math|Integer")
+	UFUNCTION(BlueprintPure, meta=(DisplayName = "% (Integer)", CompactNodeTitle = "%", Keywords = "% modulus"), Category="Math|Integer")
 	static int32 Percent_IntInt(int32 A, int32 B = 1);
 
 	/** Addition (A + B) */
@@ -308,18 +308,18 @@ class ENGINE_API UKismetMathLibrary : public UBlueprintFunctionLibrary
 	static bool GreaterEqual_IntInt(int32 A, int32 B);
 
 	/** Returns true if A is equal to B (A == B) */
-	UFUNCTION(BlueprintPure, meta=(DisplayName = "Equal (integer)", CompactNodeTitle = "==", Keywords = "== equal"), Category="Math|Integer")
+	UFUNCTION(BlueprintPure, meta=(DisplayName = "Equal (Integer)", CompactNodeTitle = "==", Keywords = "== equal"), Category="Math|Integer")
 	static bool EqualEqual_IntInt(int32 A, int32 B);
 
 	/** Returns true if A is not equal to B (A != B) */
-	UFUNCTION(BlueprintPure, meta=(DisplayName = "NotEqual (integer)", CompactNodeTitle = "!=", Keywords = "!= not equal"), Category="Math|Integer")
+	UFUNCTION(BlueprintPure, meta=(DisplayName = "Not Equal (Integer)", CompactNodeTitle = "!=", Keywords = "!= not equal"), Category="Math|Integer")
 	static bool NotEqual_IntInt(int32 A, int32 B);
 
 	/** Returns true if value is between Min and Max (V >= Min && V <= Max)
 	 * If InclusiveMin is true, value needs to be equal or larger than Min, else it needs to be larger
 	 * If InclusiveMax is true, value needs to be smaller or equal than Max, else it needs to be smaller
 	 */
-	UFUNCTION(BlueprintPure, meta = (DisplayName = "InRange (integer)", Min = "0", Max = "10"), Category = "Math|Integer")
+	UFUNCTION(BlueprintPure, meta = (DisplayName = "In Range (Integer)", Min = "0", Max = "10"), Category = "Math|Integer")
 	static bool InRange_IntInt(int32 Value, int32 Min, int32 Max, bool InclusiveMin = true, bool InclusiveMax = true);
 
 	/** Bitwise AND (A & B) */
@@ -339,7 +339,7 @@ class ENGINE_API UKismetMathLibrary : public UBlueprintFunctionLibrary
 	static int32 Not_Int(int32 A);
 
 	/** Sign (integer, returns -1 if A < 0, 0 if A is zero, and +1 if A > 0) */
-	UFUNCTION(BlueprintPure, meta=(DisplayName = "Sign (integer)"), Category="Math|Integer")
+	UFUNCTION(BlueprintPure, meta=(DisplayName = "Sign (Integer)"), Category="Math|Integer")
 	static int32 SignOfInteger(int32 A);
 
 	/** Returns a uniformly distributed random number between 0 and Max - 1 */
@@ -351,23 +351,23 @@ class ENGINE_API UKismetMathLibrary : public UBlueprintFunctionLibrary
 	static int32 RandomIntegerInRange(int32 Min, int32 Max);
 
 	/** Returns the minimum value of A and B */
-	UFUNCTION(BlueprintPure, meta=(DisplayName = "Min (integer)", CompactNodeTitle = "MIN", CommutativeAssociativeBinaryOperator = "true"), Category="Math|Integer")
+	UFUNCTION(BlueprintPure, meta=(DisplayName = "Min (Integer)", CompactNodeTitle = "MIN", CommutativeAssociativeBinaryOperator = "true"), Category="Math|Integer")
 	static int32 Min(int32 A, int32 B);
 
 	/** Returns the maximum value of A and B */
-	UFUNCTION(BlueprintPure, meta=(DisplayName = "Max (integer)", CompactNodeTitle = "MAX", CommutativeAssociativeBinaryOperator = "true"), Category="Math|Integer")
+	UFUNCTION(BlueprintPure, meta=(DisplayName = "Max (Integer)", CompactNodeTitle = "MAX", CommutativeAssociativeBinaryOperator = "true"), Category="Math|Integer")
 	static int32 Max(int32 A, int32 B);
 
 	/** Returns Value clamped to be between A and B (inclusive) */
-	UFUNCTION(BlueprintPure, meta=(DisplayName = "Clamp (integer)"), Category="Math|Integer")
+	UFUNCTION(BlueprintPure, meta=(DisplayName = "Clamp (Integer)"), Category="Math|Integer")
 	static int32 Clamp(int32 Value, int32 Min, int32 Max);
 
 	/** Returns Value between A and B (inclusive) that wraps around */
-	UFUNCTION(BlueprintPure, meta = (DisplayName = "Wrap (integer)", Min = "0", Max = "100"), Category = "Math|Integer")
+	UFUNCTION(BlueprintPure, meta = (DisplayName = "Wrap (Integer)", Min = "0", Max = "100"), Category = "Math|Integer")
 	static int32 Wrap(int32 Value, int32 Min, int32 Max);
 
 	/** Returns the absolute (positive) value of A */
-	UFUNCTION(BlueprintPure, meta=(DisplayName = "Absolute (integer)", CompactNodeTitle = "ABS"), Category="Math|Integer")
+	UFUNCTION(BlueprintPure, meta=(DisplayName = "Absolute (Integer)", CompactNodeTitle = "ABS"), Category="Math|Integer")
 	static int32 Abs_Int(int32 A);
 
 	//
@@ -407,18 +407,18 @@ class ENGINE_API UKismetMathLibrary : public UBlueprintFunctionLibrary
 	static bool GreaterEqual_Int64Int64(int64 A, int64 B);
 
 	/** Returns true if A is equal to B (A == B) */
-	UFUNCTION(BlueprintPure, meta=(DisplayName = "Equal (integer64)", CompactNodeTitle = "==", Keywords = "== equal"), Category="Math|Integer64")
+	UFUNCTION(BlueprintPure, meta=(DisplayName = "Equal (Integer64)", CompactNodeTitle = "==", Keywords = "== equal"), Category="Math|Integer64")
 	static bool EqualEqual_Int64Int64(int64 A, int64 B);
 
 	/** Returns true if A is not equal to B (A != B) */
-	UFUNCTION(BlueprintPure, meta=(DisplayName = "NotEqual (integer64)", CompactNodeTitle = "!=", Keywords = "!= not equal"), Category="Math|Integer64")
+	UFUNCTION(BlueprintPure, meta=(DisplayName = "Not Equal (Integer64)", CompactNodeTitle = "!=", Keywords = "!= not equal"), Category="Math|Integer64")
 	static bool NotEqual_Int64Int64(int64 A, int64 B);
 	
 	/** Returns true if value is between Min and Max (V >= Min && V <= Max)
 	 * If InclusiveMin is true, value needs to be equal or larger than Min, else it needs to be larger
 	 * If InclusiveMax is true, value needs to be smaller or equal than Max, else it needs to be smaller
 	 */
-	UFUNCTION(BlueprintPure, meta = (DisplayName = "InRange (integer64)", Min = "0", Max = "10"), Category = "Math|Integer64")
+	UFUNCTION(BlueprintPure, meta = (DisplayName = "In Range (Integer64)", Min = "0", Max = "10"), Category = "Math|Integer64")
 	static bool InRange_Int64Int64(int64 Value, int64 Min, int64 Max, bool InclusiveMin = true, bool InclusiveMax = true);
 
 	/** Bitwise AND (A & B) */
@@ -438,7 +438,7 @@ class ENGINE_API UKismetMathLibrary : public UBlueprintFunctionLibrary
 	static int64 Not_Int64(int64 A);
 
 	/** Sign (integer64, returns -1 if A < 0, 0 if A is zero, and +1 if A > 0) */
-	UFUNCTION(BlueprintPure, meta=(DisplayName = "Sign (integer64)"), Category="Math|Integer64")
+	UFUNCTION(BlueprintPure, meta=(DisplayName = "Sign (Integer64)"), Category="Math|Integer64")
 	static int64 SignOfInteger64(int64 A);
 
 	/** Returns a uniformly distributed random number between 0 and Max - 1 */
@@ -450,19 +450,19 @@ class ENGINE_API UKismetMathLibrary : public UBlueprintFunctionLibrary
 	static int64 RandomInteger64InRange(int64 Min, int64 Max);
 
 	/** Returns the minimum value of A and B */
-	UFUNCTION(BlueprintPure, meta=(DisplayName = "Min (integer64)", CompactNodeTitle = "MIN", CommutativeAssociativeBinaryOperator = "true"), Category="Math|Integer64")
+	UFUNCTION(BlueprintPure, meta=(DisplayName = "Min (Integer64)", CompactNodeTitle = "MIN", CommutativeAssociativeBinaryOperator = "true"), Category="Math|Integer64")
 	static int64 MinInt64(int64 A, int64 B);
 
 	/** Returns the maximum value of A and B */
-	UFUNCTION(BlueprintPure, meta=(DisplayName = "Max (integer64)", CompactNodeTitle = "MAX", CommutativeAssociativeBinaryOperator = "true"), Category="Math|Integer64")
+	UFUNCTION(BlueprintPure, meta=(DisplayName = "Max (Integer64)", CompactNodeTitle = "MAX", CommutativeAssociativeBinaryOperator = "true"), Category="Math|Integer64")
 	static int64 MaxInt64(int64 A, int64 B);
 
 	/** Returns Value clamped to be between A and B (inclusive) */
-	UFUNCTION(BlueprintPure, meta=(DisplayName = "Clamp (integer64)"), Category="Math|Integer64")
+	UFUNCTION(BlueprintPure, meta=(DisplayName = "Clamp (Integer64)"), Category="Math|Integer64")
 	static int64 ClampInt64(int64 Value, int64 Min, int64 Max);
 
 	/** Returns the absolute (positive) value of A */
-	UFUNCTION(BlueprintPure, meta=(DisplayName = "Absolute (integer64)", CompactNodeTitle = "ABS"), Category="Math|Integer64")
+	UFUNCTION(BlueprintPure, meta=(DisplayName = "Absolute (Integer64)", CompactNodeTitle = "ABS"), Category="Math|Integer64")
 	static int64 Abs_Int64(int64 A);
 
 	//
@@ -506,7 +506,7 @@ class ENGINE_API UKismetMathLibrary : public UBlueprintFunctionLibrary
 	}
 
 	/** Modulo (A % B) */
-	UFUNCTION(BlueprintPure, CustomThunk, meta = (DisplayName = "% (float)", CompactNodeTitle = "%", Keywords = "% modulus"), Category = "Math|Float")
+	UFUNCTION(BlueprintPure, CustomThunk, meta = (DisplayName = "% (Float)", CompactNodeTitle = "%", Keywords = "% modulus"), Category = "Math|Float")
 	static float Percent_FloatFloat(float A, float B = 1.f);
 
 	static float GenericPercent_FloatFloat(float A, float B);
@@ -566,7 +566,7 @@ class ENGINE_API UKismetMathLibrary : public UBlueprintFunctionLibrary
 	static bool GreaterEqual_DoubleDouble(double A, double B);
 
 	/** Returns true if A is exactly equal to B (A == B)*/
-	UFUNCTION(BlueprintPure, meta = (DisplayName = "Equal (double)", CompactNodeTitle = "==", Keywords = "== equal"), Category = "Math|double")
+	UFUNCTION(BlueprintPure, meta = (DisplayName = "Equal (Double)", CompactNodeTitle = "==", Keywords = "== equal"), Category = "Math|double")
 	static bool EqualEqual_DoubleDouble(double A, double B);
 
 	/** Addition (A + B) */
@@ -594,26 +594,26 @@ class ENGINE_API UKismetMathLibrary : public UBlueprintFunctionLibrary
 	static bool GreaterEqual_FloatFloat(float A, float B);
 
 	/** Returns true if A is exactly equal to B (A == B)*/
-	UFUNCTION(BlueprintPure, meta=(DisplayName = "Equal (float)", CompactNodeTitle = "==", Keywords = "== equal"), Category="Math|Float")
+	UFUNCTION(BlueprintPure, meta=(DisplayName = "Equal (Float)", CompactNodeTitle = "==", Keywords = "== equal"), Category="Math|Float")
 	static bool EqualEqual_FloatFloat(float A, float B);
 
 	/** Returns true if A is nearly equal to B (|A - B| < ErrorTolerance) */
-	UFUNCTION(BlueprintPure, meta=(DisplayName = "Nearly Equal (float)", Keywords = "== equal"), Category="Math|Float")
+	UFUNCTION(BlueprintPure, meta=(DisplayName = "Nearly Equal (Float)", Keywords = "== equal"), Category="Math|Float")
 	static bool NearlyEqual_FloatFloat(float A, float B, float ErrorTolerance = 1.e-6f);
 
 	/** Returns true if A does not equal B (A != B)*/
-	UFUNCTION(BlueprintPure, meta=(DisplayName = "NotEqual (float)", CompactNodeTitle = "!=", Keywords = "!= not equal"), Category="Math|Float")
+	UFUNCTION(BlueprintPure, meta=(DisplayName = "Not Equal (Float)", CompactNodeTitle = "!=", Keywords = "!= not equal"), Category="Math|Float")
 	static bool NotEqual_FloatFloat(float A, float B);
 
 	/** Returns true if A does not equal B (A != B)*/
-	UFUNCTION(BlueprintPure, meta = (DisplayName = "NotEqual (Double)", CompactNodeTitle = "!=", Keywords = "!= not equal"), Category = "Math|Double")
+	UFUNCTION(BlueprintPure, meta = (DisplayName = "Not Equal (Double)", CompactNodeTitle = "!=", Keywords = "!= not equal"), Category = "Math|Double")
 	static bool NotEqual_DoubleDouble(double A, double B);
 
 	/** Returns true if value is between Min and Max (V >= Min && V <= Max)
 	 * If InclusiveMin is true, value needs to be equal or larger than Min, else it needs to be larger
 	 * If InclusiveMax is true, value needs to be smaller or equal than Max, else it needs to be smaller
 	 */
-	UFUNCTION(BlueprintPure, meta=(DisplayName = "InRange (float)", Min="0.0", Max="1.0"), Category="Math|Float")
+	UFUNCTION(BlueprintPure, meta=(DisplayName = "In Range (Float)", Min="0.0", Max="1.0"), Category="Math|Float")
 	static bool InRange_FloatFloat(float Value, float Min, float Max, bool InclusiveMin = true, bool InclusiveMax = true);
 
 	/** Returns the hypotenuse of a right-angled triangle given the width and height. */
@@ -625,11 +625,11 @@ class ENGINE_API UKismetMathLibrary : public UBlueprintFunctionLibrary
 	 * If GridSize is 0 Location is returned
 	 * if GridSize is very small precision issues may occur.
 	 */
-	UFUNCTION(BlueprintPure, meta = (DisplayName = "Snap to grid (float)"), Category = "Math|Float")
+	UFUNCTION(BlueprintPure, meta = (DisplayName = "Snap To Grid (Float)"), Category = "Math|Float")
 	static float GridSnap_Float(float Location, float GridSize);
 
 	/** Returns the absolute (positive) value of A */
-	UFUNCTION(BlueprintPure, meta=(DisplayName = "Absolute (float)", CompactNodeTitle = "ABS"), Category="Math|Float")
+	UFUNCTION(BlueprintPure, meta=(DisplayName = "Absolute (Float)", CompactNodeTitle = "ABS"), Category="Math|Float")
 	static float Abs(float A);
 
 	/** Returns the sine of A (expects Radians)*/
@@ -743,19 +743,19 @@ class ENGINE_API UKismetMathLibrary : public UBlueprintFunctionLibrary
 	static float ClampAngle(float AngleDegrees, float MinAngleDegrees, float MaxAngleDegrees);
 
 	/** Returns the minimum value of A and B */
-	UFUNCTION(BlueprintPure, meta=(DisplayName = "Min (float)", CompactNodeTitle = "MIN", CommutativeAssociativeBinaryOperator = "true"), Category="Math|Float")
+	UFUNCTION(BlueprintPure, meta=(DisplayName = "Min (Float)", CompactNodeTitle = "MIN", CommutativeAssociativeBinaryOperator = "true"), Category="Math|Float")
 	static float FMin(float A, float B);
 
 	/** Returns the maximum value of A and B */
-	UFUNCTION(BlueprintPure, meta=(DisplayName = "Max (float)", CompactNodeTitle = "MAX", CommutativeAssociativeBinaryOperator = "true"), Category="Math|Float")
+	UFUNCTION(BlueprintPure, meta=(DisplayName = "Max (Float)", CompactNodeTitle = "MAX", CommutativeAssociativeBinaryOperator = "true"), Category="Math|Float")
 	static float FMax(float A, float B);
 
 	/** Returns Value clamped between A and B (inclusive) */
-	UFUNCTION(BlueprintPure, meta=(DisplayName = "Clamp (float)", Min="0.0", Max="1.0"), Category="Math|Float")
+	UFUNCTION(BlueprintPure, meta=(DisplayName = "Clamp (Float)", Min="0.0", Max="1.0"), Category="Math|Float")
 	static float FClamp(float Value, float Min, float Max);
 
 	/** Returns Value wrapped from A and B (inclusive) */
-	UFUNCTION(BlueprintPure, meta = (DisplayName = "Wrap (float)", Min = "0.0", Max = "1.0"), Category = "Math|Float")
+	UFUNCTION(BlueprintPure, meta = (DisplayName = "Wrap (Float)", Min = "0.0", Max = "1.0"), Category = "Math|Float")
 	static float FWrap(float Value, float Min, float Max);
 
 	/** This functions returns 0 if B (the denominator) is zero */
@@ -814,27 +814,27 @@ class ENGINE_API UKismetMathLibrary : public UBlueprintFunctionLibrary
 	static int32 FCeil(float A);
 
 	/** Rounds A to the nearest integer (e.g., -1.6 becomes -2 and 1.6 becomes 2) */
-	UFUNCTION(BlueprintPure, meta = (DisplayName = "Round to Int64"), Category = "Math|Float")
+	UFUNCTION(BlueprintPure, meta = (DisplayName = "Round to Integer64"), Category = "Math|Float")
 	static int64 Round64(float A);
 
 	/** Rounds A down towards negative infinity / down to the previous integer (e.g., -1.6 becomes -2 and 1.6 becomes 1) */
-	UFUNCTION(BlueprintPure, meta = (DisplayName = "Floor to Int64"), Category = "Math|Float")
+	UFUNCTION(BlueprintPure, meta = (DisplayName = "Floor to Integer64"), Category = "Math|Float")
 	static int64 FFloor64(float A);
 
 	/** Rounds A towards zero, truncating the fractional part (e.g., -1.6 becomes -1 and 1.6 becomes 1) */
-	UFUNCTION(BlueprintPure, meta = (DisplayName = "Truncate to Int64", BlueprintAutocast), Category = "Math|Float")
+	UFUNCTION(BlueprintPure, meta = (DisplayName = "Truncate to Integer64", BlueprintAutocast), Category = "Math|Float")
 	static int64 FTrunc64(float A);
 
 	/** Rounds A up towards positive infinity / up to the next integer (e.g., -1.6 becomes -1 and 1.6 becomes 2) */
-	UFUNCTION(BlueprintPure, meta = (DisplayName = "Ceil to Int64"), Category = "Math|Float")
+	UFUNCTION(BlueprintPure, meta = (DisplayName = "Ceil to Integer64"), Category = "Math|Float")
 	static int64 FCeil64(float A);
 
 	/** Returns the number of times Divisor will go into Dividend (i.e., Dividend divided by Divisor), as well as the remainder */
-	UFUNCTION(BlueprintPure, meta=(DisplayName = "Division (whole and remainder)"), Category="Math|Float")
+	UFUNCTION(BlueprintPure, meta=(DisplayName = "Division (Whole and Remainder)"), Category="Math|Float")
 	static int32 FMod(float Dividend, float Divisor, float& Remainder);
 
 	/** Sign (float, returns -1 if A < 0, 0 if A is zero, and +1 if A > 0) */
-	UFUNCTION(BlueprintPure, meta=(DisplayName = "Sign (float)"), Category="Math|Float")
+	UFUNCTION(BlueprintPure, meta=(DisplayName = "Sign (Float)"), Category="Math|Float")
 	static float SignOfFloat(float A);
 
 	/** Returns Value normalized to the given range.  (e.g. 20 normalized to the range 10->50 would result in 0.25) */
@@ -1029,7 +1029,7 @@ class ENGINE_API UKismetMathLibrary : public UBlueprintFunctionLibrary
 	static FVector2D Divide_Vector2DFloat(FVector2D A, float B = 1.f);
 
 	/** Returns true if vector A is equal to vector B (A == B) */
-	UFUNCTION(BlueprintPure, meta=(DisplayName = "Equal Exactly (Vector2D)", CompactNodeTitle = "===", ScriptMethod = "Equals", ScriptOperator = "==", Keywords = "== equal"), Category="Math|Vector")
+	UFUNCTION(BlueprintPure, meta=(DisplayName = "Equal Exactly (Vector2D)", CompactNodeTitle = "===", ScriptMethod = "Equals", ScriptOperator = "==", Keywords = "== equal"), Category="Math|Vector2D")
 	static bool EqualExactly_Vector2DVector2D(FVector2D A, FVector2D B);
 
 	/** Returns true if vector2D A is equal to vector2D B (A == B) within a specified error tolerance */
@@ -1185,7 +1185,7 @@ class ENGINE_API UKismetMathLibrary : public UBlueprintFunctionLibrary
 	static FVector2D NormalSafe2D(FVector2D A, float Tolerance = 1.e-8f);
 
 	/** Returns a unit normal version of the 2D vector */
-	UFUNCTION(BlueprintPure, meta = (DisplayName = "Normalize2D", Keywords = "Unit Vector", ScriptMethod = "NormalUnsafe"), Category = "Math|Vector2D")
+	UFUNCTION(BlueprintPure, meta = (DisplayName = "Normalize 2D", Keywords = "Unit Vector", ScriptMethod = "NormalUnsafe"), Category = "Math|Vector2D")
 	static FVector2D Normal2D(FVector2D A);
 
 	/**
@@ -1224,15 +1224,15 @@ class ENGINE_API UKismetMathLibrary : public UBlueprintFunctionLibrary
 	*
 	* @return A copy of the vector with each component set to +1 or -1
 	*/
-	UFUNCTION(BlueprintPure, meta = (DisplayName = "To sign (+1/-1) 2D", ScriptMethod = "ToSign"), Category = "Math|Vector2D")
+	UFUNCTION(BlueprintPure, meta = (DisplayName = "To Sign (+1/-1) 2D", ScriptMethod = "ToSign"), Category = "Math|Vector2D")
 	static FVector2D ToSign2D(FVector2D A);
 
 	/** Returns the length of a 2D Vector. */
-	UFUNCTION(BlueprintPure, meta = (DisplayName = "Vector2dLength", Keywords = "magnitude", ScriptMethod = "Length"), Category = "Math|Vector2D")
+	UFUNCTION(BlueprintPure, meta = (DisplayName = "Vector2D Length", Keywords = "magnitude", ScriptMethod = "Length"), Category = "Math|Vector2D")
 	static float VSize2D(FVector2D A);
 
 	/** Returns the squared length of a 2D Vector. */
-	UFUNCTION(BlueprintPure, meta = (DisplayName = "Vector2dLengthSquared", Keywords = "magnitude", ScriptMethod = "LengthSquared"), Category = "Math|Vector2D")
+	UFUNCTION(BlueprintPure, meta = (DisplayName = "Vector2D Length Squared", Keywords = "magnitude", ScriptMethod = "LengthSquared"), Category = "Math|Vector2D")
 	static float VSize2DSquared(FVector2D A);
 
 
@@ -1307,15 +1307,15 @@ class ENGINE_API UKismetMathLibrary : public UBlueprintFunctionLibrary
 	static void BreakVector(FVector InVec, float& X, float& Y, float& Z);
 
 	/** Converts a vector to LinearColor */
-	UFUNCTION(BlueprintPure, meta=(DisplayName = "ToLinearColor (Vector)", CompactNodeTitle = "->", ScriptMethod = "LinearColor", Keywords="cast convert", BlueprintAutocast), Category="Math|Conversions")
+	UFUNCTION(BlueprintPure, meta=(DisplayName = "To LinearColor (Vector)", CompactNodeTitle = "->", ScriptMethod = "LinearColor", Keywords="cast convert", BlueprintAutocast), Category="Math|Conversions")
 	static FLinearColor Conv_VectorToLinearColor(FVector InVec);
 
 	/** Convert a vector to a transform. Uses vector as location */
-	UFUNCTION(BlueprintPure, meta=(DisplayName = "ToTransform (Vector)", CompactNodeTitle = "->", ScriptMethod = "Transform", Keywords="cast convert", BlueprintAutocast), Category="Math|Conversions")
+	UFUNCTION(BlueprintPure, meta=(DisplayName = "To Transform (Vector)", CompactNodeTitle = "->", ScriptMethod = "Transform", Keywords="cast convert", BlueprintAutocast), Category="Math|Conversions")
 	static FTransform Conv_VectorToTransform(FVector InLocation);
 	
 	/** Convert a Vector to a Vector2D using the Vector's (X, Y) coordinates */
-	UFUNCTION(BlueprintPure, meta=(DisplayName = "ToVector2D (Vector)", CompactNodeTitle = "->", ScriptMethod = "Vector2D", Keywords="cast convert", BlueprintAutocast), Category="Math|Conversions")
+	UFUNCTION(BlueprintPure, meta=(DisplayName = "To Vector2D (Vector)", CompactNodeTitle = "->", ScriptMethod = "Vector2D", Keywords="cast convert", BlueprintAutocast), Category="Math|Conversions")
 	static FVector2D Conv_VectorToVector2D(FVector InVector);
 
 	/**
@@ -1324,7 +1324,7 @@ class ENGINE_API UKismetMathLibrary : public UBlueprintFunctionLibrary
 	 *
 	 * @return FRotator from the Vector's direction, without any roll.
 	 */
-	UFUNCTION(BlueprintPure, meta=(DisplayName = "RotationFromXVector", ScriptMethod = "Rotator", Keywords="rotation rotate cast convert", BlueprintAutocast), Category="Math|Conversions")
+	UFUNCTION(BlueprintPure, meta=(DisplayName = "Rotation From X Vector", ScriptMethod = "Rotator", Keywords="rotation rotate cast convert", BlueprintAutocast), Category="Math|Conversions")
 	static FRotator Conv_VectorToRotator(FVector InVec);
 
 	/** Create a rotation from an this axis and supplied angle (in degrees) */
@@ -1407,7 +1407,7 @@ class ENGINE_API UKismetMathLibrary : public UBlueprintFunctionLibrary
 	static bool EqualEqual_VectorVector(FVector A, FVector B, float ErrorTolerance = 1.e-4f);
 
 	/** Returns true if vector A is not equal to vector B (A != B) */
-	UFUNCTION(BlueprintPure, meta = (DisplayName = "Not Equal Exactly (Vector)", CompactNodeTitle = "!==", ScriptMethod = "NotEqual", ScriptOperator = "!=", Keywords = "!= not equal"), Category = "Math|Vector2D")
+	UFUNCTION(BlueprintPure, meta = (DisplayName = "Not Equal Exactly (Vector)", CompactNodeTitle = "!==", ScriptMethod = "NotEqual", ScriptOperator = "!=", Keywords = "!= not equal"), Category = "Math|Vector")
 	static bool NotEqualExactly_VectorVector(FVector A, FVector B);
 
 	/** Returns true if vector A is not equal to vector B (A != B) within a specified error tolerance */
@@ -1423,15 +1423,15 @@ class ENGINE_API UKismetMathLibrary : public UBlueprintFunctionLibrary
 	static FVector Cross_VectorVector(FVector A, FVector B);
 
 	/** Returns result of vector A rotated by Rotator B */
-	UFUNCTION(BlueprintPure, meta=(DisplayName = "RotateVector", ScriptMethod = "Rotate"), Category="Math|Vector")
+	UFUNCTION(BlueprintPure, meta=(DisplayName = "Rotate Vector", ScriptMethod = "Rotate"), Category="Math|Vector")
 	static FVector GreaterGreater_VectorRotator(FVector A, FRotator B);
 
 	/** Returns result of vector A rotated by AngleDeg around Axis */
-	UFUNCTION(BlueprintPure, meta=(DisplayName = "RotateVectorAroundAxis", ScriptMethod = "RotateAngleAxis"), Category="Math|Vector")
+	UFUNCTION(BlueprintPure, meta=(DisplayName = "Rotate Vector Around Axis", ScriptMethod = "RotateAngleAxis"), Category="Math|Vector")
 	static FVector RotateAngleAxis(FVector InVect, float AngleDeg, FVector Axis);
 
 	/** Returns result of vector A rotated by the inverse of Rotator B */
-	UFUNCTION(BlueprintPure, meta=(DisplayName = "UnrotateVector", ScriptMethod = "Unrotate"), Category="Math|Vector")
+	UFUNCTION(BlueprintPure, meta=(DisplayName = "Unrotate Vector", ScriptMethod = "Unrotate"), Category="Math|Vector")
 	static FVector LessLess_VectorRotator(FVector A, FRotator B);
 
 	/** When this vector contains Euler angles (degrees), ensure that angles are between +/-180 */
@@ -1605,19 +1605,19 @@ class ENGINE_API UKismetMathLibrary : public UBlueprintFunctionLibrary
 	static float Vector_Distance2DSquared(FVector V1, FVector V2);
 
 	/** Returns the length of the vector */
-	UFUNCTION(BlueprintPure, meta=(DisplayName = "VectorLength", ScriptMethod = "Length", Keywords="magnitude"), Category="Math|Vector")
+	UFUNCTION(BlueprintPure, meta=(DisplayName = "Vector Length", ScriptMethod = "Length", Keywords="magnitude"), Category="Math|Vector")
 	static float VSize(FVector A);
 
 	/** Returns the squared length of the vector */
-	UFUNCTION(BlueprintPure, meta=(DisplayName = "VectorLengthSquared", ScriptMethod = "LengthSquared", Keywords="magnitude"), Category="Math|Vector")
+	UFUNCTION(BlueprintPure, meta=(DisplayName = "Vector Length Squared", ScriptMethod = "LengthSquared", Keywords="magnitude"), Category="Math|Vector")
 	static float VSizeSquared(FVector A);
 
 	/** Returns the length of the vector's XY components. */
-	UFUNCTION(BlueprintPure, meta=(DisplayName = "VectorLengthXY", ScriptMethod = "Length2D", Keywords="magnitude"), Category="Math|Vector")
+	UFUNCTION(BlueprintPure, meta=(DisplayName = "Vector Length XY", ScriptMethod = "Length2D", Keywords="magnitude"), Category="Math|Vector")
 	static float VSizeXY(FVector A);
 
 	/** Returns the squared length of the vector's XY components. */
-	UFUNCTION(BlueprintPure, meta=(DisplayName = "VectorLengthXYSquared", ScriptMethod = "Length2DSquared", Keywords="magnitude"), Category="Math|Vector")
+	UFUNCTION(BlueprintPure, meta=(DisplayName = "Vector Length XY Squared", ScriptMethod = "Length2DSquared", Keywords="magnitude"), Category="Math|Vector")
 	static float VSizeXYSquared(FVector A);
 
 	/**
@@ -1695,7 +1695,7 @@ class ENGINE_API UKismetMathLibrary : public UBlueprintFunctionLibrary
 	 *
 	 * @return Normalized version of vector.
 	 */
-	UFUNCTION(BlueprintPure, meta=(DisplayName = "Normal unsafe (Vector)", ScriptMethod = "NormalUnsafe", Keywords="Unit Vector"), Category="Math|Vector")
+	UFUNCTION(BlueprintPure, meta=(DisplayName = "Normal Unsafe (Vector)", ScriptMethod = "NormalUnsafe", Keywords="Unit Vector"), Category="Math|Vector")
 	static FVector Vector_NormalUnsafe(const FVector& A);
 
 	/**
@@ -2132,11 +2132,11 @@ class ENGINE_API UKismetMathLibrary : public UBlueprintFunctionLibrary
 	static float Vector4_SizeSquared(const FVector4& A);
 
 	/** Returns the length of the vector. The W element is ignored. */
-	UFUNCTION(BlueprintPure, meta = (DisplayName = "LengthXYZ (Vector4)", ScriptMethod = "Length3", Keywords = "magnitude"), Category = "Math|Vector4")
+	UFUNCTION(BlueprintPure, meta = (DisplayName = "Length XYZ (Vector4)", ScriptMethod = "Length3", Keywords = "magnitude"), Category = "Math|Vector4")
 	static float Vector4_Size3(const FVector4& A);
 
 	/** Returns the squared length of the vector. The W element is ignored. */
-	UFUNCTION(BlueprintPure, meta = (DisplayName = "LengthXYZ Squared (Vector4)", ScriptMethod = "LengthSquared3", Keywords = "magnitude"), Category = "Math|Vector4")
+	UFUNCTION(BlueprintPure, meta = (DisplayName = "Length XYZ Squared (Vector4)", ScriptMethod = "LengthSquared3", Keywords = "magnitude"), Category = "Math|Vector4")
 	static float Vector4_SizeSquared3(const FVector4& A);
 
 	/**
@@ -2170,7 +2170,7 @@ class ENGINE_API UKismetMathLibrary : public UBlueprintFunctionLibrary
 	 *
 	 * @return Normalized version of vector.
 	 */
-	UFUNCTION(BlueprintPure, meta=(DisplayName = "Normal unsafe XYZ (Vector4)", ScriptMethod = "NormalUnsafe3", Keywords = "Unit Vector"), Category = "Math|Vector4")
+	UFUNCTION(BlueprintPure, meta=(DisplayName = "Normal Unsafe XYZ (Vector4)", ScriptMethod = "NormalUnsafe3", Keywords = "Unit Vector"), Category = "Math|Vector4")
 	static FVector4 Vector4_NormalUnsafe3(const FVector4& A);
 
 	/**
@@ -2278,19 +2278,19 @@ class ENGINE_API UKismetMathLibrary : public UBlueprintFunctionLibrary
 	static bool NotEqual_RotatorRotator(FRotator A, FRotator B, float ErrorTolerance = 1.e-4f);
 
 	/** Returns rotator representing rotator A scaled by B */
-	UFUNCTION(BlueprintPure, meta=(DisplayName = "ScaleRotator", CompactNodeTitle = "*", ScriptMethod = "Scale", Keywords = "* multiply rotate rotation"), Category="Math|Rotator")
+	UFUNCTION(BlueprintPure, meta=(DisplayName = "Scale Rotator (Float)", CompactNodeTitle = "*", ScriptMethod = "Scale", Keywords = "* multiply rotate rotation"), Category="Math|Rotator")
 	static FRotator Multiply_RotatorFloat(FRotator A, float B);
 	
 	/** Returns rotator representing rotator A scaled by B */
-	UFUNCTION(BlueprintPure, meta=(DisplayName = "ScaleRotator (integer)", CompactNodeTitle = "*", ScriptMethod = "ScaleInteger", Keywords = "* multiply rotate rotation"), Category="Math|Rotator")
+	UFUNCTION(BlueprintPure, meta=(DisplayName = "Scale Rotator (Integer)", CompactNodeTitle = "*", ScriptMethod = "ScaleInteger", Keywords = "* multiply rotate rotation"), Category="Math|Rotator")
 	static FRotator Multiply_RotatorInt(FRotator A, int32 B);
 
 	/** Combine 2 rotations to give you the resulting rotation of first applying A, then B. */
-	UFUNCTION(BlueprintPure, meta=(DisplayName = "CombineRotators", ScriptMethod = "Combine", Keywords="rotate rotation add"), Category="Math|Rotator")
+	UFUNCTION(BlueprintPure, meta=(DisplayName = "Combine Rotators", ScriptMethod = "Combine", Keywords="rotate rotation add"), Category="Math|Rotator")
 	static FRotator ComposeRotators(FRotator A, FRotator B);
 
 	/** Negate a rotator*/
-	UFUNCTION(BlueprintPure, meta=(DisplayName="InvertRotator", ScriptMethod = "Inversed", ScriptOperator = "neg", Keywords="rotate rotation"), Category="Math|Rotator")
+	UFUNCTION(BlueprintPure, meta=(DisplayName="Invert Rotator", ScriptMethod = "Inversed", ScriptOperator = "neg", Keywords="rotate rotation"), Category="Math|Rotator")
 	static FRotator NegateRotator(FRotator A);
 
 	/** Rotate the world forward vector by the given rotation */
@@ -2306,11 +2306,11 @@ class ENGINE_API UKismetMathLibrary : public UBlueprintFunctionLibrary
 	static FVector GetUpVector(FRotator InRot);
 
 	/** Get the X direction vector after this rotation */
-	UFUNCTION(BlueprintPure, meta=(ScriptMethod = "ToVector", DisplayName = "GetRotationXVector", Keywords="rotation rotate cast convert", BlueprintAutocast), Category="Math|Rotator")
+	UFUNCTION(BlueprintPure, meta=(ScriptMethod = "ToVector", DisplayName = "Get Rotation X Vector", Keywords="rotation rotate cast convert", BlueprintAutocast), Category="Math|Rotator")
 	static FVector Conv_RotatorToVector(FRotator InRot);
 
 	/** Convert Rotator to Transform */
-	UFUNCTION(BlueprintPure, meta = (DisplayName = "ToTransform (Rotator)", CompactNodeTitle = "->", ScriptMethod = "Transform", Keywords = "cast convert", BlueprintAutocast), Category = "Math|Conversions")
+	UFUNCTION(BlueprintPure, meta = (DisplayName = "To Transform (Rotator)", CompactNodeTitle = "->", ScriptMethod = "Transform", Keywords = "cast convert", BlueprintAutocast), Category = "Math|Conversions")
 	static FTransform Conv_RotatorToTransform(const FRotator& InRotator);
 
 	/** Get the reference frame direction vectors (axes) described by this rotation */
@@ -2744,7 +2744,7 @@ class ENGINE_API UKismetMathLibrary : public UBlueprintFunctionLibrary
 	static float Quat_AngularDistance(const FQuat& A, const FQuat& B);
 
 	/** Modify the quaternion to ensure that the delta between it and B represents the shortest possible rotation angle. */
-	UFUNCTION(BlueprintCallable, meta = (DisplayName = "Ensure shortest arc to (Quat)", ScriptMethod = "EnsureShortestArcTo"), Category = "Math|Quat")
+	UFUNCTION(BlueprintCallable, meta = (DisplayName = "Ensure Shortest Arc To (Quat)", ScriptMethod = "EnsureShortestArcTo"), Category = "Math|Quat")
 	static void Quat_EnforceShortestArcWith(UPARAM(ref) FQuat& A, const FQuat& B);
 
 	/**	Convert a Quaternion into floating-point Euler angles (in degrees). */
@@ -2919,7 +2919,7 @@ class ENGINE_API UKismetMathLibrary : public UBlueprintFunctionLibrary
 	static FLinearColor LinearColor_Yellow();
 
 	/** Transparent linear color - black with 0 opacity/alpha */
-	UFUNCTION(BlueprintPure, meta = (ScriptConstant = "Tansparent", ScriptConstantHost = "LinearColor"), Category = "Math|Color")
+	UFUNCTION(BlueprintPure, meta = (ScriptConstant = "Transparent", ScriptConstantHost = "LinearColor"), Category = "Math|Color")
 	static FLinearColor LinearColor_Transparent();
 
 
@@ -2970,7 +2970,7 @@ class ENGINE_API UKismetMathLibrary : public UBlueprintFunctionLibrary
 	static void LinearColor_SetRandomHue(UPARAM(ref) FLinearColor& InOutColor);
 
 	/** Convert a float into a LinearColor, where each element is that float */
-	UFUNCTION(BlueprintPure, meta = (DisplayName = "ToLinearColor (float)", CompactNodeTitle = "->", Keywords = "cast convert", BlueprintAutocast), Category = "Math|Conversions")
+	UFUNCTION(BlueprintPure, meta = (DisplayName = "To LinearColor (Float)", CompactNodeTitle = "->", Keywords = "cast convert", BlueprintAutocast), Category = "Math|Conversions")
 	static FLinearColor Conv_FloatToLinearColor(float InFloat);
 
 	/** Make a color from individual color components (HSV space; Hue is [0..360) while Saturation and Value are 0..1) */
@@ -2994,19 +2994,19 @@ class ENGINE_API UKismetMathLibrary : public UBlueprintFunctionLibrary
 	static void RGBToHSV_Vector(FLinearColor RGB, FLinearColor& HSV);
 
 	/** Converts a RGB linear color to HSV (where H is in R, S is in G, and V is in B) */
-	UFUNCTION(BlueprintPure, meta = (DisplayName = "RGB to HSV linear color", ScriptMethod = "RGBToHSV", Keywords = "cast convert"), Category = "Math|Color")
+	UFUNCTION(BlueprintPure, meta = (DisplayName = "RGB to HSV (LinearColor)", ScriptMethod = "RGBToHSV", Keywords = "cast convert"), Category = "Math|Color")
 	static FLinearColor RGBLinearToHSV(FLinearColor RGB);
 
 	/** Converts a LinearColor to a vector */
-	UFUNCTION(BlueprintPure, meta = (DisplayName = "ToVector (LinearColor)", ScriptMethod = "ToRGBVector", CompactNodeTitle = "->", Keywords = "cast convert", BlueprintAutocast), Category = "Math|Conversions")
+	UFUNCTION(BlueprintPure, meta = (DisplayName = "To Vector (LinearColor)", ScriptMethod = "ToRGBVector", CompactNodeTitle = "->", Keywords = "cast convert", BlueprintAutocast), Category = "Math|Conversions")
 	static FVector Conv_LinearColorToVector(FLinearColor InLinearColor);
 
 	/** Convert from linear to 8-bit RGBE as outlined in Gregory Ward's Real Pixels article, Graphics Gems II, page 80. */
-	UFUNCTION(BlueprintPure, meta = (DisplayName = "ToRGBE (LinearColor)", ScriptMethod = "ToRGBE", Keywords = "cast convert", BlueprintAutocast), Category = "Math|Color")
+	UFUNCTION(BlueprintPure, meta = (DisplayName = "To RGBE (LinearColor)", ScriptMethod = "ToRGBE", Keywords = "cast convert", BlueprintAutocast), Category = "Math|Color")
 	static FColor LinearColor_ToRGBE(FLinearColor InLinearColor);
 
 	/** Quantizes the linear color and returns the result as a FColor with optional sRGB conversion and quality as goal. */
-	UFUNCTION(BlueprintPure, meta = (DisplayName = "ToColor (LinearColor)", ScriptMethod = "ToColor", CompactNodeTitle = "->", Keywords = "cast convert", BlueprintAutocast), Category = "Math|Conversions")
+	UFUNCTION(BlueprintPure, meta = (DisplayName = "To Color (LinearColor)", ScriptMethod = "ToColor", CompactNodeTitle = "->", Keywords = "cast convert", BlueprintAutocast), Category = "Math|Conversions")
 	static FColor Conv_LinearColorToColor(FLinearColor InLinearColor, bool InUseSRGB = true);
 
 	/** Quantizes the linear color and returns the result as an 8-bit color.  This bypasses the SRGB conversion. */
@@ -3014,7 +3014,7 @@ class ENGINE_API UKismetMathLibrary : public UBlueprintFunctionLibrary
 	static FColor LinearColor_Quantize(FLinearColor InColor);
 
 	/** Quantizes the linear color with rounding and returns the result as an 8-bit color.  This bypasses the SRGB conversion. */
-	UFUNCTION(BlueprintPure, meta = (DisplayName = "Quantize with rounding to 8-bit (LinearColor)", ScriptMethod = "QuantizeRound", Keywords = "cast convert"), Category = "Math|Color")
+	UFUNCTION(BlueprintPure, meta = (DisplayName = "Quantize With Rounding To 8-bit (LinearColor)", ScriptMethod = "QuantizeRound", Keywords = "cast convert"), Category = "Math|Color")
 	static FColor LinearColor_QuantizeRound(FLinearColor InColor);
 
 	/**
@@ -3183,75 +3183,75 @@ class ENGINE_API UKismetMathLibrary : public UBlueprintFunctionLibrary
 	static bool LessEqual_DateTimeDateTime( FDateTime A, FDateTime B );
 
 	/** Returns the date component of A */
-	UFUNCTION(BlueprintPure, meta=(DisplayName="GetDate"), Category="Math|DateTime")
+	UFUNCTION(BlueprintPure, meta=(DisplayName="Get Date"), Category="Math|DateTime")
 	static FDateTime GetDate( FDateTime A );
 
 	/** Returns the day component of A (1 to 31) */
-	UFUNCTION(BlueprintPure, meta=(DisplayName="GetDay"), Category="Math|DateTime")
+	UFUNCTION(BlueprintPure, meta=(DisplayName="Get Day"), Category="Math|DateTime")
 	static int32 GetDay( FDateTime A );
 
 	/** Returns the day of year of A */
-	UFUNCTION(BlueprintPure, meta=(DisplayName="GetDayOfYear"), Category="Math|DateTime")
+	UFUNCTION(BlueprintPure, meta=(DisplayName="Get Day Of Year"), Category="Math|DateTime")
 	static int32 GetDayOfYear( FDateTime A );
 
 	/** Returns the hour component of A (24h format) */
-	UFUNCTION(BlueprintPure, meta=(DisplayName="GetHour"), Category="Math|DateTime")
+	UFUNCTION(BlueprintPure, meta=(DisplayName="Get Hour"), Category="Math|DateTime")
 	static int32 GetHour( FDateTime A );
 
 	/** Returns the hour component of A (12h format) */
-	UFUNCTION(BlueprintPure, meta=(DisplayName="GetHour12"), Category="Math|DateTime")
+	UFUNCTION(BlueprintPure, meta=(DisplayName="Get Hour 12"), Category="Math|DateTime")
 	static int32 GetHour12( FDateTime A );
 
 	/** Returns the millisecond component of A */
-	UFUNCTION(BlueprintPure, meta=(DisplayName="GetMillisecond"), Category="Math|DateTime")
+	UFUNCTION(BlueprintPure, meta=(DisplayName="Get Millisecond"), Category="Math|DateTime")
 	static int32 GetMillisecond( FDateTime A );
 
 	/** Returns the minute component of A */
-	UFUNCTION(BlueprintPure, meta=(DisplayName="GetMinute"), Category="Math|DateTime")
+	UFUNCTION(BlueprintPure, meta=(DisplayName="Get Minute"), Category="Math|DateTime")
 	static int32 GetMinute( FDateTime A );
 
 	/** Returns the month component of A */
-	UFUNCTION(BlueprintPure, meta=(DisplayName="GetMonth"), Category="Math|DateTime")
+	UFUNCTION(BlueprintPure, meta=(DisplayName="Get Month"), Category="Math|DateTime")
 	static int32 GetMonth( FDateTime A );
 
 	/** Returns the second component of A */
-	UFUNCTION(BlueprintPure, meta=(DisplayName="GetSecond"), Category="Math|DateTime")
+	UFUNCTION(BlueprintPure, meta=(DisplayName="Get Second"), Category="Math|DateTime")
 	static int32 GetSecond( FDateTime A );
 
 	/** Returns the time elapsed since midnight of A */
-	UFUNCTION(BlueprintPure, meta=(DisplayName="GetTimeOfDay"), Category="Math|DateTime")
+	UFUNCTION(BlueprintPure, meta=(DisplayName="Get Time Of Day"), Category="Math|DateTime")
 	static FTimespan GetTimeOfDay( FDateTime A );
 
 	/** Returns the year component of A */
-	UFUNCTION(BlueprintPure, meta=(DisplayName="GetYear"), Category="Math|DateTime")
+	UFUNCTION(BlueprintPure, meta=(DisplayName="Get Year"), Category="Math|DateTime")
 	static int32 GetYear( FDateTime A );
 
 	/** Returns whether A's time is in the afternoon */
-	UFUNCTION(BlueprintPure, meta=(DisplayName="IsAfternoon"), Category="Math|DateTime")
+	UFUNCTION(BlueprintPure, meta=(DisplayName="Is Afternoon"), Category="Math|DateTime")
 	static bool IsAfternoon( FDateTime A );
 
 	/** Returns whether A's time is in the morning */
-	UFUNCTION(BlueprintPure, meta=(DisplayName="IsMorning"), Category="Math|DateTime")
+	UFUNCTION(BlueprintPure, meta=(DisplayName="Is Morning"), Category="Math|DateTime")
 	static bool IsMorning( FDateTime A );
 
 	/** Returns the number of days in the given year and month */
-	UFUNCTION(BlueprintPure, meta=(DisplayName="DaysInMonth"), Category="Math|DateTime")
+	UFUNCTION(BlueprintPure, meta=(DisplayName="Days In Month"), Category="Math|DateTime")
 	static int32 DaysInMonth( int32 Year, int32 Month );
 
 	/** Returns the number of days in the given year */
-	UFUNCTION(BlueprintPure, meta=(DisplayName="DaysInYear"), Category="Math|DateTime")
+	UFUNCTION(BlueprintPure, meta=(DisplayName="Days In Year"), Category="Math|DateTime")
 	static int32 DaysInYear( int32 Year );
 
 	/** Returns whether given year is a leap year */
-	UFUNCTION(BlueprintPure, meta=(DisplayName="IsLeapYear"), Category="Math|DateTime")
+	UFUNCTION(BlueprintPure, meta=(DisplayName="Is Leap Year"), Category="Math|DateTime")
 	static bool IsLeapYear( int32 Year );
 
 	/** Returns the maximum date and time value */
-	UFUNCTION(BlueprintPure, meta=(DisplayName="MaxValue"), Category="Math|DateTime")
+	UFUNCTION(BlueprintPure, meta=(DisplayName="Max Value (DateTime)"), Category="Math|DateTime")
 	static FDateTime DateTimeMaxValue( );
 
 	/** Returns the minimum date and time value */
-	UFUNCTION(BlueprintPure, meta=(DisplayName="MinValue"), Category="Math|DateTime")
+	UFUNCTION(BlueprintPure, meta=(DisplayName="MinValue (DateTime)"), Category="Math|DateTime")
 	static FDateTime DateTimeMinValue( );
 
 	/** Returns the local date and time on this computer */
@@ -3263,7 +3263,7 @@ class ENGINE_API UKismetMathLibrary : public UBlueprintFunctionLibrary
 	static FDateTime Today( );
 
 	/** Returns the UTC date and time on this computer */
-	UFUNCTION(BlueprintPure, meta=(DisplayName="UtcNow"), Category="Math|DateTime")
+	UFUNCTION(BlueprintPure, meta=(DisplayName="UTC Now"), Category="Math|DateTime")
 	static FDateTime UtcNow( );
 
 	/** Converts a date string in ISO-8601 format to a DateTime object */
@@ -3280,15 +3280,15 @@ class ENGINE_API UKismetMathLibrary : public UBlueprintFunctionLibrary
 	//
 
 	/** Returns the maximum time span value */
-	UFUNCTION(BlueprintPure, meta=(DisplayName="MaxValue", ScriptConstant = "MaxValue", ScriptConstantHost = "Timespan"), Category="Math|Timespan")
+	UFUNCTION(BlueprintPure, meta=(DisplayName="Max Value (Timespan)", ScriptConstant = "MaxValue", ScriptConstantHost = "Timespan"), Category="Math|Timespan")
 	static FTimespan TimespanMaxValue( );
 
 	/** Returns the minimum time span value */
-	UFUNCTION(BlueprintPure, meta=(DisplayName="MinValue", ScriptConstant = "MinValue", ScriptConstantHost = "Timespan"), Category="Math|Timespan")
+	UFUNCTION(BlueprintPure, meta=(DisplayName="Min Value (Timespan)", ScriptConstant = "MinValue", ScriptConstantHost = "Timespan"), Category="Math|Timespan")
 	static FTimespan TimespanMinValue( );
 
 	/** Returns a zero time span value */
-	UFUNCTION(BlueprintPure, meta=(DisplayName="ZeroValue", ScriptConstant = "Zero", ScriptConstantHost = "Timespan"), Category="Math|Timespan")
+	UFUNCTION(BlueprintPure, meta=(DisplayName="Zero Value (Timespan)", ScriptConstant = "Zero", ScriptConstantHost = "Timespan"), Category="Math|Timespan")
 	static FTimespan TimespanZeroValue( );
 
 	//
@@ -3352,71 +3352,71 @@ class ENGINE_API UKismetMathLibrary : public UBlueprintFunctionLibrary
 	static bool LessEqual_TimespanTimespan( FTimespan A, FTimespan B );
 
 	/** Returns the days component of A */
-	UFUNCTION(BlueprintPure, meta=(DisplayName="GetDays"), Category="Math|Timespan")
+	UFUNCTION(BlueprintPure, meta=(DisplayName="Get Days"), Category="Math|Timespan")
 	static int32 GetDays( FTimespan A );
 
 	/** Returns the absolute value of A */
-	UFUNCTION(BlueprintPure, meta=(DisplayName="GetDuration"), Category="Math|Timespan")
+	UFUNCTION(BlueprintPure, meta=(DisplayName="Get Duration"), Category="Math|Timespan")
 	static FTimespan GetDuration( FTimespan A );
 
 	/** Returns the hours component of A */
-	UFUNCTION(BlueprintPure, meta=(DisplayName="GetHours"), Category="Math|Timespan")
+	UFUNCTION(BlueprintPure, meta=(DisplayName="Get Hours"), Category="Math|Timespan")
 	static int32 GetHours( FTimespan A );
 
 	/** Returns the milliseconds component of A */
-	UFUNCTION(BlueprintPure, meta=(DisplayName="GetMilliseconds"), Category="Math|Timespan")
+	UFUNCTION(BlueprintPure, meta=(DisplayName="Get Milliseconds"), Category="Math|Timespan")
 	static int32 GetMilliseconds( FTimespan A );
 
 	/** Returns the minutes component of A */
-	UFUNCTION(BlueprintPure, meta=(DisplayName="GetMinutes"), Category="Math|Timespan")
+	UFUNCTION(BlueprintPure, meta=(DisplayName="Get Minutes"), Category="Math|Timespan")
 	static int32 GetMinutes( FTimespan A );
 
 	/** Returns the seconds component of A */
-	UFUNCTION(BlueprintPure, meta=(DisplayName="GetSeconds"), Category="Math|Timespan")
+	UFUNCTION(BlueprintPure, meta=(DisplayName="Get Seconds"), Category="Math|Timespan")
 	static int32 GetSeconds( FTimespan A );
 
 	/** Returns the total number of days in A */
-	UFUNCTION(BlueprintPure, meta=(DisplayName="GetTotalDays"), Category="Math|Timespan")
+	UFUNCTION(BlueprintPure, meta=(DisplayName="Get Total Days"), Category="Math|Timespan")
 	static float GetTotalDays( FTimespan A );
 
 	/** Returns the total number of hours in A */
-	UFUNCTION(BlueprintPure, meta=(DisplayName="GetTotalHours"), Category="Math|Timespan")
+	UFUNCTION(BlueprintPure, meta=(DisplayName="Get Total Hours"), Category="Math|Timespan")
 	static float GetTotalHours( FTimespan A );
 
 	/** Returns the total number of milliseconds in A */
-	UFUNCTION(BlueprintPure, meta=(DisplayName="GetTotalMilliseconds"), Category="Math|Timespan")
+	UFUNCTION(BlueprintPure, meta=(DisplayName="Get Total Milliseconds"), Category="Math|Timespan")
 	static float GetTotalMilliseconds( FTimespan A );
 
 	/** Returns the total number of minutes in A */
-	UFUNCTION(BlueprintPure, meta=(DisplayName="GetTotalMinutes"), Category="Math|Timespan")
+	UFUNCTION(BlueprintPure, meta=(DisplayName="Get Total Minutes"), Category="Math|Timespan")
 	static float GetTotalMinutes( FTimespan A );
 
 	/** Returns the total number of seconds in A */
-	UFUNCTION(BlueprintPure, meta=(DisplayName="GetTotalSeconds"), Category="Math|Timespan")
+	UFUNCTION(BlueprintPure, meta=(DisplayName="Get Total Seconds"), Category="Math|Timespan")
 	static float GetTotalSeconds( FTimespan A );
 
 	/** Returns a time span that represents the specified number of days */
-	UFUNCTION(BlueprintPure, meta=(DisplayName="FromDays"), Category="Math|Timespan")
+	UFUNCTION(BlueprintPure, meta=(DisplayName="From Days"), Category="Math|Timespan")
 	static FTimespan FromDays( float Days );
 
 	/** Returns a time span that represents the specified number of hours */
-	UFUNCTION(BlueprintPure, meta=(DisplayName="FromHours"), Category="Math|Timespan")
+	UFUNCTION(BlueprintPure, meta=(DisplayName="From Hours"), Category="Math|Timespan")
 	static FTimespan FromHours( float Hours );
 
 	/** Returns a time span that represents the specified number of milliseconds */
-	UFUNCTION(BlueprintPure, meta=(DisplayName="FromMilliseconds"), Category="Math|Timespan")
+	UFUNCTION(BlueprintPure, meta=(DisplayName="From Milliseconds"), Category="Math|Timespan")
 	static FTimespan FromMilliseconds( float Milliseconds );
 
 	/** Returns a time span that represents the specified number of minutes */
-	UFUNCTION(BlueprintPure, meta=(DisplayName="FromMinutes"), Category="Math|Timespan")
+	UFUNCTION(BlueprintPure, meta=(DisplayName="From Minutes"), Category="Math|Timespan")
 	static FTimespan FromMinutes( float Minutes );
 
 	/** Returns a time span that represents the specified number of seconds */
-	UFUNCTION(BlueprintPure, meta=(DisplayName="FromSeconds"), Category="Math|Timespan")
+	UFUNCTION(BlueprintPure, meta=(DisplayName="From Seconds"), Category="Math|Timespan")
 	static FTimespan FromSeconds( float Seconds );
 
 	/** Returns the ratio between two time spans (A / B), handles zero values */
-	UFUNCTION(BlueprintPure, meta=(DisplayName="TimespanRatio"), Category="Math|Timespan")
+	UFUNCTION(BlueprintPure, meta=(DisplayName="Timespan Ratio"), Category="Math|Timespan")
 	static float TimespanRatio( FTimespan A, FTimespan B );
 
 	/** Converts a time span string to a Timespan object */
@@ -3428,113 +3428,113 @@ class ENGINE_API UKismetMathLibrary : public UBlueprintFunctionLibrary
 	//
 
 	/** Creates a FQualifiedFrameTime out of a frame number, frame rate, and optional 0-1 clamped subframe. */
-	UFUNCTION(BlueprintPure, Category = "Time Management", meta = (NativeMakeFunc))
+	UFUNCTION(BlueprintPure, Category = "Utilities|Time Management", meta = (NativeMakeFunc))
 	static FQualifiedFrameTime MakeQualifiedFrameTime(FFrameNumber Frame, FFrameRate FrameRate, float SubFrame = 0.f);
 
 	/** Breaks a FQualifiedFrameTime into its component parts again. */
-	UFUNCTION(BlueprintPure, Category = "Time Management", meta = (NativeBreakFunc))
+	UFUNCTION(BlueprintPure, Category = "Utilities|Time Management", meta = (NativeBreakFunc))
 	static void BreakQualifiedFrameTime(const FQualifiedFrameTime& InFrameTime, FFrameNumber& Frame, FFrameRate& FrameRate, float& SubFrame);
 
 	/** Creates a FFrameRate from a Numerator and a Denominator. Enforces that the Denominator is at least one. */
-	UFUNCTION(BlueprintPure, Category = "Time Management", meta = (NativeMakeFunc))
+	UFUNCTION(BlueprintPure, Category = "Utilities|Time Management", meta = (NativeMakeFunc))
 	static FFrameRate MakeFrameRate(int32 Numerator, int32 Denominator = 1);
 
 	/** Breaks a FFrameRate into a numerator and denominator. */
-	UFUNCTION(BlueprintPure, Category = "Time Management", meta = (NativeBreakFunc))
+	UFUNCTION(BlueprintPure, Category = "Utilities|Time Management", meta = (NativeBreakFunc))
 	static void BreakFrameRate(const FFrameRate& InFrameRate, int32& Numerator, int32& Denominator);
 	
 	// -- Begin K2 utilities
 
 	/** Converts a byte to a float */
-	UFUNCTION(BlueprintPure, meta=(DisplayName = "ToFloat (byte)", CompactNodeTitle = "->", Keywords="cast convert", BlueprintAutocast), Category="Math|Conversions")
+	UFUNCTION(BlueprintPure, meta=(DisplayName = "To Float (Byte)", CompactNodeTitle = "->", Keywords="cast convert", BlueprintAutocast), Category="Math|Conversions")
 	static float Conv_ByteToFloat(uint8 InByte);
 
 	/** Converts a byte to a double */
-	UFUNCTION(BlueprintPure, meta = (DisplayName = "ToDouble (byte)", CompactNodeTitle = "->", Keywords = "cast convert", BlueprintAutocast), Category = "Math|Conversions")
+	UFUNCTION(BlueprintPure, meta = (DisplayName = "To Double (Byte)", CompactNodeTitle = "->", Keywords = "cast convert", BlueprintAutocast), Category = "Math|Conversions")
 	static double Conv_ByteToDouble(uint8 InByte);
 
 	/** Converts an integer to a float */
-	UFUNCTION(BlueprintPure, meta=(DisplayName = "ToFloat (integer)", CompactNodeTitle = "->", Keywords="cast convert", BlueprintAutocast), Category="Math|Conversions")
+	UFUNCTION(BlueprintPure, meta=(DisplayName = "To Float (Integer)", CompactNodeTitle = "->", Keywords="cast convert", BlueprintAutocast), Category="Math|Conversions")
 	static float Conv_IntToFloat(int32 InInt);
 
 	/** Converts an integer to a double */
-	UFUNCTION(BlueprintPure, meta = (DisplayName = "ToDouble (integer)", CompactNodeTitle = "->", Keywords = "cast convert", BlueprintAutocast), Category = "Math|Conversions")
+	UFUNCTION(BlueprintPure, meta = (DisplayName = "To Double (Integer)", CompactNodeTitle = "->", Keywords = "cast convert", BlueprintAutocast), Category = "Math|Conversions")
 	static double Conv_IntToDouble(int32 InInt);
 
 	/** Converts an integer to a 64 bit integer */
-	UFUNCTION(BlueprintPure, meta = (DisplayName = "ToInt64 (integer)", CompactNodeTitle = "->", Keywords = "cast convert", BlueprintAutocast), Category = "Math|Conversions")
+	UFUNCTION(BlueprintPure, meta = (DisplayName = "To Integer64 (Integer)", CompactNodeTitle = "->", Keywords = "cast convert", BlueprintAutocast), Category = "Math|Conversions")
 	static int64 Conv_IntToInt64(int32 InInt);
 
 	/** Converts an integer to a byte (if the integer is too large, returns the low 8 bits) */
-	UFUNCTION(BlueprintPure, meta=(DisplayName = "ToByte (integer)", CompactNodeTitle = "->", Keywords="cast convert", BlueprintAutocast), Category="Math|Conversions")
+	UFUNCTION(BlueprintPure, meta=(DisplayName = "To Byte (Integer)", CompactNodeTitle = "->", Keywords="cast convert", BlueprintAutocast), Category="Math|Conversions")
 	static uint8 Conv_IntToByte(int32 InInt);
 
 	/** Converts a 64 bit integer to a 32 bit integer (if the integer is too large, returns the low 32 bits) */
-	UFUNCTION(BlueprintPure, meta=(DisplayName = "ToInt (Int64)", CompactNodeTitle = "->", Keywords="cast convert", BlueprintAutocast), Category="Math|Conversions")
+	UFUNCTION(BlueprintPure, meta=(DisplayName = "To Integer (Integer64)", CompactNodeTitle = "->", Keywords="cast convert", BlueprintAutocast), Category="Math|Conversions")
 	static int32 Conv_Int64ToInt(int64 InInt);
 
 	/** Converts a 64 bit floating point to a 32 bit floating point (if the float is too large, returns the low 32 bits) */
-	UFUNCTION(BlueprintPure, meta = (DisplayName = "ToFloat (Double)", CompactNodeTitle = "->", Keywords = "cast convert", BlueprintAutocast), Category = "Math|Conversions")
+	UFUNCTION(BlueprintPure, meta = (DisplayName = "To Float (Double)", CompactNodeTitle = "->", Keywords = "cast convert", BlueprintAutocast), Category = "Math|Conversions")
 	static float Conv_DoubleToFloat(double InDouble);
 
 	/** Converts a 32 bit floating point to a 64 bit floating point */
-	UFUNCTION(BlueprintPure, meta = (DisplayName = "ToDouble (Float)", CompactNodeTitle = "->", Keywords = "cast convert", BlueprintAutocast), Category = "Math|Conversions")
+	UFUNCTION(BlueprintPure, meta = (DisplayName = "To Double (Float)", CompactNodeTitle = "->", Keywords = "cast convert", BlueprintAutocast), Category = "Math|Conversions")
 	static double Conv_FloatToDouble(float InFloat);
 
 	/** Converts a 64 bit integer to a byte (if the integer is too large, returns the low 8 bits) */
-	UFUNCTION(BlueprintPure, meta=(DisplayName = "ToByte (Int64)", CompactNodeTitle = "->", Keywords="cast convert", BlueprintAutocast), Category="Math|Conversions")
+	UFUNCTION(BlueprintPure, meta=(DisplayName = "To Byte (Integer64)", CompactNodeTitle = "->", Keywords="cast convert", BlueprintAutocast), Category="Math|Conversions")
 	static uint8 Conv_Int64ToByte(int64 InInt);
 
 	/** Converts a 64 bit floating point to a 64 bit integer */
-	UFUNCTION(BlueprintPure, meta = (DisplayName = "ToInt64 (Double)", CompactNodeTitle = "->", Keywords = "cast convert", BlueprintAutocast), Category = "Math|Conversions")
+	UFUNCTION(BlueprintPure, meta = (DisplayName = "To Integer64 (Double)", CompactNodeTitle = "->", Keywords = "cast convert", BlueprintAutocast), Category = "Math|Conversions")
 	static int64 Conv_DoubleToInt64(double InDouble);
 
 	/** Converts a 64 bit integer to a 64 bit floating point */
-	UFUNCTION(BlueprintPure, meta = (DisplayName = "ToDouble (Int64)", CompactNodeTitle = "->", Keywords = "cast convert", BlueprintAutocast), Category = "Math|Conversions")
+	UFUNCTION(BlueprintPure, meta = (DisplayName = "To Double (Integer64)", CompactNodeTitle = "->", Keywords = "cast convert", BlueprintAutocast), Category = "Math|Conversions")
 	static double Conv_Int64ToDouble(int64 InInt);
 
 	/** Converts an integer to an IntVector*/
-	UFUNCTION(BlueprintPure, meta = (DisplayName = "ToIntVector (integer)", CompactNodeTitle = "->", Keywords = "cast convert", BlueprintAutocast), Category = "Math|Conversions")
+	UFUNCTION(BlueprintPure, meta = (DisplayName = "To IntVector (Integer)", CompactNodeTitle = "->", Keywords = "cast convert", BlueprintAutocast), Category = "Math|Conversions")
 	static FIntVector Conv_IntToIntVector(int32 InInt);
 
 	/** Converts an integer to a FVector */
-	UFUNCTION(BlueprintPure, meta = (DisplayName = "ToVector (integer)", CompactNodeTitle = "->", Keywords = "cast convert", BlueprintAutocast), Category = "Math|Conversions")
+	UFUNCTION(BlueprintPure, meta = (DisplayName = "To Vector (Integer)", CompactNodeTitle = "->", Keywords = "cast convert", BlueprintAutocast), Category = "Math|Conversions")
 	static FVector Conv_IntToVector(int32 InInt);
 
 	/** Converts a int to a bool*/
-	UFUNCTION(BlueprintPure, meta=(DisplayName = "ToBool (integer)", CompactNodeTitle = "->", Keywords="cast convert", BlueprintAutocast), Category="Math|Conversions")
+	UFUNCTION(BlueprintPure, meta=(DisplayName = "To Boolean (Integer)", CompactNodeTitle = "->", Keywords="cast convert", BlueprintAutocast), Category="Math|Conversions")
 	static bool Conv_IntToBool(int32 InInt);
 
 	/** Converts a bool to an int */
-	UFUNCTION(BlueprintPure, meta=(DisplayName = "ToInt (bool)", CompactNodeTitle = "->", Keywords="cast convert", BlueprintAutocast), Category="Math|Conversions")
+	UFUNCTION(BlueprintPure, meta=(DisplayName = "To Integer (Boolean)", CompactNodeTitle = "->", Keywords="cast convert", BlueprintAutocast), Category="Math|Conversions")
 	static int32 Conv_BoolToInt(bool InBool);
 
 	/** Converts a bool to a float (0.0f or 1.0f) */
-	UFUNCTION(BlueprintPure, meta=(DisplayName = "ToFloat (bool)", CompactNodeTitle = "->", Keywords="cast convert", BlueprintAutocast), Category="Math|Conversions")
+	UFUNCTION(BlueprintPure, meta=(DisplayName = "To Float (Boolean)", CompactNodeTitle = "->", Keywords="cast convert", BlueprintAutocast), Category="Math|Conversions")
 	static float Conv_BoolToFloat(bool InBool);
 
 	/** Converts a bool to a byte */
-	UFUNCTION(BlueprintPure, meta=(DisplayName = "ToByte (bool)", CompactNodeTitle = "->", Keywords="cast convert", BlueprintAutocast), Category="Math|Conversions")
+	UFUNCTION(BlueprintPure, meta=(DisplayName = "To Byte (Boolean)", CompactNodeTitle = "->", Keywords="cast convert", BlueprintAutocast), Category="Math|Conversions")
 	static uint8 Conv_BoolToByte(bool InBool);
 	
 	/** Converts a byte to an integer */
-	UFUNCTION(BlueprintPure, meta=(DisplayName = "ToInt (byte)", CompactNodeTitle = "->", Keywords="cast convert", BlueprintAutocast), Category="Math|Conversions")
+	UFUNCTION(BlueprintPure, meta=(DisplayName = "To Integer (Byte)", CompactNodeTitle = "->", Keywords="cast convert", BlueprintAutocast), Category="Math|Conversions")
 	static int32 Conv_ByteToInt(uint8 InByte);
 
 	/** Converts a byte to an integer */
-	UFUNCTION(BlueprintPure, meta = (DisplayName = "ToInt64 (byte)", CompactNodeTitle = "->", Keywords = "cast convert", BlueprintAutocast), Category = "Math|Conversions")
+	UFUNCTION(BlueprintPure, meta = (DisplayName = "To Integer64 (Byte)", CompactNodeTitle = "->", Keywords = "cast convert", BlueprintAutocast), Category = "Math|Conversions")
 	static int64 Conv_ByteToInt64(uint8 InByte);
 
 	/** Converts a color to LinearColor */
-	UFUNCTION(BlueprintPure, meta = (DisplayName = "ToLinearColor (Color)", CompactNodeTitle = "->", Keywords = "cast convert", BlueprintAutocast), Category = "Math|Conversions")
+	UFUNCTION(BlueprintPure, meta = (DisplayName = "To LinearColor (Color)", CompactNodeTitle = "->", Keywords = "cast convert", BlueprintAutocast), Category = "Math|Conversions")
 	static FLinearColor Conv_ColorToLinearColor(FColor InColor);
 
 	/** Convert an IntVector to a vector */
-	UFUNCTION(BlueprintPure, meta = (DisplayName = "ToVector (IntVector)", CompactNodeTitle = "->", Keywords = "cast convert", BlueprintAutocast), Category = "Math|Conversions")
+	UFUNCTION(BlueprintPure, meta = (DisplayName = "To Vector (IntVector)", CompactNodeTitle = "->", Keywords = "cast convert", BlueprintAutocast), Category = "Math|Conversions")
 	static FVector Conv_IntVectorToVector(const FIntVector& InIntVector);
 
 	/** Convert a float into a vector, where each element is that float */
-	UFUNCTION(BlueprintPure, meta=(DisplayName = "ToVector (float)", CompactNodeTitle = "->", Keywords="cast convert", BlueprintAutocast), Category="Math|Conversions")
+	UFUNCTION(BlueprintPure, meta=(DisplayName = "To Vector (Float)", CompactNodeTitle = "->", Keywords="cast convert", BlueprintAutocast), Category="Math|Conversions")
 	static FVector Conv_FloatToVector(float InFloat);
 
 
@@ -3614,7 +3614,7 @@ class ENGINE_API UKismetMathLibrary : public UBlueprintFunctionLibrary
 	static bool EqualEqual_ObjectObject(class UObject* A, class UObject* B);
 
 	/** Returns true if A and B are not equal (A != B) */
-	UFUNCTION(BlueprintPure, meta=(DisplayName = "NotEqual (Object)", CompactNodeTitle = "!=", Keywords = "!= not equal"), Category="Utilities")
+	UFUNCTION(BlueprintPure, meta=(DisplayName = "Not Equal (Object)", CompactNodeTitle = "!=", Keywords = "!= not equal"), Category="Utilities")
 	static bool NotEqual_ObjectObject(class UObject* A, class UObject* B);
 
 	//
@@ -3626,7 +3626,7 @@ class ENGINE_API UKismetMathLibrary : public UBlueprintFunctionLibrary
 	static bool EqualEqual_ClassClass(class UClass* A, class UClass* B);
 
 	/** Returns true if A and B are not equal (A != B) */
-	UFUNCTION(BlueprintPure, meta=(DisplayName = "NotEqual (Class)", CompactNodeTitle = "!=", Keywords = "!= not equal"), Category="Utilities")
+	UFUNCTION(BlueprintPure, meta=(DisplayName = "Not Equal (Class)", CompactNodeTitle = "!=", Keywords = "!= not equal"), Category="Utilities")
 	static bool NotEqual_ClassClass(class UClass* A, class UClass* B);
 
 	/**
@@ -3648,7 +3648,7 @@ class ENGINE_API UKismetMathLibrary : public UBlueprintFunctionLibrary
 	static bool EqualEqual_NameName(FName A, FName B);
 
 	/** Returns true if A and B are not equal (A != B) */
-	UFUNCTION(BlueprintPure, meta=(DisplayName = "NotEqual (Name)", CompactNodeTitle = "!=", Keywords = "!= not equal"), Category="Utilities|Name")
+	UFUNCTION(BlueprintPure, meta=(DisplayName = "Not Equal (Name)", CompactNodeTitle = "!=", Keywords = "!= not equal"), Category="Utilities|Name")
 	static bool NotEqual_NameName(FName A, FName B);
 
 
@@ -3665,7 +3665,7 @@ class ENGINE_API UKismetMathLibrary : public UBlueprintFunctionLibrary
 	static void BreakTransform(const FTransform& InTransform, FVector& Location, FRotator& Rotation, FVector& Scale);
 
 	/** Returns true if transform A is equal to transform B */
-	UFUNCTION(BlueprintPure, meta = (DisplayName = "Equal Transform", CompactNodeTitle = "==", ScriptMethod = "Equals", ScriptOperator = "==", Keywords = "== equal"), Category="Math|Transform")
+	UFUNCTION(BlueprintPure, meta = (DisplayName = "Equal (Transform)", CompactNodeTitle = "==", ScriptMethod = "Equals", ScriptOperator = "==", Keywords = "== equal"), Category="Math|Transform")
 	static bool EqualEqual_TransformTransform(const FTransform& A, const FTransform& B);
 
 	/** 
diff --git a/Engine/Source/Runtime/Engine/Classes/Kismet/KismetStringLibrary.h b/Engine/Source/Runtime/Engine/Classes/Kismet/KismetStringLibrary.h
--- a/Engine/Source/Runtime/Engine/Classes/Kismet/KismetStringLibrary.h
+++ b/Engine/Source/Runtime/Engine/Classes/Kismet/KismetStringLibrary.h
@@ -13,67 +13,67 @@ class ENGINE_API UKismetStringLibrary : public UBlueprintFunctionLibrary
 	GENERATED_UCLASS_BODY()
 
 	/** Converts a float value to a string */
-	UFUNCTION(BlueprintPure, meta=(DisplayName = "ToString (float)", CompactNodeTitle = "->", BlueprintAutocast), Category="Utilities|String")
+	UFUNCTION(BlueprintPure, meta=(DisplayName = "To String (Float)", CompactNodeTitle = "->", BlueprintAutocast), Category="Utilities|String")
 	static FString Conv_FloatToString(float InFloat);
 
 	/** Converts a double value to a string */
-	UFUNCTION(BlueprintPure, meta = (DisplayName = "ToString (double)", CompactNodeTitle = "->", BlueprintAutocast), Category = "Utilities|String")
+	UFUNCTION(BlueprintPure, meta = (DisplayName = "To String (Double)", CompactNodeTitle = "->", BlueprintAutocast), Category = "Utilities|String")
 	static FString Conv_DoubleToString(double InDouble);
 
 	/** Converts an integer value to a string */
-	UFUNCTION(BlueprintPure, meta=(DisplayName = "ToString (integer)", CompactNodeTitle = "->", BlueprintAutocast), Category="Utilities|String")
+	UFUNCTION(BlueprintPure, meta=(DisplayName = "To String (Integer)", CompactNodeTitle = "->", BlueprintAutocast), Category="Utilities|String")
 	static FString Conv_IntToString(int32 InInt);
 
 	/** Converts an 64-bit integer value to a string */
-	UFUNCTION(BlueprintPure, meta = (DisplayName = "ToString (Int64)", CompactNodeTitle = "->", BlueprintAutocast), Category = "Utilities|String")
+	UFUNCTION(BlueprintPure, meta = (DisplayName = "To String (Integer64)", CompactNodeTitle = "->", BlueprintAutocast), Category = "Utilities|String")
 	static FString Conv_Int64ToString(int64 InInt);
 	
 	/** Converts a byte value to a string */
-	UFUNCTION(BlueprintPure, meta=(DisplayName = "ToString (byte)", CompactNodeTitle = "->", BlueprintAutocast), Category="Utilities|String")
+	UFUNCTION(BlueprintPure, meta=(DisplayName = "To String (Byte)", CompactNodeTitle = "->", BlueprintAutocast), Category="Utilities|String")
 	static FString Conv_ByteToString(uint8 InByte);
 
 	/** Converts a boolean value to a string, either 'true' or 'false' */
-	UFUNCTION(BlueprintPure, meta=(DisplayName = "ToString (boolean)", CompactNodeTitle = "->", BlueprintAutocast), Category="Utilities|String")
+	UFUNCTION(BlueprintPure, meta=(DisplayName = "To String (Boolean)", CompactNodeTitle = "->", BlueprintAutocast), Category="Utilities|String")
 	static FString Conv_BoolToString(bool InBool);
 
 	/** Converts a vector value to a string, in the form 'X= Y= Z=' */
-	UFUNCTION(BlueprintPure, meta=(DisplayName = "ToString (Vector)", CompactNodeTitle = "->", BlueprintAutocast), Category="Utilities|String")
+	UFUNCTION(BlueprintPure, meta=(DisplayName = "To String (Vector)", CompactNodeTitle = "->", BlueprintAutocast), Category="Utilities|String")
 	static FString Conv_VectorToString(FVector InVec);
 
 	/** Converts an IntVector value to a string, in the form 'X= Y= Z=' */
-	UFUNCTION(BlueprintPure, meta=(DisplayName = "ToString (IntVector)", CompactNodeTitle = "->", BlueprintAutocast), Category="Utilities|String")
+	UFUNCTION(BlueprintPure, meta=(DisplayName = "To String (IntVector)", CompactNodeTitle = "->", BlueprintAutocast), Category="Utilities|String")
 	static FString Conv_IntVectorToString(FIntVector InIntVec);
 
 	/** Converts an IntPoint value to a string, in the form 'X= Y=' */
-	UFUNCTION(BlueprintPure, meta = (DisplayName = "ToString (IntPoint)", CompactNodeTitle = "->", BlueprintAutocast), Category = "Utilities|String")
+	UFUNCTION(BlueprintPure, meta = (DisplayName = "To String (IntPoint)", CompactNodeTitle = "->", BlueprintAutocast), Category = "Utilities|String")
 	static FString Conv_IntPointToString(FIntPoint InIntPoint);
 
 	/** Converts a vector2d value to a string, in the form 'X= Y=' */
-	UFUNCTION(BlueprintPure, meta=(DisplayName = "ToString (vector2d)", CompactNodeTitle = "->", BlueprintAutocast), Category="Utilities|String")
+	UFUNCTION(BlueprintPure, meta=(DisplayName = "To String (Vector2d)", CompactNodeTitle = "->", BlueprintAutocast), Category="Utilities|String")
 	static FString Conv_Vector2dToString(FVector2D InVec);
 
 	/** Converts a rotator value to a string, in the form 'P= Y= R=' */
-	UFUNCTION(BlueprintPure, meta=(DisplayName = "ToString (rotator)", CompactNodeTitle = "->", BlueprintAutocast), Category="Utilities|String")
+	UFUNCTION(BlueprintPure, meta=(DisplayName = "To String (Rotator)", CompactNodeTitle = "->", BlueprintAutocast), Category="Utilities|String")
 	static FString Conv_RotatorToString(FRotator InRot);
 
 	/** Converts a transform value to a string, in the form 'Translation: X= Y= Z= Rotation: P= Y= R= Scale: X= Y= Z=' */
-	UFUNCTION(BlueprintPure, meta=(DisplayName = "ToString (transform)", CompactNodeTitle = "->", BlueprintAutocast), Category="Utilities|String")
+	UFUNCTION(BlueprintPure, meta=(DisplayName = "To String (Transform)", CompactNodeTitle = "->", BlueprintAutocast), Category="Utilities|String")
 	static FString Conv_TransformToString(const FTransform& InTrans);
 
 	/** Converts a UObject value to a string by calling the object's GetName method */
-	UFUNCTION(BlueprintPure, meta=(DisplayName = "ToString (object)", CompactNodeTitle = "->", BlueprintAutocast), Category="Utilities|String")
+	UFUNCTION(BlueprintPure, meta=(DisplayName = "To String (Object)", CompactNodeTitle = "->", BlueprintAutocast), Category="Utilities|String")
 	static FString Conv_ObjectToString(class UObject* InObj);
 
 	/** Converts a linear color value to a string, in the form '(R=,G=,B=,A=)' */
-	UFUNCTION(BlueprintPure, meta=(DisplayName = "ToString (linear color)", CompactNodeTitle = "->", BlueprintAutocast), Category="Utilities|String")
+	UFUNCTION(BlueprintPure, meta=(DisplayName = "To String (LinearColor)", CompactNodeTitle = "->", BlueprintAutocast), Category="Utilities|String")
 	static FString Conv_ColorToString(FLinearColor InColor);
 
 	/** Converts a name value to a string */
-	UFUNCTION(BlueprintPure, meta=(DisplayName = "ToString (name)", CompactNodeTitle = "->", BlueprintAutocast), Category="Utilities|String")
+	UFUNCTION(BlueprintPure, meta=(DisplayName = "To String (Name)", CompactNodeTitle = "->", BlueprintAutocast), Category="Utilities|String")
 	static FString Conv_NameToString(FName InName);
 
 	/** Converts a name value to a string */
-	UFUNCTION(BlueprintPure, meta = (DisplayName = "ToString (Matrix)", CompactNodeTitle = "->", BlueprintAutocast), Category = "Utilities|String")
+	UFUNCTION(BlueprintPure, meta = (DisplayName = "To String (Matrix)", CompactNodeTitle = "->", BlueprintAutocast), Category = "Utilities|String")
 	static FString Conv_MatrixToString(const FMatrix& InMatrix);
 
 	/** Converts a string to a name value */
@@ -81,7 +81,7 @@ class ENGINE_API UKismetStringLibrary : public UBlueprintFunctionLibrary
 	static FName Conv_StringToName(const FString& InString);
 
 	/** Converts a string to a int value */
-	UFUNCTION(BlueprintPure, meta=(DisplayName = "String To Int", CompactNodeTitle = "->", BlueprintAutocast), Category="Utilities|String")
+	UFUNCTION(BlueprintPure, meta=(DisplayName = "String To Integer", CompactNodeTitle = "->", BlueprintAutocast), Category="Utilities|String")
 	static int32 Conv_StringToInt(const FString& InString);
 
 	/** Converts a string to a float value */
@@ -89,19 +89,19 @@ class ENGINE_API UKismetStringLibrary : public UBlueprintFunctionLibrary
 	static float Conv_StringToFloat(const FString& InString);
 	
 	/** Convert String Back To Vector. IsValid indicates whether or not the string could be successfully converted. */
-	UFUNCTION(BlueprintPure, meta=(DisplayName = "String to Vector", CompactNodeTitle = "->"), Category="Utilities|String")
+	UFUNCTION(BlueprintPure, meta=(DisplayName = "String To Vector", CompactNodeTitle = "->"), Category="Utilities|String")
 	static void Conv_StringToVector(const FString& InString, FVector& OutConvertedVector, bool& OutIsValid);
 
 	/** Convert String Back To Vector2D. IsValid indicates whether or not the string could be successfully converted. */
-	UFUNCTION(BlueprintPure, meta=(DisplayName = "String to Vector2D", CompactNodeTitle = "->"), Category="Utilities|String")
+	UFUNCTION(BlueprintPure, meta=(DisplayName = "String To Vector2D", CompactNodeTitle = "->"), Category="Utilities|String")
 	static void Conv_StringToVector2D(const FString& InString, FVector2D& OutConvertedVector2D, bool& OutIsValid);
 
 	/** Convert String Back To Rotator. IsValid indicates whether or not the string could be successfully converted. */
-	UFUNCTION(BlueprintPure, meta=(DisplayName = "String to Rotator", CompactNodeTitle = "->"), Category="Utilities|String")
+	UFUNCTION(BlueprintPure, meta=(DisplayName = "String To Rotator", CompactNodeTitle = "->"), Category="Utilities|String")
 	static void Conv_StringToRotator(const FString& InString, FRotator& OutConvertedRotator, bool& OutIsValid);
 
 	/** Convert String Back To Color. IsValid indicates whether or not the string could be successfully converted. */
-	UFUNCTION(BlueprintPure, meta=(DisplayName = "String to Color", CompactNodeTitle = "->"), Category="Utilities|String")
+	UFUNCTION(BlueprintPure, meta=(DisplayName = "String To Color", CompactNodeTitle = "->"), Category="Utilities|String")
 	static void Conv_StringToColor(const FString& InString, FLinearColor& OutConvertedColor, bool& OutIsValid);
 
 	/** 
@@ -112,7 +112,7 @@ class ENGINE_API UKismetStringLibrary : public UBlueprintFunctionLibrary
 	 * @param Suffix - A suffix to append to the end of the conversion string
 	 * @return A new string built from the passed parameters
 	 */
-	UFUNCTION(BlueprintPure, meta=(DisplayName = "BuildString (float)"), Category="Utilities|String")
+	UFUNCTION(BlueprintPure, meta=(DisplayName = "Build String (Float)"), Category="Utilities|String")
 	static FString BuildString_Float(const FString& AppendTo, const FString& Prefix, float InFloat, const FString& Suffix);
 
 	/** 
@@ -123,7 +123,7 @@ class ENGINE_API UKismetStringLibrary : public UBlueprintFunctionLibrary
 	 * @param Suffix - A suffix to append to the end of the conversion string
 	 * @return A new string built from the passed parameters
 	 */
-	UFUNCTION(BlueprintPure, meta=(DisplayName = "BuildString (integer)"), Category="Utilities|String")
+	UFUNCTION(BlueprintPure, meta=(DisplayName = "Build String (Integer)"), Category="Utilities|String")
 	static FString BuildString_Int(const FString& AppendTo, const FString& Prefix, int32 InInt, const FString& Suffix);
 
 	/** 
@@ -134,7 +134,7 @@ class ENGINE_API UKismetStringLibrary : public UBlueprintFunctionLibrary
 	 * @param Suffix - A suffix to append to the end of the conversion string
 	 * @return A new string built from the passed parameters
 	 */
-	UFUNCTION(BlueprintPure, meta=(DisplayName = "BuildString (boolean)"), Category="Utilities|String")
+	UFUNCTION(BlueprintPure, meta=(DisplayName = "Build String (Boolean)"), Category="Utilities|String")
 	static FString BuildString_Bool(const FString& AppendTo, const FString& Prefix, bool InBool, const FString& Suffix);
 
 	/** 
@@ -145,7 +145,7 @@ class ENGINE_API UKismetStringLibrary : public UBlueprintFunctionLibrary
 	 * @param Suffix - A suffix to append to the end of the conversion string
 	 * @return A new string built from the passed parameters
 	 */
-	UFUNCTION(BlueprintPure, meta=(DisplayName = "BuildString (vector)"), Category="Utilities|String")
+	UFUNCTION(BlueprintPure, meta=(DisplayName = "Build String (Vector)"), Category="Utilities|String")
 	static FString BuildString_Vector(const FString& AppendTo, const FString& Prefix, FVector InVector, const FString& Suffix);
 
 	/**
@@ -156,7 +156,7 @@ class ENGINE_API UKismetStringLibrary : public UBlueprintFunctionLibrary
 	 * @param Suffix - A suffix to append to the end of the conversion string
 	 * @return A new string built from the passed parameters
 	 */
-	UFUNCTION(BlueprintPure, meta = (DisplayName = "BuildString (IntVector)"), Category = "Utilities|String")
+	UFUNCTION(BlueprintPure, meta = (DisplayName = "Build String (IntVector)"), Category = "Utilities|String")
 	static FString BuildString_IntVector(const FString& AppendTo, const FString& Prefix, FIntVector InIntVector, const FString& Suffix);
 	/** 
 	 * Converts a vector2d->string, creating a new string in the form AppendTo+Prefix+InVector2d+Suffix
@@ -166,7 +166,7 @@ class ENGINE_API UKismetStringLibrary : public UBlueprintFunctionLibrary
 	 * @param Suffix - A suffix to append to the end of the conversion string
 	 * @return A new string built from the passed parameters
 	 */
-	UFUNCTION(BlueprintPure, meta=(DisplayName = "BuildString (vector2d)"), Category="Utilities|String")
+	UFUNCTION(BlueprintPure, meta=(DisplayName = "Build String (Vector2D)"), Category="Utilities|String")
 	static FString BuildString_Vector2d(const FString& AppendTo, const FString& Prefix, FVector2D InVector2d, const FString& Suffix);
 
 	/** 
@@ -177,7 +177,7 @@ class ENGINE_API UKismetStringLibrary : public UBlueprintFunctionLibrary
 	 * @param Suffix - A suffix to append to the end of the conversion string
 	 * @return A new string built from the passed parameters
 	 */
-	UFUNCTION(BlueprintPure, meta=(DisplayName = "BuildString (rotator)"), Category="Utilities|String")
+	UFUNCTION(BlueprintPure, meta=(DisplayName = "Build String (Rotator)"), Category="Utilities|String")
 	static FString BuildString_Rotator(const FString& AppendTo, const FString& Prefix, FRotator InRot, const FString& Suffix);
 
 	/** 
@@ -188,7 +188,7 @@ class ENGINE_API UKismetStringLibrary : public UBlueprintFunctionLibrary
 	 * @param Suffix - A suffix to append to the end of the conversion string
 	 * @return A new string built from the passed parameters
 	 */
-	UFUNCTION(BlueprintPure, meta=(DisplayName = "BuildString (object)"), Category="Utilities|String")
+	UFUNCTION(BlueprintPure, meta=(DisplayName = "Build String (Object)"), Category="Utilities|String")
 	static FString BuildString_Object(const FString& AppendTo, const FString& Prefix, class UObject* InObj, const FString& Suffix);
 
 	/** 
@@ -199,7 +199,7 @@ class ENGINE_API UKismetStringLibrary : public UBlueprintFunctionLibrary
 	 * @param Suffix - A suffix to append to the end of the conversion string
 	 * @return A new string built from the passed parameters
 	 */
-	UFUNCTION(BlueprintPure, meta=(DisplayName = "BuildString (color)"), Category="Utilities|String")
+	UFUNCTION(BlueprintPure, meta=(DisplayName = "Build String (LinearColor)"), Category="Utilities|String")
 	static FString BuildString_Color(const FString& AppendTo, const FString& Prefix, FLinearColor InColor, const FString& Suffix);
 
 	/** 
@@ -210,7 +210,7 @@ class ENGINE_API UKismetStringLibrary : public UBlueprintFunctionLibrary
 	 * @param Suffix - A suffix to append to the end of the conversion string
 	 * @return A new string built from the passed parameters
 	 */
-	UFUNCTION(BlueprintPure, meta=(DisplayName = "BuildString (name)"), Category="Utilities|String")
+	UFUNCTION(BlueprintPure, meta=(DisplayName = "Build String (Name)"), Category="Utilities|String")
 	static FString BuildString_Name(const FString& AppendTo, const FString& Prefix, FName InName, const FString& Suffix);
 
 
@@ -233,7 +233,7 @@ class ENGINE_API UKismetStringLibrary : public UBlueprintFunctionLibrary
 	 * @param B - The string to compare
 	 * @returns True if the strings are equal, false otherwise
 	 */
-	UFUNCTION(BlueprintPure, meta=(DisplayName = "Equal (string)", CompactNodeTitle = "=="), Category="Utilities|String")
+	UFUNCTION(BlueprintPure, meta=(DisplayName = "Equal Exactly (String)", CompactNodeTitle = "==="), Category="Utilities|String")
 	static bool EqualEqual_StrStr(const FString& A, const FString& B);
 
 	/**
@@ -242,7 +242,7 @@ class ENGINE_API UKismetStringLibrary : public UBlueprintFunctionLibrary
 	 * @param B - The string to compare
 	 * @returns True if the strings are equal, false otherwise
 	 */
-	UFUNCTION(BlueprintPure, meta=(DisplayName = "Equal, Case Insensitive (string)", CompactNodeTitle = "=="), Category="Utilities|String")
+	UFUNCTION(BlueprintPure, meta=(DisplayName = "Equal, Case Insensitive (String)", CompactNodeTitle = "=="), Category="Utilities|String")
 	static bool EqualEqual_StriStri(const FString& A, const FString& B);
 
 	/** 
@@ -251,7 +251,7 @@ class ENGINE_API UKismetStringLibrary : public UBlueprintFunctionLibrary
 	 * @param B - The string to compare
 	 * @return Returns true if the input strings are not equal, false if they are equal
 	 */
-	UFUNCTION(BlueprintPure, meta=(DisplayName = "NotEqual (string)", CompactNodeTitle = "!="), Category="Utilities|String")
+	UFUNCTION(BlueprintPure, meta=(DisplayName = "Not Equal Exactly (String)", CompactNodeTitle = "!=="), Category="Utilities|String")
 	static bool NotEqual_StrStr(const FString& A, const FString& B);
 
 	/** Test if the input string are not equal (A != B), ignoring case differences
@@ -259,7 +259,7 @@ class ENGINE_API UKismetStringLibrary : public UBlueprintFunctionLibrary
 	 * @param B - The string to compare
 	 * @return Returns true if the input strings are not equal, false if they are equal
 	 */
-	UFUNCTION(BlueprintPure, meta=(DisplayName = "NotEqual, case insenstive (string)", CompactNodeTitle = "!="), Category="Utilities|String")
+	UFUNCTION(BlueprintPure, meta=(DisplayName = "Not Equal, Case Insenstive (String)", CompactNodeTitle = "!="), Category="Utilities|String")
 	static bool NotEqual_StriStri(const FString& A, const FString& B);
 
 	/** 
diff --git a/Engine/Source/Runtime/Engine/Classes/Kismet/KismetSystemLibrary.h b/Engine/Source/Runtime/Engine/Classes/Kismet/KismetSystemLibrary.h
--- a/Engine/Source/Runtime/Engine/Classes/Kismet/KismetSystemLibrary.h
+++ b/Engine/Source/Runtime/Engine/Classes/Kismet/KismetSystemLibrary.h
@@ -107,7 +107,7 @@ class ENGINE_API UKismetSystemLibrary : public UBlueprintFunctionLibrary
 
 	// Returns the full system path to a UObject
 	// If given a non-asset UObject, it will return an empty string
-	UFUNCTION(BlueprintPure, Category = "Utilities")
+	UFUNCTION(BlueprintPure, Category = "Utilities|Paths")
 	static FString GetSystemPath(const UObject* Object);
 
 	// Returns the display name (or actor label), for displaying as a debugging aid.
@@ -117,7 +117,7 @@ class ENGINE_API UKismetSystemLibrary : public UBlueprintFunctionLibrary
 	static FString GetDisplayName(const UObject* Object);
 
 	// Returns the display name of a class
-	UFUNCTION(BlueprintPure, Category = "Utilities", meta = (DisplayName = "Get Display Name"))
+	UFUNCTION(BlueprintPure, Category = "Utilities", meta = (DisplayName = "Get Class Display Name"))
 	static FString GetClassDisplayName(UClass* Class);
 
 	// Returns the outer object of an object.
@@ -172,6 +172,7 @@ class ENGINE_API UKismetSystemLibrary : public UBlueprintFunctionLibrary
 	UFUNCTION(BlueprintPure, Category = "Utilities|Platform")
 	static FString GetPlatformUserDir();
 
+	/** Checks if this object implements a specific interface, works for both native and blueprint interfacse */
 	UFUNCTION(BlueprintPure, Category="Utilities")
 	static bool DoesImplementInterface(const UObject* TestObject, TSubclassOf<UInterface> Interface);
 
@@ -200,7 +201,7 @@ class ENGINE_API UKismetSystemLibrary : public UBlueprintFunctionLibrary
 	static bool IsStandalone(const UObject* WorldContextObject);
 	
 	/** Returns whether we're currently running in split screen (more than one local player). */
-	UFUNCTION(BlueprintPure, Category = "Utilities", meta = (WorldContext = "WorldContextObject"))
+	UFUNCTION(BlueprintPure, Category = "Viewport", meta = (WorldContext = "WorldContextObject"))
 	static bool IsSplitScreen(const UObject* WorldContextObject);
 
 	/** Returns whether this is a build that is packaged for distribution */
@@ -216,31 +217,31 @@ class ENGINE_API UKismetSystemLibrary : public UBlueprintFunctionLibrary
 	static FString GetDeviceId();
 
 	/** Converts an interfance into an object */
-	UFUNCTION(BlueprintPure, meta=(DisplayName = "ToObject (interface)", CompactNodeTitle = "->"), Category="Utilities")
+	UFUNCTION(BlueprintPure, meta=(DisplayName = "To Object (Interface)", CompactNodeTitle = "->"), Category="Utilities")
 	static UObject* Conv_InterfaceToObject(const FScriptInterface& Interface); 
 
 	/** Builds a SoftObjectPath struct. Generally you should be using Soft Object References/Ptr types instead */
-	UFUNCTION(BlueprintPure, Category = "SoftObjectPath", meta = (Keywords = "construct build", NativeMakeFunc, BlueprintThreadSafe))
+	UFUNCTION(BlueprintPure, Category = "Utilities", meta = (Keywords = "construct build", NativeMakeFunc, BlueprintThreadSafe))
 	static FSoftObjectPath MakeSoftObjectPath(const FString& PathString);
 
 	/** Gets the path string out of a Soft Object Path */
-	UFUNCTION(BlueprintPure, Category = "SoftObjectPath", meta = (NativeBreakFunc, BlueprintThreadSafe))
+	UFUNCTION(BlueprintPure, Category = "Utilities", meta = (NativeBreakFunc, BlueprintThreadSafe))
 	static void BreakSoftObjectPath(FSoftObjectPath InSoftObjectPath, FString& PathString);
 
 	/** Converts a Soft Object Path into a base Soft Object Reference, this is not guaranteed to be resolvable */
-	UFUNCTION(BlueprintPure, meta = (DisplayName = "ToSoftObjectReference (SoftObjectPath)", CompactNodeTitle = "->"), Category = "Utilities")
+	UFUNCTION(BlueprintPure, meta = (DisplayName = "To Soft Object Reference", CompactNodeTitle = "->"), Category = "Utilities")
 	static TSoftObjectPtr<UObject> Conv_SoftObjPathToSoftObjRef(const FSoftObjectPath& SoftObjectPath);
 
 	/** Builds a SoftClassPath struct. Generally you should be using Soft Class References/Ptr types instead */
-	UFUNCTION(BlueprintPure, Category = "SoftClassPath", meta = (Keywords = "construct build", NativeMakeFunc, BlueprintThreadSafe))
+	UFUNCTION(BlueprintPure, Category = "Utilities", meta = (Keywords = "construct build", NativeMakeFunc, BlueprintThreadSafe))
 	static FSoftClassPath MakeSoftClassPath(const FString& PathString);
 
 	/** Gets the path string out of a Soft Class Path */
-	UFUNCTION(BlueprintPure, Category = "SoftClassPath", meta = (NativeBreakFunc, BlueprintThreadSafe))
+	UFUNCTION(BlueprintPure, Category = "Utilities", meta = (NativeBreakFunc, BlueprintThreadSafe))
 	static void BreakSoftClassPath(FSoftClassPath InSoftClassPath, FString& PathString);
 
 	/** Converts a Soft Class Path into a base Soft Class Reference, this is not guaranteed to be resolvable */
-	UFUNCTION(BlueprintPure, meta = (DisplayName = "ToSoftClassReference (SoftClassPath)", CompactNodeTitle = "->"), Category = "Utilities")
+	UFUNCTION(BlueprintPure, meta = (DisplayName = "To Soft Class Reference", CompactNodeTitle = "->"), Category = "Utilities")
 	static TSoftClassPtr<UObject> Conv_SoftClassPathToSoftClassRef(const FSoftClassPath& SoftClassPath);
 
 	/** Returns true if the Soft Object Reference is not null */
@@ -248,7 +249,7 @@ class ENGINE_API UKismetSystemLibrary : public UBlueprintFunctionLibrary
 	static bool IsValidSoftObjectReference(const TSoftObjectPtr<UObject>& SoftObjectReference);
 
 	/** Converts a Soft Object Reference to a string. The other direction is not provided because it cannot be validated */
-	UFUNCTION(BlueprintPure, meta = (DisplayName = "ToString (SoftObjectReference)", CompactNodeTitle = "->", BlueprintThreadSafe), Category = "Utilities")
+	UFUNCTION(BlueprintPure, meta = (DisplayName = "To String (SoftObjectReference)", CompactNodeTitle = "->", BlueprintThreadSafe), Category = "Utilities")
 	static FString Conv_SoftObjectReferenceToString(const TSoftObjectPtr<UObject>& SoftObjectReference);
 
 	/** Returns true if the values are equal (A == B) */
@@ -256,7 +257,7 @@ class ENGINE_API UKismetSystemLibrary : public UBlueprintFunctionLibrary
 	static bool EqualEqual_SoftObjectReference(const TSoftObjectPtr<UObject>& A, const TSoftObjectPtr<UObject>& B);
 
 	/** Returns true if the values are not equal (A != B) */
-	UFUNCTION(BlueprintPure, meta = (DisplayName = "NotEqual (SoftObjectReference)", CompactNodeTitle = "!=", BlueprintThreadSafe), Category = "Utilities")
+	UFUNCTION(BlueprintPure, meta = (DisplayName = "Not Equal (SoftObjectReference)", CompactNodeTitle = "!=", BlueprintThreadSafe), Category = "Utilities")
 	static bool NotEqual_SoftObjectReference(const TSoftObjectPtr<UObject>& A, const TSoftObjectPtr<UObject>& B);
 
 	/** Resolves or loads a Soft Object Reference immediately, this will cause hitches and Async Load Asset should be used if possible */
@@ -268,7 +269,7 @@ class ENGINE_API UKismetSystemLibrary : public UBlueprintFunctionLibrary
 	static bool IsValidSoftClassReference(const TSoftClassPtr<UObject>& SoftClassReference);
 
 	/** Converts a Soft Class Reference to a string. The other direction is not provided because it cannot be validated */
-	UFUNCTION(BlueprintPure, meta = (DisplayName = "ToString (SoftObjectReference)", CompactNodeTitle = "->", BlueprintThreadSafe), Category = "Utilities")
+	UFUNCTION(BlueprintPure, meta = (DisplayName = "To String (SoftClassReference)", CompactNodeTitle = "->", BlueprintThreadSafe), Category = "Utilities")
 	static FString Conv_SoftClassReferenceToString(const TSoftClassPtr<UObject>& SoftClassReference);
 
 	/** Returns true if the values are equal (A == B) */
@@ -276,7 +277,7 @@ class ENGINE_API UKismetSystemLibrary : public UBlueprintFunctionLibrary
 	static bool EqualEqual_SoftClassReference(const TSoftClassPtr<UObject>& A, const TSoftClassPtr<UObject>& B);
 
 	/** Returns true if the values are not equal (A != B) */
-	UFUNCTION(BlueprintPure, meta = (DisplayName = "NotEqual (SoftClassReference)", CompactNodeTitle = "!=", BlueprintThreadSafe), Category = "Utilities")
+	UFUNCTION(BlueprintPure, meta = (DisplayName = "Not Equal (SoftClassReference)", CompactNodeTitle = "!=", BlueprintThreadSafe), Category = "Utilities")
 	static bool NotEqual_SoftClassReference(const TSoftClassPtr<UObject>& A, const TSoftClassPtr<UObject>& B);
 
 	/** Resolves or loads a Soft Class Reference immediately, this will cause hitches and Async Load Class Asset should be used if possible */
@@ -390,7 +391,7 @@ class ENGINE_API UKismetSystemLibrary : public UBlueprintFunctionLibrary
 	 * @param	TextColor		Whether or not to print the output to the console
 	 * @param	Duration		The display duration (if Print to Screen is True). Using negative number will result in loading the duration time from the config.
 	 */
-	UFUNCTION(BlueprintCallable, meta=(WorldContext="WorldContextObject", CallableWithoutWorldContext, Keywords = "log print", AdvancedDisplay = "2", DevelopmentOnly), Category="Utilities|String")
+	UFUNCTION(BlueprintCallable, meta=(WorldContext="WorldContextObject", CallableWithoutWorldContext, Keywords = "log print", AdvancedDisplay = "2", DevelopmentOnly), Category="Development")
 	static void PrintString(const UObject* WorldContextObject, const FString& InString = FString(TEXT("Hello")), bool bPrintToScreen = true, bool bPrintToLog = true, FLinearColor TextColor = FLinearColor(0.0, 0.66, 1.0), float Duration = 2.f);
 
 	/**
@@ -404,7 +405,7 @@ class ENGINE_API UKismetSystemLibrary : public UBlueprintFunctionLibrary
 	 * @param	TextColor		Whether or not to print the output to the console
 	 * @param	Duration		The display duration (if Print to Screen is True). Using negative number will result in loading the duration time from the config.
 	 */
-	UFUNCTION(BlueprintCallable, meta=(WorldContext="WorldContextObject", CallableWithoutWorldContext, Keywords = "log", AdvancedDisplay = "2", DevelopmentOnly), Category="Utilities|Text")
+	UFUNCTION(BlueprintCallable, meta=(WorldContext="WorldContextObject", CallableWithoutWorldContext, Keywords = "log", AdvancedDisplay = "2", DevelopmentOnly), Category="Development")
 	static void PrintText(const UObject* WorldContextObject, const FText InText = INVTEXT("Hello"), bool bPrintToScreen = true, bool bPrintToLog = true, FLinearColor TextColor = FLinearColor(0.0, 0.66, 1.0), float Duration = 2.f);
 
 	/**
@@ -418,7 +419,7 @@ class ENGINE_API UKismetSystemLibrary : public UBlueprintFunctionLibrary
 	static void PrintWarning(const FString& InString);
 
 	/** Sets the game window title */
-	UFUNCTION(BlueprintCallable, Category = "Utilities")
+	UFUNCTION(BlueprintCallable, Category = "Utilities|Platform")
 	static void SetWindowTitle(const FText& Title);
 
 	/**
@@ -611,7 +612,7 @@ class ENGINE_API UKismetSystemLibrary : public UBlueprintFunctionLibrary
 	 * @param Handle		The handle of the timer to check validity of.
 	 * @return				Whether the timer handle is valid.
 	 */
-	UFUNCTION(BlueprintPure, meta=(DisplayName = "Is Valid", ScriptName = "IsValidTimerHandle"), Category="Utilities|Time")
+	UFUNCTION(BlueprintPure, meta=(DisplayName = "Is Valid Timer Handle", ScriptName = "IsValidTimerHandle"), Category="Utilities|Time")
 	static bool K2_IsValidTimerHandle(FTimerHandle Handle);
 
 	/**
@@ -619,7 +620,7 @@ class ENGINE_API UKismetSystemLibrary : public UBlueprintFunctionLibrary
 	 * @param Handle		The handle of the timer to invalidate.
 	 * @return				Return the invalidated timer handle for convenience.
 	 */
-	UFUNCTION(BlueprintCallable, meta=(DisplayName = "Invalidate", ScriptName = "InvalidateTimerHandle"), Category="Utilities|Time")
+	UFUNCTION(BlueprintCallable, meta=(DisplayName = "Invalidate Timer Handle", ScriptName = "InvalidateTimerHandle"), Category="Utilities|Time")
 	static FTimerHandle K2_InvalidateTimerHandle(UPARAM(ref) FTimerHandle& Handle);
 
 	/**
@@ -819,7 +820,7 @@ class ENGINE_API UKismetSystemLibrary : public UBlueprintFunctionLibrary
 	static void SetClassPropertyByName(UObject* Object, FName PropertyName, TSubclassOf<UObject> Value);
 
 	/** Set an INTERFACE property by name */
-	UFUNCTION(BlueprintCallable, Category = "Collision", meta = (BlueprintInternalUseOnly = "true"))
+	UFUNCTION(BlueprintCallable, meta = (BlueprintInternalUseOnly = "true"))
 	static void SetInterfacePropertyByName(UObject* Object, FName PropertyName, const FScriptInterface& Value);
 
 	/** Set a NAME property by name */
@@ -918,7 +919,7 @@ class ENGINE_API UKismetSystemLibrary : public UBlueprintFunctionLibrary
 	 * @param OutActors		Returned array of actors. Unsorted.
 	 * @return				true if there was an overlap that passed the filters, false otherwise.
 	 */
-	UFUNCTION(BlueprintCallable, Category="Collision", meta=(WorldContext="WorldContextObject", AutoCreateRefTerm="ActorsToIgnore", DisplayName = "SphereOverlapActors"))
+	UFUNCTION(BlueprintCallable, Category="Collision", meta=(WorldContext="WorldContextObject", AutoCreateRefTerm="ActorsToIgnore", DisplayName = "Sphere Overlap Actors"))
 	static bool SphereOverlapActors(const UObject* WorldContextObject, const FVector SpherePos, float SphereRadius, const TArray<TEnumAsByte<EObjectTypeQuery> > & ObjectTypes, UClass* ActorClassFilter, const TArray<AActor*>& ActorsToIgnore, TArray<class AActor*>& OutActors);
 
 	/**
@@ -932,7 +933,7 @@ class ENGINE_API UKismetSystemLibrary : public UBlueprintFunctionLibrary
 	 * @param OutActors		Returned array of actors. Unsorted.
 	 * @return				true if there was an overlap that passed the filters, false otherwise.
 	 */
-	UFUNCTION(BlueprintCallable, Category="Collision", meta=(WorldContext="WorldContextObject", AutoCreateRefTerm="ActorsToIgnore", DisplayName="SphereOverlapComponents"))
+	UFUNCTION(BlueprintCallable, Category="Collision", meta=(WorldContext="WorldContextObject", AutoCreateRefTerm="ActorsToIgnore", DisplayName="Sphere Overlap Components"))
 	static bool SphereOverlapComponents(const UObject* WorldContextObject, const FVector SpherePos, float SphereRadius, const TArray<TEnumAsByte<EObjectTypeQuery> > & ObjectTypes, UClass* ComponentClassFilter, const TArray<AActor*>& ActorsToIgnore, TArray<class UPrimitiveComponent*>& OutComponents);
 	
 
@@ -947,7 +948,7 @@ class ENGINE_API UKismetSystemLibrary : public UBlueprintFunctionLibrary
 	 * @param OutActors		Returned array of actors. Unsorted.
 	 * @return				true if there was an overlap that passed the filters, false otherwise.
 	 */
-	UFUNCTION(BlueprintCallable, Category="Collision", meta=(WorldContext="WorldContextObject", AutoCreateRefTerm="ActorsToIgnore", DisplayName="BoxOverlapActors"))
+	UFUNCTION(BlueprintCallable, Category="Collision", meta=(WorldContext="WorldContextObject", AutoCreateRefTerm="ActorsToIgnore", DisplayName="Box Overlap Actors"))
 	static bool BoxOverlapActors(const UObject* WorldContextObject, const FVector BoxPos, FVector BoxExtent, const TArray<TEnumAsByte<EObjectTypeQuery> > & ObjectTypes, UClass* ActorClassFilter, const TArray<AActor*>& ActorsToIgnore, TArray<class AActor*>& OutActors);
 
 	/**
@@ -961,7 +962,7 @@ class ENGINE_API UKismetSystemLibrary : public UBlueprintFunctionLibrary
 	 * @param OutActors		Returned array of actors. Unsorted.
 	 * @return				true if there was an overlap that passed the filters, false otherwise.
 	 */
-	UFUNCTION(BlueprintCallable, Category="Collision", meta=(WorldContext="WorldContextObject", AutoCreateRefTerm="ActorsToIgnore", DisplayName="BoxOverlapComponents"))
+	UFUNCTION(BlueprintCallable, Category="Collision", meta=(WorldContext="WorldContextObject", AutoCreateRefTerm="ActorsToIgnore", DisplayName="Box Overlap Components"))
 	static bool BoxOverlapComponents(const UObject* WorldContextObject, const FVector BoxPos, FVector Extent, const TArray<TEnumAsByte<EObjectTypeQuery> > & ObjectTypes, UClass* ComponentClassFilter, const TArray<AActor*>& ActorsToIgnore, TArray<class UPrimitiveComponent*>& OutComponents);
 
 
@@ -977,7 +978,7 @@ class ENGINE_API UKismetSystemLibrary : public UBlueprintFunctionLibrary
 	 * @param OutActors		Returned array of actors. Unsorted.
 	 * @return				true if there was an overlap that passed the filters, false otherwise.
 	 */
-	UFUNCTION(BlueprintCallable, Category="Collision", meta=(WorldContext="WorldContextObject", AutoCreateRefTerm="ActorsToIgnore", DisplayName="CapsuleOverlapActors"))
+	UFUNCTION(BlueprintCallable, Category="Collision", meta=(WorldContext="WorldContextObject", AutoCreateRefTerm="ActorsToIgnore", DisplayName="Capsule Overlap Actors"))
 	static bool CapsuleOverlapActors(const UObject* WorldContextObject, const FVector CapsulePos, float Radius, float HalfHeight, const TArray<TEnumAsByte<EObjectTypeQuery> > & ObjectTypes, UClass* ActorClassFilter, const TArray<AActor*>& ActorsToIgnore, TArray<class AActor*>& OutActors);
 
 	/**
@@ -992,7 +993,7 @@ class ENGINE_API UKismetSystemLibrary : public UBlueprintFunctionLibrary
 	 * @param OutActors		Returned array of actors. Unsorted.
 	 * @return				true if there was an overlap that passed the filters, false otherwise.
 	 */
-	UFUNCTION(BlueprintCallable, Category="Collision", meta=(WorldContext="WorldContextObject", AutoCreateRefTerm="ActorsToIgnore", DisplayName="CapsuleOverlapComponents") )
+	UFUNCTION(BlueprintCallable, Category="Collision", meta=(WorldContext="WorldContextObject", AutoCreateRefTerm="ActorsToIgnore", DisplayName="Capsule Overlap Components") )
 	static bool CapsuleOverlapComponents(const UObject* WorldContextObject, const FVector CapsulePos, float Radius, float HalfHeight, const TArray<TEnumAsByte<EObjectTypeQuery> > & ObjectTypes, UClass* ComponentClassFilter, const TArray<AActor*>& ActorsToIgnore, TArray<class UPrimitiveComponent*>& OutComponents);
 
 
@@ -1006,7 +1007,7 @@ class ENGINE_API UKismetSystemLibrary : public UBlueprintFunctionLibrary
 	 * @param OutActors				Returned array of actors. Unsorted.
 	 * @return						true if there was an overlap that passed the filters, false otherwise.
 	 */
-	UFUNCTION(BlueprintCallable, Category="Collision", meta=(AutoCreateRefTerm="ActorsToIgnore", DisplayName="ComponentOverlapActors"))
+	UFUNCTION(BlueprintCallable, Category="Collision", meta=(AutoCreateRefTerm="ActorsToIgnore", DisplayName="Component Overlap Actors"))
 	static bool ComponentOverlapActors(UPrimitiveComponent* Component, const FTransform& ComponentTransform, const TArray<TEnumAsByte<EObjectTypeQuery> > & ObjectTypes, UClass* ActorClassFilter, const TArray<AActor*>& ActorsToIgnore, TArray<class AActor*>& OutActors);
 
 	/**
@@ -1019,7 +1020,7 @@ class ENGINE_API UKismetSystemLibrary : public UBlueprintFunctionLibrary
 	 * @param OutActors				Returned array of actors. Unsorted.
 	 * @return						true if there was an overlap that passed the filters, false otherwise.
 	 */
-	UFUNCTION(BlueprintCallable, Category="Collision", meta=(AutoCreateRefTerm="ActorsToIgnore", DisplayName="ComponentOverlapComponents"))
+	UFUNCTION(BlueprintCallable, Category="Collision", meta=(AutoCreateRefTerm="ActorsToIgnore", DisplayName="Component Overlap Components"))
 	static bool ComponentOverlapComponents(UPrimitiveComponent* Component, const FTransform& ComponentTransform, const TArray<TEnumAsByte<EObjectTypeQuery> > & ObjectTypes, UClass* ComponentClassFilter, const TArray<AActor*>& ActorsToIgnore, TArray<class UPrimitiveComponent*>& OutComponents);
 
 
@@ -1035,7 +1036,7 @@ class ENGINE_API UKismetSystemLibrary : public UBlueprintFunctionLibrary
 	 * @param OutHit		Properties of the trace hit.
 	 * @return				True if there was a hit, false otherwise.
 	 */
-	UFUNCTION(BlueprintCallable, Category="Collision", meta=(bIgnoreSelf="true", WorldContext="WorldContextObject", AutoCreateRefTerm="ActorsToIgnore", DisplayName="LineTraceByChannel", AdvancedDisplay="TraceColor,TraceHitColor,DrawTime", Keywords="raycast"))
+	UFUNCTION(BlueprintCallable, Category="Collision", meta=(bIgnoreSelf="true", WorldContext="WorldContextObject", AutoCreateRefTerm="ActorsToIgnore", DisplayName="Line Trace By Channel", AdvancedDisplay="TraceColor,TraceHitColor,DrawTime", Keywords="raycast"))
 	static bool LineTraceSingle(const UObject* WorldContextObject, const FVector Start, const FVector End, ETraceTypeQuery TraceChannel, bool bTraceComplex, const TArray<AActor*>& ActorsToIgnore, EDrawDebugTrace::Type DrawDebugType, FHitResult& OutHit, bool bIgnoreSelf, FLinearColor TraceColor = FLinearColor::Red, FLinearColor TraceHitColor = FLinearColor::Green, float DrawTime = 5.0f);
 	
 	/**
@@ -1050,7 +1051,7 @@ class ENGINE_API UKismetSystemLibrary : public UBlueprintFunctionLibrary
 	 * @param OutHit		Properties of the trace hit.
 	 * @return				True if there was a blocking hit, false otherwise.
 	 */
-	UFUNCTION(BlueprintCallable, Category="Collision", meta=(bIgnoreSelf="true", WorldContext="WorldContextObject", AutoCreateRefTerm="ActorsToIgnore", DisplayName = "MultiLineTraceByChannel", AdvancedDisplay="TraceColor,TraceHitColor,DrawTime", Keywords="raycast"))
+	UFUNCTION(BlueprintCallable, Category="Collision", meta=(bIgnoreSelf="true", WorldContext="WorldContextObject", AutoCreateRefTerm="ActorsToIgnore", DisplayName = "Multi Line Trace By Channel", AdvancedDisplay="TraceColor,TraceHitColor,DrawTime", Keywords="raycast"))
 	static bool LineTraceMulti(const UObject* WorldContextObject, const FVector Start, const FVector End, ETraceTypeQuery TraceChannel, bool bTraceComplex, const TArray<AActor*>& ActorsToIgnore, EDrawDebugTrace::Type DrawDebugType, TArray<FHitResult>& OutHits, bool bIgnoreSelf, FLinearColor TraceColor = FLinearColor::Red, FLinearColor TraceHitColor = FLinearColor::Green, float DrawTime = 5.0f);
 
 	/**
@@ -1065,7 +1066,7 @@ class ENGINE_API UKismetSystemLibrary : public UBlueprintFunctionLibrary
 	 * @param OutHit		Properties of the trace hit.
 	 * @return				True if there was a hit, false otherwise.
 	 */
-	UFUNCTION(BlueprintCallable, Category="Collision", meta=(bIgnoreSelf="true", WorldContext="WorldContextObject", AutoCreateRefTerm="ActorsToIgnore", DisplayName = "SphereTraceByChannel", AdvancedDisplay="TraceColor,TraceHitColor,DrawTime", Keywords="sweep"))
+	UFUNCTION(BlueprintCallable, Category="Collision", meta=(bIgnoreSelf="true", WorldContext="WorldContextObject", AutoCreateRefTerm="ActorsToIgnore", DisplayName = "Sphere Trace By Channel", AdvancedDisplay="TraceColor,TraceHitColor,DrawTime", Keywords="sweep"))
 	static bool SphereTraceSingle(const UObject* WorldContextObject, const FVector Start, const FVector End, float Radius, ETraceTypeQuery TraceChannel, bool bTraceComplex, const TArray<AActor*>& ActorsToIgnore, EDrawDebugTrace::Type DrawDebugType, FHitResult& OutHit, bool bIgnoreSelf, FLinearColor TraceColor = FLinearColor::Red, FLinearColor TraceHitColor = FLinearColor::Green, float DrawTime = 5.0f);
 
 	/**
@@ -1081,7 +1082,7 @@ class ENGINE_API UKismetSystemLibrary : public UBlueprintFunctionLibrary
 	 * @param OutHits		A list of hits, sorted along the trace from start to finish.  The blocking hit will be the last hit, if there was one.
 	 * @return				True if there was a blocking hit, false otherwise.
 	 */
-	UFUNCTION(BlueprintCallable, Category="Collision", meta=(bIgnoreSelf="true", WorldContext="WorldContextObject", AutoCreateRefTerm="ActorsToIgnore", DisplayName = "MultiSphereTraceByChannel", AdvancedDisplay="TraceColor,TraceHitColor,DrawTime", Keywords="sweep"))
+	UFUNCTION(BlueprintCallable, Category="Collision", meta=(bIgnoreSelf="true", WorldContext="WorldContextObject", AutoCreateRefTerm="ActorsToIgnore", DisplayName = "Multi Sphere Trace By Channel", AdvancedDisplay="TraceColor,TraceHitColor,DrawTime", Keywords="sweep"))
 	static bool SphereTraceMulti(const UObject* WorldContextObject, const FVector Start, const FVector End, float Radius, ETraceTypeQuery TraceChannel, bool bTraceComplex, const TArray<AActor*>& ActorsToIgnore, EDrawDebugTrace::Type DrawDebugType, TArray<FHitResult>& OutHits, bool bIgnoreSelf, FLinearColor TraceColor = FLinearColor::Red, FLinearColor TraceHitColor = FLinearColor::Green, float DrawTime = 5.0f);
 
 	/**
@@ -1097,7 +1098,7 @@ class ENGINE_API UKismetSystemLibrary : public UBlueprintFunctionLibrary
 	* @param OutHit			Properties of the trace hit.
 	* @return				True if there was a hit, false otherwise.
 	*/
-	UFUNCTION(BlueprintCallable, Category = "Collision", meta = (bIgnoreSelf = "true", WorldContext="WorldContextObject", AutoCreateRefTerm = "ActorsToIgnore", DisplayName = "BoxTraceByChannel", AdvancedDisplay="TraceColor,TraceHitColor,DrawTime", Keywords="sweep"))
+	UFUNCTION(BlueprintCallable, Category = "Collision", meta = (bIgnoreSelf = "true", WorldContext="WorldContextObject", AutoCreateRefTerm = "ActorsToIgnore", DisplayName = "Box Trace By Channel", AdvancedDisplay="TraceColor,TraceHitColor,DrawTime", Keywords="sweep"))
 	static bool BoxTraceSingle(const UObject* WorldContextObject, const FVector Start, const FVector End, const FVector HalfSize, const FRotator Orientation, ETraceTypeQuery TraceChannel, bool bTraceComplex, const TArray<AActor*>& ActorsToIgnore, EDrawDebugTrace::Type DrawDebugType, FHitResult& OutHit, bool bIgnoreSelf, FLinearColor TraceColor = FLinearColor::Red, FLinearColor TraceHitColor = FLinearColor::Green, float DrawTime = 5.0f);
 
 	/**
@@ -1113,7 +1114,7 @@ class ENGINE_API UKismetSystemLibrary : public UBlueprintFunctionLibrary
 	* @param OutHits		A list of hits, sorted along the trace from start to finish. The blocking hit will be the last hit, if there was one.
 	* @return				True if there was a blocking hit, false otherwise.
 	*/
-	UFUNCTION(BlueprintCallable, Category = "Collision", meta = (bIgnoreSelf = "true", WorldContext="WorldContextObject", AutoCreateRefTerm = "ActorsToIgnore", DisplayName = "MultiBoxTraceByChannel", AdvancedDisplay="TraceColor,TraceHitColor,DrawTime", Keywords="sweep"))
+	UFUNCTION(BlueprintCallable, Category = "Collision", meta = (bIgnoreSelf = "true", WorldContext="WorldContextObject", AutoCreateRefTerm = "ActorsToIgnore", DisplayName = "Multi Box Trace By Channel", AdvancedDisplay="TraceColor,TraceHitColor,DrawTime", Keywords="sweep"))
 	static bool BoxTraceMulti(const UObject* WorldContextObject, const FVector Start, const FVector End, FVector HalfSize, const FRotator Orientation, ETraceTypeQuery TraceChannel, bool bTraceComplex, const TArray<AActor*>& ActorsToIgnore, EDrawDebugTrace::Type DrawDebugType, TArray<FHitResult>& OutHits, bool bIgnoreSelf, FLinearColor TraceColor = FLinearColor::Red, FLinearColor TraceHitColor = FLinearColor::Green, float DrawTime = 5.0f);
 
 	
@@ -1131,7 +1132,7 @@ class ENGINE_API UKismetSystemLibrary : public UBlueprintFunctionLibrary
 	 * @param OutHit		Properties of the trace hit.
 	 * @return				True if there was a hit, false otherwise.
 	 */
-	UFUNCTION(BlueprintCallable, Category="Collision", meta=(bIgnoreSelf="true", WorldContext="WorldContextObject", AutoCreateRefTerm="ActorsToIgnore", DisplayName = "CapsuleTraceByChannel", AdvancedDisplay="TraceColor,TraceHitColor,DrawTime", Keywords="sweep"))
+	UFUNCTION(BlueprintCallable, Category="Collision", meta=(bIgnoreSelf="true", WorldContext="WorldContextObject", AutoCreateRefTerm="ActorsToIgnore", DisplayName = "Capsule Trace By Channel", AdvancedDisplay="TraceColor,TraceHitColor,DrawTime", Keywords="sweep"))
 	static bool CapsuleTraceSingle(const UObject* WorldContextObject, const FVector Start, const FVector End, float Radius, float HalfHeight, ETraceTypeQuery TraceChannel, bool bTraceComplex, const TArray<AActor*>& ActorsToIgnore, EDrawDebugTrace::Type DrawDebugType, FHitResult& OutHit, bool bIgnoreSelf, FLinearColor TraceColor = FLinearColor::Red, FLinearColor TraceHitColor = FLinearColor::Green, float DrawTime = 5.0f);
 
 	/**
@@ -1148,7 +1149,7 @@ class ENGINE_API UKismetSystemLibrary : public UBlueprintFunctionLibrary
 	 * @param OutHits		A list of hits, sorted along the trace from start to finish.  The blocking hit will be the last hit, if there was one.
 	 * @return				True if there was a blocking hit, false otherwise.
 	 */
-	UFUNCTION(BlueprintCallable, Category="Collision", meta=(bIgnoreSelf="true", WorldContext="WorldContextObject", AutoCreateRefTerm="ActorsToIgnore", DisplayName = "MultiCapsuleTraceByChannel", AdvancedDisplay="TraceColor,TraceHitColor,DrawTime", Keywords="sweep"))
+	UFUNCTION(BlueprintCallable, Category="Collision", meta=(bIgnoreSelf="true", WorldContext="WorldContextObject", AutoCreateRefTerm="ActorsToIgnore", DisplayName = "Multi Capsule Trace By Channel", AdvancedDisplay="TraceColor,TraceHitColor,DrawTime", Keywords="sweep"))
 	static bool CapsuleTraceMulti(const UObject* WorldContextObject, const FVector Start, const FVector End, float Radius, float HalfHeight, ETraceTypeQuery TraceChannel, bool bTraceComplex, const TArray<AActor*>& ActorsToIgnore, EDrawDebugTrace::Type DrawDebugType, TArray<FHitResult>& OutHits, bool bIgnoreSelf, FLinearColor TraceColor = FLinearColor::Red, FLinearColor TraceHitColor = FLinearColor::Green, float DrawTime = 5.0f);
 
 	/**
@@ -1163,7 +1164,7 @@ class ENGINE_API UKismetSystemLibrary : public UBlueprintFunctionLibrary
 	 * @param OutHit		Properties of the trace hit.
 	 * @return				True if there was a hit, false otherwise.
 	 */
-	UFUNCTION(BlueprintCallable, Category="Collision", meta=(bIgnoreSelf="true", WorldContext="WorldContextObject", AutoCreateRefTerm="ActorsToIgnore", DisplayName = "LineTraceForObjects", AdvancedDisplay="TraceColor,TraceHitColor,DrawTime", Keywords="raycast"))
+	UFUNCTION(BlueprintCallable, Category="Collision", meta=(bIgnoreSelf="true", WorldContext="WorldContextObject", AutoCreateRefTerm="ActorsToIgnore", DisplayName = "Line Trace For Objects", AdvancedDisplay="TraceColor,TraceHitColor,DrawTime", Keywords="raycast"))
 	static bool LineTraceSingleForObjects(const UObject* WorldContextObject, const FVector Start, const FVector End, const TArray<TEnumAsByte<EObjectTypeQuery> > & ObjectTypes, bool bTraceComplex, const TArray<AActor*>& ActorsToIgnore, EDrawDebugTrace::Type DrawDebugType, FHitResult& OutHit, bool bIgnoreSelf, FLinearColor TraceColor = FLinearColor::Red, FLinearColor TraceHitColor = FLinearColor::Green, float DrawTime = 5.0f );
 	
 	/**
@@ -1178,7 +1179,7 @@ class ENGINE_API UKismetSystemLibrary : public UBlueprintFunctionLibrary
 	 * @param OutHit		Properties of the trace hit.
 	 * @return				True if there was a hit, false otherwise.
 	 */
-	UFUNCTION(BlueprintCallable, Category="Collision", meta=(bIgnoreSelf="true", WorldContext="WorldContextObject", AutoCreateRefTerm="ActorsToIgnore", DisplayName = "MultiLineTraceForObjects", AdvancedDisplay="TraceColor,TraceHitColor,DrawTime", Keywords="raycast"))
+	UFUNCTION(BlueprintCallable, Category="Collision", meta=(bIgnoreSelf="true", WorldContext="WorldContextObject", AutoCreateRefTerm="ActorsToIgnore", DisplayName = "Multi Line Trace For Objects", AdvancedDisplay="TraceColor,TraceHitColor,DrawTime", Keywords="raycast"))
 	static bool LineTraceMultiForObjects(const UObject* WorldContextObject, const FVector Start, const FVector End, const TArray<TEnumAsByte<EObjectTypeQuery> > & ObjectTypes, bool bTraceComplex, const TArray<AActor*>& ActorsToIgnore, EDrawDebugTrace::Type DrawDebugType, TArray<FHitResult>& OutHits, bool bIgnoreSelf, FLinearColor TraceColor = FLinearColor::Red, FLinearColor TraceHitColor = FLinearColor::Green, float DrawTime = 5.0f);
 
 	/**
@@ -1193,7 +1194,7 @@ class ENGINE_API UKismetSystemLibrary : public UBlueprintFunctionLibrary
 	 * @param OutHit		Properties of the trace hit.
 	 * @return				True if there was a hit, false otherwise.
 	 */
-	UFUNCTION(BlueprintCallable, Category="Collision", meta=(bIgnoreSelf="true", WorldContext="WorldContextObject", AutoCreateRefTerm="ActorsToIgnore", DisplayName = "SphereTraceForObjects", AdvancedDisplay="TraceColor,TraceHitColor,DrawTime", Keywords="sweep"))
+	UFUNCTION(BlueprintCallable, Category="Collision", meta=(bIgnoreSelf="true", WorldContext="WorldContextObject", AutoCreateRefTerm="ActorsToIgnore", DisplayName = "Sphere Trace For Objects", AdvancedDisplay="TraceColor,TraceHitColor,DrawTime", Keywords="sweep"))
 	static bool SphereTraceSingleForObjects(const UObject* WorldContextObject, const FVector Start, const FVector End, float Radius, const TArray<TEnumAsByte<EObjectTypeQuery> > & ObjectTypes, bool bTraceComplex, const TArray<AActor*>& ActorsToIgnore, EDrawDebugTrace::Type DrawDebugType, FHitResult& OutHit, bool bIgnoreSelf, FLinearColor TraceColor = FLinearColor::Red, FLinearColor TraceHitColor = FLinearColor::Green, float DrawTime = 5.0f);
 
 	/**
@@ -1209,7 +1210,7 @@ class ENGINE_API UKismetSystemLibrary : public UBlueprintFunctionLibrary
 	 * @param OutHits		A list of hits, sorted along the trace from start to finish.  The blocking hit will be the last hit, if there was one.
 	 * @return				True if there was a hit, false otherwise.
 	 */
-	UFUNCTION(BlueprintCallable, Category="Collision", meta=(bIgnoreSelf="true", WorldContext="WorldContextObject", AutoCreateRefTerm="ActorsToIgnore", DisplayName = "MultiSphereTraceForObjects", AdvancedDisplay="TraceColor,TraceHitColor,DrawTime", Keywords="sweep"))
+	UFUNCTION(BlueprintCallable, Category="Collision", meta=(bIgnoreSelf="true", WorldContext="WorldContextObject", AutoCreateRefTerm="ActorsToIgnore", DisplayName = "Multi Sphere Trace For Objects", AdvancedDisplay="TraceColor,TraceHitColor,DrawTime", Keywords="sweep"))
 	static bool SphereTraceMultiForObjects(const UObject* WorldContextObject, const FVector Start, const FVector End, float Radius, const TArray<TEnumAsByte<EObjectTypeQuery> > & ObjectTypes, bool bTraceComplex, const TArray<AActor*>& ActorsToIgnore, EDrawDebugTrace::Type DrawDebugType, TArray<FHitResult>& OutHits, bool bIgnoreSelf, FLinearColor TraceColor = FLinearColor::Red, FLinearColor TraceHitColor = FLinearColor::Green, float DrawTime = 5.0f);
 	
 
@@ -1226,7 +1227,7 @@ class ENGINE_API UKismetSystemLibrary : public UBlueprintFunctionLibrary
 	* @param OutHit			Properties of the trace hit.
 	* @return				True if there was a hit, false otherwise.
 	*/
-	UFUNCTION(BlueprintCallable, Category = "Collision", meta = (bIgnoreSelf = "true", WorldContext="WorldContextObject", AutoCreateRefTerm = "ActorsToIgnore", DisplayName = "BoxTraceForObjects", AdvancedDisplay="TraceColor,TraceHitColor,DrawTime", Keywords="sweep"))
+	UFUNCTION(BlueprintCallable, Category = "Collision", meta = (bIgnoreSelf = "true", WorldContext="WorldContextObject", AutoCreateRefTerm = "ActorsToIgnore", DisplayName = "Box Trace For Objects", AdvancedDisplay="TraceColor,TraceHitColor,DrawTime", Keywords="sweep"))
 	static bool BoxTraceSingleForObjects(const UObject* WorldContextObject, const FVector Start, const FVector End, const FVector HalfSize, const FRotator Orientation, const TArray<TEnumAsByte<EObjectTypeQuery> > & ObjectTypes, bool bTraceComplex, const TArray<AActor*>& ActorsToIgnore, EDrawDebugTrace::Type DrawDebugType, FHitResult& OutHit, bool bIgnoreSelf, FLinearColor TraceColor = FLinearColor::Red, FLinearColor TraceHitColor = FLinearColor::Green, float DrawTime = 5.0f);
 
 
@@ -1243,7 +1244,7 @@ class ENGINE_API UKismetSystemLibrary : public UBlueprintFunctionLibrary
 	* @param OutHits		A list of hits, sorted along the trace from start to finish.  The blocking hit will be the last hit, if there was one.
 	* @return				True if there was a hit, false otherwise.
 	*/
-	UFUNCTION(BlueprintCallable, Category = "Collision", meta = (bIgnoreSelf = "true", WorldContext="WorldContextObject", AutoCreateRefTerm = "ActorsToIgnore", DisplayName = "MultiBoxTraceForObjects", AdvancedDisplay="TraceColor,TraceHitColor,DrawTime", Keywords="sweep"))
+	UFUNCTION(BlueprintCallable, Category = "Collision", meta = (bIgnoreSelf = "true", WorldContext="WorldContextObject", AutoCreateRefTerm = "ActorsToIgnore", DisplayName = "Multi Box Trace For Objects", AdvancedDisplay="TraceColor,TraceHitColor,DrawTime", Keywords="sweep"))
 	static bool BoxTraceMultiForObjects(const UObject* WorldContextObject, const FVector Start, const FVector End, const FVector HalfSize, const FRotator Orientation, const TArray<TEnumAsByte<EObjectTypeQuery> > & ObjectTypes, bool bTraceComplex, const TArray<AActor*>& ActorsToIgnore, EDrawDebugTrace::Type DrawDebugType, TArray<FHitResult>& OutHits, bool bIgnoreSelf, FLinearColor TraceColor = FLinearColor::Red, FLinearColor TraceHitColor = FLinearColor::Green, float DrawTime = 5.0f);
 
 	/**
@@ -1260,7 +1261,7 @@ class ENGINE_API UKismetSystemLibrary : public UBlueprintFunctionLibrary
 	 * @param OutHit		Properties of the trace hit.
 	 * @return				True if there was a hit, false otherwise.
 	 */
-	UFUNCTION(BlueprintCallable, Category="Collision", meta=(bIgnoreSelf="true", WorldContext="WorldContextObject", AutoCreateRefTerm="ActorsToIgnore", DisplayName = "CapsuleTraceForObjects", AdvancedDisplay="TraceColor,TraceHitColor,DrawTime", Keywords="sweep"))
+	UFUNCTION(BlueprintCallable, Category="Collision", meta=(bIgnoreSelf="true", WorldContext="WorldContextObject", AutoCreateRefTerm="ActorsToIgnore", DisplayName = "Capsule Trace For Objects", AdvancedDisplay="TraceColor,TraceHitColor,DrawTime", Keywords="sweep"))
 	static bool CapsuleTraceSingleForObjects(const UObject* WorldContextObject, const FVector Start, const FVector End, float Radius, float HalfHeight, const TArray<TEnumAsByte<EObjectTypeQuery> > & ObjectTypes, bool bTraceComplex, const TArray<AActor*>& ActorsToIgnore, EDrawDebugTrace::Type DrawDebugType, FHitResult& OutHit, bool bIgnoreSelf, FLinearColor TraceColor = FLinearColor::Red, FLinearColor TraceHitColor = FLinearColor::Green, float DrawTime = 5.0f);
 
 	/**
@@ -1277,7 +1278,7 @@ class ENGINE_API UKismetSystemLibrary : public UBlueprintFunctionLibrary
 	 * @param OutHits		A list of hits, sorted along the trace from start to finish.  The blocking hit will be the last hit, if there was one.
 	 * @return				True if there was a hit, false otherwise.
 	 */
-	UFUNCTION(BlueprintCallable, Category="Collision", meta=(bIgnoreSelf="true", WorldContext="WorldContextObject", AutoCreateRefTerm="ActorsToIgnore", DisplayName = "MultiCapsuleTraceForObjects", AdvancedDisplay="TraceColor,TraceHitColor,DrawTime", Keywords="sweep"))
+	UFUNCTION(BlueprintCallable, Category="Collision", meta=(bIgnoreSelf="true", WorldContext="WorldContextObject", AutoCreateRefTerm="ActorsToIgnore", DisplayName = "Multi Capsule Trace For Objects", AdvancedDisplay="TraceColor,TraceHitColor,DrawTime", Keywords="sweep"))
 	static bool CapsuleTraceMultiForObjects(const UObject* WorldContextObject, const FVector Start, const FVector End, float Radius, float HalfHeight, const TArray<TEnumAsByte<EObjectTypeQuery> > & ObjectTypes, bool bTraceComplex, const TArray<AActor*>& ActorsToIgnore, EDrawDebugTrace::Type DrawDebugType, TArray<FHitResult>& OutHits, bool bIgnoreSelf, FLinearColor TraceColor = FLinearColor::Red, FLinearColor TraceHitColor = FLinearColor::Green, float DrawTime = 5.0f);
 
 	// BY PROFILE
@@ -1293,7 +1294,7 @@ class ENGINE_API UKismetSystemLibrary : public UBlueprintFunctionLibrary
 	* @param OutHit			Properties of the trace hit.
 	* @return				True if there was a hit, false otherwise.
 	*/
-	UFUNCTION(BlueprintCallable, Category = "Collision", meta = (bIgnoreSelf = "true", WorldContext = "WorldContextObject", AutoCreateRefTerm = "ActorsToIgnore", DisplayName = "LineTraceByProfile", AdvancedDisplay = "TraceColor,TraceHitColor,DrawTime", Keywords = "raycast"))
+	UFUNCTION(BlueprintCallable, Category = "Collision", meta = (bIgnoreSelf = "true", WorldContext = "WorldContextObject", AutoCreateRefTerm = "ActorsToIgnore", DisplayName = "Line Trace By Profile", AdvancedDisplay = "TraceColor,TraceHitColor,DrawTime", Keywords = "raycast"))
 	static bool LineTraceSingleByProfile(const UObject* WorldContextObject, const FVector Start, const FVector End, FName ProfileName, bool bTraceComplex, const TArray<AActor*>& ActorsToIgnore, EDrawDebugTrace::Type DrawDebugType, FHitResult& OutHit, bool bIgnoreSelf, FLinearColor TraceColor = FLinearColor::Red, FLinearColor TraceHitColor = FLinearColor::Green, float DrawTime = 5.0f);
 
 	/**
@@ -1309,7 +1310,7 @@ class ENGINE_API UKismetSystemLibrary : public UBlueprintFunctionLibrary
 	* @param OutHit		Properties of the trace hit.
 	* @return				True if there was a blocking hit, false otherwise.
 	*/
-	UFUNCTION(BlueprintCallable, Category = "Collision", meta = (bIgnoreSelf = "true", WorldContext = "WorldContextObject", AutoCreateRefTerm = "ActorsToIgnore", DisplayName = "MultiLineTraceByProfile", AdvancedDisplay = "TraceColor,TraceHitColor,DrawTime", Keywords = "raycast"))
+	UFUNCTION(BlueprintCallable, Category = "Collision", meta = (bIgnoreSelf = "true", WorldContext = "WorldContextObject", AutoCreateRefTerm = "ActorsToIgnore", DisplayName = "Multi Line Trace By Profile", AdvancedDisplay = "TraceColor,TraceHitColor,DrawTime", Keywords = "raycast"))
 	static bool LineTraceMultiByProfile(const UObject* WorldContextObject, const FVector Start, const FVector End, FName ProfileName, bool bTraceComplex, const TArray<AActor*>& ActorsToIgnore, EDrawDebugTrace::Type DrawDebugType, TArray<FHitResult>& OutHits, bool bIgnoreSelf, FLinearColor TraceColor = FLinearColor::Red, FLinearColor TraceHitColor = FLinearColor::Green, float DrawTime = 5.0f);
 
 	/**
@@ -1323,7 +1324,7 @@ class ENGINE_API UKismetSystemLibrary : public UBlueprintFunctionLibrary
 	* @param OutHit			Properties of the trace hit.
 	* @return				True if there was a hit, false otherwise.
 	*/
-	UFUNCTION(BlueprintCallable, Category = "Collision", meta = (bIgnoreSelf = "true", WorldContext = "WorldContextObject", AutoCreateRefTerm = "ActorsToIgnore", DisplayName = "SphereTraceByProfile", AdvancedDisplay = "TraceColor,TraceHitColor,DrawTime", Keywords = "sweep"))
+	UFUNCTION(BlueprintCallable, Category = "Collision", meta = (bIgnoreSelf = "true", WorldContext = "WorldContextObject", AutoCreateRefTerm = "ActorsToIgnore", DisplayName = "Sphere Trace By Profile", AdvancedDisplay = "TraceColor,TraceHitColor,DrawTime", Keywords = "sweep"))
 	static bool SphereTraceSingleByProfile(const UObject* WorldContextObject, const FVector Start, const FVector End, float Radius, FName ProfileName, bool bTraceComplex, const TArray<AActor*>& ActorsToIgnore, EDrawDebugTrace::Type DrawDebugType, FHitResult& OutHit, bool bIgnoreSelf, FLinearColor TraceColor = FLinearColor::Red, FLinearColor TraceHitColor = FLinearColor::Green, float DrawTime = 5.0f);
 
 	/**
@@ -1340,7 +1341,7 @@ class ENGINE_API UKismetSystemLibrary : public UBlueprintFunctionLibrary
 	* @param OutHits		A list of hits, sorted along the trace from start to finish.  The blocking hit will be the last hit, if there was one.
 	* @return				True if there was a blocking hit, false otherwise.
 	*/
-	UFUNCTION(BlueprintCallable, Category = "Collision", meta = (bIgnoreSelf = "true", WorldContext = "WorldContextObject", AutoCreateRefTerm = "ActorsToIgnore", DisplayName = "MultiSphereTraceByProfile", AdvancedDisplay = "TraceColor,TraceHitColor,DrawTime", Keywords = "sweep"))
+	UFUNCTION(BlueprintCallable, Category = "Collision", meta = (bIgnoreSelf = "true", WorldContext = "WorldContextObject", AutoCreateRefTerm = "ActorsToIgnore", DisplayName = "Multi Sphere Trace By Profile", AdvancedDisplay = "TraceColor,TraceHitColor,DrawTime", Keywords = "sweep"))
 	static bool SphereTraceMultiByProfile(const UObject* WorldContextObject, const FVector Start, const FVector End, float Radius, FName ProfileName, bool bTraceComplex, const TArray<AActor*>& ActorsToIgnore, EDrawDebugTrace::Type DrawDebugType, TArray<FHitResult>& OutHits, bool bIgnoreSelf, FLinearColor TraceColor = FLinearColor::Red, FLinearColor TraceHitColor = FLinearColor::Green, float DrawTime = 5.0f);
 
 	/**
@@ -1355,7 +1356,7 @@ class ENGINE_API UKismetSystemLibrary : public UBlueprintFunctionLibrary
 	* @param OutHit			Properties of the trace hit.
 	* @return				True if there was a hit, false otherwise.
 	*/
-	UFUNCTION(BlueprintCallable, Category = "Collision", meta = (bIgnoreSelf = "true", WorldContext = "WorldContextObject", AutoCreateRefTerm = "ActorsToIgnore", DisplayName = "BoxTraceByProfile", AdvancedDisplay = "TraceColor,TraceHitColor,DrawTime", Keywords = "sweep"))
+	UFUNCTION(BlueprintCallable, Category = "Collision", meta = (bIgnoreSelf = "true", WorldContext = "WorldContextObject", AutoCreateRefTerm = "ActorsToIgnore", DisplayName = "Box Trace By Profile", AdvancedDisplay = "TraceColor,TraceHitColor,DrawTime", Keywords = "sweep"))
 	static bool BoxTraceSingleByProfile(const UObject* WorldContextObject, const FVector Start, const FVector End, const FVector HalfSize, const FRotator Orientation, FName ProfileName, bool bTraceComplex, const TArray<AActor*>& ActorsToIgnore, EDrawDebugTrace::Type DrawDebugType, FHitResult& OutHit, bool bIgnoreSelf, FLinearColor TraceColor = FLinearColor::Red, FLinearColor TraceHitColor = FLinearColor::Green, float DrawTime = 5.0f);
 
 	/**
@@ -1372,7 +1373,7 @@ class ENGINE_API UKismetSystemLibrary : public UBlueprintFunctionLibrary
 	* @param OutHits		A list of hits, sorted along the trace from start to finish. The blocking hit will be the last hit, if there was one.
 	* @return				True if there was a blocking hit, false otherwise.
 	*/
-	UFUNCTION(BlueprintCallable, Category = "Collision", meta = (bIgnoreSelf = "true", WorldContext = "WorldContextObject", AutoCreateRefTerm = "ActorsToIgnore", DisplayName = "MultiBoxTraceByProfile", AdvancedDisplay = "TraceColor,TraceHitColor,DrawTime", Keywords = "sweep"))
+	UFUNCTION(BlueprintCallable, Category = "Collision", meta = (bIgnoreSelf = "true", WorldContext = "WorldContextObject", AutoCreateRefTerm = "ActorsToIgnore", DisplayName = "Multi Box Trace By Profile", AdvancedDisplay = "TraceColor,TraceHitColor,DrawTime", Keywords = "sweep"))
 	static bool BoxTraceMultiByProfile(const UObject* WorldContextObject, const FVector Start, const FVector End, FVector HalfSize, const FRotator Orientation, FName ProfileName, bool bTraceComplex, const TArray<AActor*>& ActorsToIgnore, EDrawDebugTrace::Type DrawDebugType, TArray<FHitResult>& OutHits, bool bIgnoreSelf, FLinearColor TraceColor = FLinearColor::Red, FLinearColor TraceHitColor = FLinearColor::Green, float DrawTime = 5.0f);
 
 
@@ -1389,7 +1390,7 @@ class ENGINE_API UKismetSystemLibrary : public UBlueprintFunctionLibrary
 	* @param OutHit			Properties of the trace hit.
 	* @return				True if there was a hit, false otherwise.
 	*/
-	UFUNCTION(BlueprintCallable, Category = "Collision", meta = (bIgnoreSelf = "true", WorldContext = "WorldContextObject", AutoCreateRefTerm = "ActorsToIgnore", DisplayName = "CapsuleTraceByProfile", AdvancedDisplay = "TraceColor,TraceHitColor,DrawTime", Keywords = "sweep"))
+	UFUNCTION(BlueprintCallable, Category = "Collision", meta = (bIgnoreSelf = "true", WorldContext = "WorldContextObject", AutoCreateRefTerm = "ActorsToIgnore", DisplayName = "Capsule Trace By Profile", AdvancedDisplay = "TraceColor,TraceHitColor,DrawTime", Keywords = "sweep"))
 	static bool CapsuleTraceSingleByProfile(const UObject* WorldContextObject, const FVector Start, const FVector End, float Radius, float HalfHeight, FName ProfileName, bool bTraceComplex, const TArray<AActor*>& ActorsToIgnore, EDrawDebugTrace::Type DrawDebugType, FHitResult& OutHit, bool bIgnoreSelf, FLinearColor TraceColor = FLinearColor::Red, FLinearColor TraceHitColor = FLinearColor::Green, float DrawTime = 5.0f);
 
 	/**
@@ -1407,7 +1408,7 @@ class ENGINE_API UKismetSystemLibrary : public UBlueprintFunctionLibrary
 	* @param OutHits		A list of hits, sorted along the trace from start to finish.  The blocking hit will be the last hit, if there was one.
 	* @return				True if there was a blocking hit, false otherwise.
 	*/
-	UFUNCTION(BlueprintCallable, Category = "Collision", meta = (bIgnoreSelf = "true", WorldContext = "WorldContextObject", AutoCreateRefTerm = "ActorsToIgnore", DisplayName = "MultiCapsuleTraceByProfile", AdvancedDisplay = "TraceColor,TraceHitColor,DrawTime", Keywords = "sweep"))
+	UFUNCTION(BlueprintCallable, Category = "Collision", meta = (bIgnoreSelf = "true", WorldContext = "WorldContextObject", AutoCreateRefTerm = "ActorsToIgnore", DisplayName = "Multi Capsule Trace By Profile", AdvancedDisplay = "TraceColor,TraceHitColor,DrawTime", Keywords = "sweep"))
 	static bool CapsuleTraceMultiByProfile(const UObject* WorldContextObject, const FVector Start, const FVector End, float Radius, float HalfHeight, FName ProfileName, bool bTraceComplex, const TArray<AActor*>& ActorsToIgnore, EDrawDebugTrace::Type DrawDebugType, TArray<FHitResult>& OutHits, bool bIgnoreSelf, FLinearColor TraceColor = FLinearColor::Red, FLinearColor TraceHitColor = FLinearColor::Green, float DrawTime = 5.0f);
 
 
@@ -1418,7 +1419,7 @@ class ENGINE_API UKismetSystemLibrary : public UBlueprintFunctionLibrary
 	 * @param ClassFilter		If set, will only return results of this class or subclasses of it.
 	 * @param OutActorList		Start of line segment.
 	 */
-	UFUNCTION(BlueprintCallable, Category="Utilities")
+	UFUNCTION(BlueprintCallable, Category="Actor")
 	static void GetActorListFromComponentList(const TArray<class UPrimitiveComponent*>& ComponentList, UClass* ActorClassFilter, TArray<class AActor*>& OutActorList);
 
 
@@ -1464,7 +1465,7 @@ class ENGINE_API UKismetSystemLibrary : public UBlueprintFunctionLibrary
 	 * Draw a debug cone 
 	 * Angles are specified in degrees
 	 */
-	UFUNCTION(BlueprintCallable, Category="Rendering|Debug", meta=(WorldContext="WorldContextObject", DisplayName="DrawDebugCone", DevelopmentOnly))
+	UFUNCTION(BlueprintCallable, Category="Rendering|Debug", meta=(WorldContext="WorldContextObject", DisplayName="Draw Debug Cone", DevelopmentOnly))
 	static void DrawDebugConeInDegrees(const UObject* WorldContextObject, const FVector Origin, const FVector Direction, float Length=100.f, float AngleWidth=45.f, float AngleHeight=45.f, int32 NumSides = 12, FLinearColor LineColor = FLinearColor::White, float Duration=0.f, float Thickness = 0.f);
 
 	/** Draw a debug capsule */
@@ -1733,14 +1734,14 @@ class ENGINE_API UKismetSystemLibrary : public UBlueprintFunctionLibrary
 	 * @param WorldContextObject	World context
 	 * @param State					set true to suppress transition message
 	 */
-	UFUNCTION(BlueprintCallable, Category = "Utilities", meta = (WorldContext="WorldContextObject"))
+	UFUNCTION(BlueprintCallable, Category = "Viewport", meta = (WorldContext="WorldContextObject"))
 	static void SetSuppressViewportTransitionMessage(const UObject* WorldContextObject, bool bState);
 
 	/**
 	 * Returns an array of the user's preferred languages in order of preference
 	 * @return An array of language IDs ordered from most preferred to least
 	 */
-	UFUNCTION(BlueprintCallable, Category = "Utilities|Platform")
+	UFUNCTION(BlueprintCallable, Category = "Utilities|Internationalization")
 	static TArray<FString> GetPreferredLanguages();
 	
 	/**
@@ -1752,7 +1753,7 @@ class ENGINE_API UKismetSystemLibrary : public UBlueprintFunctionLibrary
 	 *  - An optional two-letter ISO 3166-1 country code (eg, "CN")
 	 * @return The language as an IETF language tag (eg, "zh-Hans-CN")
 	 */
-	UFUNCTION(BlueprintPure, Category = "Utilities|Platform")
+	UFUNCTION(BlueprintPure, Category = "Utilities|Internationalization")
 	static FString GetDefaultLanguage();
 
 	/**
@@ -1763,21 +1764,21 @@ class ENGINE_API UKismetSystemLibrary : public UBlueprintFunctionLibrary
 	 *  - An optional two-letter ISO 3166-1 country code (eg, "CN")
 	 * @return The locale as an IETF language tag (eg, "zh-Hans-CN")
 	 */
-	UFUNCTION(BlueprintPure, Category = "Utilities|Platform")
+	UFUNCTION(BlueprintPure, Category = "Utilities|Internationalization")
 	static FString GetDefaultLocale();
 
 	/**
 	* Returns the currency code associated with the device's locale
 	* @return the currency code associated with the device's locale
 	*/
-	UFUNCTION(BlueprintCallable, Category = "Utilities|Platform")
+	UFUNCTION(BlueprintCallable, Category = "Utilities|Internationalization")
 	static FString GetLocalCurrencyCode();
 
 	/**
 	* Returns the currency symbol associated with the device's locale
 	* @return the currency symbol associated with the device's locale
 	*/
-	UFUNCTION(BlueprintCallable, Category = "Utilities|Platform")
+	UFUNCTION(BlueprintCallable, Category = "Utilities|Internationalization")
 	static FString GetLocalCurrencySymbol();
 
 	/**
@@ -1797,7 +1798,7 @@ class ENGINE_API UKismetSystemLibrary : public UBlueprintFunctionLibrary
 	/**
 	 * Tells the engine what the user is doing for debug, analytics, etc.
 	 */
-	UFUNCTION(BlueprintCallable, Category = "Utilities")
+	UFUNCTION(BlueprintCallable, Category = "Development")
 	static void SetUserActivity(const FUserActivity& UserActivity);
 
 	/**
@@ -1969,7 +1970,7 @@ class ENGINE_API UKismetSystemLibrary : public UBlueprintFunctionLibrary
 	static bool IsValidPrimaryAssetId(FPrimaryAssetId PrimaryAssetId);
 
 	/** Converts a Primary Asset Id to a string. The other direction is not provided because it cannot be validated */
-	UFUNCTION(BlueprintPure, meta = (DisplayName = "ToString (PrimaryAssetId)", CompactNodeTitle = "->", ScriptMethod="ToString", BlueprintThreadSafe), Category = "AssetManager")
+	UFUNCTION(BlueprintPure, meta = (DisplayName = "To String (PrimaryAssetId)", CompactNodeTitle = "->", ScriptMethod="ToString", BlueprintThreadSafe), Category = "AssetManager")
 	static FString Conv_PrimaryAssetIdToString(FPrimaryAssetId PrimaryAssetId);
 
 	/** Returns true if the values are equal (A == B) */
@@ -1977,7 +1978,7 @@ class ENGINE_API UKismetSystemLibrary : public UBlueprintFunctionLibrary
 	static bool EqualEqual_PrimaryAssetId(FPrimaryAssetId A, FPrimaryAssetId B);
 
 	/** Returns true if the values are not equal (A != B) */
-	UFUNCTION(BlueprintPure, meta = (DisplayName = "NotEqual (PrimaryAssetId)", CompactNodeTitle = "!=", ScriptOperator="!=", BlueprintThreadSafe), Category = "AssetManager")
+	UFUNCTION(BlueprintPure, meta = (DisplayName = "Not Equal (PrimaryAssetId)", CompactNodeTitle = "!=", ScriptOperator="!=", BlueprintThreadSafe), Category = "AssetManager")
 	static bool NotEqual_PrimaryAssetId(FPrimaryAssetId A, FPrimaryAssetId B);
 
 	/** Returns list of Primary Asset Ids for a PrimaryAssetType */
@@ -1985,7 +1986,7 @@ class ENGINE_API UKismetSystemLibrary : public UBlueprintFunctionLibrary
 	static bool IsValidPrimaryAssetType(FPrimaryAssetType PrimaryAssetType);
 
 	/** Converts a Primary Asset Type to a string. The other direction is not provided because it cannot be validated */
-	UFUNCTION(BlueprintPure, meta = (DisplayName = "ToString (PrimaryAssetType)", CompactNodeTitle = "->", ScriptMethod="ToString", BlueprintThreadSafe), Category = "AssetManager")
+	UFUNCTION(BlueprintPure, meta = (DisplayName = "To String (PrimaryAssetType)", CompactNodeTitle = "->", ScriptMethod="ToString", BlueprintThreadSafe), Category = "AssetManager")
 	static FString Conv_PrimaryAssetTypeToString(FPrimaryAssetType PrimaryAssetType);
 
 	/** Returns true if the values are equal (A == B) */
@@ -1993,7 +1994,7 @@ class ENGINE_API UKismetSystemLibrary : public UBlueprintFunctionLibrary
 	static bool EqualEqual_PrimaryAssetType(FPrimaryAssetType A, FPrimaryAssetType B);
 
 	/** Returns true if the values are not equal (A != B) */
-	UFUNCTION(BlueprintPure, meta = (DisplayName = "NotEqual (PrimaryAssetType)", CompactNodeTitle = "!=", ScriptOperator="!=", BlueprintThreadSafe), Category = "AssetManager")
+	UFUNCTION(BlueprintPure, meta = (DisplayName = "Not Equal (PrimaryAssetType)", CompactNodeTitle = "!=", ScriptOperator="!=", BlueprintThreadSafe), Category = "AssetManager")
 	static bool NotEqual_PrimaryAssetType(FPrimaryAssetType A, FPrimaryAssetType B);
 
 	/** Unloads a primary asset, which allows it to be garbage collected if nothing else is referencing it */
diff --git a/Engine/Source/Runtime/Engine/Classes/Kismet/KismetTextLibrary.h b/Engine/Source/Runtime/Engine/Classes/Kismet/KismetTextLibrary.h
--- a/Engine/Source/Runtime/Engine/Classes/Kismet/KismetTextLibrary.h
+++ b/Engine/Source/Runtime/Engine/Classes/Kismet/KismetTextLibrary.h
@@ -87,50 +87,50 @@ class ENGINE_API UKismetTextLibrary : public UBlueprintFunctionLibrary
 	GENERATED_UCLASS_BODY()
 
 	/** Converts a vector value to localized formatted text, in the form 'X= Y= Z=' */
-	UFUNCTION(BlueprintPure, meta = (DisplayName = "ToText (Vector)", CompactNodeTitle = "->", BlueprintAutocast), Category = "Utilities|Text")
+	UFUNCTION(BlueprintPure, meta = (DisplayName = "To Text (Vector)", CompactNodeTitle = "->", BlueprintAutocast), Category = "Utilities|Text")
 	static FText Conv_VectorToText(FVector InVec);
 
 	/** Converts a vector2d value to localized formatted text, in the form 'X= Y=' */
-	UFUNCTION(BlueprintPure, meta = (DisplayName = "ToText (vector2d)", CompactNodeTitle = "->", BlueprintAutocast), Category = "Utilities|Text")
+	UFUNCTION(BlueprintPure, meta = (DisplayName = "To Text (Vector2d)", CompactNodeTitle = "->", BlueprintAutocast), Category = "Utilities|Text")
 	static FText Conv_Vector2dToText(FVector2D InVec);
 
 	/** Converts a rotator value to localized formatted text, in the form 'P= Y= R=' */
-	UFUNCTION(BlueprintPure, meta = (DisplayName = "ToText (rotator)", CompactNodeTitle = "->", BlueprintAutocast), Category = "Utilities|Text")
+	UFUNCTION(BlueprintPure, meta = (DisplayName = "To Text (Rotator)", CompactNodeTitle = "->", BlueprintAutocast), Category = "Utilities|Text")
 	static FText Conv_RotatorToText(FRotator InRot);
 
 	/** Converts a transform value to localized formatted text, in the form 'Translation: X= Y= Z= Rotation: P= Y= R= Scale: X= Y= Z=' */
-	UFUNCTION(BlueprintPure, meta = (DisplayName = "ToText (transform)", CompactNodeTitle = "->", BlueprintAutocast), Category = "Utilities|Text")
+	UFUNCTION(BlueprintPure, meta = (DisplayName = "To Text (Transform)", CompactNodeTitle = "->", BlueprintAutocast), Category = "Utilities|Text")
 	static FText Conv_TransformToText(const FTransform& InTrans);
 
 	/** Converts a UObject value to culture invariant text by calling the object's GetName method */
-	UFUNCTION(BlueprintPure, meta = (DisplayName = "ToText (object)", BlueprintAutocast), Category = "Utilities|Text")
+	UFUNCTION(BlueprintPure, meta = (DisplayName = "To Text (Object)", BlueprintAutocast), Category = "Utilities|Text")
 	static FText Conv_ObjectToText(class UObject* InObj);
 
 	/** Converts a linear color value to localized formatted text, in the form '(R=,G=,B=,A=)' */
-	UFUNCTION(BlueprintPure, meta = (DisplayName = "ToText (linear color)", CompactNodeTitle = "->", BlueprintAutocast), Category = "Utilities|Text")
+	UFUNCTION(BlueprintPure, meta = (DisplayName = "To Text (LinearColor)", CompactNodeTitle = "->", BlueprintAutocast), Category = "Utilities|Text")
 	static FText Conv_ColorToText(FLinearColor InColor);
 
 	/** Converts localizable text to the string */
-	UFUNCTION(BlueprintPure, meta = (DisplayName = "ToString (text)", CompactNodeTitle = "->", BlueprintAutocast), Category = "Utilities|String")
+	UFUNCTION(BlueprintPure, meta = (DisplayName = "To String (Text)", CompactNodeTitle = "->", BlueprintAutocast), Category = "Utilities|String")
 	static FString Conv_TextToString(const FText& InText);
 
 	/** Converts string to culture invariant text. Use Format or Make Literal Text to create localizable text */
-	UFUNCTION(BlueprintPure, meta=(DisplayName = "ToText (string)", BlueprintAutocast), Category="Utilities|Text")
+	UFUNCTION(BlueprintPure, meta=(DisplayName = "To Text (String)", BlueprintAutocast), Category="Utilities|Text")
 	static FText Conv_StringToText(const FString& InString);
 
 	/** Converts Name to culture invariant text */
-	UFUNCTION(BlueprintPure, meta=(DisplayName = "ToText (name)", BlueprintAutocast), Category="Utilities|Text")
+	UFUNCTION(BlueprintPure, meta=(DisplayName = "To Text (Name)", BlueprintAutocast), Category="Utilities|Text")
 	static FText Conv_NameToText(FName InName);
 
-	/* Returns true if text is empty. */
+	/** Returns true if text is empty. */
 	UFUNCTION(BlueprintPure, Category="Utilities|Text")
 	static bool TextIsEmpty(const FText& InText);
 
-	/* Returns true if text is transient. */
+	/** Returns true if text is transient. */
 	UFUNCTION(BlueprintPure, Category="Utilities|Text")
 	static bool TextIsTransient(const FText& InText);
 
-	/* Returns true if text is culture invariant. */
+	/** Returns true if text is culture invariant. */
 	UFUNCTION(BlueprintPure, Category="Utilities|Text")
 	static bool TextIsCultureInvariant(const FText& InText);
 
@@ -148,62 +148,62 @@ class ENGINE_API UKismetTextLibrary : public UBlueprintFunctionLibrary
 	UFUNCTION(BlueprintPure, Category="Utilities|Text")
 	static FText TextToUpper(const FText& InText);
 
-	/* Removes whitespace characters from the front of the text. */
+	/** Removes whitespace characters from the front of the text. */
 	UFUNCTION(BlueprintPure, Category="Utilities|Text")
 	static FText TextTrimPreceding(const FText& InText);
 
-	/* Removes trailing whitespace characters. */
+	/** Removes trailing whitespace characters. */
 	UFUNCTION(BlueprintPure, Category="Utilities|Text")
 	static FText TextTrimTrailing(const FText& InText);
 
-	/* Removes whitespace characters from the front and end of the text. */
+	/** Removes whitespace characters from the front and end of the text. */
 	UFUNCTION(BlueprintPure, Category="Utilities|Text")
 	static FText TextTrimPrecedingAndTrailing(const FText& InText);
 
-	/* Returns an empty piece of text. */
+	/** Returns an empty piece of text. */
 	UFUNCTION(BlueprintPure, Category="Utilities|Text")
 	static FText GetEmptyText();
 
-	/* Attempts to find existing Text using the representation found in the loc tables for the specified namespace and key. */
+	/** Attempts to find existing Text using the representation found in the loc tables for the specified namespace and key. */
 	UFUNCTION(BlueprintPure, Category="Utilities|Text")
 	static bool FindTextInLocalizationTable(const FString& Namespace, const FString& Key, FText& OutText);
 
-	/* Returns true if A and B are linguistically equal (A == B). */
-	UFUNCTION(BlueprintPure, meta=(DisplayName = "Equal (text)", CompactNodeTitle = "=="), Category="Utilities|Text")
+	/** Returns true if A and B are linguistically equal (A == B). */
+	UFUNCTION(BlueprintPure, meta=(DisplayName = "Equal Exactly (Text)", CompactNodeTitle = "==="), Category="Utilities|Text")
 	static bool EqualEqual_TextText(const FText& A, const FText& B);
 
-	/* Returns true if A and B are linguistically equal (A == B), ignoring case. */
-	UFUNCTION(BlueprintPure, meta=(DisplayName = "Equal, Case Insensitive (text)", CompactNodeTitle = "=="), Category="Utilities|Text")
+	/** Returns true if A and B are linguistically equal (A == B), ignoring case. */
+	UFUNCTION(BlueprintPure, meta=(DisplayName = "Equal, Case Insensitive (Text)", CompactNodeTitle = "=="), Category="Utilities|Text")
 	static bool EqualEqual_IgnoreCase_TextText(const FText& A, const FText& B);
 				
-	/* Returns true if A and B are linguistically not equal (A != B). */
-	UFUNCTION(BlueprintPure, meta=(DisplayName = "NotEqual (text)", CompactNodeTitle = "!="), Category="Utilities|Text")
+	/** Returns true if A and B are linguistically not equal (A != B). */
+	UFUNCTION(BlueprintPure, meta=(DisplayName = "Not Equal Exactly (Text)", CompactNodeTitle = "!=="), Category="Utilities|Text")
 	static bool NotEqual_TextText(const FText& A, const FText& B);
 
-	/* Returns true if A and B are linguistically not equal (A != B), ignoring case. */
-	UFUNCTION(BlueprintPure, meta=(DisplayName = "NotEqual, Case Insensitive (text)", CompactNodeTitle = "!="), Category="Utilities|Text")
+	/** Returns true if A and B are linguistically not equal (A != B), ignoring case. */
+	UFUNCTION(BlueprintPure, meta=(DisplayName = "Not Equal, Case Insensitive (Text)", CompactNodeTitle = "!="), Category="Utilities|Text")
 	static bool NotEqual_IgnoreCase_TextText(const FText& A, const FText& B);
 
 	/** Converts a boolean value to formatted text, either 'true' or 'false' */
-	UFUNCTION(BlueprintPure, meta=(DisplayName = "ToText (boolean)", CompactNodeTitle = "->", BlueprintAutocast), Category="Utilities|Text")
+	UFUNCTION(BlueprintPure, meta=(DisplayName = "To Text (Boolean)", CompactNodeTitle = "->", BlueprintAutocast), Category="Utilities|Text")
 	static FText Conv_BoolToText(bool InBool);
 
 	/** Converts a byte value to formatted text */
-	UFUNCTION(BlueprintPure, meta=(DisplayName = "ToText (byte)", CompactNodeTitle = "->", BlueprintAutocast), Category="Utilities|Text")
+	UFUNCTION(BlueprintPure, meta=(DisplayName = "To Text (Byte)", CompactNodeTitle = "->", BlueprintAutocast), Category="Utilities|Text")
 	static FText Conv_ByteToText(uint8 Value);
 
 	// Default values are duplicated from FNumberFormattingOptions and should be replicated in all functions and in the struct when changed!
-	/* Converts a passed in integer to text based on formatting options */
-	UFUNCTION(BlueprintPure, meta=(DisplayName = "ToText (integer)", AdvancedDisplay = "1", BlueprintAutocast), Category="Utilities|Text")
+	/** Converts a passed in integer to text based on formatting options */
+	UFUNCTION(BlueprintPure, meta=(DisplayName = "To Text (Integer)", AdvancedDisplay = "1", BlueprintAutocast), Category="Utilities|Text")
 	static FText Conv_IntToText(int32 Value, bool bAlwaysSign = false, bool bUseGrouping = true, int32 MinimumIntegralDigits = 1, int32 MaximumIntegralDigits = 324);
 
-	/* Converts a passed in integer to text based on formatting options */
-	UFUNCTION(BlueprintPure, meta = (DisplayName = "ToText (int64)", AdvancedDisplay = "1", BlueprintAutocast), Category = "Utilities|Text")
+	/** Converts a passed in integer to text based on formatting options */
+	UFUNCTION(BlueprintPure, meta = (DisplayName = "To Text (Integer64)", AdvancedDisplay = "1", BlueprintAutocast), Category = "Utilities|Text")
 	static FText Conv_Int64ToText(int64 Value, bool bAlwaysSign = false, bool bUseGrouping = true, int32 MinimumIntegralDigits = 1, int32 MaximumIntegralDigits = 324);
 
 	// Default values are duplicated from FNumberFormattingOptions and should be replicated in all functions and in the struct when changed!
-	/* Converts a passed in float to text based on formatting options */
-	UFUNCTION(BlueprintPure, meta=(DisplayName = "ToText (float)", AdvancedDisplay = "1", BlueprintAutocast), Category="Utilities|Text")
+	/** Converts a passed in float to text based on formatting options */
+	UFUNCTION(BlueprintPure, meta=(DisplayName = "To Text (Float)", AdvancedDisplay = "1", BlueprintAutocast), Category="Utilities|Text")
 	static FText Conv_FloatToText(float Value, TEnumAsByte<ERoundingMode> RoundingMode, bool bAlwaysSign = false, bool bUseGrouping = true, int32 MinimumIntegralDigits = 1, int32 MaximumIntegralDigits = 324, int32 MinimumFractionalDigits = 0, int32 MaximumFractionalDigits = 3);
 	
 	/**
@@ -212,53 +212,53 @@ class ENGINE_API UKismetTextLibrary : public UBlueprintFunctionLibrary
 	 * Keep in mind the CurrencyCode is completely independent of the culture it's displayed in (and they do not imply one another).
 	 * For example: FText::AsCurrencyBase(650, TEXT("EUR")); would return an FText of "<EUR>6.50" in most English cultures (en_US/en_UK) and "6,50<EUR>" in Spanish (es_ES) (where <EUR> is U+20AC)
 	 */
-	UFUNCTION(BlueprintPure, meta = (DisplayName = "AsCurrency"), Category = "Utilities|Text")
+	UFUNCTION(BlueprintPure, meta = (DisplayName = "As Currency"), Category = "Utilities|Text")
 	static FText AsCurrencyBase(int32 BaseValue, const FString& CurrencyCode);
 
 	// Default values are duplicated from FNumberFormattingOptions and should be replicated in all functions and in the struct when changed!
-	/* Converts a passed in integer to a text formatted as a currency */
+	/** Converts a passed in integer to a text formatted as a currency */
 	UFUNCTION(BlueprintPure, meta=(DisplayName = "AsCurrency (integer) - DEPRECATED (use AsCurrency)"), Category="Utilities|Text")
 	static FText AsCurrency_Integer(int32 Value, TEnumAsByte<ERoundingMode> RoundingMode, bool bAlwaysSign = false, bool bUseGrouping = true, int32 MinimumIntegralDigits = 1, int32 MaximumIntegralDigits = 324, int32 MinimumFractionalDigits = 0, int32 MaximumFractionalDigits = 3, const FString& CurrencyCode = TEXT(""));
 
 	// Default values are duplicated from FNumberFormattingOptions and should be replicated in all functions and in the struct when changed!
-	/* Converts a passed in float to a text formatted as a currency */
+	/** Converts a passed in float to a text formatted as a currency */
 	UFUNCTION(BlueprintPure, meta=(DisplayName = "AsCurrency (float) - DEPRECATED (use AsCurrency)"), Category="Utilities|Text")
 	static FText AsCurrency_Float(float Value, TEnumAsByte<ERoundingMode> RoundingMode, bool bAlwaysSign = false, bool bUseGrouping = true, int32 MinimumIntegralDigits = 1, int32 MaximumIntegralDigits = 324, int32 MinimumFractionalDigits = 0, int32 MaximumFractionalDigits = 3, const FString& CurrencyCode = TEXT(""));
 
 	// Default values are duplicated from FNumberFormattingOptions and should be replicated in all functions and in the struct when changed!
-	/* Converts a passed in float to a text, formatted as a percent */
-	UFUNCTION(BlueprintPure, meta=(DisplayName = "AsPercent", AdvancedDisplay = "1"), Category="Utilities|Text")
+	/** Converts a passed in float to a text, formatted as a percent */
+	UFUNCTION(BlueprintPure, meta=(DisplayName = "As Percent", AdvancedDisplay = "1"), Category="Utilities|Text")
 	static FText AsPercent_Float(float Value, TEnumAsByte<ERoundingMode> RoundingMode, bool bAlwaysSign = false, bool bUseGrouping = true, int32 MinimumIntegralDigits = 1, int32 MaximumIntegralDigits = 324, int32 MinimumFractionalDigits = 0, int32 MaximumFractionalDigits = 3);
 
-	/* Converts a passed in date & time to a text, formatted as a date using an invariant timezone. This will use the given date & time as-is, so it's assumed to already be in the correct timezone. */
-	UFUNCTION(BlueprintPure, meta=(DisplayName = "AsDate", AdvancedDisplay = "1"), Category="Utilities|Text")
+	/** Converts a passed in date & time to a text, formatted as a date using an invariant timezone. This will use the given date & time as-is, so it's assumed to already be in the correct timezone. */
+	UFUNCTION(BlueprintPure, meta=(DisplayName = "As Date", AdvancedDisplay = "1"), Category="Utilities|Text")
 	static FText AsDate_DateTime(const FDateTime& InDateTime);
 
-	/* Converts a passed in date & time to a text, formatted as a date using the given timezone (default is the local timezone). This will convert the given date & time from UTC to the given timezone (taking into account DST). */
-	UFUNCTION(BlueprintPure, meta=(DisplayName = "AsDate (from UTC)", AdvancedDisplay = "1"), Category="Utilities|Text")
+	/** Converts a passed in date & time to a text, formatted as a date using the given timezone (default is the local timezone). This will convert the given date & time from UTC to the given timezone (taking into account DST). */
+	UFUNCTION(BlueprintPure, meta=(DisplayName = "As Date (from UTC)", AdvancedDisplay = "1"), Category="Utilities|Text")
 	static FText AsTimeZoneDate_DateTime(const FDateTime& InDateTime, const FString& InTimeZone = TEXT(""));
 
-	/* Converts a passed in date & time to a text, formatted as a date & time using an invariant timezone. This will use the given date & time as-is, so it's assumed to already be in the correct timezone. */
-	UFUNCTION(BlueprintPure, meta=(DisplayName = "AsDateTime", AdvancedDisplay = "1"), Category="Utilities|Text")
+	/** Converts a passed in date & time to a text, formatted as a date & time using an invariant timezone. This will use the given date & time as-is, so it's assumed to already be in the correct timezone. */
+	UFUNCTION(BlueprintPure, meta=(DisplayName = "As DateTime", AdvancedDisplay = "1"), Category="Utilities|Text")
 	static FText AsDateTime_DateTime(const FDateTime& In);
 
-	/* Converts a passed in date & time to a text, formatted as a date & time using the given timezone (default is the local timezone). This will convert the given date & time from UTC to the given timezone (taking into account DST). */
-	UFUNCTION(BlueprintPure, meta=(DisplayName = "AsDateTime (from UTC)", AdvancedDisplay = "1"), Category="Utilities|Text")
+	/** Converts a passed in date & time to a text, formatted as a date & time using the given timezone (default is the local timezone). This will convert the given date & time from UTC to the given timezone (taking into account DST). */
+	UFUNCTION(BlueprintPure, meta=(DisplayName = "As DateTime (from UTC)", AdvancedDisplay = "1"), Category="Utilities|Text")
 	static FText AsTimeZoneDateTime_DateTime(const FDateTime& InDateTime, const FString& InTimeZone = TEXT(""));
 
-	/* Converts a passed in date & time to a text, formatted as a time using an invariant timezone. This will use the given date & time as-is, so it's assumed to already be in the correct timezone. */
-	UFUNCTION(BlueprintPure, meta=(DisplayName = "AsTime", AdvancedDisplay = "1"), Category="Utilities|Text")
+	/** Converts a passed in date & time to a text, formatted as a time using an invariant timezone. This will use the given date & time as-is, so it's assumed to already be in the correct timezone. */
+	UFUNCTION(BlueprintPure, meta=(DisplayName = "As Time", AdvancedDisplay = "1"), Category="Utilities|Text")
 	static FText AsTime_DateTime(const FDateTime& In);
 
-	/* Converts a passed in date & time to a text, formatted as a time using the given timezone (default is the local timezone). This will convert the given date & time from UTC to the given timezone (taking into account DST). */
-	UFUNCTION(BlueprintPure, meta=(DisplayName = "AsTime (from UTC)", AdvancedDisplay = "1"), Category="Utilities|Text")
+	/** Converts a passed in date & time to a text, formatted as a time using the given timezone (default is the local timezone). This will convert the given date & time from UTC to the given timezone (taking into account DST). */
+	UFUNCTION(BlueprintPure, meta=(DisplayName = "As Time (from UTC)", AdvancedDisplay = "1"), Category="Utilities|Text")
 	static FText AsTimeZoneTime_DateTime(const FDateTime& InDateTime, const FString& InTimeZone = TEXT(""));
 
-	/* Converts a passed in time span to a text, formatted as a time span */
-	UFUNCTION(BlueprintPure, meta=(DisplayName = "AsTimespan", AdvancedDisplay = "1"), Category="Utilities|Text")
+	/** Converts a passed in time span to a text, formatted as a time span */
+	UFUNCTION(BlueprintPure, meta=(DisplayName = "As Timespan", AdvancedDisplay = "1"), Category="Utilities|Text")
 	static FText AsTimespan_Timespan(const FTimespan& InTimespan);
 
-	/* Used for formatting text using the FText::Format function and utilized by the UK2Node_FormatText */
+	/** Used for formatting text using the FText::Format function and utilized by the UK2Node_FormatText */
 	UFUNCTION(BlueprintPure, meta=(BlueprintInternalUseOnly = "true"))
 	static FText Format(FText InPattern, TArray<FFormatArgumentData> InArgs);
 
diff --git a/Engine/Source/Runtime/Engine/Public/Net/NetPushModelHelpers.h b/Engine/Source/Runtime/Engine/Public/Net/NetPushModelHelpers.h
--- a/Engine/Source/Runtime/Engine/Public/Net/NetPushModelHelpers.h
+++ b/Engine/Source/Runtime/Engine/Public/Net/NetPushModelHelpers.h
@@ -11,10 +11,11 @@ class ENGINE_API UNetPushModelHelpers : public UBlueprintFunctionLibrary
 	GENERATED_BODY()
 
 public:
-
-	UFUNCTION(BlueprintCallable, Category="Net")
+	
+	/** Mark replicated property as dirty with the Push Model networking system */
+	UFUNCTION(BlueprintCallable, Category="Networking")
 	static void MarkPropertyDirty(UObject* Object, FName PropertyName);
 
-	UFUNCTION(BlueprintCallable, Category = "Net", Meta=(BlueprintInternalUseOnly = "true", HidePin = "Object|RepIndex|PropertyName"))
+	UFUNCTION(BlueprintCallable, Category = "Networking", Meta=(BlueprintInternalUseOnly = "true", HidePin = "Object|RepIndex|PropertyName"))
 	static void MarkPropertyDirtyFromRepIndex(UObject* Object, int32 RepIndex, FName PropertyName);
 };
\ No newline at end of file
diff --git a/Engine/Source/Runtime/Landscape/Classes/LandscapeProxy.h b/Engine/Source/Runtime/Landscape/Classes/LandscapeProxy.h
--- a/Engine/Source/Runtime/Landscape/Classes/LandscapeProxy.h
+++ b/Engine/Source/Runtime/Landscape/Classes/LandscapeProxy.h
@@ -354,7 +354,7 @@ struct FLandscapeProxyMaterialOverride
 #endif
 };
 
-UCLASS(Abstract, MinimalAPI, NotBlueprintable, NotPlaceable, hidecategories=(Display, Attachment, Physics, Debug, Lighting, LOD), showcategories=(Lighting, Rendering, "Utilities|Transformation"), hidecategories=(Mobility))
+UCLASS(Abstract, MinimalAPI, NotBlueprintable, NotPlaceable, hidecategories=(Display, Attachment, Physics, Debug, Lighting, LOD), showcategories=(Lighting, Rendering, Transformation), hidecategories=(Mobility))
 class ALandscapeProxy : public APartitionActor, public ILandscapeSplineInterface
 {
 	GENERATED_BODY()
diff --git a/Engine/Source/Runtime/NavigationSystem/Public/NavMesh/RecastNavMesh.h b/Engine/Source/Runtime/NavigationSystem/Public/NavMesh/RecastNavMesh.h
--- a/Engine/Source/Runtime/NavigationSystem/Public/NavMesh/RecastNavMesh.h
+++ b/Engine/Source/Runtime/NavigationSystem/Public/NavMesh/RecastNavMesh.h
@@ -406,7 +406,7 @@ namespace FNavMeshConfig
 }
 
 
-UCLASS(config=Engine, defaultconfig, hidecategories=(Input,Rendering,Tags,"Utilities|Transformation",Actor,Layers,Replication), notplaceable)
+UCLASS(config=Engine, defaultconfig, hidecategories=(Input,Rendering,Tags,Transformation,Actor,Layers,Replication), notplaceable)
 class NAVIGATIONSYSTEM_API ARecastNavMesh : public ANavigationData
 {
 	GENERATED_UCLASS_BODY()
diff --git a/Engine/Source/Runtime/TimeManagement/Public/TimeManagementBlueprintLibrary.h b/Engine/Source/Runtime/TimeManagement/Public/TimeManagementBlueprintLibrary.h
--- a/Engine/Source/Runtime/TimeManagement/Public/TimeManagementBlueprintLibrary.h
+++ b/Engine/Source/Runtime/TimeManagement/Public/TimeManagementBlueprintLibrary.h
@@ -19,11 +19,11 @@ class TIMEMANAGEMENT_API UTimeManagementBlueprintLibrary : public UBlueprintFunc
 
 public:
 	/** Converts an FrameRate to a float ie: 1/30 returns 0.0333333 */
-	UFUNCTION(BlueprintPure, meta = (DisplayName = "FrameRate to Seconds", BlueprintAutocast), Category = "Utilities|Time Management")
+	UFUNCTION(BlueprintPure, meta = (DisplayName = "FrameRate To Seconds", BlueprintAutocast), Category = "Utilities|Time Management")
 	static float Conv_FrameRateToSeconds(const FFrameRate& InFrameRate);
 
 	/** Converts an QualifiedFrameTime to seconds. */
-	UFUNCTION(BlueprintPure, meta = (DisplayName = "QualifiedFrameTime to Seconds", BlueprintAutocast), Category = "Utilities|Time Management")
+	UFUNCTION(BlueprintPure, meta = (DisplayName = "QualifiedFrameTime To Seconds", BlueprintAutocast), Category = "Utilities|Time Management")
 	static float Conv_QualifiedFrameTimeToSeconds(const FQualifiedFrameTime& InFrameTime);
 
 	/** Multiplies a value in seconds against a FrameRate to get a new FrameTime. */
@@ -31,11 +31,11 @@ public:
 	static FFrameTime Multiply_SecondsFrameRate(float TimeInSeconds, const FFrameRate& FrameRate);
 
 	/** Converts an Timecode to a string (hh:mm:ss:ff). If bForceSignDisplay then the number sign will always be prepended instead of just when expressing a negative time. */
-	UFUNCTION(BlueprintPure, meta = (DisplayName = "Timecode to String", BlueprintAutocast), Category = "Utilities|Time Management")
+	UFUNCTION(BlueprintPure, meta = (DisplayName = "Timecode To String", BlueprintAutocast), Category = "Utilities|Time Management")
 	static FString Conv_TimecodeToString(const FTimecode& InTimecode, bool bForceSignDisplay = false);
 
 	/** Verifies that this is a valid framerate with a non-zero denominator. */
-	UFUNCTION(BlueprintPure, meta = (DisplayName = "IsValid"), Category = "Utilities|Time Management")
+	UFUNCTION(BlueprintPure, meta = (DisplayName = "Is Valid Frame Rate"), Category = "Utilities|Time Management")
 	static bool IsValid_Framerate(const FFrameRate& InFrameRate);
 
 	/** Checks if this framerate is an even multiple of another framerate, ie: 60 is a multiple of 30, but 59.94 is not. */
