#ue5 - Fix CharacterMovement potentially using wrong movement mode if Pack/UnpackNetworkMovementMode virtuals are modified by users.

#udn https://udn.unrealengine.com/s/question/0D54z00006tK6UTCA0/character-movement-possible-bug

diff --git a/Engine/Source/Runtime/Engine/Private/Components/CharacterMovementComponent.cpp b/Engine/Source/Runtime/Engine/Private/Components/CharacterMovementComponent.cpp
--- a/Engine/Source/Runtime/Engine/Private/Components/CharacterMovementComponent.cpp
+++ b/Engine/Source/Runtime/Engine/Private/Components/CharacterMovementComponent.cpp
@@ -8992,7 +8992,7 @@ void FCharacterNetworkMoveData::ClientFillNetworkMoveData(const FSavedMove_Chara
 	Acceleration = ClientMove.Acceleration;
 	ControlRotation = ClientMove.SavedControlRotation;
 	CompressedMoveFlags = ClientMove.GetCompressedFlags();
-	MovementMode = ClientMove.MovementMode;
+	MovementMode = ClientMove.EndPackedMovementMode;
 
 	// Location, relative movement base, and ending movement mode is only used for error checking, so only fill in the more complex parts if actually required.
 	if (MoveType == ENetworkMoveType::NewMove)
@@ -9360,10 +9360,17 @@ void UCharacterMovementComponent::ServerMoveHandleClientError(float ClientTimeSt
 
 	// Client may send a null movement base when walking on bases with no relative location (to save bandwidth).
 	// In this case don't check movement base in error conditions, use the server one (which avoids an error based on differing bases). Position will still be validated.
-	if (ClientMovementBase == nullptr && ClientMovementMode == MOVE_Walking)
+	if (ClientMovementBase == nullptr)
 	{
-		ClientMovementBase = CharacterOwner->GetBasedMovement().MovementBase;
-		ClientBaseBoneName = CharacterOwner->GetBasedMovement().BoneName;
+		TEnumAsByte<EMovementMode> NetMovementMode(MOVE_None);
+		TEnumAsByte<EMovementMode> NetGroundMode(MOVE_None);
+		uint8 NetCustomMode(0);
+		UnpackNetworkMovementMode(ClientMovementMode, NetMovementMode, NetCustomMode, NetGroundMode);
+		if (NetMovementMode == MOVE_Walking)
+		{
+			ClientMovementBase = CharacterOwner->GetBasedMovement().MovementBase;
+			ClientBaseBoneName = CharacterOwner->GetBasedMovement().BoneName;
+		}
 	}
 
 	// Compute the client error from the server's position
