Strata: renamed "shared normal" to "shared tangent basis"

diff --git a/Engine/Shaders/Private/BasePassPixelShader.usf b/Engine/Shaders/Private/BasePassPixelShader.usf
--- a/Engine/Shaders/Private/BasePassPixelShader.usf
+++ b/Engine/Shaders/Private/BasePassPixelShader.usf
@@ -982,7 +982,7 @@ void FPixelShaderInOut_MainPS(
 	{
 		StrataPixelHeader.BSDFCount += StrataData.Layers[l].BSDFCount;
 	}
-	StrataPixelHeader.SharedNormals = MaterialParameters.SharedNormals;
+	StrataPixelHeader.SharedLocalBases = MaterialParameters.SharedLocalBases;
 
 #if !GBUFFER_HAS_PRECSHADOWFACTOR
 	StrataPixelHeader.SkipPrecShadowMask = true;
@@ -1005,7 +1005,7 @@ void FPixelShaderInOut_MainPS(
 	Specular = SLW_SPECULAR(SLWBSDF);
 	Roughness = SLW_ROUGHNESS(SLWBSDF);
 	Opacity = SLW_TOPMATERIALOPACITY(SLWBSDF);
-	MaterialParameters.WorldNormal = StrataPixelHeader.SharedNormals.Normals[BSDF_GETNORMALID(SLWBSDF)];
+	MaterialParameters.WorldNormal = StrataPixelHeader.SharedLocalBases.Normals[BSDF_GETSHAREDLOCALBASISID(SLWBSDF)];
 	MaterialParameters.TangentToWorld[2] = MaterialParameters.WorldNormal;
 	Anisotropy = 0.0f;
 	ShadingModel = MATERIAL_SHADINGMODEL_SINGLELAYERWATER;
@@ -1542,12 +1542,12 @@ void FPixelShaderInOut_MainPS(
 	// We must normalize each normal and tangent to avoid non normalised vectors due to per vertex interpolation or texture filtering,
 	// for the deferred (our packing relies on normalized normal) and forward (normals are going to be used as-is from registers) paths.
 	UNROLL
-	for (uint i = 0; i < StrataPixelHeader.SharedNormals.NormalCount; ++i)
+	for (uint i = 0; i < StrataPixelHeader.SharedLocalBases.Count; ++i)
 	{
-		StrataPixelHeader.SharedNormals.Normals[i] = normalize(StrataPixelHeader.SharedNormals.Normals[i]);
-		if (StrataGetSharedNormalType(StrataPixelHeader.SharedNormals.NormalTypes, i) == STRATA_BASIS_TYPE_TANGENT)
+		StrataPixelHeader.SharedLocalBases.Normals[i] = normalize(StrataPixelHeader.SharedLocalBases.Normals[i]);
+		if (StrataGetSharedLocalBasisType(StrataPixelHeader.SharedLocalBases.Types, i) == STRATA_BASIS_TYPE_TANGENT)
 		{
-			StrataPixelHeader.SharedNormals.Tangents[i] = normalize(StrataPixelHeader.SharedNormals.Tangents[i]);
+			StrataPixelHeader.SharedLocalBases.Tangents[i] = normalize(StrataPixelHeader.SharedLocalBases.Tangents[i]);
 		}
 	}
 #endif
diff --git a/Engine/Shaders/Private/HairStrands/HairStrandsMaterialCommon.ush b/Engine/Shaders/Private/HairStrands/HairStrandsMaterialCommon.ush
--- a/Engine/Shaders/Private/HairStrands/HairStrandsMaterialCommon.ush
+++ b/Engine/Shaders/Private/HairStrands/HairStrandsMaterialCommon.ush
@@ -205,7 +205,7 @@ FEvaluateOutput Evaluate(
 		OutSample.Specular	= HAIR_SPECULAR(HairBSDF);
 		OutSample.Backlit	= HAIR_BACKLIT(HairBSDF);
 		OutSample.Emissive	= BSDF_GETEMISSIVE(HairBSDF);
-		float3 Tangent		= normalize(MaterialParameters.SharedNormals.Normals[BSDF_GETNORMALID(HairBSDF)]);
+		float3 Tangent		= normalize(MaterialParameters.SharedLocalBases.Normals[BSDF_GETSHAREDLOCALBASISID(HairBSDF)]);
 		
 	  #if MATERIAL_TANGENTSPACENORMAL 
 		// Convert the tangent space input tangent 
diff --git a/Engine/Shaders/Private/HairStrands/HairStrandsVisibilityPS.usf b/Engine/Shaders/Private/HairStrands/HairStrandsVisibilityPS.usf
--- a/Engine/Shaders/Private/HairStrands/HairStrandsVisibilityPS.usf
+++ b/Engine/Shaders/Private/HairStrands/HairStrandsVisibilityPS.usf
@@ -137,7 +137,7 @@ void MainVisibility(
 		Backlit			 = HAIR_BACKLIT(HairBSDF);
 		float MaterialAO = GetMaterialAmbientOcclusion(PixelMaterialInputs);
 
-		Tangent = MaterialParameters.SharedNormals.Normals[BSDF_GETNORMALID(HairBSDF)];	// Strata normals are all world space already
+		Tangent = MaterialParameters.SharedLocalBases.Normals[BSDF_GETSHAREDLOCALBASISID(HairBSDF)];	// Strata normals are all world space already
 #else
 		BaseColor		 = GetMaterialBaseColor(PixelMaterialInputs);
 		Roughness		 = GetMaterialRoughness(PixelMaterialInputs);
diff --git a/Engine/Shaders/Private/MaterialTemplate.ush b/Engine/Shaders/Private/MaterialTemplate.ush
--- a/Engine/Shaders/Private/MaterialTemplate.ush
+++ b/Engine/Shaders/Private/MaterialTemplate.ush
@@ -398,7 +398,7 @@ struct FMaterialPixelParameters
 #endif
 
 #if TEMPLATE_USES_STRATA
-	FSharedNormals SharedNormals;
+	FSharedLocalBases SharedLocalBases;
 #endif
 };
 
@@ -2587,7 +2587,7 @@ void CalcPixelMaterialInputs(in out FMaterialPixelParameters Parameters, in out
 %s
 
 #if TEMPLATE_USES_STRATA
-	Parameters.SharedNormals = StrataInitialiseSharedNormals();
+	Parameters.SharedLocalBases = StrataInitialiseSharedLocalBases();
 #endif
 
 	// Note that here MaterialNormal can be in world space or tangent space
@@ -2646,7 +2646,7 @@ void CalcPixelMaterialInputsAnalyticDerivatives(in out FMaterialPixelParameters
 %s
 
 #if TEMPLATE_USES_STRATA
-	Parameters.SharedNormals = StrataInitialiseSharedNormals();
+	Parameters.SharedLocalBases = StrataInitialiseSharedLocalBases();
 #endif
 
 	// Note that here MaterialNormal can be in world space or tangent space
diff --git a/Engine/Shaders/Private/Strata/Strata.ush b/Engine/Shaders/Private/Strata/Strata.ush
--- a/Engine/Shaders/Private/Strata/Strata.ush
+++ b/Engine/Shaders/Private/Strata/Strata.ush
@@ -12,7 +12,7 @@
 #include "/Engine/Shared/StrataDefinitions.h"
 
 
-// During the base pass or forward rendering, shared normals are simply available from registers.
+// During the base pass or forward rendering, shared local bases are simply available from registers.
 // If a BSDF is loaded for processing from the strata buffer, normals are unpacked on demand to reduce VGPR pressure and increase occupancy.
 #ifndef STRATA_INLINE_SHADING
 #define STRATA_INLINE_SHADING 1
@@ -46,7 +46,7 @@
 struct FStrataBSDF
 {
 	float3				Weight;		// Weight is a float3 because, when loaded from the material buffer, it contains the BSDF throughput (coverage and transmittance combined)
-	uint				State;		// (27 unused bits) (1bit as HasAnisotropy) (2bits as SharedNormalIndex) (3bits as BSDFType)
+	uint				State;		// Extra informations and other enabled BSDF features
 
 #if STRATA_INLINE_SHADING
 	float3				Emissive;	// All bsdf can potientially have emissive. This is never written out to the strata gbuffer but instead accumulated in the emissive buffer during the base pass.
@@ -71,10 +71,10 @@ struct FStrataBSDF
 #define WRITE_BITS(X, BitCount, BitOffset, NewValue)	X = (X & (~TO_OFFSET_BITMASK(BitCount, BitOffset))) | (((NewValue) & TO_BITMASK(BitCount)) << (BitOffset))
 
 // Bit count for each component stored in a Strata header
-// For HEADER_BIT_BSDF_COUNT and  HEADER_BIT_SHAREDNORMAL_X, please refer to values in StrataDefinitions.h
+// For HEADER_BIT_BSDF_COUNT and  HEADER_BIT_SHAREDLOCALBASES_X, please refer to values in StrataDefinitions.h
 #define HEADER_BIT_COUNT_BSDF_COUNT				4
-#define HEADER_BIT_COUNT_SHAREDNORMAL_COUNT		3
-#define HEADER_BIT_COUNT_SHAREDNORMAL_TYPE		4
+#define HEADER_BIT_COUNT_SHAREDLOCALBASES_COUNT	3
+#define HEADER_BIT_COUNT_SHAREDLOCALBASES_TYPE	4
 #define HEADER_BIT_COUNT_AO						8
 #define HEADER_BIT_COUNT_SKIP_PRECSHADOW_MASK	1
 #define HEADER_BIT_COUNT_ZERO_PRECSHADOW_MASK	1
@@ -82,9 +82,9 @@ struct FStrataBSDF
 #define HEADER_BIT_COUNT___UNUSED___			10
 
 #define HEADER_BIT_OFFSET_BSDF_COUNT			0
-#define HEADER_BIT_OFFSET_SHAREDNORMAL_COUNT	(HEADER_BIT_OFFSET_BSDF_COUNT				+ HEADER_BIT_COUNT_BSDF_COUNT)
-#define HEADER_BIT_OFFSET_SHAREDNORMAL_TYPE		(HEADER_BIT_OFFSET_SHAREDNORMAL_COUNT		+ HEADER_BIT_COUNT_SHAREDNORMAL_COUNT)
-#define HEADER_BIT_OFFSET_AO					(HEADER_BIT_OFFSET_SHAREDNORMAL_TYPE		+ HEADER_BIT_COUNT_SHAREDNORMAL_TYPE)
+#define HEADER_BIT_OFFSET_SHAREDLOCALBASES_COUNT (HEADER_BIT_OFFSET_BSDF_COUNT				+ HEADER_BIT_COUNT_BSDF_COUNT)
+#define HEADER_BIT_OFFSET_SHAREDLOCALBASES_TYPE	(HEADER_BIT_OFFSET_SHAREDLOCALBASES_COUNT	+ HEADER_BIT_COUNT_SHAREDLOCALBASES_COUNT)
+#define HEADER_BIT_OFFSET_AO					(HEADER_BIT_OFFSET_SHAREDLOCALBASES_TYPE	+ HEADER_BIT_COUNT_SHAREDLOCALBASES_TYPE)
 #define HEADER_BIT_OFFSET_SKIP_PRECSHADOW_MASK	(HEADER_BIT_OFFSET_AO						+ HEADER_BIT_COUNT_AO)
 #define HEADER_BIT_OFFSET_ZERO_PRECSHADOW_MASK	(HEADER_BIT_OFFSET_SKIP_PRECSHADOW_MASK		+ HEADER_BIT_COUNT_SKIP_PRECSHADOW_MASK)
 #define HEADER_BIT_OFFSET_ISSINGLELAYERWATER	(HEADER_BIT_OFFSET_ZERO_PRECSHADOW_MASK		+ HEADER_BIT_COUNT_ZERO_PRECSHADOW_MASK)
@@ -92,11 +92,11 @@ struct FStrataBSDF
 #define HEADER_GETBSDFCOUNT(X)					READ_BITS( X, HEADER_BIT_COUNT_BSDF_COUNT, HEADER_BIT_OFFSET_BSDF_COUNT)
 #define HEADER_SETBSDFCOUNT(X, NewValue)		WRITE_BITS(X, HEADER_BIT_COUNT_BSDF_COUNT, HEADER_BIT_OFFSET_BSDF_COUNT, NewValue)
 
-#define HEADER_GETSHAREDNORMALCOUNT(X)			READ_BITS( X, HEADER_BIT_COUNT_SHAREDNORMAL_COUNT, HEADER_BIT_OFFSET_SHAREDNORMAL_COUNT)
-#define HEADER_SETSHAREDNORMALCOUNT(X, NewValue)WRITE_BITS(X, HEADER_BIT_COUNT_SHAREDNORMAL_COUNT, HEADER_BIT_OFFSET_SHAREDNORMAL_COUNT, NewValue)
+#define HEADER_GETSHAREDLOCALBASESCOUNT(X)			READ_BITS( X, HEADER_BIT_COUNT_SHAREDLOCALBASES_COUNT, HEADER_BIT_OFFSET_SHAREDLOCALBASES_COUNT)
+#define HEADER_SETSHAREDLOCALBASESCOUNT(X, NewValue)WRITE_BITS(X, HEADER_BIT_COUNT_SHAREDLOCALBASES_COUNT, HEADER_BIT_OFFSET_SHAREDLOCALBASES_COUNT, NewValue)
 
-#define HEADER_GETSHAREDNORMALTYPE(X)			READ_BITS( X, HEADER_BIT_COUNT_SHAREDNORMAL_TYPE, HEADER_BIT_OFFSET_SHAREDNORMAL_TYPE)
-#define HEADER_SETSHAREDNORMALTYPE(X, NewValue)	WRITE_BITS(X, HEADER_BIT_COUNT_SHAREDNORMAL_TYPE, HEADER_BIT_OFFSET_SHAREDNORMAL_TYPE, NewValue)
+#define HEADER_GETSHAREDLOCALBASESTYPE(X)			READ_BITS( X, HEADER_BIT_COUNT_SHAREDLOCALBASES_TYPE, HEADER_BIT_OFFSET_SHAREDLOCALBASES_TYPE)
+#define HEADER_SETSHAREDLOCALBASESTYPE(X, NewValue)	WRITE_BITS(X, HEADER_BIT_COUNT_SHAREDLOCALBASES_TYPE, HEADER_BIT_OFFSET_SHAREDLOCALBASES_TYPE, NewValue)
 
 #define HEADER_GETAO(X)							READ_BITS( X, HEADER_BIT_COUNT_AO, HEADER_BIT_OFFSET_AO)
 #define HEADER_SETAO(X, NewValue)				WRITE_BITS(X, HEADER_BIT_COUNT_AO, HEADER_BIT_OFFSET_AO, NewValue)
@@ -112,7 +112,7 @@ struct FStrataBSDF
 
 // Bit count for each component stored in a Strata BSDF state
 #define STATE_BIT_COUNT_BSDF					3
-#define STATE_BIT_COUNT_SHAREDNORMALID			2
+#define STATE_BIT_COUNT_SHAREDLOCALBASESID		2
 #define STATE_BIT_COUNT_ANISOTROPY				1
 #define STATE_BIT_COUNT_TOPLAYER				1
 #define STATE_BIT_COUNT_HASSSS					1
@@ -130,8 +130,8 @@ struct FStrataBSDF
 
 // Bit offset of each component store in BSDF state 
 #define STATE_BIT_OFFSET_BSDF			0
-#define STATE_BIT_OFFSET_SHAREDNORMALID	(STATE_BIT_OFFSET_BSDF				+ STATE_BIT_COUNT_BSDF)
-#define STATE_BIT_OFFSET_ANISOTROPY		(STATE_BIT_OFFSET_SHAREDNORMALID	+ STATE_BIT_COUNT_SHAREDNORMALID)
+#define STATE_BIT_OFFSET_SHAREDLOCALBASESID (STATE_BIT_OFFSET_BSDF			+ STATE_BIT_COUNT_BSDF)
+#define STATE_BIT_OFFSET_ANISOTROPY		(STATE_BIT_OFFSET_SHAREDLOCALBASESID+ STATE_BIT_COUNT_SHAREDLOCALBASESID)
 #define STATE_BIT_OFFSET_TOPLAYER		(STATE_BIT_OFFSET_ANISOTROPY		+ STATE_BIT_COUNT_ANISOTROPY)
 #define STATE_BIT_OFFSET_HASSSS			(STATE_BIT_OFFSET_TOPLAYER			+ STATE_BIT_COUNT_TOPLAYER)
 #define STATE_BIT_OFFSET_HASSSSPROFILE	(STATE_BIT_OFFSET_HASSSS			+ STATE_BIT_COUNT_HASSSS)
@@ -149,10 +149,10 @@ struct FStrataBSDF
 #define BSDF_GETTYPE(X)					READ_BITS(X.State,  STATE_BIT_COUNT_BSDF, STATE_BIT_OFFSET_BSDF)
 #define BSDF_SETTYPE(X, NewType)		WRITE_BITS(X.State, STATE_BIT_COUNT_BSDF, STATE_BIT_OFFSET_BSDF, NewType)
 
-// Accessors for Shared Normal
-// This assumes Normals[STRATA_MAX_SHARED_NORMAL_REGISTERS]
-#define BSDF_GETNORMALID(X)				READ_BITS(X.State,  STATE_BIT_COUNT_SHAREDNORMALID, STATE_BIT_OFFSET_SHAREDNORMALID)
-#define BSDF_SETNORMALID(X, NrmlIdx)	WRITE_BITS(X.State, STATE_BIT_COUNT_SHAREDNORMALID, STATE_BIT_OFFSET_SHAREDNORMALID, NrmlIdx)
+// Accessors for shared local bases
+// This assumes Normals[STRATA_MAX_SHAREDLOCALBASES_REGISTERS]
+#define BSDF_GETSHAREDLOCALBASISID(X)	READ_BITS(X.State,  STATE_BIT_COUNT_SHAREDLOCALBASESID, STATE_BIT_OFFSET_SHAREDLOCALBASESID)
+#define BSDF_SETSHAREDLOCALBASISID(X, NrmlIdx)	WRITE_BITS(X.State, STATE_BIT_COUNT_SHAREDLOCALBASESID, STATE_BIT_OFFSET_SHAREDLOCALBASESID, NrmlIdx)
 
 // Accessors for HasAnisotropy
 // Indicates if the BSDF has anisotropic feature
@@ -280,7 +280,7 @@ struct FStrataBSDF
 // Forward declarations
 FParticipatingMedia StrataSlabCreateParticipatingMedia(float3 DiffuseColor, float3 MeanFreePathCentimeters);
 void EnableSlabBSDFSimpleVolumetric(inout FStrataBSDF BSDF);
-void StrataRequestSharedNormalTangentBasis(inout uint NormalTypes, uint Index);
+void StrataRequestSharedLocalBasisTangent(inout uint Types, uint Index);
 float3 StrataUnpackNormal(uint PackedNormal);
 void StrataUnpackNormalAndTangent(inout float3 Normal, inout float3 Tangent, in uint InPacked);
 uint   PackRGBA8(float4 rgba);
@@ -478,7 +478,7 @@ FStrataData GetStrataSlabBSDF(
 	float ThinFilmThickness,
 	float FuzzAmount, float3 FuzzColor,
 	float Thickness, 
-	uint SharedNormalIndex, inout uint SharedNormalTypes)
+	uint SharedLocalBasisIndex, inout uint SharedLocalBasisTypes)
 {
 	FStrataData StrataData = GetInitialisedStrataDataWithReservedFirstSlot();
 
@@ -488,7 +488,7 @@ FStrataData GetStrataSlabBSDF(
 	const uint bHasFuzz = FuzzAmount > 0.0f ? 1 : 0;
 	if (bHasAnisotropy)
 	{
-		StrataRequestSharedNormalTangentBasis(SharedNormalTypes, SharedNormalIndex);
+		StrataRequestSharedLocalBasisTangent(SharedLocalBasisTypes, SharedLocalBasisIndex);
 	}
 
 	const float2 RoughnessXY = GetAnisotropicRoughness(Roughness, Anisotropy);
@@ -508,7 +508,7 @@ FStrataData GetStrataSlabBSDF(
 	//  - It will only be ran for the bottom layer only. Otherwise, the simple volume should be used. As such, Two-sided lighting can also be disabled later if part of the bottom layer.
 
 	BSDF_SETTYPE			(StrataData.Layers[0].BSDFs[0],		STRATA_BSDF_TYPE_SLAB);
-	BSDF_SETNORMALID		(StrataData.Layers[0].BSDFs[0],		SharedNormalIndex);
+	BSDF_SETSHAREDLOCALBASISID(StrataData.Layers[0].BSDFs[0],		SharedLocalBasisIndex);
 	BSDF_SETEMISSIVE		(StrataData.Layers[0].BSDFs[0],		Emissive);
 	BSDF_SETANISOTROPY		(StrataData.Layers[0].BSDFs[0],		bHasAnisotropy);
 	BSDF_SETISTOPLAYER		(StrataData.Layers[0].BSDFs[0],		0);
@@ -602,12 +602,12 @@ FStrataData GetStrataUnlitBSDF(float3 Emissive, float3 TransmittanceColor)
 	return StrataData;
 }
 
-FStrataData GetStrataHairBSDF(float3 BaseColor, float Scatter, float Specular, float Roughness, float Backlit, float3 Emissive, uint SharedNormalIndex)
+FStrataData GetStrataHairBSDF(float3 BaseColor, float Scatter, float Specular, float Roughness, float Backlit, float3 Emissive, uint SharedLocalBasisIndex)
 {
 	FStrataData StrataData = GetInitialisedStrataDataWithReservedFirstSlot();
 
 	BSDF_SETTYPE					(StrataData.Layers[0].BSDFs[0], STRATA_BSDF_TYPE_HAIR);
-	BSDF_SETNORMALID				(StrataData.Layers[0].BSDFs[0], SharedNormalIndex);
+	BSDF_SETSHAREDLOCALBASISID		(StrataData.Layers[0].BSDFs[0], SharedLocalBasisIndex);
 	BSDF_SETEMISSIVE				(StrataData.Layers[0].BSDFs[0], Emissive);
 	HAIR_BASECOLOR					(StrataData.Layers[0].BSDFs[0]) = BaseColor;
 	HAIR_SCATTER					(StrataData.Layers[0].BSDFs[0]) = Scatter;
@@ -622,12 +622,12 @@ FStrataData GetStrataHairBSDF(float3 BaseColor, float Scatter, float Specular, f
 
 FStrataData GetStrataSingleLayerWaterBSDF(
 	float3 BaseColor, float Metallic, float Specular, float Roughness, float3 Emissive, float TopMaterialOpacity,
-	float3 WaterAlbedo, float3 WaterExtinction, float WaterPhaseG, float3 ColorScaleBehindWater, uint SharedNormalIndex)
+	float3 WaterAlbedo, float3 WaterExtinction, float WaterPhaseG, float3 ColorScaleBehindWater, uint SharedLocalBasisIndex)
 {
 	FStrataData StrataData = GetInitialisedStrataDataWithReservedFirstSlot();
 
 	BSDF_SETTYPE					(StrataData.Layers[0].BSDFs[0], STRATA_BSDF_TYPE_SINGLELAYERWATER);
-	BSDF_SETNORMALID				(StrataData.Layers[0].BSDFs[0], SharedNormalIndex);
+	BSDF_SETSHAREDLOCALBASISID		(StrataData.Layers[0].BSDFs[0], SharedLocalBasisIndex);
 	BSDF_SETEMISSIVE				(StrataData.Layers[0].BSDFs[0], Emissive);
 	SLW_BASECOLOR					(StrataData.Layers[0].BSDFs[0]) = BaseColor;
 	SLW_METALLIC					(StrataData.Layers[0].BSDFs[0]) = Metallic;
@@ -728,32 +728,32 @@ uint  StrataLoadUint1(ByteAddressBuffer StrataBuffer, inout FStrataAddressing St
 
 
 ///////////////////////////////////////////////////////////////////////////////
-// Shared normals
+// Shared local bases
 
 #define STRATA_BASIS_TYPE_NORMAL 0u
 #define STRATA_BASIS_TYPE_TANGENT 1u
-struct FSharedNormals
+struct FSharedLocalBases
 {
-	uint NormalCount;
-	uint NormalTypes;
-	float3 Normals[STRATA_MAX_SHARED_NORMAL_REGISTERS];	// once registered, normals are always world space
-	float3 Tangents[STRATA_MAX_SHARED_NORMAL_REGISTERS];// idem for tangents
+	uint Count;
+	uint Types;
+	float3 Normals[STRATA_MAX_SHAREDLOCALBASES_REGISTERS];	// once registered, normals are always world space
+	float3 Tangents[STRATA_MAX_SHAREDLOCALBASES_REGISTERS];// idem for tangents
 };
 
-FSharedNormals StrataInitialiseSharedNormals()
+FSharedLocalBases StrataInitialiseSharedLocalBases()
 {
-	FSharedNormals SharedNormals = (FSharedNormals)0;
-	return SharedNormals;
+	FSharedLocalBases SharedLocalBases = (FSharedLocalBases)0;
+	return SharedLocalBases;
 }
 
-void StrataRequestSharedNormalTangentBasis(inout uint NormalTypes, uint Index)
+void StrataRequestSharedLocalBasisTangent(inout uint Types, uint Index)
 {
-	NormalTypes = NormalTypes | (STRATA_BASIS_TYPE_TANGENT << Index);
+	Types = Types | (STRATA_BASIS_TYPE_TANGENT << Index);
 }
 
-uint StrataGetSharedNormalType(in uint NormalTypes, uint Index)
+uint StrataGetSharedLocalBasisType(in uint Types, uint Index)
 {
-	return (NormalTypes >> Index) & 0x1 ? STRATA_BASIS_TYPE_TANGENT : STRATA_BASIS_TYPE_NORMAL;
+	return (Types >> Index) & 0x1 ? STRATA_BASIS_TYPE_TANGENT : STRATA_BASIS_TYPE_NORMAL;
 }
 
 struct FStrataPixelHeader
@@ -761,12 +761,12 @@ struct FStrataPixelHeader
 	uint BSDFCount;
 
 #if STRATA_INLINE_SHADING
-	// All the shared normal in VGPRS. This is only used in basepass and during forward rendering.
-	FSharedNormals SharedNormals;
+	// All the shared local bases in VGPRS. This is only used in basepass and during forward rendering.
+	FSharedLocalBases SharedLocalBases;
 #else
 	ByteAddressBuffer StrataBuffer;
-	uint SharedNormalByteOffset;
-	uint SharedNormalTypes;
+	uint SharedLocalBasesByteOffset;
+	uint SharedLocalBasesTypes;
 #endif
 
 	float MaterialAO;
@@ -780,10 +780,10 @@ FStrataPixelHeader InitialiseStrataPixelHeader()
 	FStrataPixelHeader Out;
 	Out.BSDFCount = 0;
 #if STRATA_INLINE_SHADING
-	Out.SharedNormals = (FSharedNormals)0;
+	Out.SharedLocalBases = (FSharedLocalBases)0;
 #else
-	Out.SharedNormalByteOffset = 0;
-	Out.SharedNormalTypes = 0;
+	Out.SharedLocalBasesByteOffset = 0;
+	Out.SharedLocalBasesTypes = 0;
 #endif
 	Out.MaterialAO = 1.0f;
 	Out.SkipPrecShadowMask = false;
@@ -796,14 +796,14 @@ float3x3 StrataGetBSDFSharedBasis(in FStrataPixelHeader StrataPixelHeader, in FS
 {
 #if STRATA_INLINE_SHADING
 	// Assume Normal & Tangent are already been normalized
-	float3 Normal  = StrataPixelHeader.SharedNormals.Normals[BSDF_GETNORMALID(BSDF)];
-	float3 Tangent = StrataPixelHeader.SharedNormals.Tangents[BSDF_GETNORMALID(BSDF)];
+	float3 Normal  = StrataPixelHeader.SharedLocalBases.Normals[BSDF_GETSHAREDLOCALBASISID(BSDF)];
+	float3 Tangent = StrataPixelHeader.SharedLocalBases.Tangents[BSDF_GETSHAREDLOCALBASISID(BSDF)];
 	float3 Bitangent = cross(Normal, Tangent);
 	return float3x3(Tangent, Bitangent, Normal);
 #else
 	const bool bHasAnisotropy = BSDF_GETANISOTROPY(BSDF) != 0;
-	uint PackedBasis = StrataPixelHeader.StrataBuffer.Load(StrataPixelHeader.SharedNormalByteOffset + BSDF_GETNORMALID(BSDF) * StrataAddressing.ByteOffsetToNextUint);
-	uint BasisType   = StrataGetSharedNormalType(StrataPixelHeader.SharedNormalTypes, BSDF_GETNORMALID(BSDF));
+	uint PackedBasis = StrataPixelHeader.StrataBuffer.Load(StrataPixelHeader.SharedLocalBasesByteOffset + BSDF_GETSHAREDLOCALBASISID(BSDF) * StrataAddressing.ByteOffsetToNextUint);
+	uint BasisType   = StrataGetSharedLocalBasisType(StrataPixelHeader.SharedLocalBasesTypes, BSDF_GETSHAREDLOCALBASISID(BSDF));
 	float3x3 OutTangentBasis;
 	if (BasisType == STRATA_BASIS_TYPE_NORMAL)
 	{
@@ -918,7 +918,7 @@ FStrataData StrataHorizontalMixingParameterBlending(FStrataData Foreground, FStr
 
 	// We can only parameter blend the slab 
 	BSDF_SETTYPE(ResultBSDF,			STRATA_BSDF_TYPE_SLAB);
-	BSDF_SETNORMALID(ResultBSDF,		NewNormalIndex);
+	BSDF_SETSHAREDLOCALBASISID(ResultBSDF,NewNormalIndex);
 
 	BSDF_SETISTOPLAYER(ResultBSDF,		0);
 	BSDF_SETANISOTROPY(ResultBSDF,		BSDF_GETANISOTROPY(ForegroundBSDF)				| BSDF_GETANISOTROPY(BackgroundBSDF));
@@ -981,7 +981,7 @@ FStrataData StrataAddParameterBlending(FStrataData A, FStrataData B, uint NewNor
 
 	// We can only parameter blend the slab 
 	BSDF_SETTYPE(ResultBSDF,			STRATA_BSDF_TYPE_SLAB);
-	BSDF_SETNORMALID(ResultBSDF,		NewNormalIndex);
+	BSDF_SETSHAREDLOCALBASISID(ResultBSDF,NewNormalIndex);
 
 	BSDF_SETISTOPLAYER(ResultBSDF,		0);
 	BSDF_SETANISOTROPY(ResultBSDF,		BSDF_GETANISOTROPY(ABSDF)						| BSDF_GETANISOTROPY(BBSDF));
@@ -1045,7 +1045,7 @@ FStrataData StrataVerticalLayeringParameterBlending(FStrataData Top, FStrataData
 
 	// We can only parameter blend the slab 
 	BSDF_SETTYPE(ResultBSDF,			STRATA_BSDF_TYPE_SLAB);
-	BSDF_SETNORMALID(ResultBSDF,		NewNormalIndex);
+	BSDF_SETSHAREDLOCALBASISID(ResultBSDF,NewNormalIndex);
 
 	BSDF_SETISTOPLAYER(ResultBSDF,		0);
 	BSDF_SETANISOTROPY(ResultBSDF,		0);																					// Anisotropy is disabled because it changes a lot of things
@@ -1727,16 +1727,16 @@ void PackStrataOut(
 	STRATA_STORE_UINT1(PackStrataHeaderIn(BSDFCount, StrataPixelHeader));
 #if STRATA_INLINE_SHADING
 	UNROLL
-	for (uint i = 0; i < StrataPixelHeader.SharedNormals.NormalCount; ++i)
+	for (uint i = 0; i < StrataPixelHeader.SharedLocalBases.Count; ++i)
 	{
-		const uint BasisType = StrataGetSharedNormalType(StrataPixelHeader.SharedNormals.NormalTypes, i);
+		const uint BasisType = StrataGetSharedLocalBasisType(StrataPixelHeader.SharedLocalBases.Types, i);
 		if (BasisType == STRATA_BASIS_TYPE_NORMAL)
 		{
-			STRATA_STORE_UINT1(StrataPackNormal(StrataPixelHeader.SharedNormals.Normals[i]));
+			STRATA_STORE_UINT1(StrataPackNormal(StrataPixelHeader.SharedLocalBases.Normals[i]));
 		}
 		else // if (BasisType == STRATA_BASIS_TYPE_TANGENT)
 		{
-			STRATA_STORE_UINT1(StrataPackNormalAndTangent(StrataPixelHeader.SharedNormals.Normals[i], StrataPixelHeader.SharedNormals.Tangents[i]));
+			STRATA_STORE_UINT1(StrataPackNormalAndTangent(StrataPixelHeader.SharedLocalBases.Normals[i], StrataPixelHeader.SharedLocalBases.Tangents[i]));
 		}
 	}
 #endif
@@ -1840,8 +1840,8 @@ uint PackStrataHeaderIn(uint InBSDFCount, FStrataPixelHeader InHeader)
 	HEADER_SETBSDFCOUNT(Out, InBSDFCount);
 
 #if STRATA_INLINE_SHADING
-	HEADER_SETSHAREDNORMALCOUNT(Out, InHeader.SharedNormals.NormalCount);
-	HEADER_SETSHAREDNORMALTYPE(Out, InHeader.SharedNormals.NormalTypes);
+	HEADER_SETSHAREDLOCALBASESCOUNT(Out, InHeader.SharedLocalBases.Count);
+	HEADER_SETSHAREDLOCALBASESTYPE(Out, InHeader.SharedLocalBases.Types);
 #endif
 
 	HEADER_SETAO(Out, uint(InHeader.MaterialAO * 255.0f));
@@ -1861,20 +1861,20 @@ FStrataPixelHeader UnpackStrataHeaderIn(ByteAddressBuffer StrataBuffer, inout FS
 	uint PackedHeader = 0;
 	STRATA_LOAD_UINT1(PackedHeader);
 	Out.BSDFCount    = HEADER_GETBSDFCOUNT(PackedHeader);
-	uint NormalCount = HEADER_GETSHAREDNORMALCOUNT(PackedHeader);
-	uint NormalTypes = HEADER_GETSHAREDNORMALTYPE(PackedHeader);
+	uint SharedLocalBasesCount = HEADER_GETSHAREDLOCALBASESCOUNT(PackedHeader);
+	uint SharedLocalBasesTypes = HEADER_GETSHAREDLOCALBASESTYPE(PackedHeader);
 
 #if STRATA_INLINE_SHADING==0
-	// Only keep the offset to the shared normal memory. We are going to load them on demand.
-	Out.SharedNormalByteOffset = StrataAddressing.CurrentAddr;
+	// Only keep the offset to the shared local bases memory. We are going to load them on demand.
+	Out.SharedLocalBasesByteOffset = StrataAddressing.CurrentAddr;
 	Out.StrataBuffer = StrataBuffer;
-	Out.SharedNormalTypes = NormalTypes;
+	Out.SharedLocalBasesTypes = SharedLocalBasesTypes;
 
-	// Now skip over the shared normal memory to be able to load the BSDF content
-	StrataAddressing.CurrentAddr += NormalCount * StrataAddressing.ByteOffsetToNextUint;
+	// Now skip over the shared local basis memory to be able to load the BSDF content
+	StrataAddressing.CurrentAddr += SharedLocalBasesCount * StrataAddressing.ByteOffsetToNextUint;
 
 	// And account for the space used by normal for the debug information
-	StrataAddressing.ReadBytes += NormalCount * STRATA_PACKED_NORMAL_STRIDE_BYTES;
+	StrataAddressing.ReadBytes += SharedLocalBasesCount * STRATA_PACKED_SHAREDLOCALBASIS_STRIDE_BYTES;
 #endif
 
 	Out.MaterialAO = float(HEADER_GETAO(PackedHeader)) * (1.0f / 255.0f);
diff --git a/Engine/Shaders/Private/Strata/StrataFurnaceTest.usf b/Engine/Shaders/Private/Strata/StrataFurnaceTest.usf
--- a/Engine/Shaders/Private/Strata/StrataFurnaceTest.usf
+++ b/Engine/Shaders/Private/Strata/StrataFurnaceTest.usf
@@ -137,7 +137,7 @@ void MainPS(
 			}
 
 			BSDF_SETTYPE(BSDF, STRATA_BSDF_TYPE_SLAB);
-			BSDF_SETNORMALID(BSDF, 0);
+			BSDF_SETSHAREDLOCALBASISID(BSDF, 0);
 			BSDF_SETANISOTROPY(BSDF, Anisoptropy != 0 ? 1 : 0);
 			BSDF_SETISTOPLAYER(BSDF, 1);
 			BSDF_SETHASF90(BSDF, 0);
diff --git a/Engine/Shaders/Private/Strata/StrataVisualize.usf b/Engine/Shaders/Private/Strata/StrataVisualize.usf
--- a/Engine/Shaders/Private/Strata/StrataVisualize.usf
+++ b/Engine/Shaders/Private/Strata/StrataVisualize.usf
@@ -188,7 +188,7 @@ void VisualizeMaterialPS(
 
 #define PRINTBSDFWEIGHT()	PRINTBSDFCHAR(_SPC_); PRINTBSDFCHAR(_MINUS_); PRINTBSDFCHAR(_W_); PRINTBSDFCHAR(_E_); PRINTBSDFCHAR(_I_); PRINTBSDFCHAR(_G_); PRINTBSDFCHAR(_H_);PRINTBSDFCHAR(_T_);  PRINTBSDFCHAR(_SPC_); PrintFloat(PixelPos, OutColor.xyz, FontColor, TopLeft, BSDF.Weight.x); TopLeft.x+=8*7; PRINTBSDFCHAR(_SPC_); PrintFloat(PixelPos, OutColor.xyz, FontColor, TopLeft, BSDF.Weight.y); TopLeft.x+=8*7; PRINTBSDFCHAR(_SPC_); PrintFloat(PixelPos, OutColor.xyz, FontColor, TopLeft, BSDF.Weight.z); TopLeft.x+=8*7; 
 #define PRINTBSDFGREYWEIGHT()PRINTBSDFCHAR(_SPC_);PRINTBSDFCHARGREY(_MINUS_); PRINTBSDFCHARGREY(_W_); PRINTBSDFCHARGREY(_E_); PRINTBSDFCHARGREY(_I_); PRINTBSDFCHARGREY(_G_); PRINTBSDFCHARGREY(_H_);PRINTBSDFCHARGREY(_T_);  PRINTBSDFCHAR(_SPC_); PrintFloat(PixelPos, OutColor.xyz, FontColor, TopLeft, BSDF_GETGREYWEIGHT(BSDF)); TopLeft.x+=8*23; 
-#define PRINTBSDFNORMALID()	PRINTBSDFCHAR(_SPC_); PRINTBSDFCHAR(_MINUS_); PRINTBSDFCHAR(_N_); PRINTBSDFCHAR(_R_); PRINTBSDFCHAR(_M_); PRINTBSDFCHAR(_I_); PRINTBSDFCHAR(_D_); PRINTBSDFCHAR(_SPC_); PrintFloatNoFraction(PixelPos, OutColor.xyz, FontColor, TopLeft, float(BSDF_GETNORMALID(BSDF)), 1); TopLeft.x+=8*1;
+#define PRINTBSDFSHAREDLOCALBASISID()	PRINTBSDFCHAR(_SPC_); PRINTBSDFCHAR(_MINUS_); PRINTBSDFCHAR(_N_); PRINTBSDFCHAR(_R_); PRINTBSDFCHAR(_M_); PRINTBSDFCHAR(_I_); PRINTBSDFCHAR(_D_); PRINTBSDFCHAR(_SPC_); PrintFloatNoFraction(PixelPos, OutColor.xyz, FontColor, TopLeft, float(BSDF_GETSHAREDLOCALBASISID(BSDF)), 1); TopLeft.x+=8*1;
 #define PRINTBSDFANISO()	PRINTBSDFCHAR(_SPC_); PRINTBSDFCHAR(_MINUS_); PRINTBSDFCHAR(_A_); PRINTBSDFCHAR(_N_); PRINTBSDFCHAR(_I_); PRINTBSDFCHAR(_S_); PRINTBSDFCHAR(_O_); PRINTBSDFCHAR(_SPC_); PrintFloatNoFraction(PixelPos, OutColor.xyz, FontColor, TopLeft, float(BSDF_GETANISOTROPY(BSDF)), 1); TopLeft.x+=8*1;
 #define PRINTBSDFTOPLAY()	PRINTBSDFCHAR(_SPC_); PRINTBSDFCHAR(_MINUS_); PRINTBSDFCHAR(_T_); PRINTBSDFCHAR(_O_); PRINTBSDFCHAR(_P_); PRINTBSDFCHAR(_L_); PRINTBSDFCHAR(_A_); PRINTBSDFCHAR(_Y_); PRINTBSDFCHAR(_SPC_); PrintFloatNoFraction(PixelPos, OutColor.xyz, FontColor, TopLeft, float(BSDF_GETISTOPLAYER(BSDF)), 1); TopLeft.x+=8*1;
 #define PRINTBSDFHASSCAT()	PRINTBSDFCHAR(_SPC_); PRINTBSDFCHAR(_MINUS_); PRINTBSDFCHAR(_H_); PRINTBSDFCHAR(_A_); PRINTBSDFCHAR(_S_); PRINTBSDFCHAR(_S_); PRINTBSDFCHAR(_C_); PRINTBSDFCHAR(_A_); PRINTBSDFCHAR(_T_); PRINTBSDFCHAR(_SPC_); PrintFloatNoFraction(PixelPos, OutColor.xyz, FontColor, TopLeft, float(BSDF_GETHASSSS(BSDF)), 1); TopLeft.x+=8*1;
@@ -314,7 +314,7 @@ void VisualizeMaterialPS(
 				{
 					PRINTBSDFWEIGHT();
 				}
-				PRINTBSDFNORMALID();
+				PRINTBSDFSHAREDLOCALBASISID();
 				PRINTBSDFANISO();
 				PRINTBSDFTOPLAY();
 				PRINTBSDFHASSCAT();
@@ -469,7 +469,7 @@ void VisualizeMaterialPS(
 				{
 					PRINTBSDFWEIGHT();
 				}
-				PRINTBSDFNORMALID();
+				PRINTBSDFSHAREDLOCALBASISID();
 				PRINTBSDFANISO();
 				PRINTBSDFTOPLAY();
 				NEXTLINE;
diff --git a/Engine/Shaders/Shared/StrataDefinitions.h b/Engine/Shaders/Shared/StrataDefinitions.h
--- a/Engine/Shaders/Shared/StrataDefinitions.h
+++ b/Engine/Shaders/Shared/StrataDefinitions.h
@@ -16,10 +16,10 @@
 #define STRATA_MAX_LAYER_COUNT				4
 #define STRATA_MAX_BSDF_COUNT				15
 
-// It this is changed, STATE_BITS_SHAREDNORMAL also needs to be updated
-#define STRATA_MAX_SHARED_NORMAL_REGISTERS	4
+// It this is changed, STATE_BIT_COUNT_SHAREDLOCALBASESID and HEADER_BIT_COUNT_SHAREDLOCALBASES_COUNT also needs to be updated
+#define STRATA_MAX_SHAREDLOCALBASES_REGISTERS		4
 
-#define STRATA_PACKED_NORMAL_STRIDE_BYTES	4
+#define STRATA_PACKED_SHAREDLOCALBASIS_STRIDE_BYTES	4
 
 
 #define STRATA_BSDF_TYPE_SLAB				0
diff --git a/Engine/Source/Runtime/Engine/Private/Materials/HLSLMaterialTranslator.cpp b/Engine/Source/Runtime/Engine/Private/Materials/HLSLMaterialTranslator.cpp
--- a/Engine/Source/Runtime/Engine/Private/Materials/HLSLMaterialTranslator.cpp
+++ b/Engine/Source/Runtime/Engine/Private/Materials/HLSLMaterialTranslator.cpp
@@ -160,7 +160,7 @@ FHLSLMaterialTranslator::FHLSLMaterialTranslator(FMaterial* InMaterial,
 ,	DynamicParticleParameterMask(0)
 ,	StrataValidFrontMaterialCodeChunkPostTranslate(INDEX_NONE)
 ,	NextFreeStrataShaderNormalIndex(0)
-,	FinalUsedSharedNormalCount(0)
+,	FinalUsedSharedLocalBasesCount(0)
 ,	NumVtSamples(0)
 ,	TargetPlatform(InTargetPlatform)
 {
@@ -1647,34 +1647,34 @@ void FHLSLMaterialTranslator::GetMaterialEnvironment(EShaderPlatform InPlatform,
 	if (bMaterialIsStrata)
 	{
 		StrataPixelNormalInitializerValues += FString::Printf(TEXT("\n\n\n\t// Strata normal and tangent\n"));
-		FinalUsedSharedNormalCount = 0;
-		FStrataRegisteredSharedNormal UsedSharedNormalInfo[STRATA_MAX_SHARED_NORMAL_REGISTERS];
+		FinalUsedSharedLocalBasesCount = 0;
+		FStrataRegisteredSharedLocalBasis UsedSharedLocalBasesInfo[STRATA_MAX_SHAREDLOCALBASES_REGISTERS];
 
 
 		/*
 		* The final output code/workflow for shared tangent basis should look like
 		*
-		* #define NORMAL_INDEX_0 0		// default, unused
-		* #define NORMAL_INDEX_1 0		// default, unused
-		* #define NORMAL_INDEX_2 0
+		* #define SHAREDLOCALBASIS_INDEX_0 0		// default, unused
+		* #define SHAREDLOCALBASIS_INDEX_1 0		// default, unused
+		* #define SHAREDLOCALBASIS_INDEX_2 0
 		*
-		* FStrataData BSDF0 = GetStrataSlabBSDF(... NORMAL_INDEX_0, NormalCode0 ...)
-		* FStrataData BSDF1 = GetStrataSlabBSDF(... NORMAL_INDEX_1, NormalCode1 ...)
+		* FStrataData BSDF0 = GetStrataSlabBSDF(... SHAREDLOCALBASIS_0, NormalCode0 ...)
+		* FStrataData BSDF1 = GetStrataSlabBSDF(... SHAREDLOCALBASIS_1, NormalCode1 ...)
 		*
 		* float3 NormalCode2 = lerp(NormalCode0, NormalCode1, mix)
-		* FStrataData BSDF2 = StrataHorizontalMixingParameterBlending(BSDF0, BSDF1, mix, NormalCode2, NORMAL_INDEX_2) // will internally create NormalCode2
+		* FStrataData BSDF2 = StrataHorizontalMixingParameterBlending(BSDF0, BSDF1, mix, NormalCode2, SHAREDLOCALBASIS_INDEX_2) // will internally create NormalCode2
 		*
-		* tParameters.SharedNormals.Normals[NORMAL_INDEX_2] = NormalCode2;
+		* tParameters.SharedLocalBases.Normals[SHAREDLOCALBASIS_INDEX_2] = NormalCode2;
 		* #if MATERIAL_TANGENTSPACENORMAL
-		* Parameters.SharedNormals.Normals[NORMAL_INDEX_2] *= Parameters.TwoSidedSign;
+		* Parameters.SharedLocalBases.Normals[SHAREDLOCALBASIS_INDEX_2] *= Parameters.TwoSidedSign;
 		* #endif
 		*/
 
 
-		// We go over all output BSDF and reserve actualspace for there shared tangent basis. Only up to STRATA_MAX_SHARED_NORMAL_REGISTERS.
-		// When we know we will need a basis, it is allocated as FinalSharedNormalIndex. A macro is created with propoer output index for the correct value to be output ot the correct slot by the BSDF.
-		// If too many output are required, we faile the compilation.
-		// STRATA_TODO: parameter blend in case it is not possible to output according to per pixel cost retriction.
+		// We go over all output BSDF and reserve actual space for there shared tangent basis. Only up to STRATA_MAX_SHAREDLOCALBASES_REGISTERS.
+		// When we know we will need a basis, it is allocated as FinalSharedLocalBasisIndex. A macro is created with proper output index for the correct value to be output ot the correct slot by the BSDF.
+		// If too many output are required, we fail the compilation.
+		// STRATA_TODO: parameter blend in case it is not possible to output according to per pixel cost restriction.
 		const FStrataMaterialCompilationInfo& StrataCompilationInfo = CodeChunkToStrataCompilationInfoMap[StrataValidFrontMaterialCodeChunkPostTranslate];
 		for (uint32 LayerIt = 0; LayerIt < StrataCompilationInfo.LayerCount; ++LayerIt)
 		{
@@ -1682,20 +1682,20 @@ void FHLSLMaterialTranslator::GetMaterialEnvironment(EShaderPlatform InPlatform,
 			for (uint32 BSDFIt = 0; BSDFIt < Layer.BSDFCount; ++BSDFIt)
 			{
 				const FStrataMaterialCompilationInfo::FBSDF& BSDF = Layer.BSDFs[BSDFIt];
-				if (BSDF.RegisteredSharedNormal.NormalCodeChunk == INDEX_NONE && BSDF.RegisteredSharedNormal.TangentCodeChunk == INDEX_NONE)
+				if (BSDF.RegisteredSharedLocalBasis.NormalCodeChunk == INDEX_NONE && BSDF.RegisteredSharedLocalBasis.TangentCodeChunk == INDEX_NONE)
 				{
 					// We skip null normal on certain BSDF, for instance unlit.
 					continue;
 				}
 
-				const FStrataSharedNormalInfo StrataSharedNormalInfo = StrataCompilationInfoGetMatchingSharedNormalInfo(BSDF.RegisteredSharedNormal);
+				const FStrataSharedLocalBasesInfo StrataSharedLocalBasesInfo = StrataCompilationInfoGetMatchingSharedLocalBasisInfo(BSDF.RegisteredSharedLocalBasis);
 
 				// First, we check that the normal/tangent has not already written out (avoid 2 BSDFs sharing the same normal to note generate the same code twice)
 				bool bAlreadyProcessed = false;
-				for (uint8 i = 0; i < FinalUsedSharedNormalCount; ++i)
+				for (uint8 i = 0; i < FinalUsedSharedLocalBasesCount; ++i)
 				{
-					if (UsedSharedNormalInfo[i].NormalCodeChunkHash == StrataSharedNormalInfo.SharedData.NormalCodeChunkHash &&
-						(UsedSharedNormalInfo[i].TangentCodeChunkHash == StrataSharedNormalInfo.SharedData.TangentCodeChunkHash || BSDF.RegisteredSharedNormal.TangentCodeChunk == INDEX_NONE))
+					if (UsedSharedLocalBasesInfo[i].NormalCodeChunkHash == StrataSharedLocalBasesInfo.SharedData.NormalCodeChunkHash &&
+						(UsedSharedLocalBasesInfo[i].TangentCodeChunkHash == StrataSharedLocalBasesInfo.SharedData.TangentCodeChunkHash || BSDF.RegisteredSharedLocalBasis.TangentCodeChunk == INDEX_NONE))
 					{
 						bAlreadyProcessed = true;
 						break;
@@ -1706,48 +1706,48 @@ void FHLSLMaterialTranslator::GetMaterialEnvironment(EShaderPlatform InPlatform,
 					continue;
 				}
 
-				check(FinalUsedSharedNormalCount < STRATA_MAX_SHARED_NORMAL_REGISTERS);
-				const uint8 FinalSharedNormalIndex = FinalUsedSharedNormalCount++;
-				UsedSharedNormalInfo[FinalSharedNormalIndex] = StrataSharedNormalInfo.SharedData;
+				check(FinalUsedSharedLocalBasesCount < STRATA_MAX_SHAREDLOCALBASES_REGISTERS);
+				const uint8 FinalSharedLocalBasisIndex = FinalUsedSharedLocalBasesCount++;
+				UsedSharedLocalBasesInfo[FinalSharedLocalBasisIndex] = StrataSharedLocalBasesInfo.SharedData;
 
 				// Write out normals
-				StrataPixelNormalInitializerValues += FString::Printf(TEXT("\tParameters.SharedNormals.Normals[%u] = %s;\n"), FinalSharedNormalIndex, *StrataSharedNormalInfo.SharedNormalCode);
+				StrataPixelNormalInitializerValues += FString::Printf(TEXT("\tParameters.SharedLocalBases.Normals[%u] = %s;\n"), FinalSharedLocalBasisIndex, *StrataSharedLocalBasesInfo.NormalCode);
 				StrataPixelNormalInitializerValues += FString::Printf(TEXT("\t#if MATERIAL_TANGENTSPACENORMAL\n"));
-				StrataPixelNormalInitializerValues += FString::Printf(TEXT("\tParameters.SharedNormals.Normals[%u] *= Parameters.TwoSidedSign;\n"), FinalSharedNormalIndex);
+				StrataPixelNormalInitializerValues += FString::Printf(TEXT("\tParameters.SharedLocalBases.Normals[%u] *= Parameters.TwoSidedSign;\n"), FinalSharedLocalBasisIndex);
 				StrataPixelNormalInitializerValues += FString::Printf(TEXT("\t#endif\n"));
 
 				// Write out tangents
-				if (StrataSharedNormalInfo.SharedData.TangentCodeChunk != INDEX_NONE)
+				if (StrataSharedLocalBasesInfo.SharedData.TangentCodeChunk != INDEX_NONE)
 				{
-					StrataPixelNormalInitializerValues += FString::Printf(TEXT("\tParameters.SharedNormals.Tangents[%u] = %s;\n"), FinalSharedNormalIndex, *StrataSharedNormalInfo.SharedTangentCode);
+					StrataPixelNormalInitializerValues += FString::Printf(TEXT("\tParameters.SharedLocalBases.Tangents[%u] = %s;\n"), FinalSharedLocalBasisIndex, *StrataSharedLocalBasesInfo.TangentCode);
 				}
 				else
 				{
-					StrataPixelNormalInitializerValues += FString::Printf(TEXT("\tParameters.SharedNormals.Tangents[%u] = float3(1,0,0);\n"), FinalSharedNormalIndex);
+					StrataPixelNormalInitializerValues += FString::Printf(TEXT("\tParameters.SharedLocalBases.Tangents[%u] = float3(1,0,0);\n"), FinalSharedLocalBasisIndex);
 				}
 				StrataPixelNormalInitializerValues += FString::Printf(TEXT("\t#if MATERIAL_TANGENTSPACENORMAL\n"));
-				StrataPixelNormalInitializerValues += FString::Printf(TEXT("\tParameters.SharedNormals.Tangents[%u] *= Parameters.TwoSidedSign;\n"), FinalSharedNormalIndex);
+				StrataPixelNormalInitializerValues += FString::Printf(TEXT("\tParameters.SharedLocalBases.Tangents[%u] *= Parameters.TwoSidedSign;\n"), FinalSharedLocalBasisIndex);
 				StrataPixelNormalInitializerValues += FString::Printf(TEXT("\t#endif\n"));
 			}
 		}
-		StrataPixelNormalInitializerValues += FString::Printf(TEXT("\tParameters.SharedNormals.NormalCount = %u;\n"), FinalUsedSharedNormalCount);
+		StrataPixelNormalInitializerValues += FString::Printf(TEXT("\tParameters.SharedLocalBases.Count = %u;\n"), FinalUsedSharedLocalBasesCount);
 
-		// Now write out all the macros, them mapping from the BSDF to the effective position/index in the shared normal array they should write to.
-		for (TMultiMap<uint64, FStrataSharedNormalInfo>::TIterator It(CodeChunkToStrataSharedNormal); It; ++It)
+		// Now write out all the macros, them mapping from the BSDF to the effective position/index in the shared local basis array they should write to.
+		for (TMultiMap<uint64, FStrataSharedLocalBasesInfo>::TIterator It(CodeChunkToStrataSharedLocalBasis); It; ++It)
 		{
-			// The default linear output index will be 0 by default, and different if in fact the shared normal point to one that is effectively in used in the array of shared normals.
+			// The default linear output index will be 0 by default, and different if in fact the shared local basis points to one that is effectively in used in the array of shared local bases.
 			uint8 LinearIndex = 0;
-			for (uint8 i = 0; i < FinalUsedSharedNormalCount; ++i)
+			for (uint8 i = 0; i < FinalUsedSharedLocalBasesCount; ++i)
 			{
-				if (UsedSharedNormalInfo[i].NormalCodeChunkHash == It->Value.SharedData.NormalCodeChunkHash &&
-					(UsedSharedNormalInfo[i].TangentCodeChunkHash == It->Value.SharedData.TangentCodeChunkHash || It->Value.SharedData.TangentCodeChunk == INDEX_NONE))
+				if (UsedSharedLocalBasesInfo[i].NormalCodeChunkHash == It->Value.SharedData.NormalCodeChunkHash &&
+					(UsedSharedLocalBasesInfo[i].TangentCodeChunkHash == It->Value.SharedData.TangentCodeChunkHash || It->Value.SharedData.TangentCodeChunk == INDEX_NONE))
 				{
 					LinearIndex = i;
 					break;
 				}
 			}
 
-			OutEnvironment.SetDefine(*GetStrataSharedNormalIndexMacro(It->Value.SharedData), LinearIndex);
+			OutEnvironment.SetDefine(*GetStrataSharedLocalBasisIndexMacro(It->Value.SharedData), LinearIndex);
 		}
 
 		// Now write some feedback to the user
@@ -1768,7 +1768,7 @@ void FHLSLMaterialTranslator::GetMaterialEnvironment(EShaderPlatform InPlatform,
 			StrataMaterialDescription += FString::Printf(TEXT("----- STRATA -----\r\n"));
 			StrataMaterialDescription += FString::Printf(TEXT("StrataCompilationInfo -\r\n"));
 			StrataMaterialDescription += FString::Printf(TEXT(" - TotalBSDFCount    = %i\r\n"), StrataCompilationInfo.TotalBSDFCount);
-			StrataMaterialDescription += FString::Printf(TEXT(" - SharedNormalCount = %i\r\n"), StrataCompilationInfoGetSharedNormalCount());
+			StrataMaterialDescription += FString::Printf(TEXT(" - SharedLocalBasesCount = %i\r\n"), StrataCompilationInfoGetSharedLocalBasesCount());
 
 			for (uint32 LayerIt = 0; LayerIt < StrataCompilationInfo.LayerCount; ++LayerIt)
 			{
@@ -1777,7 +1777,7 @@ void FHLSLMaterialTranslator::GetMaterialEnvironment(EShaderPlatform InPlatform,
 				const FStrataMaterialCompilationInfo::FLayer& Layer = StrataCompilationInfo.Layers[LayerIt];
 				for (uint32 BSDFIt = 0; BSDFIt < Layer.BSDFCount; ++BSDFIt)
 				{
-					StrataMaterialDescription += FString::Printf(TEXT("     - %s - SharedNormalIndexMacro = %s \r\n"), *GetStrataBSDFName(Layer.BSDFs[BSDFIt].Type), *GetStrataSharedNormalIndexMacro(Layer.BSDFs[BSDFIt].RegisteredSharedNormal));
+					StrataMaterialDescription += FString::Printf(TEXT("     - %s - SharedLocalBasisIndexMacro = %s \r\n"), *GetStrataBSDFName(Layer.BSDFs[BSDFIt].Type), *GetStrataSharedLocalBasisIndexMacro(Layer.BSDFs[BSDFIt].RegisteredSharedLocalBasis));
 				}
 			}
 
@@ -1939,7 +1939,7 @@ void FHLSLMaterialTranslator::GetSharedInputsMaterialCode(FString& PixelMembersD
 			PixelMembersInitializationEpilog += DerivativeVariations[DerivativeVariation].TranslatedCodeChunkDefinitions[LastProperty] + TEXT("\n");
 		}
 
-		if (CodeChunkToStrataSharedNormal.Num() > 0)
+		if (CodeChunkToStrataSharedLocalBasis.Num() > 0)
 		{
 			PixelInputInitializerValues += StrataPixelNormalInitializerValues;
 		}
@@ -8892,52 +8892,52 @@ const FStrataMaterialCompilationInfo& FHLSLMaterialTranslator::GetStrataCompilat
 	return CodeChunkToStrataCompilationInfoMap[CodeChunk];
 }
 
-FStrataRegisteredSharedNormal FHLSLMaterialTranslator::StrataCompilationInfoRegisterSharedNormal(int32 NormalCodeChunk)
+FStrataRegisteredSharedLocalBasis FHLSLMaterialTranslator::StrataCompilationInfoRegisterSharedLocalBasis(int32 NormalCodeChunk)
 {
 	check(NormalCodeChunk != INDEX_NONE);
-	check(NextFreeStrataShaderNormalIndex < 255);	// Out of shared normal slots
+	check(NextFreeStrataShaderNormalIndex < 255);	// Out of shared local basis slots
 
-	FStrataRegisteredSharedNormal StrataRegisteredSharedNormal;
-	StrataRegisteredSharedNormal.NormalCodeChunk		= NormalCodeChunk;
-	StrataRegisteredSharedNormal.NormalCodeChunkHash	= GetParameterHash(NormalCodeChunk); 
-	StrataRegisteredSharedNormal.TangentCodeChunk		= INDEX_NONE;
-	StrataRegisteredSharedNormal.TangentCodeChunkHash	= GetParameterHash(INDEX_NONE);
-	StrataRegisteredSharedNormal.GraphSharedNormalIndex = NextFreeStrataShaderNormalIndex++;
+	FStrataRegisteredSharedLocalBasis StrataRegisteredSharedLocalBasis;
+	StrataRegisteredSharedLocalBasis.NormalCodeChunk		= NormalCodeChunk;
+	StrataRegisteredSharedLocalBasis.NormalCodeChunkHash	= GetParameterHash(NormalCodeChunk);
+	StrataRegisteredSharedLocalBasis.TangentCodeChunk		= INDEX_NONE;
+	StrataRegisteredSharedLocalBasis.TangentCodeChunkHash	= GetParameterHash(INDEX_NONE);
+	StrataRegisteredSharedLocalBasis.GraphSharedLocalBasisIndex = NextFreeStrataShaderNormalIndex++;
 
 	// Find any basis which match the Normal code chunk
 	// A normal can be duplicated when it is paired with different tangent, so find the first one which matches
-	TArray<FStrataSharedNormalInfo*> NormalInfos;
-	CodeChunkToStrataSharedNormal.MultiFindPointer(StrataRegisteredSharedNormal.NormalCodeChunkHash, NormalInfos);
+	TArray<FStrataSharedLocalBasesInfo*> NormalInfos;
+	CodeChunkToStrataSharedLocalBasis.MultiFindPointer(StrataRegisteredSharedLocalBasis.NormalCodeChunkHash, NormalInfos);
 	if (NormalInfos.Num() == 0)
 	{
-		CodeChunkToStrataSharedNormal.Add(StrataRegisteredSharedNormal.NormalCodeChunkHash, { StrataRegisteredSharedNormal, *GetParameterCode(StrataRegisteredSharedNormal.NormalCodeChunk), FString() });
-		return StrataRegisteredSharedNormal;
+		CodeChunkToStrataSharedLocalBasis.Add(StrataRegisteredSharedLocalBasis.NormalCodeChunkHash, { StrataRegisteredSharedLocalBasis, *GetParameterCode(StrataRegisteredSharedLocalBasis.NormalCodeChunk), FString() });
+		return StrataRegisteredSharedLocalBasis;
 	}
 	// Return the first existing code chunk which match the normal chunk code
 	return NormalInfos[0]->SharedData;
 }
 
-FStrataRegisteredSharedNormal FHLSLMaterialTranslator::StrataCompilationInfoRegisterSharedNormal(int32 NormalCodeChunk, int32 TangentCodeChunk)
+FStrataRegisteredSharedLocalBasis FHLSLMaterialTranslator::StrataCompilationInfoRegisterSharedLocalBasis(int32 NormalCodeChunk, int32 TangentCodeChunk)
 {
 	check(NormalCodeChunk != INDEX_NONE);
-	check(NextFreeStrataShaderNormalIndex < 255);	// Out of shared normal slots
+	check(NextFreeStrataShaderNormalIndex < 255);	// Out of shared local basis slots
 	
 
-	FStrataRegisteredSharedNormal StrataRegisteredSharedNormal;
-	StrataRegisteredSharedNormal.NormalCodeChunk		= NormalCodeChunk;
-	StrataRegisteredSharedNormal.NormalCodeChunkHash	= GetParameterHash(NormalCodeChunk); 
-	StrataRegisteredSharedNormal.TangentCodeChunk		= TangentCodeChunk;
-	StrataRegisteredSharedNormal.TangentCodeChunkHash	= GetParameterHash(TangentCodeChunk);
-	StrataRegisteredSharedNormal.GraphSharedNormalIndex = NextFreeStrataShaderNormalIndex++;
+	FStrataRegisteredSharedLocalBasis StrataRegisteredSharedLocalBasis;
+	StrataRegisteredSharedLocalBasis.NormalCodeChunk		= NormalCodeChunk;
+	StrataRegisteredSharedLocalBasis.NormalCodeChunkHash	= GetParameterHash(NormalCodeChunk);
+	StrataRegisteredSharedLocalBasis.TangentCodeChunk		= TangentCodeChunk;
+	StrataRegisteredSharedLocalBasis.TangentCodeChunkHash	= GetParameterHash(TangentCodeChunk);
+	StrataRegisteredSharedLocalBasis.GraphSharedLocalBasisIndex = NextFreeStrataShaderNormalIndex++;
 
 	// Find a basis which matches both the Normal & the Tangent code chunks
-	TArray<FStrataSharedNormalInfo*> NormalInfos;
-	CodeChunkToStrataSharedNormal.MultiFindPointer(StrataRegisteredSharedNormal.NormalCodeChunkHash, NormalInfos);
-	for (FStrataSharedNormalInfo* NormalInfo : NormalInfos)
+	TArray<FStrataSharedLocalBasesInfo*> NormalInfos;
+	CodeChunkToStrataSharedLocalBasis.MultiFindPointer(StrataRegisteredSharedLocalBasis.NormalCodeChunkHash, NormalInfos);
+	for (FStrataSharedLocalBasesInfo* NormalInfo : NormalInfos)
 	{
 		// * Either we find a perfect match (normal & tangent matches)
 		// * Or we find a normal which doesn't have a tangent associated with, and we set the tangent for code
-		if (StrataRegisteredSharedNormal.TangentCodeChunkHash == NormalInfo->SharedData.TangentCodeChunk)
+		if (StrataRegisteredSharedLocalBasis.TangentCodeChunkHash == NormalInfo->SharedData.TangentCodeChunk)
 		{
 			return NormalInfo->SharedData;
 			break;
@@ -8945,32 +8945,31 @@ FStrataRegisteredSharedNormal FHLSLMaterialTranslator::StrataCompilationInfoRegi
 		else if (NormalInfo->SharedData.TangentCodeChunk == INDEX_NONE)
 		{
 			NormalInfo->SharedData.TangentCodeChunk		= TangentCodeChunk;
-			NormalInfo->SharedData.TangentCodeChunkHash = StrataRegisteredSharedNormal.TangentCodeChunkHash;
-			NormalInfo->SharedTangentCode				= *GetParameterCode(TangentCodeChunk);
+			NormalInfo->SharedData.TangentCodeChunkHash = StrataRegisteredSharedLocalBasis.TangentCodeChunkHash;
+			NormalInfo->TangentCode				= *GetParameterCode(TangentCodeChunk);
 			return NormalInfo->SharedData;
 			break;
 		}
 	}
 
-	// Allocate a new slot for a new shared normal & tangent
-	const FString NewSharedNormalIndex = GetStrataSharedNormalIndexMacro(StrataRegisteredSharedNormal);
-	CodeChunkToStrataSharedNormal.Add(StrataRegisteredSharedNormal.NormalCodeChunkHash, { StrataRegisteredSharedNormal, *GetParameterCode(StrataRegisteredSharedNormal.NormalCodeChunk), *GetParameterCode(StrataRegisteredSharedNormal.TangentCodeChunk) });
-	return StrataRegisteredSharedNormal;
+	// Allocate a new slot for a new shared local basis
+	CodeChunkToStrataSharedLocalBasis.Add(StrataRegisteredSharedLocalBasis.NormalCodeChunkHash, { StrataRegisteredSharedLocalBasis, *GetParameterCode(StrataRegisteredSharedLocalBasis.NormalCodeChunk), *GetParameterCode(StrataRegisteredSharedLocalBasis.TangentCodeChunk) });
+	return StrataRegisteredSharedLocalBasis;
 }
 
-FHLSLMaterialTranslator::FStrataSharedNormalInfo FHLSLMaterialTranslator::StrataCompilationInfoGetMatchingSharedNormalInfo(const FStrataRegisteredSharedNormal& SearchedSharedNormal)
+FHLSLMaterialTranslator::FStrataSharedLocalBasesInfo FHLSLMaterialTranslator::StrataCompilationInfoGetMatchingSharedLocalBasisInfo(const FStrataRegisteredSharedLocalBasis& SearchedSharedLocalBasis)
 {
-	check(NextFreeStrataShaderNormalIndex < 255);	// Out of shared normal slots
+	check(NextFreeStrataShaderNormalIndex < 255);	// Out of shared local basis slots
 
 	// Find a basis which matches both the Normal & the Tangent code chunks
-	TArray<FStrataSharedNormalInfo*> NormalInfos;
-	CodeChunkToStrataSharedNormal.MultiFindPointer(SearchedSharedNormal.NormalCodeChunkHash, NormalInfos);
+	TArray<FStrataSharedLocalBasesInfo*> NormalInfos;
+	CodeChunkToStrataSharedLocalBasis.MultiFindPointer(SearchedSharedLocalBasis.NormalCodeChunkHash, NormalInfos);
 
 	// We first try to find a perfect match for normal and tangent from all the registered element.
-	for (FStrataSharedNormalInfo* NormalInfo : NormalInfos)
+	for (FStrataSharedLocalBasesInfo* NormalInfo : NormalInfos)
 	{
-		if (SearchedSharedNormal.TangentCodeChunk == INDEX_NONE ||											// We selected the first available normal if there is no tangent specified on the material.
-			SearchedSharedNormal.TangentCodeChunkHash == NormalInfo->SharedData.TangentCodeChunkHash)		// Otherwise we select the normal+tangent that exactly matches the request.
+		if (SearchedSharedLocalBasis.TangentCodeChunk == INDEX_NONE ||											// We selected the first available normal if there is no tangent specified on the material.
+			SearchedSharedLocalBasis.TangentCodeChunkHash == NormalInfo->SharedData.TangentCodeChunkHash)		// Otherwise we select the normal+tangent that exactly matches the request.
 		{
 			return *NormalInfo;
 			break;
@@ -8978,12 +8977,12 @@ FHLSLMaterialTranslator::FStrataSharedNormalInfo FHLSLMaterialTranslator::Strata
 	}
 
 	check(0);	// When the compiler is querying, this is to get a result to generate code from a fully processed graph. No result means a bug happened during graph processing.
-	return FStrataSharedNormalInfo();
+	return FStrataSharedLocalBasesInfo();
 }
 
-uint8 FHLSLMaterialTranslator::StrataCompilationInfoGetSharedNormalCount()
+uint8 FHLSLMaterialTranslator::StrataCompilationInfoGetSharedLocalBasesCount()
 {
-	return FinalUsedSharedNormalCount;
+	return FinalUsedSharedLocalBasesCount;
 }
 
 int32 FHLSLMaterialTranslator::StrataCreateAndRegisterNullMaterial()
@@ -9004,12 +9003,12 @@ int32 FHLSLMaterialTranslator::StrataSlabBSDF(
 	int32 ThinFilmThickness, 
 	int32 FuzzAmount, int32 FuzzColor,
 	int32 Thickness, 
-	int32 Normal, int32 Tangent, const FString& SharedNormalIndexMacro)
+	int32 Normal, int32 Tangent, const FString& SharedLocalBasisIndexMacro)
 {
 	const FString NormalCode = GetParameterCode(Normal);
 
 	return AddCodeChunk(
-		MCT_Strata, TEXT("GetStrataSlabBSDF(%s, %s, %s, %s, %s, %s, %s, %s, %s, %s, %s, %s, %s, %s, %s, %s, %s, %s, %s, %s, Parameters.SharedNormals.NormalTypes) /* Normal = %s ; Tangent = %s */"),
+		MCT_Strata, TEXT("GetStrataSlabBSDF(%s, %s, %s, %s, %s, %s, %s, %s, %s, %s, %s, %s, %s, %s, %s, %s, %s, %s, %s, %s, Parameters.SharedLocalBases.Types) /* Normal = %s ; Tangent = %s */"),
 		* GetParameterCode(UseMetalness),
 		*GetParameterCode(BaseColor),		*GetParameterCode(EdgeColor),	*GetParameterCode(Specular), *GetParameterCode(Metallic),
 		*GetParameterCode(DiffuseAlbedo),	*GetParameterCode(F0),			*GetParameterCode(F90),
@@ -9022,19 +9021,19 @@ int32 FHLSLMaterialTranslator::StrataSlabBSDF(
 		*GetParameterCode(ThinFilmThickness),
 		*GetParameterCode(FuzzAmount),		*GetParameterCode(FuzzColor),
 		*GetParameterCode(Thickness),
-		*SharedNormalIndexMacro,
+		*SharedLocalBasisIndexMacro,
 		*NormalCode,
 		Tangent!=INDEX_NONE ? *GetParameterCode(Tangent) : TEXT("NONE")
 	);
 }
 
-int32 FHLSLMaterialTranslator::StrataSheenBSDF(int32 BaseColor, int32 Roughness, int32 Normal, const FString& SharedNormalIndexMacro)
+int32 FHLSLMaterialTranslator::StrataSheenBSDF(int32 BaseColor, int32 Roughness, int32 Normal, const FString& SharedLocalBasisIndexMacro)
 {
 	return AddCodeChunk(
 		MCT_Strata, TEXT("GetStrataSheenBSDF(%s, %s, %s) /* %s */"),
 		*GetParameterCode(BaseColor),
 		*GetParameterCode(Roughness),
-		*SharedNormalIndexMacro,
+		*SharedLocalBasisIndexMacro,
 		*GetParameterCode(Normal)
 	);
 }
@@ -9059,7 +9058,7 @@ int32 FHLSLMaterialTranslator::StrataUnlitBSDF(int32 EmissiveColor, int32 Transm
 	);
 }
 
-int32 FHLSLMaterialTranslator::StrataHairBSDF(int32 BaseColor, int32 Scatter, int32 Specular, int32 Roughness, int32 Backlit, int32 EmissiveColor, int32 Tangent, const FString& SharedNormalIndexMacro)
+int32 FHLSLMaterialTranslator::StrataHairBSDF(int32 BaseColor, int32 Scatter, int32 Specular, int32 Roughness, int32 Backlit, int32 EmissiveColor, int32 Tangent, const FString& SharedLocalBasisIndexMacro)
 {
 	return AddCodeChunk(
 		MCT_Strata, TEXT("GetStrataHairBSDF(%s, %s, %s, %s, %s, %s, %s) /* %s */"),
@@ -9069,14 +9068,14 @@ int32 FHLSLMaterialTranslator::StrataHairBSDF(int32 BaseColor, int32 Scatter, in
 		*GetParameterCode(Roughness),
 		*GetParameterCode(Backlit),
 		*GetParameterCode(EmissiveColor),
-		*SharedNormalIndexMacro,
+		*SharedLocalBasisIndexMacro,
 		*GetParameterCode(Tangent)
 	);
 }
 
 int32 FHLSLMaterialTranslator::StrataSingleLayerWaterBSDF(
 	int32 BaseColor, int32 Metallic, int32 Specular, int32 Roughness, int32 EmissiveColor, int32 TopMaterialOpacity,
-	int32 WaterAlbedo, int32 WaterExtinction, int32 WaterPhaseG, int32 ColorScaleBehindWater, int32 Normal, const FString& SharedNormalIndexMacro)
+	int32 WaterAlbedo, int32 WaterExtinction, int32 WaterPhaseG, int32 ColorScaleBehindWater, int32 Normal, const FString& SharedLocalBasisIndexMacro)
 {
 	return AddCodeChunk(
 		MCT_Strata, TEXT("GetStrataSingleLayerWaterBSDF(%s, %s, %s, %s, %s, %s, %s, %s, %s, %s, %s) /* %s */"),
@@ -9090,7 +9089,7 @@ int32 FHLSLMaterialTranslator::StrataSingleLayerWaterBSDF(
 		*GetParameterCode(WaterExtinction),
 		*GetParameterCode(WaterPhaseG),
 		*GetParameterCode(ColorScaleBehindWater),
-		*SharedNormalIndexMacro,
+		*SharedLocalBasisIndexMacro,
 		*GetParameterCode(Normal)
 	);
 }
@@ -9109,7 +9108,7 @@ int32 FHLSLMaterialTranslator::StrataHorizontalMixing(int32 Foreground, int32 Ba
 	);
 }
 
-int32 FHLSLMaterialTranslator::StrataHorizontalMixingParameterBlending(int32 Foreground, int32 Background, int32 Mix, const FString& SharedNormalIndexMacro)
+int32 FHLSLMaterialTranslator::StrataHorizontalMixingParameterBlending(int32 Foreground, int32 Background, int32 Mix, const FString& SharedLocalBasisIndexMacro)
 {
 	if (Foreground == INDEX_NONE || Background == INDEX_NONE || Mix == INDEX_NONE)
 	{
@@ -9120,7 +9119,7 @@ int32 FHLSLMaterialTranslator::StrataHorizontalMixingParameterBlending(int32 For
 		*GetParameterCode(Foreground),
 		*GetParameterCode(Background),
 		*GetParameterCode(Mix),
-		*SharedNormalIndexMacro
+		*SharedLocalBasisIndexMacro
 	);
 }
 
@@ -9137,7 +9136,7 @@ int32 FHLSLMaterialTranslator::StrataVerticalLayering(int32 Top, int32 Base)
 	);
 }
 
-int32 FHLSLMaterialTranslator::StrataVerticalLayeringParameterBlending(int32 Top, int32 Base, const FString& SharedNormalIndexMacro, int32 TopBSDFNormalCodeChunk)
+int32 FHLSLMaterialTranslator::StrataVerticalLayeringParameterBlending(int32 Top, int32 Base, const FString& SharedLocalBasisIndexMacro, int32 TopBSDFNormalCodeChunk)
 {
 	if (Top == INDEX_NONE || Base == INDEX_NONE)
 	{
@@ -9147,7 +9146,7 @@ int32 FHLSLMaterialTranslator::StrataVerticalLayeringParameterBlending(int32 Top
 		MCT_Strata, TEXT("StrataVerticalLayeringParameterBlending(%s, %s, %s, dot(%s, %s))"),
 		*GetParameterCode(Top),
 		*GetParameterCode(Base),
-		*SharedNormalIndexMacro,
+		*SharedLocalBasisIndexMacro,
 		*GetParameterCode(TopBSDFNormalCodeChunk),
 		*GetParameterCode(CameraVector())
 	);
@@ -9166,7 +9165,7 @@ int32 FHLSLMaterialTranslator::StrataAdd(int32 A, int32 B)
 	);
 }
 
-int32 FHLSLMaterialTranslator::StrataAddParameterBlending(int32 A, int32 B, const FString& SharedNormalIndexMacro)
+int32 FHLSLMaterialTranslator::StrataAddParameterBlending(int32 A, int32 B, const FString& SharedLocalBasisIndexMacro)
 {
 	if (A == INDEX_NONE || B == INDEX_NONE)
 	{
@@ -9176,7 +9175,7 @@ int32 FHLSLMaterialTranslator::StrataAddParameterBlending(int32 A, int32 B, cons
 		MCT_Strata, TEXT("StrataAddParameterBlending(%s, %s, %s)"),
 		*GetParameterCode(A),
 		*GetParameterCode(B),
-		*SharedNormalIndexMacro
+		*SharedLocalBasisIndexMacro
 	);
 }
 
diff --git a/Engine/Source/Runtime/Engine/Private/Materials/HLSLMaterialTranslator.h b/Engine/Source/Runtime/Engine/Private/Materials/HLSLMaterialTranslator.h
--- a/Engine/Source/Runtime/Engine/Private/Materials/HLSLMaterialTranslator.h
+++ b/Engine/Source/Runtime/Engine/Private/Materials/HLSLMaterialTranslator.h
@@ -396,22 +396,22 @@ protected:
 	TMap<int32, FStrataMaterialCompilationInfo> CodeChunkToStrataCompilationInfoMap;
 	/** Only valid after Translate() is called. This is the code chunk representing the front material. It can be used to recover the material topology. */
 	int32 StrataValidFrontMaterialCodeChunkPostTranslate;
-	/** The code initializing the array of shared normal/tangent. */
+	/** The code initializing the array of shared local bases. */
 	FString StrataPixelNormalInitializerValues;
-	/** The next free index that can be used to represent a unique macros pointing to he shared normal position in the array of shared normal written to memory once the shader is executed. */
+	/** The next free index that can be used to represent a unique macros pointing to the position in the array of shared local bases written to memory once the shader is executed. */
 	uint8 NextFreeStrataShaderNormalIndex;
-	/** The effective final shared normal count used by the final shader. */
-	uint8 FinalUsedSharedNormalCount;
-	/** Represent a shared normal description with its associated code. */
-	struct FStrataSharedNormalInfo
+	/** The effective final shared local bases count used by the final shader. */
+	uint8 FinalUsedSharedLocalBasesCount;
+	/** Represent a shared local basis description with its associated code. */
+	struct FStrataSharedLocalBasesInfo
 	{
-		FStrataRegisteredSharedNormal SharedData;
-		FString SharedNormalCode;
-		FString SharedTangentCode;
+		FStrataRegisteredSharedLocalBasis SharedData;
+		FString NormalCode;
+		FString TangentCode;
 	};
-	/** Tracks shared normal & tangent used by strata materials, mapping a normal code chunk hash to a SharedMaterialInfo.
+	/** Tracks shared local bases used by strata materials, mapping a normal code chunk hash to a SharedMaterialInfo.
 	 * A normal code chunk hash can point to multiple shared info in case it is paired with different tangents. */
-	TMultiMap<uint64, FStrataSharedNormalInfo> CodeChunkToStrataSharedNormal;
+	TMultiMap<uint64, FStrataSharedLocalBasesInfo> CodeChunkToStrataSharedLocalBasis;
 
 	/** Tracks the total number of vt samples in the shader. */
 	uint32 NumVtSamples;
@@ -987,32 +987,32 @@ protected:
 		int32 ThinFilmThickness, 
 		int32 FuzzAmount, int32 FuzzColor, 
 		int32 Thickness,
-		int32 Normal, int32 Tangent, const FString& SharedNormalIndexMacro) override;
-	virtual int32 StrataSheenBSDF(int32 BaseColor, int32 Roughness, int32 Normal, const FString& SharedNormalIndexMacro) override;
+		int32 Normal, int32 Tangent, const FString& SharedLocalBasisIndexMacro) override;
+	virtual int32 StrataSheenBSDF(int32 BaseColor, int32 Roughness, int32 Normal, const FString& SharedLocalBasisIndexMacro) override;
 	virtual int32 StrataVolumetricFogCloudBSDF(int32 Albedo, int32 Extinction, int32 EmissiveColor, int32 AmbientOcclusion) override;
 	virtual int32 StrataUnlitBSDF(int32 EmissiveColor, int32 TransmittanceColor) override;
-	virtual int32 StrataHairBSDF(int32 BaseColor, int32 Scatter, int32 Specular, int32 Roughness, int32 Backlit, int32 EmissiveColor, int32 Tangent, const FString& SharedNormalIndexMacro) override;
+	virtual int32 StrataHairBSDF(int32 BaseColor, int32 Scatter, int32 Specular, int32 Roughness, int32 Backlit, int32 EmissiveColor, int32 Tangent, const FString& SharedLocalBasisIndexMacro) override;
 	virtual int32 StrataSingleLayerWaterBSDF(
 		int32 BaseColor, int32 Metallic, int32 Specular, int32 Roughness, 
 		int32 EmissiveColor, int32 TopMaterialOpacity, int32 WaterAlbedo, int32 WaterExtinction, int32 WaterPhaseG, 
-		int32 ColorScaleBehindWater, int32 Normal, const FString& SharedNormalIndexMacro) override;
+		int32 ColorScaleBehindWater, int32 Normal, const FString& SharedLocalBasisIndexMacro) override;
 	virtual int32 StrataHorizontalMixing(int32 Foreground, int32 Background, int32 Mix) override;
-	virtual int32 StrataHorizontalMixingParameterBlending(int32 Foreground, int32 Background, int32 Mix, const FString& SharedNormalIndexMacro) override;
+	virtual int32 StrataHorizontalMixingParameterBlending(int32 Foreground, int32 Background, int32 Mix, const FString& SharedLocalBasisIndexMacro) override;
 	virtual int32 StrataVerticalLayering(int32 Top, int32 Base) override;
-	virtual int32 StrataVerticalLayeringParameterBlending(int32 Top, int32 Base, const FString& SharedNormalIndexMacro, int32 TopBSDFNormalCodeChunk) override;
+	virtual int32 StrataVerticalLayeringParameterBlending(int32 Top, int32 Base, const FString& SharedLocalBasisIndexMacro, int32 TopBSDFNormalCodeChunk) override;
 	virtual int32 StrataAdd(int32 A, int32 B) override;
-	virtual int32 StrataAddParameterBlending(int32 A, int32 B, const FString& SharedNormalIndexMacro) override;
+	virtual int32 StrataAddParameterBlending(int32 A, int32 B, const FString& SharedLocalBasisIndexMacro) override;
 	virtual int32 StrataWeight(int32 A, int32 Weight) override;
 	virtual int32 StrataTransmittanceToMFP(int32 TransmittanceColor, int32 DesiredThickness, int32 OutputIndex) override;
 
 	virtual void StrataCompilationInfoRegisterCodeChunk(int32 CodeChunk, FStrataMaterialCompilationInfo& StrataMaterialCompilationInfo) override;
 	virtual bool StrataCompilationInfoContainsCodeChunk(int32 CodeChunk) override;
 	virtual const FStrataMaterialCompilationInfo& GetStrataCompilationInfo(int32 CodeChunk) override;
-	virtual FStrataRegisteredSharedNormal StrataCompilationInfoRegisterSharedNormal(int32 NormalCodeChunk) override;
-	virtual FStrataRegisteredSharedNormal StrataCompilationInfoRegisterSharedNormal(int32 NormalCodeChunk, int32 TangentCodeChunk) override;
-	virtual uint8 StrataCompilationInfoGetSharedNormalCount() override;
+	virtual FStrataRegisteredSharedLocalBasis StrataCompilationInfoRegisterSharedLocalBasis(int32 NormalCodeChunk) override;
+	virtual FStrataRegisteredSharedLocalBasis StrataCompilationInfoRegisterSharedLocalBasis(int32 NormalCodeChunk, int32 TangentCodeChunk) override;
+	virtual uint8 StrataCompilationInfoGetSharedLocalBasesCount() override;
 
-	FStrataSharedNormalInfo StrataCompilationInfoGetMatchingSharedNormalInfo(const FStrataRegisteredSharedNormal& SearchedSharedNormal);
+	FStrataSharedLocalBasesInfo StrataCompilationInfoGetMatchingSharedLocalBasisInfo(const FStrataRegisteredSharedLocalBasis& SearchedSharedLocalBasis);
 
 #if HANDLE_CUSTOM_OUTPUTS_AS_MATERIAL_ATTRIBUTES
 	/** Used to translate code for custom output attributes such as ClearCoatBottomNormal */
diff --git a/Engine/Source/Runtime/Engine/Private/Materials/MaterialExpressions.cpp b/Engine/Source/Runtime/Engine/Private/Materials/MaterialExpressions.cpp
--- a/Engine/Source/Runtime/Engine/Private/Materials/MaterialExpressions.cpp
+++ b/Engine/Source/Runtime/Engine/Private/Materials/MaterialExpressions.cpp
@@ -19982,7 +19982,7 @@ int32 UMaterialExpressionStrataSlabBSDF::Compile(class FMaterialCompiler* Compil
 	int32 NormalCodeChunk = CompileWithDefaultNormalWS(Compiler, Normal);
 
 	int32 TangentCodeChunk = bHasAnisotropy ? CompileWithDefaultTangentWS(Compiler, Tangent) : INDEX_NONE;
-	const FStrataRegisteredSharedNormal NewRegisteredSharedNormal = StrataCompilationInfoCreateSharedNormal(Compiler, NormalCodeChunk, TangentCodeChunk);
+	const FStrataRegisteredSharedLocalBasis NewRegisteredSharedLocalBasis = StrataCompilationInfoCreateSharedLocalBasis(Compiler, NormalCodeChunk, TangentCodeChunk);
 
 	const bool bHasEdgeColor = HasEdgeColor(); // This accounts for EdgeColor and also F90 when the non metalness worfklow is selected.
 	const bool bHasThinFilm = HasThinFilm();
@@ -20032,8 +20032,8 @@ int32 UMaterialExpressionStrataSlabBSDF::Compile(class FMaterialCompiler* Compil
 		CompileWithDefaultFloat1(Compiler, Thickness, STRATA_LAYER_DEFAULT_THICKNESS_CM),
 		NormalCodeChunk,
 		TangentCodeChunk,
-		Compiler->GetStrataSharedNormalIndexMacro(NewRegisteredSharedNormal));
-	StrataCompilationInfoCreateSingleBSDFMaterial(Compiler, OutputCodeChunk, NewRegisteredSharedNormal, STRATA_BSDF_TYPE_SLAB, bHasSSS, bHasDMFPPluggedIn, bHasEdgeColor, bHasThinFilm, bHasFuzz, bHasHaziness);
+		Compiler->GetStrataSharedLocalBasisIndexMacro(NewRegisteredSharedLocalBasis));
+	StrataCompilationInfoCreateSingleBSDFMaterial(Compiler, OutputCodeChunk, NewRegisteredSharedLocalBasis, STRATA_BSDF_TYPE_SLAB, bHasSSS, bHasDMFPPluggedIn, bHasEdgeColor, bHasThinFilm, bHasFuzz, bHasHaziness);
 
 	return OutputCodeChunk;
 }
@@ -20482,7 +20482,7 @@ int32 UMaterialExpressionStrataVolumetricFogCloudBSDF::Compile(class FMaterialCo
 		CompileWithDefaultFloat3(Compiler, EmissiveColor, 0.0f, 0.0f, 0.0f),
 		CompileWithDefaultFloat1(Compiler, AmbientOcclusion, 1.0f));
 
-	StrataCompilationInfoCreateSingleBSDFMaterial(Compiler, OutputCodeChunk, StrataCompilationInfoCreateNullSharedNormal(), STRATA_BSDF_TYPE_VOLUMETRICFOGCLOUD);
+	StrataCompilationInfoCreateSingleBSDFMaterial(Compiler, OutputCodeChunk, StrataCompilationInfoCreateNullSharedLocalBasis(), STRATA_BSDF_TYPE_VOLUMETRICFOGCLOUD);
 
 	return OutputCodeChunk;
 }
@@ -20553,7 +20553,7 @@ int32 UMaterialExpressionStrataUnlitBSDF::Compile(class FMaterialCompiler* Compi
 		CompileWithDefaultFloat3(Compiler, EmissiveColor, 0.0f, 0.0f, 0.0f),
 		CompileWithDefaultFloat3(Compiler, TransmittanceColor, 1.0f, 1.0f, 1.0f));
 
-	StrataCompilationInfoCreateSingleBSDFMaterial(Compiler, OutputCodeChunk, StrataCompilationInfoCreateNullSharedNormal(), STRATA_BSDF_TYPE_UNLIT);
+	StrataCompilationInfoCreateSingleBSDFMaterial(Compiler, OutputCodeChunk, StrataCompilationInfoCreateNullSharedLocalBasis(), STRATA_BSDF_TYPE_UNLIT);
 
 	return OutputCodeChunk;
 }
@@ -20614,9 +20614,9 @@ UMaterialExpressionStrataHairBSDF::UMaterialExpressionStrataHairBSDF(const FObje
 #if WITH_EDITOR
 int32 UMaterialExpressionStrataHairBSDF::Compile(class FMaterialCompiler* Compiler, int32 OutputIndex)
 {
-	// For hair, the shared normal index in fact represent the tangent
+	// For hair, the shared local basis normal in fact represent the tangent
 	int32 TangentCodeChunk = CompileWithDefaultTangentWS(Compiler, Tangent);
-	const FStrataRegisteredSharedNormal NewRegisteredSharedNormal = StrataCompilationInfoCreateSharedNormal(Compiler, TangentCodeChunk);
+	const FStrataRegisteredSharedLocalBasis NewRegisteredSharedLocalBasis = StrataCompilationInfoCreateSharedLocalBasis(Compiler, TangentCodeChunk);
 
 	int32 OutputCodeChunk = Compiler->StrataHairBSDF(
 		CompileWithDefaultFloat3(Compiler, BaseColor,	0.0f, 0.0f, 0.0f),
@@ -20626,9 +20626,9 @@ int32 UMaterialExpressionStrataHairBSDF::Compile(class FMaterialCompiler* Compil
 		CompileWithDefaultFloat1(Compiler, Backlit,		0.0f),
 		CompileWithDefaultFloat3(Compiler, EmissiveColor,1.0f, 0.0f, 0.0f),
 		TangentCodeChunk,
-		Compiler->GetStrataSharedNormalIndexMacro(NewRegisteredSharedNormal));
+		Compiler->GetStrataSharedLocalBasisIndexMacro(NewRegisteredSharedLocalBasis));
 
-	StrataCompilationInfoCreateSingleBSDFMaterial(Compiler, OutputCodeChunk, NewRegisteredSharedNormal, STRATA_BSDF_TYPE_HAIR);
+	StrataCompilationInfoCreateSingleBSDFMaterial(Compiler, OutputCodeChunk, NewRegisteredSharedLocalBasis, STRATA_BSDF_TYPE_HAIR);
 
 	return OutputCodeChunk;
 }
@@ -20705,7 +20705,7 @@ UMaterialExpressionStrataSingleLayerWaterBSDF::UMaterialExpressionStrataSingleLa
 int32 UMaterialExpressionStrataSingleLayerWaterBSDF::Compile(class FMaterialCompiler* Compiler, int32 OutputIndex)
 {
 	int32 NormalCodeChunk = CompileWithDefaultNormalWS(Compiler, Normal);
-	const FStrataRegisteredSharedNormal NewRegisteredSharedNormal = StrataCompilationInfoCreateSharedNormal(Compiler, NormalCodeChunk);
+	const FStrataRegisteredSharedLocalBasis NewRegisteredSharedLocalBasis = StrataCompilationInfoCreateSharedLocalBasis(Compiler, NormalCodeChunk);
 
 	int32 OutputCodeChunk = Compiler->StrataSingleLayerWaterBSDF(
 		CompileWithDefaultFloat3(Compiler, BaseColor, 0.0f, 0.0f, 0.0f),
@@ -20719,9 +20719,9 @@ int32 UMaterialExpressionStrataSingleLayerWaterBSDF::Compile(class FMaterialComp
 		CompileWithDefaultFloat1(Compiler, WaterPhaseG, 0.0f),
 		CompileWithDefaultFloat3(Compiler, ColorScaleBehindWater, 1.0f, 1.0f, 1.0f),
 		NormalCodeChunk,
-		Compiler->GetStrataSharedNormalIndexMacro(NewRegisteredSharedNormal));
+		Compiler->GetStrataSharedLocalBasisIndexMacro(NewRegisteredSharedLocalBasis));
 
-	StrataCompilationInfoCreateSingleBSDFMaterial(Compiler, OutputCodeChunk, NewRegisteredSharedNormal, STRATA_BSDF_TYPE_SINGLELAYERWATER);
+	StrataCompilationInfoCreateSingleBSDFMaterial(Compiler, OutputCodeChunk, NewRegisteredSharedLocalBasis, STRATA_BSDF_TYPE_SINGLELAYERWATER);
 
 	return OutputCodeChunk;
 }
@@ -20854,27 +20854,27 @@ int32 UMaterialExpressionStrataHorizontalMixing::Compile(class FMaterialCompiler
 		}
 
 		// Compute the new Normal and Tangent resulting from the blending using code chunk
-		const int32 NewNormalCodeChunk = StrataBlendNormal(Compiler, BackgroundBSDF.RegisteredSharedNormal.NormalCodeChunk, ForegroundBSDF.RegisteredSharedNormal.NormalCodeChunk, MixCodeChunk);
+		const int32 NewNormalCodeChunk = StrataBlendNormal(Compiler, BackgroundBSDF.RegisteredSharedLocalBasis.NormalCodeChunk, ForegroundBSDF.RegisteredSharedLocalBasis.NormalCodeChunk, MixCodeChunk);
 		// The tangent is optional so we treat it differently if INDEX_NONE is specified
 		int32 NewTangentCodeChunk = INDEX_NONE;
-		if (ForegroundBSDF.RegisteredSharedNormal.TangentCodeChunk != INDEX_NONE && BackgroundBSDF.RegisteredSharedNormal.TangentCodeChunk != INDEX_NONE)
+		if (ForegroundBSDF.RegisteredSharedLocalBasis.TangentCodeChunk != INDEX_NONE && BackgroundBSDF.RegisteredSharedLocalBasis.TangentCodeChunk != INDEX_NONE)
 		{
-			NewTangentCodeChunk = StrataBlendNormal(Compiler, ForegroundBSDF.RegisteredSharedNormal.TangentCodeChunk, BackgroundBSDF.RegisteredSharedNormal.TangentCodeChunk, MixCodeChunk);
+			NewTangentCodeChunk = StrataBlendNormal(Compiler, ForegroundBSDF.RegisteredSharedLocalBasis.TangentCodeChunk, BackgroundBSDF.RegisteredSharedLocalBasis.TangentCodeChunk, MixCodeChunk);
 		}
-		else if (ForegroundBSDF.RegisteredSharedNormal.TangentCodeChunk != INDEX_NONE)
+		else if (ForegroundBSDF.RegisteredSharedLocalBasis.TangentCodeChunk != INDEX_NONE)
 		{
-			NewTangentCodeChunk = ForegroundBSDF.RegisteredSharedNormal.TangentCodeChunk;
+			NewTangentCodeChunk = ForegroundBSDF.RegisteredSharedLocalBasis.TangentCodeChunk;
 		}
-		else if (BackgroundBSDF.RegisteredSharedNormal.TangentCodeChunk != INDEX_NONE)
+		else if (BackgroundBSDF.RegisteredSharedLocalBasis.TangentCodeChunk != INDEX_NONE)
 		{
-			NewTangentCodeChunk = BackgroundBSDF.RegisteredSharedNormal.TangentCodeChunk;
+			NewTangentCodeChunk = BackgroundBSDF.RegisteredSharedLocalBasis.TangentCodeChunk;
 		}
-		const FStrataRegisteredSharedNormal NewRegisteredSharedNormal = StrataCompilationInfoCreateSharedNormal(Compiler, NewNormalCodeChunk, NewTangentCodeChunk);
+		const FStrataRegisteredSharedLocalBasis NewRegisteredSharedLocalBasis = StrataCompilationInfoCreateSharedLocalBasis(Compiler, NewNormalCodeChunk, NewTangentCodeChunk);
 
-		OutputCodeChunk = Compiler->StrataHorizontalMixingParameterBlending(ForegroundCodeChunk, BackgroundCodeChunk, MixCodeChunk, Compiler->GetStrataSharedNormalIndexMacro(NewRegisteredSharedNormal));
+		OutputCodeChunk = Compiler->StrataHorizontalMixingParameterBlending(ForegroundCodeChunk, BackgroundCodeChunk, MixCodeChunk, Compiler->GetStrataSharedLocalBasisIndexMacro(NewRegisteredSharedLocalBasis));
 
-		StrataCompilationInfoCreateSingleBSDFMaterial(Compiler, OutputCodeChunk, NewRegisteredSharedNormal, STRATA_BSDF_TYPE_SLAB);
-		FStrataMaterialCompilationInfo StrataInfo = StrataCompilationInfoHorizontalMixingParamBlend(Compiler, ForegroundStrataData, BackgroundStrataData, NewRegisteredSharedNormal);
+		StrataCompilationInfoCreateSingleBSDFMaterial(Compiler, OutputCodeChunk, NewRegisteredSharedLocalBasis, STRATA_BSDF_TYPE_SLAB);
+		FStrataMaterialCompilationInfo StrataInfo = StrataCompilationInfoHorizontalMixingParamBlend(Compiler, ForegroundStrataData, BackgroundStrataData, NewRegisteredSharedLocalBasis);
 
 		Compiler->StrataCompilationInfoRegisterCodeChunk(OutputCodeChunk, StrataInfo);
 	}
@@ -20997,27 +20997,27 @@ int32 UMaterialExpressionStrataVerticalLayering::Compile(class FMaterialCompiler
 		}
 
 		// Compute the new Normal and Tangent resulting from the blending using code chunk
-		const int32 NewNormalCodeChunk = StrataBlendNormal(Compiler, TopBSDF.RegisteredSharedNormal.NormalCodeChunk, BaseBSDF.RegisteredSharedNormal.NormalCodeChunk, MixCodeChunk);
+		const int32 NewNormalCodeChunk = StrataBlendNormal(Compiler, TopBSDF.RegisteredSharedLocalBasis.NormalCodeChunk, BaseBSDF.RegisteredSharedLocalBasis.NormalCodeChunk, MixCodeChunk);
 		// The tangent is optional so we treat it differently if INDEX_NONE is specified
 		int32 NewTangentCodeChunk = INDEX_NONE;
-		if (TopBSDF.RegisteredSharedNormal.TangentCodeChunk != INDEX_NONE && BaseBSDF.RegisteredSharedNormal.TangentCodeChunk != INDEX_NONE)
+		if (TopBSDF.RegisteredSharedLocalBasis.TangentCodeChunk != INDEX_NONE && BaseBSDF.RegisteredSharedLocalBasis.TangentCodeChunk != INDEX_NONE)
 		{
-			NewTangentCodeChunk = StrataBlendNormal(Compiler, TopBSDF.RegisteredSharedNormal.TangentCodeChunk, BaseBSDF.RegisteredSharedNormal.TangentCodeChunk, MixCodeChunk);
+			NewTangentCodeChunk = StrataBlendNormal(Compiler, TopBSDF.RegisteredSharedLocalBasis.TangentCodeChunk, BaseBSDF.RegisteredSharedLocalBasis.TangentCodeChunk, MixCodeChunk);
 		}
-		else if (TopBSDF.RegisteredSharedNormal.TangentCodeChunk != INDEX_NONE)
+		else if (TopBSDF.RegisteredSharedLocalBasis.TangentCodeChunk != INDEX_NONE)
 		{
-			NewTangentCodeChunk = TopBSDF.RegisteredSharedNormal.TangentCodeChunk;
+			NewTangentCodeChunk = TopBSDF.RegisteredSharedLocalBasis.TangentCodeChunk;
 		}
-		else if (BaseBSDF.RegisteredSharedNormal.TangentCodeChunk != INDEX_NONE)
+		else if (BaseBSDF.RegisteredSharedLocalBasis.TangentCodeChunk != INDEX_NONE)
 		{
-			NewTangentCodeChunk = BaseBSDF.RegisteredSharedNormal.TangentCodeChunk;
+			NewTangentCodeChunk = BaseBSDF.RegisteredSharedLocalBasis.TangentCodeChunk;
 		}
-		const FStrataRegisteredSharedNormal NewRegisteredSharedNormal = StrataCompilationInfoCreateSharedNormal(Compiler, NewNormalCodeChunk, NewTangentCodeChunk);
+		const FStrataRegisteredSharedLocalBasis NewRegisteredSharedLocalBasis = StrataCompilationInfoCreateSharedLocalBasis(Compiler, NewNormalCodeChunk, NewTangentCodeChunk);
 
-		OutputCodeChunk = Compiler->StrataVerticalLayeringParameterBlending(TopCodeChunk, BaseCodeChunk, Compiler->GetStrataSharedNormalIndexMacro(NewRegisteredSharedNormal), TopBSDF.RegisteredSharedNormal.NormalCodeChunk);
+		OutputCodeChunk = Compiler->StrataVerticalLayeringParameterBlending(TopCodeChunk, BaseCodeChunk, Compiler->GetStrataSharedLocalBasisIndexMacro(NewRegisteredSharedLocalBasis), TopBSDF.RegisteredSharedLocalBasis.NormalCodeChunk);
 
-		StrataCompilationInfoCreateSingleBSDFMaterial(Compiler, OutputCodeChunk, NewRegisteredSharedNormal, STRATA_BSDF_TYPE_SLAB);
-		FStrataMaterialCompilationInfo StrataInfo = StrataCompilationInfoVerticalLayeringParamBlend(Compiler, TopStrataData, BaseStrataData, NewRegisteredSharedNormal);
+		StrataCompilationInfoCreateSingleBSDFMaterial(Compiler, OutputCodeChunk, NewRegisteredSharedLocalBasis, STRATA_BSDF_TYPE_SLAB);
+		FStrataMaterialCompilationInfo StrataInfo = StrataCompilationInfoVerticalLayeringParamBlend(Compiler, TopStrataData, BaseStrataData, NewRegisteredSharedLocalBasis);
 
 		Compiler->StrataCompilationInfoRegisterCodeChunk(OutputCodeChunk, StrataInfo);
 	}
@@ -21138,27 +21138,27 @@ int32 UMaterialExpressionStrataAdd::Compile(class FMaterialCompiler* Compiler, i
 		}
 
 		// Compute the new Normal and Tangent resulting from the blending using code chunk
-		const int32 NewNormalCodeChunk = StrataBlendNormal(Compiler, ABSDF.RegisteredSharedNormal.NormalCodeChunk, BBSDF.RegisteredSharedNormal.NormalCodeChunk, MixCodeChunk);
+		const int32 NewNormalCodeChunk = StrataBlendNormal(Compiler, ABSDF.RegisteredSharedLocalBasis.NormalCodeChunk, BBSDF.RegisteredSharedLocalBasis.NormalCodeChunk, MixCodeChunk);
 		// The tangent is optional so we treat it differently if INDEX_NONE is specified
 		int32 NewTangentCodeChunk = INDEX_NONE;
-		if (ABSDF.RegisteredSharedNormal.TangentCodeChunk != INDEX_NONE && BBSDF.RegisteredSharedNormal.TangentCodeChunk != INDEX_NONE)
+		if (ABSDF.RegisteredSharedLocalBasis.TangentCodeChunk != INDEX_NONE && BBSDF.RegisteredSharedLocalBasis.TangentCodeChunk != INDEX_NONE)
 		{
-			NewTangentCodeChunk = StrataBlendNormal(Compiler, ABSDF.RegisteredSharedNormal.TangentCodeChunk, BBSDF.RegisteredSharedNormal.TangentCodeChunk, MixCodeChunk);
+			NewTangentCodeChunk = StrataBlendNormal(Compiler, ABSDF.RegisteredSharedLocalBasis.TangentCodeChunk, BBSDF.RegisteredSharedLocalBasis.TangentCodeChunk, MixCodeChunk);
 		}
-		else if (ABSDF.RegisteredSharedNormal.TangentCodeChunk != INDEX_NONE)
+		else if (ABSDF.RegisteredSharedLocalBasis.TangentCodeChunk != INDEX_NONE)
 		{
-			NewTangentCodeChunk = ABSDF.RegisteredSharedNormal.TangentCodeChunk;
+			NewTangentCodeChunk = ABSDF.RegisteredSharedLocalBasis.TangentCodeChunk;
 		}
-		else if (BBSDF.RegisteredSharedNormal.TangentCodeChunk != INDEX_NONE)
+		else if (BBSDF.RegisteredSharedLocalBasis.TangentCodeChunk != INDEX_NONE)
 		{
-			NewTangentCodeChunk = BBSDF.RegisteredSharedNormal.TangentCodeChunk;
+			NewTangentCodeChunk = BBSDF.RegisteredSharedLocalBasis.TangentCodeChunk;
 		}
-		const FStrataRegisteredSharedNormal NewRegisteredSharedNormal = StrataCompilationInfoCreateSharedNormal(Compiler, NewNormalCodeChunk, NewTangentCodeChunk);
+		const FStrataRegisteredSharedLocalBasis NewRegisteredSharedLocalBasis = StrataCompilationInfoCreateSharedLocalBasis(Compiler, NewNormalCodeChunk, NewTangentCodeChunk);
 
-		OutputCodeChunk = Compiler->StrataAddParameterBlending(ACodeChunk, BCodeChunk, Compiler->GetStrataSharedNormalIndexMacro(NewRegisteredSharedNormal));
+		OutputCodeChunk = Compiler->StrataAddParameterBlending(ACodeChunk, BCodeChunk, Compiler->GetStrataSharedLocalBasisIndexMacro(NewRegisteredSharedLocalBasis));
 
-		StrataCompilationInfoCreateSingleBSDFMaterial(Compiler, OutputCodeChunk, NewRegisteredSharedNormal, STRATA_BSDF_TYPE_SLAB);
-		FStrataMaterialCompilationInfo StrataInfo = StrataCompilationInfoAddParamBlend(Compiler, AStrataData, BStrataData, NewRegisteredSharedNormal);
+		StrataCompilationInfoCreateSingleBSDFMaterial(Compiler, OutputCodeChunk, NewRegisteredSharedLocalBasis, STRATA_BSDF_TYPE_SLAB);
+		FStrataMaterialCompilationInfo StrataInfo = StrataCompilationInfoAddParamBlend(Compiler, AStrataData, BStrataData, NewRegisteredSharedLocalBasis);
 
 		Compiler->StrataCompilationInfoRegisterCodeChunk(OutputCodeChunk, StrataInfo);
 	}
diff --git a/Engine/Source/Runtime/Engine/Private/Materials/StrataMaterial.cpp b/Engine/Source/Runtime/Engine/Private/Materials/StrataMaterial.cpp
--- a/Engine/Source/Runtime/Engine/Private/Materials/StrataMaterial.cpp
+++ b/Engine/Source/Runtime/Engine/Private/Materials/StrataMaterial.cpp
@@ -43,35 +43,35 @@ static void UpdateTotalBSDFCount(FMaterialCompiler* Compiler, FStrataMaterialCom
 	}
 }
 
-FStrataRegisteredSharedNormal StrataCompilationInfoCreateNullSharedNormal()
+FStrataRegisteredSharedLocalBasis StrataCompilationInfoCreateNullSharedLocalBasis()
 {
-	FStrataRegisteredSharedNormal FakeSharedNormal;
-	FakeSharedNormal.NormalCodeChunk = INDEX_NONE;
-	FakeSharedNormal.TangentCodeChunk = INDEX_NONE;
-	FakeSharedNormal.NormalCodeChunkHash = 0;
-	FakeSharedNormal.TangentCodeChunkHash = 0;
-	FakeSharedNormal.GraphSharedNormalIndex = 0;
-	return FakeSharedNormal;
+	FStrataRegisteredSharedLocalBasis FakeSharedLocalBasis;
+	FakeSharedLocalBasis.NormalCodeChunk = INDEX_NONE;
+	FakeSharedLocalBasis.TangentCodeChunk = INDEX_NONE;
+	FakeSharedLocalBasis.NormalCodeChunkHash = 0;
+	FakeSharedLocalBasis.TangentCodeChunkHash = 0;
+	FakeSharedLocalBasis.GraphSharedLocalBasisIndex = 0;
+	return FakeSharedLocalBasis;
 }
 
-FStrataRegisteredSharedNormal StrataCompilationInfoCreateSharedNormal(FMaterialCompiler* Compiler, int32 NormalCodeChunk, int32 TangentCodeChunk)
+FStrataRegisteredSharedLocalBasis StrataCompilationInfoCreateSharedLocalBasis(FMaterialCompiler* Compiler, int32 NormalCodeChunk, int32 TangentCodeChunk)
 {
 	if (TangentCodeChunk == INDEX_NONE)
 	{
-		return Compiler->StrataCompilationInfoRegisterSharedNormal(NormalCodeChunk);
+		return Compiler->StrataCompilationInfoRegisterSharedLocalBasis(NormalCodeChunk);
 	}
-	return Compiler->StrataCompilationInfoRegisterSharedNormal(NormalCodeChunk, TangentCodeChunk);
+	return Compiler->StrataCompilationInfoRegisterSharedLocalBasis(NormalCodeChunk, TangentCodeChunk);
 }
 
 void StrataCompilationInfoCreateSingleBSDFMaterial(FMaterialCompiler* Compiler, int32 CodeChunk,
-	const FStrataRegisteredSharedNormal& RegisteredSharedNormal,
+	const FStrataRegisteredSharedLocalBasis& RegisteredSharedLocalBasis,
 	uint8 BSDFType, bool bHasSSS, bool bHasDMFPPluggedIn, bool bHasEdgeColor, bool bHasThinFilm, bool bHasFuzz, bool bHasHaziness)
 {
 	FStrataMaterialCompilationInfo StrataInfo;
 	StrataInfo.LayerCount = 1;
 	StrataInfo.Layers[0].BSDFCount = 1;
 	StrataInfo.Layers[0].BSDFs[0].Type = BSDFType;
-	StrataInfo.Layers[0].BSDFs[0].RegisteredSharedNormal = RegisteredSharedNormal;
+	StrataInfo.Layers[0].BSDFs[0].RegisteredSharedLocalBasis = RegisteredSharedLocalBasis;
 	StrataInfo.Layers[0].BSDFs[0].bHasSSS = bHasSSS;
 	StrataInfo.Layers[0].BSDFs[0].bHasDMFPPluggedIn = bHasDMFPPluggedIn;
 	StrataInfo.Layers[0].BSDFs[0].bHasEdgeColor = bHasEdgeColor;
@@ -129,7 +129,7 @@ FStrataMaterialCompilationInfo StrataCompilationInfoAdd(FMaterialCompiler* Compi
 
 
 // ==> NOTE: Always pair with the shader behavior in StrataAddParameterBlending.
-FStrataMaterialCompilationInfo StrataCompilationInfoAddParamBlend(FMaterialCompiler* Compiler, const FStrataMaterialCompilationInfo& A, const FStrataMaterialCompilationInfo& B, const FStrataRegisteredSharedNormal& RegisteredSharedNormal)
+FStrataMaterialCompilationInfo StrataCompilationInfoAddParamBlend(FMaterialCompiler* Compiler, const FStrataMaterialCompilationInfo& A, const FStrataMaterialCompilationInfo& B, const FStrataRegisteredSharedLocalBasis& RegisteredSharedLocalBasis)
 {
 	check(A.TotalBSDFCount == 1);
 	check(B.TotalBSDFCount == 1);
@@ -138,7 +138,7 @@ FStrataMaterialCompilationInfo StrataCompilationInfoAddParamBlend(FMaterialCompi
 	FStrataMaterialCompilationInfo::FBSDF& NewBSDF = StrataInfo.Layers[0].BSDFs[0];
 	const FStrataMaterialCompilationInfo::FBSDF& OtherBSDF = B.Layers[0].BSDFs[0];
 
-	NewBSDF.RegisteredSharedNormal = RegisteredSharedNormal;
+	NewBSDF.RegisteredSharedLocalBasis = RegisteredSharedLocalBasis;
 
 	// When parameter blending is used, we take the union of all the features activated by input BSDFs.
 	NewBSDF.bHasSSS				|=	OtherBSDF.bHasSSS;
@@ -159,7 +159,7 @@ FStrataMaterialCompilationInfo StrataCompilationInfoHorizontalMixing(FMaterialCo
 
 
 // ==> NOTE: Always pair with the shader behavior in StrataHorizontalMixingParameterBlending
-FStrataMaterialCompilationInfo StrataCompilationInfoHorizontalMixingParamBlend(FMaterialCompiler* Compiler, const FStrataMaterialCompilationInfo& A, const FStrataMaterialCompilationInfo& B, const FStrataRegisteredSharedNormal& RegisteredSharedNormal)
+FStrataMaterialCompilationInfo StrataCompilationInfoHorizontalMixingParamBlend(FMaterialCompiler* Compiler, const FStrataMaterialCompilationInfo& A, const FStrataMaterialCompilationInfo& B, const FStrataRegisteredSharedLocalBasis& RegisteredSharedLocalBasis)
 {
 	check(A.TotalBSDFCount == 1);
 	check(B.TotalBSDFCount == 1);
@@ -168,7 +168,7 @@ FStrataMaterialCompilationInfo StrataCompilationInfoHorizontalMixingParamBlend(F
 	FStrataMaterialCompilationInfo::FBSDF& NewBSDF = StrataInfo.Layers[0].BSDFs[0];
 	const FStrataMaterialCompilationInfo::FBSDF& OtherBSDF = B.Layers[0].BSDFs[0];
 
-	NewBSDF.RegisteredSharedNormal = RegisteredSharedNormal;
+	NewBSDF.RegisteredSharedLocalBasis = RegisteredSharedLocalBasis;
 
 	// When parameter blending is used, we take the union of all the features activated by input BSDFs.
 	NewBSDF.bHasSSS				|=	OtherBSDF.bHasSSS;
@@ -206,7 +206,7 @@ FStrataMaterialCompilationInfo StrataCompilationInfoVerticalLayering(FMaterialCo
 
 
 // ==> NOTE: Always pair with the shader behavior in StrataVerticalLayeringParameterBlending
-FStrataMaterialCompilationInfo StrataCompilationInfoVerticalLayeringParamBlend(FMaterialCompiler* Compiler, const FStrataMaterialCompilationInfo& Top, const FStrataMaterialCompilationInfo& Base, const FStrataRegisteredSharedNormal& RegisteredSharedNormal)
+FStrataMaterialCompilationInfo StrataCompilationInfoVerticalLayeringParamBlend(FMaterialCompiler* Compiler, const FStrataMaterialCompilationInfo& Top, const FStrataMaterialCompilationInfo& Base, const FStrataRegisteredSharedLocalBasis& RegisteredSharedLocalBasis)
 {
 	check(Top.TotalBSDFCount == 1);
 	check(Base.TotalBSDFCount == 1);
@@ -215,7 +215,7 @@ FStrataMaterialCompilationInfo StrataCompilationInfoVerticalLayeringParamBlend(F
 	FStrataMaterialCompilationInfo::FBSDF& NewBSDF = StrataInfo.Layers[0].BSDFs[0];
 	const FStrataMaterialCompilationInfo::FBSDF& TopBSDF = Top.Layers[0].BSDFs[0];
 
-	NewBSDF.RegisteredSharedNormal = RegisteredSharedNormal;
+	NewBSDF.RegisteredSharedLocalBasis = RegisteredSharedLocalBasis;
 
 //	NewBSDF.bHasSSS				= TopBSDF.bHasSSS;				// We keep SSS only if the base layer has it. Otherwise it will be simple volume and the throughput wll be applied on the parameters.
 //	NewBSDF.bHasDMFPPluggedIn	= TopBSDF.bHasDMFPPluggedIn;	// Idem
@@ -338,10 +338,10 @@ FStrataMaterialAnalysisResult StrataCompilationInfoMaterialAnalysis(FMaterialCom
 
 	// 1. Header
 
-	// SharedNormals_BSDFCount
+	// SharedLocalBases_BSDFCount
 	Result.RequestedByteCount += UintByteSize;
-	// Shared normals between BSDFs
-	Result.RequestedByteCount += Compiler->StrataCompilationInfoGetSharedNormalCount() * STRATA_PACKED_NORMAL_STRIDE_BYTES;
+	// shared local bases between BSDFs
+	Result.RequestedByteCount += Compiler->StrataCompilationInfoGetSharedLocalBasesCount() * STRATA_PACKED_SHAREDLOCALBASIS_STRIDE_BYTES;
 	
 	// 2. The list of BSDFs
 
diff --git a/Engine/Source/Runtime/Engine/Private/Materials/StrataMaterial.h b/Engine/Source/Runtime/Engine/Private/Materials/StrataMaterial.h
--- a/Engine/Source/Runtime/Engine/Private/Materials/StrataMaterial.h
+++ b/Engine/Source/Runtime/Engine/Private/Materials/StrataMaterial.h
@@ -18,7 +18,7 @@ struct FStrataMaterialCompilationInfo
 	struct FBSDF
 	{
 		uint8 Type;
-		FStrataRegisteredSharedNormal RegisteredSharedNormal;
+		FStrataRegisteredSharedLocalBasis RegisteredSharedLocalBasis;
 
 		// Those properties are centered around the SLAB node. Maybe this can be abstracted at some point.
 		bool  bHasSSS;
@@ -50,21 +50,21 @@ struct FStrataMaterialCompilationInfo
 
 FString GetStrataBSDFName(uint8 BSDFType);
 
-FStrataRegisteredSharedNormal StrataCompilationInfoCreateNullSharedNormal();
-FStrataRegisteredSharedNormal StrataCompilationInfoCreateSharedNormal(FMaterialCompiler* Compiler, int32 NormalCodeChunk, int32 TangentCodeChunk = INDEX_NONE);
+FStrataRegisteredSharedLocalBasis StrataCompilationInfoCreateNullSharedLocalBasis();
+FStrataRegisteredSharedLocalBasis StrataCompilationInfoCreateSharedLocalBasis(FMaterialCompiler* Compiler, int32 NormalCodeChunk, int32 TangentCodeChunk = INDEX_NONE);
 
 void  StrataCompilationInfoCreateNullBSDF(FMaterialCompiler* Compiler, int32 CodeChunk);
 void  StrataCompilationInfoCreateSingleBSDFMaterial(FMaterialCompiler* Compiler, int32 CodeChunk, 
-	const FStrataRegisteredSharedNormal& RegisteredSharedNormal,
+	const FStrataRegisteredSharedLocalBasis& RegisteredSharedLocalBasis,
 	uint8 BSDFType, bool bHasSSS = false, bool bHasDMFPPluggedIn = false, bool bHasEdgeColor = false, bool bHasThinFilm = false, bool bHasFuzz = false, bool bHasHaziness = false);
 
 FStrataMaterialCompilationInfo StrataCompilationInfoWeight(FMaterialCompiler* Compiler, const FStrataMaterialCompilationInfo& A);
 FStrataMaterialCompilationInfo StrataCompilationInfoAdd(FMaterialCompiler* Compiler, const FStrataMaterialCompilationInfo& A, const FStrataMaterialCompilationInfo& B);
-FStrataMaterialCompilationInfo StrataCompilationInfoAddParamBlend(FMaterialCompiler* Compiler, const FStrataMaterialCompilationInfo& A, const FStrataMaterialCompilationInfo& B, const FStrataRegisteredSharedNormal& RegisteredSharedNormal);
+FStrataMaterialCompilationInfo StrataCompilationInfoAddParamBlend(FMaterialCompiler* Compiler, const FStrataMaterialCompilationInfo& A, const FStrataMaterialCompilationInfo& B, const FStrataRegisteredSharedLocalBasis& RegisteredSharedLocalBasis);
 FStrataMaterialCompilationInfo StrataCompilationInfoHorizontalMixing(FMaterialCompiler* Compiler, const FStrataMaterialCompilationInfo& A, const FStrataMaterialCompilationInfo& B);
-FStrataMaterialCompilationInfo StrataCompilationInfoHorizontalMixingParamBlend(FMaterialCompiler* Compiler, const FStrataMaterialCompilationInfo& A, const FStrataMaterialCompilationInfo& B, const FStrataRegisteredSharedNormal& RegisteredSharedNormal);
+FStrataMaterialCompilationInfo StrataCompilationInfoHorizontalMixingParamBlend(FMaterialCompiler* Compiler, const FStrataMaterialCompilationInfo& A, const FStrataMaterialCompilationInfo& B, const FStrataRegisteredSharedLocalBasis& RegisteredSharedLocalBasis);
 FStrataMaterialCompilationInfo StrataCompilationInfoVerticalLayering(FMaterialCompiler* Compiler, const FStrataMaterialCompilationInfo& Top, const FStrataMaterialCompilationInfo& Base);
-FStrataMaterialCompilationInfo StrataCompilationInfoVerticalLayeringParamBlend(FMaterialCompiler* Compiler, const FStrataMaterialCompilationInfo& Top, const FStrataMaterialCompilationInfo& Base, const FStrataRegisteredSharedNormal& RegisteredSharedNormal);
+FStrataMaterialCompilationInfo StrataCompilationInfoVerticalLayeringParamBlend(FMaterialCompiler* Compiler, const FStrataMaterialCompilationInfo& Top, const FStrataMaterialCompilationInfo& Base, const FStrataRegisteredSharedLocalBasis& RegisteredSharedLocalBasis);
 
 bool StrataIsVolumetricFogCloudOnly(FMaterialCompiler* Compiler, const FStrataMaterialCompilationInfo& Material);
 bool StrataIsUnlitOnly(FMaterialCompiler* Compiler, const FStrataMaterialCompilationInfo& Material);
diff --git a/Engine/Source/Runtime/Engine/Public/MaterialCompiler.h b/Engine/Source/Runtime/Engine/Public/MaterialCompiler.h
--- a/Engine/Source/Runtime/Engine/Public/MaterialCompiler.h
+++ b/Engine/Source/Runtime/Engine/Public/MaterialCompiler.h
@@ -48,13 +48,13 @@ enum class EMaterialCompilerType
 	MaterialProxy, /** Flat material proxy compiler */
 };
 
-struct FStrataRegisteredSharedNormal
+struct FStrataRegisteredSharedLocalBasis
 {
 	int32 NormalCodeChunk;
 	int32 TangentCodeChunk;
 	uint64 NormalCodeChunkHash;
 	uint64 TangentCodeChunkHash;
-	uint8 GraphSharedNormalIndex;
+	uint8 GraphSharedLocalBasisIndex;
 };
 
 /** 
@@ -431,32 +431,32 @@ public:
 		int32 ThinFilmThickness, 
 		int32 FuzzAmount, int32 FuzzColor,
 		int32 Thickness,
-		int32 Normal, int32 Tangent, const FString& SharedNormalIndexMacro) = 0;
-	virtual int32 StrataSheenBSDF(int32 BaseColor, int32 Roughness, int32 Normal, const FString& SharedNormalIndexMacro) = 0;
+		int32 Normal, int32 Tangent, const FString& SharedLocalBasisIndexMacro) = 0;
+	virtual int32 StrataSheenBSDF(int32 BaseColor, int32 Roughness, int32 Normal, const FString& SharedLocalBasisIndexMacro) = 0;
 	virtual int32 StrataVolumetricFogCloudBSDF(int32 Albedo, int32 Extinction, int32 EmissiveColor, int32 AmbientOcclusion) = 0;
 	virtual int32 StrataUnlitBSDF(int32 EmissiveColor, int32 TransmittanceColor) = 0;
-	virtual int32 StrataHairBSDF(int32 BaseColor, int32 Scatter, int32 Specular, int32 Roughness, int32 Backlit, int32 EmissiveColor, int32 Tangent, const FString& SharedNormalIndexMacro) = 0;
+	virtual int32 StrataHairBSDF(int32 BaseColor, int32 Scatter, int32 Specular, int32 Roughness, int32 Backlit, int32 EmissiveColor, int32 Tangent, const FString& SharedLocalBasisIndexMacro) = 0;
 	virtual int32 StrataSingleLayerWaterBSDF(
 		int32 BaseColor, int32 Metallic, int32 Specular, int32 Roughness, int32 EmissiveColor, int32 TopMaterialOpacity, 
-		int32 WaterAlbedo, int32 WaterExtinction, int32 WaterPhaseG, int32 ColorScaleBehindWater, int32 Normal, const FString& SharedNormalIndexMacro) = 0;
+		int32 WaterAlbedo, int32 WaterExtinction, int32 WaterPhaseG, int32 ColorScaleBehindWater, int32 Normal, const FString& SharedLocalBasisIndexMacro) = 0;
 	virtual int32 StrataHorizontalMixing(int32 Foreground, int32 Background, int32 Mix) = 0;
-	virtual int32 StrataHorizontalMixingParameterBlending(int32 Foreground, int32 Background, int32 Mix, const FString& SharedNormalIndexMacro) = 0;
+	virtual int32 StrataHorizontalMixingParameterBlending(int32 Foreground, int32 Background, int32 Mix, const FString& SharedLocalBasisIndexMacro) = 0;
 	virtual int32 StrataVerticalLayering(int32 Top, int32 Base) = 0;
-	virtual int32 StrataVerticalLayeringParameterBlending(int32 Top, int32 Base, const FString& SharedNormalIndexMacro, int32 TopBSDFNormalCodeChunk) = 0;
+	virtual int32 StrataVerticalLayeringParameterBlending(int32 Top, int32 Base, const FString& SharedLocalBasisIndexMacro, int32 TopBSDFNormalCodeChunk) = 0;
 	virtual int32 StrataAdd(int32 A, int32 B) = 0;
-	virtual int32 StrataAddParameterBlending(int32 A, int32 B, const FString& SharedNormalIndexMacro) = 0;
+	virtual int32 StrataAddParameterBlending(int32 A, int32 B, const FString& SharedLocalBasisIndexMacro) = 0;
 	virtual int32 StrataWeight(int32 A, int32 Weight) = 0;
 	virtual int32 StrataTransmittanceToMFP(int32 TransmittanceColor, int32 DesiredThickness, int32 OutputIndex) = 0;
 
 	virtual void StrataCompilationInfoRegisterCodeChunk(int32 CodeChunk, FStrataMaterialCompilationInfo& StrataMaterialCompilationInfo) = 0;
 	virtual bool StrataCompilationInfoContainsCodeChunk(int32 CodeChunk) = 0;
 	virtual const FStrataMaterialCompilationInfo& GetStrataCompilationInfo(int32 CodeChunk) = 0;
-	virtual FStrataRegisteredSharedNormal StrataCompilationInfoRegisterSharedNormal(int32 NormalCodeChunk) = 0;
-	virtual FStrataRegisteredSharedNormal StrataCompilationInfoRegisterSharedNormal(int32 NormalCodeChunk, int32 TangentCodeChunk) = 0;
-	virtual uint8 StrataCompilationInfoGetSharedNormalCount() = 0;
-	FString GetStrataSharedNormalIndexMacro(const FStrataRegisteredSharedNormal& SharedNormal)
+	virtual FStrataRegisteredSharedLocalBasis StrataCompilationInfoRegisterSharedLocalBasis(int32 NormalCodeChunk) = 0;
+	virtual FStrataRegisteredSharedLocalBasis StrataCompilationInfoRegisterSharedLocalBasis(int32 NormalCodeChunk, int32 TangentCodeChunk) = 0;
+	virtual uint8 StrataCompilationInfoGetSharedLocalBasesCount() = 0;
+	FString GetStrataSharedLocalBasisIndexMacro(const FStrataRegisteredSharedLocalBasis& SharedLocalBasis)
 	{
-		return FString::Printf(TEXT("NORMAL_INDEX_%u"), SharedNormal.GraphSharedNormalIndex);
+		return FString::Printf(TEXT("SHAREDLOCALBASIS_INDEX_%u"), SharedLocalBasis.GraphSharedLocalBasisIndex);
 	}
 
 	// Water
@@ -931,7 +931,7 @@ public:
 		int32 ThinFilmThickness, 
 		int32 FuzzAmount, int32 FuzzColor,
 		int32 Thickness,
-		int32 Normal, int32 Tangent, const FString& SharedNormalIndexMacro) override
+		int32 Normal, int32 Tangent, const FString& SharedLocalBasisIndexMacro) override
 	{
 		return Compiler->StrataSlabBSDF(
 			UseMetalness,
@@ -944,12 +944,12 @@ public:
 			ThinFilmThickness,
 			FuzzAmount, FuzzColor,
 			Thickness,
-			Normal, Tangent, SharedNormalIndexMacro);
+			Normal, Tangent, SharedLocalBasisIndexMacro);
 	}
 
-	virtual int32 StrataSheenBSDF(int32 BaseColor, int32 Roughness, int32 Normal, const FString& SharedNormalIndexMacro) override
+	virtual int32 StrataSheenBSDF(int32 BaseColor, int32 Roughness, int32 Normal, const FString& SharedLocalBasisIndexMacro) override
 	{
-		return Compiler->StrataSheenBSDF(BaseColor, Roughness, Normal, SharedNormalIndexMacro);
+		return Compiler->StrataSheenBSDF(BaseColor, Roughness, Normal, SharedLocalBasisIndexMacro);
 	}
 
 	virtual int32 StrataVolumetricFogCloudBSDF(int32 Albedo, int32 Extinction, int32 EmissiveColor, int32 AmbientOcclusion) override
@@ -962,18 +962,18 @@ public:
 		return Compiler->StrataUnlitBSDF(EmissiveColor, TransmittanceColor);
 	}
 
-	virtual int32 StrataHairBSDF(int32 BaseColor, int32 Scatter, int32 Specular, int32 Roughness, int32 Backlit, int32 EmissiveColor, int32 Tangent, const FString& SharedNormalIndexMacro) override
+	virtual int32 StrataHairBSDF(int32 BaseColor, int32 Scatter, int32 Specular, int32 Roughness, int32 Backlit, int32 EmissiveColor, int32 Tangent, const FString& SharedLocalBasisIndexMacro) override
 	{
-		return Compiler->StrataHairBSDF(BaseColor, Scatter, Specular, Roughness, Backlit, EmissiveColor, Tangent, SharedNormalIndexMacro);
+		return Compiler->StrataHairBSDF(BaseColor, Scatter, Specular, Roughness, Backlit, EmissiveColor, Tangent, SharedLocalBasisIndexMacro);
 	}
 
 	virtual int32 StrataSingleLayerWaterBSDF(
 		int32 BaseColor, int32 Metallic, int32 Specular, int32 Roughness, int32 EmissiveColor, int32 TopMaterialOpacity, 
-		int32 WaterAlbedo, int32 WaterExtinction, int32 WaterPhaseG, int32 ColorScaleBehindWater, int32 Normal, const FString& SharedNormalIndexMacro) override
+		int32 WaterAlbedo, int32 WaterExtinction, int32 WaterPhaseG, int32 ColorScaleBehindWater, int32 Normal, const FString& SharedLocalBasisIndexMacro) override
 	{
 		return Compiler->StrataSingleLayerWaterBSDF(
 			BaseColor, Metallic, Specular, Roughness, EmissiveColor, TopMaterialOpacity, 
-			WaterAlbedo, WaterExtinction, WaterPhaseG, ColorScaleBehindWater, Normal, SharedNormalIndexMacro);
+			WaterAlbedo, WaterExtinction, WaterPhaseG, ColorScaleBehindWater, Normal, SharedLocalBasisIndexMacro);
 	}
 
 	virtual int32 StrataHorizontalMixing(int32 Foreground, int32 Background, int32 Mix) override
@@ -981,9 +981,9 @@ public:
 		return Compiler->StrataHorizontalMixing(Foreground, Background, Mix);
 	}
 
-	virtual int32 StrataHorizontalMixingParameterBlending(int32 Foreground, int32 Background, int32 Mix, const FString& SharedNormalIndexMacro) override
+	virtual int32 StrataHorizontalMixingParameterBlending(int32 Foreground, int32 Background, int32 Mix, const FString& SharedLocalBasisIndexMacro) override
 	{
-		return Compiler->StrataHorizontalMixingParameterBlending(Foreground, Background, Mix, SharedNormalIndexMacro);
+		return Compiler->StrataHorizontalMixingParameterBlending(Foreground, Background, Mix, SharedLocalBasisIndexMacro);
 	}
 
 	virtual int32 StrataVerticalLayering(int32 Top, int32 Base) override
@@ -991,9 +991,9 @@ public:
 		return Compiler->StrataVerticalLayering(Top, Base);
 	}
 
-	virtual int32 StrataVerticalLayeringParameterBlending(int32 Top, int32 Base, const FString& SharedNormalIndexMacro, int32 TopBSDFNormalCodeChunk) override
+	virtual int32 StrataVerticalLayeringParameterBlending(int32 Top, int32 Base, const FString& SharedLocalBasisIndexMacro, int32 TopBSDFNormalCodeChunk) override
 	{
-		return Compiler->StrataVerticalLayeringParameterBlending(Top, Base, SharedNormalIndexMacro, TopBSDFNormalCodeChunk);
+		return Compiler->StrataVerticalLayeringParameterBlending(Top, Base, SharedLocalBasisIndexMacro, TopBSDFNormalCodeChunk);
 	}
 
 	virtual int32 StrataAdd(int32 A, int32 B) override
@@ -1001,9 +1001,9 @@ public:
 		return Compiler->StrataAdd(A, B);
 	}
 
-	virtual int32 StrataAddParameterBlending(int32 A, int32 B, const FString& SharedNormalIndexMacro) override
+	virtual int32 StrataAddParameterBlending(int32 A, int32 B, const FString& SharedLocalBasisIndexMacro) override
 	{
-		return Compiler->StrataAddParameterBlending(A, B, SharedNormalIndexMacro);
+		return Compiler->StrataAddParameterBlending(A, B, SharedLocalBasisIndexMacro);
 	}
 
 	virtual int32 StrataWeight(int32 A, int32 Weight) override
@@ -1031,19 +1031,19 @@ public:
 		return Compiler->GetStrataCompilationInfo(CodeChunk);
 	}
 
-	virtual FStrataRegisteredSharedNormal StrataCompilationInfoRegisterSharedNormal(int32 NormalCodeChunk) override
+	virtual FStrataRegisteredSharedLocalBasis StrataCompilationInfoRegisterSharedLocalBasis(int32 NormalCodeChunk) override
 	{
-		return Compiler->StrataCompilationInfoRegisterSharedNormal(NormalCodeChunk);
+		return Compiler->StrataCompilationInfoRegisterSharedLocalBasis(NormalCodeChunk);
 	}
 
-	virtual FStrataRegisteredSharedNormal StrataCompilationInfoRegisterSharedNormal(int32 NormalCodeChunk, int32 TangentCodeChunk) override
+	virtual FStrataRegisteredSharedLocalBasis StrataCompilationInfoRegisterSharedLocalBasis(int32 NormalCodeChunk, int32 TangentCodeChunk) override
 	{
-		return Compiler->StrataCompilationInfoRegisterSharedNormal(NormalCodeChunk, TangentCodeChunk);
+		return Compiler->StrataCompilationInfoRegisterSharedLocalBasis(NormalCodeChunk, TangentCodeChunk);
 	}
 
-	virtual uint8 StrataCompilationInfoGetSharedNormalCount() override
+	virtual uint8 StrataCompilationInfoGetSharedLocalBasesCount() override
 	{
-		return Compiler->StrataCompilationInfoGetSharedNormalCount();
+		return Compiler->StrataCompilationInfoGetSharedLocalBasesCount();
 	}
 
 protected:
