Add warnings when asset manager is told to load an invalid primary asset id, can be disabled in asset manager settings
Move some utility functions to be public and runtime available

diff --git a/Engine/Source/Runtime/Engine/Classes/Engine/AssetManager.h b/Engine/Source/Runtime/Engine/Classes/Engine/AssetManager.h
--- a/Engine/Source/Runtime/Engine/Classes/Engine/AssetManager.h
+++ b/Engine/Source/Runtime/Engine/Classes/Engine/AssetManager.h
@@ -477,6 +477,18 @@ public:
 	/** Loads the redirector maps */
 	virtual void LoadRedirectorMaps();
 
+	/** Refresh the entire set of asset data, can call from editor when things have changed dramatically. Will only refresh if force is true or it thinks something has changed */
+	virtual void RefreshPrimaryAssetDirectory(bool bForceRefresh = false);
+
+	/** Invalidate cached asset data so it knows to rescan when needed */
+	virtual void InvalidatePrimaryAssetDirectory();
+
+	/** Warn about this primary asset id being missing, but only if this is the first time this session */
+	virtual void WarnAboutInvalidPrimaryAsset(const FPrimaryAssetId& PrimaryAssetId, const FString& Message) const;
+
+	/** Helper function to write out asset reports */
+	virtual bool WriteCustomReport(FString FileName, TArray<FString>& FileLines) const;
+
 #if WITH_EDITOR
 	// EDITOR ONLY FUNCTIONALITY
 
@@ -528,9 +540,6 @@ public:
 	/** Returns the list of chunks assigned to the list of primary assets, which is usually a manager list. This is called by GetPackageChunkIds */
 	virtual bool GetPrimaryAssetSetChunkIds(const TSet<FPrimaryAssetId>& PrimaryAssetSet, const class ITargetPlatform* TargetPlatform, TArrayView<const int32> ExistingChunkList, TArray<int32>& OutChunkList) const;
 
-	/** Refresh the entire set of asset data, can call from editor when things have changed dramatically. Will only refresh if force is true or it thinks something has changed */
-	virtual void RefreshPrimaryAssetDirectory(bool bForceRefresh = false);
-
 	/** Resets all asset manager data, called in the editor to reinitialize the config */
 	virtual void ReinitializeFromConfig();
 
@@ -601,9 +610,6 @@ protected:
 	/** Called when an internal load handle finishes, handles setting to pending state */
 	virtual void OnAssetStateChangeCompleted(FPrimaryAssetId PrimaryAssetId, TSharedPtr<FStreamableHandle> BoundHandle, FStreamableDelegate WrappedDelegate);
 
-	/** Helper function to write out asset reports */
-	virtual bool WriteCustomReport(FString FileName, TArray<FString>& FileLines) const;
-
 	/** Scans all asset types specified in DefaultGame */
 	virtual void ScanPrimaryAssetTypesFromConfig();
 
@@ -662,9 +668,6 @@ protected:
 	/** Called after PIE ends, resets loading state */
 	virtual void EndPIE(bool bStartSimulate);
 
-	/** Invalidate cached asset data so it knows to rescan when needed */
-	virtual void InvalidatePrimaryAssetDirectory();
-
 	/** Copy of the asset state before PIE was entered, return to that when PIE completes */
 	TMap<FPrimaryAssetId, TArray<FName>> PrimaryAssetStateBeforePIE;
 
@@ -684,6 +687,9 @@ protected:
 	/** Cached map of asset bundles, global and per primary asset */
 	TMap<FPrimaryAssetId, TSharedPtr<FAssetBundleData, ESPMode::ThreadSafe>> CachedAssetBundles;
 
+	/** List of assets we have warned about being missing */
+	mutable TSet<FPrimaryAssetId> WarningInvalidAssets;
+
 	/** List of directories that have already been synchronously scanned */
 	mutable TArray<FString> AlreadyScannedDirectories;
 
diff --git a/Engine/Source/Runtime/Engine/Classes/Engine/AssetManagerSettings.h b/Engine/Source/Runtime/Engine/Classes/Engine/AssetManagerSettings.h
--- a/Engine/Source/Runtime/Engine/Classes/Engine/AssetManagerSettings.h
+++ b/Engine/Source/Runtime/Engine/Classes/Engine/AssetManagerSettings.h
@@ -69,7 +69,12 @@ class ENGINE_API UAssetManagerSettings : public UDeveloperSettings
 	GENERATED_BODY()
 
 public:
-	UAssetManagerSettings() : bOnlyCookProductionAssets(false), bShouldGuessTypeAndNameInEditor(true), bShouldAcquireMissingChunksOnLoad(false) {}
+	UAssetManagerSettings() 
+	: bOnlyCookProductionAssets(false)
+	, bShouldGuessTypeAndNameInEditor(true)
+	, bShouldAcquireMissingChunksOnLoad(false) 
+	, bShouldWarnAboutInvalidAssets(true)
+	{}
 
 	/** List of asset types to scan at startup */
 	UPROPERTY(config, EditAnywhere, Category = "Asset Manager", meta = (TitleProperty = "PrimaryAssetType"))
@@ -109,6 +114,10 @@ public:
 	UPROPERTY(config, EditAnywhere, Category = "Asset Manager")
 	bool bShouldAcquireMissingChunksOnLoad;
 
+	/** If true, the asset manager will warn when it is told to load or do something with assets it does not know about */
+	UPROPERTY(config, EditAnywhere, Category = "Asset Manager")
+	bool bShouldWarnAboutInvalidAssets;
+
 	/** Redirect from Type:Name to Type:NameNew */
 	UPROPERTY(config, EditAnywhere, Category = "Redirects")
 	TArray<FAssetManagerRedirect> PrimaryAssetIdRedirects;
diff --git a/Engine/Source/Runtime/Engine/Private/AssetManager.cpp b/Engine/Source/Runtime/Engine/Private/AssetManager.cpp
--- a/Engine/Source/Runtime/Engine/Private/AssetManager.cpp
+++ b/Engine/Source/Runtime/Engine/Private/AssetManager.cpp
@@ -1644,7 +1644,7 @@
 		}
 		else
 		{
- 			UE_LOG(LogAssetManager, Verbose, TEXT("%s - UAssetManager::ChangeBundleStateForPrimaryAssets found no NameData for this primary asset."), *PrimaryAssetId.ToString());
+			WarnAboutInvalidPrimaryAsset(PrimaryAssetId, TEXT("ChangeBundleStateForPrimaryAssets failed to find NameData"));
 		}
 	}
 
@@ -1741,6 +1741,10 @@
 			OutAssetLoadSet.Append(Entry.BundleAssets);
 		}
 	}
+	else
+	{
+		WarnAboutInvalidPrimaryAsset(PrimaryAssetId, TEXT("GetPrimaryAssetLoadSet failed to find NameData"));
+	}
 	return NameData != nullptr;
 }
 
@@ -1940,6 +1944,10 @@
 				NameData->PendingState.Reset(true);
 			}
 		}
+		else
+		{
+			WarnAboutInvalidPrimaryAsset(PrimaryAssetId, TEXT("UnloadPrimaryAssets failed to find NameData"));
+		}
 	}
 
 	return NumUnloaded;
@@ -2194,6 +2202,10 @@
 				}
 			}
 		}
+		else
+		{
+			WarnAboutInvalidPrimaryAsset(PrimaryAssetId, TEXT("AcquireResourcesForPrimaryAssetList failed to find NameData"));
+		}
 	}
 
 	AcquireResourcesForAssetList(PathsToLoad.Array(), MoveTemp(CompleteDelegate), Priority);
@@ -3403,6 +3415,77 @@
 	return MoveTemp(ReturnPath);
 }
 
+void UAssetManager::WarnAboutInvalidPrimaryAsset(const FPrimaryAssetId& PrimaryAssetId, const FString& Message) const
+{
+	if (!WarningInvalidAssets.Contains(PrimaryAssetId))
+	{
+		WarningInvalidAssets.Add(PrimaryAssetId);
+
+		const UAssetManagerSettings& Settings = GetSettings();
+		if (Settings.bShouldWarnAboutInvalidAssets)
+		{
+			const TSharedRef<FPrimaryAssetTypeData>* FoundType = AssetTypeMap.Find(PrimaryAssetId.PrimaryAssetType);
+
+			if (FoundType)
+			{
+				UE_LOG(LogAssetManager, Warning, TEXT("Invalid Primary Asset Id %s: %s"), *PrimaryAssetId.ToString(), *Message);
+			}
+			else
+			{
+				UE_LOG(LogAssetManager, Warning, TEXT("Invalid Primary Asset Type %s: %s"), *PrimaryAssetId.ToString(), *Message);
+			}
+		}
+	}
+}
+
+void UAssetManager::InvalidatePrimaryAssetDirectory()
+{
+	bIsPrimaryAssetDirectoryCurrent = false;
+}
+
+void UAssetManager::RefreshPrimaryAssetDirectory(bool bForceRefresh)
+{
+	WarningInvalidAssets.Reset();
+
+	if (bForceRefresh || !bIsPrimaryAssetDirectoryCurrent)
+	{
+		PushBulkScanning();
+
+		for (TPair<FName, TSharedRef<FPrimaryAssetTypeData>>& TypePair : AssetTypeMap)
+		{
+			FPrimaryAssetTypeData& TypeData = TypePair.Value.Get();
+
+			// Rescan the runtime data, the class may have gotten changed by hot reload or config changes
+			bool bIsValid, bBaseClassWasLoaded;
+			TypeData.Info.FillRuntimeData(bIsValid, bBaseClassWasLoaded);
+
+			if (bBaseClassWasLoaded)
+			{
+				// Had to load a class, mark that the temporary cache needs to be updated
+				GetAssetRegistry().SetTemporaryCachingModeInvalidated();
+			}
+
+			if (!bIsValid)
+			{
+				continue;
+			}
+
+			if (TypeData.Info.AssetScanPaths.Num())
+			{
+				// Clear old data if this type has actual scan paths
+				TypeData.AssetMap.Reset();
+
+				// Rescan all assets. We don't force synchronous here as in the editor it was already loaded async
+				ScanPathsForPrimaryAssets(TypePair.Key, TypeData.Info.AssetScanPaths, TypeData.Info.AssetBaseClassLoaded, TypeData.Info.bHasBlueprintClasses, TypeData.Info.bIsEditorOnly, false);
+			}
+		}
+
+		PopBulkScanning();
+
+		PostInitialAssetScan();
+	}
+}
+
 #if WITH_EDITOR
 
 EAssetSetManagerResult::Type UAssetManager::ShouldSetManager(const FAssetIdentifier& Manager, const FAssetIdentifier& Source, const FAssetIdentifier& Target, EAssetRegistryDependencyType::Type DependencyType, EAssetSetManagerFlags::Type Flags) const
@@ -4020,52 +4103,6 @@
 	}
 }
 
-void UAssetManager::InvalidatePrimaryAssetDirectory()
-{
-	bIsPrimaryAssetDirectoryCurrent = false;
-}
-
-void UAssetManager::RefreshPrimaryAssetDirectory(bool bForceRefresh)
-{
-	if (bForceRefresh || !bIsPrimaryAssetDirectoryCurrent)
-	{
-		PushBulkScanning();
-
-		for (TPair<FName, TSharedRef<FPrimaryAssetTypeData>>& TypePair : AssetTypeMap)
-		{
-			FPrimaryAssetTypeData& TypeData = TypePair.Value.Get();
-
-			// Rescan the runtime data, the class may have gotten changed by hot reload or config changes
-			bool bIsValid, bBaseClassWasLoaded;
-			TypeData.Info.FillRuntimeData(bIsValid, bBaseClassWasLoaded);
-
-			if (bBaseClassWasLoaded)
-			{
-				// Had to load a class, mark that the temporary cache needs to be updated
-				GetAssetRegistry().SetTemporaryCachingModeInvalidated();
-			}
-
-			if (!bIsValid)
-			{
-				continue;
-			}
-
-			if (TypeData.Info.AssetScanPaths.Num())
-			{
-				// Clear old data
-				TypeData.AssetMap.Reset();
-
-				// Rescan all assets. We don't force synchronous here as in the editor it was already loaded async
-				ScanPathsForPrimaryAssets(TypePair.Key, TypeData.Info.AssetScanPaths, TypeData.Info.AssetBaseClassLoaded, TypeData.Info.bHasBlueprintClasses, TypeData.Info.bIsEditorOnly, false);
-			}
-		}
-
-		PopBulkScanning();
-
-		PostInitialAssetScan();
-	}
-}
-
 void UAssetManager::ReinitializeFromConfig()
 {
 	// We specifically do not reset AssetRuleOverrides as those can be set by something other than inis
diff --git a/Engine/Source/Runtime/Engine/Private/AssetManagerTypes.cpp b/Engine/Source/Runtime/Engine/Private/AssetManagerTypes.cpp
--- a/Engine/Source/Runtime/Engine/Private/AssetManagerTypes.cpp
+++ b/Engine/Source/Runtime/Engine/Private/AssetManagerTypes.cpp
@@ -52,13 +52,7 @@ void FPrimaryAssetTypeInfo::FillRuntimeData(bool& bIsValid, bool& bBaseClassWasL
 		}
 	}
 
-	if (AssetScanPaths.Num() == 0)
-	{
-		// No scan locations picked out
-		return;
-	}
-
-	// Valid data
+	// Valid data, it's fine for a type to have no scan directories
 	bIsValid = true;
 }
 
