Cleaned-up redundant include of UInstancedStaticMeshComponent

This was being included in a header that is included throughout the engine, resulting in any changes to UInstancedStaticMeshComponent triggering an extensive build (2,500-3,000 actions).
After this change it only triggers a build of ~140 actions (EngineTest).

Code that was relying on this implicit include has been updated to include the header if required.

diff --git a/Engine/Plugins/Enterprise/DataprepEditor/Source/DataprepEditor/Private/DataprepStats.cpp b/Engine/Plugins/Enterprise/DataprepEditor/Source/DataprepEditor/Private/DataprepStats.cpp
--- a/Engine/Plugins/Enterprise/DataprepEditor/Source/DataprepEditor/Private/DataprepStats.cpp
+++ b/Engine/Plugins/Enterprise/DataprepEditor/Source/DataprepEditor/Private/DataprepStats.cpp
@@ -7,6 +7,7 @@
 #include "Components/SkeletalMeshComponent.h"
 #include "Components/StaticMeshComponent.h"
 #include "Components/PrimitiveComponent.h"
+#include "Components/InstancedStaticMeshComponent.h"
 #include "Engine/MapBuildDataRegistry.h"
 #include "Engine/World.h"
 #include "Engine/StaticMesh.h"
diff --git a/Engine/Plugins/Runtime/nDisplay/Source/DisplayClusterConfigurator/Private/Views/Viewport/DisplayClusterConfiguratorSCSEditorViewportClient.cpp b/Engine/Plugins/Runtime/nDisplay/Source/DisplayClusterConfigurator/Private/Views/Viewport/DisplayClusterConfiguratorSCSEditorViewportClient.cpp
--- a/Engine/Plugins/Runtime/nDisplay/Source/DisplayClusterConfigurator/Private/Views/Viewport/DisplayClusterConfiguratorSCSEditorViewportClient.cpp
+++ b/Engine/Plugins/Runtime/nDisplay/Source/DisplayClusterConfigurator/Private/Views/Viewport/DisplayClusterConfiguratorSCSEditorViewportClient.cpp
@@ -24,6 +24,7 @@
 #include "UnrealEdGlobals.h"
 #include "UnrealWidget.h"
 #include "Components/PostProcessComponent.h"
+#include "Components/InstancedStaticMeshComponent.h"
 #include "Editor/EditorPerProjectUserSettings.h"
 #include "Editor/UnrealEdEngine.h"
 #include "Kismet2/ComponentEditorUtils.h"
diff --git a/Engine/Source/Developer/WorldPartitionHLODUtilities/Private/HLODBuilder.cpp b/Engine/Source/Developer/WorldPartitionHLODUtilities/Private/HLODBuilder.cpp
--- a/Engine/Source/Developer/WorldPartitionHLODUtilities/Private/HLODBuilder.cpp
+++ b/Engine/Source/Developer/WorldPartitionHLODUtilities/Private/HLODBuilder.cpp
@@ -3,6 +3,7 @@
 #include "HLODBuilder.h"
 
 #include "Engine/StaticMesh.h"
+#include "Components/InstancedStaticMeshComponent.h"
 
 #include "WorldPartition/HLOD/HLODActor.h"
 #include "WorldPartition/HLOD/HLODLayer.h"
diff --git a/Engine/Source/Developer/WorldPartitionHLODUtilities/Private/HLODBuilderInstancing.cpp b/Engine/Source/Developer/WorldPartitionHLODUtilities/Private/HLODBuilderInstancing.cpp
--- a/Engine/Source/Developer/WorldPartitionHLODUtilities/Private/HLODBuilderInstancing.cpp
+++ b/Engine/Source/Developer/WorldPartitionHLODUtilities/Private/HLODBuilderInstancing.cpp
@@ -4,6 +4,7 @@
 
 #include "Misc/HashBuilder.h"
 #include "Engine/StaticMesh.h"
+#include "Components/InstancedStaticMeshComponent.h"
 #include "WorldPartition/HLOD/HLODActor.h"
 #include "WorldPartition/HLOD/HLODLayer.h"
 
diff --git a/Engine/Source/Runtime/Engine/Classes/Engine/MeshMerging.h b/Engine/Source/Runtime/Engine/Classes/Engine/MeshMerging.h
--- a/Engine/Source/Runtime/Engine/Classes/Engine/MeshMerging.h
+++ b/Engine/Source/Runtime/Engine/Classes/Engine/MeshMerging.h
@@ -5,10 +5,11 @@
 #include "CoreMinimal.h"
 #include "UObject/ObjectMacros.h"
 #include "Engine/MaterialMerging.h"
-#include "GameFramework/Actor.h"
-#include "Components/InstancedStaticMeshComponent.h"
 #include "MeshMerging.generated.h"
 
+class AActor;
+class UInstancedStaticMeshComponent;
+
 /** The importance of a mesh feature when automatically generating mesh LODs. */
 UENUM()
 namespace EMeshFeatureImportance
@@ -638,17 +639,11 @@ struct FSectionInfo
 
 /** Mesh instance-replacement settings */
 USTRUCT(Blueprintable)
-struct FMeshInstancingSettings
+struct ENGINE_API FMeshInstancingSettings
 {
 	GENERATED_BODY()
 
-	FMeshInstancingSettings()
-		: ActorClassToUse(AActor::StaticClass())
-		, InstanceReplacementThreshold(2)
-		, bSkipMeshesWithVertexColors(true)
-		, bUseHLODVolumes(true)
-		, ISMComponentToUse(UInstancedStaticMeshComponent::StaticClass())
-	{}
+	FMeshInstancingSettings();
 
 	/** The actor class to attach new instance static mesh components to */
 	UPROPERTY(BlueprintReadWrite, EditAnywhere, NoClear, Category="Instancing")
diff --git a/Engine/Source/Runtime/Engine/Classes/Engine/StaticMeshSourceData.h b/Engine/Source/Runtime/Engine/Classes/Engine/StaticMeshSourceData.h
--- a/Engine/Source/Runtime/Engine/Classes/Engine/StaticMeshSourceData.h
+++ b/Engine/Source/Runtime/Engine/Classes/Engine/StaticMeshSourceData.h
@@ -7,6 +7,7 @@
 #include "Engine/EngineTypes.h"
 #include "Engine/MeshMerging.h"
 #include "UObject/ObjectMacros.h"
+#include "PerPlatformProperties.h"
 
 #include "StaticMeshSourceData.generated.h"
 
diff --git a/Engine/Source/Runtime/Engine/Private/MeshMerging.cpp b/Engine/Source/Runtime/Engine/Private/MeshMerging.cpp
new file mode 100644
--- /dev/null
+++ b/Engine/Source/Runtime/Engine/Private/MeshMerging.cpp
@@ -0,0 +1,14 @@
+// Copyright Epic Games, Inc. All Rights Reserved.
+
+#include "Engine/MeshMerging.h"
+#include "GameFramework/Actor.h"
+#include "Components/InstancedStaticMeshComponent.h"
+
+FMeshInstancingSettings::FMeshInstancingSettings()
+	: ActorClassToUse(AActor::StaticClass())
+	, InstanceReplacementThreshold(2)
+	, bSkipMeshesWithVertexColors(true)
+	, bUseHLODVolumes(true)
+	, ISMComponentToUse(UInstancedStaticMeshComponent::StaticClass())
+{
+}
