Move gather of compute graph data providers to execution phase.
Fixes timing issue with skin cache provider.
Added support for RDG buffer allocation from data provider which also needed the gathering to happen at the execute phase.

diff --git a/Engine/Source/Runtime/Engine/Private/ComputeFramework/ComputeGraphComponent.cpp b/Engine/Source/Runtime/Engine/Private/ComputeFramework/ComputeGraphComponent.cpp
--- a/Engine/Source/Runtime/Engine/Private/ComputeFramework/ComputeGraphComponent.cpp
+++ b/Engine/Source/Runtime/Engine/Private/ComputeFramework/ComputeGraphComponent.cpp
@@ -61,15 +61,10 @@ void UComputeGraphComponent::SendRenderDynamicData_Concurrent()
 	}
 
 	ENQUEUE_RENDER_COMMAND(ComputeFrameworkEnqueueExecutionCommand)(
-		[ComputeGraphScheduler, ComputeGraphProxy, ProviderProxies = MoveTemp(ComputeDataProviderProxies)](FRHICommandListImmediate& RHICmdList)
+		[ComputeGraphScheduler, ComputeGraphProxy, DataProviderProxies = MoveTemp(ComputeDataProviderProxies)](FRHICommandListImmediate& RHICmdList)
 		{
-			ComputeGraphScheduler->EnqueueForExecution(ComputeGraphProxy, ProviderProxies);
-			
-			for (FComputeDataProviderRenderProxy* ProviderProxy : ProviderProxies)
-			{
-				delete ProviderProxy;
-			}
-
+			// Compute graph scheduler will take ownership of the provider proxies.
+			ComputeGraphScheduler->EnqueueForExecution(ComputeGraphProxy, DataProviderProxies);
 			delete ComputeGraphProxy;
 		});
 }
diff --git a/Engine/Source/Runtime/Engine/Private/ComputeFramework/ComputeGraphScheduler.cpp b/Engine/Source/Runtime/Engine/Private/ComputeFramework/ComputeGraphScheduler.cpp
--- a/Engine/Source/Runtime/Engine/Private/ComputeFramework/ComputeGraphScheduler.cpp
+++ b/Engine/Source/Runtime/Engine/Private/ComputeFramework/ComputeGraphScheduler.cpp
@@ -35,8 +35,10 @@ void FComputeGraphProxy::Initialize(UComputeGraph* ComputeGraph)
 	}
 }
 
-void FComputeGraphScheduler::EnqueueForExecution(const FComputeGraphProxy* ComputeGraph, TArrayView<FComputeDataProviderRenderProxy* const> ComputeDataProviders)
+void FComputeGraphScheduler::EnqueueForExecution(const FComputeGraphProxy* ComputeGraph, TArray<FComputeDataProviderRenderProxy*> ComputeDataProviders)
 {
+	FGraphInvocation& GraphInvocation = GraphInvocations.AddDefaulted_GetRef();
+
 	// todo[CF]: Allocate a specific data provider per kernel to drive the number of invocations?
 	int32 FirstProvider = -1;
 	for (int32 ProviderIndex = 0; ProviderIndex < ComputeDataProviders.Num(); ++ProviderIndex)
@@ -64,29 +66,16 @@ void FComputeGraphScheduler::EnqueueForExecution(const FComputeGraphProxy* Compu
 					Invocation.KernelName,
 					Invocation.InvocationName,
 					DispatchDim,
-					Invocation.Kernel->GetShaderParamMetadata(),
+					Invocation.ShaderMetadata,
 					Invocation.Kernel->GetShader(),
-					Bindings.Num() + SubInvocationIndex
-					};
+					SubInvocationIndex };
 
-				ComputeShaders.Emplace(ShaderInvocation);
+				GraphInvocation.ComputeShaders.Emplace(ShaderInvocation);
 			}
 		}
 	}
 
-	// todo[CF]: Move GetBindings() calls to ExecuteBatches() so that RDG buffers can be registered and bound.
-	Bindings.AddDefaulted(SubInvocationCount);
-	for (int32 SubInvocationIndex = 0; SubInvocationIndex < SubInvocationCount; ++SubInvocationIndex)
-	{
-		for (int32 ProviderIndex = 0; ProviderIndex < ComputeDataProviders.Num(); ++ProviderIndex)
-		{
-			if (ComputeDataProviders[ProviderIndex] != nullptr)
-			{
-				TCHAR const* UID = UComputeGraph::GetDataInterfaceUID(ProviderIndex);
-				ComputeDataProviders[ProviderIndex]->GetBindings(SubInvocationIndex, UID, Bindings[SubInvocationIndex]);
-			}
-		}
-	}
+	GraphInvocation.DataProviders = MoveTemp(ComputeDataProviders);
 }
 
 struct FComputeExecutionBuffer
@@ -108,7 +97,7 @@ void FComputeGraphScheduler::ExecuteBatches(
 	ERHIFeatureLevel::Type FeatureLevel
 	)
 {
-	if (ComputeShaders.IsEmpty())
+	if (GraphInvocations.IsEmpty())
 	{
 		return;
 	}
@@ -121,137 +110,124 @@ void FComputeGraphScheduler::ExecuteBatches(
 		FRDGBuilder GraphBuilder(RHICmdList);
 		FGlobalShaderMap* ShaderMap = GetGlobalShaderMap(FeatureLevel);
 
-		TArray<FComputeExecutionBuffer> ExecutionBuffers;
-
-		// todo[CF]: Implement RDG allocations in the data providers. Probably need to move GetBindings() into the RDG block here?
-		/*
-		for (auto& ExtRes : ExternalBuffers)
-		{
-			TCHAR ResName[FComputeResource::MAX_NAME_LENGTH];
-			ExtRes.Name.ToString(ResName);
-
-			ExecutionBuffers.Emplace(
-				ExtRes.Name,
-				GraphBuilder.RegisterExternalBuffer(ExtRes.Buffer, ResName)
-				);
-		}
-
-		for (auto& TransRes : TransientBuffers)
+		for (int32 GraphIndex = 0; GraphIndex < GraphInvocations.Num(); ++GraphIndex)
 		{
-			TCHAR ResName[FComputeResource::MAX_NAME_LENGTH];
-			TransRes.Name.ToString(ResName);
+			// Gather from all providers.
+			// todo[CF]: This is first pass and needs profiling. Probably with some care we can remove a bunch of heap allocations.
+			TArray<FComputeDataProviderRenderProxy::FBindings> AllBindings;
 
-			FRDGBufferDesc BufferDesc;
-		
-			EComputeExecutionBufferType BufferType = (EComputeExecutionBufferType)TransRes.BufferType;
-			switch (BufferType)
+			TArray<FComputeDataProviderRenderProxy*> const& DataProviders = GraphInvocations[GraphIndex].DataProviders;
+			for (int32 DataProviderIndex = 0; DataProviderIndex < DataProviders.Num(); ++DataProviderIndex)
 			{
-			case EComputeExecutionBufferType::Buffer:
-				BufferDesc = FRDGBufferDesc::CreateBufferDesc(TransRes.BytesPerElement, TransRes.ElementCount);
-				break;
-
-			case EComputeExecutionBufferType::StructuredBuffer:
-				BufferDesc = FRDGBufferDesc::CreateStructuredDesc(TransRes.BytesPerElement, TransRes.ElementCount);
-				break;
+				FComputeDataProviderRenderProxy* DataProvider = DataProviders[DataProviderIndex];
+				if (DataProvider != nullptr)
+				{
+					if (AllBindings.Num() < DataProvider->GetInvocationCount())
+					{
+						AllBindings.SetNum(DataProvider->GetInvocationCount());
+					}
 
-			case EComputeExecutionBufferType::ByteAddressBuffer:
-				BufferDesc = FRDGBufferDesc::CreateByteAddressDesc(TransRes.BytesPerElement * TransRes.ElementCount);
-				break;
+					DataProvider->AllocateResources(GraphBuilder);
 
-			default:
-				check(!"Unsupported buffer type");
-				break;
+					for (int32 InvocationIndex = 0; InvocationIndex < DataProvider->GetInvocationCount(); ++InvocationIndex)
+					{
+						TCHAR const* UID = UComputeGraph::GetDataInterfaceUID(DataProviderIndex);
+						DataProvider->GetBindings(InvocationIndex, UID, AllBindings[InvocationIndex]);
+					}
+				}
 			}
 
-			ExecutionBuffers.Emplace(
-				TransRes.Name,
-				GraphBuilder.CreateBuffer(BufferDesc, ResName)
-				);
-		}
-		*/
-
-		for (int32 KernelIndex = 0; KernelIndex < ComputeShaders.Num(); KernelIndex++)
-		{
-			FShaderInvocation& Compute = ComputeShaders[KernelIndex];
-			FComputeDataProviderRenderProxy::FBindings& Binding = Bindings[Compute.BindingsIndex];
+			// Add compute passes.
+			TArray<FShaderInvocation> const& ComputeShaders = GraphInvocations[GraphIndex].ComputeShaders;
+			for (FShaderInvocation const& Compute : ComputeShaders)
+			{
+				FRDGBufferDesc BufferDesc = FRDGBufferDesc::CreateBufferDesc(4, 32);
 
-			FRDGBufferDesc BufferDesc = FRDGBufferDesc::CreateBufferDesc(4, 32);
+				void* RawBuffer = GraphBuilder.Alloc(Compute.ShaderParamMetadata->GetSize(), SHADER_PARAMETER_STRUCT_ALIGNMENT);
+				FMemory::Memzero(RawBuffer, Compute.ShaderParamMetadata->GetSize());
 
-			void* RawBuffer = GraphBuilder.Alloc(Compute.ShaderParamMetadata->GetSize(), SHADER_PARAMETER_STRUCT_ALIGNMENT);
-			FMemory::Memzero(RawBuffer, Compute.ShaderParamMetadata->GetSize());
+				FComputeDataProviderRenderProxy::FBindings& Bindings = AllBindings[Compute.SubInvocationIndex];
 
-			bool bBindFailed = false;
-			uint8* ParamBuffer = static_cast<uint8*>(RawBuffer);
-			TArray<FShaderParametersMetadata::FMember> ParamMembers = Compute.ShaderParamMetadata->GetMembers();
-			for (auto& Member : ParamMembers)
-			{
-				switch (Member.GetBaseType())
+				bool bBindFailed = false;
+				uint8* ParamBuffer = static_cast<uint8*>(RawBuffer);
+				TArray<FShaderParametersMetadata::FMember> ParamMembers = Compute.ShaderParamMetadata->GetMembers();
+				for (auto& Member : ParamMembers)
 				{
-				case EUniformBufferBaseType::UBMT_INT32:
+					switch (Member.GetBaseType())
 					{
-						int32* ParamValue = Binding.ParamsInt.Find(Member.GetName());
-						*reinterpret_cast<int32*>(&ParamBuffer[Member.GetOffset()]) = ParamValue != nullptr ? *ParamValue : 0;
-					}
-					break;
-
-				case EUniformBufferBaseType::UBMT_UINT32:
-					{
-						uint32* ParamValue = Binding.ParamsUint.Find(Member.GetName());
-						*reinterpret_cast<uint32*>(&ParamBuffer[Member.GetOffset()]) = ParamValue != nullptr ? *ParamValue : 0;
-					}
-				break;
+					case EUniformBufferBaseType::UBMT_INT32:
+						{
+							int32* ParamValue = Bindings.ParamsInt.Find(Member.GetName());
+							*reinterpret_cast<int32*>(&ParamBuffer[Member.GetOffset()]) = ParamValue != nullptr ? *ParamValue : 0;
+						}
+						break;
 
-				case EUniformBufferBaseType::UBMT_FLOAT32:
-					{
-						float* ParamValue = Binding.ParamsFloat.Find(Member.GetName());
-						*reinterpret_cast<float*>(&ParamBuffer[Member.GetOffset()]) = ParamValue != nullptr ? *ParamValue : 0;
-					}
+					case EUniformBufferBaseType::UBMT_UINT32:
+						{
+							uint32* ParamValue = Bindings.ParamsUint.Find(Member.GetName());
+							*reinterpret_cast<uint32*>(&ParamBuffer[Member.GetOffset()]) = ParamValue != nullptr ? *ParamValue : 0;
+						}
 					break;
 
-				case EUniformBufferBaseType::UBMT_NESTED_STRUCT:
-					{
-						TArray<uint8>* ParamValue = Binding.Structs.Find(Member.GetName());
-						if (ParamValue != nullptr)
+					case EUniformBufferBaseType::UBMT_FLOAT32:
 						{
-							FMemory::Memcpy(&ParamBuffer[Member.GetOffset()], ParamValue->GetData(), Member.GetStructMetadata()->GetSize());
+							float* ParamValue = Bindings.ParamsFloat.Find(Member.GetName());
+							*reinterpret_cast<float*>(&ParamBuffer[Member.GetOffset()]) = ParamValue != nullptr ? *ParamValue : 0;
 						}
-						else
+						break;
+
+					case EUniformBufferBaseType::UBMT_NESTED_STRUCT:
 						{
-							bBindFailed = true;
+							TArray<uint8>* ParamValue = Bindings.Structs.Find(Member.GetName());
+							if (ParamValue != nullptr)
+							{
+								FMemory::Memcpy(&ParamBuffer[Member.GetOffset()], ParamValue->GetData(), Member.GetStructMetadata()->GetSize());
+							}
+							else
+							{
+								bBindFailed = true;
+							}
 						}
-					}
-					break;
+						break;
 
-				default:
-					check(!"Unsupported type");
-					bBindFailed = true;
-					break;
-				};
-			}
+					default:
+						check(!"Unsupported type");
+						bBindFailed = true;
+						break;
+					};
+				}
 
-			if (ensure(!bBindFailed))
-			{
-				TCHAR KernelName[128];
-				Compute.KernelName.ToString(KernelName);
-
-				TCHAR InvocationName[128];
-				Compute.InvocationName.ToString(InvocationName);
-
-				FComputeShaderUtils::AddPass(
-					GraphBuilder,
-					RDG_EVENT_NAME("Compute[%s]: %s", KernelName, InvocationName),
-					ERDGPassFlags::Compute | ERDGPassFlags::NeverCull,
-					Compute.Shader,
-					Compute.ShaderParamMetadata,
-					static_cast<FComputeKernelShader::FParameters*>(RawBuffer),
-					Compute.DispatchDim
-					);		
+				if (ensure(!bBindFailed))
+				{
+					TCHAR KernelName[128];
+					Compute.KernelName.ToString(KernelName);
+
+					TCHAR InvocationName[128];
+					Compute.InvocationName.ToString(InvocationName);
+
+					FComputeShaderUtils::AddPass(
+						GraphBuilder,
+						RDG_EVENT_NAME("Compute[%s]: %s", KernelName, InvocationName),
+						ERDGPassFlags::Compute | ERDGPassFlags::NeverCull,
+						Compute.Shader,
+						Compute.ShaderParamMetadata,
+						static_cast<FComputeKernelShader::FParameters*>(RawBuffer),
+						Compute.DispatchDim
+						);		
+				}
 			}
 		}
 
 		GraphBuilder.Execute();
 
-		ComputeShaders.Reset();
-		Bindings.Reset();
+		GraphInvocations.Reset();
+	}
+}
+
+FComputeGraphScheduler::FGraphInvocation::~FGraphInvocation()
+{
+	for (FComputeDataProviderRenderProxy* DataProvider : DataProviders)
+	{
+		delete DataProvider;
 	}
 }
diff --git a/Engine/Source/Runtime/Engine/Private/GPUSkinCache.cpp b/Engine/Source/Runtime/Engine/Private/GPUSkinCache.cpp
--- a/Engine/Source/Runtime/Engine/Private/GPUSkinCache.cpp
+++ b/Engine/Source/Runtime/Engine/Private/GPUSkinCache.cpp
@@ -2232,7 +2232,7 @@ FRWBuffer* FGPUSkinCache::GetPositionBuffer(uint32 ComponentId, uint32 SectionIn
 		if (Entry && Entry->GPUSkin && Entry->GPUSkin->GetComponentId() == ComponentId)
 		{
 			FGPUSkinCacheEntry::FSectionDispatchData& DispatchData = Entry->GetDispatchData()[SectionIndex];
-			FSkinCacheRWBuffer* SkinCacheRWBuffer = DispatchData.PreviousPositionBuffer != nullptr ? DispatchData.PreviousPositionBuffer : DispatchData.PositionBuffer;
+			FSkinCacheRWBuffer* SkinCacheRWBuffer = DispatchData.PositionBuffer;
 			return SkinCacheRWBuffer != nullptr ? &SkinCacheRWBuffer->Buffer : nullptr;
 		}
 	}
diff --git a/Engine/Source/Runtime/Engine/Public/ComputeFramework/ComputeDataProvider.h b/Engine/Source/Runtime/Engine/Public/ComputeFramework/ComputeDataProvider.h
--- a/Engine/Source/Runtime/Engine/Public/ComputeFramework/ComputeDataProvider.h
+++ b/Engine/Source/Runtime/Engine/Public/ComputeFramework/ComputeDataProvider.h
@@ -7,10 +7,7 @@
 
 class FComputeDataProviderRenderProxy;
 struct FComputeKernelPermutationSet;
-class FRDGBuffer;
-class FRDGTexture;
-class FRHIShaderResourceView;
-class FRHIUnorderedAccessView;
+class FRDGBuilder;
 
 /**
  * Compute Framework Data Provider.
@@ -52,6 +49,9 @@ public:
 	 */
 	virtual void GetPermutations(int32 InvocationIndex, FComputeKernelPermutationSet& OutPermutationSet) const {}
 
+	/* Called once before any calls to GetBindings() to allow any RDG resource allocation. */
+	virtual void AllocateResources(FRDGBuilder& GraphBuilder) {}
+
 	/** 
 	 * The name-value shader bindings that are collected from data providers.
 	 * All names should already have been registered by the associated data interface. 
diff --git a/Engine/Source/Runtime/Engine/Public/ComputeFramework/ComputeGraphScheduler.h b/Engine/Source/Runtime/Engine/Public/ComputeFramework/ComputeGraphScheduler.h
--- a/Engine/Source/Runtime/Engine/Public/ComputeFramework/ComputeGraphScheduler.h
+++ b/Engine/Source/Runtime/Engine/Public/ComputeFramework/ComputeGraphScheduler.h
@@ -39,20 +39,18 @@ public:
  * Class that manages the scheduling of Compute Graph work.
  * Work can be enqueued on the render thread for the execution at the next call to ExecuteBatches().
  */
-class FComputeGraphScheduler
+class ENGINE_API FComputeGraphScheduler
 {
 public:
 	/** Enqueue a compute graph for execution. */
 	void EnqueueForExecution(
 		const FComputeGraphProxy* ComputeGraph,
-		TArrayView<FComputeDataProviderRenderProxy* const> ComputeDataProviders
-		);
+		TArray<FComputeDataProviderRenderProxy*> ComputeDataProviders );
 
 	/** Submit enqueued compute graph work. */
 	void ExecuteBatches(
 		FRHICommandListImmediate& RHICmdList,
-		ERHIFeatureLevel::Type FeatureLevel
-		);
+		ERHIFeatureLevel::Type FeatureLevel );
 
 private:
 	/** Description of each dispatch that is enqueued. */
@@ -63,18 +61,21 @@ private:
 		FIntVector DispatchDim;
 		const FShaderParametersMetadata* ShaderParamMetadata = nullptr;
 		TShaderRef<FComputeKernelShader> Shader;
-		int32 BindingsIndex = 0;
+		int32 SubInvocationIndex = 0;
 	};
 
-	/** 
-	 * Shader invocations to dispatch. Stored in required order for dispatch. 
+	/**
+	 * Description of each graph that is enqueued.
 	 * todo[CF]: We probably need more context for dispatching work with minimal overhead. For example we would like to overlap UAVs on any skin cache writing.
 	 */
-	TArray<FShaderInvocation> ComputeShaders;
+	struct FGraphInvocation
+	{
+		/** Shader invocations to dispatch. */
+		TArray<FShaderInvocation> ComputeShaders;
+		TArray<FComputeDataProviderRenderProxy*> DataProviders;
+	
+		~FGraphInvocation();
+	};
 
-	/** 
-	 * Parameter and resource binding information for queued dispatches.
-	 * todo[CF]: Replace this with something more efficient. Hopefully using the existing StructMetadata classes.
-	 */
-	TArray<FComputeDataProviderRenderProxy::FBindings> Bindings;
+	TArray<FGraphInvocation> GraphInvocations;
 };
