Strata - fixed add parameter blending to account for weight when evaluating the normals.

diff --git a/Engine/Shaders/Private/Strata/Strata.ush b/Engine/Shaders/Private/Strata/Strata.ush
--- a/Engine/Shaders/Private/Strata/Strata.ush
+++ b/Engine/Shaders/Private/Strata/Strata.ush
@@ -970,10 +970,19 @@ FStrataData StrataHorizontalMixingParameterBlending(FStrataData Foreground, FStr
 	return Result;
 }
 
-// ==> NOTE: Always pair with the compiler behavior in StrataCompilationInfoAddParamBlend
-FStrataData StrataAddParameterBlending(FStrataData A, FStrataData B, uint NewNormalIndex)
+float AddParameterBlendingBSDFWeightToNormalMix(FStrataData AStrata, FStrataData BStrata) 
 {
+	const float3 SafeABSDFWeight = saturate(AStrata.Layers[0].BSDFs[0].Weight);
+	const float3 SafeBBSDFWeight = saturate(BStrata.Layers[0].BSDFs[0].Weight);
+	const float SafeGreyABSDFWeight = dot(SafeABSDFWeight, 1.0f / 3.0f);
+	const float SafeGreyBBSDFWeight = dot(SafeBBSDFWeight, 1.0f / 3.0f);
+	const float AMixFactor = SafeGreyABSDFWeight / (1e-10 + SafeGreyABSDFWeight + SafeGreyBBSDFWeight);
+	return AMixFactor;
+}
 
+// ==> NOTE: Always pair with the compiler behavior in StrataCompilationInfoAddParamBlend
+FStrataData StrataAddParameterBlending(FStrataData A, FStrataData B, float AMixFactor, uint NewNormalIndex)
+{
 	FStrataData Result = GetInitialisedStrataDataWithReservedFirstSlot();
 #define ResultBSDF		Result.Layers[0].BSDFs[0]
 #define ABSDF			A.Layers[0].BSDFs[0]
@@ -981,9 +990,6 @@ FStrataData StrataAddParameterBlending(FStrataData A, FStrataData B, uint NewNor
 
 	const float3 SafeABSDFWeight = saturate(ABSDF.Weight);
 	const float3 SafeBBSDFWeight = saturate(BBSDF.Weight);
-	const float SafeGreyABSDFWeight = dot(SafeABSDFWeight, 1.0f / 3.0f);
-	const float SafeGreyBBSDFWeight = dot(SafeBBSDFWeight, 1.0f / 3.0f);
-	const float AMixFactor = SafeGreyABSDFWeight / (1e-10 + SafeGreyABSDFWeight + SafeGreyBBSDFWeight);
 
 	// We can only parameter blend the slab 
 	BSDF_SETTYPE(ResultBSDF,			STRATA_BSDF_TYPE_SLAB);
diff --git a/Engine/Source/Runtime/Engine/Private/Materials/HLSLMaterialTranslator.cpp b/Engine/Source/Runtime/Engine/Private/Materials/HLSLMaterialTranslator.cpp
--- a/Engine/Source/Runtime/Engine/Private/Materials/HLSLMaterialTranslator.cpp
+++ b/Engine/Source/Runtime/Engine/Private/Materials/HLSLMaterialTranslator.cpp
@@ -8985,6 +8985,11 @@ uint8 FHLSLMaterialTranslator::StrataCompilationInfoGetSharedLocalBasesCount()
 	return FinalUsedSharedLocalBasesCount;
 }
 
+int32 FHLSLMaterialTranslator::StrataAddParameterBlendingBSDFWeightToNormalMixCodeChunk(int32 ACodeChunk, int32 BCodeChunk)
+{
+	return AddCodeChunk(MCT_Float, TEXT("AddParameterBlendingBSDFWeightToNormalMix(%s, %s)"), *GetParameterCode(ACodeChunk), *GetParameterCode(BCodeChunk));
+}
+
 int32 FHLSLMaterialTranslator::StrataCreateAndRegisterNullMaterial()
 {
 	int32 OutputCodeChunk = AddInlinedCodeChunk(MCT_Strata, TEXT("GetInitialisedStrataData()"));
@@ -9165,16 +9170,17 @@ int32 FHLSLMaterialTranslator::StrataAdd(int32 A, int32 B)
 	);
 }
 
-int32 FHLSLMaterialTranslator::StrataAddParameterBlending(int32 A, int32 B, const FString& SharedLocalBasisIndexMacro)
+int32 FHLSLMaterialTranslator::StrataAddParameterBlending(int32 A, int32 B, int32 AMixWeight, const FString& SharedLocalBasisIndexMacro)
 {
 	if (A == INDEX_NONE || B == INDEX_NONE)
 	{
 		return INDEX_NONE;
 	}
 	return AddCodeChunk(
-		MCT_Strata, TEXT("StrataAddParameterBlending(%s, %s, %s)"),
+		MCT_Strata, TEXT("StrataAddParameterBlending(%s, %s, %s, %s)"),
 		*GetParameterCode(A),
 		*GetParameterCode(B),
+		*GetParameterCode(AMixWeight),
 		*SharedLocalBasisIndexMacro
 	);
 }
diff --git a/Engine/Source/Runtime/Engine/Private/Materials/HLSLMaterialTranslator.h b/Engine/Source/Runtime/Engine/Private/Materials/HLSLMaterialTranslator.h
--- a/Engine/Source/Runtime/Engine/Private/Materials/HLSLMaterialTranslator.h
+++ b/Engine/Source/Runtime/Engine/Private/Materials/HLSLMaterialTranslator.h
@@ -1001,7 +1001,7 @@ protected:
 	virtual int32 StrataVerticalLayering(int32 Top, int32 Base) override;
 	virtual int32 StrataVerticalLayeringParameterBlending(int32 Top, int32 Base, const FString& SharedLocalBasisIndexMacro, int32 TopBSDFNormalCodeChunk) override;
 	virtual int32 StrataAdd(int32 A, int32 B) override;
-	virtual int32 StrataAddParameterBlending(int32 A, int32 B, const FString& SharedLocalBasisIndexMacro) override;
+	virtual int32 StrataAddParameterBlending(int32 A, int32 B, int32 AMixWeight, const FString& SharedLocalBasisIndexMacro) override;
 	virtual int32 StrataWeight(int32 A, int32 Weight) override;
 	virtual int32 StrataTransmittanceToMFP(int32 TransmittanceColor, int32 DesiredThickness, int32 OutputIndex) override;
 
@@ -1011,6 +1011,7 @@ protected:
 	virtual FStrataRegisteredSharedLocalBasis StrataCompilationInfoRegisterSharedLocalBasis(int32 NormalCodeChunk) override;
 	virtual FStrataRegisteredSharedLocalBasis StrataCompilationInfoRegisterSharedLocalBasis(int32 NormalCodeChunk, int32 TangentCodeChunk) override;
 	virtual uint8 StrataCompilationInfoGetSharedLocalBasesCount() override;
+	virtual int32 StrataAddParameterBlendingBSDFWeightToNormalMixCodeChunk(int32 ACodeChunk, int32 BCodeChunk) override;
 
 	FStrataSharedLocalBasesInfo StrataCompilationInfoGetMatchingSharedLocalBasisInfo(const FStrataRegisteredSharedLocalBasis& SearchedSharedLocalBasis);
 
diff --git a/Engine/Source/Runtime/Engine/Private/Materials/MaterialExpressions.cpp b/Engine/Source/Runtime/Engine/Private/Materials/MaterialExpressions.cpp
--- a/Engine/Source/Runtime/Engine/Private/Materials/MaterialExpressions.cpp
+++ b/Engine/Source/Runtime/Engine/Private/Materials/MaterialExpressions.cpp
@@ -20971,7 +20971,8 @@ int32 UMaterialExpressionStrataVerticalLayering::Compile(class FMaterialCompiler
 	const FStrataMaterialCompilationInfo& TopStrataData = Compiler->GetStrataCompilationInfo(TopCodeChunk);
 	const FStrataMaterialCompilationInfo& BaseStrataData = Compiler->GetStrataCompilationInfo(BaseCodeChunk);
 
-	const int32 MixCodeChunk = Compiler->Constant(0.5f);
+	// STRATA_TODO: use a weight accounting for weights and transmittance?
+	const int32 NormalMixCodeChunk = Compiler->Constant(0.5f);
 
 	int32 OutputCodeChunk = INDEX_NONE;
 	if (bUseParameterBlending)
@@ -20997,12 +20998,12 @@ int32 UMaterialExpressionStrataVerticalLayering::Compile(class FMaterialCompiler
 		}
 
 		// Compute the new Normal and Tangent resulting from the blending using code chunk
-		const int32 NewNormalCodeChunk = StrataBlendNormal(Compiler, TopBSDF.RegisteredSharedLocalBasis.NormalCodeChunk, BaseBSDF.RegisteredSharedLocalBasis.NormalCodeChunk, MixCodeChunk);
+		const int32 NewNormalCodeChunk = StrataBlendNormal(Compiler, TopBSDF.RegisteredSharedLocalBasis.NormalCodeChunk, BaseBSDF.RegisteredSharedLocalBasis.NormalCodeChunk, NormalMixCodeChunk);
 		// The tangent is optional so we treat it differently if INDEX_NONE is specified
 		int32 NewTangentCodeChunk = INDEX_NONE;
 		if (TopBSDF.RegisteredSharedLocalBasis.TangentCodeChunk != INDEX_NONE && BaseBSDF.RegisteredSharedLocalBasis.TangentCodeChunk != INDEX_NONE)
 		{
-			NewTangentCodeChunk = StrataBlendNormal(Compiler, TopBSDF.RegisteredSharedLocalBasis.TangentCodeChunk, BaseBSDF.RegisteredSharedLocalBasis.TangentCodeChunk, MixCodeChunk);
+			NewTangentCodeChunk = StrataBlendNormal(Compiler, TopBSDF.RegisteredSharedLocalBasis.TangentCodeChunk, BaseBSDF.RegisteredSharedLocalBasis.TangentCodeChunk, NormalMixCodeChunk);
 		}
 		else if (TopBSDF.RegisteredSharedLocalBasis.TangentCodeChunk != INDEX_NONE)
 		{
@@ -21112,7 +21113,7 @@ int32 UMaterialExpressionStrataAdd::Compile(class FMaterialCompiler* Compiler, i
 	const FStrataMaterialCompilationInfo& AStrataData = Compiler->GetStrataCompilationInfo(ACodeChunk);
 	const FStrataMaterialCompilationInfo& BStrataData = Compiler->GetStrataCompilationInfo(BCodeChunk);
 
-	const int32 MixCodeChunk = Compiler->Constant(0.5f);
+	const int32 ANormalMixCodeChunk = Compiler->StrataAddParameterBlendingBSDFWeightToNormalMixCodeChunk(ACodeChunk, BCodeChunk);
 
 	int32 OutputCodeChunk = INDEX_NONE;
 	if (bUseParameterBlending)
@@ -21138,12 +21139,12 @@ int32 UMaterialExpressionStrataAdd::Compile(class FMaterialCompiler* Compiler, i
 		}
 
 		// Compute the new Normal and Tangent resulting from the blending using code chunk
-		const int32 NewNormalCodeChunk = StrataBlendNormal(Compiler, ABSDF.RegisteredSharedLocalBasis.NormalCodeChunk, BBSDF.RegisteredSharedLocalBasis.NormalCodeChunk, MixCodeChunk);
+		const int32 NewNormalCodeChunk = StrataBlendNormal(Compiler, BBSDF.RegisteredSharedLocalBasis.NormalCodeChunk, ABSDF.RegisteredSharedLocalBasis.NormalCodeChunk, ANormalMixCodeChunk);
 		// The tangent is optional so we treat it differently if INDEX_NONE is specified
 		int32 NewTangentCodeChunk = INDEX_NONE;
 		if (ABSDF.RegisteredSharedLocalBasis.TangentCodeChunk != INDEX_NONE && BBSDF.RegisteredSharedLocalBasis.TangentCodeChunk != INDEX_NONE)
 		{
-			NewTangentCodeChunk = StrataBlendNormal(Compiler, ABSDF.RegisteredSharedLocalBasis.TangentCodeChunk, BBSDF.RegisteredSharedLocalBasis.TangentCodeChunk, MixCodeChunk);
+			NewTangentCodeChunk = StrataBlendNormal(Compiler, ABSDF.RegisteredSharedLocalBasis.TangentCodeChunk, BBSDF.RegisteredSharedLocalBasis.TangentCodeChunk, ANormalMixCodeChunk);
 		}
 		else if (ABSDF.RegisteredSharedLocalBasis.TangentCodeChunk != INDEX_NONE)
 		{
@@ -21155,7 +21156,7 @@ int32 UMaterialExpressionStrataAdd::Compile(class FMaterialCompiler* Compiler, i
 		}
 		const FStrataRegisteredSharedLocalBasis NewRegisteredSharedLocalBasis = StrataCompilationInfoCreateSharedLocalBasis(Compiler, NewNormalCodeChunk, NewTangentCodeChunk);
 
-		OutputCodeChunk = Compiler->StrataAddParameterBlending(ACodeChunk, BCodeChunk, Compiler->GetStrataSharedLocalBasisIndexMacro(NewRegisteredSharedLocalBasis));
+		OutputCodeChunk = Compiler->StrataAddParameterBlending(ACodeChunk, BCodeChunk, ANormalMixCodeChunk, Compiler->GetStrataSharedLocalBasisIndexMacro(NewRegisteredSharedLocalBasis));
 
 		StrataCompilationInfoCreateSingleBSDFMaterial(Compiler, OutputCodeChunk, NewRegisteredSharedLocalBasis, STRATA_BSDF_TYPE_SLAB);
 		FStrataMaterialCompilationInfo StrataInfo = StrataCompilationInfoAddParamBlend(Compiler, AStrataData, BStrataData, NewRegisteredSharedLocalBasis);
diff --git a/Engine/Source/Runtime/Engine/Public/MaterialCompiler.h b/Engine/Source/Runtime/Engine/Public/MaterialCompiler.h
--- a/Engine/Source/Runtime/Engine/Public/MaterialCompiler.h
+++ b/Engine/Source/Runtime/Engine/Public/MaterialCompiler.h
@@ -444,7 +444,7 @@ public:
 	virtual int32 StrataVerticalLayering(int32 Top, int32 Base) = 0;
 	virtual int32 StrataVerticalLayeringParameterBlending(int32 Top, int32 Base, const FString& SharedLocalBasisIndexMacro, int32 TopBSDFNormalCodeChunk) = 0;
 	virtual int32 StrataAdd(int32 A, int32 B) = 0;
-	virtual int32 StrataAddParameterBlending(int32 A, int32 B, const FString& SharedLocalBasisIndexMacro) = 0;
+	virtual int32 StrataAddParameterBlending(int32 A, int32 B, int32 AMixWeight, const FString& SharedLocalBasisIndexMacro) = 0;
 	virtual int32 StrataWeight(int32 A, int32 Weight) = 0;
 	virtual int32 StrataTransmittanceToMFP(int32 TransmittanceColor, int32 DesiredThickness, int32 OutputIndex) = 0;
 
@@ -458,6 +458,7 @@ public:
 	{
 		return FString::Printf(TEXT("SHAREDLOCALBASIS_INDEX_%u"), SharedLocalBasis.GraphSharedLocalBasisIndex);
 	}
+	virtual int32 StrataAddParameterBlendingBSDFWeightToNormalMixCodeChunk(int32 ACodeChunk, int32 BCodeChunk) = 0;
 
 	// Water
 	virtual int32 SceneDepthWithoutWater(int32 Offset, int32 ViewportUV, bool bUseOffset, float FallbackDepth) = 0;
@@ -1001,9 +1002,9 @@ public:
 		return Compiler->StrataAdd(A, B);
 	}
 
-	virtual int32 StrataAddParameterBlending(int32 A, int32 B, const FString& SharedLocalBasisIndexMacro) override
+	virtual int32 StrataAddParameterBlending(int32 A, int32 B, int32 AMixWeight, const FString& SharedLocalBasisIndexMacro) override
 	{
-		return Compiler->StrataAddParameterBlending(A, B, SharedLocalBasisIndexMacro);
+		return Compiler->StrataAddParameterBlending(A, B, AMixWeight, SharedLocalBasisIndexMacro);
 	}
 
 	virtual int32 StrataWeight(int32 A, int32 Weight) override
@@ -1025,6 +1026,11 @@ public:
 	{
 		return Compiler->StrataCompilationInfoContainsCodeChunk(CodeChunk);
 	}
+
+	virtual int32 StrataAddParameterBlendingBSDFWeightToNormalMixCodeChunk(int32 ACodeChunk, int32 BCodeChunk) override
+	{
+		return Compiler->StrataAddParameterBlendingBSDFWeightToNormalMixCodeChunk(ACodeChunk, BCodeChunk);
+	}
 	
 	virtual const FStrataMaterialCompilationInfo& GetStrataCompilationInfo(int32 CodeChunk) override
 	{
