Unify asset compilers notification to fix some discrepancies

#rnx

diff --git a/Engine/Source/Editor/UnrealEd/Private/CardRepresentationBuildNotification.cpp b/Engine/Source/Editor/UnrealEd/Private/CardRepresentationBuildNotification.cpp
deleted file mode 100644
diff --git a/Engine/Source/Editor/UnrealEd/Private/DistanceFieldBuildNotification.cpp b/Engine/Source/Editor/UnrealEd/Private/DistanceFieldBuildNotification.cpp
deleted file mode 100644
diff --git a/Engine/Source/Editor/UnrealEd/Private/EditorEngine.cpp b/Engine/Source/Editor/UnrealEd/Private/EditorEngine.cpp
--- a/Engine/Source/Editor/UnrealEd/Private/EditorEngine.cpp
+++ b/Engine/Source/Editor/UnrealEd/Private/EditorEngine.cpp
@@ -703,12 +703,6 @@ void UEditorEngine::InitEditor(IEngineLoop* InEngineLoop)
 	FCoreUObjectDelegates::IsPackageOKToSaveDelegate.BindUObject(this, &UEditorEngine::IsPackageOKToSave);
 	FCoreUObjectDelegates::AutoPackageBackupDelegate.BindStatic(&FAutoPackageBackup::BackupPackage);
 
-	extern void SetupDistanceFieldBuildNotification();
-	SetupDistanceFieldBuildNotification();
-
-	extern void SetupCardRepresentationBuildNotification();
-	SetupCardRepresentationBuildNotification();
-
 	// Update recents
 	UpdateRecentlyLoadedProjectFiles();
 
@@ -1301,12 +1295,6 @@ void UEditorEngine::FinishDestroy()
 		// Destroy selection sets.
 		PrivateEditorSelection::DestroySelectionSets();
 
-		extern void TearDownDistanceFieldBuildNotification();
-		TearDownDistanceFieldBuildNotification();
-
-		extern void TearDownCardRepresentationBuildNotification();
-		TearDownCardRepresentationBuildNotification();
-
 		// Remove editor array from root.
 		UE_LOG(LogExit, Log, TEXT("Editor shut down") );
 
diff --git a/Engine/Source/Editor/UnrealEd/Private/ShaderCompilingNotification.cpp b/Engine/Source/Editor/UnrealEd/Private/ShaderCompilingNotification.cpp
deleted file mode 100644
diff --git a/Engine/Source/Runtime/Engine/Private/AsyncCompilationHelpers.cpp b/Engine/Source/Runtime/Engine/Private/AsyncCompilationHelpers.cpp
--- a/Engine/Source/Runtime/Engine/Private/AsyncCompilationHelpers.cpp
+++ b/Engine/Source/Runtime/Engine/Private/AsyncCompilationHelpers.cpp
@@ -24,38 +24,23 @@
 
 DEFINE_LOG_CATEGORY(LogAsyncCompilation);
 
-#if WITH_EDITOR
-
 #define LOCTEXT_NAMESPACE "AsyncCompilation"
 
-namespace AsyncCompilationHelpers
-{
-	void DumpStallStacks();
-}
-
-static FAutoConsoleCommand CVarAsyncAssetDumpStallStacks(
-	TEXT("Editor.AsyncAssetDumpStallStacks"),
-	TEXT("Dump all the callstacks that have caused waits on async compilation."),
-	FConsoleCommandWithArgsDelegate::CreateLambda([](const TArray<FString>& Args)
-	{
-		AsyncCompilationHelpers::DumpStallStacks();
-	})
-);
-
 void FAsyncCompilationNotification::Update(int32 NumJobs)
 {
+#if WITH_EDITOR
 	check(IsInGameThread());
 
 	// Use a lambda to only for progress message in code-path where it's needed
 	auto GetProgressMessage =
 		[this, NumJobs]()
-		{
-			FFormatNamedArguments Args;
-			// Always use the plural form when displaying the notification
-			Args.Add(TEXT("AssetType"), FText::Format(AssetNameFormat, FText::AsNumber(100)));
-			Args.Add(TEXT("NumJobs"), FText::AsNumber(NumJobs));
-			return FText::Format(LOCTEXT("AsyncCompilationProgress", "Preparing {AssetType} ({NumJobs})"), Args);
-		};
+	{
+		FFormatNamedArguments Args;
+		// Always use the plural form when displaying the notification
+		Args.Add(TEXT("AssetType"), FText::Format(AssetNameFormat, FText::AsNumber(100)));
+		Args.Add(TEXT("NumJobs"), FText::AsNumber(NumJobs));
+		return FText::Format(LOCTEXT("AsyncCompilationProgress", "Preparing {AssetType} ({NumJobs})"), Args);
+	};
 
 	if (NumJobs == 0)
 	{
@@ -66,7 +51,7 @@ void FAsyncCompilationNotification::Update(int32 NumJobs)
 		StartNumJobs = 0;
 		NotificationHandle = FProgressNotificationHandle();
 	}
-	else 
+	else
 	{
 		if (!NotificationHandle.IsValid())
 		{
@@ -82,8 +67,25 @@ void FAsyncCompilationNotification::Update(int32 NumJobs)
 			FSlateNotificationManager::Get().UpdateProgressNotification(NotificationHandle, StartNumJobs - NumJobs, StartNumJobs, GetProgressMessage());
 		}
 	}
+#endif
 };
 
+#if WITH_EDITOR
+
+namespace AsyncCompilationHelpers
+{
+	void DumpStallStacks();
+}
+
+static FAutoConsoleCommand CVarAsyncAssetDumpStallStacks(
+	TEXT("Editor.AsyncAssetDumpStallStacks"),
+	TEXT("Dump all the callstacks that have caused waits on async compilation."),
+	FConsoleCommandWithArgsDelegate::CreateLambda([](const TArray<FString>& Args)
+	{
+		AsyncCompilationHelpers::DumpStallStacks();
+	})
+);
+
 namespace AsyncCompilationHelpers
 {
 	void FinishCompilation(
@@ -397,6 +399,6 @@ namespace AsyncCompilationHelpers
 	}
 }
 
-#undef LOCTEXT_NAMESPACE
-
 #endif // #if WITH_EDITOR
+
+#undef LOCTEXT_NAMESPACE
\ No newline at end of file
diff --git a/Engine/Source/Runtime/Engine/Private/DistanceFieldAtlas.cpp b/Engine/Source/Runtime/Engine/Private/DistanceFieldAtlas.cpp
--- a/Engine/Source/Runtime/Engine/Private/DistanceFieldAtlas.cpp
+++ b/Engine/Source/Runtime/Engine/Private/DistanceFieldAtlas.cpp
@@ -272,6 +272,7 @@ void FDistanceFieldAsyncQueue::OnAssetPostCompile(const TArray<FAssetCompileData
 }
 
 FDistanceFieldAsyncQueue::FDistanceFieldAsyncQueue() 
+	: Notification(GetAssetNameFormat())
 {
 	FAssetCompilingManager::Get().RegisterManager(this);
 	FAssetCompilingManager::Get().OnAssetPostCompileEvent().AddRaw(this, &FDistanceFieldAsyncQueue::OnAssetPostCompile);
@@ -765,6 +766,8 @@ void FDistanceFieldAsyncQueue::ProcessAsyncTasks(bool bLimitExecutionTime)
 
 		delete Task;
 	}
+
+	Notification.Update(GetNumRemainingAssets());
 #endif
 }
 
diff --git a/Engine/Source/Runtime/Engine/Private/MeshCardRepresentation.cpp b/Engine/Source/Runtime/Engine/Private/MeshCardRepresentation.cpp
--- a/Engine/Source/Runtime/Engine/Private/MeshCardRepresentation.cpp
+++ b/Engine/Source/Runtime/Engine/Private/MeshCardRepresentation.cpp
@@ -168,7 +168,8 @@ static FAutoConsoleVariableRef CVarCardRepresentationAsyncBuildQueue(
 	ECVF_Default | ECVF_ReadOnly
 	);
 
-FCardRepresentationAsyncQueue::FCardRepresentationAsyncQueue() 
+FCardRepresentationAsyncQueue::FCardRepresentationAsyncQueue()
+	: Notification(GetAssetNameFormat())
 {
 #if WITH_EDITOR
 	MeshUtilities = NULL;
@@ -619,6 +620,8 @@ void FCardRepresentationAsyncQueue::ProcessAsyncTasks(bool bLimitExecutionTime)
 
 		delete Task;
 	}
+
+	Notification.Update(GetNumRemainingAssets());
 #endif
 }
 
diff --git a/Engine/Source/Runtime/Engine/Private/ShaderCompiler/ShaderCompiler.cpp b/Engine/Source/Runtime/Engine/Private/ShaderCompiler/ShaderCompiler.cpp
--- a/Engine/Source/Runtime/Engine/Private/ShaderCompiler/ShaderCompiler.cpp
+++ b/Engine/Source/Runtime/Engine/Private/ShaderCompiler/ShaderCompiler.cpp
@@ -2821,7 +2821,8 @@ FShaderCompilingManager::FShaderCompilingManager() :
 #endif
 	SuppressedShaderPlatforms(0),
 	BuildDistributionController(nullptr),
-	bNoShaderCompilation(false)
+	bNoShaderCompilation(false),
+	Notification(GetAssetNameFormat())
 {
 	// don't perform any initialization if compiling is not allowed
 	if (!AllowShaderCompiling())
@@ -4509,6 +4510,8 @@ void FShaderCompilingManager::ProcessAsyncResults(bool bLimitExecutionTime, bool
 	{
 		check(AllJobs.GetNumPendingJobs() == 0);
 	}
+
+	Notification.Update(GetNumRemainingAssets());
 }
 
 bool FShaderCompilingManager::IsShaderCompilerWorkerRunning(FProcHandle & WorkerHandle)
diff --git a/Engine/Source/Runtime/Engine/Public/DistanceFieldAtlas.h b/Engine/Source/Runtime/Engine/Public/DistanceFieldAtlas.h
--- a/Engine/Source/Runtime/Engine/Public/DistanceFieldAtlas.h
+++ b/Engine/Source/Runtime/Engine/Public/DistanceFieldAtlas.h
@@ -14,6 +14,7 @@
 #include "RenderResource.h"
 #include "RenderingThread.h"
 #include "TextureLayout3d.h"
+#include "AsyncCompilationHelpers.h"
 #include "AssetCompilingManager.h"
 #include "Templates/UniquePtr.h"
 #include "DerivedMeshDataTaskUtils.h"
@@ -422,6 +423,8 @@ private:
 	class IMeshUtilities* MeshUtilities;
 
 	mutable FCriticalSection CriticalSection;
+
+	FAsyncCompilationNotification Notification;
 };
 
 /** Global build queue. */
diff --git a/Engine/Source/Runtime/Engine/Public/MeshCardRepresentation.h b/Engine/Source/Runtime/Engine/Public/MeshCardRepresentation.h
--- a/Engine/Source/Runtime/Engine/Public/MeshCardRepresentation.h
+++ b/Engine/Source/Runtime/Engine/Public/MeshCardRepresentation.h
@@ -11,6 +11,7 @@
 #include "ProfilingDebugging/ResourceSize.h"
 #include "Engine/EngineTypes.h"
 #include "UObject/GCObject.h"
+#include "AsyncCompilationHelpers.h"
 #include "AssetCompilingManager.h"
 #include "RenderResource.h"
 #include "RenderingThread.h"
@@ -283,6 +284,8 @@ private:
 	class IMeshUtilities* MeshUtilities;
 
 	mutable FCriticalSection CriticalSection;
+	
+	FAsyncCompilationNotification Notification;
 };
 
 /** Global build queue. */
diff --git a/Engine/Source/Runtime/Engine/Public/ShaderCompiler.h b/Engine/Source/Runtime/Engine/Public/ShaderCompiler.h
--- a/Engine/Source/Runtime/Engine/Public/ShaderCompiler.h
+++ b/Engine/Source/Runtime/Engine/Public/ShaderCompiler.h
@@ -21,6 +21,7 @@
 #include "GBufferInfo.h"
 #include "ShaderMaterial.h"
 #include "Misc/ScopeRWLock.h"
+#include "AsyncCompilationHelpers.h"
 #include "AssetCompilingManager.h"
 #include "Containers/HashTable.h"
 #include "Containers/List.h"
@@ -1118,6 +1119,8 @@ public:
 	 * Returns true if the given shader compile worker is still running.
 	 */
 	static bool IsShaderCompilerWorkerRunning(FProcHandle & WorkerHandle);
+
+	FAsyncCompilationNotification Notification;
 };
 
 /** The global shader compiling thread manager. */
