Fix ImportanceSampling when used with a texture that is still compiling

#rnx

diff --git a/Engine/Source/Runtime/Engine/Private/Kismet/ImportanceSamplingLibrary.cpp b/Engine/Source/Runtime/Engine/Private/Kismet/ImportanceSamplingLibrary.cpp
--- a/Engine/Source/Runtime/Engine/Private/Kismet/ImportanceSamplingLibrary.cpp
+++ b/Engine/Source/Runtime/Engine/Private/Kismet/ImportanceSamplingLibrary.cpp
@@ -4,6 +4,7 @@
 #include "Math/Sobol.h"
 #include "Math/UnrealMathUtility.h"
 #include "Engine/Texture2D.h"
+#include "TextureCompiler.h"
 #include "UObject/Stack.h"
 #include "UObject/NoExportTypes.h"
 
@@ -71,6 +72,15 @@ float FImportanceTexture::ImportanceWeight(FColor Texel, TEnumAsByte<EImportance
 
 void FImportanceTexture::Initialize(UTexture2D *SourceTexture, TEnumAsByte<EImportanceWeight::Type> WeightingFunc)
 {
+#if WITH_EDITOR
+	// For now, this doesn't support being refreshed once the source texture compilation is done.
+	// So force it to finish its compilation if it is not yet done to avoid getting the default's texture data.
+	if (SourceTexture && SourceTexture->IsCompiling())
+	{
+		FTextureCompilingManager::Get().FinishCompilation({ SourceTexture });
+	}
+#endif
+
 	if (!SourceTexture || SourceTexture->GetPixelFormat() != EPixelFormat::PF_B8G8R8A8)
 	{
 		Texture = 0;
