Fix for UTexture resource access during loading screen

- CL 16264522 changed the task scoping checks on UTexture resource access
  - need to account for the Slate threads as well

diff --git a/Engine/Source/Runtime/Engine/Private/Texture.cpp b/Engine/Source/Runtime/Engine/Private/Texture.cpp
--- a/Engine/Source/Runtime/Engine/Private/Texture.cpp
+++ b/Engine/Source/Runtime/Engine/Private/Texture.cpp
@@ -129,7 +129,7 @@ UTexture::UTexture(const FObjectInitializer& ObjectInitializer)
 
 const FTextureResource* UTexture::GetResource() const
 {
-	if (IsInParallelGameThread() || IsInGameThread())
+	if (IsInParallelGameThread() || IsInGameThread() || IsInSlateThread())
 	{
 		return PrivateResource;
 	}
@@ -144,7 +144,7 @@ const FTextureResource* UTexture::GetResource() const
 
 FTextureResource* UTexture::GetResource()
 {
-	if (IsInParallelGameThread() || IsInGameThread())
+	if (IsInParallelGameThread() || IsInGameThread() || IsInSlateThread())
 	{
 		return PrivateResource;
 	}
