32MB memory optimization

- The default bloom texture was always loaded at engine startup.  Access to the texture is hidden behind an accessor on the engine and is loaded on-demand when the bloom feature that requires it is used.
- The default bloom texture can't be streamed in because it only has a single mip so we must do a load.

diff --git a/Engine/Source/Runtime/Engine/Classes/Engine/Engine.h b/Engine/Source/Runtime/Engine/Classes/Engine/Engine.h
--- a/Engine/Source/Runtime/Engine/Classes/Engine/Engine.h
+++ b/Engine/Source/Runtime/Engine/Classes/Engine/Engine.h
@@ -1724,6 +1724,10 @@ public:
 		#endif
 	}
 
+	/** Delay loading this texture until it is needed by the renderer.
+	* The texture is uncompressed and contains no mips so it can't be streamed.
+	*/
+	void LoadDefaultBloomTexture();
 
 private:
 	#if WITH_DYNAMIC_RESOLUTION
diff --git a/Engine/Source/Runtime/Engine/Private/SceneView.cpp b/Engine/Source/Runtime/Engine/Private/SceneView.cpp
--- a/Engine/Source/Runtime/Engine/Private/SceneView.cpp
+++ b/Engine/Source/Runtime/Engine/Private/SceneView.cpp
@@ -2105,6 +2105,12 @@ void FSceneView::EndFinalPostprocessSettings(const FSceneViewInitOptions& ViewIn
 	}
 
 	check(VerifyMembersChecks());
+
+	// If we are using FFT bloom and we don't have a convolution texture load the default one.
+	if (FinalPostProcessSettings.BloomMethod == EBloomMethod::BM_FFT && FinalPostProcessSettings.BloomConvolutionTexture == nullptr)
+	{
+		GEngine->LoadDefaultBloomTexture();
+	}
 }
 
 void FSceneView::ConfigureBufferVisualizationSettings()
diff --git a/Engine/Source/Runtime/Engine/Private/UnrealEngine.cpp b/Engine/Source/Runtime/Engine/Private/UnrealEngine.cpp
--- a/Engine/Source/Runtime/Engine/Private/UnrealEngine.cpp
+++ b/Engine/Source/Runtime/Engine/Private/UnrealEngine.cpp
@@ -2747,7 +2747,6 @@ void UEngine::InitializeObjectReferences()
 	LoadEngineTexture(DefaultDiffuseTexture, *DefaultDiffuseTextureName.ToString());
 	LoadEngineTexture(HighFrequencyNoiseTexture, *HighFrequencyNoiseTextureName.ToString());
 	LoadEngineTexture(DefaultBokehTexture, *DefaultBokehTextureName.ToString());
-	LoadEngineTexture(DefaultBloomKernelTexture, *DefaultBloomKernelTextureName.ToString());
 	LoadEngineTexture(PreIntegratedSkinBRDFTexture, *PreIntegratedSkinBRDFTextureName.ToString());
 	LoadEngineTexture(MiniFontTexture, *MiniFontTextureName.ToString());
 	LoadEngineTexture(WeightMapPlaceholderTexture, *WeightMapPlaceholderTextureName.ToString());
@@ -2892,6 +2891,14 @@ void UEngine::InitializeObjectReferences()
 	UISettings->ForceLoadResources();
 }
 
+void UEngine::LoadDefaultBloomTexture()
+{
+	if (DefaultBloomKernelTexture == nullptr)
+	{
+		LoadEngineTexture(DefaultBloomKernelTexture, *DefaultBloomKernelTextureName.ToString());
+	}
+}
+
 void UEngine::InitializePortalServices()
 {
 	IMessagingRpcModule* MessagingRpcModule = nullptr;
