WorldPartition: DataLayer Multiplayer Support

diff --git a/Engine/Source/Editor/DataLayerEditor/Private/DataLayer/DataLayerHierarchy.cpp b/Engine/Source/Editor/DataLayerEditor/Private/DataLayer/DataLayerHierarchy.cpp
--- a/Engine/Source/Editor/DataLayerEditor/Private/DataLayer/DataLayerHierarchy.cpp
+++ b/Engine/Source/Editor/DataLayerEditor/Private/DataLayer/DataLayerHierarchy.cpp
@@ -8,6 +8,7 @@
 #include "DataLayerMode.h"
 #include "DataLayer/DataLayerEditorSubsystem.h"
 #include "WorldPartition/DataLayer/DataLayer.h"
+#include "WorldPartition/DataLayer/WorldDataLayers.h"
 #include "Engine/World.h"
 #include "EngineUtils.h"
 #include "Editor.h"
diff --git a/Engine/Source/Editor/DataLayerEditor/Private/DataLayer/DataLayerMode.cpp b/Engine/Source/Editor/DataLayerEditor/Private/DataLayer/DataLayerMode.cpp
--- a/Engine/Source/Editor/DataLayerEditor/Private/DataLayer/DataLayerMode.cpp
+++ b/Engine/Source/Editor/DataLayerEditor/Private/DataLayer/DataLayerMode.cpp
@@ -4,6 +4,7 @@
 #include "SSceneOutliner.h"
 #include "SDataLayerBrowser.h"
 #include "DataLayer/DataLayerEditorSubsystem.h"
+#include "WorldPartition/DataLayer/WorldDataLayers.h"
 #include "DataLayerHierarchy.h"
 #include "DataLayerActorTreeItem.h"
 #include "DataLayerTreeItem.h"
diff --git a/Engine/Source/Runtime/Engine/Private/ActorPartition/ActorPartitionSubsystem.cpp b/Engine/Source/Runtime/Engine/Private/ActorPartition/ActorPartitionSubsystem.cpp
--- a/Engine/Source/Runtime/Engine/Private/ActorPartition/ActorPartitionSubsystem.cpp
+++ b/Engine/Source/Runtime/Engine/Private/ActorPartition/ActorPartitionSubsystem.cpp
@@ -6,6 +6,7 @@
 #include "WorldPartition/WorldPartitionHelpers.h"
 #include "WorldPartition/WorldPartitionSubsystem.h"
 #include "WorldPartition/DataLayer/DataLayerSubsystem.h"
+#include "WorldPartition/DataLayer/WorldDataLayers.h"
 #include "Engine/World.h"
 #include "EngineUtils.h"
 
diff --git a/Engine/Source/Runtime/Engine/Private/WorldPartition/DataLayer/DataLayer.cpp b/Engine/Source/Runtime/Engine/Private/WorldPartition/DataLayer/DataLayer.cpp
--- a/Engine/Source/Runtime/Engine/Private/WorldPartition/DataLayer/DataLayer.cpp
+++ b/Engine/Source/Runtime/Engine/Private/WorldPartition/DataLayer/DataLayer.cpp
@@ -5,6 +5,7 @@
 =============================================================================*/
 
 #include "WorldPartition/DataLayer/DataLayer.h"
+#include "WorldPartition/DataLayer/WorldDataLayers.h"
 
 #define LOCTEXT_NAMESPACE "DataLayer"
 
diff --git a/Engine/Source/Runtime/Engine/Private/WorldPartition/DataLayer/DataLayerSubsystem.cpp b/Engine/Source/Runtime/Engine/Private/WorldPartition/DataLayer/DataLayerSubsystem.cpp
--- a/Engine/Source/Runtime/Engine/Private/WorldPartition/DataLayer/DataLayerSubsystem.cpp
+++ b/Engine/Source/Runtime/Engine/Private/WorldPartition/DataLayer/DataLayerSubsystem.cpp
@@ -43,25 +43,18 @@ void UDataLayerSubsystem::Initialize(FSubsystemCollectionBase& Collection)
 #endif
 }
 
-void UDataLayerSubsystem::PostInitialize()
+const TSet<FName>& UDataLayerSubsystem::GetActiveDataLayerNames() const
 {
-	Super::PostInitialize();
+	static TSet<FName> EmptySet;
+	const AWorldDataLayers* WorldDataLayers = GetWorld()->GetWorldDataLayers();
+	return WorldDataLayers ? WorldDataLayers->GetActiveDataLayerNames() : EmptySet;
+}
 
-	if (GetWorld()->IsGameWorld())
-	{
-		// Initialize Dynamically loaded Data Layers state
-		if (const AWorldDataLayers* WorldDataLayers = GetWorld()->GetWorldDataLayers())
-		{
-			WorldDataLayers->ForEachDataLayer([this](class UDataLayer* DataLayer)
-			{
-				if (DataLayer && DataLayer->IsDynamicallyLoaded())
-				{
-					SetDataLayerState(DataLayer, DataLayer->GetInitialState());
-				}
-				return true;
-			});
-		}
-	}
+const TSet<FName>& UDataLayerSubsystem::GetLoadedDataLayerNames() const
+{
+	static TSet<FName> EmptySet;
+	const AWorldDataLayers* WorldDataLayers = GetWorld()->GetWorldDataLayers();
+	return WorldDataLayers ? WorldDataLayers->GetLoadedDataLayerNames() : EmptySet;
 }
 
 UDataLayer* UDataLayerSubsystem::GetDataLayer(const FActorDataLayer& InDataLayer) const
@@ -81,63 +74,11 @@ UDataLayer* UDataLayerSubsystem::GetDataLayerFromName(FName InDataLayerName) con
 	return WorldDataLayers ? const_cast<UDataLayer*>(WorldDataLayers->GetDataLayerFromName(InDataLayerName)) : nullptr;
 }
 
-// Deprecated start
-void UDataLayerSubsystem::ActivateDataLayer(const FActorDataLayer& InDataLayer, bool bInActivate)
-{
-	// Keep deprecated behavior (!Activated == Unloaded)
-	SetDataLayerState(GetDataLayerFromName(InDataLayer.Name), bInActivate ? EDataLayerState::Activated : EDataLayerState::Unloaded);
-}
-
-void UDataLayerSubsystem::ActivateDataLayerByLabel(const FName& InDataLayerLabel, bool bInActivate)
-{
-	// Keep deprecated behavior (!Activated == Unloaded)
-	SetDataLayerState(GetDataLayerFromLabel(InDataLayerLabel), bInActivate ? EDataLayerState::Activated : EDataLayerState::Unloaded);
-}
-
-bool UDataLayerSubsystem::IsDataLayerActive(const FActorDataLayer& InDataLayer) const
-{
-	return GetDataLayerStateByName(InDataLayer.Name) == EDataLayerState::Activated;
-}
-
-bool UDataLayerSubsystem::IsDataLayerActiveByLabel(const FName& InDataLayerLabel) const
-{
-	return GetDataLayerStateByLabel(InDataLayerLabel) == EDataLayerState::Activated;
-}
-// Deprecation end
-
 void UDataLayerSubsystem::SetDataLayerState(const UDataLayer* InDataLayer, EDataLayerState InState)
 {
-	if (!InDataLayer || !InDataLayer->IsDynamicallyLoaded())
-	{
-		return;
-	}
-
-	FName DataLayerName = InDataLayer->GetFName();
-	EDataLayerState CurrentState = GetDataLayerStateByName(DataLayerName);
-	if (CurrentState != InState)
+	if (AWorldDataLayers* WorldDataLayers = GetWorld()->GetWorldDataLayers())
 	{
-		LoadedDataLayerNames.Remove(DataLayerName);
-		ActiveDataLayerNames.Remove(DataLayerName);
-		
-		if (InState == EDataLayerState::Loaded)
-		{
-			LoadedDataLayerNames.Add(DataLayerName);
-		}
-		else if(InState == EDataLayerState::Activated)
-		{
-			ActiveDataLayerNames.Add(DataLayerName);
-		}
-		
-		// todo_ow: remove
-		if (OnDataLayerActivationStateChanged.IsBound())
-		{
-			OnDataLayerActivationStateChanged.Broadcast(InDataLayer, InState == EDataLayerState::Activated);
-		}
-
-		if (OnDataLayerStateChanged.IsBound())
-		{
-			OnDataLayerStateChanged.Broadcast(InDataLayer, InState);
-		}
+		WorldDataLayers->SetDataLayerState(FActorDataLayer(InDataLayer->GetFName()), InState);
 	}
 }
 
@@ -168,18 +109,13 @@ EDataLayerState UDataLayerSubsystem::GetDataLayerState(const UDataLayer* InDataL
 
 EDataLayerState UDataLayerSubsystem::GetDataLayerStateByName(const FName& InDataLayerName) const
 {
-	if (ActiveDataLayerNames.Contains(InDataLayerName))
-	{
-		check(!LoadedDataLayerNames.Contains(InDataLayerName));
-		return EDataLayerState::Activated;
-	}
-	else if (LoadedDataLayerNames.Contains(InDataLayerName))
+	const AWorldDataLayers* WorldDataLayers = GetWorld()->GetWorldDataLayers();
+	if (!WorldDataLayers)
 	{
-		check(!ActiveDataLayerNames.Contains(InDataLayerName));
-		return EDataLayerState::Loaded;
+		return EDataLayerState::Unloaded;
 	}
 
-	return EDataLayerState::Unloaded;
+	return WorldDataLayers->GetDataLayerStateByName(InDataLayerName);
 }
 
 EDataLayerState UDataLayerSubsystem::GetDataLayerState(const FActorDataLayer& InDataLayer) const
@@ -234,8 +170,8 @@ void UDataLayerSubsystem::DrawDataLayersStatus(UCanvas* Canvas, FVector2D& Offse
 		}
 	};
 
-	DrawLayerNames(TEXT("Loaded Data Layers"), FColor::Cyan, LoadedDataLayerNames);
-	DrawLayerNames(TEXT("Active Data Layers"), FColor::Green, ActiveDataLayerNames);
+	DrawLayerNames(TEXT("Loaded Data Layers"), FColor::Cyan, GetLoadedDataLayerNames());
+	DrawLayerNames(TEXT("Active Data Layers"), FColor::Green, GetActiveDataLayerNames());
 
 	Offset.X += MaxTextWidth + 10;
 }
diff --git a/Engine/Source/Runtime/Engine/Private/WorldPartition/DataLayer/WorldDataLayers.cpp b/Engine/Source/Runtime/Engine/Private/WorldPartition/DataLayer/WorldDataLayers.cpp
--- a/Engine/Source/Runtime/Engine/Private/WorldPartition/DataLayer/WorldDataLayers.cpp
+++ b/Engine/Source/Runtime/Engine/Private/WorldPartition/DataLayer/WorldDataLayers.cpp
@@ -6,16 +6,122 @@
 
 #include "WorldPartition/DataLayer/WorldDataLayers.h"
 #include "WorldPartition/DataLayer/DataLayer.h"
+#include "WorldPartition/DataLayer/DataLayerSubsystem.h"
 #include "EngineUtils.h"
+#include "Net/UnrealNetwork.h"
 #if WITH_EDITOR
 #include "WorldPartition/WorldPartitionEditorPerProjectUserSettings.h"
 #endif
 
 #define LOCTEXT_NAMESPACE "WorldDataLayers"
 
+int32 AWorldDataLayers::DataLayersStateEpoch = 0;
+
 AWorldDataLayers::AWorldDataLayers(const FObjectInitializer& ObjectInitializer)
 	: Super(ObjectInitializer)
 {
+	bAlwaysRelevant = true;
+	bReplicates = true;
+}
+
+void AWorldDataLayers::GetLifetimeReplicatedProps(TArray< FLifetimeProperty >& OutLifetimeProps) const
+{
+	Super::GetLifetimeReplicatedProps(OutLifetimeProps);
+
+	DOREPLIFETIME(AWorldDataLayers, RepLoadedDataLayerNames);
+	DOREPLIFETIME(AWorldDataLayers, RepActiveDataLayerNames);
+}
+
+void AWorldDataLayers::InitializeDataLayerStates()
+{
+	if (GetWorld()->IsGameWorld())
+	{
+		ForEachDataLayer([this](class UDataLayer* DataLayer)
+		{
+			if (DataLayer && DataLayer->IsDynamicallyLoaded())
+			{
+				if (DataLayer->GetInitialState() == EDataLayerState::Activated)
+				{
+					ActiveDataLayerNames.Add(DataLayer->GetFName());
+				}
+				else if (DataLayer->GetInitialState() == EDataLayerState::Loaded)
+				{
+					LoadedDataLayerNames.Add(DataLayer->GetFName());
+				}
+			}
+			return true;
+		});
+
+		RepActiveDataLayerNames = ActiveDataLayerNames.Array();
+		RepLoadedDataLayerNames = LoadedDataLayerNames.Array();
+	}
+}
+
+void AWorldDataLayers::SetDataLayerState_Implementation(FActorDataLayer InDataLayer, EDataLayerState InState)
+{
+	const UDataLayer* DataLayer = GetDataLayerFromName(InDataLayer.Name);
+	if (!DataLayer || !DataLayer->IsDynamicallyLoaded())
+	{
+		return;
+	}
+
+	EDataLayerState CurrentState = GetDataLayerStateByName(InDataLayer.Name);
+	if (CurrentState != InState)
+	{
+		LoadedDataLayerNames.Remove(InDataLayer.Name);
+		ActiveDataLayerNames.Remove(InDataLayer.Name);
+
+		if (InState == EDataLayerState::Loaded)
+		{
+			LoadedDataLayerNames.Add(InDataLayer.Name);
+		}
+		else if (InState == EDataLayerState::Activated)
+		{
+			ActiveDataLayerNames.Add(InDataLayer.Name);
+		}
+			
+		// Update Replicated Properties
+		RepActiveDataLayerNames = ActiveDataLayerNames.Array();
+		RepLoadedDataLayerNames = LoadedDataLayerNames.Array();
+
+		++DataLayersStateEpoch;
+
+		OnDataLayerStateChanged(DataLayer, InState);
+	}
+}
+
+void AWorldDataLayers::OnDataLayerStateChanged_Implementation(const UDataLayer* InDataLayer, EDataLayerState InState)
+{
+	UDataLayerSubsystem* DataLayerSubsystem = GetWorld()->GetSubsystem<UDataLayerSubsystem>();
+	DataLayerSubsystem->OnDataLayerStateChanged.Broadcast(InDataLayer, InState);
+}
+
+void AWorldDataLayers::OnRep_ActiveDataLayerNames()
+{
+	ActiveDataLayerNames.Reset();
+	ActiveDataLayerNames.Append(RepActiveDataLayerNames);
+}
+
+void AWorldDataLayers::OnRep_LoadedDataLayerNames()
+{
+	LoadedDataLayerNames.Reset();
+	LoadedDataLayerNames.Append(RepLoadedDataLayerNames);
+}
+
+EDataLayerState AWorldDataLayers::GetDataLayerStateByName(FName InDataLayerName) const
+{
+	if (ActiveDataLayerNames.Contains(InDataLayerName))
+	{
+		check(!LoadedDataLayerNames.Contains(InDataLayerName));
+		return EDataLayerState::Activated;
+	}
+	else if (LoadedDataLayerNames.Contains(InDataLayerName))
+	{
+		check(!ActiveDataLayerNames.Contains(InDataLayerName));
+		return EDataLayerState::Loaded;
+	}
+
+	return EDataLayerState::Unloaded;
 }
 
 #if WITH_EDITOR
@@ -233,6 +339,8 @@ void AWorldDataLayers::PostLoad()
 		NameToDataLayer.Add(DataLayer->GetFName(), DataLayer);
 	}
 #endif
+
+	InitializeDataLayerStates();
 }
 
 #undef LOCTEXT_NAMESPACE
\ No newline at end of file
diff --git a/Engine/Source/Runtime/Engine/Private/WorldPartition/HLOD/HLODSubsystem.cpp b/Engine/Source/Runtime/Engine/Private/WorldPartition/HLOD/HLODSubsystem.cpp
--- a/Engine/Source/Runtime/Engine/Private/WorldPartition/HLOD/HLODSubsystem.cpp
+++ b/Engine/Source/Runtime/Engine/Private/WorldPartition/HLOD/HLODSubsystem.cpp
@@ -94,7 +94,7 @@ void UHLODSubsystem::OnWorldPartitionRegistered(UWorldPartition* InWorldPartitio
 	check(CellsHLODMapping.IsEmpty());
 
 	TSet<const UWorldPartitionRuntimeCell*> StreamingCells;
-	InWorldPartition->RuntimeHash->GetAllStreamingCells(StreamingCells, /*bIncludeDataLayers*/ true);
+	InWorldPartition->RuntimeHash->GetAllStreamingCells(StreamingCells, /*bAllDataLayers*/ true);
 
 	// Build cell to HLOD mapping
 	for (const UWorldPartitionRuntimeCell* Cell : StreamingCells)
diff --git a/Engine/Source/Runtime/Engine/Private/WorldPartition/WorldPartitionActorCluster.cpp b/Engine/Source/Runtime/Engine/Private/WorldPartition/WorldPartitionActorCluster.cpp
--- a/Engine/Source/Runtime/Engine/Private/WorldPartition/WorldPartitionActorCluster.cpp
+++ b/Engine/Source/Runtime/Engine/Private/WorldPartition/WorldPartitionActorCluster.cpp
@@ -10,6 +10,7 @@
 #include "WorldPartition/ActorDescContainer.h"
 #include "WorldPartition/WorldPartitionActorDescView.h"
 #include "WorldPartition/WorldPartitionRuntimeHash.h"
+#include "WorldPartition/DataLayer/WorldDataLayers.h"
 #include "Hash/CityHash.h"
 
 template<class LayerNameContainer>
diff --git a/Engine/Source/Runtime/Engine/Private/WorldPartition/WorldPartitionLevelStreamingPolicy.cpp b/Engine/Source/Runtime/Engine/Private/WorldPartition/WorldPartitionLevelStreamingPolicy.cpp
--- a/Engine/Source/Runtime/Engine/Private/WorldPartition/WorldPartitionLevelStreamingPolicy.cpp
+++ b/Engine/Source/Runtime/Engine/Private/WorldPartition/WorldPartitionLevelStreamingPolicy.cpp
@@ -72,7 +72,7 @@ TSubclassOf<UWorldPartitionRuntimeCell> UWorldPartitionLevelStreamingPolicy::Get
 void UWorldPartitionLevelStreamingPolicy::PrepareActorToCellRemapping()
 {
 	TSet<const UWorldPartitionRuntimeCell*> StreamingCells;
-	WorldPartition->RuntimeHash->GetAllStreamingCells(StreamingCells, /*bIncludeDataLayers*/ true);
+	WorldPartition->RuntimeHash->GetAllStreamingCells(StreamingCells, /*bAllDataLayers*/ true);
 
 	// Build Actor-to-Cell remapping
 	for (const UWorldPartitionRuntimeCell* Cell : StreamingCells)
diff --git a/Engine/Source/Runtime/Engine/Private/WorldPartition/WorldPartitionRuntimeHash.cpp b/Engine/Source/Runtime/Engine/Private/WorldPartition/WorldPartitionRuntimeHash.cpp
--- a/Engine/Source/Runtime/Engine/Private/WorldPartition/WorldPartitionRuntimeHash.cpp
+++ b/Engine/Source/Runtime/Engine/Private/WorldPartition/WorldPartitionRuntimeHash.cpp
@@ -9,6 +9,7 @@
 
 #if WITH_EDITOR
 #include "WorldPartition/WorldPartitionHandle.h"
+#include "WorldPartition/DataLayer/WorldDataLayers.h"
 #include "Logging/MessageLog.h"
 #include "Misc/UObjectToken.h"
 #include "Misc/MapErrors.h"
diff --git a/Engine/Source/Runtime/Engine/Private/WorldPartition/WorldPartitionRuntimeSpatialHash.cpp b/Engine/Source/Runtime/Engine/Private/WorldPartition/WorldPartitionRuntimeSpatialHash.cpp
--- a/Engine/Source/Runtime/Engine/Private/WorldPartition/WorldPartitionRuntimeSpatialHash.cpp
+++ b/Engine/Source/Runtime/Engine/Private/WorldPartition/WorldPartitionRuntimeSpatialHash.cpp
@@ -1014,7 +1014,7 @@ FAutoConsoleCommand UWorldPartitionRuntimeSpatialHash::OverrideLoadingRangeComma
 );
 
 // Streaming interface
-int32 UWorldPartitionRuntimeSpatialHash::GetAllStreamingCells(TSet<const UWorldPartitionRuntimeCell*>& Cells, bool bIncludeDataLayers /*=false*/) const
+int32 UWorldPartitionRuntimeSpatialHash::GetAllStreamingCells(TSet<const UWorldPartitionRuntimeCell*>& Cells, bool bAllDataLayers, bool bDataLayersOnly, const TSet<FName>& InDataLayers) const
 {
 	for (const FSpatialHashStreamingGrid& StreamingGrid : StreamingGrids)
 	{
@@ -1024,7 +1024,11 @@ int32 UWorldPartitionRuntimeSpatialHash::GetAllStreamingCells(TSet<const UWorldP
 			{
 				for (const UWorldPartitionRuntimeSpatialHashCell* Cell : LayerCell.GridCells)
 				{
-					if (bIncludeDataLayers || !Cell->HasDataLayers())
+					if (!bDataLayersOnly && !Cell->HasDataLayers())
+					{
+						Cells.Add(Cell);
+					}
+					else if (Cell->HasDataLayers() && (bAllDataLayers || Cell->HasAnyDataLayer(InDataLayers)))
 					{
 						Cells.Add(Cell);
 					}
diff --git a/Engine/Source/Runtime/Engine/Private/WorldPartition/WorldPartitionStreamingPolicy.cpp b/Engine/Source/Runtime/Engine/Private/WorldPartition/WorldPartitionStreamingPolicy.cpp
--- a/Engine/Source/Runtime/Engine/Private/WorldPartition/WorldPartitionStreamingPolicy.cpp
+++ b/Engine/Source/Runtime/Engine/Private/WorldPartition/WorldPartitionStreamingPolicy.cpp
@@ -9,6 +9,7 @@
 #include "WorldPartition/WorldPartitionStreamingSource.h"
 #include "WorldPartition/WorldPartition.h"
 #include "WorldPartition/DataLayer/DataLayerSubsystem.h"
+#include "WorldPartition/DataLayer/WorldDataLayers.h"
 #include "WorldPartition/HLOD/HLODActor.h"
 #include "WorldPartition/HLOD/HLODSubsystem.h"
 #include "GameFramework/PlayerController.h"
@@ -34,9 +35,9 @@ static FAutoConsoleVariableRef CMaxLoadingStreamingCells(
 
 UWorldPartitionStreamingPolicy::UWorldPartitionStreamingPolicy(const FObjectInitializer& ObjectInitializer)
 : Super(ObjectInitializer) 
-, bIsServerLoadingDone(false)
 , UpdateStreamingStateEpoch(0)
 , SortedAddToWorldCellsEpoch(INT_MIN)
+, DataLayersStatesServerEpoch(INT_MIN)
 {
 	if (!IsTemplate())
 	{
@@ -74,7 +75,8 @@ void UWorldPartitionStreamingPolicy::UpdateStreamingSources()
 #endif
 	{
 		UWorld* World = WorldPartition->GetWorld();
-		if (World->GetNetMode() != NM_DedicatedServer)
+		const ENetMode NetMode = World->GetNetMode();
+		if(NetMode == NM_Standalone || NetMode == NM_Client)
 		{
 			const int32 NumPlayers = GEngine->GetNumGamePlayers(World);
 			for (int32 PlayerIndex = 0; PlayerIndex < NumPlayers; ++PlayerIndex)
@@ -131,40 +133,21 @@ void UWorldPartitionStreamingPolicy::UpdateStreamingState()
 	
 	// Update streaming sources
 	UpdateStreamingSources();
-
-	if (World->GetNetMode() == NM_DedicatedServer)
-	{
-		// Load all cells on server (do this once)
-		if (!bIsServerLoadingDone)
-		{
-			TSet<const UWorldPartitionRuntimeCell*> AllStreamingCells;
-			WorldPartition->RuntimeHash->GetAllStreamingCells(AllStreamingCells);
-			TSet<const UWorldPartitionRuntimeCell*> ToActivateCells = AllStreamingCells.Difference(ActivatedCells);
-
-			// Mark Runtime Cells to be Always Loaded
-			for (const UWorldPartitionRuntimeCell* Cell : ToActivateCells)
-			{
-				UWorldPartitionRuntimeCell* MutableCell = const_cast<UWorldPartitionRuntimeCell*>(Cell);
-				MutableCell->SetIsAlwaysLoaded(true);
-			}
-
-			SetTargetStateForCells(EWorldPartitionRuntimeCellState::Activated, ToActivateCells);
-			bIsServerLoadingDone = true;
-		}
-	}
-	else
+		
+	TSet<const UWorldPartitionRuntimeCell*>& ActivateStreamingCells = FrameActivateCells.GetCells();
+	TSet<const UWorldPartitionRuntimeCell*>& LoadStreamingCells = FrameLoadCells.GetCells();
+	check(ActivateStreamingCells.IsEmpty());
+	check(LoadStreamingCells.IsEmpty());
+
+	const ENetMode NetMode = World->GetNetMode();
+	if(NetMode == NM_Standalone || NetMode == NM_Client)
 	{
 		// Early out if nothing loaded and no streaming source
 		if ((StreamingSources.Num() == 0) && (ActivatedCells.Num() == 0) && (LoadedCells.Num() == 0))
 		{
 			return;
 		}
-
-		TSet<const UWorldPartitionRuntimeCell*>& ActivateStreamingCells = FrameActivateCells.GetCells();
-		TSet<const UWorldPartitionRuntimeCell*>& LoadStreamingCells = FrameLoadCells.GetCells();
-		check(ActivateStreamingCells.IsEmpty());
-		check(LoadStreamingCells.IsEmpty());
-
+				
 		// When uninitializing, UpdateStreamingState is called, but we don't want any cells to be loaded
 		if (WorldPartition->IsInitialized())
 		{
@@ -191,49 +174,71 @@ void UWorldPartitionStreamingPolicy::UpdateStreamingState()
 			RemoveHLODCells(ActivateStreamingCells);
 			RemoveHLODCells(LoadStreamingCells);
 		}
+	} 
+	else 
+	{
+		// Server will activate all non data layer cells at first and then load/activate/unload data layer cells only when the data layer states change
+		if (DataLayersStatesServerEpoch == AWorldDataLayers::GetDataLayersStateEpoch())
+		{
+			// Server as nothing to do early out
+			return; 
+		}
 
-		// Determine cells to load/unload
-		TSet<const UWorldPartitionRuntimeCell*> ToActivateCells = ActivateStreamingCells.Difference(ActivatedCells);
-		TSet<const UWorldPartitionRuntimeCell*> ToLoadCells = LoadStreamingCells.Difference(LoadedCells);
-		TSet<const UWorldPartitionRuntimeCell*> ToUnloadCells = ActivatedCells.Union(LoadedCells).Difference(ActivateStreamingCells.Union(LoadStreamingCells));
+		const UDataLayerSubsystem* DataLayerSubsystem = WorldPartition->GetWorld()->GetSubsystem<UDataLayerSubsystem>();
+		DataLayersStatesServerEpoch = AWorldDataLayers::GetDataLayersStateEpoch();
 
-		UE_SUPPRESS(LogWorldPartition, Verbose,
-		{
-			if (ToActivateCells.Num() > 0 || ToUnloadCells.Num() > 0)
-			{
-				UE_LOG(LogWorldPartition, Verbose, TEXT("UWorldPartitionStreamingPolicy: CellsToLoad(%d), CellsToUnload(%d)"), ToActivateCells.Num(), ToUnloadCells.Num());
-				
-				FTransform LocalToWorld = WorldPartition->GetInstanceTransform();
-				for (int i = 0; i < StreamingSources.Num(); ++i)
-				{
-					FVector ViewLocation = LocalToWorld.TransformPosition(StreamingSources[i].Location);
-					FRotator ViewRotation = LocalToWorld.TransformRotation(StreamingSources[i].Rotation.Quaternion()).Rotator();
-					UE_LOG(LogWorldPartition, Verbose, TEXT("UWorldPartitionStreamingPolicy: Sources[%d] = %s,%s"), i, *ViewLocation.ToString(), *ViewRotation.ToString());
-				}
-			}
-		});
+		const AWorldDataLayers* WorldDataLayers = GetWorld()->GetWorldDataLayers();
 
-		// Process unload first, so that LoadedCells is up-to-date when calling LoadCells
-		if (ToUnloadCells.Num() > 0)
-		{
-			SetTargetStateForCells(EWorldPartitionRuntimeCellState::Unloaded, ToUnloadCells);
-		}
+		// Non Data Layer Cells + Active Data Layers
+		WorldPartition->RuntimeHash->GetAllStreamingCells(ActivateStreamingCells, /*bAllDataLayers=*/ false, /*bDataLayersOnly=*/ false, DataLayerSubsystem->GetActiveDataLayerNames());
 
-		// Do Activation State first as it is higher prio than Load State (if we have a limited number of loading cells per frame)
-		if (ToActivateCells.Num() > 0)
+		// Loaded Data Layers Cells only
+		if (DataLayerSubsystem->GetLoadedDataLayerNames().Num())
 		{
-			SetTargetStateForCells(EWorldPartitionRuntimeCellState::Activated, ToActivateCells);
+			WorldPartition->RuntimeHash->GetAllStreamingCells(LoadStreamingCells, /*bAllDataLayers=*/ false, /*bDataLayersOnly=*/ true, DataLayerSubsystem->GetLoadedDataLayerNames());
 		}
+	}
+
+	// Determine cells to load/unload
+	TSet<const UWorldPartitionRuntimeCell*> ToActivateCells = ActivateStreamingCells.Difference(ActivatedCells);
+	TSet<const UWorldPartitionRuntimeCell*> ToLoadCells = LoadStreamingCells.Difference(LoadedCells);
+	TSet<const UWorldPartitionRuntimeCell*> ToUnloadCells = ActivatedCells.Union(LoadedCells).Difference(ActivateStreamingCells.Union(LoadStreamingCells));
 
-		if (ToLoadCells.Num() > 0)
+	UE_SUPPRESS(LogWorldPartition, Verbose,
+	{
+		if (ToActivateCells.Num() > 0 || ToLoadCells.Num() > 0 || ToUnloadCells.Num() > 0)
 		{
-			SetTargetStateForCells(EWorldPartitionRuntimeCellState::Loaded, ToLoadCells);
+			UE_LOG(LogWorldPartition, Verbose, TEXT("UWorldPartitionStreamingPolicy: CellsToActivate(%d), CellsToLoad(%d), CellsToUnload(%d)"), ToActivateCells.Num(), ToLoadCells.Num(), ToUnloadCells.Num());
+
+			FTransform LocalToWorld = WorldPartition->GetInstanceTransform();
+			for (int i = 0; i < StreamingSources.Num(); ++i)
+			{
+				FVector ViewLocation = LocalToWorld.TransformPosition(StreamingSources[i].Location);
+				FRotator ViewRotation = LocalToWorld.TransformRotation(StreamingSources[i].Rotation.Quaternion()).Rotator();
+				UE_LOG(LogWorldPartition, Verbose, TEXT("UWorldPartitionStreamingPolicy: Sources[%d] = %s,%s"), i, *ViewLocation.ToString(), *ViewRotation.ToString());
+			}
 		}
+	});
 
-		// Reset frame StreamingSourceCells (optimization to avoid reallocation at every call to UpdateStreamingState)
-		FrameActivateCells.Reset();
-		FrameLoadCells.Reset();
+	if (ToUnloadCells.Num() > 0)
+	{
+		SetTargetStateForCells(EWorldPartitionRuntimeCellState::Unloaded, ToUnloadCells);
+	}
+
+	// Do Activation State first as it is higher prio than Load State (if we have a limited number of loading cells per frame)
+	if (ToActivateCells.Num() > 0)
+	{
+		SetTargetStateForCells(EWorldPartitionRuntimeCellState::Activated, ToActivateCells);
 	}
+
+	if (ToLoadCells.Num() > 0)
+	{
+		SetTargetStateForCells(EWorldPartitionRuntimeCellState::Loaded, ToLoadCells);
+	}
+
+	// Reset frame StreamingSourceCells (optimization to avoid reallocation at every call to UpdateStreamingState)
+	FrameActivateCells.Reset();
+	FrameLoadCells.Reset();
 }
 
 int32 UWorldPartitionStreamingPolicy::GetMaxCellsToLoad() const
diff --git a/Engine/Source/Runtime/Engine/Public/WorldPartition/DataLayer/DataLayer.h b/Engine/Source/Runtime/Engine/Public/WorldPartition/DataLayer/DataLayer.h
--- a/Engine/Source/Runtime/Engine/Public/WorldPartition/DataLayer/DataLayer.h
+++ b/Engine/Source/Runtime/Engine/Public/WorldPartition/DataLayer/DataLayer.h
@@ -6,7 +6,6 @@
 #include "UObject/ObjectMacros.h"
 #include "UObject/UObjectGlobals.h"
 #include "UObject/Object.h"
-#include "WorldDataLayers.h"
 #include "ActorDataLayer.h"
 #include "Math/Color.h"
 
diff --git a/Engine/Source/Runtime/Engine/Public/WorldPartition/DataLayer/DataLayerSubsystem.h b/Engine/Source/Runtime/Engine/Public/WorldPartition/DataLayer/DataLayerSubsystem.h
--- a/Engine/Source/Runtime/Engine/Public/WorldPartition/DataLayer/DataLayerSubsystem.h
+++ b/Engine/Source/Runtime/Engine/Public/WorldPartition/DataLayer/DataLayerSubsystem.h
@@ -14,9 +14,6 @@
 
 class UCanvas;
 
-// Deprecate
-DECLARE_DYNAMIC_MULTICAST_DELEGATE_TwoParams(FOnDataLayerActivationStateChanged, const UDataLayer*, DataLayer, bool, bIsActive);
-
 DECLARE_DYNAMIC_MULTICAST_DELEGATE_TwoParams(FOnDataLayerStateChanged, const UDataLayer*, DataLayer, EDataLayerState, State);
 
 UCLASS()
@@ -32,23 +29,7 @@ public:
 	virtual void Initialize(FSubsystemCollectionBase& Collection) override;
 	//~ End USubsystem Interface.
 
-	//~ Begin UWorldSubsystem Interface.
-	virtual void PostInitialize() override;
-	//~ End UWorldSubsystem Interface.
-
 	//~ Begin Blueprint callable functions
-	UFUNCTION(BlueprintCallable, Category = DataLayers, meta=(DeprecatedFunction, DeprecationMessage="Use SetDataLayerState instead"))
-	void ActivateDataLayer(const FActorDataLayer& InDataLayer, bool bInActivate);
-	
-	UFUNCTION(BlueprintCallable, Category = DataLayers, meta = (DeprecatedFunction, DeprecationMessage = "Use SetDataLayerStateByLabel instead"))
-	void ActivateDataLayerByLabel(const FName& InDataLayerLabel, bool bInActivate);
-
-	UFUNCTION(BlueprintCallable, Category = DataLayers, meta = (DeprecatedFunction, DeprecationMessage = "Use GetDataLayerState instead"))
-	bool IsDataLayerActive(const FActorDataLayer& InDataLayer) const;
-
-	UFUNCTION(BlueprintCallable, Category = DataLayers, meta = (DeprecatedFunction, DeprecationMessage = "Use GetDataLayerStateByLabel instead"))
-	bool IsDataLayerActiveByLabel(const FName& InDataLayerLabel) const;
-
 	UFUNCTION(BlueprintCallable, Category = DataLayers)
 	void SetDataLayerState(const FActorDataLayer& InDataLayer, EDataLayerState InState);
 
@@ -64,9 +45,6 @@ public:
 	UFUNCTION(BlueprintCallable, Category = DataLayers)
 	EDataLayerState GetDataLayerStateByLabel(const FName& InDataLayerLabel) const;
 
-	UPROPERTY(BlueprintAssignable)
-	FOnDataLayerActivationStateChanged OnDataLayerActivationStateChanged;
-
 	UPROPERTY(BlueprintAssignable)
 	FOnDataLayerStateChanged OnDataLayerStateChanged;
 
@@ -78,10 +56,10 @@ public:
 	//~ End Blueprint callable functions
 
 	UFUNCTION(BlueprintCallable, Category = DataLayers)
-	const TSet<FName>& GetActiveDataLayerNames() const { return ActiveDataLayerNames; }
+	const TSet<FName>& GetActiveDataLayerNames() const;
 
 	UFUNCTION(BlueprintCallable, Category = DataLayers)
-	const TSet<FName>& GetLoadedDataLayerNames() const { return LoadedDataLayerNames; }
+	const TSet<FName>& GetLoadedDataLayerNames() const;
 
 	void SetDataLayerState(const UDataLayer* InDataLayer, EDataLayerState InState);
 	void SetDataLayerStateByName(const FName& InDataLayerName, EDataLayerState InState);
@@ -108,7 +86,4 @@ private:
 
 	/** Console command used to toggle activation of a DataLayer */
 	static class FAutoConsoleCommand ToggleDataLayerActivation;
-
-	TSet<FName> ActiveDataLayerNames;
-	TSet<FName> LoadedDataLayerNames;
 };
\ No newline at end of file
diff --git a/Engine/Source/Runtime/Engine/Public/WorldPartition/DataLayer/WorldDataLayers.h b/Engine/Source/Runtime/Engine/Public/WorldPartition/DataLayer/WorldDataLayers.h
--- a/Engine/Source/Runtime/Engine/Public/WorldPartition/DataLayer/WorldDataLayers.h
+++ b/Engine/Source/Runtime/Engine/Public/WorldPartition/DataLayer/WorldDataLayers.h
@@ -5,6 +5,7 @@
 #include "CoreMinimal.h"
 #include "UObject/ObjectMacros.h"
 #include "GameFramework/Info.h"
+#include "WorldPartition/DataLayer/ActorDataLayer.h"
 #include "WorldDataLayers.generated.h"
 
 class UDataLayer;
@@ -39,6 +40,28 @@ public:
 	void ForEachDataLayer(TFunctionRef<bool(class UDataLayer*)> Func);
 	void ForEachDataLayer(TFunctionRef<bool(class UDataLayer*)> Func) const;
 
+	UFUNCTION(Server, Reliable)
+	void SetDataLayerState(FActorDataLayer InDataLayer, EDataLayerState InState);
+
+	UFUNCTION(NetMulticast, Reliable)
+	void OnDataLayerStateChanged(const UDataLayer* InDataLayer, EDataLayerState InState);
+
+	EDataLayerState GetDataLayerStateByName(FName InDataLAyerName) const;
+
+	const TSet<FName>& GetActiveDataLayerNames() const { return ActiveDataLayerNames; }
+	const TSet<FName>& GetLoadedDataLayerNames() const { return LoadedDataLayerNames; }
+
+	static int32 GetDataLayersStateEpoch() { return DataLayersStateEpoch; }
+
+protected:
+	void InitializeDataLayerStates();
+
+	UFUNCTION()
+	void OnRep_ActiveDataLayerNames();
+
+	UFUNCTION()
+	void OnRep_LoadedDataLayerNames();
+
 private:
 #if !WITH_EDITOR
 	TMap<FName, const UDataLayer*> LabelToDataLayer;
@@ -47,4 +70,16 @@ private:
 
 	UPROPERTY()
 	TSet<TObjectPtr<UDataLayer>> WorldDataLayers;
+
+	UPROPERTY(Transient, Replicated, ReplicatedUsing=OnRep_ActiveDataLayerNames)
+	TArray<FName> RepActiveDataLayerNames;
+		
+	UPROPERTY(Transient, Replicated, ReplicatedUsing=OnRep_LoadedDataLayerNames)
+	TArray<FName> RepLoadedDataLayerNames;
+
+	// TSet do not support replication so we replicate an array and update the set in the OnRep_ActiveDataLayerNames/OnRep_LoadedDataLayerNames
+	TSet<FName> ActiveDataLayerNames;
+	TSet<FName> LoadedDataLayerNames;
+
+	static int32 DataLayersStateEpoch;
 };
\ No newline at end of file
diff --git a/Engine/Source/Runtime/Engine/Public/WorldPartition/WorldPartitionRuntimeCell.h b/Engine/Source/Runtime/Engine/Public/WorldPartition/WorldPartitionRuntimeCell.h
--- a/Engine/Source/Runtime/Engine/Public/WorldPartition/WorldPartitionRuntimeCell.h
+++ b/Engine/Source/Runtime/Engine/Public/WorldPartition/WorldPartitionRuntimeCell.h
@@ -6,6 +6,7 @@
 #include "WorldPartition/DataLayer/DataLayer.h"
 #include "WorldPartition/WorldPartitionActorDescView.h"
 #include "ProfilingDebugging/ProfilingHelpers.h"
+#include "Algo/AnyOf.h"
 #include "WorldPartitionRuntimeCell.generated.h"
 
 USTRUCT()
@@ -136,6 +137,10 @@ class UWorldPartitionRuntimeCell : public UObject
 
 	bool HasDataLayers() const { return !DataLayers.IsEmpty(); }
 	const TArray<FName>& GetDataLayers() const { return DataLayers; }
+	bool HasAnyDataLayer(const TSet<FName>& InDataLayers) const
+	{
+		return Algo::AnyOf(DataLayers, [&InDataLayers](const FName& DataLayer) { return InDataLayers.Contains(DataLayer); });
+	}
 
 #if WITH_EDITOR
 	void SetDataLayers(const TArray<const UDataLayer*>& InDataLayers);
diff --git a/Engine/Source/Runtime/Engine/Public/WorldPartition/WorldPartitionRuntimeHash.h b/Engine/Source/Runtime/Engine/Public/WorldPartition/WorldPartitionRuntimeHash.h
--- a/Engine/Source/Runtime/Engine/Public/WorldPartition/WorldPartitionRuntimeHash.h
+++ b/Engine/Source/Runtime/Engine/Public/WorldPartition/WorldPartitionRuntimeHash.h
@@ -56,7 +56,7 @@ class ENGINE_API UWorldPartitionRuntimeHash : public UObject
 	};
 
 	// Streaming interface
-	virtual int32 GetAllStreamingCells(TSet<const UWorldPartitionRuntimeCell*>& Cells, bool bIncludeDataLayers = false) const { return 0; }
+	virtual int32 GetAllStreamingCells(TSet<const UWorldPartitionRuntimeCell*>& Cells, bool bAllDataLayers = false, bool bDataLayersOnly = false, const TSet<FName>& InDataLayers = TSet<FName>()) const { return 0; }
 	virtual bool GetStreamingCells(const FWorldPartitionStreamingQuerySource& QuerySource, TSet<const UWorldPartitionRuntimeCell*>& OutCells) const { return false; }
 	virtual bool GetStreamingCells(const TArray<FWorldPartitionStreamingSource>& Sources, UWorldPartitionRuntimeHash::FStreamingSourceCells& OutActivateCells, UWorldPartitionRuntimeHash::FStreamingSourceCells& OutLoadCells) const { return false; };
 	virtual void SortStreamingCellsByImportance(const TSet<const UWorldPartitionRuntimeCell*>& InCells, const TArray<FWorldPartitionStreamingSource>& InSources, TArray<const UWorldPartitionRuntimeCell*, TInlineAllocator<256>>& OutSortedCells) const;
diff --git a/Engine/Source/Runtime/Engine/Public/WorldPartition/WorldPartitionRuntimeSpatialHash.h b/Engine/Source/Runtime/Engine/Public/WorldPartition/WorldPartitionRuntimeSpatialHash.h
--- a/Engine/Source/Runtime/Engine/Public/WorldPartition/WorldPartitionRuntimeSpatialHash.h
+++ b/Engine/Source/Runtime/Engine/Public/WorldPartition/WorldPartitionRuntimeSpatialHash.h
@@ -217,7 +217,7 @@ public:
 #endif
 
 	// streaming interface
-	virtual int32 GetAllStreamingCells(TSet<const UWorldPartitionRuntimeCell*>& Cells, bool bIncludeDataLayers = false) const override;
+	virtual int32 GetAllStreamingCells(TSet<const UWorldPartitionRuntimeCell*>& Cells, bool bAllDataLayers, bool bDataLayersOnly, const TSet<FName>& InDataLayers) const override;
 	virtual bool GetStreamingCells(const FWorldPartitionStreamingQuerySource& QuerySource, TSet<const UWorldPartitionRuntimeCell*>& OutCells) const override;
 	virtual bool GetStreamingCells(const TArray<FWorldPartitionStreamingSource>& Sources, UWorldPartitionRuntimeHash::FStreamingSourceCells& OutActivateCells, UWorldPartitionRuntimeHash::FStreamingSourceCells& OutLoadCells) const override;
 		
diff --git a/Engine/Source/Runtime/Engine/Public/WorldPartition/WorldPartitionStreamingPolicy.h b/Engine/Source/Runtime/Engine/Public/WorldPartition/WorldPartitionStreamingPolicy.h
--- a/Engine/Source/Runtime/Engine/Public/WorldPartition/WorldPartitionStreamingPolicy.h
+++ b/Engine/Source/Runtime/Engine/Public/WorldPartition/WorldPartitionStreamingPolicy.h
@@ -55,7 +55,6 @@ protected:
 	virtual int32 GetMaxCellsToLoad() const;
 	virtual void UpdateStreamingSources();
 
-	bool bIsServerLoadingDone;
 	const UWorldPartition* WorldPartition;
 	TSet<const UWorldPartitionRuntimeCell*> LoadedCells;
 	TSet<const UWorldPartitionRuntimeCell*> ActivatedCells;
@@ -67,4 +66,6 @@ protected:
 	int32 UpdateStreamingStateEpoch;
 	mutable int32 SortedAddToWorldCellsEpoch;
 	mutable TArray<const UWorldPartitionRuntimeCell*, TInlineAllocator<256>> SortedAddToWorldCells;
+
+	int32 DataLayersStatesServerEpoch;
 };
\ No newline at end of file
