World Partition: properly handle deleted actors when unloading cells.

diff --git a/Engine/Source/Runtime/Engine/Private/WorldPartition/WorldPartitionActorDesc.cpp b/Engine/Source/Runtime/Engine/Private/WorldPartition/WorldPartitionActorDesc.cpp
--- a/Engine/Source/Runtime/Engine/Private/WorldPartition/WorldPartitionActorDesc.cpp
+++ b/Engine/Source/Runtime/Engine/Private/WorldPartition/WorldPartitionActorDesc.cpp
@@ -232,9 +232,9 @@ bool FWorldPartitionActorDesc::IsLoaded() const
 	return ActorPtr.IsValid();
 }
 
-AActor* FWorldPartitionActorDesc::GetActor() const
+AActor* FWorldPartitionActorDesc::GetActor(bool bEvenIfPendingKill, bool bEvenIfUnreachable) const
 {
-	return ActorPtr.Get(true);
+	return bEvenIfUnreachable ? ActorPtr.GetEvenIfUnreachable() : ActorPtr.Get(bEvenIfPendingKill);
 }
 
 AActor* FWorldPartitionActorDesc::Load() const
diff --git a/Engine/Source/Runtime/Engine/Private/WorldPartition/WorldPartitionEditorCell.cpp b/Engine/Source/Runtime/Engine/Private/WorldPartition/WorldPartitionEditorCell.cpp
--- a/Engine/Source/Runtime/Engine/Private/WorldPartition/WorldPartitionEditorCell.cpp
+++ b/Engine/Source/Runtime/Engine/Private/WorldPartition/WorldPartitionEditorCell.cpp
@@ -5,21 +5,24 @@
 #include "WorldPartition/WorldPartition.h"
 #include "WorldPartition/WorldPartitionActorDesc.h"
 
-void UWorldPartitionEditorCell::Serialize(FArchive& Ar)
-{
 #if WITH_EDITOR
-	if (Ar.IsTransacting())
+void UWorldPartitionEditorCell::AddReferencedObjects(UObject* InThis, FReferenceCollector& Collector)
+{
+	UWorldPartitionEditorCell* This = CastChecked<UWorldPartitionEditorCell>(InThis);
+	
+	Collector.AllowEliminatingReferences(false);
+	for (const FActorReference& ActorReference: This->LoadedActors)
 	{
-		bool bLoadedBool = bLoaded;
-		Ar << bLoadedBool;
-		bLoaded = bLoadedBool;
+		AActor* LoadedActor = ActorReference->GetActor(/*bEvenIfPendingKill*/true, /*bEvenIfUnreachable*/true);
+		check(LoadedActor);
+	
+		Collector.AddReferencedObject(LoadedActor);
 	}
-#endif
+	Collector.AllowEliminatingReferences(true);
 
-	Super::Serialize(Ar);
+	Super::AddReferencedObjects(InThis, Collector);
 }
 
-#if WITH_EDITOR
 void UWorldPartitionEditorCell::BeginDestroy()
 {
 	// Release WorldPartition Actor Handles/References
diff --git a/Engine/Source/Runtime/Engine/Public/WorldPartition/WorldPartitionActorDesc.h b/Engine/Source/Runtime/Engine/Public/WorldPartition/WorldPartitionActorDesc.h
--- a/Engine/Source/Runtime/Engine/Public/WorldPartition/WorldPartitionActorDesc.h
+++ b/Engine/Source/Runtime/Engine/Public/WorldPartition/WorldPartitionActorDesc.h
@@ -129,7 +129,7 @@ public:
 	FString ToString() const;
 
 	bool IsLoaded() const;
-	AActor* GetActor() const;
+	AActor* GetActor(bool bEvenIfPendingKill=true, bool bEvenIfUnreachable=false) const;
 	AActor* Load() const;
 	virtual void Unload();
 
diff --git a/Engine/Source/Runtime/Engine/Public/WorldPartition/WorldPartitionEditorCell.h b/Engine/Source/Runtime/Engine/Public/WorldPartition/WorldPartitionEditorCell.h
--- a/Engine/Source/Runtime/Engine/Public/WorldPartition/WorldPartitionEditorCell.h
+++ b/Engine/Source/Runtime/Engine/Public/WorldPartition/WorldPartitionEditorCell.h
@@ -18,9 +18,11 @@ class UWorldPartitionEditorCell: public UObject
 	GENERATED_UCLASS_BODY()
 
 public:
-	virtual void Serialize(FArchive& Ar) override;
-
 #if WITH_EDITOR
+	//~ UObject interface
+	static void AddReferencedObjects(UObject* InThis, FReferenceCollector& Collector);
+	//~ UObject interface
+
 	template<typename Type>
 	struct TActorHandle
 	{
