Removed hack in UWorld to filter autosaves from UActorDescContainer::OnObjectPreSave by using SaveContext.GetSaveFlags instead.

diff --git a/Engine/Source/Editor/UnrealEd/Private/FileHelpers.cpp b/Engine/Source/Editor/UnrealEd/Private/FileHelpers.cpp
--- a/Engine/Source/Editor/UnrealEd/Private/FileHelpers.cpp
+++ b/Engine/Source/Editor/UnrealEd/Private/FileHelpers.cpp
@@ -2862,7 +2862,6 @@ EAutosaveContentPackagesResult::Type FEditorFileUtils::AutosaveMapEx(const FStri
 			// Now gather the world external packages and save them if needed
 			if (World->PersistentLevel)
 			{
-				World->PersistentLevel->SetIsAutoSaveExternalActorPackages(true);
 				for (UPackage* ExternalPackage : World->PersistentLevel->GetLoadedExternalActorPackages())
 				{
 					if (ExternalPackage->IsDirty() && (bForceIfNotInList || DirtyPackagesForAutoSave.Contains(ExternalPackage))
@@ -2897,7 +2896,6 @@ EAutosaveContentPackagesResult::Type FEditorFileUtils::AutosaveMapEx(const FStri
 						}
 					}
 				}
-				World->PersistentLevel->SetIsAutoSaveExternalActorPackages(false);
 			}
 		}
 
diff --git a/Engine/Source/Runtime/Engine/Classes/Engine/Level.h b/Engine/Source/Runtime/Engine/Classes/Engine/Level.h
--- a/Engine/Source/Runtime/Engine/Classes/Engine/Level.h
+++ b/Engine/Source/Runtime/Engine/Classes/Engine/Level.h
@@ -610,9 +610,6 @@ public:
 
 	/** Whether the level has been saved after introducing actor GUIDs */
 	uint8										bContainsStableActorGUIDs:1;
-
-	/** Whether we are autosaviong external actors packages or not */
-	uint8										bIsAutoSaveExternalActorPackages:1;
 #endif
 	
 	/** The below variables are used temporarily while making a level visible.				*/
@@ -1085,16 +1082,6 @@ public:
 	 */
 	ENGINE_API TArray<UPackage*> GetLoadedExternalActorPackages() const;
 
-	/**
-	 * Tells the level that we're autosaving external actors or not
-	 */
-	ENGINE_API void SetIsAutoSaveExternalActorPackages(bool bInIsAutoSaveExternalActorPackages);
-
-	/**
-	 * Returns if we are autosaving external actor packages or not
-	 */
-	ENGINE_API bool GetIsAutoSaveExternalActorPackages() const;
-
 	/**
 	 * Create an package for this actor
 	 * @param InActorPath the fully qualified actor path, in the format: 'Outermost.Outer.Name'
diff --git a/Engine/Source/Runtime/Engine/Private/Level.cpp b/Engine/Source/Runtime/Engine/Private/Level.cpp
--- a/Engine/Source/Runtime/Engine/Private/Level.cpp
+++ b/Engine/Source/Runtime/Engine/Private/Level.cpp
@@ -330,7 +330,6 @@ ULevel::ULevel( const FObjectInitializer& ObjectInitializer )
 	FixupOverrideVertexColorsCount = 0;
 	bUseExternalActors = false;
 	bContainsStableActorGUIDs = true;
-	bIsAutoSaveExternalActorPackages = false;
 #endif	
 	bActorClusterCreated = false;
 	bIsPartitioned = false;
@@ -2539,17 +2538,6 @@ TArray<UPackage*> ULevel::GetLoadedExternalActorPackages() const
 	return ActorPackages.Array();
 }
 
-void ULevel::SetIsAutoSaveExternalActorPackages(bool bInIsAutoSaveExternalActorPackages)
-{
-	check(bIsAutoSaveExternalActorPackages != bInIsAutoSaveExternalActorPackages);
-	bIsAutoSaveExternalActorPackages = bInIsAutoSaveExternalActorPackages;
-}
-
-bool ULevel::GetIsAutoSaveExternalActorPackages() const
-{
-	return bIsAutoSaveExternalActorPackages;
-}
-
 UPackage* ULevel::CreateActorPackage(UPackage* InLevelPackage, const FString& InActorPath)
 {
 	UPackage* ActorPackage = CreatePackage(*GetActorPackageName(InLevelPackage, InActorPath));
diff --git a/Engine/Source/Runtime/Engine/Private/WorldPartition/ActorDescContainer.cpp b/Engine/Source/Runtime/Engine/Private/WorldPartition/ActorDescContainer.cpp
--- a/Engine/Source/Runtime/Engine/Private/WorldPartition/ActorDescContainer.cpp
+++ b/Engine/Source/Runtime/Engine/Private/WorldPartition/ActorDescContainer.cpp
@@ -150,7 +150,7 @@ void UActorDescContainer::BeginDestroy()
 #if WITH_EDITOR
 bool UActorDescContainer::ShouldHandleActorEvent(const AActor* Actor)
 {
-	if (Actor && Actor->IsPackageExternal() && (Actor->GetLevel() == World->PersistentLevel) && !World->PersistentLevel->GetIsAutoSaveExternalActorPackages() && Actor->IsMainPackageActor())
+	if (Actor && Actor->IsPackageExternal() && (Actor->GetLevel() == World->PersistentLevel) && Actor->IsMainPackageActor())
 	{
 		// Only handle assets that belongs to our level
 		auto RemoveAfterFirstDot = [](const FString& InValue)
@@ -173,41 +173,37 @@ bool UActorDescContainer::ShouldHandleActorEvent(const AActor* Actor)
 
 void UActorDescContainer::OnObjectPreSave(UObject* Object, FObjectPreSaveContext SaveContext)
 {
-	if (SaveContext.IsProceduralSave())
+	if (!SaveContext.IsProceduralSave() && !(SaveContext.GetSaveFlags() & SAVE_FromAutosave))
 	{
-		// Do not delete and recreate FWorldPartitionActorDesc when making a procedural save. Procedural saves such as
-		// EditorDomain saves can occur during a load package, and FWorldPartitionReferenceImpl::IncRefCount 
-		// can load (and therefore save) an actor package while holding a pointer to the FWorldPartitionActorDesc. 
-		return;
-	}
-	if (const AActor* Actor = Cast<AActor>(Object))
-	{
-		if (ShouldHandleActorEvent(Actor))
+		if (const AActor* Actor = Cast<AActor>(Object))
 		{
-			check(!Actor->IsPendingKill());
-			if (TUniquePtr<FWorldPartitionActorDesc>* ExistingActorDesc = GetActorDescriptor(Actor->GetActorGuid()))
+			if (ShouldHandleActorEvent(Actor))
 			{
-				// Pin the actor handle on the actor to prevent unloading it when unhashing
-				FWorldPartitionHandle ExistingActorHandle(ExistingActorDesc);
-				FWorldPartitionHandlePinRefScope ExistingActorHandlePin(ExistingActorHandle);
+				check(!Actor->IsPendingKill());
+				if (TUniquePtr<FWorldPartitionActorDesc>* ExistingActorDesc = GetActorDescriptor(Actor->GetActorGuid()))
+				{
+					// Pin the actor handle on the actor to prevent unloading it when unhashing
+					FWorldPartitionHandle ExistingActorHandle(ExistingActorDesc);
+					FWorldPartitionHandlePinRefScope ExistingActorHandlePin(ExistingActorHandle);
 
-				TUniquePtr<FWorldPartitionActorDesc> NewActorDesc(Actor->CreateActorDesc());
+					TUniquePtr<FWorldPartitionActorDesc> NewActorDesc(Actor->CreateActorDesc());
 
-				OnActorDescUpdating(ExistingActorDesc->Get());
+					OnActorDescUpdating(ExistingActorDesc->Get());
 
-				// Transfer any reference count from external sources
-				NewActorDesc->TransferRefCounts(ExistingActorDesc->Get());
+					// Transfer any reference count from external sources
+					NewActorDesc->TransferRefCounts(ExistingActorDesc->Get());
 
-				*ExistingActorDesc = MoveTemp(NewActorDesc);
+					*ExistingActorDesc = MoveTemp(NewActorDesc);
 
-				OnActorDescUpdated(ExistingActorDesc->Get());
-			}
-			// New actor
-			else
-			{
-				FWorldPartitionActorDesc* const AddedActorDesc = AddActor(Actor);
-				OnActorDescAdded(AddedActorDesc);
-				OnActorDescAddedEvent.Broadcast(AddedActorDesc);
+					OnActorDescUpdated(ExistingActorDesc->Get());
+				}
+				// New actor
+				else
+				{
+					FWorldPartitionActorDesc* const AddedActorDesc = AddActor(Actor);
+					OnActorDescAdded(AddedActorDesc);
+					OnActorDescAddedEvent.Broadcast(AddedActorDesc);
+				}
 			}
 		}
 	}
