Fix to allow accessing the specific BodyInstance when the InstancedStaticMeshComponent has an array of BodyInstance instead of retrieving the default one in the primitive component base class.
#rnx

diff --git a/Engine/Plugins/Runtime/ApexDestruction/Source/ApexDestruction/Private/DestructibleComponent.cpp b/Engine/Plugins/Runtime/ApexDestruction/Source/ApexDestruction/Private/DestructibleComponent.cpp
--- a/Engine/Plugins/Runtime/ApexDestruction/Source/ApexDestruction/Private/DestructibleComponent.cpp
+++ b/Engine/Plugins/Runtime/ApexDestruction/Source/ApexDestruction/Private/DestructibleComponent.cpp
@@ -1431,7 +1431,7 @@ void UDestructibleComponent::SetEnableGravity(bool bGravityEnabled)
 #endif //WITH_APEX
 }
 
-FBodyInstance* UDestructibleComponent::GetBodyInstance( FName BoneName /*= NAME_None*/, bool) const
+FBodyInstance* UDestructibleComponent::GetBodyInstance( FName BoneName /*= NAME_None*/, bool, int) const
 {
 #if WITH_APEX
 	if (ApexDestructibleActor != NULL)
diff --git a/Engine/Plugins/Runtime/ApexDestruction/Source/ApexDestruction/Public/DestructibleComponent.h b/Engine/Plugins/Runtime/ApexDestruction/Source/ApexDestruction/Public/DestructibleComponent.h
--- a/Engine/Plugins/Runtime/ApexDestruction/Source/ApexDestruction/Public/DestructibleComponent.h
+++ b/Engine/Plugins/Runtime/ApexDestruction/Source/ApexDestruction/Public/DestructibleComponent.h
@@ -135,7 +135,7 @@ public:
 	//~ End UActorComponent Interface.
 
 	//~ Begin UPrimitiveComponent Interface.
-	virtual FBodyInstance* GetBodyInstance(FName BoneName = NAME_None, bool bGetWelded = true) const override;
+	virtual FBodyInstance* GetBodyInstance(FName BoneName = NAME_None, bool bGetWelded = true, int32 Index = INDEX_NONE) const override;
 	virtual bool CanEditSimulatePhysics() override;
 	virtual bool IsAnySimulatingPhysics() const override;
 
diff --git a/Engine/Source/Runtime/Engine/Classes/Components/InstancedStaticMeshComponent.h b/Engine/Source/Runtime/Engine/Classes/Components/InstancedStaticMeshComponent.h
--- a/Engine/Source/Runtime/Engine/Classes/Components/InstancedStaticMeshComponent.h
+++ b/Engine/Source/Runtime/Engine/Classes/Components/InstancedStaticMeshComponent.h
@@ -338,6 +338,7 @@ public:
 
 	//~ Begin UPrimitiveComponent Interface
 	virtual FPrimitiveSceneProxy* CreateSceneProxy() override;
+	virtual FBodyInstance* GetBodyInstance(FName BoneName = NAME_None, bool bGetWelded = true, int32 Index = INDEX_NONE) const override;
 protected:
 	virtual void OnCreatePhysicsState() override;
 	virtual void OnDestroyPhysicsState() override;
diff --git a/Engine/Source/Runtime/Engine/Classes/Components/PrimitiveComponent.h b/Engine/Source/Runtime/Engine/Classes/Components/PrimitiveComponent.h
--- a/Engine/Source/Runtime/Engine/Classes/Components/PrimitiveComponent.h
+++ b/Engine/Source/Runtime/Engine/Classes/Components/PrimitiveComponent.h
@@ -1814,10 +1814,11 @@ public:
 	*
 	* @param BoneName				Used to get body associated with specific bone. NAME_None automatically gets the root most body
 	* @param bGetWelded				If the component has been welded to another component and bGetWelded is true we return the single welded BodyInstance that is used in the simulation
+	* @param Index					Index used in Components with multiple body instances
 	*
 	* @return		Returns the BodyInstance based on various states (does component have multiple bodies? Is the body welded to another body?)
 	*/
-	virtual FBodyInstance* GetBodyInstance(FName BoneName = NAME_None, bool bGetWelded = true) const;
+	virtual FBodyInstance* GetBodyInstance(FName BoneName = NAME_None, bool bGetWelded = true, int32 Index = INDEX_NONE) const;
 
 	/** 
 	 * Returns The square of the distance to closest Body Instance surface. 
diff --git a/Engine/Source/Runtime/Engine/Classes/Components/SkeletalMeshComponent.h b/Engine/Source/Runtime/Engine/Classes/Components/SkeletalMeshComponent.h
--- a/Engine/Source/Runtime/Engine/Classes/Components/SkeletalMeshComponent.h
+++ b/Engine/Source/Runtime/Engine/Classes/Components/SkeletalMeshComponent.h
@@ -1655,7 +1655,7 @@ public:
 
 	virtual class UBodySetup* GetBodySetup() override;
 	virtual bool CanEditSimulatePhysics() override;
-	virtual FBodyInstance* GetBodyInstance(FName BoneName = NAME_None, bool bGetWelded = true) const override;
+	virtual FBodyInstance* GetBodyInstance(FName BoneName = NAME_None, bool bGetWelded = true, int32 Index = INDEX_NONE) const override;
 	virtual void UpdatePhysicsToRBChannels() override;
 	virtual void SetAllPhysicsAngularVelocityInRadians(FVector const& NewVel, bool bAddToCurrent = false) override;
 	virtual void SetAllPhysicsPosition(FVector NewPos) override;
diff --git a/Engine/Source/Runtime/Engine/Classes/Particles/ParticleSystemComponent.h b/Engine/Source/Runtime/Engine/Classes/Particles/ParticleSystemComponent.h
--- a/Engine/Source/Runtime/Engine/Classes/Particles/ParticleSystemComponent.h
+++ b/Engine/Source/Runtime/Engine/Classes/Particles/ParticleSystemComponent.h
@@ -1377,7 +1377,7 @@ public:
 	virtual FPrimitiveSceneProxy* CreateSceneProxy() override;
 	virtual void GetUsedMaterials( TArray<UMaterialInterface*>& OutMaterials, bool bGetDebugMaterials = false ) const override;
 	virtual void GetStreamingRenderAssetInfo(FStreamingTextureLevelContext& LevelContext, TArray<FStreamingRenderAssetPrimitiveInfo>& OutStreamingRenderAssets) const override;
-	virtual FBodyInstance* GetBodyInstance(FName BoneName = NAME_None, bool bGetWelded = true) const override;
+	virtual FBodyInstance* GetBodyInstance(FName BoneName = NAME_None, bool bGetWelded = true, int32 Index = INDEX_NONE) const override;
 	//End UPrimitiveComponent Interface
 
 	//~ Begin USceneComonent Interface
diff --git a/Engine/Source/Runtime/Engine/Private/InstancedStaticMesh.cpp b/Engine/Source/Runtime/Engine/Private/InstancedStaticMesh.cpp
--- a/Engine/Source/Runtime/Engine/Private/InstancedStaticMesh.cpp
+++ b/Engine/Source/Runtime/Engine/Private/InstancedStaticMesh.cpp
@@ -1828,6 +1828,7 @@ UInstancedStaticMeshComponent::UInstancedStaticMeshComponent(const FObjectInitia
 	BodyInstance.bSimulatePhysics = false;
 
 	bDisallowMeshPaintPerInstance = true;
+	bMultiBodyOverlap = true;
 
 #if STATS
 	{
@@ -1950,6 +1951,15 @@ bool UInstancedStaticMeshComponent::IsHLODRelevant() const
 	return Super::IsHLODRelevant();
 }
 
+FBodyInstance* UInstancedStaticMeshComponent::GetBodyInstance(FName BoneName, bool bGetWelded, int32 Index) const
+{
+	if (Index != INDEX_NONE && IsValidInstance(Index))
+	{
+		return const_cast<FBodyInstance*>(InstanceBodies[Index]);
+	}
+	return  nullptr;
+}
+
 FPrimitiveSceneProxy* UInstancedStaticMeshComponent::CreateSceneProxy()
 {
 	static const auto NaniteProxyRenderModeVar = IConsoleManager::Get().FindConsoleVariable(TEXT("r.Nanite.ProxyRenderMode"));
diff --git a/Engine/Source/Runtime/Engine/Private/LevelActor.cpp b/Engine/Source/Runtime/Engine/Private/LevelActor.cpp
--- a/Engine/Source/Runtime/Engine/Private/LevelActor.cpp
+++ b/Engine/Source/Runtime/Engine/Private/LevelActor.cpp
@@ -1297,7 +1297,8 @@ static bool ComponentEncroachesBlockingGeometry_WithAdjustment(UWorld const* Wor
 				{
 					NumBlockingHits++;
 					FCollisionShape const NonShrunkenCollisionShape = PrimComp->GetCollisionShape();
-					bool bSuccess = OverlapComponent->ComputePenetration(MTDResult, NonShrunkenCollisionShape, TestWorldTransform.GetLocation(), TestWorldTransform.GetRotation());
+					const FBodyInstance* OverlapBodyInstance = OverlapComponent->GetBodyInstance(NAME_None, true, Overlaps[HitIdx].ItemIndex);
+					bool bSuccess = OverlapBodyInstance && OverlapBodyInstance->OverlapTest(TestWorldTransform.GetLocation(), TestWorldTransform.GetRotation(), NonShrunkenCollisionShape, &MTDResult);
 					if (bSuccess)
 					{
 						OutProposedAdjustment += MTDResult.Direction * MTDResult.Distance;
@@ -1316,7 +1317,7 @@ static bool ComponentEncroachesBlockingGeometry_WithAdjustment(UWorld const* Wor
 					if (bSuccess && FMath::IsNearlyZero(MTDResult.Distance))
 					{
 						FCollisionShape const ShrunkenCollisionShape = PrimComp->GetCollisionShape(-Epsilon);
-						bSuccess = OverlapComponent->ComputePenetration(MTDResult, ShrunkenCollisionShape, TestWorldTransform.GetLocation(), TestWorldTransform.GetRotation());
+						bSuccess = OverlapBodyInstance && OverlapBodyInstance->OverlapTest(TestWorldTransform.GetLocation(), TestWorldTransform.GetRotation(), ShrunkenCollisionShape, &MTDResult);
 						if (bSuccess)
 						{
 							OutProposedAdjustment += MTDResult.Direction * MTDResult.Distance;
diff --git a/Engine/Source/Runtime/Engine/Private/Particles/ParticleComponents.cpp b/Engine/Source/Runtime/Engine/Private/Particles/ParticleComponents.cpp
--- a/Engine/Source/Runtime/Engine/Private/Particles/ParticleComponents.cpp
+++ b/Engine/Source/Runtime/Engine/Private/Particles/ParticleComponents.cpp
@@ -7768,7 +7768,7 @@ void UParticleSystemComponent::GetStreamingRenderAssetInfo(FStreamingTextureLeve
 	}
 }
 
-FBodyInstance* UParticleSystemComponent::GetBodyInstance(FName BoneName /*= NAME_None*/, bool bGetWelded /*= true*/) const
+FBodyInstance* UParticleSystemComponent::GetBodyInstance(FName BoneName /*= NAME_None*/, bool bGetWelded /*= true*/, int32 Index /*=INDEX_NONE*/) const
 {
 	return nullptr;
 }
diff --git a/Engine/Source/Runtime/Engine/Private/PrimitiveComponentPhysics.cpp b/Engine/Source/Runtime/Engine/Private/PrimitiveComponentPhysics.cpp
--- a/Engine/Source/Runtime/Engine/Private/PrimitiveComponentPhysics.cpp
+++ b/Engine/Source/Runtime/Engine/Private/PrimitiveComponentPhysics.cpp
@@ -803,7 +803,7 @@ void UPrimitiveComponent::UnWeldChildren()
 	}
 }
 
-FBodyInstance* UPrimitiveComponent::GetBodyInstance(FName BoneName, bool bGetWelded) const
+FBodyInstance* UPrimitiveComponent::GetBodyInstance(FName BoneName, bool bGetWelded, int32 Index) const
 {
 	return const_cast<FBodyInstance*>((bGetWelded && BodyInstance.WeldParent) ? BodyInstance.WeldParent : &BodyInstance);
 }
diff --git a/Engine/Source/Runtime/Engine/Private/SkeletalMeshComponentPhysics.cpp b/Engine/Source/Runtime/Engine/Private/SkeletalMeshComponentPhysics.cpp
--- a/Engine/Source/Runtime/Engine/Private/SkeletalMeshComponentPhysics.cpp
+++ b/Engine/Source/Runtime/Engine/Private/SkeletalMeshComponentPhysics.cpp
@@ -1708,7 +1708,7 @@ FName USkeletalMeshComponent::FindConstraintBoneName( int32 ConstraintIndex )
 }
 
 
-FBodyInstance* USkeletalMeshComponent::GetBodyInstance(FName BoneName, bool) const
+FBodyInstance* USkeletalMeshComponent::GetBodyInstance(FName BoneName, bool, int32) const
 {
 	UPhysicsAsset * const PhysicsAsset = GetPhysicsAsset();
 	FBodyInstance* BodyInst = NULL;
diff --git a/Engine/Source/Runtime/Experimental/GeometryCollectionEngine/Private/GeometryCollection/GeometryCollectionComponent.cpp b/Engine/Source/Runtime/Experimental/GeometryCollectionEngine/Private/GeometryCollection/GeometryCollectionComponent.cpp
--- a/Engine/Source/Runtime/Experimental/GeometryCollectionEngine/Private/GeometryCollection/GeometryCollectionComponent.cpp
+++ b/Engine/Source/Runtime/Experimental/GeometryCollectionEngine/Private/GeometryCollection/GeometryCollectionComponent.cpp
@@ -612,7 +612,7 @@ void UGeometryCollectionComponent::SetNotifyBreaks(bool bNewNotifyBreaks)
 	}
 }
 
-FBodyInstance* UGeometryCollectionComponent::GetBodyInstance(FName BoneName /*= NAME_None*/, bool bGetWelded /*= true*/) const
+FBodyInstance* UGeometryCollectionComponent::GetBodyInstance(FName BoneName /*= NAME_None*/, bool bGetWelded /*= true*/, int32 Index /*=INDEX_NONE*/) const
 {
 	return nullptr;// const_cast<FBodyInstance*>(&DummyBodyInstance);
 }
diff --git a/Engine/Source/Runtime/Experimental/GeometryCollectionEngine/Public/GeometryCollection/GeometryCollectionComponent.h b/Engine/Source/Runtime/Experimental/GeometryCollectionEngine/Public/GeometryCollection/GeometryCollectionComponent.h
--- a/Engine/Source/Runtime/Experimental/GeometryCollectionEngine/Public/GeometryCollection/GeometryCollectionComponent.h
+++ b/Engine/Source/Runtime/Experimental/GeometryCollectionEngine/Public/GeometryCollection/GeometryCollectionComponent.h
@@ -361,7 +361,7 @@ public:
 	//~ Begin UPrimitiveComponent Interface.
 	virtual FPrimitiveSceneProxy* CreateSceneProxy() override;
 	virtual void OnRegister() override;
-	virtual FBodyInstance* GetBodyInstance(FName BoneName = NAME_None, bool bGetWelded = true) const override;
+	virtual FBodyInstance* GetBodyInstance(FName BoneName = NAME_None, bool bGetWelded = true, int32 Index = -1) const override;
 	virtual void SetNotifyRigidBodyCollision(bool bNewNotifyRigidBodyCollision) override;
 	//~ End UPrimitiveComponent Interface.
 
