Move Audio plugin init logs to 'Display' so they show up in console and unify format when no plugin set (using build-in) for easy look-up

diff --git a/Engine/Source/Runtime/Engine/Private/AudioDevice.cpp b/Engine/Source/Runtime/Engine/Private/AudioDevice.cpp
--- a/Engine/Source/Runtime/Engine/Private/AudioDevice.cpp
+++ b/Engine/Source/Runtime/Engine/Private/AudioDevice.cpp
@@ -458,11 +458,11 @@ bool FAudioDevice::Init(Audio::FDeviceId InDeviceID, int32 InMaxSources)
 		bSpatializationInterfaceEnabled = true;
 		bSpatializationIsExternalSend = SpatializationPluginFactory->IsExternalSend();
 		MaxChannelsSupportedBySpatializationPlugin = SpatializationPluginFactory->GetMaxSupportedChannels();
-		UE_LOG(LogAudio, Log, TEXT("Using Audio Spatialization Plugin: %s is external send: %d"), *(SpatializationPluginFactory->GetDisplayName()), bSpatializationIsExternalSend);
+		UE_LOG(LogAudio, Display, TEXT("Audio Spatialization Plugin: %s is external send: %d"), *(SpatializationPluginFactory->GetDisplayName()), bSpatializationIsExternalSend);
 	}
 	else
 	{
-		UE_LOG(LogAudio, Log, TEXT("Using built-in audio spatialization."));
+		UE_LOG(LogAudio, Display, TEXT("Audio Spatialization Plugin: None (built-in)."));
 	}
 
 	//Get the requested reverb plugin and set it up:
@@ -472,11 +472,11 @@ bool FAudioDevice::Init(Audio::FDeviceId InDeviceID, int32 InMaxSources)
 		ReverbPluginInterface = ReverbPluginFactory->CreateNewReverbPlugin(this);
 		bReverbInterfaceEnabled = true;
 		bReverbIsExternalSend = ReverbPluginFactory->IsExternalSend();
-		UE_LOG(LogAudio, Log, TEXT("Audio Reverb Plugin: %s"), *(ReverbPluginFactory->GetDisplayName()));
+		UE_LOG(LogAudio, Display, TEXT("Audio Reverb Plugin: %s"), *(ReverbPluginFactory->GetDisplayName()));
 	}
 	else
 	{
-		UE_LOG(LogAudio, Log, TEXT("Using built-in audio reverb."));
+		UE_LOG(LogAudio, Display, TEXT("Audio Reverb Plugin: None (built-in)."));
 	}
 
 	//Get the requested occlusion plugin and set it up.
@@ -490,7 +490,7 @@ bool FAudioDevice::Init(Audio::FDeviceId InDeviceID, int32 InMaxSources)
 	}
 	else
 	{
-		UE_LOG(LogAudio, Display, TEXT("Using built-in audio occlusion."));
+		UE_LOG(LogAudio, Display, TEXT("Audio Occlusion Plugin: None (built-in)."));
 	}
 
 	//Get the requested modulation plugin and set it up.
