Fix non-unity build error.

diff --git a/Engine/Source/Runtime/Engine/Classes/Engine/NetConnection.h b/Engine/Source/Runtime/Engine/Classes/Engine/NetConnection.h
--- a/Engine/Source/Runtime/Engine/Classes/Engine/NetConnection.h
+++ b/Engine/Source/Runtime/Engine/Classes/Engine/NetConnection.h
@@ -39,6 +39,7 @@ class FObjectReplicator;
 class StatelessConnectHandlerComponent;
 class UActorChannel;
 class UChildConnection;
+struct FEncryptionKeyResponse;
 
 typedef TMap<TWeakObjectPtr<AActor>, UActorChannel*, FDefaultSetAllocator, TWeakObjectPtrMapKeyFuncs<TWeakObjectPtr<AActor>, UActorChannel*>> FActorChannelMap;
 
diff --git a/Engine/Source/Runtime/Engine/Private/NetConnection.cpp b/Engine/Source/Runtime/Engine/Private/NetConnection.cpp
--- a/Engine/Source/Runtime/Engine/Private/NetConnection.cpp
+++ b/Engine/Source/Runtime/Engine/Private/NetConnection.cpp
@@ -5,6 +5,7 @@
 =============================================================================*/
 
 #include "Engine/NetConnection.h"
+#include "Engine/NetworkDelegates.h"
 #include "Misc/CommandLine.h"
 #include "EngineStats.h"
 #include "EngineGlobals.h"
