Reduce latency of known long running tasks by preventing other long running tasks from being taken by busy wait

diff --git a/Engine/Source/Runtime/Engine/Private/DistanceFieldAtlas.cpp b/Engine/Source/Runtime/Engine/Private/DistanceFieldAtlas.cpp
--- a/Engine/Source/Runtime/Engine/Private/DistanceFieldAtlas.cpp
+++ b/Engine/Source/Runtime/Engine/Private/DistanceFieldAtlas.cpp
@@ -422,7 +422,7 @@ void FDistanceFieldAsyncQueue::StartBackgroundTask(FAsyncDistanceFieldTask* Task
 {
 	check(Task->AsyncTask == nullptr);
 	Task->AsyncTask = MakeUnique<FAsyncTask<FAsyncDistanceFieldTaskWorker>>(*Task);
-	Task->AsyncTask->StartBackgroundTask(ThreadPool.Get(), EQueuedWorkPriority::Lowest);
+	Task->AsyncTask->StartBackgroundTask(ThreadPool.Get(), EQueuedWorkPriority::Lowest, EQueuedWorkFlags::DoNotRunInsideBusyWait);
 }
 
 void FDistanceFieldAsyncQueue::ProcessPendingTasks()
diff --git a/Engine/Source/Runtime/Engine/Private/SkeletalMesh.cpp b/Engine/Source/Runtime/Engine/Private/SkeletalMesh.cpp
--- a/Engine/Source/Runtime/Engine/Private/SkeletalMesh.cpp
+++ b/Engine/Source/Runtime/Engine/Private/SkeletalMesh.cpp
@@ -2042,7 +2042,7 @@ void USkeletalMesh::Build()
 		EQueuedWorkPriority BasePriority = FSkeletalMeshCompilingManager::Get().GetBasePriority(this);
 		check(AsyncTask == nullptr);
 		AsyncTask = MakeUnique<FSkeletalMeshAsyncBuildTask>(this, MoveTemp(Context));
-		AsyncTask->StartBackgroundTask(SkeletalMeshThreadPool, BasePriority);
+		AsyncTask->StartBackgroundTask(SkeletalMeshThreadPool, BasePriority, EQueuedWorkFlags::DoNotRunInsideBusyWait);
 		FSkeletalMeshCompilingManager::Get().AddSkeletalMeshes({ this });
 	}
 	else
@@ -3155,7 +3155,7 @@ void USkeletalMesh::PostLoad()
 		EQueuedWorkPriority BasePriority = FSkeletalMeshCompilingManager::Get().GetBasePriority(this);
 
 		AsyncTask = MakeUnique<FSkeletalMeshAsyncBuildTask>(this, MoveTemp(Context));
-		AsyncTask->StartBackgroundTask(SkeletalMeshThreadPool, BasePriority);
+		AsyncTask->StartBackgroundTask(SkeletalMeshThreadPool, BasePriority, EQueuedWorkFlags::DoNotRunInsideBusyWait);
 		FSkeletalMeshCompilingManager::Get().AddSkeletalMeshes({ this });
 	}
 	else
diff --git a/Engine/Source/Runtime/Engine/Private/StaticMesh.cpp b/Engine/Source/Runtime/Engine/Private/StaticMesh.cpp
--- a/Engine/Source/Runtime/Engine/Private/StaticMesh.cpp
+++ b/Engine/Source/Runtime/Engine/Private/StaticMesh.cpp
@@ -5017,7 +5017,7 @@ void UStaticMesh::PostLoad()
 		}
 
 		AsyncTask = MakeUnique<FStaticMeshAsyncBuildTask>(this, MakeUnique<FStaticMeshPostLoadContext>(MoveTemp(Context)));
-		AsyncTask->StartBackgroundTask(StaticMeshThreadPool, BasePriority, EQueuedWorkFlags::None, GetBuildRequiredMemory());
+		AsyncTask->StartBackgroundTask(StaticMeshThreadPool, BasePriority, EQueuedWorkFlags::DoNotRunInsideBusyWait, GetBuildRequiredMemory());
 		FStaticMeshCompilingManager::Get().AddStaticMeshes({this});
 	}
 	else
diff --git a/Engine/Source/Runtime/Engine/Private/StaticMeshBuild.cpp b/Engine/Source/Runtime/Engine/Private/StaticMeshBuild.cpp
--- a/Engine/Source/Runtime/Engine/Private/StaticMeshBuild.cpp
+++ b/Engine/Source/Runtime/Engine/Private/StaticMeshBuild.cpp
@@ -253,7 +253,7 @@ void UStaticMesh::BatchBuild(const TArray<UStaticMesh*>& InStaticMeshes, bool bI
 					EQueuedWorkPriority BasePriority = FStaticMeshCompilingManager::Get().GetBasePriority(StaticMesh);
 					check(StaticMesh->AsyncTask == nullptr);
 					StaticMesh->AsyncTask = MakeUnique<FStaticMeshAsyncBuildTask>(StaticMesh, MoveTemp(Context));
-					StaticMesh->AsyncTask->StartBackgroundTask(StaticMeshThreadPool, BasePriority, EQueuedWorkFlags::None, StaticMesh->GetBuildRequiredMemory());
+					StaticMesh->AsyncTask->StartBackgroundTask(StaticMeshThreadPool, BasePriority, EQueuedWorkFlags::DoNotRunInsideBusyWait, StaticMesh->GetBuildRequiredMemory());
 					FStaticMeshCompilingManager::Get().AddStaticMeshes({ StaticMesh });
 					return true;
 				}
diff --git a/Engine/Source/Runtime/Engine/Private/TextureDerivedData.cpp b/Engine/Source/Runtime/Engine/Private/TextureDerivedData.cpp
--- a/Engine/Source/Runtime/Engine/Private/TextureDerivedData.cpp
+++ b/Engine/Source/Runtime/Engine/Private/TextureDerivedData.cpp
@@ -915,7 +915,7 @@ void FTexturePlatformData::Cache(
 		EQueuedWorkPriority BasePriority      = FTextureCompilingManager::Get().GetBasePriority(&InTexture);
 
 		AsyncTask = new FTextureAsyncCacheDerivedDataTask(Compressor, this, &InTexture, InSettingsPerLayer, Flags);
-		AsyncTask->StartBackgroundTask(TextureThreadPool, BasePriority, EQueuedWorkFlags::None, AsyncTask->GetTask().GetRequiredMemoryEstimate());
+		AsyncTask->StartBackgroundTask(TextureThreadPool, BasePriority, EQueuedWorkFlags::DoNotRunInsideBusyWait, AsyncTask->GetTask().GetRequiredMemoryEstimate());
 	}
 	else
 	{
