Added the concept of thread safe Blueprint functions

Either enabling the 'thread safe' checkbox for a function or overriding an already thread safe (either native or BP) function will enable compile-time checking for thread safety.
Checks are based around the complied statement list rather than the bytecode itself.

diff --git a/Engine/Source/Editor/BlueprintGraph/Classes/EdGraphSchema_K2.h b/Engine/Source/Editor/BlueprintGraph/Classes/EdGraphSchema_K2.h
--- a/Engine/Source/Editor/BlueprintGraph/Classes/EdGraphSchema_K2.h
+++ b/Engine/Source/Editor/BlueprintGraph/Classes/EdGraphSchema_K2.h
@@ -250,6 +250,12 @@ public:
 	/** Namespace into which a type can be optionally defined; if empty or not set, the type will belong to the global namespace (default). */
 	static const FName MD_Namespace;
 
+	/** Function or class marked as thread-safe. Opts class/function compilation into thread-safety checks. */
+	static const FName MD_ThreadSafe;
+
+	/** Function marked as explicitly not thread-safe. Opts function out of class-level thread-safety checks. */
+	static const FName MD_NotThreadSafe;
+	
 private:
 	// This class should never be instantiated
 	FBlueprintMetadata() {}
@@ -1093,6 +1099,13 @@ public:
 	 */
 	bool DoesGraphSupportImpureFunctions(const UEdGraph* InGraph) const;
 
+	/** 
+	 * Checks if the graph is marked as thread safe
+	 * @param InGraph		Graph to check
+	 * @return			True if the graph is marked theead safe
+	 */
+	bool IsGraphMarkedThreadSafe(const UEdGraph* InGraph) const;
+	
 	/**
 	 * Checks to see if the passed in function is valid in the graph for the current class
 	 *
diff --git a/Engine/Source/Editor/BlueprintGraph/Classes/K2Node_EditablePinBase.h b/Engine/Source/Editor/BlueprintGraph/Classes/K2Node_EditablePinBase.h
--- a/Engine/Source/Editor/BlueprintGraph/Classes/K2Node_EditablePinBase.h
+++ b/Engine/Source/Editor/BlueprintGraph/Classes/K2Node_EditablePinBase.h
@@ -71,6 +71,9 @@ struct FKismetUserDeclaredFunctionMetadata
 	UPROPERTY()
 	bool bCallInEditor;
 
+	UPROPERTY()
+	bool bThreadSafe;
+
 	/** Cached value for whether or not the graph has latent functions, positive for TRUE, zero for FALSE, and INDEX_None for undetermined */
 	UPROPERTY(transient)
 	int8 HasLatentFunctions;
@@ -80,6 +83,7 @@ public:
 		: InstanceTitleColor(FLinearColor::White)
 		, bIsDeprecated(false)
 		, bCallInEditor(false)
+		, bThreadSafe(false)
 		, HasLatentFunctions(INDEX_NONE)
 	{
 	}
diff --git a/Engine/Source/Editor/BlueprintGraph/Private/BlueprintActionFilter.cpp b/Engine/Source/Editor/BlueprintGraph/Private/BlueprintActionFilter.cpp
--- a/Engine/Source/Editor/BlueprintGraph/Private/BlueprintActionFilter.cpp
+++ b/Engine/Source/Editor/BlueprintGraph/Private/BlueprintActionFilter.cpp
@@ -184,6 +184,16 @@ namespace BlueprintActionFilterImpl
 	*/
 	static bool IsPropertyAccessorNode(FBlueprintActionFilter const& Filter, FBlueprintActionInfo& BlueprintAction);
 
+	/**
+	* Rejection test that checks to see if the supplied node-spawner would produce a node with thread safety
+	* characteristics that are incompatible with the specification of the containing graph. 
+	* 
+	* @param  Filter			Filter context (unused) for this test.
+	* @param  BlueprintAction	The action you wish to query.
+	* @return True if the action would spawn a node that has incompatible thread safety.
+	*/
+	static bool IsThreadSafetyIncompatible(FBlueprintActionFilter const& Filter, FBlueprintActionInfo& BlueprintAction);
+
 	/**
 	 * Rejection test that checks to see if the supplied node-spawner would 
 	 * produce an impure node, incompatible with the specified graphs.
@@ -998,6 +1008,32 @@ static bool BlueprintActionFilterImpl::IsPropertyAccessorNode(FBlueprintActionFi
 	return bIsAccessor;
 }
 
+//------------------------------------------------------------------------------
+static bool BlueprintActionFilterImpl::IsThreadSafetyIncompatible(FBlueprintActionFilter const& Filter, FBlueprintActionInfo& BlueprintAction)
+{
+	if (UFunction const* Function = BlueprintAction.GetAssociatedFunction())
+	{	
+		bool bAllowNonThreadSafeNodes = true;
+
+		FBlueprintActionContext const& FilterContext = Filter.Context;
+
+		for (UEdGraph* Graph : FilterContext.Graphs)
+		{
+			if (const UEdGraphSchema_K2* K2Schema = Cast<UEdGraphSchema_K2>(Graph->GetSchema()))
+			{
+				bAllowNonThreadSafeNodes &= !K2Schema->IsGraphMarkedThreadSafe(Graph);
+			}
+		}
+
+		if(!bAllowNonThreadSafeNodes)
+		{
+			return !FBlueprintEditorUtils::HasFunctionBlueprintThreadSafeMetaData(Function);
+		}
+	}
+	
+	return false;
+}
+
 //------------------------------------------------------------------------------
 static bool BlueprintActionFilterImpl::IsIncompatibleImpureNode(FBlueprintActionFilter const& Filter, FBlueprintActionInfo& BlueprintAction)
 {
@@ -1993,6 +2029,10 @@ FBlueprintActionFilter::FBlueprintActionFilter(uint32 Flags/*= 0x00*/)
 	AddRejectionTest(FRejectionTestDelegate::CreateStatic(IsIncompatibleLatentNode));
 	AddRejectionTest(FRejectionTestDelegate::CreateStatic(IsIncompatibleImpureNode));
 	AddRejectionTest(FRejectionTestDelegate::CreateStatic(IsPropertyAccessorNode));
+	if (Flags & BPFILTER_RejectIncompatibleThreadSafety)
+	{
+		AddRejectionTest(FRejectionTestDelegate::CreateStatic(IsThreadSafetyIncompatible));
+	}
 	
 	AddRejectionTest(FRejectionTestDelegate::CreateStatic(IsActionHiddenByConfig));
 	AddRejectionTest(FRejectionTestDelegate::CreateStatic(IsFieldCategoryHidden));
diff --git a/Engine/Source/Editor/BlueprintGraph/Private/EdGraphSchema_K2.cpp b/Engine/Source/Editor/BlueprintGraph/Private/EdGraphSchema_K2.cpp
--- a/Engine/Source/Editor/BlueprintGraph/Private/EdGraphSchema_K2.cpp
+++ b/Engine/Source/Editor/BlueprintGraph/Private/EdGraphSchema_K2.cpp
@@ -199,6 +199,9 @@ const FName FBlueprintMetadata::MD_AllowAbstractClasses(TEXT("AllowAbstract"));
 
 const FName FBlueprintMetadata::MD_Namespace(TEXT("Namespace"));
 
+const FName FBlueprintMetadata::MD_ThreadSafe(TEXT("BlueprintThreadSafe"));
+const FName FBlueprintMetadata::MD_NotThreadSafe(TEXT("NotBlueprintThreadSafe"));
+
 //////////////////////////////////////////////////////////////////////////
 
 #define LOCTEXT_NAMESPACE "KismetSchema"
@@ -971,8 +974,8 @@ bool UEdGraphSchema_K2::HasFunctionAnyOutputParameter(const UFunction* InFunctio
 
 bool UEdGraphSchema_K2::FunctionCanBePlacedAsEvent(const UFunction* InFunction)
 {
-	// First check we are override-able, non-static and non-const
-	if (!InFunction || !CanKismetOverrideFunction(InFunction) || InFunction->HasAnyFunctionFlags(FUNC_Static|FUNC_Const))
+	// First check we are override-able, non-static, non-const and not marked thread safe
+	if (!InFunction || !CanKismetOverrideFunction(InFunction) || InFunction->HasAnyFunctionFlags(FUNC_Static|FUNC_Const) || FBlueprintEditorUtils::HasFunctionBlueprintThreadSafeMetaData(InFunction))
 	{
 		return false;
 	}
@@ -1267,6 +1270,29 @@ bool UEdGraphSchema_K2::DoesGraphSupportImpureFunctions(const UEdGraph* InGraph)
 	return bAllowImpureFuncs;
 }
 
+bool UEdGraphSchema_K2::IsGraphMarkedThreadSafe(const UEdGraph* InGraph) const
+{
+	TArray<UK2Node_FunctionEntry*> EntryNodes;
+	InGraph->GetNodesOfClass(EntryNodes);
+
+	for(UK2Node_FunctionEntry* EntryNode : EntryNodes)
+	{
+		if(EntryNode->MetaData.bThreadSafe)
+		{
+			return true;
+		}
+		else if(UFunction* Function = FFunctionFromNodeHelper::FunctionFromNode(EntryNode))
+		{
+			if(FBlueprintEditorUtils::HasFunctionBlueprintThreadSafeMetaData(Function))
+			{
+				return true;
+			}
+		}
+	}
+
+	return false;
+}
+
 bool UEdGraphSchema_K2::IsPropertyExposedOnSpawn(const FProperty* Property)
 {
 	Property = FBlueprintEditorUtils::GetMostUpToDateProperty(Property);
diff --git a/Engine/Source/Editor/BlueprintGraph/Public/BlueprintActionFilter.h b/Engine/Source/Editor/BlueprintGraph/Public/BlueprintActionFilter.h
--- a/Engine/Source/Editor/BlueprintGraph/Public/BlueprintActionFilter.h
+++ b/Engine/Source/Editor/BlueprintGraph/Public/BlueprintActionFilter.h
@@ -194,6 +194,12 @@ public:
 		 * imported namespaces will be filtered out.
 		 */
 		BPFILTER_RejectOutOfScopeMembers	= (1<<4),
+
+		/**
+		 * Rejects actions that would spawn nodes with thread safety characteristics
+		 * that are incompatible with the specification of the containing graph
+		 */
+		BPFILTER_RejectIncompatibleThreadSafety = (1<<5)
 	};
 	FBlueprintActionFilter(uint32 const Flags = 0x00);
 	
diff --git a/Engine/Source/Editor/Kismet/Private/BlueprintActionMenuUtils.cpp b/Engine/Source/Editor/Kismet/Private/BlueprintActionMenuUtils.cpp
--- a/Engine/Source/Editor/Kismet/Private/BlueprintActionMenuUtils.cpp
+++ b/Engine/Source/Editor/Kismet/Private/BlueprintActionMenuUtils.cpp
@@ -461,6 +461,11 @@ void FBlueprintActionMenuUtils::MakeContextMenu(FBlueprintActionContext const& C
 		FilterFlags |= FBlueprintActionFilter::BPFILTER_RejectOutOfScopeMembers;
 	}
 
+	if(bIsContextSensitive)
+	{
+		FilterFlags |= FBlueprintActionFilter::BPFILTER_RejectIncompatibleThreadSafety;
+	}
+
 	FBlueprintActionFilter MainMenuFilter(FilterFlags);
 	MainMenuFilter.Context = Context;
 	MainMenuFilter.Context.SelectedObjects.Empty();
diff --git a/Engine/Source/Editor/Kismet/Private/BlueprintDetailsCustomization.cpp b/Engine/Source/Editor/Kismet/Private/BlueprintDetailsCustomization.cpp
--- a/Engine/Source/Editor/Kismet/Private/BlueprintDetailsCustomization.cpp
+++ b/Engine/Source/Editor/Kismet/Private/BlueprintDetailsCustomization.cpp
@@ -3508,6 +3508,23 @@ void FBlueprintGraphActionDetails::CustomizeDetails( IDetailLayoutBuilder& Detai
 					.OnCheckStateChanged( this, &FBlueprintGraphActionDetails::OnIsExecFunctionModified )
 				];
 			}
+			if (IsThreadSafeFunctionVisible())
+			{
+				Category.AddCustomRow( LOCTEXT( "FunctionThreadSafe_Tooltip", "Thread Safe" ), true )
+				.NameContent()
+				[
+					SNew(STextBlock)
+					.Text( LOCTEXT( "FunctionThreadSafe_Tooltip", "Thread Safe" ) )
+				 	.ToolTipText( LOCTEXT("FunctionIsThreadSafe_Tooltip", "Enable thread-safety checks on this function. Only thread-safe functions and operations are allowed in this function.") )
+					.Font( IDetailLayoutBuilder::GetDetailFont() )
+				]
+				.ValueContent()
+				[
+					SNew( SCheckBox )
+					.IsChecked( this, &FBlueprintGraphActionDetails::GetIsThreadSafeFunction )
+					.OnCheckStateChanged( this, &FBlueprintGraphActionDetails::OnIsThreadSafeFunctionModified )
+				];
+			}
 		}
 
 		if (bIsCustomEvent)
@@ -5091,6 +5108,48 @@ ECheckBoxState FBlueprintGraphActionDetails::GetIsExecFunction() const
 	return (EntryNode->GetFunctionFlags() & FUNC_Exec) ? ECheckBoxState::Checked : ECheckBoxState::Unchecked;
 }
 
+bool FBlueprintGraphActionDetails::IsThreadSafeFunctionVisible() const
+{
+	bool bSupportedType = false;
+	bool bIsEditable = false;
+	
+	UK2Node_EditablePinBase* FunctionEntryNode = FunctionEntryNodePtr.Get();
+	if(FunctionEntryNode)
+	{
+		bSupportedType = FunctionEntryNode->IsA<UK2Node_FunctionEntry>();
+		bIsEditable = FunctionEntryNode->IsEditable();
+	}
+
+	return bIsEditable && bSupportedType;
+}
+
+void FBlueprintGraphActionDetails::OnIsThreadSafeFunctionModified(const ECheckBoxState NewCheckedState)
+{
+	if( FunctionEntryNodePtr.IsValid() )
+	{
+		const bool bThreadSafe = NewCheckedState == ECheckBoxState::Checked;
+		const FText TransactionType = bThreadSafe ? LOCTEXT( "DisableThreadSafe", "Disable Thread Safe" ) : LOCTEXT( "EnableThreadSafe", "Enable Thread Safe" );
+
+		if( UK2Node_FunctionEntry* EntryPoint = Cast<UK2Node_FunctionEntry>(FunctionEntryNodePtr.Get()) )
+		{
+			const FScopedTransaction Transaction( TransactionType );
+			EntryPoint->MetaData.bThreadSafe = bThreadSafe;
+			FBlueprintEditorUtils::MarkBlueprintAsModified( EntryPoint->GetBlueprint() );
+		}
+	}
+}
+
+ECheckBoxState FBlueprintGraphActionDetails::GetIsThreadSafeFunction() const
+{
+	UK2Node_EditablePinBase* FunctionEntryNode = FunctionEntryNodePtr.Get();
+	UK2Node_FunctionEntry* EntryNode = Cast<UK2Node_FunctionEntry>(FunctionEntryNode);
+	if (!EntryNode)
+	{
+		return ECheckBoxState::Undetermined;
+	}
+	return EntryNode->MetaData.bThreadSafe ? ECheckBoxState::Checked : ECheckBoxState::Unchecked;
+}
+
 FReply FBaseBlueprintGraphActionDetails::OnAddNewInputClicked()
 {
 	UK2Node_EditablePinBase* FunctionEntryNode = FunctionEntryNodePtr.Get();
diff --git a/Engine/Source/Editor/Kismet/Private/BlueprintDetailsCustomization.h b/Engine/Source/Editor/Kismet/Private/BlueprintDetailsCustomization.h
--- a/Engine/Source/Editor/Kismet/Private/BlueprintDetailsCustomization.h
+++ b/Engine/Source/Editor/Kismet/Private/BlueprintDetailsCustomization.h
@@ -614,6 +614,10 @@ private:
 	void OnIsExecFunctionModified(const ECheckBoxState NewCheckedState);
 	ECheckBoxState GetIsExecFunction() const;
 
+	bool IsThreadSafeFunctionVisible() const;
+	void OnIsThreadSafeFunctionModified(const ECheckBoxState NewCheckedState);
+	ECheckBoxState GetIsThreadSafeFunction() const;
+	
 	/** Determines if the selected event is identified as editor callable */
 	ECheckBoxState GetIsEditorCallableEvent() const;
 
diff --git a/Engine/Source/Editor/KismetCompiler/Private/KismetCompiler.cpp b/Engine/Source/Editor/KismetCompiler/Private/KismetCompiler.cpp
--- a/Engine/Source/Editor/KismetCompiler/Private/KismetCompiler.cpp
+++ b/Engine/Source/Editor/KismetCompiler/Private/KismetCompiler.cpp
@@ -2229,6 +2229,12 @@ void FKismetCompilerContext::CompileFunction(FKismetFunctionContext& Context)
 		Context.UnsortedSeparateExecutionGroups = FKismetCompilerUtilities::FindUnsortedSeparateExecutionGroups(Context.LinearExecutionList);
 	}
 
+	// Propagate thread-safe flags in this first pass. Also gets called from SetCalculatedMetaDataAndFlags in the second
+	// pass to catch skeleton class generation
+	if (Context.EntryPoint->MetaData.bThreadSafe)
+	{
+		Context.Function->SetMetaData(FBlueprintMetadata::MD_ThreadSafe, TEXT("true"));
+	}
 }
 
 /**
@@ -2276,6 +2282,12 @@ void FKismetCompilerContext::FinishCompilingFunction(FKismetFunctionContext& Con
 		}
 	}
 #endif//VALIDATE_UBER_GRAPH_PERSISTENT_FRAME
+
+	// Check thread safety
+	if (FBlueprintEditorUtils::HasFunctionBlueprintThreadSafeMetaData(Context.Function))
+	{
+		FKismetCompilerUtilities::CheckFunctionThreadSafety(Context, MessageLog);
+	}
 }
 
 void FKismetCompilerContext::SetCalculatedMetaDataAndFlags(UFunction* Function, UK2Node_FunctionEntry* EntryNode, const UEdGraphSchema_K2* K2Schema)
@@ -2360,6 +2372,12 @@ void FKismetCompilerContext::SetCalculatedMetaDataAndFlags(UFunction* Function,
 			Function->SetMetaData(FBlueprintMetadata::MD_DeprecationMessage, *(EntryNode->MetaData.DeprecationMessage));
 		}
 	}
+
+	// Just copy the thread-safe flag with no checking, we verify thread safety elsewhere
+	if (EntryNode->MetaData.bThreadSafe)
+	{
+		Function->SetMetaData(FBlueprintMetadata::MD_ThreadSafe, TEXT("true"));
+	}
 	
 	if (UEdGraphPin* WorldContextPin = EntryNode->GetAutoWorldContextPin())
 	{
@@ -2372,6 +2390,12 @@ void FKismetCompilerContext::SetCalculatedMetaDataAndFlags(UFunction* Function,
 	{
 		// Copy metadata from parent function as well
 		UMetaData::CopyMetadata(OverriddenFunction, Function);
+
+		// Native thread safety metadata works on the class level as well as functions, so we need to check that too.
+		if(FBlueprintEditorUtils::HasFunctionBlueprintThreadSafeMetaData(OverriddenFunction))
+		{
+			Function->SetMetaData(FBlueprintMetadata::MD_ThreadSafe, TEXT("true"));
+		}
 	}
 }
 
diff --git a/Engine/Source/Editor/KismetCompiler/Private/KismetCompilerMisc.cpp b/Engine/Source/Editor/KismetCompiler/Private/KismetCompilerMisc.cpp
--- a/Engine/Source/Editor/KismetCompiler/Private/KismetCompilerMisc.cpp
+++ b/Engine/Source/Editor/KismetCompiler/Private/KismetCompilerMisc.cpp
@@ -1694,6 +1694,209 @@ void FKismetCompilerUtilities::UpdateDependentBlueprints(UBlueprint* ForBP)
 	}
 }
 
+bool FKismetCompilerUtilities::CheckFunctionThreadSafety(const FKismetFunctionContext& InContext, FCompilerResultsLog& InMessageLog, bool InbEmitErrors)
+{
+	bool bIsThreadSafe = true;
+	
+	for(const TPair<UEdGraphNode*, TArray<FBlueprintCompiledStatement*>>& StatementPair : InContext.StatementsPerNode)
+	{
+		bIsThreadSafe &= CheckFunctionCompiledStatementsThreadSafety(StatementPair.Key, InContext.SourceGraph, StatementPair.Value, InMessageLog, InbEmitErrors);
+	}
+
+	return bIsThreadSafe;
+}
+
+// Helper used to emit to log as errors/warnings 
+struct FLogThreadSafetyHelper
+{
+	FLogThreadSafetyHelper(FCompilerResultsLog& InLog, bool bInEmitErrors)
+		: Log(InLog)
+		, bEmitErrors(bInEmitErrors)
+	{}
+		
+	FCompilerResultsLog& Log;
+	bool bEmitErrors;
+		
+	template<typename... Args>
+	void Message(const TCHAR* Format, Args... args)
+	{
+		if(bEmitErrors)
+		{
+			Log.Error(Format, args...);
+		}
+		else
+		{
+			Log.Warning(Format, args...);
+		}
+	}
+};
+
+#define LOG_THREADSAFETY_HELPER(EmitErrors, CategoryName, Format, ...) \
+	if(EmitErrors) \
+	{ \
+		UE_LOG(CategoryName, Error, Format, ##__VA_ARGS__); \
+	} \
+	else \
+	{ \
+		UE_LOG(CategoryName, Warning, Format, ##__VA_ARGS__); \
+	}
+
+bool FKismetCompilerUtilities::CheckFunctionCompiledStatementsThreadSafety(const UEdGraphNode* InNode, const UEdGraph* InSourceGraph, const TArray<FBlueprintCompiledStatement*>& InStatements, FCompilerResultsLog& InMessageLog, bool InbEmitErrors)
+{
+	bool bIsThreadSafe = true;
+
+	const FText GenericThreadSafetyErrorOneParam = LOCTEXT("ThreadSafety_Error_Generic", "This is not thread safe when compiled. See the output log for more details.");
+
+	FLogThreadSafetyHelper LogHelper(InMessageLog, InbEmitErrors);
+	
+	for(const FBlueprintCompiledStatement* Statement : InStatements)
+	{
+		auto LogDelegateUsage = [&bIsThreadSafe, &LogHelper, InNode, InbEmitErrors]()
+		{
+			LogHelper.Message(*LOCTEXT("ThreadSafety_Error_Delegate", "@@ Delegate usage is not thread-safe").ToString(), InNode);
+			bIsThreadSafe = false;
+		};
+
+		auto CheckForInvalidInstancedObjectContext = [&bIsThreadSafe, &LogHelper, InNode, &GenericThreadSafetyErrorOneParam, InbEmitErrors](const FBPTerminal* InTerm)
+		{
+			const FBPTerminal* Context = InTerm;
+			while(Context)
+			{
+				if(Context != nullptr && Context->IsObjectContextType() && Context->IsInstancedVarTerm())
+				{
+					if(Context->SourcePin && Context->SourcePin->GetOwningNode())
+					{
+						// @TODO: we could possibly make exceptions for 'assets' here
+						LogHelper.Message(*LOCTEXT("ThreadSafety_Error_InstancedObjectWithPin", "@@ Accessing an object reference is not thread-safe").ToString(), Context->SourcePin->GetOwningNode());
+					}
+					else
+					{
+						LogHelper.Message(*GenericThreadSafetyErrorOneParam.ToString(), InNode);
+						LOG_THREADSAFETY_HELPER(InbEmitErrors, LogBlueprint, TEXT("Expression that accesses an instanced object context is not thread-safe"));
+					}
+					bIsThreadSafe = false;
+				}
+
+				Context = Context->Context;
+			}
+		};
+
+		switch (Statement->Type)
+		{
+			case KCST_Nop:
+				break;
+			case KCST_CallFunction:
+			{
+				check(Statement->FunctionToCall);
+
+				if(Statement->FunctionContext)
+				{
+					CheckForInvalidInstancedObjectContext(Statement->FunctionContext);
+				}
+				
+				// Check RHS (function inputs) for invalid object access 
+				for(FBPTerminal* RHSTerm : Statement->RHS)
+				{
+					if(RHSTerm->Context)
+					{
+						CheckForInvalidInstancedObjectContext(RHSTerm->Context);
+					}
+				}
+
+				if(!FBlueprintEditorUtils::HasFunctionBlueprintThreadSafeMetaData(Statement->FunctionToCall))
+				{
+					LogHelper.Message(*LOCTEXT("ThreadSafety_Error_NativeFunction", "@@ Non-thread safe function @@ called from thread-safe graph @@").ToString(), InNode, Statement->FunctionToCall, InSourceGraph);
+					bIsThreadSafe = false;
+				}
+				break;
+			}
+			case KCST_Assignment:
+			{
+				// Check LHS/RHS for invalid object access.
+				if(Statement->LHS && Statement->LHS->Context)
+				{
+					CheckForInvalidInstancedObjectContext(Statement->LHS->Context);
+				}
+					
+				for(FBPTerminal* RHSTerm : Statement->RHS)
+				{
+					if(RHSTerm->Context)
+					{
+						CheckForInvalidInstancedObjectContext(RHSTerm->Context);
+					}
+				}
+				break;
+			}
+			case KCST_CompileError:
+		    case KCST_UnconditionalGoto:
+			case KCST_PushState:
+			case KCST_GotoIfNot:
+			case KCST_Return:
+			case KCST_EndOfThread:
+			case KCST_Comment:
+			case KCST_ComputedGoto:
+			case KCST_EndOfThreadIfNot:
+			case KCST_DebugSite:
+			case KCST_CastObjToInterface:
+			case KCST_DynamicCast:
+			case KCST_ObjectToBool:
+				break;
+			case KCST_AddMulticastDelegate:
+			case KCST_ClearMulticastDelegate:
+				LogDelegateUsage();
+				break;
+			case KCST_WireTraceSite:
+				break;
+			case KCST_BindDelegate:
+			case KCST_RemoveMulticastDelegate:
+			case KCST_CallDelegate:
+				LogDelegateUsage();
+				break;
+			case KCST_CreateArray:
+			case KCST_CrossInterfaceCast:
+			case KCST_MetaCast:
+				break;
+			case KCST_AssignmentOnPersistentFrame:
+				LogHelper.Message(*GenericThreadSafetyErrorOneParam.ToString(), InSourceGraph);
+				LOG_THREADSAFETY_HELPER(InbEmitErrors, LogBlueprint, TEXT("Persistent frame assignment is not supported in thread-safe function"));
+				bIsThreadSafe = false;
+				break;
+			case KCST_CastInterfaceToObj:
+			case KCST_GotoReturn:
+			case KCST_GotoReturnIfNot:
+			case KCST_SwitchValue:
+				break;
+			case KCST_InstrumentedEvent:
+			case KCST_InstrumentedEventStop:
+			case KCST_InstrumentedPureNodeEntry:
+			case KCST_InstrumentedWireEntry:
+			case KCST_InstrumentedWireExit:
+			case KCST_InstrumentedStatePush:
+			case KCST_InstrumentedStateRestore:
+			case KCST_InstrumentedStateReset:
+			case KCST_InstrumentedStateSuspend:
+			case KCST_InstrumentedStatePop:
+			case KCST_InstrumentedTunnelEndOfThread:
+				// No way to test instrumentation right now, so this can potentially be removed
+				LogHelper.Message(*GenericThreadSafetyErrorOneParam.ToString(), InSourceGraph);
+				LOG_THREADSAFETY_HELPER(InbEmitErrors, LogBlueprint, TEXT("Instrumentation not supported in thread-safe function"));
+				bIsThreadSafe = false;
+				break;
+			case KCST_ArrayGetByRef:
+			case KCST_CreateSet:
+			case KCST_CreateMap:
+				break;
+			default:
+				LogHelper.Message(*GenericThreadSafetyErrorOneParam.ToString(), InSourceGraph);
+				LOG_THREADSAFETY_HELPER(InbEmitErrors, LogBlueprint, TEXT("Non-thread safe unknown statement type %d (from %s) used in thread-safe context %s"), (int32)Statement->Type, *InNode->GetName(), *InSourceGraph->GetName());
+				bIsThreadSafe = false;
+				break;
+		}
+	}
+
+	return bIsThreadSafe;
+}
+
 //////////////////////////////////////////////////////////////////////////
 // FNodeHandlingFunctor
 
diff --git a/Engine/Source/Editor/KismetCompiler/Public/KismetCompilerMisc.h b/Engine/Source/Editor/KismetCompiler/Public/KismetCompilerMisc.h
--- a/Engine/Source/Editor/KismetCompiler/Public/KismetCompilerMisc.h
+++ b/Engine/Source/Editor/KismetCompiler/Public/KismetCompilerMisc.h
@@ -111,6 +111,12 @@ public:
 
 	/** Add this BP to any BPs that it in*/
 	static void UpdateDependentBlueprints(UBlueprint* BP);
+
+	// Check the passed-in function to verify its thread safety. This makes sure that it only uses/calls thread-safe functions/nodes.
+	static bool CheckFunctionThreadSafety(const FKismetFunctionContext& InContext, FCompilerResultsLog& InMessageLog, bool InbEmitErrors = true);
+
+	// Helper function used by CheckFunctionThreadSafety. Split out to allow the ability to examine individual compiled statement lists (e.g. for the ubergraph)
+	static bool CheckFunctionCompiledStatementsThreadSafety(const UEdGraphNode* InNode, const UEdGraph* InSourceGraph, const TArray<FBlueprintCompiledStatement*>& InStatements, FCompilerResultsLog& InMessageLog, bool InbEmitErrors = true);
 };
 
 //////////////////////////////////////////////////////////////////////////
diff --git a/Engine/Source/Editor/UnrealEd/Private/Kismet2/BlueprintEditorUtils.cpp b/Engine/Source/Editor/UnrealEd/Private/Kismet2/BlueprintEditorUtils.cpp
--- a/Engine/Source/Editor/UnrealEd/Private/Kismet2/BlueprintEditorUtils.cpp
+++ b/Engine/Source/Editor/UnrealEd/Private/Kismet2/BlueprintEditorUtils.cpp
@@ -2410,7 +2410,7 @@ void FBlueprintEditorUtils::CreateMatchingFunction(UK2Node_CallFunction* InNode,
 
 bool FBlueprintEditorUtils::IsFunctionConvertableToEvent(UBlueprint* const BlueprintObj, UFunction* const Function)
 {
-	return BlueprintObj && BlueprintObj->BlueprintType != BPTYPE_FunctionLibrary && BlueprintObj->BlueprintType != BPTYPE_Interface && Function;
+	return BlueprintObj && BlueprintObj->BlueprintType != BPTYPE_FunctionLibrary && BlueprintObj->BlueprintType != BPTYPE_Interface && Function && !HasFunctionBlueprintThreadSafeMetaData(Function);
 }
 
 UFunction* FBlueprintEditorUtils::FindFunctionInImplementedInterfaces(const UBlueprint* Blueprint, const FName& FunctionName, bool * bOutInvalidInterface, bool bGetAllInterfaces)
@@ -10074,4 +10074,24 @@ void FBlueprintEditorUtils::RecombineNestedSubPins(UK2Node* Node)
 	}
 }
 
+bool FBlueprintEditorUtils::HasFunctionBlueprintThreadSafeMetaData(const UFunction* InFunction)
+{
+	if(InFunction)
+	{
+		const bool bHasThreadSafeMetaData = InFunction->HasMetaData(FBlueprintMetadata::MD_ThreadSafe);
+		const bool bHasNotThreadSafeMetaData = InFunction->HasMetaData(FBlueprintMetadata::MD_NotThreadSafe);
+		const bool bClassHasThreadSafeMetaData = InFunction->GetOwnerClass() && InFunction->GetOwnerClass()->HasMetaData(FBlueprintMetadata::MD_ThreadSafe);
+
+		// Native (or BP event) functions need to just have the correct class/function metadata
+		const bool bThreadSafeNative = InFunction->HasAnyFunctionFlags(FUNC_Native | FUNC_BlueprintEvent) && (bHasThreadSafeMetaData || (bClassHasThreadSafeMetaData && !bHasNotThreadSafeMetaData));
+
+		// Script functions get their flag propagated from their entry point, and dont pay heed to class metadata
+		const bool bThreadSafeScript = !InFunction->HasAnyFunctionFlags(FUNC_Native) && bHasThreadSafeMetaData;
+		
+		return bThreadSafeNative || bThreadSafeScript;
+	}
+	
+	return false;
+};
+
 #undef LOCTEXT_NAMESPACE
diff --git a/Engine/Source/Editor/UnrealEd/Public/Kismet2/BlueprintEditorUtils.h b/Engine/Source/Editor/UnrealEd/Public/Kismet2/BlueprintEditorUtils.h
--- a/Engine/Source/Editor/UnrealEd/Public/Kismet2/BlueprintEditorUtils.h
+++ b/Engine/Source/Editor/UnrealEd/Public/Kismet2/BlueprintEditorUtils.h
@@ -1885,6 +1885,14 @@ public:
 	 * Returns true if this BP is currently based on a type that returns true for the UObject::ImplementsGetWorld() call:
 	 */
 	static bool ImplementsGetWorld(const UBlueprint* BP);
+
+	/*
+	 * Checks a function for 'blueprint thread safety'. Note this is not strict thread safety and does not perform any
+	 * analysis on the function, it only inspects the metadata. Analysis is performed by the compiler for BP functions
+	 * that are marked thread safe and metadata is used both here and in the compiler for native functions that have
+	 * been validated for thread safe use in a blueprint.
+	 */
+	static bool HasFunctionBlueprintThreadSafeMetaData(const UFunction* InFunction);
 };
 
 struct UNREALED_API FBlueprintDuplicationScopeFlags
diff --git a/Engine/Source/Runtime/Engine/Classes/Animation/AnimInstance.h b/Engine/Source/Runtime/Engine/Classes/Animation/AnimInstance.h
--- a/Engine/Source/Runtime/Engine/Classes/Animation/AnimInstance.h
+++ b/Engine/Source/Runtime/Engine/Classes/Animation/AnimInstance.h
@@ -367,7 +367,7 @@ struct FMontageEvaluationState
 	bool bIsActive;
 };
 
-UCLASS(transient, Blueprintable, hideCategories=AnimInstance, BlueprintType, meta=(BlueprintThreadSafe), Within=SkeletalMeshComponent)
+UCLASS(transient, Blueprintable, hideCategories=AnimInstance, BlueprintType, Within=SkeletalMeshComponent)
 class ENGINE_API UAnimInstance : public UObject
 {
 	GENERATED_UCLASS_BODY()
@@ -498,7 +498,7 @@ public:
 
 	// kismet event functions
 
-	UFUNCTION(BlueprintCallable, Category = "Animation", meta=(NotBlueprintThreadSafe))
+	UFUNCTION(BlueprintCallable, Category = "Animation") 
 	virtual APawn* TryGetPawnOwner() const;
 
 	/** 
@@ -564,11 +564,11 @@ private:
 
 public:
 	/** Returns the owning actor of this AnimInstance */
-	UFUNCTION(BlueprintCallable, Category = "Animation", meta=(NotBlueprintThreadSafe))
+	UFUNCTION(BlueprintCallable, Category = "Animation")
 	AActor* GetOwningActor() const;
 	
 	// Returns the skeletal mesh component that has created this AnimInstance
-	UFUNCTION(BlueprintCallable, Category = "Animation", meta=(NotBlueprintThreadSafe))
+	UFUNCTION(BlueprintCallable, Category = "Animation")
 	USkeletalMeshComponent* GetOwningComponent() const;
 
 public:
@@ -681,20 +681,20 @@ public:
 	void Montage_SetPlayRate(const UAnimMontage* Montage, float NewPlayRate = 1.f);
 
 	/** Returns true if the animation montage is active. If the Montage reference is NULL, it will return true if any Montage is active. */
-	UFUNCTION(BlueprintPure, Category="Montage", meta=(NotBlueprintThreadSafe))
+	UFUNCTION(BlueprintPure, Category="Montage")
 	bool Montage_IsActive(const UAnimMontage* Montage) const;
 
 	/** Returns true if the animation montage is currently active and playing. 
 	If reference is NULL, it will return true is ANY montage is currently active and playing. */
-	UFUNCTION(BlueprintPure, Category="Montage", meta = (NotBlueprintThreadSafe))
+	UFUNCTION(BlueprintPure, Category="Montage")
 	bool Montage_IsPlaying(const UAnimMontage* Montage) const;
 
 	/** Returns the name of the current animation montage section. */
-	UFUNCTION(BlueprintPure, Category="Montage", meta = (NotBlueprintThreadSafe))
+	UFUNCTION(BlueprintPure, Category="Montage")
 	FName Montage_GetCurrentSection(const UAnimMontage* Montage = NULL) const;
 
 	/** Get Current Montage Position */
-	UFUNCTION(BlueprintPure, Category = "Montage", meta = (NotBlueprintThreadSafe))
+	UFUNCTION(BlueprintPure, Category = "Montage")
 	float Montage_GetPosition(const UAnimMontage* Montage) const;
 	
 	/** Set position. */
@@ -702,27 +702,27 @@ public:
 	void Montage_SetPosition(const UAnimMontage* Montage, float NewPosition);
 	
 	/** return true if Montage is not currently active. (not valid or blending out) */
-	UFUNCTION(BlueprintPure, Category = "Montage", meta = (NotBlueprintThreadSafe))
+	UFUNCTION(BlueprintPure, Category = "Montage")
 	bool Montage_GetIsStopped(const UAnimMontage* Montage) const;
 
 	/** Get the current blend time of the Montage.
 	If Montage reference is NULL, it will return the current blend time on the first active Montage found. */
-	UFUNCTION(BlueprintPure, Category = "Montage", meta = (NotBlueprintThreadSafe))
+	UFUNCTION(BlueprintPure, Category = "Montage")
 	float Montage_GetBlendTime(const UAnimMontage* Montage) const;
 
 	/** Get PlayRate for Montage.
 	If Montage reference is NULL, PlayRate for any Active Montage will be returned.
 	If Montage is not playing, 0 is returned. */
-	UFUNCTION(BlueprintPure, Category = "Montage", meta = (NotBlueprintThreadSafe))
+	UFUNCTION(BlueprintPure, Category = "Montage")
 	float Montage_GetPlayRate(const UAnimMontage* Montage) const;
 
 	/** Returns true if any montage is playing currently. Doesn't mean it's active though, it could be blending out. */
-	UFUNCTION(BlueprintPure, Category = "Montage", meta = (NotBlueprintThreadSafe))
+	UFUNCTION(BlueprintPure, Category = "Montage")
 	bool IsAnyMontagePlaying() const;
 
 	/** Get a current Active Montage in this AnimInstance. 
 		Note that there might be multiple Active at the same time. This will only return the first active one it finds. **/
-	UFUNCTION(BlueprintPure, Category = "Montage", meta = (NotBlueprintThreadSafe))
+	UFUNCTION(BlueprintPure, Category = "Montage")
 	UAnimMontage* GetCurrentActiveMontage() const;
 
 	/** Called when a montage starts blending out, whether interrupted or finished */
@@ -956,67 +956,67 @@ public:
 	 */
 
 	/** Gets the length in seconds of the asset referenced in an asset player node */
-	UFUNCTION(BlueprintPure, Category="Asset Player", meta=(DisplayName="Length", BlueprintInternalUseOnly="true", AnimGetter="true"))
+	UFUNCTION(BlueprintPure, Category="Asset Player", meta=(DisplayName="Length", BlueprintInternalUseOnly="true", AnimGetter="true", BlueprintThreadSafe))
 	float GetInstanceAssetPlayerLength(int32 AssetPlayerIndex);
 
 	/** Get the current accumulated time in seconds for an asset player node */
-	UFUNCTION(BlueprintPure, Category="Asset Player", meta = (DisplayName = "Current Time", BlueprintInternalUseOnly = "true", AnimGetter = "true"))
+	UFUNCTION(BlueprintPure, Category="Asset Player", meta = (DisplayName = "Current Time", BlueprintInternalUseOnly = "true", AnimGetter = "true", BlueprintThreadSafe))
 	float GetInstanceAssetPlayerTime(int32 AssetPlayerIndex);
 
 	/** Get the current accumulated time as a fraction for an asset player node */
-	UFUNCTION(BlueprintPure, Category="Asset Player", meta=(DisplayName="Current Time (ratio)", BlueprintInternalUseOnly="true", AnimGetter="true"))
+	UFUNCTION(BlueprintPure, Category="Asset Player", meta=(DisplayName="Current Time (ratio)", BlueprintInternalUseOnly="true", AnimGetter="true", BlueprintThreadSafe))
 	float GetInstanceAssetPlayerTimeFraction(int32 AssetPlayerIndex);
 
 	/** Get the time in seconds from the end of an animation in an asset player node */
-	UFUNCTION(BlueprintPure, Category="Asset Player", meta=(DisplayName="Time Remaining", BlueprintInternalUseOnly="true", AnimGetter="true"))
+	UFUNCTION(BlueprintPure, Category="Asset Player", meta=(DisplayName="Time Remaining", BlueprintInternalUseOnly="true", AnimGetter="true", BlueprintThreadSafe))
 	float GetInstanceAssetPlayerTimeFromEnd(int32 AssetPlayerIndex);
 
 	/** Get the time as a fraction of the asset length of an animation in an asset player node */
-	UFUNCTION(BlueprintPure, Category="Asset Player", meta=(DisplayName="Time Remaining (ratio)", BlueprintInternalUseOnly="true", AnimGetter="true"))
+	UFUNCTION(BlueprintPure, Category="Asset Player", meta=(DisplayName="Time Remaining (ratio)", BlueprintInternalUseOnly="true", AnimGetter="true", BlueprintThreadSafe))
 	float GetInstanceAssetPlayerTimeFromEndFraction(int32 AssetPlayerIndex);
 
 	/** Get the blend weight of a specified state machine */
-	UFUNCTION(BlueprintPure, Category = "States", meta = (DisplayName = "Machine Weight", BlueprintInternalUseOnly = "true", AnimGetter = "true"))
+	UFUNCTION(BlueprintPure, Category = "States", meta = (DisplayName = "Machine Weight", BlueprintInternalUseOnly = "true", AnimGetter = "true", BlueprintThreadSafe))
 	float GetInstanceMachineWeight(int32 MachineIndex);
 
 	/** Get the blend weight of a specified state */
-	UFUNCTION(BlueprintPure, Category="States", meta = (DisplayName="State Weight", BlueprintInternalUseOnly = "true", AnimGetter="true"))
+	UFUNCTION(BlueprintPure, Category="States", meta = (DisplayName="State Weight", BlueprintInternalUseOnly = "true", AnimGetter="true", BlueprintThreadSafe))
 	float GetInstanceStateWeight(int32 MachineIndex, int32 StateIndex);
 
 	/** Get the current elapsed time of a state within the specified state machine */
-	UFUNCTION(BlueprintPure, Category="States", meta = (DisplayName="Current State Time", BlueprintInternalUseOnly = "true", AnimGetter="true", GetterContext="Transition"))
+	UFUNCTION(BlueprintPure, Category="States", meta = (DisplayName="Current State Time", BlueprintInternalUseOnly = "true", AnimGetter="true", GetterContext="Transition", BlueprintThreadSafe))
 	float GetInstanceCurrentStateElapsedTime(int32 MachineIndex);
 
 	/** Get the crossfade duration of a specified transition */
-	UFUNCTION(BlueprintPure, Category="Transitions", meta = (DisplayName="Get Transition Crossfade Duration", BlueprintInternalUseOnly = "true", AnimGetter="true"))
+	UFUNCTION(BlueprintPure, Category="Transitions", meta = (DisplayName="Get Transition Crossfade Duration", BlueprintInternalUseOnly = "true", AnimGetter="true", BlueprintThreadSafe))
 	float GetInstanceTransitionCrossfadeDuration(int32 MachineIndex, int32 TransitionIndex);
 
 	/** Get the elapsed time in seconds of a specified transition */
-	UFUNCTION(BlueprintPure, Category="Transitions", meta = (DisplayName="Get Transition Time Elapsed", BlueprintInternalUseOnly = "true", AnimGetter="true", GetterContext="CustomBlend"))
+	UFUNCTION(BlueprintPure, Category="Transitions", meta = (DisplayName="Get Transition Time Elapsed", BlueprintInternalUseOnly = "true", AnimGetter="true", GetterContext="CustomBlend", BlueprintThreadSafe))
 	float GetInstanceTransitionTimeElapsed(int32 MachineIndex, int32 TransitionIndex);
 
 	/** Get the elapsed time as a fraction of the crossfade duration of a specified transition */
-	UFUNCTION(BlueprintPure, Category="Transitions", meta = (DisplayName="Get Transition Time Elapsed (ratio)", BlueprintInternalUseOnly = "true", AnimGetter="true", GetterContext="CustomBlend"))
+	UFUNCTION(BlueprintPure, Category="Transitions", meta = (DisplayName="Get Transition Time Elapsed (ratio)", BlueprintInternalUseOnly = "true", AnimGetter="true", GetterContext="CustomBlend", BlueprintThreadSafe))
 	float GetInstanceTransitionTimeElapsedFraction(int32 MachineIndex, int32 TransitionIndex);
 
 	/** Get the time remaining in seconds for the most relevant animation in the source state */
-	UFUNCTION(BlueprintPure, Category="Asset Player", meta = (BlueprintInternalUseOnly = "true", AnimGetter="true", GetterContext="Transition"))
+	UFUNCTION(BlueprintPure, Category="Asset Player", meta = (BlueprintInternalUseOnly = "true", AnimGetter="true", GetterContext="Transition", BlueprintThreadSafe))
 	float GetRelevantAnimTimeRemaining(int32 MachineIndex, int32 StateIndex);
 
 	/** Get the time remaining as a fraction of the duration for the most relevant animation in the source state */
-	UFUNCTION(BlueprintPure, Category = "Asset Player", meta = (BlueprintInternalUseOnly = "true", AnimGetter = "true", GetterContext = "Transition"))
+	UFUNCTION(BlueprintPure, Category = "Asset Player", meta = (BlueprintInternalUseOnly = "true", AnimGetter = "true", GetterContext = "Transition", BlueprintThreadSafe))
 	float GetRelevantAnimTimeRemainingFraction(int32 MachineIndex, int32 StateIndex);
 
 	/** Get the length in seconds of the most relevant animation in the source state */
-	UFUNCTION(BlueprintPure, Category = "Asset Player", meta = (BlueprintInternalUseOnly = "true", AnimGetter = "true", GetterContext = "Transition"))
+	UFUNCTION(BlueprintPure, Category = "Asset Player", meta = (BlueprintInternalUseOnly = "true", AnimGetter = "true", GetterContext = "Transition", BlueprintThreadSafe))
 	float GetRelevantAnimLength(int32 MachineIndex, int32 StateIndex);
 
 	/** Get the current accumulated time in seconds for the most relevant animation in the source state */
-	UFUNCTION(BlueprintPure, Category = "Asset Player", meta = (BlueprintInternalUseOnly = "true", AnimGetter = "true", GetterContext = "Transition"))
+	UFUNCTION(BlueprintPure, Category = "Asset Player", meta = (BlueprintInternalUseOnly = "true", AnimGetter = "true", GetterContext = "Transition", BlueprintThreadSafe))
 	float GetRelevantAnimTime(int32 MachineIndex, int32 StateIndex);
 
 	/** Get the current accumulated time as a fraction of the length of the most relevant animation in the source state */
-	UFUNCTION(BlueprintPure, Category = "Asset Player", meta = (BlueprintInternalUseOnly = "true", AnimGetter = "true", GetterContext = "Transition"))
+	UFUNCTION(BlueprintPure, Category = "Asset Player", meta = (BlueprintInternalUseOnly = "true", AnimGetter = "true", GetterContext = "Transition", BlueprintThreadSafe))
 	float GetRelevantAnimTimeFraction(int32 MachineIndex, int32 StateIndex);
 
 	/** Get whether a particular notify state was active in any state machine last tick.*/
@@ -1094,22 +1094,22 @@ public:
 
 public:
 	/** Returns the value of a named curve. */
-	UFUNCTION(BlueprintPure, Category="Animation")
+	UFUNCTION(BlueprintPure, Category="Animation", meta=(BlueprintThreadSafe))
 	float GetCurveValue(FName CurveName) const;
 
 	/** This returns last up-to-date list of active curve names */
-	UFUNCTION(BlueprintPure, Category = "Animation")
+	UFUNCTION(BlueprintPure, Category = "Animation", meta=(BlueprintThreadSafe))
 	void GetActiveCurveNames(EAnimCurveType CurveType, TArray<FName>& OutNames) const;
 
 	/* This returns all curve names */
-	UFUNCTION(BlueprintPure, Category = "Animation")
+	UFUNCTION(BlueprintPure, Category = "Animation", meta=(BlueprintThreadSafe))
 	void GetAllCurveNames(TArray<FName>& OutNames) const;
 
 	/** Returns value of named curved in OutValue, returns whether the curve was actually found or not. */
 	bool GetCurveValue(FName CurveName, float& OutValue) const;
 
 	/** Returns the name of a currently active state in a state machine. */
-	UFUNCTION(BlueprintPure, Category="Animation", meta=(BlueprintInternalUseOnly = "true", AnimGetter = "true"))
+	UFUNCTION(BlueprintPure, Category="Animation", meta=(BlueprintInternalUseOnly = "true", AnimGetter = "true", BlueprintThreadSafe))
 	FName GetCurrentStateName(int32 MachineIndex);
 
 	/** Sets a morph target to a certain weight. */
@@ -1121,7 +1121,7 @@ public:
 	void ClearMorphTargets();
 
 	UE_DEPRECATED(5.0, "Please use UKismetAnimationLibrary::CalculateDirection instead")
-	UFUNCTION(BlueprintCallable, Category="Animation")
+	UFUNCTION(BlueprintCallable, Category="Animation", meta=(BlueprintThreadSafe))
 	float CalculateDirection(const FVector& Velocity, const FRotator& BaseRotation) const;
 
 	//--- AI communication start ---//
@@ -1136,16 +1136,16 @@ public:
 	void UnlockAIResources(bool bUnlockMovement, bool UnlockAILogic);
 	//--- AI communication end ---//
 
-	UFUNCTION(BlueprintCallable, Category = "SyncGroup")
+	UFUNCTION(BlueprintCallable, Category = "SyncGroup", meta=(BlueprintThreadSafe))
 	bool GetTimeToClosestMarker(FName SyncGroup, FName MarkerName, float& OutMarkerTime) const;
 
-	UFUNCTION(BlueprintCallable, Category = "SyncGroup")
+	UFUNCTION(BlueprintCallable, Category = "SyncGroup", meta=(BlueprintThreadSafe))
 	bool HasMarkerBeenHitThisFrame(FName SyncGroup, FName MarkerName) const;
 
-	UFUNCTION(BlueprintCallable, Category = "SyncGroup")
+	UFUNCTION(BlueprintCallable, Category = "SyncGroup", meta=(BlueprintThreadSafe))
 	bool IsSyncGroupBetweenMarkers(FName InSyncGroupName, FName PreviousMarker, FName NextMarker, bool bRespectMarkerOrder = true) const;
 
-	UFUNCTION(BlueprintCallable, Category = "SyncGroup")
+	UFUNCTION(BlueprintCallable, Category = "SyncGroup", meta=(BlueprintThreadSafe))
 	FMarkerSyncAnimPosition GetSyncGroupPosition(FName InSyncGroupName) const;
 
 public:
@@ -1269,7 +1269,7 @@ public:
 	virtual void DisplayDebugInstance(FDisplayDebugManager& DisplayDebugManager, float& Indent);
 
 	/** Reset any dynamics running simulation-style updates (e.g. on teleport, time skip etc.) */
-	UFUNCTION(BlueprintCallable, Category = "Animation", meta = (NotBlueprintThreadSafe))
+	UFUNCTION(BlueprintCallable, Category = "Animation")
 	void ResetDynamics(ETeleportType InTeleportType);
 
 	UE_DEPRECATED(4.20, "Please use ResetDynamics with a ETeleportType argument")
diff --git a/Engine/Source/Runtime/Engine/Classes/Kismet/BlueprintMapLibrary.h b/Engine/Source/Runtime/Engine/Classes/Kismet/BlueprintMapLibrary.h
--- a/Engine/Source/Runtime/Engine/Classes/Kismet/BlueprintMapLibrary.h
+++ b/Engine/Source/Runtime/Engine/Classes/Kismet/BlueprintMapLibrary.h
@@ -10,7 +10,7 @@
 #include "Kismet/BlueprintFunctionLibrary.h"
 #include "BlueprintMapLibrary.generated.h"
 
-UCLASS()
+UCLASS(meta=(BlueprintThreadSafe))
 class ENGINE_API UBlueprintMapLibrary : public UBlueprintFunctionLibrary
 {
 	GENERATED_BODY()
diff --git a/Engine/Source/Runtime/Engine/Classes/Kismet/BlueprintSetLibrary.h b/Engine/Source/Runtime/Engine/Classes/Kismet/BlueprintSetLibrary.h
--- a/Engine/Source/Runtime/Engine/Classes/Kismet/BlueprintSetLibrary.h
+++ b/Engine/Source/Runtime/Engine/Classes/Kismet/BlueprintSetLibrary.h
@@ -10,7 +10,7 @@
 #include "Kismet/BlueprintFunctionLibrary.h"
 #include "BlueprintSetLibrary.generated.h"
 
-UCLASS()
+UCLASS(meta=(BlueprintThreadSafe))
 class ENGINE_API UBlueprintSetLibrary : public UBlueprintFunctionLibrary
 {
 	GENERATED_BODY()
diff --git a/Engine/Source/Runtime/Engine/Classes/Kismet/KismetArrayLibrary.h b/Engine/Source/Runtime/Engine/Classes/Kismet/KismetArrayLibrary.h
--- a/Engine/Source/Runtime/Engine/Classes/Kismet/KismetArrayLibrary.h
+++ b/Engine/Source/Runtime/Engine/Classes/Kismet/KismetArrayLibrary.h
@@ -15,7 +15,7 @@
 
 class AActor;
 
-UCLASS()
+UCLASS(meta=(BlueprintThreadSafe))
 class ENGINE_API UKismetArrayLibrary : public UBlueprintFunctionLibrary
 {
 	GENERATED_UCLASS_BODY()
diff --git a/Engine/Source/Runtime/Engine/Classes/Kismet/KismetNodeHelperLibrary.h b/Engine/Source/Runtime/Engine/Classes/Kismet/KismetNodeHelperLibrary.h
--- a/Engine/Source/Runtime/Engine/Classes/Kismet/KismetNodeHelperLibrary.h
+++ b/Engine/Source/Runtime/Engine/Classes/Kismet/KismetNodeHelperLibrary.h
@@ -12,7 +12,7 @@
 #include "Kismet/BlueprintFunctionLibrary.h"
 #include "KismetNodeHelperLibrary.generated.h"
 
-UCLASS()
+UCLASS(meta=(BlueprintThreadSafe))
 class ENGINE_API UKismetNodeHelperLibrary : public UBlueprintFunctionLibrary
 {
 	GENERATED_UCLASS_BODY()
@@ -84,7 +84,7 @@ class ENGINE_API UKismetNodeHelperLibrary : public UBlueprintFunctionLibrary
 	 * @param bRandom - Whether to select a random index or not
 	 * @return - The index that was selected (returns INDEX_NONE if there was no unmarked bits to choose from)
 	 */
-	UFUNCTION(BlueprintCallable, meta=(BlueprintInternalUseOnly = "TRUE"))
+	UFUNCTION(BlueprintCallable, meta=(BlueprintInternalUseOnly = "TRUE", NotBlueprintThreadSafe))
 	static int32 GetUnmarkedBit(int32 Data, int32 StartIdx, int32 NumBits, bool bRandom);
 
 	/**
@@ -94,7 +94,7 @@ class ENGINE_API UKismetNodeHelperLibrary : public UBlueprintFunctionLibrary
 	 * @param NumBits - The logical number of bits we want to track
 	 * @return - The index that was selected (returns INDEX_NONE if there was no unmarked bits to choose from)
 	 */
-	UFUNCTION(BlueprintCallable, meta=(BlueprintInternalUseOnly = "TRUE"))
+	UFUNCTION(BlueprintCallable, meta=(BlueprintInternalUseOnly = "TRUE", NotBlueprintThreadSafe))
 	static int32 GetRandomUnmarkedBit(int32 Data, int32 StartIdx, int32 NumBits);
 
 	/**
