WorldPartition: Make UDataLayerSubsystem::SetDataLayerState BlueprintAuthorityOnly and replace Server only implementation by a if (ensure(GetLocalRole() == ROLE_Authority)) block

diff --git a/Engine/Source/Runtime/Engine/Private/WorldPartition/DataLayer/WorldDataLayers.cpp b/Engine/Source/Runtime/Engine/Private/WorldPartition/DataLayer/WorldDataLayers.cpp
--- a/Engine/Source/Runtime/Engine/Private/WorldPartition/DataLayer/WorldDataLayers.cpp
+++ b/Engine/Source/Runtime/Engine/Private/WorldPartition/DataLayer/WorldDataLayers.cpp
@@ -57,36 +57,39 @@ void AWorldDataLayers::InitializeDataLayerStates()
 	}
 }
 
-void AWorldDataLayers::SetDataLayerState_Implementation(FActorDataLayer InDataLayer, EDataLayerState InState)
+void AWorldDataLayers::SetDataLayerState(FActorDataLayer InDataLayer, EDataLayerState InState)
 {
-	const UDataLayer* DataLayer = GetDataLayerFromName(InDataLayer.Name);
-	if (!DataLayer || !DataLayer->IsDynamicallyLoaded())
+	if (ensure(GetLocalRole() == ROLE_Authority))
 	{
-		return;
-	}
-
-	EDataLayerState CurrentState = GetDataLayerStateByName(InDataLayer.Name);
-	if (CurrentState != InState)
-	{
-		LoadedDataLayerNames.Remove(InDataLayer.Name);
-		ActiveDataLayerNames.Remove(InDataLayer.Name);
-
-		if (InState == EDataLayerState::Loaded)
+		const UDataLayer* DataLayer = GetDataLayerFromName(InDataLayer.Name);
+		if (!DataLayer || !DataLayer->IsDynamicallyLoaded())
 		{
-			LoadedDataLayerNames.Add(InDataLayer.Name);
+			return;
 		}
-		else if (InState == EDataLayerState::Activated)
+
+		EDataLayerState CurrentState = GetDataLayerStateByName(InDataLayer.Name);
+		if (CurrentState != InState)
 		{
-			ActiveDataLayerNames.Add(InDataLayer.Name);
-		}
-			
-		// Update Replicated Properties
-		RepActiveDataLayerNames = ActiveDataLayerNames.Array();
-		RepLoadedDataLayerNames = LoadedDataLayerNames.Array();
+			LoadedDataLayerNames.Remove(InDataLayer.Name);
+			ActiveDataLayerNames.Remove(InDataLayer.Name);
 
-		++DataLayersStateEpoch;
+			if (InState == EDataLayerState::Loaded)
+			{
+				LoadedDataLayerNames.Add(InDataLayer.Name);
+			}
+			else if (InState == EDataLayerState::Activated)
+			{
+				ActiveDataLayerNames.Add(InDataLayer.Name);
+			}
 
-		OnDataLayerStateChanged(DataLayer, InState);
+			// Update Replicated Properties
+			RepActiveDataLayerNames = ActiveDataLayerNames.Array();
+			RepLoadedDataLayerNames = LoadedDataLayerNames.Array();
+
+			++DataLayersStateEpoch;
+
+			OnDataLayerStateChanged(DataLayer, InState);
+		}
 	}
 }
 
diff --git a/Engine/Source/Runtime/Engine/Public/WorldPartition/DataLayer/DataLayerSubsystem.h b/Engine/Source/Runtime/Engine/Public/WorldPartition/DataLayer/DataLayerSubsystem.h
--- a/Engine/Source/Runtime/Engine/Public/WorldPartition/DataLayer/DataLayerSubsystem.h
+++ b/Engine/Source/Runtime/Engine/Public/WorldPartition/DataLayer/DataLayerSubsystem.h
@@ -30,27 +30,36 @@ public:
 	//~ End USubsystem Interface.
 
 	//~ Begin Blueprint callable functions
-	UFUNCTION(BlueprintCallable, Category = DataLayers)
+
+	/** Set the Data Layer state using its name. */
+	UFUNCTION(BlueprintCallable, Category = DataLayers, BlueprintAuthorityOnly)
 	void SetDataLayerState(const FActorDataLayer& InDataLayer, EDataLayerState InState);
 
+	/** Find a Data Layer by name. */
 	UFUNCTION(BlueprintCallable, Category = DataLayers)
 	UDataLayer* GetDataLayer(const FActorDataLayer& InDataLayer) const;
 
-	UFUNCTION(BlueprintCallable, Category = DataLayers)
+	/** Set the Data Layer state using its label. */
+	UFUNCTION(BlueprintCallable, Category = DataLayers, BlueprintAuthorityOnly)
 	void SetDataLayerStateByLabel(const FName& InDataLayerLabel, EDataLayerState InState);
 		
+	/** Get the Data Layer state using its name. */
 	UFUNCTION(BlueprintCallable, Category = DataLayers)
 	EDataLayerState GetDataLayerState(const FActorDataLayer& InDataLayer) const;
 
+	/** Get the Data Layer state using its label. */
 	UFUNCTION(BlueprintCallable, Category = DataLayers)
 	EDataLayerState GetDataLayerStateByLabel(const FName& InDataLayerLabel) const;
 
+	/** Called when a Data Layer changes state. */
 	UPROPERTY(BlueprintAssignable)
 	FOnDataLayerStateChanged OnDataLayerStateChanged;
 
+	/** Find a Data Layer by label. */
 	UFUNCTION(BlueprintCallable, Category = DataLayers)
 	UDataLayer* GetDataLayerFromLabel(FName InDataLayerLabel) const;
 
+	/** Find a Data Layer by name. */
 	UFUNCTION(BlueprintCallable, Category = DataLayers)
 	UDataLayer* GetDataLayerFromName(FName InDataLayerName) const;
 	//~ End Blueprint callable functions
diff --git a/Engine/Source/Runtime/Engine/Public/WorldPartition/DataLayer/WorldDataLayers.h b/Engine/Source/Runtime/Engine/Public/WorldPartition/DataLayer/WorldDataLayers.h
--- a/Engine/Source/Runtime/Engine/Public/WorldPartition/DataLayer/WorldDataLayers.h
+++ b/Engine/Source/Runtime/Engine/Public/WorldPartition/DataLayer/WorldDataLayers.h
@@ -40,7 +40,6 @@ public:
 	void ForEachDataLayer(TFunctionRef<bool(class UDataLayer*)> Func);
 	void ForEachDataLayer(TFunctionRef<bool(class UDataLayer*)> Func) const;
 
-	UFUNCTION(Server, Reliable)
 	void SetDataLayerState(FActorDataLayer InDataLayer, EDataLayerState InState);
 
 	UFUNCTION(NetMulticast, Reliable)
