Disable world partiton editor cells referencing loaded actors until actor deletion is properly handled.

diff --git a/Engine/Source/Runtime/Engine/Private/WorldPartition/WorldPartitionEditorCell.cpp b/Engine/Source/Runtime/Engine/Private/WorldPartition/WorldPartitionEditorCell.cpp
--- a/Engine/Source/Runtime/Engine/Private/WorldPartition/WorldPartitionEditorCell.cpp
+++ b/Engine/Source/Runtime/Engine/Private/WorldPartition/WorldPartitionEditorCell.cpp
@@ -10,6 +10,7 @@ void UWorldPartitionEditorCell::AddReferencedObjects(UObject* InThis, FReference
 {
 	UWorldPartitionEditorCell* This = CastChecked<UWorldPartitionEditorCell>(InThis);
 	
+#if 0
 	Collector.AllowEliminatingReferences(false);
 	for (const FActorReference& ActorReference: This->LoadedActors)
 	{
@@ -19,6 +20,7 @@ void UWorldPartitionEditorCell::AddReferencedObjects(UObject* InThis, FReference
 		Collector.AddReferencedObject(LoadedActor);
 	}
 	Collector.AllowEliminatingReferences(true);
+#endif
 
 	Super::AddReferencedObjects(InThis, Collector);
 }
