Limit asset compilation to number of physical cores like it was always meant to be

diff --git a/Engine/Source/Runtime/Engine/Private/AssetCompilingManager.cpp b/Engine/Source/Runtime/Engine/Private/AssetCompilingManager.cpp
--- a/Engine/Source/Runtime/Engine/Private/AssetCompilingManager.cpp
+++ b/Engine/Source/Runtime/Engine/Private/AssetCompilingManager.cpp
@@ -208,12 +208,17 @@ FQueuedThreadPool* FAssetCompilingManager::GetThreadPool() const
 	{
 		AssetCompilingManagerImpl::EnsureInitializedCVars();
 
-		// Wrapping GThreadPool to give AssetThreadPool it's own set of priorities and allow Pausable functionality
-		// We're using GThreadPool instead of GLargeThreadPool because asset compilation is hard on total memory and memory bandwidth and can run slower when going wider than actual cores.
-		// All asset priorities will resolve to a Low priority once being scheduled in the LargeThreadPool.
+		// We limit concurrency to half the workers because asset compilation is hard on total memory and memory bandwidth and can run slower when going wider than actual cores.
+		// Most asset also have some form of inner parallelism during compilation.
+		// Recently found out that GThreadPool and GLargeThreadPool have the same amount of workers, so can't rely on GThreadPool to be our limiter here.
+		// FPlatformMisc::NumberOfCores() and FPlatformMisc::NumberOfCoresIncludingHyperthreads() also return the same value when -corelimit is used so we can't use FPlatformMisc::NumberOfCores()
+		// if we want to keep the same 1:2 relationship with worker count.
+		FQueuedThreadPoolWrapper* LimitedThreadPool = new FQueuedThreadPoolWrapper(GLargeThreadPool, FMath::Max(GLargeThreadPool->GetNumThreads() / 2, 1));
+
+		// All asset priorities will resolve to a Low priority once being scheduled.
 		// Any asset supporting being built async should be scheduled lower than Normal to let non-async stuff go first
 		// However, we let Highest priority pass-through as it to benefit from going to foreground threads when required (i.e. Game-thread is waiting on some assets)
-		GAssetThreadPool = new FMemoryBoundQueuedThreadPoolWrapper(GThreadPool, -1, [](EQueuedWorkPriority Priority) { return Priority == EQueuedWorkPriority::Highest ? Priority : EQueuedWorkPriority::Low; });
+		GAssetThreadPool = new FMemoryBoundQueuedThreadPoolWrapper(LimitedThreadPool, -1, [](EQueuedWorkPriority Priority) { return Priority == EQueuedWorkPriority::Highest ? Priority : EQueuedWorkPriority::Low; });
 
 		AsyncCompilationHelpers::BindThreadPoolToCVar(
 			GAssetThreadPool,
