more information in BuildTextureStatus log line

diff --git a/Engine/Source/Runtime/Engine/Private/TextureDerivedDataTask.cpp b/Engine/Source/Runtime/Engine/Private/TextureDerivedDataTask.cpp
--- a/Engine/Source/Runtime/Engine/Private/TextureDerivedDataTask.cpp
+++ b/Engine/Source/Runtime/Engine/Private/TextureDerivedDataTask.cpp
@@ -252,17 +252,21 @@ void FTextureCacheDerivedDataWorker::BuildTexture(bool bReplaceExistingDDC)
 	TRACE_CPUPROFILER_EVENT_SCOPE(FTextureCacheDerivedDataWorker::BuildTexture);
 
 	const bool bHasValidMip0 = TextureData.Blocks.Num() && TextureData.Blocks[0].MipsPerLayer.Num() && TextureData.Blocks[0].MipsPerLayer[0].Num();
+	const bool bForVirtualTextureStreamingBuild = (CacheFlags & ETextureCacheFlags::ForVirtualTextureStreamingBuild) != 0;	
 
 	FFormatNamedArguments Args;
 	Args.Add(TEXT("TextureName"), FText::FromString(Texture.GetName()));
 	Args.Add(TEXT("TextureFormatName"), FText::FromString(BuildSettingsPerLayer[0].TextureFormatName.GetPlainNameString()));
+	Args.Add(TEXT("IsVT"), FText::FromString( FString( bForVirtualTextureStreamingBuild ? TEXT(" VT") : TEXT("") ) ) );
 	Args.Add(TEXT("TextureResolutionX"), FText::FromString(FString::FromInt(bHasValidMip0 ? TextureData.Blocks[0].MipsPerLayer[0][0].SizeX : 0)));
 	Args.Add(TEXT("TextureResolutionY"), FText::FromString(FString::FromInt(bHasValidMip0 ? TextureData.Blocks[0].MipsPerLayer[0][0].SizeY : 0)));
+	Args.Add(TEXT("NumBlocks"), FText::FromString(FString::FromInt(TextureData.Blocks.Num())));
+	Args.Add(TEXT("NumLayers"), FText::FromString(FString::FromInt(TextureData.Layers.Num())));
 	Args.Add(TEXT("EstimatedMemory"), FText::FromString(FString::SanitizeFloat(double(RequiredMemoryEstimate) / (1024.0*1024.0), 3)));
 
 	FTextureStatusMessageContext StatusMessage(
 		FText::Format(
-			NSLOCTEXT("Engine", "BuildTextureStatus", "Building textures: {TextureName} ({TextureFormatName}, {TextureResolutionX}X{TextureResolutionY}) (Required Memory Estimate: {EstimatedMemory} MB)"), 
+			NSLOCTEXT("Engine", "BuildTextureStatus", "Building textures: {TextureName} ({TextureFormatName}{IsVT}, {TextureResolutionX}X{TextureResolutionY} X{NumBlocks}X{NumLayers}) (Required Memory Estimate: {EstimatedMemory} MB)"), 
 			Args
 		));
 
@@ -272,7 +276,6 @@ void FTextureCacheDerivedDataWorker::BuildTexture(bool bReplaceExistingDDC)
 		return;
 	}
 
-	const bool bForVirtualTextureStreamingBuild = (CacheFlags & ETextureCacheFlags::ForVirtualTextureStreamingBuild) != 0;	
 	if (bForVirtualTextureStreamingBuild)
 	{
 		if (DerivedData->VTData == nullptr)
@@ -646,9 +649,9 @@ void FTextureCacheDerivedDataWorker::DoWork()
 		{
 			CompositeTextureData.GetAsyncSourceMips(ImageWrapper);
 			if ((bool)Texture.CompositeTexture)
-			{
-				BuildExporter.ExportCompositeTextureSourceBulkData(CompositeTextureData.AsyncSource);
-			}
+		{
+			BuildExporter.ExportCompositeTextureSourceBulkData(CompositeTextureData.AsyncSource);
+		}
 			CompositeTextureData.AsyncSource.RemoveBulkData();
 		}
 
