no unity fixes

diff --git a/Engine/Plugins/Experimental/MeshModelingToolset/Source/MeshModelingToolsEditorOnly/Public/ToolTargets/VolumeDynamicMeshToolTarget.h b/Engine/Plugins/Experimental/MeshModelingToolset/Source/MeshModelingToolsEditorOnly/Public/ToolTargets/VolumeDynamicMeshToolTarget.h
--- a/Engine/Plugins/Experimental/MeshModelingToolset/Source/MeshModelingToolsEditorOnly/Public/ToolTargets/VolumeDynamicMeshToolTarget.h
+++ b/Engine/Plugins/Experimental/MeshModelingToolset/Source/MeshModelingToolsEditorOnly/Public/ToolTargets/VolumeDynamicMeshToolTarget.h
@@ -12,6 +12,8 @@
 
 #include "VolumeDynamicMeshToolTarget.generated.h"
 
+struct FMeshDescription;
+
 /**
  * A tool target backed by AVolume
  */
diff --git a/Engine/Source/Developer/MeshMergeUtilities/Private/MeshMergeUtilities.h b/Engine/Source/Developer/MeshMergeUtilities/Private/MeshMergeUtilities.h
--- a/Engine/Source/Developer/MeshMergeUtilities/Private/MeshMergeUtilities.h
+++ b/Engine/Source/Developer/MeshMergeUtilities/Private/MeshMergeUtilities.h
@@ -6,6 +6,7 @@
 #include "SceneTypes.h"
 
 class FMeshMergeDataTracker;
+class FMaterialUpdateContext;
 struct FMeshLODKey;
 
 typedef TPair<uint32, uint32> MaterialRemapPair;
diff --git a/Engine/Source/Editor/DataLayerEditor/Private/DataLayer/DataLayerActorTreeItem.h b/Engine/Source/Editor/DataLayerEditor/Private/DataLayer/DataLayerActorTreeItem.h
--- a/Engine/Source/Editor/DataLayerEditor/Private/DataLayer/DataLayerActorTreeItem.h
+++ b/Engine/Source/Editor/DataLayerEditor/Private/DataLayer/DataLayerActorTreeItem.h
@@ -3,6 +3,7 @@
 #pragma once
 
 #include "CoreMinimal.h"
+#include "GameFramework/Actor.h"
 #include "SceneOutlinerFwd.h"
 #include "ActorTreeItem.h"
 #include "WorldPartition/DataLayer/DataLayer.h"
diff --git a/Engine/Source/Editor/SceneOutliner/Private/ActorHierarchy.cpp b/Engine/Source/Editor/SceneOutliner/Private/ActorHierarchy.cpp
--- a/Engine/Source/Editor/SceneOutliner/Private/ActorHierarchy.cpp
+++ b/Engine/Source/Editor/SceneOutliner/Private/ActorHierarchy.cpp
@@ -2,6 +2,7 @@
 
 #include "ActorHierarchy.h"
 
+#include "Engine/Engine.h"
 #include "Engine/World.h"
 #include "EngineUtils.h"
 #include "WorldTreeItem.h"
diff --git a/Engine/Source/Editor/WorldPartitionEditor/Private/WorldPartition/SWorldPartitionConvertDialog.cpp b/Engine/Source/Editor/WorldPartitionEditor/Private/WorldPartition/SWorldPartitionConvertDialog.cpp
--- a/Engine/Source/Editor/WorldPartitionEditor/Private/WorldPartition/SWorldPartitionConvertDialog.cpp
+++ b/Engine/Source/Editor/WorldPartitionEditor/Private/WorldPartition/SWorldPartitionConvertDialog.cpp
@@ -4,7 +4,7 @@
 #include "Widgets/Input/SButton.h"
 #include "WorldPartition/WorldPartitionConvertOptions.h"
 #include "Modules/ModuleManager.h"
-
+#include "PropertyEditorModule.h"
 
 #define LOCTEXT_NAMESPACE "WorldPartitionConvertDialog"
 
diff --git a/Engine/Source/Editor/WorldPartitionEditor/Private/WorldPartition/SWorldPartitionEditorGridSpatialHash.cpp b/Engine/Source/Editor/WorldPartitionEditor/Private/WorldPartition/SWorldPartitionEditorGridSpatialHash.cpp
--- a/Engine/Source/Editor/WorldPartitionEditor/Private/WorldPartition/SWorldPartitionEditorGridSpatialHash.cpp
+++ b/Engine/Source/Editor/WorldPartitionEditor/Private/WorldPartition/SWorldPartitionEditorGridSpatialHash.cpp
@@ -22,6 +22,7 @@
 #include "EditorStyleSet.h"
 #include "Selection.h"
 #include "Engine/World.h"
+#include "Engine/Texture2D.h"
 #include "Engine/Texture2DDynamic.h"
 #include "Misc/HashBuilder.h"
 
diff --git a/Engine/Source/Runtime/Engine/Private/WorldPartition/RuntimeSpatialHash/RuntimeSpatialHashHLOD.cpp b/Engine/Source/Runtime/Engine/Private/WorldPartition/RuntimeSpatialHash/RuntimeSpatialHashHLOD.cpp
--- a/Engine/Source/Runtime/Engine/Private/WorldPartition/RuntimeSpatialHash/RuntimeSpatialHashHLOD.cpp
+++ b/Engine/Source/Runtime/Engine/Private/WorldPartition/RuntimeSpatialHash/RuntimeSpatialHashHLOD.cpp
@@ -11,6 +11,7 @@
 #include "WorldPartition/HLOD/HLODActorDesc.h"
 #include "WorldPartition/WorldPartitionHelpers.h"
 
+#include "Engine/Engine.h"
 #include "HAL/PlatformFileManager.h"
 #include "Misc/PackageName.h"
 #include "Misc/ScopedSlowTask.h"
diff --git a/Engine/Source/Runtime/Engine/Private/WorldPartition/WorldPartitionRuntimeSpatialHashCell.cpp b/Engine/Source/Runtime/Engine/Private/WorldPartition/WorldPartitionRuntimeSpatialHashCell.cpp
--- a/Engine/Source/Runtime/Engine/Private/WorldPartition/WorldPartitionRuntimeSpatialHashCell.cpp
+++ b/Engine/Source/Runtime/Engine/Private/WorldPartition/WorldPartitionRuntimeSpatialHashCell.cpp
@@ -1,8 +1,13 @@
 // Copyright Epic Games, Inc. All Rights Reserved.
+
 #include "WorldPartition/WorldPartitionRuntimeSpatialHashCell.h"
 #include "WorldPartition/WorldPartitionStreamingSource.h"
 #include "HAL/IConsoleManager.h"
 
+#include "Engine/Level.h"
+#include "GameFramework/Actor.h"
+#include "HAL/ConsoleManager.h"
+
 static float GRuntimeSpatialHashCellToSourceAngleContributionToCellImportance = 0.4f; // Value between [0, 1]
 static FAutoConsoleVariableRef CVarRuntimeSpatialHashCellToSourceAngleContributionToCellImportance(
 	TEXT("wp.Runtime.RuntimeSpatialHashCellToSourceAngleContributionToCellImportance"),
diff --git a/Engine/Source/Runtime/Engine/Public/WorldPartition/HLOD/HLODLayer.h b/Engine/Source/Runtime/Engine/Public/WorldPartition/HLOD/HLODLayer.h
--- a/Engine/Source/Runtime/Engine/Public/WorldPartition/HLOD/HLODLayer.h
+++ b/Engine/Source/Runtime/Engine/Public/WorldPartition/HLOD/HLODLayer.h
@@ -18,6 +18,7 @@
 class UMaterial;
 class UWorldPartition;
 class AActor;
+class FWorldPartitionActorDesc;
 
 UENUM()
 enum class EHLODLayerType : uint8
diff --git a/Engine/Source/Runtime/Engine/Public/WorldPartition/WorldPartitionRuntimeSpatialHashCell.h b/Engine/Source/Runtime/Engine/Public/WorldPartition/WorldPartitionRuntimeSpatialHashCell.h
--- a/Engine/Source/Runtime/Engine/Public/WorldPartition/WorldPartitionRuntimeSpatialHashCell.h
+++ b/Engine/Source/Runtime/Engine/Public/WorldPartition/WorldPartitionRuntimeSpatialHashCell.h
@@ -5,6 +5,8 @@
 #include "WorldPartition/WorldPartitionRuntimeCell.h"
 #include "WorldPartitionRuntimeSpatialHashCell.generated.h"
 
+class UActorContainer;
+
 UCLASS(Abstract, Within = WorldPartition)
 class UWorldPartitionRuntimeSpatialHashCell : public UWorldPartitionRuntimeCell
 {
