WorldPartition: Support ClientOnlyVisible streaming grids (HLODs are setup this way, not loaded on server)

diff --git a/Engine/Source/Runtime/Engine/Classes/Engine/LevelStreaming.h b/Engine/Source/Runtime/Engine/Classes/Engine/LevelStreaming.h
--- a/Engine/Source/Runtime/Engine/Classes/Engine/LevelStreaming.h
+++ b/Engine/Source/Runtime/Engine/Classes/Engine/LevelStreaming.h
@@ -164,6 +164,10 @@ public:
 	UPROPERTY(EditAnywhere, Category=LevelStreaming, BlueprintReadWrite)
 	FTransform LevelTransform;
 
+	/** Applied to LoadedLevel						                                                                            */
+	UPROPERTY()
+	bool bClientOnlyVisible;
+
 private:
 
 	/** Requested LOD. Non LOD sub-levels have Index = -1  */
diff --git a/Engine/Source/Runtime/Engine/Private/LevelStreaming.cpp b/Engine/Source/Runtime/Engine/Private/LevelStreaming.cpp
--- a/Engine/Source/Runtime/Engine/Private/LevelStreaming.cpp
+++ b/Engine/Source/Runtime/Engine/Private/LevelStreaming.cpp
@@ -1284,6 +1284,9 @@ void ULevelStreaming::AsyncLevelLoadComplete(const FName& InPackageName, UPackag
 						Level->bClientOnlyVisible = LODPackageNames.Contains(InLoadedPackage->GetFName());
 					}
 				}
+
+				// Apply streaming level property to level
+				Level->bClientOnlyVisible |= bClientOnlyVisible;
 			
 				// In the editor levels must be in the levels array regardless of whether they are visible or not
 				if (ensure(LevelOwningWorld) && LevelOwningWorld->WorldType == EWorldType::Editor)
diff --git a/Engine/Source/Runtime/Engine/Private/WorldPartition/RuntimeSpatialHash/RuntimeSpatialHashHLOD.cpp b/Engine/Source/Runtime/Engine/Private/WorldPartition/RuntimeSpatialHash/RuntimeSpatialHashHLOD.cpp
--- a/Engine/Source/Runtime/Engine/Private/WorldPartition/RuntimeSpatialHash/RuntimeSpatialHashHLOD.cpp
+++ b/Engine/Source/Runtime/Engine/Private/WorldPartition/RuntimeSpatialHash/RuntimeSpatialHashHLOD.cpp
@@ -289,6 +289,7 @@ static void GatherHLODGrids(UWorldPartition* WorldPartition, TMap<FName, FSpatia
 		HLODGrid.LoadingRange = HLODLayer->GetLoadingRange();
 		HLODGrid.DebugColor = FLinearColor::Red;
 		HLODGrid.GridName = HLODLayer->GetRuntimeGrid(InitialHLODLevel);
+		HLODGrid.bClientOnlyVisible = true;
 
 		OutHLODGrids.Emplace(HLODGrid.GridName, HLODGrid);
 
diff --git a/Engine/Source/Runtime/Engine/Private/WorldPartition/WorldPartitionLevelStreamingDynamic.cpp b/Engine/Source/Runtime/Engine/Private/WorldPartition/WorldPartitionLevelStreamingDynamic.cpp
--- a/Engine/Source/Runtime/Engine/Private/WorldPartition/WorldPartitionLevelStreamingDynamic.cpp
+++ b/Engine/Source/Runtime/Engine/Private/WorldPartition/WorldPartitionLevelStreamingDynamic.cpp
@@ -20,6 +20,7 @@
 #include "WorldPartition/WorldPartitionLevelStreamingPolicy.h"
 #include "WorldPartition/WorldPartitionLevelHelper.h"
 #include "WorldPartition/WorldPartitionRuntimeHash.h"
+#include "WorldPartition/HLOD/HLODActor.h"
 #endif
 
 #define LOCTEXT_NAMESPACE "World"
@@ -89,6 +90,9 @@ void UWorldPartitionLevelStreamingDynamic::CreateRuntimeLevel()
 	RuntimeLevel = FWorldPartitionLevelHelper::CreateEmptyLevelForRuntimeCell(PlayWorld, GetWorldAsset().ToString());
 	check(RuntimeLevel);
 
+	// Set flag here as this level isn't async loaded
+	RuntimeLevel->bClientOnlyVisible = bClientOnlyVisible;
+
 	// Attach ourself to Level cleanup to do our own cleanup
 	OnCleanupLevelDelegateHandle = RuntimeLevel->OnCleanupLevel.AddUObject(this, &UWorldPartitionLevelStreamingDynamic::OnCleanupLevel);
 }
diff --git a/Engine/Source/Runtime/Engine/Private/WorldPartition/WorldPartitionRuntimeLevelStreamingCell.cpp b/Engine/Source/Runtime/Engine/Private/WorldPartition/WorldPartitionRuntimeLevelStreamingCell.cpp
--- a/Engine/Source/Runtime/Engine/Private/WorldPartition/WorldPartitionRuntimeLevelStreamingCell.cpp
+++ b/Engine/Source/Runtime/Engine/Private/WorldPartition/WorldPartitionRuntimeLevelStreamingCell.cpp
@@ -108,6 +108,7 @@ ULevelStreaming* UWorldPartitionRuntimeLevelStreamingCell::CreateLevelStreaming(
 		NewLevelStreaming->SetWorldAsset(WorldAsset);
 		// Transfer WorldPartition's transform to Level
 		NewLevelStreaming->LevelTransform = WorldPartition->GetInstanceTransform();
+		NewLevelStreaming->bClientOnlyVisible = GetClientOnlyVisible();
 
 		if (UWorldPartitionLevelStreamingDynamic* WorldPartitionLevelStreamingDynamic = Cast<UWorldPartitionLevelStreamingDynamic>(NewLevelStreaming))
 		{
diff --git a/Engine/Source/Runtime/Engine/Private/WorldPartition/WorldPartitionRuntimeSpatialHash.cpp b/Engine/Source/Runtime/Engine/Private/WorldPartition/WorldPartitionRuntimeSpatialHash.cpp
--- a/Engine/Source/Runtime/Engine/Private/WorldPartition/WorldPartitionRuntimeSpatialHash.cpp
+++ b/Engine/Source/Runtime/Engine/Private/WorldPartition/WorldPartitionRuntimeSpatialHash.cpp
@@ -73,6 +73,7 @@ FSpatialHashStreamingGrid::FSpatialHashStreamingGrid()
 	, LoadingRange(0.0f)
 	, DebugColor(ForceInitToZero)
 	, WorldBounds(ForceInitToZero)
+	, bClientOnlyVisible(false)
 	, OverrideLoadingRange(-1.f)
 	, GridHelper(nullptr)
 {
@@ -767,6 +768,7 @@ bool UWorldPartitionRuntimeSpatialHash::CreateStreamingGrid(const FSpatialHashRu
 	CurrentStreamingGrid.WorldBounds = PartionedActors.WorldBounds;
 	CurrentStreamingGrid.LoadingRange = RuntimeGrid.LoadingRange;
 	CurrentStreamingGrid.DebugColor = RuntimeGrid.DebugColor;
+	CurrentStreamingGrid.bClientOnlyVisible = RuntimeGrid.bClientOnlyVisible;
 
 	// Move actors into the final streaming grids
 	CurrentStreamingGrid.GridLevels.Reserve(PartionedActors.Levels.Num());
@@ -864,6 +866,7 @@ bool UWorldPartitionRuntimeSpatialHash::CreateStreamingGrid(const FSpatialHashRu
 				verify(TempLevel.GetCellBounds(FIntVector2(CellCoordX, CellCoordY), Bounds));
 				StreamingCell->Position = FVector(Bounds.GetCenter(), 0.f);
 				StreamingCell->SetDebugInfo(CellGlobalCoords, CurrentStreamingGrid.GridName);
+				StreamingCell->SetClientOnlyVisible(CurrentStreamingGrid.bClientOnlyVisible);
 
 				UE_LOG(LogWorldPartition, Verbose, TEXT("Cell%s %s Actors = %d Bounds (%s)"), bIsCellAlwaysLoaded ? TEXT(" (AlwaysLoaded)") : TEXT(""), *StreamingCell->GetName(), FilteredActors.Num(), *Bounds.ToString());
 
@@ -1013,24 +1016,43 @@ FAutoConsoleCommand UWorldPartitionRuntimeSpatialHash::OverrideLoadingRangeComma
 	})
 );
 
+bool UWorldPartitionRuntimeSpatialHash::ShouldConsiderClientOnlyVisibleCells() const
+{
+	const UWorld* World = GetWorld();
+	if (World->IsGameWorld())
+	{
+		ENetMode NetMode = World->GetNetMode();
+		if (NetMode == NM_DedicatedServer || NetMode == NM_ListenServer)
+		{
+			return false;
+		}
+	}
+	return true;
+}
+
 // Streaming interface
 int32 UWorldPartitionRuntimeSpatialHash::GetAllStreamingCells(TSet<const UWorldPartitionRuntimeCell*>& Cells, bool bAllDataLayers, bool bDataLayersOnly, const TSet<FName>& InDataLayers) const
 {
+	const bool bShouldConsiderClientOnlyVisible = ShouldConsiderClientOnlyVisibleCells();
+
 	for (const FSpatialHashStreamingGrid& StreamingGrid : StreamingGrids)
 	{
-		for (const FSpatialHashStreamingGridLevel& GridLevel : StreamingGrid.GridLevels)
+		if (!StreamingGrid.bClientOnlyVisible || bShouldConsiderClientOnlyVisible)
 		{
-			for (const FSpatialHashStreamingGridLayerCell& LayerCell : GridLevel.LayerCells)
+			for (const FSpatialHashStreamingGridLevel& GridLevel : StreamingGrid.GridLevels)
 			{
-				for (const UWorldPartitionRuntimeSpatialHashCell* Cell : LayerCell.GridCells)
+				for (const FSpatialHashStreamingGridLayerCell& LayerCell : GridLevel.LayerCells)
 				{
-					if (!bDataLayersOnly && !Cell->HasDataLayers())
-					{
-						Cells.Add(Cell);
-					}
-					else if (Cell->HasDataLayers() && (bAllDataLayers || Cell->HasAnyDataLayer(InDataLayers)))
+					for (const UWorldPartitionRuntimeSpatialHashCell* Cell : LayerCell.GridCells)
 					{
-						Cells.Add(Cell);
+						if (!bDataLayersOnly && !Cell->HasDataLayers())
+						{
+							Cells.Add(Cell);
+						}
+						else if (Cell->HasDataLayers() && (bAllDataLayers || Cell->HasAnyDataLayer(InDataLayers)))
+						{
+							Cells.Add(Cell);
+						}
 					}
 				}
 			}
@@ -1042,9 +1064,14 @@ int32 UWorldPartitionRuntimeSpatialHash::GetAllStreamingCells(TSet<const UWorldP
 
 bool UWorldPartitionRuntimeSpatialHash::GetStreamingCells(const FWorldPartitionStreamingQuerySource& QuerySource, TSet<const UWorldPartitionRuntimeCell*>& OutCells) const 
 {
+	const bool bShouldConsiderClientOnlyVisible = ShouldConsiderClientOnlyVisibleCells();
+
 	for (const FSpatialHashStreamingGrid& StreamingGrid : StreamingGrids)
 	{
-		StreamingGrid.GetCells(QuerySource, OutCells);
+		if (!StreamingGrid.bClientOnlyVisible || bShouldConsiderClientOnlyVisible)
+		{
+			StreamingGrid.GetCells(QuerySource, OutCells);
+		}
 	}
 
 	return !!OutCells.Num();
@@ -1055,13 +1082,17 @@ bool UWorldPartitionRuntimeSpatialHash::GetStreamingCells(const TArray<FWorldPar
 	TRACE_CPUPROFILER_EVENT_SCOPE(UWorldPartitionRuntimeSpatialHash::GetStreamingCells);
 
 	const UDataLayerSubsystem* DataLayerSubsystem = GetWorld()->GetSubsystem<UDataLayerSubsystem>();
+	const bool bShouldConsiderClientOnlyVisible = ShouldConsiderClientOnlyVisibleCells();
 
 	if (Sources.Num() == 0)
 	{
 		// Get always loaded cells
 		for (const FSpatialHashStreamingGrid& StreamingGrid : StreamingGrids)
 		{
-			StreamingGrid.GetAlwaysLoadedCells(DataLayerSubsystem, OutActivateCells.GetCells(), OutLoadCells.GetCells());
+			if (!StreamingGrid.bClientOnlyVisible || bShouldConsiderClientOnlyVisible)
+			{
+				StreamingGrid.GetAlwaysLoadedCells(DataLayerSubsystem, OutActivateCells.GetCells(), OutLoadCells.GetCells());
+			}
 		}
 	}
 	else
@@ -1069,7 +1100,10 @@ bool UWorldPartitionRuntimeSpatialHash::GetStreamingCells(const TArray<FWorldPar
 		// Get cells based on streaming sources
 		for (const FSpatialHashStreamingGrid& StreamingGrid : StreamingGrids)
 		{
-			StreamingGrid.GetCells(Sources, DataLayerSubsystem, OutActivateCells, OutLoadCells);
+			if (!StreamingGrid.bClientOnlyVisible || bShouldConsiderClientOnlyVisible)
+			{
+				StreamingGrid.GetCells(Sources, DataLayerSubsystem, OutActivateCells, OutLoadCells);
+			}
 		}
 	}
 
diff --git a/Engine/Source/Runtime/Engine/Public/WorldPartition/WorldPartitionRuntimeCell.h b/Engine/Source/Runtime/Engine/Public/WorldPartition/WorldPartitionRuntimeCell.h
--- a/Engine/Source/Runtime/Engine/Public/WorldPartition/WorldPartitionRuntimeCell.h
+++ b/Engine/Source/Runtime/Engine/Public/WorldPartition/WorldPartitionRuntimeCell.h
@@ -143,6 +143,9 @@ class UWorldPartitionRuntimeCell : public UObject
 	}
 
 #if WITH_EDITOR
+	void SetClientOnlyVisible(bool bInClientOnlyVisible) { bClientOnlyVisible = bInClientOnlyVisible; }
+	bool GetClientOnlyVisible() const { return bClientOnlyVisible; }
+
 	void SetDataLayers(const TArray<const UDataLayer*>& InDataLayers);
 	void SetDebugInfo(FIntVector InCoords, FName InGridName);
 	void AddCellData(const UWorldPartitionRuntimeCellData* InCellData);
@@ -189,6 +192,9 @@ private:
 	UPROPERTY()
 	int32 Priority;
 
+	UPROPERTY()
+	bool bClientOnlyVisible;
+
 	// Source Priority
 	mutable int32 CachedSourcePriority;
 
diff --git a/Engine/Source/Runtime/Engine/Public/WorldPartition/WorldPartitionRuntimeSpatialHash.h b/Engine/Source/Runtime/Engine/Public/WorldPartition/WorldPartitionRuntimeSpatialHash.h
--- a/Engine/Source/Runtime/Engine/Public/WorldPartition/WorldPartitionRuntimeSpatialHash.h
+++ b/Engine/Source/Runtime/Engine/Public/WorldPartition/WorldPartitionRuntimeSpatialHash.h
@@ -112,6 +112,9 @@ struct FSpatialHashStreamingGrid
 	UPROPERTY()
 	FBox WorldBounds;
 
+	UPROPERTY()
+	bool bClientOnlyVisible;
+
 	float OverrideLoadingRange;
 
 	float GetLoadingRange() const
@@ -153,6 +156,7 @@ struct FSpatialHashRuntimeGrid
 		, LoadingRange(1000)
 		, Priority(0)
 		, DebugColor(FLinearColor::MakeRandomColor())
+		, bClientOnlyVisible(false)
 #endif
 	{}
 
@@ -171,6 +175,9 @@ struct FSpatialHashRuntimeGrid
 
 	UPROPERTY(EditAnywhere, Category=Settings, meta = (IgnoreForMemberInitializationTest))
 	FLinearColor DebugColor;
+
+	UPROPERTY()
+	bool bClientOnlyVisible;
 #endif
 };
 
@@ -224,6 +231,7 @@ public:
 	static FString GetCellCoordString(const FIntVector& InCellGlobalCoord);
 
 protected:
+	bool ShouldConsiderClientOnlyVisibleCells() const;
 
 #if WITH_EDITOR
 	virtual bool GenerateStreaming(EWorldPartitionStreamingMode Mode, class UWorldPartitionStreamingPolicy* StreamingPolicy, TArray<FString>* OutPackagesToGenerate = nullptr) override;
