disable incorrect Texture DDC mip size warning

diff --git a/Engine/Source/Runtime/Engine/Private/TextureDerivedData.cpp b/Engine/Source/Runtime/Engine/Private/TextureDerivedData.cpp
--- a/Engine/Source/Runtime/Engine/Private/TextureDerivedData.cpp
+++ b/Engine/Source/Runtime/Engine/Private/TextureDerivedData.cpp
@@ -1013,9 +1013,13 @@ static void BeginLoadDerivedVTChunks(const TArray<FVirtualTextureDataChunk>& Chu
 	}
 }
 
-/** Asserts that MipSize is correct for the mipmap. */
+/** Logs a warning that MipSize is correct for the mipmap. */
 static void CheckMipSize(FTexture2DMipMap& Mip, EPixelFormat PixelFormat, int32 MipSize)
 {
+	// this check is incorrect ; it does not account of platform tiling and padding done on textures
+	// re-enable if fixed
+
+	#if 0
 	// Only volume can have SizeZ != 1
 	if (MipSize != Mip.SizeZ * CalcTextureMipMapSize(Mip.SizeX, Mip.SizeY, PixelFormat, 0))
 	{
@@ -1029,6 +1033,7 @@ static void CheckMipSize(FTexture2DMipMap& Mip, EPixelFormat PixelFormat, int32
 			*Mip.DerivedDataKey
 			);
 	}
+	#endif
 }
 
 bool FTexturePlatformData::TryInlineMipData(int32 FirstMipToLoad, UTexture* Texture)
