Refactor HLSL locals to generate a tree using single static assignment
- 'Phi expressions' are generated to represent locals which are assigned from multiple scopes
- In practice this means that constants/preshaders can propogate through local variable assignments
- Also means that code to assign local values won't be emitted unless code to read that local is also required (since phi nodes will insert required assignments in a 'lazy' fashion)

diff --git a/Engine/Source/Runtime/Engine/Private/HLSLTree/HLSLTree.cpp b/Engine/Source/Runtime/Engine/Private/HLSLTree/HLSLTree.cpp
--- a/Engine/Source/Runtime/Engine/Private/HLSLTree/HLSLTree.cpp
+++ b/Engine/Source/Runtime/Engine/Private/HLSLTree/HLSLTree.cpp
@@ -12,7 +12,7 @@ public:
 	TStringBuilder<2048> LocalStringBuilder;
 };
 
-UE::HLSLTree::EExpressionEvaluationType UE::HLSLTree::CombineEvaluationTypes(UE::HLSLTree::EExpressionEvaluationType Lhs, UE::HLSLTree::EExpressionEvaluationType Rhs)
+UE::HLSLTree::EExpressionEvaluationType UE::HLSLTree::CombineEvaluationTypes(EExpressionEvaluationType Lhs, EExpressionEvaluationType Rhs)
 {
 	if (Lhs == EExpressionEvaluationType::Constant && Rhs == EExpressionEvaluationType::Constant)
 	{
@@ -165,30 +165,10 @@ UE::HLSLTree::FEmitContext::~FEmitContext()
 	}
 }
 
-UE::HLSLTree::FEmitContext::FScopeEntry* UE::HLSLTree::FEmitContext::FindScope(UE::HLSLTree::FScope* Scope)
+UE::HLSLTree::FEmitContext::FScopeEntry* UE::HLSLTree::FEmitContext::FindScopeEntry(const FScope* Scope)
 {
-	for (int32 Index = ScopeStack.Num() - 1; Index >= 0; --Index)
-	{
-		FScopeEntry& Entry = ScopeStack[Index];
-		if (Entry.Scope == Scope)
-		{
-			return &Entry;
-		}
-	}
-	return nullptr;
-}
-
-int32 UE::HLSLTree::FEmitContext::FindScopeIndex(FScope* Scope)
-{
-	for (int32 Index = ScopeStack.Num() - 1; Index >= 0; --Index)
-	{
-		FScopeEntry& Entry = ScopeStack[Index];
-		if (Entry.Scope == Scope)
-		{
-			return Index;
-		}
-	}
-	return INDEX_NONE;
+	FEmitContext::FScopeEntry** FoundEntry = ScopeMap.Find(Scope);
+	return FoundEntry ? *FoundEntry : nullptr;
 }
 
 void UE::HLSLTree::FExpressionEmitResult::ForwardValue(FEmitContext& Context, const FEmitValue* InValue)
@@ -207,35 +187,7 @@ void UE::HLSLTree::FExpressionEmitResult::ForwardValue(FEmitContext& Context, co
 	}
 }
 
-const UE::HLSLTree::FEmitValue* UE::HLSLTree::FEmitContext::AcquireValue(UE::HLSLTree::FLocalDeclaration* Declaration)
-{
-	check(Declaration);
-	check(Declaration->Type != Shader::EValueType::Void);
-
-	FDeclarationEntry*& Entry = FunctionStack.Last().DeclarationMap.FindOrAdd(Declaration);
-	if (!Entry)
-	{
-		FScopeEntry* ScopeEntry = FindScope(Declaration->ParentScope);
-		check(ScopeEntry);
-		Entry = new(*Allocator) FDeclarationEntry();
-		Entry->Value.EvaluationType = EExpressionEvaluationType::Shader;
-		Entry->Value.ExpressionType = Declaration->Type;
-
-		FLocalHLSLCodeWriter LocalWriter;
-		LocalWriter.WriteConstant(Shader::FValue(Declaration->Type));
-
-		Entry->Value.Code = AllocateStringf(*Allocator, TEXT("Local%d"), NumExpressionLocals++);
-		const Shader::FValueTypeDescription& TypeDesc = Shader::GetValueTypeDescription(Declaration->Type);
-		ScopeEntry->ExpressionCodeWriter->WriteLinef(TEXT("%s %s = %s;"),
-			TypeDesc.Name,
-			Entry->Value.Code,
-			LocalWriter.StringBuilder->ToString());
-	}
-
-	return &Entry->Value;
-}
-
-const UE::HLSLTree::FEmitValue* UE::HLSLTree::FEmitContext::AcquireValue(UE::HLSLTree::FExpression* Expression)
+const UE::HLSLTree::FEmitValue* UE::HLSLTree::FEmitContext::AcquireValue(FExpression* Expression)
 {
 	if (!Expression)
 	{
@@ -281,33 +233,39 @@ const UE::HLSLTree::FEmitValue* UE::HLSLTree::FEmitContext::AcquireValue(UE::HLS
 				}
 				else
 				{
-					const int32 ScopeIndex = FindScopeIndex(Expression->ParentScope);
-					check(ScopeIndex != INDEX_NONE);
+					FScopeEntry* ScopeEntry = FindScopeEntry(Expression->ParentScope);
+					check(ScopeEntry);
 
 					// Check to see if we've already generated code for an equivalent expression in either this scope, or any outer visible scope
 					const FSHAHash Hash = LocalWriter.GetCodeHash();
 					const TCHAR* Declaration = nullptr;
-					for (int32 CheckScopeIndex = ScopeIndex; CheckScopeIndex >= 0; --CheckScopeIndex)
+
+					FScopeEntry* CheckScopeEntry = ScopeEntry;
+					while (CheckScopeEntry)
 					{
-						const FScopeEntry& CheckScopeEntry = ScopeStack[CheckScopeIndex];
-						const TCHAR** FoundDeclaration = CheckScopeEntry.ExpressionMap->Find(Hash);
+						const TCHAR** FoundDeclaration = CheckScopeEntry->ExpressionMap.Find(Hash);
 						if (FoundDeclaration)
 						{
 							// Re-use results from previous matching expression
 							Declaration = *FoundDeclaration;
 							break;
 						}
+						CheckScopeEntry = CheckScopeEntry->ParentEntry;
 					}
 
 					if (!Declaration)
 					{
 						const Shader::FValueTypeDescription& TypeDesc = Shader::GetValueTypeDescription(EmitResult.Type);
-						Declaration = AllocateStringf(*Allocator, TEXT("Local%d"), NumExpressionLocals++);
-						ScopeStack[ScopeIndex].ExpressionCodeWriter->WriteLinef(TEXT("const %s %s = %s;"),
+						Declaration = AcquireLocalDeclarationCode();
+						/*ScopeStack[ScopeIndex].ExpressionCodeWriter->WriteLinef(TEXT("const %s %s = %s;"),
+							TypeDesc.Name,
+							Declaration,
+							LocalWriter.GetStringBuilder().ToString());*/
+						AppendHLSLf(*ScopeEntry->ScopeCode, TEXT("const %s %s = %s;"),
 							TypeDesc.Name,
 							Declaration,
 							LocalWriter.GetStringBuilder().ToString());
-						ScopeStack[ScopeIndex].ExpressionMap->Add(Hash, Declaration);
+						ScopeEntry->ExpressionMap.Add(Hash, Declaration);
 					}
 					Entry->Value.Code = Declaration;
 				}
@@ -325,7 +283,7 @@ const UE::HLSLTree::FEmitValue* UE::HLSLTree::FEmitContext::AcquireValue(FFuncti
 	FFunctionCallEntry*& Entry = FunctionStack.Last().FunctionCallMap.FindOrAdd(FunctionCall);
 	if (!Entry)
 	{
-		FScopeEntry* ParentScopeEntry = FindScope(FunctionCall->ParentScope);
+		FScopeEntry* ParentScopeEntry = FindScopeEntry(FunctionCall->ParentScope);
 		check(ParentScopeEntry);
 
 		{
@@ -333,17 +291,11 @@ const UE::HLSLTree::FEmitValue* UE::HLSLTree::FEmitContext::AcquireValue(FFuncti
 			StackEntry.FunctionCall = FunctionCall;
 		}
 
-		{
-			FScopeEntry& ScopeEntry = ScopeStack.AddDefaulted_GetRef();
-			ScopeEntry.Scope = FunctionCall->FunctionScope;
-			ScopeEntry.ExpressionCodeWriter = ParentScopeEntry->ExpressionCodeWriter;
-			ScopeEntry.ExpressionMap = ParentScopeEntry->ExpressionMap;
-		}
+		// Link the function scope to our parent scope
+		ScopeMap.Add(FunctionCall->FunctionScope, ParentScopeEntry);
 
 		// Emit the function's HLSL into the current scope (allows sharing expressions across functions called from the same scope)
-		FLocalHLSLCodeWriter LocalWriter;
-		FunctionCall->FunctionScope->EmitUnscopedHLSL(*this, LocalWriter);
-		ParentScopeEntry->ExpressionCodeWriter->Append(LocalWriter);
+		FunctionCall->FunctionScope->EmitUnscopedHLSL(*this, *ParentScopeEntry->ScopeCode);
 
 		// Assign function outputs
 		FEmitValue* OutputValues = new(*Allocator) FEmitValue[FunctionCall->NumOutputs];
@@ -365,10 +317,7 @@ const UE::HLSLTree::FEmitValue* UE::HLSLTree::FEmitContext::AcquireValue(FFuncti
 			}
 		}
 
-		{
-			const FScopeEntry PoppedScopeEntry = ScopeStack.Pop(false);
-			check(PoppedScopeEntry.Scope == FunctionCall->FunctionScope);
-		}
+		verify(ScopeMap.Remove(FunctionCall->FunctionScope) == 1);
 
 		{
 			const FFunctionStackEntry PoppedStackEntry = FunctionStack.Pop(false);
@@ -385,6 +334,11 @@ const UE::HLSLTree::FEmitValue* UE::HLSLTree::FEmitContext::AcquireValue(FFuncti
 	return &Entry->OutputValues[OutputIndex];
 }
 
+const TCHAR* UE::HLSLTree::FEmitContext::AcquireLocalDeclarationCode()
+{
+	return AllocateStringf(*Allocator, TEXT("Local%d"), NumExpressionLocals++);
+}
+
 const TCHAR* UE::HLSLTree::FEmitContext::GetCode(const FEmitValue* Value) const
 {
 	check(Value);
@@ -453,7 +407,104 @@ void UE::HLSLTree::FEmitContext::AppendPreshader(const FEmitValue* Value, Shader
 	}
 }
 
-void UE::HLSLTree::FNodeVisitor::VisitNode(UE::HLSLTree::FNode* Node)
+UE::HLSLTree::FEmitCode* UE::HLSLTree::FEmitContext::AppendHLSL(FEmitCode& Scope, const TCHAR* Code, int32 InLength)
+{
+	const int32 Length = InLength > 0 ? InLength : FCString::Strlen(Code);
+	TCHAR* InternedCode = new(*Allocator) TCHAR[Length + 1];
+	FMemory::Memcpy(InternedCode, Code, Length * sizeof(TCHAR));
+	InternedCode[Length] = 0;
+
+	FEmitCode* EmitCode = new(*Allocator) FEmitCode();
+	EmitCode->CodeLength = Length;
+	EmitCode->Code = InternedCode;
+	EmitCode->ParentScope = &Scope;
+
+	if (Scope.LastScopeCode)
+	{
+		Scope.LastScopeCode->NextCode = EmitCode;
+	}
+	EmitCode->PrevCode = Scope.LastScopeCode;
+	Scope.LastScopeCode = EmitCode;
+	if (!Scope.FirstScopeCode)
+	{
+		Scope.FirstScopeCode = EmitCode;
+	}
+
+	TotalCodeLength += Length + 1; // include newline
+
+	return EmitCode;
+}
+
+UE::HLSLTree::FEmitCode* UE::HLSLTree::FEmitContext::InsertHLSL(FEmitCode& Scope, FEmitCode* InsertAfter, const TCHAR* Code, int32 InLength)
+{
+	const int32 Length = InLength > 0 ? InLength : FCString::Strlen(Code);
+	TCHAR* InternedCode = new(*Allocator) TCHAR[Length + 1];
+	FMemory::Memcpy(InternedCode, Code, Length * sizeof(TCHAR));
+	InternedCode[Length] = 0;
+
+	FEmitCode* EmitCode = new(*Allocator) FEmitCode();
+	EmitCode->CodeLength = Length;
+	EmitCode->Code = InternedCode;
+	EmitCode->ParentScope = &Scope;
+	EmitCode->PrevCode = InsertAfter;
+	if (InsertAfter)
+	{
+		check(InsertAfter->ParentScope == &Scope);
+
+		EmitCode->NextCode = InsertAfter;
+		if (InsertAfter->NextCode)
+		{
+			InsertAfter->NextCode->PrevCode = EmitCode;
+		}
+		InsertAfter->NextCode = EmitCode;
+	}
+	else
+	{
+		EmitCode->NextCode = Scope.FirstScopeCode;
+		if (Scope.FirstScopeCode)
+		{
+			Scope.FirstScopeCode->PrevCode = EmitCode;
+		}
+		Scope.FirstScopeCode = EmitCode;
+	}
+
+	if (Scope.LastScopeCode == InsertAfter)
+	{
+		Scope.LastScopeCode = EmitCode;
+	}
+
+	TotalCodeLength += Length + 1; // include newline
+
+	return EmitCode;
+}
+
+UE::HLSLTree::FEmitCode* UE::HLSLTree::FEmitContext::AppendHLSLf(FEmitCode& Scope, const TCHAR* Format, ...)
+{
+	TCHAR Code[1024];
+
+	va_list va;
+	va_start(va, Format);
+	const int32 Length = FCString::GetVarArgs(Code, UE_ARRAY_COUNT(Code), Format, va);
+	check(Length >= 0 && Length < UE_ARRAY_COUNT(Code));
+	va_end(va);
+
+	return AppendHLSL(Scope, Code, Length);
+}
+
+UE::HLSLTree::FEmitCode* UE::HLSLTree::FEmitContext::InsertHLSLf(FEmitCode& Scope, FEmitCode* InsertAfter, const TCHAR* Format, ...)
+{
+	TCHAR Code[1024];
+
+	va_list va;
+	va_start(va, Format);
+	const int32 Length = FCString::GetVarArgs(Code, UE_ARRAY_COUNT(Code), Format, va);
+	check(Length >= 0 && Length < UE_ARRAY_COUNT(Code));
+	va_end(va);
+
+	return InsertHLSL(Scope, InsertAfter, Code, Length);
+}
+
+void UE::HLSLTree::FNodeVisitor::VisitNode(FNode* Node)
 {
 	if (Node)
 	{
@@ -461,27 +512,22 @@ void UE::HLSLTree::FNodeVisitor::VisitNode(UE::HLSLTree::FNode* Node)
 	}
 }
 
-UE::HLSLTree::ENodeVisitResult UE::HLSLTree::FStatement::Visit(UE::HLSLTree::FNodeVisitor& Visitor)
+UE::HLSLTree::ENodeVisitResult UE::HLSLTree::FStatement::Visit(FNodeVisitor& Visitor)
 {
 	return Visitor.OnStatement(*this);
 }
 
-UE::HLSLTree::ENodeVisitResult UE::HLSLTree::FExpression::Visit(UE::HLSLTree::FNodeVisitor& Visitor)
+UE::HLSLTree::ENodeVisitResult UE::HLSLTree::FExpression::Visit(FNodeVisitor& Visitor)
 {
 	return Visitor.OnExpression(*this);
 }
 
-UE::HLSLTree::ENodeVisitResult UE::HLSLTree::FLocalDeclaration::Visit(UE::HLSLTree::FNodeVisitor& Visitor)
-{
-	return Visitor.OnLocalDeclaration(*this);
-}
-
-UE::HLSLTree::ENodeVisitResult UE::HLSLTree::FParameterDeclaration::Visit(UE::HLSLTree::FNodeVisitor& Visitor)
+UE::HLSLTree::ENodeVisitResult UE::HLSLTree::FParameterDeclaration::Visit(FNodeVisitor& Visitor)
 {
 	return Visitor.OnParameterDeclaration(*this);
 }
 
-UE::HLSLTree::ENodeVisitResult UE::HLSLTree::FTextureParameterDeclaration::Visit(UE::HLSLTree::FNodeVisitor& Visitor)
+UE::HLSLTree::ENodeVisitResult UE::HLSLTree::FTextureParameterDeclaration::Visit(FNodeVisitor& Visitor)
 {
 	return Visitor.OnTextureParameterDeclaration(*this);
 }
@@ -500,7 +546,7 @@ UE::HLSLTree::ENodeVisitResult UE::HLSLTree::FFunctionCall::Visit(FNodeVisitor&
 	return Result;
 }
 
-UE::HLSLTree::ENodeVisitResult UE::HLSLTree::FScope::Visit(UE::HLSLTree::FNodeVisitor& Visitor)
+UE::HLSLTree::ENodeVisitResult UE::HLSLTree::FScope::Visit(FNodeVisitor& Visitor)
 {
 	const ENodeVisitResult Result = Visitor.OnScope(*this);
 	if (ShouldVisitDependentNodes(Result))
@@ -515,91 +561,56 @@ UE::HLSLTree::ENodeVisitResult UE::HLSLTree::FScope::Visit(UE::HLSLTree::FNodeVi
 	return Result;
 }
 
-bool UE::HLSLTree::FScope::EmitHLSL(FEmitContext& Context, FCodeWriter& OutWriter) const
+bool UE::HLSLTree::FScope::EmitHLSL(FEmitContext& Context, FEmitCode& Scope) const
 {
-	const int32 ScopeIndentLevel = OutWriter.IndentLevel + 1;
-
-	FLocalHLSLCodeWriter DeclarationCodeWriter;
-	DeclarationCodeWriter.IndentLevel = ScopeIndentLevel;
-
-	TMap<FSHAHash, const TCHAR*> LocalExpressionMap;
+	FEmitContext::FScopeEntry* ScopeEntry = new(*Context.Allocator) FEmitContext::FScopeEntry();
+	ScopeEntry->ParentEntry = Context.CurrentScopeEntry;
+	ScopeEntry->Scope = this;
+	ScopeEntry->ScopeCode = &Scope;
 
-	FEmitContext::FScopeEntry& ScopeEntry = Context.ScopeStack.AddDefaulted_GetRef();
-	ScopeEntry.Scope = this;
-	ScopeEntry.ExpressionCodeWriter = &DeclarationCodeWriter;
-	ScopeEntry.ExpressionMap = &LocalExpressionMap;
+	Context.ScopeMap.Add(this, ScopeEntry);
+	Context.CurrentScopeEntry = ScopeEntry;
 
-	OutWriter.WriteLine(TEXT("{"));
-
-	FLocalHLSLCodeWriter StatementCodeWriter;
-	StatementCodeWriter.IndentLevel = ScopeIndentLevel;
 	const FStatement* Statement = FirstStatement;
 	bool bResult = true;
 	while (Statement)
 	{
-		if (!Statement->EmitHLSL(Context, StatementCodeWriter))
+		if (!Statement->EmitHLSL(Context, Scope))
 		{
 			bResult = false;
 			break;
 		}
-
-		// First write any expressions needed by the statement, then the statement itself
-		OutWriter.Append(DeclarationCodeWriter);
-		OutWriter.Append(StatementCodeWriter);
-
-		DeclarationCodeWriter.Reset();
-		StatementCodeWriter.Reset();
-
 		Statement = Statement->NextStatement;
 	}
 
-	OutWriter.WriteLine(TEXT("}"));
-	Context.ScopeStack.Pop(false);
+	Context.CurrentScopeEntry = ScopeEntry->ParentEntry;
 
 	return bResult;
 }
 
-bool UE::HLSLTree::FScope::EmitUnscopedHLSL(FEmitContext& Context, FCodeWriter& OutWriter) const
+bool UE::HLSLTree::FScope::EmitUnscopedHLSL(FEmitContext& Context, FEmitCode& Scope) const
 {
-	FCodeWriter& DeclarationCodeWriter = *Context.ScopeStack.Last().ExpressionCodeWriter;
-
-	FLocalHLSLCodeWriter StatementCodeWriter;
-	StatementCodeWriter.IndentLevel = OutWriter.IndentLevel;
 	const FStatement* Statement = FirstStatement;
 	bool bResult = true;
 	while (Statement)
 	{
-		if (!Statement->EmitHLSL(Context, StatementCodeWriter))
+		if (!Statement->EmitHLSL(Context, Scope))
 		{
 			bResult = false;
 			break;
 		}
-
-		// First write any expressions needed by the statement, then the statement itself
-		OutWriter.Append(DeclarationCodeWriter);
-		OutWriter.Append(StatementCodeWriter);
-
-		DeclarationCodeWriter.Reset();
-		StatementCodeWriter.Reset();
-
 		Statement = Statement->NextStatement;
 	}
 	return bResult;
 }
 
-void UE::HLSLTree::FScope::AddDeclaration(UE::HLSLTree::FLocalDeclaration* Declaration)
-{
-	check(!Declaration->ParentScope);
-	Declaration->ParentScope = this;
-}
-
-void UE::HLSLTree::FScope::AddExpression(UE::HLSLTree::FExpression* Expression)
+void UE::HLSLTree::FScope::AddExpression(FExpression* Expression)
 {
 	check(!Expression->ParentScope);
 	Expression->ParentScope = this;
 }
 
-void UE::HLSLTree::FScope::UseNode(UE::HLSLTree::FNode* Node)
+void UE::HLSLTree::FScope::UseNode(FNode* Node)
 {
 	FScope* Scope0 = this;
 	FScope* Scope1 = Node->ParentScope;
@@ -622,11 +633,6 @@ void UE::HLSLTree::FScope::UseNode(UE::HLSLTree::FNode* Node)
 	Node->ParentScope = Scope0;
 }
 
-void UE::HLSLTree::FScope::UseDeclaration(UE::HLSLTree::FLocalDeclaration* Declaration)
-{
-	UseNode(Declaration);
-}
-
 namespace UE
 {
 namespace HLSLTree
@@ -649,12 +655,6 @@ public:
 		return ENodeVisitResult::VisitDependentNodes;
 	}
 
-	virtual ENodeVisitResult OnLocalDeclaration(FLocalDeclaration& InDeclaration) override
-	{
-		Scope->UseNode(&InDeclaration);
-		return ENodeVisitResult::VisitDependentNodes;
-	}
-
 	virtual ENodeVisitResult OnFunctionCall(FFunctionCall& InFunctionCall) override
 	{
 		Scope->UseNode(&InFunctionCall);
@@ -666,7 +666,7 @@ public:
 } // namespace HLSLTree
 } // namespace UE
 
-void UE::HLSLTree::FScope::UseExpression(UE::HLSLTree::FExpression* Expression)
+void UE::HLSLTree::FScope::UseExpression(FExpression* Expression)
 {
 	// Need to move all dependent expressions/declarations into this scope
 	FNodeVisitor_MoveToScope Visitor(this);
@@ -679,7 +679,7 @@ void UE::HLSLTree::FScope::UseFunctionCall(FFunctionCall* FunctionCall)
 	Visitor.VisitNode(FunctionCall);
 }
 
-void UE::HLSLTree::FScope::AddStatement(UE::HLSLTree::FStatement* Statement)
+void UE::HLSLTree::FScope::AddStatement(FStatement* Statement)
 {
 	check(!Statement->ParentScope);
 	check(!Statement->NextStatement);
@@ -706,34 +706,70 @@ UE::HLSLTree::FTree* UE::HLSLTree::FTree::Create(FMemStackBase& Allocator)
 	return Tree;
 }
 
-bool UE::HLSLTree::FTree::EmitHLSL(UE::HLSLTree::FEmitContext& Context, UE::HLSLTree::FCodeWriter& Writer) const
+//
+static void WriteIndent(int32 IndentLevel, FStringBuilderBase& InOutString)
+{
+	for (int32 i = 0; i < IndentLevel; ++i)
+	{
+		InOutString.Append(TCHAR('\t'));
+	}
+}
+
+static void WriteScope(const UE::HLSLTree::FEmitCode& Code, int32 IndentLevel, FStringBuilderBase& InOutString)
+{
+	const UE::HLSLTree::FEmitCode* EmitCode = &Code;
+	while (EmitCode)
+	{
+		WriteIndent(IndentLevel, InOutString);
+		InOutString.Append(EmitCode->Code, EmitCode->CodeLength);
+		InOutString.Append(TCHAR('\n'));
+
+		if (EmitCode->FirstScopeCode)
+		{
+			WriteIndent(IndentLevel, InOutString);
+			InOutString.Append(TEXT("{\n"));
+			WriteScope(*EmitCode->FirstScopeCode, IndentLevel + 1, InOutString);
+			WriteIndent(IndentLevel, InOutString);
+			InOutString.Append(TEXT("}\n"));
+		}
+		EmitCode = EmitCode->NextCode;
+	}
+}
+
+bool UE::HLSLTree::FTree::EmitHLSL(UE::HLSLTree::FEmitContext& Context, FCodeWriter& Writer) const
 {
-	return RootScope->EmitHLSL(Context, Writer);
+	FEmitCode RootCode;
+	if (RootScope->EmitHLSL(Context, RootCode))
+	{
+		WriteScope(*RootCode.FirstScopeCode, 0, *Writer.StringBuilder);
+		return true;
+	}
+
+	return false;
 }
 
-UE::HLSLTree::FScope* UE::HLSLTree::FTree::NewScope(UE::HLSLTree::FScope& Scope)
+UE::HLSLTree::FScope* UE::HLSLTree::FTree::NewScope(FScope& Scope, FScope** PreviousScopes, int32 NumPreviousScopes)
 {
 	FScope* NewScope = NewNode<FScope>();
 	NewScope->ParentScope = &Scope;
 	NewScope->NestedLevel = Scope.NestedLevel + 1;
+	NewScope->NumPreviousScopes = NumPreviousScopes;
+	if (NumPreviousScopes > 0)
+	{
+		check(PreviousScopes);
+		FMemory::Memcpy(NewScope->PreviousScope, PreviousScopes, sizeof(FScope*) * NumPreviousScopes);
+	}
 	return NewScope;
 }
 
-UE::HLSLTree::FLocalDeclaration* UE::HLSLTree::FTree::NewLocalDeclaration(UE::HLSLTree::FScope& Scope, Shader::EValueType Type, const FName& Name)
-{
-	FLocalDeclaration* Declaration = NewNode<FLocalDeclaration>(Name, Type);
-	Scope.AddDeclaration(Declaration);
-	return Declaration;
-}
-
-UE::HLSLTree::FParameterDeclaration* UE::HLSLTree::FTree::NewParameterDeclaration(UE::HLSLTree::FScope& Scope, const FName& Name, const Shader::FValue& DefaultValue)
+UE::HLSLTree::FParameterDeclaration* UE::HLSLTree::FTree::NewParameterDeclaration(FScope& Scope, const FName& Name, const Shader::FValue& DefaultValue)
 {
 	FParameterDeclaration* Declaration = NewNode<FParameterDeclaration>(Name, DefaultValue);
 	Declaration->ParentScope = &Scope;
 	return Declaration;
 }
 
-UE::HLSLTree::FTextureParameterDeclaration* UE::HLSLTree::FTree::NewTextureParameterDeclaration(UE::HLSLTree::FScope& Scope, const FName& Name, const UE::HLSLTree::FTextureDescription& DefaultValue)
+UE::HLSLTree::FTextureParameterDeclaration* UE::HLSLTree::FTree::NewTextureParameterDeclaration(FScope& Scope, const FName& Name, const FTextureDescription& DefaultValue)
 {
 	FTextureParameterDeclaration* Declaration = NewNode<FTextureParameterDeclaration>(Name, DefaultValue);
 	Declaration->ParentScope = &Scope;
diff --git a/Engine/Source/Runtime/Engine/Private/HLSLTree/HLSLTreeCommon.cpp b/Engine/Source/Runtime/Engine/Private/HLSLTree/HLSLTreeCommon.cpp
--- a/Engine/Source/Runtime/Engine/Private/HLSLTree/HLSLTreeCommon.cpp
+++ b/Engine/Source/Runtime/Engine/Private/HLSLTree/HLSLTreeCommon.cpp
@@ -68,18 +68,6 @@ bool UE::HLSLTree::FExpressionConstant::EmitCode(FEmitContext& Context, FExpress
 	return true;
 }
 
-bool UE::HLSLTree::FExpressionLocalVariable::EmitCode(FEmitContext& Context, FExpressionEmitResult& OutResult) const
-{
-	const FEmitValue* DeclarationValue = Context.AcquireValue(Declaration);
-	if (!DeclarationValue)
-	{
-		return false;
-	}
-
-	OutResult.ForwardValue(Context, DeclarationValue);
-	return true;
-}
-
 bool UE::HLSLTree::FExpressionParameter::EmitCode(FEmitContext& Context, FExpressionEmitResult& OutResult) const
 {
 	OutResult.Type = Declaration->DefaultValue.GetType();
@@ -116,6 +104,55 @@ bool UE::HLSLTree::FExpressionParameter::EmitCode(FEmitContext& Context, FExpres
 	return true;
 }
 
+bool UE::HLSLTree::FExpressionLocalPHI::EmitCode(FEmitContext& Context, FExpressionEmitResult& OutResult) const
+{
+	FEmitContext::FScopeEntry* ScopeEntries[MaxNumPreviousScopes] = { nullptr };
+	const FEmitValue* EmitValues[MaxNumPreviousScopes] = { nullptr };
+	Shader::EValueType CombinedValueType = Shader::EValueType::Void;
+	for (int32 i = 0; i < NumValues; ++i)
+	{
+		const FEmitValue* Value = Context.AcquireValue(Values[i]);
+		if (Value)
+		{
+			const Shader::EValueType ValueType = Value->GetExpressionType();
+			if (CombinedValueType == Shader::EValueType::Void)
+			{
+				CombinedValueType = ValueType;
+			}
+			else if (ValueType != CombinedValueType)
+			{
+				return false;
+			}
+
+			EmitValues[i] = Value;
+		}
+		ScopeEntries[i] = Context.FindScopeEntry(Scopes[i]);
+		check(ScopeEntries[i]);
+	}
+
+	const Shader::FValueTypeDescription TypeDesc = Shader::GetValueTypeDescription(CombinedValueType);
+	const TCHAR* Declaration = Context.AcquireLocalDeclarationCode();
+
+	// Declare the local before the scopes
+	FEmitCode& ScopeCode = *ScopeEntries[0]->ScopeCode;
+	Context.InsertHLSLf(*ScopeCode.ParentScope, ScopeCode.PrevCode, TEXT("%s %s;"), TypeDesc.Name, Declaration);
+
+	// Emit code to assign the value within each scope
+	for (int32 i = 0; i < NumValues; ++i)
+	{
+		if (EmitValues[i])
+		{
+			Context.AppendHLSLf(*ScopeEntries[i]->ScopeCode, TEXT("%s = %s;"), Declaration, Context.GetCode(EmitValues[i]));
+		}
+	}
+
+	OutResult.EvaluationType = EExpressionEvaluationType::Shader;
+	OutResult.Type = CombinedValueType;
+	OutResult.bInline = true;
+	OutResult.Writer.Writef(TEXT("%s"), Declaration);
+	return true;
+}
+
 bool UE::HLSLTree::FExpressionExternalInput::EmitCode(FEmitContext& Context, FExpressionEmitResult& OutResult) const
 {
 	const int32 TypeIndex = (int32)InputType;
@@ -505,32 +542,24 @@ bool UE::HLSLTree::FExpressionFunctionOutput::EmitCode(FEmitContext& Context, FE
 	return true;
 }
 
-bool UE::HLSLTree::FStatementReturn::EmitHLSL(FEmitContext& Context, FCodeWriter& OutWriter) const
+bool UE::HLSLTree::FStatementJump::EmitHLSL(FEmitContext& Context, FEmitCode& Scope) const
 {
-	const FEmitValue* Value = Context.AcquireValue(Expression);
-	if (!Value)
-	{
-		return false;
-	}
-
-	OutWriter.WriteLinef(TEXT("return %s;"), Context.GetCode(Value));
-	return true;
+	return TargetScope->EmitHLSL(Context, Scope);
 }
 
-bool UE::HLSLTree::FStatementSetLocalVariable::EmitHLSL(FEmitContext& Context, FCodeWriter& OutWriter) const
+bool UE::HLSLTree::FStatementReturn::EmitHLSL(FEmitContext& Context, FEmitCode& Scope) const
 {
-	const FEmitValue* DeclarationValue = Context.AcquireValue(Declaration);
-	const FEmitValue* ExpressionValue = Context.AcquireValue(Expression);
-	if (!DeclarationValue || !ExpressionValue)
+	const FEmitValue* Value = Context.AcquireValue(Expression);
+	if (!Value)
 	{
 		return false;
 	}
 
-	OutWriter.WriteLinef(TEXT("%s = %s;"), Context.GetCode(DeclarationValue), Context.GetCode(ExpressionValue));
+	Context.AppendHLSLf(Scope, TEXT("return %s;"), Context.GetCode(Value));
 	return true;
 }
 
-bool UE::HLSLTree::FStatementIf::EmitHLSL(FEmitContext& Context, FCodeWriter& OutWriter) const
+bool UE::HLSLTree::FStatementIf::EmitHLSL(FEmitContext& Context, FEmitCode& Scope) const
 {
 	const FEmitValue* ConditionValue = Context.AcquireValue(ConditionExpression);
 	if (!ConditionValue)
@@ -538,34 +567,33 @@ bool UE::HLSLTree::FStatementIf::EmitHLSL(FEmitContext& Context, FCodeWriter& Ou
 		return false;
 	}
 
-	OutWriter.WriteLinef(TEXT("if (%s)"), Context.GetCode(ConditionValue));
-	ThenScope->EmitHLSL(Context, OutWriter);
+	FEmitCode* ThenScopeCode = Context.AppendHLSLf(Scope, TEXT("if (%s)"), Context.GetCode(ConditionValue));
+	ThenScope->EmitHLSL(Context, *ThenScopeCode);
 	if (ElseScope)
 	{
-		OutWriter.WriteLine(TEXT("else"));
-		ElseScope->EmitHLSL(Context, OutWriter);
+		FEmitCode* ElseScopeCode = Context.AppendHLSLf(Scope, TEXT("else"));
+		ElseScope->EmitHLSL(Context, *ElseScopeCode);
 	}
 	return true;
 }
 
-bool UE::HLSLTree::FStatementFor::EmitHLSL(FEmitContext& Context, FCodeWriter& OutWriter) const
+bool UE::HLSLTree::FStatementFor::EmitHLSL(FEmitContext& Context, FEmitCode& Scope) const
 {
-	const FEmitValue* DeclarationValue = Context.AcquireValue(LoopControlDeclaration);
 	const FEmitValue* StartExpressionValue = Context.AcquireValue(StartExpression);
 	const FEmitValue* EndExpressionValue = Context.AcquireValue(StartExpression);
-	if (!DeclarationValue || !StartExpressionValue || !EndExpressionValue)
+	if (!!StartExpressionValue || !EndExpressionValue)
 	{
 		return false;
 	}
 
-	const TCHAR* DeclarationCode = Context.GetCode(DeclarationValue);
+	const TCHAR* DeclarationCode = TEXT("");
 
-	OutWriter.WriteLinef(TEXT("for (%s = %s; %s < %s; %s++"),
+	FEmitCode* LoopScopeCode = Context.AppendHLSLf(Scope, TEXT("for (%s = %s; %s < %s; %s++"),
 		DeclarationCode,
 		Context.GetCode(StartExpressionValue),
 		DeclarationCode,
 		Context.GetCode(EndExpressionValue),
 		DeclarationCode);
-	LoopScope->EmitHLSL(Context, OutWriter);
+	LoopScope->EmitHLSL(Context, *LoopScopeCode);
 	return true;
 }
diff --git a/Engine/Source/Runtime/Engine/Private/Materials/MaterialExpressionHLSL.cpp b/Engine/Source/Runtime/Engine/Private/Materials/MaterialExpressionHLSL.cpp
--- a/Engine/Source/Runtime/Engine/Private/Materials/MaterialExpressionHLSL.cpp
+++ b/Engine/Source/Runtime/Engine/Private/Materials/MaterialExpressionHLSL.cpp
@@ -102,8 +102,11 @@ EMaterialGenerateHLSLStatus UMaterialExpressionStaticSwitch::GenerateHLSLExpress
 
 EMaterialGenerateHLSLStatus UMaterialExpressionGetLocal::GenerateHLSLExpression(FMaterialHLSLGenerator& Generator, UE::HLSLTree::FScope& Scope, int32 OutputIndex, UE::HLSLTree::FExpression*& OutExpression)
 {
-	UE::HLSLTree::FLocalDeclaration* Declaration = Generator.AcquireLocalDeclaration(Scope, UE::Shader::EValueType::Float3, LocalName);
-	OutExpression = Generator.GetTree().NewExpression<UE::HLSLTree::FExpressionLocalVariable>(Scope, Declaration);
+	OutExpression = Generator.AcquireLocalValue(Scope, LocalName);
+	if (!OutExpression)
+	{
+		return Generator.Error(TEXT("Local accessed before assigned"));
+	}
 	return EMaterialGenerateHLSLStatus::Success;
 }
 
@@ -258,30 +261,7 @@ EMaterialGenerateHLSLStatus UMaterialExpressionReflectionVectorWS::GenerateHLSLE
 
 EMaterialGenerateHLSLStatus UMaterialExpressionFunctionInput::GenerateHLSLExpression(FMaterialHLSLGenerator& Generator, UE::HLSLTree::FScope& Scope, int32 OutputIndex, UE::HLSLTree::FExpression*& OutExpression)
 {
-	UE::Shader::EValueType ExpressionType = UE::Shader::EValueType::Void;
-	switch (InputType)
-	{
-	case FunctionInput_Scalar: ExpressionType = UE::Shader::EValueType::Float1; break;
-	case FunctionInput_Vector2: ExpressionType = UE::Shader::EValueType::Float2; break;
-	case FunctionInput_Vector3: ExpressionType = UE::Shader::EValueType::Float3; break;
-	case FunctionInput_Vector4: ExpressionType = UE::Shader::EValueType::Float4; break;
-	case FunctionInput_Texture2D:
-	case FunctionInput_TextureCube:
-	case FunctionInput_Texture2DArray:
-	case FunctionInput_VolumeTexture:
-	case FunctionInput_StaticBool:
-	case FunctionInput_MaterialAttributes:
-	case FunctionInput_TextureExternal:
-		break;
-	}
-	if (ExpressionType == UE::Shader::EValueType::Void)
-	{
-		return Generator.Error(TEXT("Invalid input connection"));
-	}
-
-	UE::HLSLTree::FLocalDeclaration* Declaration = Generator.AcquireLocalDeclaration(Scope, ExpressionType, InputName);
-	OutExpression = Generator.GetTree().NewExpression<UE::HLSLTree::FExpressionLocalVariable>(Scope, Declaration);
-	return EMaterialGenerateHLSLStatus::Success;
+	return Generator.Error(TEXT("Invalid"));
 }
 
 EMaterialGenerateHLSLStatus UMaterialExpressionMaterialFunctionCall::GenerateHLSLExpression(FMaterialHLSLGenerator& Generator, UE::HLSLTree::FScope& Scope, int32 OutputIndex, UE::HLSLTree::FExpression*& OutExpression)
@@ -325,24 +305,14 @@ EMaterialGenerateHLSLStatus UMaterialExpressionExecEnd::GenerateHLSLStatements(F
 
 EMaterialGenerateHLSLStatus UMaterialExpressionSetLocal::GenerateHLSLStatements(FMaterialHLSLGenerator& Generator, UE::HLSLTree::FScope& Scope)
 {
-	UE::HLSLTree::FExpression* ValueExpression = Value.AcquireHLSLExpressionWithCast(Generator, Scope, UE::Shader::EValueType::Float3);
+	UE::HLSLTree::FExpression* ValueExpression = Value.AcquireHLSLExpression(Generator, Scope);
 	if (!ValueExpression)
 	{
 		return Generator.Error(TEXT("Missing value connection"));
 	}
 
-	UE::HLSLTree::FLocalDeclaration* Declaration = Generator.AcquireLocalDeclaration(Scope, UE::Shader::EValueType::Float3, LocalName);
-	if (!Declaration)
-	{
-		return EMaterialGenerateHLSLStatus::Error;
-	}
-
-	UE::HLSLTree::FStatementSetLocalVariable* Statement = Generator.GetTree().NewStatement<UE::HLSLTree::FStatementSetLocalVariable>(Scope);
-	Statement->Declaration = Declaration;
-	Statement->Expression = ValueExpression;
-
+	Generator.GenerateAssignLocal(Scope, LocalName, ValueExpression);
 	Exec.GenerateHLSLStatements(Generator, Scope);
-
 	return EMaterialGenerateHLSLStatus::Success;
 }
 
@@ -392,12 +362,10 @@ EMaterialGenerateHLSLStatus UMaterialExpressionForLoop::GenerateHLSLStatements(F
 		return Generator.Error(TEXT("Missing EndIndex connection"));
 	}
 
-	UE::HLSLTree::FLocalDeclaration* LoopControlDeclaration = Generator.GetTree().NewLocalDeclaration(Scope, UE::Shader::EValueType::Int1, TEXT("ForLoopControl"));
 	UE::HLSLTree::FScope* LoopScope = Generator.GetTree().NewScope(Scope);
 	LoopBody.GenerateHLSLStatements(Generator, *LoopScope);
 
 	UE::HLSLTree::FStatementFor* ForStatement = Generator.GetTree().NewStatement<UE::HLSLTree::FStatementFor>(Scope);
-	ForStatement->LoopControlDeclaration = LoopControlDeclaration;
 	ForStatement->StartExpression = StartExpression;
 	ForStatement->EndExpression = EndExpression;
 	ForStatement->LoopScope = LoopScope;
diff --git a/Engine/Source/Runtime/Engine/Private/Materials/MaterialHLSLGenerator.cpp b/Engine/Source/Runtime/Engine/Private/Materials/MaterialHLSLGenerator.cpp
--- a/Engine/Source/Runtime/Engine/Private/Materials/MaterialHLSLGenerator.cpp
+++ b/Engine/Source/Runtime/Engine/Private/Materials/MaterialHLSLGenerator.cpp
@@ -256,30 +256,6 @@ UE::HLSLTree::FExpression* FMaterialHLSLGenerator::NewFunctionInput(UE::HLSLTree
 	return Expression;
 }
 
-UE::HLSLTree::FLocalDeclaration* FMaterialHLSLGenerator::AcquireLocalDeclaration(UE::HLSLTree::FScope& Scope, UE::Shader::EValueType Type, const FName& Name)
-{
-	UE::HLSLTree::FLocalDeclaration*& Declaration = LocalDeclarationMap.FindOrAdd(Name);
-	if (!Declaration)
-	{
-		Declaration = HLSLTree->NewLocalDeclaration(Scope, Type, Name);
-	}
-	else
-	{
-		if (Declaration->Type == Type)
-		{
-			Scope.UseDeclaration(Declaration);
-		}
-		else
-		{
-			Errorf(TEXT("Local %s first accessed as type %s, now type %s"), *Name.ToString(),
-				UE::Shader::GetValueTypeDescription(Declaration->Type).Name,
-				UE::Shader::GetValueTypeDescription(Type).Name);
-			return nullptr;
-		}
-	}
-	return Declaration;
-}
-
 UE::HLSLTree::FParameterDeclaration* FMaterialHLSLGenerator::AcquireParameterDeclaration(UE::HLSLTree::FScope& Scope, const FName& Name, const UE::Shader::FValue& DefaultValue)
 {
 	UE::HLSLTree::FParameterDeclaration*& Declaration = ParameterDeclarationMap.FindOrAdd(Name);
@@ -353,6 +329,43 @@ UE::HLSLTree::FFunctionCall* FMaterialHLSLGenerator::AcquireFunctionCall(UE::HLS
 	return FunctionCall;
 }
 
+bool FMaterialHLSLGenerator::GenerateAssignLocal(UE::HLSLTree::FScope& Scope, const FName& LocalName, UE::HLSLTree::FExpression* Value)
+{
+	const FLocalKey Key(&Scope, LocalName);
+	LocalMap.Add(Key, Value);
+	return true;
+}
+
+UE::HLSLTree::FExpression* FMaterialHLSLGenerator::AcquireLocalValue(UE::HLSLTree::FScope& Scope, const FName& LocalName)
+{
+	const FLocalKey Key(&Scope, LocalName);
+	UE::HLSLTree::FExpression** FoundExpression = LocalMap.Find(Key);
+	if (FoundExpression)
+	{
+		return *FoundExpression;
+	}
+
+	const TArrayView<UE::HLSLTree::FScope*> PreviousScopes = Scope.GetPreviousScopes();
+	if (PreviousScopes.Num() > 1)
+	{
+		UE::HLSLTree::FExpressionLocalPHI* Expression = HLSLTree->NewExpression<UE::HLSLTree::FExpressionLocalPHI>(Scope);
+		Expression->NumValues = PreviousScopes.Num();
+		for (int32 i = 0; i < PreviousScopes.Num(); ++i)
+		{
+			Expression->Scopes[i] = PreviousScopes[i];
+			Expression->Values[i] = AcquireLocalValue(*PreviousScopes[i], LocalName);
+		}
+		LocalMap.Add(Key, Expression);
+		return Expression;
+	}
+
+	if (Scope.ParentScope)
+	{
+		return AcquireLocalValue(*Scope.ParentScope, LocalName);
+	}
+	return nullptr;
+}
+
 UE::HLSLTree::FExpression* FMaterialHLSLGenerator::AcquireExpression(UE::HLSLTree::FScope& Scope, UMaterialExpression* MaterialExpression, int32 OutputIndex)
 {
 	const FExpressionKey Key(MaterialExpression, OutputIndex);
@@ -384,7 +397,8 @@ UE::HLSLTree::FTextureParameterDeclaration* FMaterialHLSLGenerator::AcquireTextu
 bool FMaterialHLSLGenerator::GenerateStatements(UE::HLSLTree::FScope& Scope, UMaterialExpression* MaterialExpression)
 {
 	FStatementEntry& Entry = StatementMap.FindOrAdd(MaterialExpression);
-	Entry.NumInputs++;
+	Entry.PreviousScope[Entry.NumInputs++] = &Scope;
+	check(Entry.NumInputs <= MaxNumPreviousScopes);
 	check(Entry.NumInputs <= MaterialExpression->NumExecutionInputs);
 
 	bool bResult = true;
@@ -393,7 +407,10 @@ bool FMaterialHLSLGenerator::GenerateStatements(UE::HLSLTree::FScope& Scope, UMa
 		UE::HLSLTree::FScope* ScopeToUse = &Scope;
 		if (MaterialExpression->NumExecutionInputs > 1u)
 		{
-			ScopeToUse = ScopeToUse->ParentScope;
+			ScopeToUse = HLSLTree->NewScope(*ScopeToUse->ParentScope, Entry.PreviousScope, Entry.NumInputs);
+			// Add a jump to link the previous scope to this new scope
+			UE::HLSLTree::FStatementJump* JumpStatement = HLSLTree->NewStatement<UE::HLSLTree::FStatementJump>(*Scope.ParentScope);
+			JumpStatement->TargetScope = ScopeToUse;
 		}
 
 		const FExpressionKey Key(MaterialExpression);
diff --git a/Engine/Source/Runtime/Engine/Public/HLSLTree/HLSLTree.h b/Engine/Source/Runtime/Engine/Public/HLSLTree/HLSLTree.h
--- a/Engine/Source/Runtime/Engine/Public/HLSLTree/HLSLTree.h
+++ b/Engine/Source/Runtime/Engine/Public/HLSLTree/HLSLTree.h
@@ -3,6 +3,7 @@
 
 #include "CoreMinimal.h"
 #include "Engine/EngineTypes.h"
+#include "Containers/ArrayView.h"
 #include "Misc/StringBuilder.h"
 #include "Misc/MemStack.h"
 #include "HLSLTree/HLSLTreeTypes.h"
@@ -27,6 +28,8 @@ class FPreshaderData;
 namespace HLSLTree
 {
 
+static constexpr int32 MaxNumPreviousScopes = 2;
+
 /** Allows building a string incrementally, with some additional features to support code generation, such as indent handling */
 class FCodeWriter
 {
@@ -85,6 +88,20 @@ public:
 	int32 IndentLevel;
 };
 
+class FEmitCode
+{
+public:
+	const TCHAR* Code = nullptr;
+	FEmitCode* ParentScope = nullptr;
+	FEmitCode* NextCode = nullptr;
+	FEmitCode* PrevCode = nullptr;
+	FEmitCode* FirstScopeCode = nullptr;
+	FEmitCode* LastScopeCode = nullptr;
+	int32 CodeLength = 0;
+
+	friend class FEmitContext;
+};
+
 class FEmitValue
 {
 public:
@@ -112,23 +129,30 @@ public:
 	/** Returns a value that references the given expression */
 	const FEmitValue* AcquireValue(FExpression* Expression);
 
-	/** Returns value that references the given local variable */
-	const FEmitValue* AcquireValue(FLocalDeclaration* Declaration);
-
 	/** Returns value that references the given local function output */
 	const FEmitValue* AcquireValue(FFunctionCall* FunctionCall, int32 OutputIndex);
 
+	/** Get a unique local variable name */
+	const TCHAR* AcquireLocalDeclarationCode();
+
 	/** Gets HLSL code that references the given value */
 	const TCHAR* GetCode(const FEmitValue* Value) const;
 
 	/** Append preshader bytecode that represents the given value */
 	void AppendPreshader(const FEmitValue* Value, Shader::FPreshaderData& InOutPreshader) const;
 
+	FEmitCode* AppendHLSL(FEmitCode& Scope, const TCHAR* Code, int32 Length = 0);
+	FEmitCode* InsertHLSL(FEmitCode& Scope, FEmitCode* InsertAfter, const TCHAR* Code, int32 Length = 0);
+
+	FEmitCode* AppendHLSLf(FEmitCode& Scope, const TCHAR* Format, ...);
+	FEmitCode* InsertHLSLf(FEmitCode& Scope, FEmitCode* InsertAfter, const TCHAR* Format, ...);
+
 	struct FScopeEntry
 	{
+		FScopeEntry* ParentEntry;
 		const FScope* Scope;
-		FCodeWriter* ExpressionCodeWriter;
-		TMap<FSHAHash, const TCHAR*>* ExpressionMap;
+		FEmitCode* ScopeCode;
+		TMap<FSHAHash, const TCHAR*> ExpressionMap;
 	};
 
 	struct FDeclarationEntry
@@ -150,16 +174,18 @@ public:
 		TMap<FFunctionCall*, FFunctionCallEntry*> FunctionCallMap;
 	};
 
-	FScopeEntry* FindScope(FScope* Scope);
-	int32 FindScopeIndex(FScope* Scope);
+	FScopeEntry* FindScopeEntry(const FScope* Scope);
+
+	TMap<const FScope*, FScopeEntry*> ScopeMap;
+	FScopeEntry* CurrentScopeEntry = nullptr;
 
-	TArray<FScopeEntry> ScopeStack;
 	TArray<FFunctionStackEntry> FunctionStack;
 	TArray<Shader::FPreshaderData*> TempPreshaders;
 	FMemStackBase* Allocator = nullptr;
 	const FMaterial* Material = nullptr; // TODO - remove preshader material dependency
 	const FStaticParameterSet* StaticParameters = nullptr;
 	FMaterialCompilationOutput* MaterialCompilationOutput = nullptr;
+	int32 TotalCodeLength = 0;
 	int32 NumExpressionLocals = 0;
 	int32 NumTexCoords = 0;
 };
@@ -200,7 +226,6 @@ public:
 	virtual ENodeVisitResult OnScope(FScope& Scope) { return ENodeVisitResult::VisitDependentNodes; }
 	virtual ENodeVisitResult OnStatement(FStatement& Statement) { return ENodeVisitResult::VisitDependentNodes; }
 	virtual ENodeVisitResult OnExpression(FExpression& Expression) { return ENodeVisitResult::VisitDependentNodes; }
-	virtual ENodeVisitResult OnLocalDeclaration(FLocalDeclaration& Declaration) { return ENodeVisitResult::VisitDependentNodes; }
 	virtual ENodeVisitResult OnParameterDeclaration(FParameterDeclaration& Declaration) { return ENodeVisitResult::VisitDependentNodes; }
 	virtual ENodeVisitResult OnTextureParameterDeclaration(FTextureParameterDeclaration& Declaration) { return ENodeVisitResult::VisitDependentNodes; }
 	virtual ENodeVisitResult OnFunctionCall(FFunctionCall& FunctionCall) { return ENodeVisitResult::VisitDependentNodes; }
@@ -227,11 +252,11 @@ public:
 class FStatement : public FNode
 {
 public:
-	/** Emits HLSL code for the statement. The generated code should include any required semi-colons and newlines */
-	virtual bool EmitHLSL(FEmitContext& Context, FCodeWriter& Writer) const = 0;
-
 	virtual ENodeVisitResult Visit(FNodeVisitor& Visitor) override;
 
+	/** Emits HLSL code for the statement. The generated code should include any required semi-colons and newlines */
+	virtual bool EmitHLSL(FEmitContext& Context, FEmitCode& Scope) const = 0;
+
 	FStatement* NextStatement = nullptr;
 };
 
@@ -241,7 +266,6 @@ public:
  * Expressions track the outer-most scope in which they're accessed. The generated HLSL code will ensure they are defined in that scope.
  * Examples include constant literals, variable accessors, and various types of math operations
  * This is an abstract base class, with derived classes representing various types of expression
- * Derived classes are expected to implement at least one of EmitHLSL or EmitPreshader
  */
 class FExpression : public FNode
 {
@@ -252,20 +276,6 @@ public:
 	virtual bool EmitCode(FEmitContext& Context, FExpressionEmitResult& OutResult) const = 0;
 };
 
-/**
- * Represents an HLSL local variable.  This is used by statements/expressions to refer to a particular variable
- */
-class FLocalDeclaration final : public FNode
-{
-public:
-	FLocalDeclaration(const FName& InName, Shader::EValueType InType) : Name(InName), Type(InType) {}
-
-	virtual ENodeVisitResult Visit(FNodeVisitor& Visitor) override;
-
-	FName Name;
-	Shader::EValueType Type;
-};
-
 /**
  * Represents an HLSL parameter.  This is a uniform constant passed to the generated shader
  */
@@ -322,16 +332,20 @@ public:
 class FScope final : public FNode
 {
 public:
+	inline TArrayView<FScope*> GetPreviousScopes() const
+	{
+		// const_cast needed, otherwise type of array view is 'FScope*const' which doesn't make sense
+		return MakeArrayView(const_cast<FScope*>(this)->PreviousScope, NumPreviousScopes);
+	}
+
 	virtual ENodeVisitResult Visit(FNodeVisitor& Visitor) override;
-	bool EmitHLSL(FEmitContext& Context, FCodeWriter& OutWriter) const;
+	bool EmitHLSL(FEmitContext& Context, FEmitCode& Scope) const;
 
-	bool EmitUnscopedHLSL(FEmitContext& Context, FCodeWriter& OutWriter) const;
+	bool EmitUnscopedHLSL(FEmitContext& Context, FEmitCode& Scope) const;
 
-	void AddDeclaration(FLocalDeclaration* Declaration);
 	void AddExpression(FExpression* Expression);
 	void AddStatement(FStatement* Statement);
 
-	void UseDeclaration(FLocalDeclaration* Declaration);
 	void UseFunctionCall(FFunctionCall* FunctionCall);
 	void UseExpression(FExpression* Expression);
 private:
@@ -340,8 +354,10 @@ private:
 
 	void UseNode(FNode* Node);
 
+	FScope* PreviousScope[MaxNumPreviousScopes];
 	FStatement* FirstStatement = nullptr;
 	FStatement* LastStatement = nullptr;
+	int32 NumPreviousScopes = 0;
 	int32 NestedLevel = 0;
 };
 
@@ -373,8 +389,7 @@ public:
 		return Statement;
 	}
 
-	FScope* NewScope(FScope& Scope);
-	FLocalDeclaration* NewLocalDeclaration(FScope& Scope, Shader::EValueType Type, const FName& Name);
+	FScope* NewScope(FScope& Scope, FScope** PreviousScopes = nullptr, int32 NumPreviousScopes = 0);
 	FParameterDeclaration* NewParameterDeclaration(FScope& Scope, const FName& Name, const Shader::FValue& DefaultValue);
 	FTextureParameterDeclaration* NewTextureParameterDeclaration(FScope& Scope, const FName& Name, const FTextureDescription& DefaultValue);
 
diff --git a/Engine/Source/Runtime/Engine/Public/HLSLTree/HLSLTreeCommon.h b/Engine/Source/Runtime/Engine/Public/HLSLTree/HLSLTreeCommon.h
--- a/Engine/Source/Runtime/Engine/Public/HLSLTree/HLSLTreeCommon.h
+++ b/Engine/Source/Runtime/Engine/Public/HLSLTree/HLSLTreeCommon.h
@@ -29,26 +29,6 @@ public:
 	virtual bool EmitCode(FEmitContext& Context, FExpressionEmitResult& OutResult) const override;
 };
 
-class FExpressionLocalVariable : public FExpression
-{
-public:
-	explicit FExpressionLocalVariable(FLocalDeclaration* InDeclaration) : Declaration(InDeclaration) {}
-
-	FLocalDeclaration* Declaration;
-
-	virtual ENodeVisitResult Visit(FNodeVisitor& Visitor) override
-	{
-		const ENodeVisitResult Result = FExpression::Visit(Visitor);
-		if (ShouldVisitDependentNodes(Result))
-		{
-			Visitor.VisitNode(Declaration);
-		}
-		return Result;
-	}
-
-	virtual bool EmitCode(FEmitContext& Context, FExpressionEmitResult& OutResult) const override;
-};
-
 class FExpressionParameter : public FExpression
 {
 public:
@@ -71,6 +51,21 @@ public:
 	virtual bool EmitCode(FEmitContext& Context, FExpressionEmitResult& OutResult) const override;
 };
 
+/**
+ * Represents a phi node (see various topics on single static assignment)
+ * A phi node takes on a value based on the previous scope that was executed.
+ * In practice, this means the generated HLSL code will declare a local variable before all the previous scopes, then assign that variable the proper value from within each scope
+ */
+class FExpressionLocalPHI final : public FExpression
+{
+public:
+	virtual bool EmitCode(FEmitContext& Context, FExpressionEmitResult& OutResult) const override;
+
+	FScope* Scopes[MaxNumPreviousScopes];
+	FExpression* Values[MaxNumPreviousScopes];
+	int32 NumValues;
+};
+
 enum class EExternalInputType
 {
 	TexCoord0,
@@ -323,28 +318,27 @@ public:
 	virtual bool EmitCode(FEmitContext& Context, FExpressionEmitResult& OutResult) const override;
 };
 
-class FStatementReturn : public FStatement
+class FStatementJump : public FStatement
 {
 public:
-	FExpression* Expression;
+	FScope* TargetScope;
 
 	virtual ENodeVisitResult Visit(FNodeVisitor& Visitor) override
 	{
 		const ENodeVisitResult Result = FStatement::Visit(Visitor);
 		if (ShouldVisitDependentNodes(Result))
 		{
-			Visitor.VisitNode(Expression);
+			Visitor.VisitNode(TargetScope);
 		}
 		return Result;
 	}
 
-	virtual bool EmitHLSL(FEmitContext& Context, FCodeWriter& Writer) const override;
+	virtual bool EmitHLSL(FEmitContext& Context, FEmitCode& Scope) const override;
 };
 
-class FStatementSetLocalVariable : public FStatement
+class FStatementReturn : public FStatement
 {
 public:
-	FLocalDeclaration* Declaration;
 	FExpression* Expression;
 
 	virtual ENodeVisitResult Visit(FNodeVisitor& Visitor) override
@@ -352,13 +346,12 @@ public:
 		const ENodeVisitResult Result = FStatement::Visit(Visitor);
 		if (ShouldVisitDependentNodes(Result))
 		{
-			Visitor.VisitNode(Declaration);
 			Visitor.VisitNode(Expression);
 		}
 		return Result;
 	}
 
-	virtual bool EmitHLSL(FEmitContext& Context, FCodeWriter& Writer) const override;
+	virtual bool EmitHLSL(FEmitContext& Context, FEmitCode& Scope) const override;
 };
 
 class FStatementIf : public FStatement
@@ -380,13 +373,12 @@ public:
 		return Result;
 	}
 
-	virtual bool EmitHLSL(FEmitContext& Context, FCodeWriter& Writer) const override;
+	virtual bool EmitHLSL(FEmitContext& Context, FEmitCode& Scope) const override;
 };
 
 class FStatementFor : public FStatement
 {
 public:
-	FLocalDeclaration* LoopControlDeclaration;
 	FExpression* StartExpression;
 	FExpression* EndExpression;
 	FScope* LoopScope;
@@ -396,7 +388,6 @@ public:
 		const ENodeVisitResult Result = FStatement::Visit(Visitor);
 		if (ShouldVisitDependentNodes(Result))
 		{
-			Visitor.VisitNode(LoopControlDeclaration);
 			Visitor.VisitNode(StartExpression);
 			Visitor.VisitNode(EndExpression);
 			Visitor.VisitNode(LoopScope);
@@ -404,7 +395,7 @@ public:
 		return Result;
 	}
 
-	virtual bool EmitHLSL(FEmitContext& Context, FCodeWriter& Writer) const override;
+	virtual bool EmitHLSL(FEmitContext& Context, FEmitCode& Scope) const override;
 };
 
 }
diff --git a/Engine/Source/Runtime/Engine/Public/HLSLTree/HLSLTreeTypes.h b/Engine/Source/Runtime/Engine/Public/HLSLTree/HLSLTreeTypes.h
--- a/Engine/Source/Runtime/Engine/Public/HLSLTree/HLSLTreeTypes.h
+++ b/Engine/Source/Runtime/Engine/Public/HLSLTree/HLSLTreeTypes.h
@@ -13,7 +13,6 @@ namespace HLSLTree
 
 class FNode;
 class FExpression;
-class FLocalDeclaration;
 class FParameterDeclaration;
 class FTextureParameterDeclaration;
 class FStatement;
diff --git a/Engine/Source/Runtime/Engine/Public/MaterialHLSLGenerator.h b/Engine/Source/Runtime/Engine/Public/MaterialHLSLGenerator.h
--- a/Engine/Source/Runtime/Engine/Public/MaterialHLSLGenerator.h
+++ b/Engine/Source/Runtime/Engine/Public/MaterialHLSLGenerator.h
@@ -33,7 +33,6 @@ class FExpressionConstant;
 class FExpressionExternalInput;
 class FExpressionSwizzle;
 class FExpressionCast;
-class FFunctionCall;
 }
 }
 
@@ -74,7 +73,6 @@ public:
 
 	UE::HLSLTree::FExpression* NewFunctionInput(UE::HLSLTree::FScope& Scope, int32 InputIndex, UMaterialExpressionFunctionInput* MaterialFunctionInput);
 
-	UE::HLSLTree::FLocalDeclaration* AcquireLocalDeclaration(UE::HLSLTree::FScope& Scope, UE::Shader::EValueType Type, const FName& Name);
 	UE::HLSLTree::FParameterDeclaration* AcquireParameterDeclaration(UE::HLSLTree::FScope& Scope, const FName& Name, const UE::Shader::FValue& DefaultValue);
 
 	/** Returns a declaration to access the given texture, with no parameter */
@@ -85,6 +83,9 @@ public:
 
 	UE::HLSLTree::FFunctionCall* AcquireFunctionCall(UE::HLSLTree::FScope& Scope, UMaterialFunctionInterface* Function, TArrayView<UE::HLSLTree::FExpression*> Inputs);
 
+	bool GenerateAssignLocal(UE::HLSLTree::FScope& Scope, const FName& LocalName, UE::HLSLTree::FExpression* Value);
+	UE::HLSLTree::FExpression* AcquireLocalValue(UE::HLSLTree::FScope& Scope, const FName& LocalName);
+
 	/**
 	 * Returns the appropriate HLSLNode representing the given UMaterialExpression.
 	 * The node will be created if it doesn't exist. Otherwise, the tree will be updated to ensure the node is visible in the given scope
@@ -95,6 +96,7 @@ public:
 	bool GenerateStatements(UE::HLSLTree::FScope& Scope, UMaterialExpression* MaterialExpression);
 
 private:
+	static constexpr int32 MaxNumPreviousScopes = UE::HLSLTree::MaxNumPreviousScopes;
 	struct FExpressionKey
 	{
 		explicit FExpressionKey(UMaterialExpression* InExpression, int32 InOutputIndex = INDEX_NONE) : Expression(InExpression), OutputIndex(InOutputIndex) {}
@@ -129,8 +131,27 @@ private:
 		}
 	};
 
+	struct FLocalKey
+	{
+		FLocalKey(UE::HLSLTree::FScope* InScope, const FName& InName) : Scope(InScope), Name(InName) {}
+
+		UE::HLSLTree::FScope* Scope;
+		FName Name;
+
+		friend inline uint32 GetTypeHash(const FLocalKey& Value)
+		{
+			return HashCombine(GetTypeHash(Value.Scope), GetTypeHash(Value.Name));
+		}
+
+		friend inline bool operator==(const FLocalKey& Lhs, const FLocalKey& Rhs)
+		{
+			return Lhs.Scope == Rhs.Scope && Lhs.Name == Rhs.Name;
+		}
+	};
+
 	struct FStatementEntry
 	{
+		UE::HLSLTree::FScope* PreviousScope[MaxNumPreviousScopes];
 		int32 NumInputs = 0;
 	};
 
@@ -142,11 +163,11 @@ private:
 	TArray<FExpressionKey> ExpressionStack;
 	TArray<FString> CompileErrors;
 	TArray<UMaterialExpression*> ErrorExpressions;
-	TMap<FName, UE::HLSLTree::FLocalDeclaration*> LocalDeclarationMap;
 	TMap<FName, UE::HLSLTree::FParameterDeclaration*> ParameterDeclarationMap;
 	TMap<UE::HLSLTree::FTextureDescription, UE::HLSLTree::FTextureParameterDeclaration*> TextureDeclarationMap;
 	TMap<FName, UE::HLSLTree::FTextureParameterDeclaration*> TextureParameterDeclarationMap;
 	TMap<FFunctionCallKey, UE::HLSLTree::FFunctionCall*> FunctionCallMap;
+	TMap<FLocalKey, UE::HLSLTree::FExpression*> LocalMap;
 	TMap<FExpressionKey, UE::HLSLTree::FExpression*> ExpressionMap;
 	TMap<UMaterialExpression*, FStatementEntry> StatementMap;
 	bool bGeneratedResult;
diff --git a/Engine/Source/Runtime/Engine/Public/MaterialHLSLTree.h b/Engine/Source/Runtime/Engine/Public/MaterialHLSLTree.h
--- a/Engine/Source/Runtime/Engine/Public/MaterialHLSLTree.h
+++ b/Engine/Source/Runtime/Engine/Public/MaterialHLSLTree.h
@@ -17,7 +17,6 @@ namespace HLSLTree
 class FTree;
 class FScope;
 class FExpression;
-class FLocalDeclaration;
 class FFunctionCall;
 }
 }
