remove unused and incorrect ComputePsnr function for DXT blocks

diff --git a/Engine/Source/Runtime/Engine/Private/TextureDerivedData.cpp b/Engine/Source/Runtime/Engine/Private/TextureDerivedData.cpp
--- a/Engine/Source/Runtime/Engine/Private/TextureDerivedData.cpp
+++ b/Engine/Source/Runtime/Engine/Private/TextureDerivedData.cpp
@@ -713,168 +713,6 @@ uint32 PutDerivedDataInCache(FTexturePlatformData* DerivedData, const FString& D
 
 #if WITH_EDITOR
 
-/**
- * Unpack a DXT 565 color to RGB32.
- */
-static uint16 UnpackDXTColor(int32* OutColors, const uint8* Block)
-{
-	uint16 PackedColor = (Block[1] << 8) | Block[0];
-	uint8 Red = (PackedColor >> 11) & 0x1f;
-	OutColors[0] = (Red << 3) | ( Red >> 2);
-	uint8 Green = (PackedColor >> 5) & 0x3f;
-	OutColors[1] = (Green << 2) | ( Green >> 4);
-	uint8 Blue = PackedColor & 0x1f;
-	OutColors[2] = (Blue << 3) | ( Blue >> 2);
-	return PackedColor;
-}
-
-/**
- * Computes the squared error between a DXT compression block and the source colors.
- */
-static double ComputeDXTColorBlockSquaredError(const uint8* Block, const FColor* Colors, int32 ColorPitch)
-{
-	int32 ColorTable[4][3];
-
-	uint16 c0 = UnpackDXTColor(ColorTable[0], Block);
-	uint16 c1 = UnpackDXTColor(ColorTable[1], Block + 2);
-	if (c0 > c1)
-	{
-		for (int32 ColorIndex = 0; ColorIndex < 3; ++ColorIndex)
-		{
-			ColorTable[2][ColorIndex] = (2 * ColorTable[0][ColorIndex]) / 3 + (1 * ColorTable[1][ColorIndex]) / 3;
-			ColorTable[3][ColorIndex] = (1 * ColorTable[0][ColorIndex]) / 3 + (2 * ColorTable[1][ColorIndex]) / 3;
-		}
-	}
-	else
-	{
-		for (int32 ColorIndex = 0; ColorIndex < 3; ++ColorIndex)
-		{
-			ColorTable[2][ColorIndex] = (1 * ColorTable[0][ColorIndex]) / 2 + (1 * ColorTable[1][ColorIndex]) / 2;
-			ColorTable[3][ColorIndex] = 0;
-		}
-	}
-
-	double SquaredError = 0.0;
-	for (int32 Y = 0; Y < 4; ++Y)
-	{
-		uint8 IndexTable[4];
-		uint8 RowIndices = Block[4+Y];
-		IndexTable[0] = RowIndices & 0x3;
-		IndexTable[1] = (RowIndices >> 2) & 0x3;
-		IndexTable[2] = (RowIndices >> 4) & 0x3;
-		IndexTable[3] = (RowIndices >> 6) & 0x3;
-
-		for (int32 X = 0; X < 4; ++X)
-		{
-			FColor Color = Colors[Y * ColorPitch + X];
-			int32* DXTColor = ColorTable[IndexTable[X]];
-			SquaredError += ((int32)Color.R - DXTColor[0]) * ((int32)Color.R - DXTColor[0]);
-			SquaredError += ((int32)Color.G - DXTColor[1]) * ((int32)Color.G - DXTColor[1]);
-			SquaredError += ((int32)Color.B - DXTColor[2]) * ((int32)Color.B - DXTColor[2]);
-		}
-	}
-	return SquaredError;
-}
-
-/**
- * Computes the squared error between the alpha values in the block and the source colors.
- */
-static double ComputeDXTAlphaBlockSquaredError(const uint8* Block, const FColor* Colors, int32 ColorPitch)
-{
-	int32 AlphaTable[8];
-
-	int32 a0 = Block[0];
-	int32 a1 = Block[1];
-
-	AlphaTable[0] = a0;
-	AlphaTable[1] = a1;
-	if (AlphaTable[0] > AlphaTable[1])
-	{
-		for (int32 AlphaIndex = 0; AlphaIndex < 6; ++AlphaIndex)
-		{
-			AlphaTable[AlphaIndex+2] = ((6 - AlphaIndex) * a0 + (1 + AlphaIndex) * a1) / 7;
-		}
-	}
-	else
-	{
-		for (int32 AlphaIndex = 0; AlphaIndex < 4; ++AlphaIndex)
-		{
-			AlphaTable[AlphaIndex+2] = ((4 - AlphaIndex) * a0 + (1 + AlphaIndex) * a1) / 5;
-		}
-		AlphaTable[6] = 0;
-		AlphaTable[7] = 255;
-	}
-
-	uint64 IndexBits = (uint64)Block[7];
-	IndexBits = (IndexBits << 8) | (uint64)Block[6];
-	IndexBits = (IndexBits << 8) | (uint64)Block[5];
-	IndexBits = (IndexBits << 8) | (uint64)Block[4];
-	IndexBits = (IndexBits << 8) | (uint64)Block[3];
-	IndexBits = (IndexBits << 8) | (uint64)Block[2];
-
-	double SquaredError = 0.0;
-	for (int32 Y = 0; Y < 4; ++Y)
-	{
-		for (int32 X = 0; X < 4; ++X)
-		{
-			const FColor Color = Colors[Y * ColorPitch + X];
-			uint8 Index = IndexBits & 0x7;
-			SquaredError += ((int32)Color.A - AlphaTable[Index]) * ((int32)Color.A - AlphaTable[Index]);
-			IndexBits = IndexBits >> 3;
-		}
-	}
-	return SquaredError;
-}
-
-/**
- * Computes the PSNR value for the compressed image.
- */
-static float ComputePSNR(const FImage& SrcImage, const FCompressedImage2D& CompressedImage)
-{
-	double SquaredError = 0.0;
-	int32 NumErrors = 0;
-	const uint8* CompressedData = CompressedImage.RawData.GetData();
-
-	if (SrcImage.Format == ERawImageFormat::BGRA8 && (CompressedImage.PixelFormat == PF_DXT1 || CompressedImage.PixelFormat == PF_DXT5))
-	{
-		int32 NumBlocksX = CompressedImage.SizeX / 4;
-		int32 NumBlocksY = CompressedImage.SizeY / 4;
-		for (int32 BlockY = 0; BlockY < NumBlocksY; ++BlockY)
-		{
-			for (int32 BlockX = 0; BlockX < NumBlocksX; ++BlockX)
-			{
-				if (CompressedImage.PixelFormat == PF_DXT1)
-				{
-					SquaredError += ComputeDXTColorBlockSquaredError(
-						CompressedData + (BlockY * NumBlocksX + BlockX) * 8,
-						(&SrcImage.AsBGRA8()[0]) + (BlockY * NumBlocksX * 16 + BlockX * 4),
-						SrcImage.SizeX
-						);
-					NumErrors += 16 * 3;
-				}
-				else if (CompressedImage.PixelFormat == PF_DXT5)
-				{
-					SquaredError += ComputeDXTAlphaBlockSquaredError(
-						CompressedData + (BlockY * NumBlocksX + BlockX) * 16,
-						(&SrcImage.AsBGRA8()[0]) + (BlockY * NumBlocksX * 16 + BlockX * 4),
-						SrcImage.SizeX
-						);
-					SquaredError += ComputeDXTColorBlockSquaredError(
-						CompressedData + (BlockY * NumBlocksX + BlockX) * 16 + 8,
-						(&SrcImage.AsBGRA8()[0]) + (BlockY * NumBlocksX * 16 + BlockX * 4),
-						SrcImage.SizeX
-						);
-					NumErrors += 16 * 4;
-				}
-			}
-		}
-	}
-
-	double MeanSquaredError = NumErrors > 0 ? SquaredError / (double)NumErrors : 0.0;
-	double RMSE = FMath::Sqrt(MeanSquaredError);
-	return RMSE > 0.0 ? 20.0f * (float)log10(255.0 / RMSE) : 500.0f;
-}
-
 void FTexturePlatformData::Cache(
 	UTexture& InTexture,
 	const FTextureBuildSettings* InSettingsPerLayer,
