[Navigation] Deprecate unused NavigationGeometryGatheringMode property in AStaticMeshActor

diff --git a/Engine/Source/Runtime/Engine/Classes/Engine/StaticMeshActor.h b/Engine/Source/Runtime/Engine/Classes/Engine/StaticMeshActor.h
--- a/Engine/Source/Runtime/Engine/Classes/Engine/StaticMeshActor.h
+++ b/Engine/Source/Runtime/Engine/Classes/Engine/StaticMeshActor.h
@@ -33,6 +33,7 @@ public:
 	UPROPERTY(Category=Actor, EditAnywhere, AdvancedDisplay)
 	bool bStaticMeshReplicateMovement;
 
+	UE_DEPRECATED(5.0, "Unused property. The actor will use the DefaultGeometryGatheringMode set in FNavigationOctree (see virtual ENavDataGatheringMode GetGeometryGatheringMode()).")
 	UPROPERTY(EditAnywhere, Category = Navigation, AdvancedDisplay)
 	ENavDataGatheringMode NavigationGeometryGatheringMode;
 
