DataLayer: Expose Debug Color to Blueprint

diff --git a/Engine/Source/Runtime/Engine/Public/WorldPartition/DataLayer/DataLayer.h b/Engine/Source/Runtime/Engine/Public/WorldPartition/DataLayer/DataLayer.h
--- a/Engine/Source/Runtime/Engine/Public/WorldPartition/DataLayer/DataLayer.h
+++ b/Engine/Source/Runtime/Engine/Public/WorldPartition/DataLayer/DataLayer.h
@@ -64,6 +64,7 @@ public:
 	UFUNCTION(Category = "Data Layer - Runtime", BlueprintCallable)
 	EDataLayerState GetInitialState() const { return IsDynamicallyLoaded() ? InitialState : EDataLayerState::Unloaded; }
 
+	UFUNCTION(Category = "Data Layer", BlueprintCallable)
 	FColor GetDebugColor() const { return DebugColor; }
 
 	/** Returns a sanitized version of the provided Data Layer Label */
