Added BlueprintThreadSafe metadata tag to notify transition window functions.

diff --git a/Engine/Source/Runtime/Engine/Classes/Animation/AnimInstance.h b/Engine/Source/Runtime/Engine/Classes/Animation/AnimInstance.h
--- a/Engine/Source/Runtime/Engine/Classes/Animation/AnimInstance.h
+++ b/Engine/Source/Runtime/Engine/Classes/Animation/AnimInstance.h
@@ -1020,39 +1020,39 @@ public:
 	float GetRelevantAnimTimeFraction(int32 MachineIndex, int32 StateIndex);
 
 	/** Get whether a particular notify state was active in any state machine last tick.*/
-	UFUNCTION(BlueprintPure, Category = "Asset Player", meta = (Keywords = "Window,TransitionWindow,NotifyState,NotifyStateTransition"))
+	UFUNCTION(BlueprintPure, Category = "Asset Player", meta = (BlueprintThreadSafe, Keywords = "Window,TransitionWindow,NotifyState,NotifyStateTransition"))
 	bool WasAnimNotifyStateActiveInAnyState(TSubclassOf<UAnimNotifyState> AnimNotifyStateType);
 
 	/** Get whether a particular notify state is active in a specific state machine last tick.  */
-	UFUNCTION(BlueprintPure, Category = "Asset Player", meta = (BlueprintInternalUseOnly = "true", AnimGetter = "true", GetterContext = "Transition", Keywords = "Window,TransitionWindow,NotifyState,NotifyStateTransition"))
+	UFUNCTION(BlueprintPure, Category = "Asset Player", meta = (BlueprintInternalUseOnly = "true", AnimGetter = "true", GetterContext = "Transition", BlueprintThreadSafe, Keywords = "Window,TransitionWindow,NotifyState,NotifyStateTransition"))
 	bool WasAnimNotifyStateActiveInStateMachine(int32 MachineIndex, TSubclassOf<UAnimNotifyState> AnimNotifyStateType);
 
 	/** Get whether a particular notify state is active in a specific state last tick. */
-	UFUNCTION(BlueprintPure, Category = "Asset Player", meta = (BlueprintInternalUseOnly = "true", AnimGetter = "true", GetterContext = "Transition", Keywords = "Window,TransitionWindow,NotifyState,NotifyStateTransition"))
+	UFUNCTION(BlueprintPure, Category = "Asset Player", meta = (BlueprintInternalUseOnly = "true", AnimGetter = "true", GetterContext = "Transition", BlueprintThreadSafe, Keywords = "Window,TransitionWindow,NotifyState,NotifyStateTransition"))
 	bool WasAnimNotifyStateActiveInSourceState(int32 MachineIndex, int32 StateIndex, TSubclassOf<UAnimNotifyState> AnimNotifyStateType);
 
 	/** Get whether the most relevant animation was in a particular notify state last tick. */
-	UFUNCTION(BlueprintPure, Category = "Asset Player", meta = (BlueprintInternalUseOnly = "true", AnimGetter = "true", GetterContext = "Transition", Keywords = "Window,TransitionWindow,NotifyState,NotifyStateTransition"))
+	UFUNCTION(BlueprintPure, Category = "Asset Player", meta = (BlueprintInternalUseOnly = "true", AnimGetter = "true", GetterContext = "Transition", BlueprintThreadSafe, Keywords = "Window,TransitionWindow,NotifyState,NotifyStateTransition"))
     bool WasAnimNotifyTriggeredInSourceState(int32 MachineIndex, int32 StateIndex,  TSubclassOf<UAnimNotify> AnimNotifyType);
 
 	/** Get whether the most relevant animation triggered the animation notify with the specified name last tick.. */
-	UFUNCTION(BlueprintPure, Category = "Asset Player", meta = (BlueprintInternalUseOnly = "true", AnimGetter = "true", GetterContext = "Transition", Keywords = "Window,TransitionWindow,NotifyState,NotifyStateTransition"))
+	UFUNCTION(BlueprintPure, Category = "Asset Player", meta = (BlueprintInternalUseOnly = "true", AnimGetter = "true", GetterContext = "Transition", BlueprintThreadSafe, Keywords = "Window,TransitionWindow,NotifyState,NotifyStateTransition"))
     bool WasAnimNotifyNameTriggeredInSourceState(int32 MachineIndex, int32 StateIndex, FName NotifyName);
 
 	/** Get whether a particular notify type was active in a specific state machine last tick.  */
-	UFUNCTION(BlueprintPure, Category = "Asset Player", meta = (BlueprintInternalUseOnly = "true", AnimGetter = "true", GetterContext = "Transition", Keywords = "Window,TransitionWindow,NotifyState,NotifyStateTransition"))
+	UFUNCTION(BlueprintPure, Category = "Asset Player", meta = (BlueprintInternalUseOnly = "true", AnimGetter = "true", GetterContext = "Transition", BlueprintThreadSafe, Keywords = "Window,TransitionWindow,NotifyState,NotifyStateTransition"))
     bool WasAnimNotifyTriggeredInStateMachine(int32 MachineIndex, TSubclassOf<UAnimNotify> AnimNotifyType);
 
 	/** Get whether the given state machine triggered the animation notify with the specified name last tick. */
-	UFUNCTION(BlueprintPure, Category = "Asset Player", meta = (BlueprintInternalUseOnly = "true", AnimGetter = "true", GetterContext = "Transition", Keywords = "Window,TransitionWindow,NotifyState,NotifyStateTransition"))
+	UFUNCTION(BlueprintPure, Category = "Asset Player", meta = (BlueprintInternalUseOnly = "true", AnimGetter = "true", GetterContext = "Transition", BlueprintThreadSafe, Keywords = "Window,TransitionWindow,NotifyState,NotifyStateTransition"))
     bool WasAnimNotifyNameTriggeredInStateMachine(int32 MachineIndex, FName NotifyName);
 	
 	/**  Get whether an animation notify of a given type was triggered last tick. */
-	UFUNCTION(BlueprintPure, Category = "Asset Player", meta = (BlueprintInternalUseOnly = "true", AnimGetter = "true", GetterContext = "Transition", Keywords = "Window,TransitionWindow,NotifyState,NotifyStateTransition"))
+	UFUNCTION(BlueprintPure, Category = "Asset Player", meta = (BlueprintInternalUseOnly = "true", AnimGetter = "true", GetterContext = "Transition", BlueprintThreadSafe, Keywords = "Window,TransitionWindow,NotifyState,NotifyStateTransition"))
     bool WasAnimNotifyTriggeredInAnyState(TSubclassOf<UAnimNotify> AnimNotifyType);
 
 	/** Get whether the animation notify with the specified name triggered last tick. */
-	UFUNCTION(BlueprintPure, Category = "Asset Player", meta = (BlueprintInternalUseOnly = "true", AnimGetter = "true", GetterContext = "Transition", Keywords = "Window,TransitionWindow,NotifyState,NotifyStateTransition"))
+	UFUNCTION(BlueprintPure, Category = "Asset Player", meta = (BlueprintInternalUseOnly = "true", AnimGetter = "true", GetterContext = "Transition", BlueprintThreadSafe, Keywords = "Window,TransitionWindow,NotifyState,NotifyStateTransition"))
     bool WasAnimNotifyNameTriggeredInAnyState(FName NotifyName);
 	
 	/** Gets the runtime instance of the specified state machine by Name */
