Fixed issue in GetBodyInstance for ISMs with default parameters.
Submitted on behalf of Jaco Van Dyk.

diff --git a/Engine/Source/Runtime/Engine/Private/InstancedStaticMesh.cpp b/Engine/Source/Runtime/Engine/Private/InstancedStaticMesh.cpp
--- a/Engine/Source/Runtime/Engine/Private/InstancedStaticMesh.cpp
+++ b/Engine/Source/Runtime/Engine/Private/InstancedStaticMesh.cpp
@@ -1957,7 +1957,7 @@ FBodyInstance* UInstancedStaticMeshComponent::GetBodyInstance(FName BoneName, bo
 	{
 		return const_cast<FBodyInstance*>(InstanceBodies[Index]);
 	}
-	return  nullptr;
+	return  const_cast<FBodyInstance*>(&BodyInstance); // If no index is specified we return the primitive component body instance instead
 }
 
 FPrimitiveSceneProxy* UInstancedStaticMeshComponent::CreateSceneProxy()
