Fixed Show Nanite Proxy functionality in Static Mesh editor, and set AncientGame proxy render mode to 2 (hide proxies except in SM editor)

diff --git a/Engine/Source/Runtime/Engine/Private/StaticMeshRender.cpp b/Engine/Source/Runtime/Engine/Private/StaticMeshRender.cpp
--- a/Engine/Source/Runtime/Engine/Private/StaticMeshRender.cpp
+++ b/Engine/Source/Runtime/Engine/Private/StaticMeshRender.cpp
@@ -2323,8 +2323,12 @@ FPrimitiveSceneProxy* UStaticMeshComponent::CreateSceneProxy()
 		return nullptr;
 	}
 
-	// Whether or not to allow Nanite for this component
-	if (ShouldCreateNaniteProxy())
+	// Is Nanite supported, and is there built Nanite data for this static mesh?
+	if (ShouldCreateNaniteProxy() 
+	#if WITH_EDITORONLY_DATA
+		&& !bDisplayNaniteProxyMesh
+	#endif
+	)
 	{
 		LLM_SCOPE(ELLMTag::StaticMesh);
 
