Fixed issue where a level could have a WorldSettings actor of a different class type than the class specified in the project settings.

diff --git a/Engine/Source/Runtime/Engine/Private/World.cpp b/Engine/Source/Runtime/Engine/Private/World.cpp
--- a/Engine/Source/Runtime/Engine/Private/World.cpp
+++ b/Engine/Source/Runtime/Engine/Private/World.cpp
@@ -1512,7 +1512,7 @@ void UWorld::RepairWorldSettings()
 
 	// If for some reason we don't have a valid WorldSettings object go ahead and spawn one to avoid crashing.
 	// This will generally happen if a map is being moved from a different project.
-	if (ExistingWorldSettings == nullptr || !ExistingWorldSettings->IsA(GEngine->WorldSettingsClass))
+	if (ExistingWorldSettings == nullptr || ExistingWorldSettings->GetClass() != GEngine->WorldSettingsClass)
 	{
 		// Rename invalid WorldSettings to avoid name collisions
 		if (ExistingWorldSettings)
