Added more debug info to track hard to repro crash.

diff --git a/Engine/Source/Runtime/Engine/Private/WorldPartition/WorldPartition.cpp b/Engine/Source/Runtime/Engine/Private/WorldPartition/WorldPartition.cpp
--- a/Engine/Source/Runtime/Engine/Private/WorldPartition/WorldPartition.cpp
+++ b/Engine/Source/Runtime/Engine/Private/WorldPartition/WorldPartition.cpp
@@ -813,11 +813,13 @@ bool UWorldPartition::UpdateEditorCells(TFunctionRef<bool(TArray<UWorldPartition
 	for (const TPair<FWorldPartitionActorDesc*, int32> Pair : UnloadCount)
 	{
 		FWorldPartitionActorDesc* ActorDesc = Pair.Key;
+
 		// Only prompt if the actor will get unloaded by the unloading cells
 		if (ActorDesc->GetHardRefCount() == Pair.Value)
 		{
 			AActor* LoadedActor = ActorDesc->GetActor();
-			check(LoadedActor);
+			checkf(LoadedActor, TEXT("Can't find loaded actor %s"), *ActorDesc->ToString());
+
 			UPackage* ActorPackage = LoadedActor->GetExternalPackage();
 			if (ActorPackage && ActorPackage->IsDirty())
 			{
@@ -829,6 +831,7 @@ bool UWorldPartition::UpdateEditorCells(TFunctionRef<bool(TArray<UWorldPartition
 
 	// Make sure we save modified actor packages before unloading
 	FEditorFileUtils::EPromptReturnCode RetCode = FEditorFileUtils::PR_Success;
+
 	// Skip this when running commandlet because MarkPackageDirty() is allowed to dirty packages at loading when running a commandlet.
 	// Usually, the main reason actor packages are dirtied is when RerunConstructionScripts is called on the actor when it is added to the level.
 	if (ModifiedPackages.Num() && !IsRunningCommandlet())
