Compilation time optimization
  - Reduce codegen time for huge switch case
  - 17.03s -> 5.34s for ShowFlags.cpp
  - 24.59s -> 13.73s for Module.Engine.45_of_56.cpp

diff --git a/Engine/Source/Runtime/Engine/Private/ShowFlags.cpp b/Engine/Source/Runtime/Engine/Private/ShowFlags.cpp
--- a/Engine/Source/Runtime/Engine/Private/ShowFlags.cpp
+++ b/Engine/Source/Runtime/Engine/Private/ShowFlags.cpp
@@ -238,6 +238,10 @@ int32 FEngineShowFlags::FindIndexByName(const TCHAR* Name, const TCHAR* CommaSep
 	}
 }
 
+// Codegen optimization degenerates for very long functions like FindNameByIndex.
+// We don't need this code to be particularly fast anyway.
+BEGIN_FUNCTION_BUILD_OPTIMIZATION
+
 FString FEngineShowFlags::FindNameByIndex(uint32 InIndex)
 {
 	FString Name;
@@ -255,6 +259,8 @@ FString FEngineShowFlags::FindNameByIndex(uint32 InIndex)
 	return Name;
 }
 
+END_FUNCTION_BUILD_OPTIMIZATION
+
 void FEngineShowFlags::AddNameByIndex(uint32 InIndex, FString& Out)
 {
 	#define SHOWFLAG_ALWAYS_ACCESSIBLE(a,...) case SF_##a: Out += PREPROCESSOR_TO_STRING(a); break;
