#fix implemented UAnimDataModel::IsEditorOnly to ensure the object type is always culled during cook
#misc switched to CreateEditorOnlyDefaultSubobject for correctness and indicate life-time / use-case

diff --git a/Engine/Source/Runtime/Engine/Classes/Animation/AnimData/AnimDataModel.h b/Engine/Source/Runtime/Engine/Classes/Animation/AnimData/AnimDataModel.h
--- a/Engine/Source/Runtime/Engine/Classes/Animation/AnimData/AnimDataModel.h
+++ b/Engine/Source/Runtime/Engine/Classes/Animation/AnimData/AnimDataModel.h
@@ -80,6 +80,7 @@ public:
 	/** Begin UAnimDataModel overrides */
 	virtual void PostLoad() override;
 	virtual void PostDuplicate(bool bDuplicateForPIE) override;
+	virtual bool IsEditorOnly() const override { return true; }
 	/** End UAnimDataModel overrides */
 
 	/**
diff --git a/Engine/Source/Runtime/Engine/Private/Animation/AnimSequenceBase.cpp b/Engine/Source/Runtime/Engine/Private/Animation/AnimSequenceBase.cpp
--- a/Engine/Source/Runtime/Engine/Private/Animation/AnimSequenceBase.cpp
+++ b/Engine/Source/Runtime/Engine/Private/Animation/AnimSequenceBase.cpp
@@ -1268,7 +1268,7 @@ void UAnimSequenceBase::CreateModel()
 {	
 	if (DataModel == nullptr)
 	{
-		DataModel = CreateDefaultSubobject<UAnimDataModel>(FName(TEXT("AnimationDataModel")));
+		DataModel = CreateEditorOnlyDefaultSubobject<UAnimDataModel>(FName(TEXT("AnimationDataModel")));
 	}
 	
 	DataModel->GetModifiedEvent().RemoveAll(this);
