Updates to light weight instances.
Add support for interfaces on light weight instances. This allows a more uniform interface between light weight instances and actors for users to interact with.
Added new library functions to interact with the system.
Fix crashes that could occur with static meshes.
Light refactoring to improve implementing subclasses.

diff --git a/Engine/Source/Runtime/Engine/Classes/Engine/EngineTypes.h b/Engine/Source/Runtime/Engine/Classes/Engine/EngineTypes.h
--- a/Engine/Source/Runtime/Engine/Classes/Engine/EngineTypes.h
+++ b/Engine/Source/Runtime/Engine/Classes/Engine/EngineTypes.h
@@ -2065,6 +2065,7 @@ struct ENGINE_API FActorInstanceHandle
 
 	FVector GetLocation() const;
 	FRotator GetRotation() const;
+	FTransform GetTransform() const;
 
 	ULevel* GetLevel() const;
 	bool IsInLevel(ULevel* InLevel) const;
@@ -2087,6 +2088,7 @@ struct ENGINE_API FActorInstanceHandle
 
 	FActorInstanceHandle& operator=(const FActorInstanceHandle& Other) = default;
 	FActorInstanceHandle& operator=(FActorInstanceHandle&& Other) = default;
+	FActorInstanceHandle& operator=(AActor* OtherActor);
 
 	bool operator==(const FActorInstanceHandle& Other) const;
 	bool operator!=(const FActorInstanceHandle& Other) const;
@@ -2097,14 +2099,24 @@ struct ENGINE_API FActorInstanceHandle
 	friend ENGINE_API FArchive& operator<<(FArchive& Ar, FActorInstanceHandle& Handle);
 
 private:
-	// this is cached here for convenience
+	/**
+	 * helper functions that let us treat the actor pointer as a UObject in templated functions
+	 * these do NOT fetch the actor so they will return nullptr if we don't have a full actor representation
+	 */
+	UObject* GetActorAsUObject();
+	const UObject* GetActorAsUObject() const;
+
+	/** Returns true if Actor is not null and not pending kill */
+	bool IsActorValid() const;
+
+	/** this is cached here for convenience */
 	UPROPERTY()
 	mutable TWeakObjectPtr<AActor> Actor;
 
-	// Identifies the light weight instance manager to use
+	/** Identifies the light weight instance manager to use */
 	int32 ManagerIndex;
 
-	// Identifies the instance within the manager
+	/** Identifies the instance within the manager */
 	int32 InstanceIndex;
 };
 
diff --git a/Engine/Source/Runtime/Engine/Classes/GameFramework/LightWeightInstanceBlueprintFunctionLibrary.h b/Engine/Source/Runtime/Engine/Classes/GameFramework/LightWeightInstanceBlueprintFunctionLibrary.h
--- a/Engine/Source/Runtime/Engine/Classes/GameFramework/LightWeightInstanceBlueprintFunctionLibrary.h
+++ b/Engine/Source/Runtime/Engine/Classes/GameFramework/LightWeightInstanceBlueprintFunctionLibrary.h
@@ -14,6 +14,11 @@ class ENGINE_API ULightWeightInstanceBlueprintFunctionLibrary : public UBlueprin
 
 public:
 
+	// Returns a handle to a new light weight instance that represents an object of type ActorClass
 	UFUNCTION(BlueprintCallable, Category = "Light Weight Instance")
 	static FActorInstanceHandle CreateNewLightWeightInstance(UClass* ActorClass, FTransform Transform, ULevel* Level);
+
+	// Returns a handle to the light weight representation and destroys Actor if successful; Returns a handle to Actor otherwise
+	UFUNCTION(BlueprintCallable, Category = "Light Weight Instance")
+	static FActorInstanceHandle ConvertActorToLightWeightInstance(AActor* Actor);
 };
\ No newline at end of file
diff --git a/Engine/Source/Runtime/Engine/Classes/GameFramework/LightWeightInstanceManager.h b/Engine/Source/Runtime/Engine/Classes/GameFramework/LightWeightInstanceManager.h
--- a/Engine/Source/Runtime/Engine/Classes/GameFramework/LightWeightInstanceManager.h
+++ b/Engine/Source/Runtime/Engine/Classes/GameFramework/LightWeightInstanceManager.h
@@ -24,6 +24,30 @@ struct ENGINE_API FLWIData
 	FTransform Transform;
 };
 
+
+// Base class for interfaces for each handle
+UCLASS()
+class ENGINE_API UActorInstanceHandleInterface : public UObject
+{
+	GENERATED_UCLASS_BODY()
+
+	friend class ALightWeightInstanceManager;
+
+	// Returns a pointer to the actor for this handle. This will create a new copy of the actor if it was being stored as a light weight instance.
+	template<typename T>
+	T* FetchActor();
+
+protected:
+
+	// Handle to the actor or instance that is associated with this interface
+	FActorInstanceHandle Handle;
+
+	// A cached pointer to the manager associated with Handle so we don't need to find it every time
+	ALightWeightInstanceManager* Manager;
+};
+
+
+
 UCLASS(BlueprintType, Blueprintable)
 class ENGINE_API ALightWeightInstanceManager : public AActor
 {
@@ -41,6 +65,12 @@ public:
 	// Returns the location of the instance specified by Handle
 	FVector GetLocation(const FActorInstanceHandle& Handle) const;
 
+	// Returns the rotation of the instance specified by Handle
+	FRotator GetRotation(const FActorInstanceHandle& Handle) const;
+
+	// Returns the transform of the instance specified by Handle
+	FTransform GetTransform(const FActorInstanceHandle& Handle) const;
+
 	// Returns the name of the instance specified by Handle
 	FString GetName(const FActorInstanceHandle& Handle) const;
 
@@ -60,13 +90,43 @@ public:
 	virtual void SetRepresentedClass(UClass* ActorClass);
 
 	// Returns the actor associated with Handle if one exists
-	AActor* GetActorFromHandle(const FActorInstanceHandle& Handle);
+	AActor* FetchActorFromHandle(const FActorInstanceHandle& Handle);
 
 	// Returns the index of the light weight instance associated with InActor if one exists; otherwise we return INDEX_NONE
 	int32 FindIndexForActor(const AActor* InActor) const;
 
+	// Returns a handle to a light weight instance representing the same object as InActor and calls destroy on InActor if successful.
+	FActorInstanceHandle ConvertActorToLightWeightInstance(AActor* InActor);
+
 	virtual int32 ConvertCollisionIndexToLightWeightIndex(int32 InIndex) const;
 
+	template<typename U>
+	bool IsInterfaceSupported() const
+	{
+		const UClass* const InterfaceClass = GetInterfaceClass();
+		return InterfaceClass && InterfaceClass->ImplementsInterface(U::StaticClass());
+	}
+
+	template<typename I>
+	I* FetchInterfaceObject(const FActorInstanceHandle& Handle)
+	{
+		// if we have a valid actor, use it
+		if (Handle.Actor.IsValid())
+		{
+			return Cast<I>(Handle.Actor.Get());
+		}
+
+		// TODO: once we have a better idea of how we're going to use this we should add some kind of caching so we don't always need to create a new uobject.
+		// if we can support the interface using just the handle then do that
+		if (I* InterfaceObj = Cast<I>(CreateInterfaceObject(Handle)))
+		{
+			return InterfaceObj;
+		}
+
+		// fallback to creating the actor
+		return Cast<I>(FetchActorFromHandle(Handle));
+	}
+
 protected:
 	// Takes a polymorphic struct to set the initial data for a new instance
 	int32 AddNewInstance(FLWIData* InitData);
@@ -90,6 +150,18 @@ protected:
 	// Called after spawning a new actor from a light weight instance
 	virtual void PostActorSpawn(const FActorInstanceHandle& Handle);
 
+	// Create an object that implements interfaces for the light weight instance specified by Handle
+	UObject* CreateInterfaceObject(const FActorInstanceHandle& Handle)
+	{
+		UActorInstanceHandleInterface* InterfaceObj = NewObject<UActorInstanceHandleInterface>(this, GetInterfaceClass());
+		InterfaceObj->Handle = Handle;
+		InterfaceObj->Manager = this;
+		return InterfaceObj;
+	}
+
+	// Gets the class that implements interfaces for light weight instances owned by this manager
+	virtual UClass* GetInterfaceClass() const;
+
 	//
 	// Data and replication functions
 	//
@@ -114,11 +186,27 @@ protected:
 	// 
 
 	// Current per instance transforms
-	UPROPERTY(ReplicatedUsing=OnRep_Transforms)
+	UPROPERTY(ReplicatedUsing=OnRep_Transforms, EditInstanceOnly, Category=LightWeightInstance)
 	TArray<FTransform> InstanceTransforms;
 	UFUNCTION()
 	virtual void OnRep_Transforms();
 
+	//
+	// Subclasses should override the following functions if they need to store any additional data 
+	//
+
+	// Increases the size of all of our data arrays to NewSize
+	virtual void GrowDataArrays();
+
+	// Populates the data array entries at Index with the data stored in InData;
+	virtual void UpdateDataAtIndex(FLWIData* InData, int32 Index);
+
+	// Allocates the appropriate subclass of FLWIData to initialize instances for this manager. The caller is responsible for freeing the memory when they are finished.
+	virtual FLWIData* AllocateInitData() const;
+
+	// Gathers info from InActor and stores it in InData. Returns true if the relevant data was successfully copied.
+	virtual bool SetDataFromActor(FLWIData* InData, AActor* InActor) const;
+
 	//
 	// Bookkeeping info
 	//
diff --git a/Engine/Source/Runtime/Engine/Classes/GameFramework/LightWeightInstanceSubsystem.h b/Engine/Source/Runtime/Engine/Classes/GameFramework/LightWeightInstanceSubsystem.h
--- a/Engine/Source/Runtime/Engine/Classes/GameFramework/LightWeightInstanceSubsystem.h
+++ b/Engine/Source/Runtime/Engine/Classes/GameFramework/LightWeightInstanceSubsystem.h
@@ -8,6 +8,7 @@
 #include "UObject/Object.h"
 #include "Engine/EngineTypes.h"
 #include "Engine/EngineBaseTypes.h"
+#include "GameFramework/LightWeightInstanceManager.h"
 #include "Templates/SharedPointer.h"
 
 #include "Actor.h"
@@ -33,18 +34,18 @@ struct ENGINE_API FLightWeightInstanceSubsystem
 		return *LWISubsystem;
 	}
 
-	// returns the instance manager that handles the given handle
+	// Returns the instance manager that handles the given handle
 	ALightWeightInstanceManager* FindLightWeightInstanceManager(const FActorInstanceHandle& Handle) const;
 
-	// returns the instance manager that handles actors of type ActorClass in level Level
+	// Returns the instance manager that handles actors of type ActorClass in level Level
 	ALightWeightInstanceManager* FindLightWeightInstanceManager(UClass* ActorClass, ULevel* Level) const;
 
-	// returns the instance manager that handles instances of type Class that live in Level
+	// Returns the instance manager that handles instances of type Class that live in Level
 	UFUNCTION(Server, Unreliable)
 	ALightWeightInstanceManager* FindOrAddLightWeightInstanceManager(UClass* ActorClass, ULevel* Level);
 
 	// Returns the actor specified by Handle. This may require loading and creating the actor object.
-	AActor* GetActor(const FActorInstanceHandle& Handle);
+	AActor* FetchActor(const FActorInstanceHandle& Handle);
 
 	// Returns the actor specified by Handle if it exists. Returns nullptr if it doesn't
 	AActor* GetActor_NoCreate(const FActorInstanceHandle& Handle) const;
@@ -58,16 +59,44 @@ struct ENGINE_API FLightWeightInstanceSubsystem
 
 	ULevel* GetLevel(const FActorInstanceHandle& Handle);
 
-	// returns true if the object represented by Handle is in InLevel
+	// Returns true if the object represented by Handle is in InLevel
 	bool IsInLevel(const FActorInstanceHandle& Handle, const ULevel* InLevel);
 
+	// Returns a handle to a new light weight instance that represents an object of type ActorClass
+	FActorInstanceHandle CreateNewLightWeightInstance(UClass* ActorClass, FLWIData* InitData, ULevel* Level);
+
+	// Returns true if the handle can return an object that implements the interface U
+	template<typename U>
+	bool IsInterfaceSupported(const FActorInstanceHandle& Handle) const
+	{
+		if (ALightWeightInstanceManager* InstanceManager = FindLightWeightInstanceManager(Handle))
+		{
+			return InstanceManager->IsInterfaceSupported<U>();
+		}
+
+		return false;
+	}
+
+	// Returns an object that implements the interface I for Handle
+	template<typename I>
+	I* FetchInterfaceObject(const FActorInstanceHandle& Handle)
+	{
+		if (ALightWeightInstanceManager* InstanceManager = FindLightWeightInstanceManager(Handle))
+		{
+			return InstanceManager->FetchInterfaceObject<I>(Handle);
+		}
+
+		return nullptr;
+	}
+
 protected:
+	// Returns the class of the insance manager best suited to support instances of type ActorClass
 	UClass* FindBestInstanceManagerClass(const UClass* ActorClass);
 
-	// returns the index associated with Manager
+	// Returns the index associated with Manager
 	int32 GetManagerIndex(const ALightWeightInstanceManager* Manager) const;
 
-	// returns the light weight instance manager at index Index
+	// Returns the light weight instance manager at index Index
 	const ALightWeightInstanceManager* GetManagerAt(int32 Index) const;
 
 private:
diff --git a/Engine/Source/Runtime/Engine/Classes/GameFramework/LightWeightInstanceUtils.h b/Engine/Source/Runtime/Engine/Classes/GameFramework/LightWeightInstanceUtils.h
new file mode 100644
--- /dev/null
+++ b/Engine/Source/Runtime/Engine/Classes/GameFramework/LightWeightInstanceUtils.h
@@ -0,0 +1,52 @@
+// Copyright Epic Games, Inc. All Rights Reserved.
+
+#pragma once
+
+#include "Kismet/BlueprintFunctionLibrary.h"
+#include "Engine/EngineTypes.h"
+
+#include "LightWeightInstanceBlueprintFunctionLibrary.generated.h"
+
+namespace LWIUtils
+{
+	/**
+	 * Returns true if the object this handle represents supports the interface of type U
+	 */
+	template<typename U>
+	static ENGINE_API bool DoesHandleSupportInterface(const FActorInstanceHandle& Handle);
+
+	/**
+	 * Returns an object implementing the interface I.
+	 */
+	template<typename I>
+	static ENGINE_API I* FetchInterfaceFromHandle(const FActorInstanceHandle& Handle);
+
+
+	/**
+	 * Implementations of inline or template functions go below here
+	 */
+
+	template<typename U>
+	static bool DoesHandleSupportInterface(const FActorInstanceHandle& Handle)
+	{
+		// if we have a valid actor, see if it supports the interface
+		if (const UObject* Obj = Handle.GetActorAsUObject())
+		{
+			return Obj->GetClass()->ImplementsInterface(U);
+		}
+
+		// no valid actor, ask the instance manager instead
+		return FLightWeightInstanceSubsystem::Get().IsInterfaceSupported<U>(Handle);
+	}
+
+	template<typename I>
+	static I* FetchInterfaceFromHandle(const FActorInstanceHandle& Handle)
+	{
+		if (Handle.IsActorValid())
+		{
+			return Cast<I>(Actor.Get());
+		}
+
+		return FLightWeightInstanceSubsystem::Get().FetchInterfaceObject<I>(Handle);
+	}
+};
\ No newline at end of file
diff --git a/Engine/Source/Runtime/Engine/Private/EngineTypes.cpp b/Engine/Source/Runtime/Engine/Private/EngineTypes.cpp
--- a/Engine/Source/Runtime/Engine/Private/EngineTypes.cpp
+++ b/Engine/Source/Runtime/Engine/Private/EngineTypes.cpp
@@ -273,7 +273,7 @@ FActorInstanceHandle::FActorInstanceHandle(int32 InManagerIndex, int32 InInstanc
 	{
 		InstanceIndex = Manager->ConvertCollisionIndexToLightWeightIndex(InInstanceIndex);
 
-		if (AActor*const* FoundActor = Manager->Actors.Find(InInstanceIndex))
+		if (AActor* const* FoundActor = Manager->Actors.Find(InInstanceIndex))
 		{
 			Actor = *FoundActor;
 		}
@@ -293,7 +293,7 @@ FActorInstanceHandle::FActorInstanceHandle(ALightWeightInstanceManager* Manager,
 	if (ensure(Manager))
 	{
 		InstanceIndex = Manager->ConvertCollisionIndexToLightWeightIndex(InInstanceIndex);
-		if (AActor*const* FoundActor = Manager->Actors.Find(InstanceIndex))
+		if (AActor* const* FoundActor = Manager->Actors.Find(InstanceIndex))
 		{
 			Actor = *FoundActor;
 		}
@@ -312,7 +312,7 @@ FActorInstanceHandle::FActorInstanceHandle(const FActorInstanceHandle& Other)
 
 bool FActorInstanceHandle::IsValid() const
 {
-	return (ManagerIndex != INDEX_NONE && InstanceIndex != INDEX_NONE) || Actor.IsValid();
+	return (ManagerIndex != INDEX_NONE && InstanceIndex != INDEX_NONE) || IsActorValid();
 }
 
 bool FActorInstanceHandle::DoesRepresentClass(const UClass* OtherClass) const
@@ -322,12 +322,12 @@ bool FActorInstanceHandle::DoesRepresentClass(const UClass* OtherClass) const
 		return false;
 	}
 
-	if (Actor.IsValid())
+	if (IsActorValid())
 	{
 		return Actor->IsA(OtherClass);
 	}
 
-	if (ALightWeightInstanceManager* Manager = FLightWeightInstanceSubsystem::Get().FindLightWeightInstanceManager(*this))
+	if (const ALightWeightInstanceManager* Manager = FLightWeightInstanceSubsystem::Get().GetManagerAt(ManagerIndex))
 	{
 		return Manager->DoesRepresentClass(OtherClass);
 	}
@@ -342,12 +342,12 @@ UClass* FActorInstanceHandle::GetRepresentedClass() const
 		return nullptr;
 	}
 
-	if (Actor.IsValid())
+	if (IsActorValid())
 	{
 		return Actor->GetClass();
 	}
 
-	if (ALightWeightInstanceManager* Manager = FLightWeightInstanceSubsystem::Get().FindLightWeightInstanceManager(*this))
+	if (const ALightWeightInstanceManager* Manager = FLightWeightInstanceSubsystem::Get().GetManagerAt(ManagerIndex))
 	{
 		return Manager->GetRepresentedClass();
 	}
@@ -357,37 +357,92 @@ UClass* FActorInstanceHandle::GetRepresentedClass() const
 
 FVector FActorInstanceHandle::GetLocation() const
 {
+	if (IsActorValid())
+	{
+		return Actor->GetActorLocation();
+	}
+
+	if (const ALightWeightInstanceManager* Manager = FLightWeightInstanceSubsystem::Get().GetManagerAt(ManagerIndex))
+	{
+		Manager->GetLocation(*this);
+	}
+
 	return FVector();
 }
 
 FRotator FActorInstanceHandle::GetRotation() const
 {
+	if (IsActorValid())
+	{
+		return Actor->GetActorRotation();
+	}
+
+	if (const ALightWeightInstanceManager* Manager = FLightWeightInstanceSubsystem::Get().GetManagerAt(ManagerIndex))
+	{
+		Manager->GetRotation(*this);
+	}
+
 	return FRotator();
 }
 
+FTransform FActorInstanceHandle::GetTransform() const
+{
+	if (IsActorValid())
+	{
+		return Actor->GetActorTransform();
+	}
+
+	if (const ALightWeightInstanceManager* Manager = FLightWeightInstanceSubsystem::Get().GetManagerAt(ManagerIndex))
+	{
+		Manager->GetTransform(*this);
+	}
+
+	return FTransform();
+}
+
 ULevel* FActorInstanceHandle::GetLevel() const
 {
+	if (IsActorValid())
+	{
+		return Actor->GetLevel();
+	}
+
 	return nullptr;
 }
 
 bool FActorInstanceHandle::IsInLevel(ULevel* Level) const
 {
-	return false;
+	return Level && Level == GetLevel();
 }
 
 FName FActorInstanceHandle::GetFName() const
 {
+	if (IsActorValid())
+	{
+		return Actor->GetFName();
+	}
+
 	return NAME_None;
 }
 
 FString FActorInstanceHandle::GetName() const
 {
+	if (IsActorValid())
+	{
+		return Actor->GetName();
+	}
+
+	if (const ALightWeightInstanceManager* Manager = FLightWeightInstanceSubsystem::Get().GetManagerAt(ManagerIndex))
+	{
+		Manager->GetName(*this);
+	}
+
 	return FString();
 }
 
 AActor* FActorInstanceHandle::GetManagingActor() const
 {
-	if (Actor.IsValid())
+	if (IsActorValid())
 	{
 		return Actor.Get();
 	}
@@ -397,7 +452,7 @@ AActor* FActorInstanceHandle::GetManagingActor() const
 
 USceneComponent* FActorInstanceHandle::GetRootComponent() const
 {
-	if (Actor.IsValid())
+	if (IsActorValid())
 	{
 		return Actor->GetRootComponent();
 	}
@@ -408,12 +463,46 @@ USceneComponent* FActorInstanceHandle::GetRootComponent() const
 
 AActor* FActorInstanceHandle::FetchActor() const
 {
-	if (Actor.IsValid())
+	if (IsActorValid())
 	{
 		return Actor.Get();
 	}
 
-	return FLightWeightInstanceSubsystem::Get().GetActor(*this);
+	return FLightWeightInstanceSubsystem::Get().FetchActor(*this);
+}
+
+UObject* FActorInstanceHandle::GetActorAsUObject()
+{
+	if (IsActorValid())
+	{
+		return Cast<UObject>(Actor.Get());
+	}
+
+	return nullptr;
+}
+
+const UObject* FActorInstanceHandle::GetActorAsUObject() const
+{
+	if (IsActorValid())
+	{
+		return Cast<UObject>(Actor.Get());
+	}
+
+	return nullptr;
+}
+
+bool FActorInstanceHandle::IsActorValid() const
+{
+	return Actor.IsValid() && !Actor->IsPendingKill();
+}
+
+FActorInstanceHandle& FActorInstanceHandle::operator=(AActor* OtherActor)
+{
+	Actor = OtherActor;
+	ManagerIndex = INDEX_NONE;
+	InstanceIndex = INDEX_NONE;
+
+	return *this;
 }
 
 bool FActorInstanceHandle::operator==(const FActorInstanceHandle& Other) const
@@ -450,7 +539,7 @@ bool FActorInstanceHandle::operator==(const AActor* OtherActor) const
 
 		return Manager->FindIndexForActor(OtherActor) == InstanceIndex;
 	}
-	
+
 	return false;
 }
 
diff --git a/Engine/Source/Runtime/Engine/Private/LightWeightInstanceBlueprintFunctionLibrary.cpp b/Engine/Source/Runtime/Engine/Private/LightWeightInstanceBlueprintFunctionLibrary.cpp
--- a/Engine/Source/Runtime/Engine/Private/LightWeightInstanceBlueprintFunctionLibrary.cpp
+++ b/Engine/Source/Runtime/Engine/Private/LightWeightInstanceBlueprintFunctionLibrary.cpp
@@ -6,16 +6,25 @@
 
 FActorInstanceHandle ULightWeightInstanceBlueprintFunctionLibrary::CreateNewLightWeightInstance(UClass* InActorClass, FTransform InTransform, ULevel* InLevel)
 {
-	// Get or create a light weight instance for this class and level
-	if (ALightWeightInstanceManager* LightWeightInstance = FLightWeightInstanceSubsystem::Get().FindOrAddLightWeightInstanceManager(InActorClass, InLevel))
+	// set up initialization data
+	FLWIData PerInstanceData;
+	PerInstanceData.Transform = InTransform;
+
+	return FLightWeightInstanceSubsystem::Get().CreateNewLightWeightInstance(InActorClass, &PerInstanceData, InLevel);
+}
+
+FActorInstanceHandle ULightWeightInstanceBlueprintFunctionLibrary::ConvertActorToLightWeightInstance(AActor* InActor)
+{
+	if (!ensure(InActor))
 	{
-		// create an instance with the given transform
-		FLWIData PerInstanceData;
-		PerInstanceData.Transform = InTransform;
-		int32 InstanceIdx = LightWeightInstance->AddNewInstance(&PerInstanceData);
+		return FActorInstanceHandle();
+	}
 
-		return FActorInstanceHandle(LightWeightInstance, InstanceIdx);
+	// Get or create a light weight instance for this class and level
+	if (ALightWeightInstanceManager* LWIManager = FLightWeightInstanceSubsystem::Get().FindOrAddLightWeightInstanceManager(InActor->GetClass(), InActor->GetLevel()))
+	{
+		return LWIManager->ConvertActorToLightWeightInstance(InActor);
 	}
 
-	return FActorInstanceHandle();
+	return FActorInstanceHandle(InActor);
 }
diff --git a/Engine/Source/Runtime/Engine/Private/LightWeightInstanceManager.cpp b/Engine/Source/Runtime/Engine/Private/LightWeightInstanceManager.cpp
--- a/Engine/Source/Runtime/Engine/Private/LightWeightInstanceManager.cpp
+++ b/Engine/Source/Runtime/Engine/Private/LightWeightInstanceManager.cpp
@@ -5,6 +5,11 @@
 #include "Engine/World.h"
 #include "Net/UnrealNetwork.h"
 
+UActorInstanceHandleInterface::UActorInstanceHandleInterface(const FObjectInitializer& ObjectInitializer)
+{
+	// do nothing
+}
+
 ALightWeightInstanceManager::ALightWeightInstanceManager(const FObjectInitializer& ObjectInitializer)
 {
 	bReplicates = true;
@@ -22,18 +27,20 @@ ALightWeightInstanceManager::~ALightWeightInstanceManager()
 
 void ALightWeightInstanceManager::SetRepresentedClass(UClass* ActorClass)
 {
-	// we should not be changing this once it has been set
-	ensure(RepresentedClass == nullptr);
-
 	RepresentedClass = ActorClass;
 }
 
+UClass* ALightWeightInstanceManager::GetInterfaceClass() const
+{
+	return UActorInstanceHandleInterface::StaticClass();
+}
+
 void ALightWeightInstanceManager::Tick(float DeltaSeconds)
 {
 	// do nothing
 }
 
-AActor* ALightWeightInstanceManager::GetActorFromHandle(const FActorInstanceHandle& Handle)
+AActor* ALightWeightInstanceManager::FetchActorFromHandle(const FActorInstanceHandle& Handle)
 {
 	// make sure the handle doesn't have an actor already
 	if (ensure(!Handle.Actor.IsValid()))
@@ -77,6 +84,73 @@ int32 ALightWeightInstanceManager::FindIndexForActor(const AActor* InActor) cons
 	return INDEX_NONE;
 }
 
+FActorInstanceHandle ALightWeightInstanceManager::ConvertActorToLightWeightInstance(AActor* InActor)
+{
+	FActorInstanceHandle ReturnHandle;
+	if (!InActor)
+	{
+		return ReturnHandle;
+	}
+
+	if (FLWIData* Data = AllocateInitData())
+	{
+		SetDataFromActor(Data, InActor);
+		// If we used to manage this actor as a light weight instance then we will already have entries associated with it. Check for these so we can use them again instead of allocating new entries
+		int32 Idx = INDEX_NONE;
+		for (auto& IndexActorPair : Actors)
+		{
+			if (IndexActorPair.Value == InActor)
+			{
+				Idx = IndexActorPair.Key;
+				break;
+			}
+		}
+
+		if (Idx == INDEX_NONE)
+		{
+			Idx = AddNewInstance(Data);
+		}
+		else
+		{
+			UpdateDataAtIndex(Data, Idx);
+		}
+
+		// Update our handle
+		ReturnHandle.ManagerIndex = FLightWeightInstanceSubsystem::Get().GetManagerIndex(this);
+		ReturnHandle.InstanceIndex = Idx;
+
+		// cleanup
+		InActor->Destroy();
+		delete Data;
+	}
+	else
+	{
+		// something went wrong and we can't manage this actor
+		// just return a handle to the actor
+		ReturnHandle.Actor = InActor;
+	}
+
+	return ReturnHandle;
+}
+
+FLWIData* ALightWeightInstanceManager::AllocateInitData() const
+{
+	return new FLWIData;
+}
+
+
+bool ALightWeightInstanceManager::SetDataFromActor(FLWIData* InData, AActor* InActor) const
+{
+	if (!InData || !InActor)
+	{
+		return false;
+	}
+
+	InData->Transform = InActor->GetActorTransform();
+
+	return true;
+}
+
 int32 ALightWeightInstanceManager::ConvertCollisionIndexToLightWeightIndex(int32 InIndex) const
 {
 	return InIndex;
@@ -96,7 +170,7 @@ void ALightWeightInstanceManager::PostActorSpawn(const FActorInstanceHandle& Han
 
 bool ALightWeightInstanceManager::IsIndexValid(int32 Index) const
 {
-	if (Index < ValidIndices.Num())
+	if (ValidIndices.IsValidIndex(Index))
 	{
 		return ValidIndices[Index];
 	}
@@ -128,6 +202,36 @@ FVector ALightWeightInstanceManager::GetLocation(const FActorInstanceHandle& Han
 	return FVector();
 }
 
+FRotator ALightWeightInstanceManager::GetRotation(const FActorInstanceHandle& Handle) const
+{
+	if (FindActorForHandle(Handle))
+	{
+		return Handle.Actor->GetActorRotation();
+	}
+
+	if (ensure(IsIndexValid(Handle.GetInstanceIndex())))
+	{
+		return InstanceTransforms[Handle.GetInstanceIndex()].Rotator();
+	}
+
+	return FRotator();
+}
+
+FTransform ALightWeightInstanceManager::GetTransform(const FActorInstanceHandle& Handle) const
+{
+	if (FindActorForHandle(Handle))
+	{
+		return Handle.Actor->GetActorTransform();
+	}
+
+	if (ensure(IsIndexValid(Handle.GetInstanceIndex())))
+	{
+		return InstanceTransforms[Handle.GetInstanceIndex()];
+	}
+
+	return FTransform();
+}
+
 FString ALightWeightInstanceManager::GetName(const FActorInstanceHandle& Handle) const
 {
 	if (FindActorForHandle(Handle))
@@ -179,17 +283,28 @@ void ALightWeightInstanceManager::AddNewInstanceAt(FLWIData* InitData, int32 Ind
 	// allocate space on the end of the array if we need to
 	if (Index >= ValidIndices.Num())
 	{
-		InstanceTransforms.AddUninitialized();
-		ValidIndices.AddUninitialized();
+		GrowDataArrays();
 	}
 	ensure(Index < ValidIndices.Num());
-	ensure(InstanceTransforms.Num() == ValidIndices.Num());
 
 	// update our data
-	InstanceTransforms[Index] = InitData->Transform;
+	UpdateDataAtIndex(InitData, Index);
 	ValidIndices[Index] = true;
 }
 
+void ALightWeightInstanceManager::GrowDataArrays()
+{
+	InstanceTransforms.AddUninitialized();
+	ValidIndices.AddUninitialized();
+
+	check(InstanceTransforms.Num() == ValidIndices.Num());
+}
+
+void ALightWeightInstanceManager::UpdateDataAtIndex(FLWIData* InData, int32 Index)
+{
+	InstanceTransforms[Index] = InData->Transform;
+}
+
 void ALightWeightInstanceManager::RemoveInstance(const int32 Index)
 {
 	if (ensure(IsIndexValid(Index)))
diff --git a/Engine/Source/Runtime/Engine/Private/LightWeightInstanceStaticMeshManager.cpp b/Engine/Source/Runtime/Engine/Private/LightWeightInstanceStaticMeshManager.cpp
--- a/Engine/Source/Runtime/Engine/Private/LightWeightInstanceStaticMeshManager.cpp
+++ b/Engine/Source/Runtime/Engine/Private/LightWeightInstanceStaticMeshManager.cpp
@@ -26,7 +26,7 @@ void ALightWeightInstanceStaticMeshManager::SetRepresentedClass(UClass* ActorCla
 {
 	Super::SetRepresentedClass(ActorClass);
 
-	AActor* ActorCDO = Cast<AActor>(RepresentedClass->GetDefaultObject());
+	AActor* ActorCDO = RepresentedClass ? Cast<AActor>(RepresentedClass->GetDefaultObject()) : nullptr;
 	if (ActorCDO)
 	{
 		BaseInstanceName = ActorCDO->GetName();
@@ -76,6 +76,7 @@ void ALightWeightInstanceStaticMeshManager::AddNewInstanceAt(FLWIData* InitData,
 void ALightWeightInstanceStaticMeshManager::RemoveInstance(const int32 Index)
 {
 	RemoveInstanceFromRendering(Index);
+
 	Super::RemoveInstance(Index);
 }
 
@@ -179,10 +180,12 @@ void ALightWeightInstanceStaticMeshManager::OnStaticMeshSet()
 			InstancedStaticMeshComponent->SetStaticMesh(StaticMesh.Get());
 		}
 
-		UStaticMesh* Mesh = StaticMesh.Get();
-		for (int32 Idx = 0; Idx < Mesh->GetStaticMaterials().Num(); ++Idx)
+		if (UStaticMesh* Mesh = StaticMesh.Get())
 		{
-			InstancedStaticMeshComponent->SetMaterial(Idx, Mesh->GetMaterial(Idx));
+			for (int32 Idx = 0; Idx < Mesh->GetStaticMaterials().Num(); ++Idx)
+			{
+				InstancedStaticMeshComponent->SetMaterial(Idx, Mesh->GetMaterial(Idx));
+			}
 		}
 	}
 }
diff --git a/Engine/Source/Runtime/Engine/Private/LightWeightInstanceSubsystem.cpp b/Engine/Source/Runtime/Engine/Private/LightWeightInstanceSubsystem.cpp
--- a/Engine/Source/Runtime/Engine/Private/LightWeightInstanceSubsystem.cpp
+++ b/Engine/Source/Runtime/Engine/Private/LightWeightInstanceSubsystem.cpp
@@ -1,7 +1,6 @@
 // Copyright Epic Games, Inc. All Rights Reserved.
 
 #include "GameFramework/LightWeightInstanceSubsystem.h"
-#include "GameFramework/LightWeightInstanceManager.h"
 #include "Engine/World.h"
 #include "UObject/UObjectIterator.h"
 
@@ -24,7 +23,7 @@ int32 FLightWeightInstanceSubsystem::GetManagerIndex(const ALightWeightInstanceM
 
 const ALightWeightInstanceManager* FLightWeightInstanceSubsystem::GetManagerAt(int32 Index) const
 {
-	return (Index >= 0 && Index < LWInstanceManagers.Num()) ? LWInstanceManagers[Index] : nullptr;
+	return LWInstanceManagers.IsValidIndex(Index) ? LWInstanceManagers[Index] : nullptr;
 }
 
 ALightWeightInstanceManager* FLightWeightInstanceSubsystem::FindLightWeightInstanceManager(const FActorInstanceHandle& Handle) const
@@ -163,7 +162,7 @@ UClass* FLightWeightInstanceSubsystem::FindBestInstanceManagerClass(const UClass
 	return BestManagerClass;
 }
 
-AActor* FLightWeightInstanceSubsystem::GetActor(const FActorInstanceHandle& Handle)
+AActor* FLightWeightInstanceSubsystem::FetchActor(const FActorInstanceHandle& Handle)
 {
 	// if the actor is valid return it
 	if (Handle.Actor.IsValid())
@@ -173,7 +172,7 @@ AActor* FLightWeightInstanceSubsystem::GetActor(const FActorInstanceHandle& Hand
 
 	if (ALightWeightInstanceManager* LWIManager = FindLightWeightInstanceManager(Handle))
 	{
-		return LWIManager->GetActorFromHandle(Handle);
+		return LWIManager->FetchActorFromHandle(Handle);
 	}
 
 	return nullptr;
@@ -265,4 +264,18 @@ bool FLightWeightInstanceSubsystem::IsInLevel(const FActorInstanceHandle& Handle
 	}
 
 	return false;
-}
\ No newline at end of file
+}
+
+FActorInstanceHandle FLightWeightInstanceSubsystem::CreateNewLightWeightInstance(UClass* InActorClass, FLWIData* InitData, ULevel* InLevel)
+{
+	// Get or create a light weight instance for this class and level
+	if (ALightWeightInstanceManager* LWIManager = FindOrAddLightWeightInstanceManager(InActorClass, InLevel))
+	{
+		// create an instance with the given data
+		int32 InstanceIdx = LWIManager->AddNewInstance(InitData);
+
+		return FActorInstanceHandle(LWIManager, InstanceIdx);
+	}
+
+	return FActorInstanceHandle();
+}
