Fixed ISM component AddInstances call that were incorrectly replaced in CL 16329274.

#rnx

diff --git a/Engine/Source/Developer/WorldPartitionHLODUtilities/Private/HLODBuilderInstancing.cpp b/Engine/Source/Developer/WorldPartitionHLODUtilities/Private/HLODBuilderInstancing.cpp
--- a/Engine/Source/Developer/WorldPartitionHLODUtilities/Private/HLODBuilderInstancing.cpp
+++ b/Engine/Source/Developer/WorldPartitionHLODUtilities/Private/HLODBuilderInstancing.cpp
@@ -110,7 +110,7 @@ TArray<UPrimitiveComponent*> FHLODBuilder_Instancing::CreateComponents(AWorldPar
 				InstanceTransforms.Add(SMC->GetComponentTransform());
 			}
 		}
-		Component->AddInstances(InstanceTransforms, /*bWorldSpace*/true);
+		Component->AddInstances(InstanceTransforms, /*bShouldReturnIndices*/false, /*bWorldSpace*/true);
 
 		Components.Add(Component);
 	};
diff --git a/Engine/Source/Runtime/Engine/Private/LODActor.cpp b/Engine/Source/Runtime/Engine/Private/LODActor.cpp
--- a/Engine/Source/Runtime/Engine/Private/LODActor.cpp
+++ b/Engine/Source/Runtime/Engine/Private/LODActor.cpp
@@ -1032,7 +1032,7 @@ void ALODActor::AddInstances(const UStaticMesh* InStaticMesh, const UMaterialInt
 	check(InTransforms.Num() > 0);
 
 	UInstancedStaticMeshComponent* Component = GetOrCreateISMComponent(FHLODInstancingKey(InStaticMesh, InMaterial));
-	Component->AddInstances(InTransforms, /*bWorldSpace*/true);
+	Component->AddInstances(InTransforms, /*bShouldReturnIndices*/false, /*bWorldSpace*/true);
 
 	// Ensure parenting is up to date and take into account the newly created component.
 	UpdateSubActorLODParents();
diff --git a/Engine/Source/Runtime/Engine/Private/LevelInstance/Packed/LevelInstanceISMPacker.cpp b/Engine/Source/Runtime/Engine/Private/LevelInstance/Packed/LevelInstanceISMPacker.cpp
--- a/Engine/Source/Runtime/Engine/Private/LevelInstance/Packed/LevelInstanceISMPacker.cpp
+++ b/Engine/Source/Runtime/Engine/Private/LevelInstance/Packed/LevelInstanceISMPacker.cpp
@@ -73,7 +73,7 @@ void FLevelInstanceISMPacker::PackActors(FPackedLevelInstanceBuilderContext& InC
 			InstanceTransforms.Add(StaticMeshComponent->GetComponentTransform());
 		}
 	}
-	PackComponent->AddInstances(InstanceTransforms, /*bWorldSpace*/true);
+	PackComponent->AddInstances(InstanceTransforms, /*bShouldReturnIndices*/false, /*bWorldSpace*/true);
 
 	FTransform NewWorldTransform = ActorTransform * CurrentPivotOffsetInverse * FTransform(InContext.GetPivotOffset());
 
diff --git a/Engine/Source/Runtime/Foliage/Private/InstancedFoliage.cpp b/Engine/Source/Runtime/Foliage/Private/InstancedFoliage.cpp
--- a/Engine/Source/Runtime/Foliage/Private/InstancedFoliage.cpp
+++ b/Engine/Source/Runtime/Foliage/Private/InstancedFoliage.cpp
@@ -1654,7 +1654,7 @@ void FFoliageStaticMesh::Reapply(const UFoliageType* FoliageType)
 
 		TArray<FTransform> InstanceTransforms;
 		Algo::Transform(Info->Instances, InstanceTransforms, &FFoliageInstance::GetInstanceWorldTransform);
-		Component->AddInstances(InstanceTransforms, /*bWorldSpace*/true);
+		Component->AddInstances(InstanceTransforms, /*bShouldReturnIndices*/false, /*bWorldSpace*/true);
 
 		Component->bAutoRebuildTreeOnInstanceChanges = true;
 		Component->BuildTreeIfOutdated(true, true);
