When cooking, world partition initialization now only allows loading of actors marked as Always Loaded.

diff --git a/Engine/Source/Runtime/Engine/Private/WorldPartition/WorldPartitionEditorSpatialHash.cpp b/Engine/Source/Runtime/Engine/Private/WorldPartition/WorldPartitionEditorSpatialHash.cpp
--- a/Engine/Source/Runtime/Engine/Private/WorldPartition/WorldPartitionEditorSpatialHash.cpp
+++ b/Engine/Source/Runtime/Engine/Private/WorldPartition/WorldPartitionEditorSpatialHash.cpp
@@ -138,6 +138,13 @@ bool UWorldPartitionEditorSpatialHash::IsActorAlwaysLoaded(const FWorldPartition
 		return true;
 	}
 
+	// When cooking, we only want actors marked as AlwaysLoaded to be loaded in the PersistentLevel.
+	// @todo_ow: https://jira.it.epicgames.com/browse/UETOOL-3601
+	if (IsRunningCookCommandlet())
+	{
+		return false;
+	}
+
 	// If an actor covers more that 4 levels in the octree (which means 32K cells), treat it as always loaded
 	const FBox ActorBounds = GetActorBounds(InActorHandle);
 	const int32 ActorLevel = GetLevelForBox(ActorBounds);
