World Partition CookOnTheFly: Enable deferred population, which is faster and supports cook on the fly.
#rnx

diff --git a/Engine/Source/Runtime/Engine/Private/WorldPartition/WorldPartitionCookPackageSplitter.cpp b/Engine/Source/Runtime/Engine/Private/WorldPartition/WorldPartitionCookPackageSplitter.cpp
--- a/Engine/Source/Runtime/Engine/Private/WorldPartition/WorldPartitionCookPackageSplitter.cpp
+++ b/Engine/Source/Runtime/Engine/Private/WorldPartition/WorldPartitionCookPackageSplitter.cpp
@@ -4,6 +4,7 @@
 
 #if WITH_EDITOR
 
+#include "Misc/ConfigCacheIni.h"
 #include "WorldPartition/WorldPartitionRuntimeLevelStreamingCell.h"
 #include "WorldPartition/WorldPartitionLevelStreamingDynamic.h"
 #include "WorldPartition/WorldPartitionRuntimeCell.h"
@@ -20,6 +21,19 @@ bool FWorldPartitionCookPackageSplitter::ShouldSplit(UObject* SplitData)
 	return World && !!World->GetWorldPartition();
 }
 
+bool FWorldPartitionCookPackageSplitter::UseDeferredPopulate()
+{
+	static struct FOneTimeRead
+	{
+		FOneTimeRead()
+		{
+			GConfig->GetBool(TEXT("CookSettings"), TEXT("SplitterDeferredPopulate"), Value, GEditorIni);
+		}
+		bool Value = true;
+	} UseDeferredPopulate;
+	return UseDeferredPopulate.Value;
+}
+
 UWorld* FWorldPartitionCookPackageSplitter::ValidateDataObject(UObject* SplitData)
 {
 	UWorld* PartitionedWorld = CastChecked<UWorld>(SplitData);
diff --git a/Engine/Source/Runtime/Engine/Private/WorldPartition/WorldPartitionCookPackageSplitter.h b/Engine/Source/Runtime/Engine/Private/WorldPartition/WorldPartitionCookPackageSplitter.h
--- a/Engine/Source/Runtime/Engine/Private/WorldPartition/WorldPartitionCookPackageSplitter.h
+++ b/Engine/Source/Runtime/Engine/Private/WorldPartition/WorldPartitionCookPackageSplitter.h
@@ -13,7 +13,7 @@ class FWorldPartitionCookPackageSplitter : public ICookPackageSplitter, public F
 public:
 	//~ Begin of ICookPackageSplitter
 	static bool ShouldSplit(UObject* SplitData);
-	virtual bool UseDeferredPopulate() { return false; }
+	virtual bool UseDeferredPopulate() override;
 	virtual ~FWorldPartitionCookPackageSplitter() {}
 	virtual TArray<ICookPackageSplitter::FGeneratedPackage> GetGenerateList(const UPackage* OwnerPackage, const UObject* OwnerObject) override;
 	virtual bool TryPopulatePackage(const UPackage* OwnerPackage, const UObject* OwnerObject,
