World partition editor cells won't report destroyed actors as references (because saving empty packages requires having no referencers to inner objects), but still keep them around when collecting garbage.

diff --git a/Engine/Source/Runtime/Engine/Private/WorldPartition/WorldPartitionEditorCell.cpp b/Engine/Source/Runtime/Engine/Private/WorldPartition/WorldPartitionEditorCell.cpp
--- a/Engine/Source/Runtime/Engine/Private/WorldPartition/WorldPartitionEditorCell.cpp
+++ b/Engine/Source/Runtime/Engine/Private/WorldPartition/WorldPartitionEditorCell.cpp
@@ -8,19 +8,24 @@
 #if WITH_EDITOR
 void UWorldPartitionEditorCell::AddReferencedObjects(UObject* InThis, FReferenceCollector& Collector)
 {
-	UWorldPartitionEditorCell* This = CastChecked<UWorldPartitionEditorCell>(InThis);
-	
-#if 0
-	Collector.AllowEliminatingReferences(false);
-	for (const FActorReference& ActorReference: This->LoadedActors)
+	// We need to keep all loaded actors alive, mainly for deleted actors. Normally, these actors are only referenced
+	// by the transaction buffer, but we clear it when unloading cells, etc. and we don't want these actors to die.
+	// Also, we must avoid reporting these references when not collecting garbage, as code such as package deletion
+	// will skip packages with actors still referenced (via GatherObjectReferencersForDeletion).
+	if (IsGarbageCollecting())
 	{
-		AActor* LoadedActor = ActorReference->GetActor(/*bEvenIfPendingKill*/true, /*bEvenIfUnreachable*/true);
-		check(LoadedActor);
+		UWorldPartitionEditorCell* This = CastChecked<UWorldPartitionEditorCell>(InThis);
+	
+		Collector.AllowEliminatingReferences(false);
+		for (const FActorReference& ActorReference: This->LoadedActors)
+		{
+			AActor* LoadedActor = ActorReference->GetActor(/*bEvenIfPendingKill*/true, /*bEvenIfUnreachable*/true);
+			check(LoadedActor);
 	
-		Collector.AddReferencedObject(LoadedActor);
+			Collector.AddReferencedObject(LoadedActor);
+		}
+		Collector.AllowEliminatingReferences(true);
 	}
-	Collector.AllowEliminatingReferences(true);
-#endif
 
 	Super::AddReferencedObjects(InThis, Collector);
 }
