InstancedStaticMeshComponent: avoid calling GetActiveLightingScenario twice in PostLoad.
Added specific GetActiveLightingScenario implementation for partitioned worlds.
Also added some profiling tags.

diff --git a/Engine/Source/Runtime/Engine/Private/HierarchicalInstancedStaticMesh.cpp b/Engine/Source/Runtime/Engine/Private/HierarchicalInstancedStaticMesh.cpp
--- a/Engine/Source/Runtime/Engine/Private/HierarchicalInstancedStaticMesh.cpp
+++ b/Engine/Source/Runtime/Engine/Private/HierarchicalInstancedStaticMesh.cpp
@@ -3200,6 +3200,7 @@ void UHierarchicalInstancedStaticMeshComponent::UpdateDensityScaling()
 void UHierarchicalInstancedStaticMeshComponent::OnPostLoadPerInstanceData()
 {
 	SCOPE_CYCLE_COUNTER(STAT_FoliagePostLoad);
+	TRACE_CPUPROFILER_EVENT_SCOPE(UHierarchicalInstancedStaticMeshComponent::OnPostLoadPerInstanceData);
 
 	// Tree will be fully rebuilt when staticmesh has finished compiling
 	if (!HasAnyFlags(RF_ClassDefaultObject | RF_ArchetypeObject) && !IsCompiling())
diff --git a/Engine/Source/Runtime/Engine/Private/InstancedStaticMesh.cpp b/Engine/Source/Runtime/Engine/Private/InstancedStaticMesh.cpp
--- a/Engine/Source/Runtime/Engine/Private/InstancedStaticMesh.cpp
+++ b/Engine/Source/Runtime/Engine/Private/InstancedStaticMesh.cpp
@@ -3621,6 +3621,8 @@ void UInstancedStaticMeshComponent::PostLoad()
 
 void UInstancedStaticMeshComponent::OnPostLoadPerInstanceData()
 {
+	TRACE_CPUPROFILER_EVENT_SCOPE(UInstancedStaticMeshComponent::OnPostLoadPerInstanceData);
+
 	if (!HasAnyFlags(RF_ClassDefaultObject|RF_ArchetypeObject))
 	{
 		InitPerInstanceRenderData(true, InstanceDataBuffers.Get());
@@ -3635,8 +3637,9 @@ void UInstancedStaticMeshComponent::OnPostLoadPerInstanceData()
 		{
 			ULevel* OwnerLevel = Owner->GetLevel();
 			UWorld* OwnerWorld = OwnerLevel ? OwnerLevel->OwningWorld : nullptr;
+			ULevel* ActiveLightingScenario = OwnerWorld ? OwnerWorld->GetActiveLightingScenario() : nullptr;
 
-			if (OwnerWorld && OwnerWorld->GetActiveLightingScenario() != nullptr && OwnerWorld->GetActiveLightingScenario() != OwnerLevel)
+			if (ActiveLightingScenario && ActiveLightingScenario != OwnerLevel)
 			{
 				//update the instance data if the lighting scenario isn't the owner level
 				InstanceUpdateCmdBuffer.Edit();
diff --git a/Engine/Source/Runtime/Engine/Private/MapBuildData.cpp b/Engine/Source/Runtime/Engine/Private/MapBuildData.cpp
--- a/Engine/Source/Runtime/Engine/Private/MapBuildData.cpp
+++ b/Engine/Source/Runtime/Engine/Private/MapBuildData.cpp
@@ -54,17 +54,27 @@ FArchive& operator<<(FArchive& Ar, FSkyAtmosphereMapBuildData& Data)
 
 ULevel* UWorld::GetActiveLightingScenario() const
 {
-	for (int32 LevelIndex = 0; LevelIndex < Levels.Num(); LevelIndex++)
+	if (PersistentLevel && PersistentLevel->bIsPartitioned)
 	{
-		ULevel* LocalLevel = Levels[LevelIndex];
-
-		if (LocalLevel->bIsVisible && LocalLevel->bIsLightingScenario)
+		if (PersistentLevel->bIsLightingScenario)
 		{
-			return LocalLevel;
+			return PersistentLevel;
 		}
 	}
+	else
+	{
+		for (int32 LevelIndex = 0; LevelIndex < Levels.Num(); LevelIndex++)
+		{
+			ULevel* LocalLevel = Levels[LevelIndex];
 
-	return NULL;
+			if (LocalLevel->bIsVisible && LocalLevel->bIsLightingScenario)
+			{
+				return LocalLevel;
+			}
+		}
+	}
+
+	return nullptr;
 }
 
 void UWorld::PropagateLightingScenarioChange()
