Fixed proxy render mode logic for hiding Nanite meshes when disabled, except in SM editor.

diff --git a/Engine/Source/Runtime/Engine/Private/StaticMeshRender.cpp b/Engine/Source/Runtime/Engine/Private/StaticMeshRender.cpp
--- a/Engine/Source/Runtime/Engine/Private/StaticMeshRender.cpp
+++ b/Engine/Source/Runtime/Engine/Private/StaticMeshRender.cpp
@@ -2339,18 +2339,18 @@ FPrimitiveSceneProxy* UStaticMeshComponent::CreateSceneProxy()
 	// If we didn't get a proxy, but Nanite was enabled on the asset when it was built, evaluate proxy creation
 	if (GetStaticMesh()->HasValidNaniteData())
 	{
-		if (GNaniteProxyRenderMode == 1) // Never render proxies
+		const bool bAllowProxyRender = GNaniteProxyRenderMode == 0
+	#if WITH_EDITORONLY_DATA
+			// Check for specific case of static mesh editor "proxy toggle"
+			|| (bDisplayNaniteProxyMesh && GNaniteProxyRenderMode == 2)
+	#endif
+		;
+
+		if (!bAllowProxyRender) // Never render proxies
 		{
 			// We don't want to fall back to Nanite proxy rendering, so just make the mesh invisible instead.
 			return nullptr;
 		}
-	#if WITH_EDITORONLY_DATA
-		// Check for specific case of static mesh editor "proxy toggle"
-		else if (bDisplayNaniteProxyMesh && GNaniteProxyRenderMode == 2)
-		{
-			// Do nothing - render proxy in this case
-		}
-	#endif
 
 		// Fall back to rendering Nanite proxy meshes with traditional static mesh scene proxies
 	}
