Change DerivedDataBuildWorker/Loop & TextureDerivedDataBuildExporter to use the DDC implementation of FBuildAction[Builder] and FBuildOutput[Builder]

diff --git a/Engine/Source/Developer/TextureBuild/Private/TextureBuildFunction.cpp b/Engine/Source/Developer/TextureBuild/Private/TextureBuildFunction.cpp
--- a/Engine/Source/Developer/TextureBuild/Private/TextureBuildFunction.cpp
+++ b/Engine/Source/Developer/TextureBuild/Private/TextureBuildFunction.cpp
@@ -269,7 +269,7 @@ void FTextureBuildFunction::Build(UE::DerivedData::FBuildContext& Context) const
 			FSharedBuffer MipHeader = FSharedBuffer::MakeView(&MipSize, sizeof(int32));
 			FCompositeBuffer CompositeMipBuffer(MipHeader, FSharedBuffer::MakeView(CompressedMips[MipIndex].RawData.GetData(), CompressedMips[MipIndex].RawData.Num()));
 
-			Context.AddPayload(UE::DerivedData::FPayloadId::FromName(*PayloadName), FCompressedBuffer::Compress(NAME_None, CompositeMipBuffer));
+			Context.AddPayload(UE::DerivedData::FPayloadId::FromName(*PayloadName), FCompressedBuffer::Compress(NAME_Default, CompositeMipBuffer));
 		}
 	}
 
@@ -373,6 +373,6 @@ void FTextureBuildFunction::Build(UE::DerivedData::FBuildContext& Context) const
 
 
 		FCompositeBuffer CompositeResult(OrderedBuffers);
-		Context.AddPayload(UE::DerivedData::FPayloadId::FromName("Texture"), FCompressedBuffer::Compress(NAME_None, CompositeResult));
+		Context.AddPayload(UE::DerivedData::FPayloadId::FromName("Texture"), FCompressedBuffer::Compress(NAME_Default, CompositeResult));
 	}
 }
diff --git a/Engine/Source/Programs/DerivedDataBuildWorker/DerivedDataBuildWorker.Build.cs b/Engine/Source/Programs/DerivedDataBuildWorker/DerivedDataBuildWorker.Build.cs
--- a/Engine/Source/Programs/DerivedDataBuildWorker/DerivedDataBuildWorker.Build.cs
+++ b/Engine/Source/Programs/DerivedDataBuildWorker/DerivedDataBuildWorker.Build.cs
@@ -13,6 +13,10 @@ public abstract class DerivedDataBuildWorkerTarget : TargetRules
 
 		SolutionDirectory = "Programs/BuildWorker";
 
+		// TODO: This is temporary to allow DDCUtils to not cause build errors when referenced
+		//		 by the build workers.  In the future the use of DDC should strip out backends
+		//		 and DDCUTils, so it will no longer be required to set bLegalToDistributeBinary.
+		bLegalToDistributeBinary		= true;
 		bUseXGEController				= false;
 		bCompileFreeType				= false;
 		bLoggingToMemoryEnabled			= true;
@@ -28,7 +32,7 @@ public abstract class DerivedDataBuildWorkerTarget : TargetRules
 		bWithLiveCoding					= false;
 		bBuildDeveloperTools			= false;
 		bUseMallocProfiler				= false;
-		bBuildWithEditorOnlyData		= false;
+		bBuildWithEditorOnlyData		= true;
 		bCompileAgainstEngine			= false;
 		bCompileAgainstCoreUObject		= false;
 		bCompileAgainstApplicationCore	= false;
@@ -60,15 +64,11 @@ public class DerivedDataBuildWorker : ModuleRules
 
 		PrivateIncludePaths.Add("Runtime/Launch/Private");		// For LaunchEngineLoop.cpp include
 
-		PrivateIncludePathModuleNames.AddRange(
-			new string[] {
-				"DerivedDataCache",
-			});
-
 		PrivateDependencyModuleNames.AddRange(
 			new string[] {
 				"Core",
 				"Projects",
+				"DerivedDataCache",
 			});
 
 		if (Target.bCompileAgainstApplicationCore)
diff --git a/Engine/Source/Programs/DerivedDataBuildWorker/Private/DerivedDataBuildLoop.cpp b/Engine/Source/Programs/DerivedDataBuildWorker/Private/DerivedDataBuildLoop.cpp
--- a/Engine/Source/Programs/DerivedDataBuildWorker/Private/DerivedDataBuildLoop.cpp
+++ b/Engine/Source/Programs/DerivedDataBuildWorker/Private/DerivedDataBuildLoop.cpp
@@ -3,6 +3,8 @@
 #include "DerivedDataBuildLoop.h"
 
 #include "Compression/CompressedBuffer.h"
+#include "DerivedDataBuild.h"
+#include "DerivedDataBuildOutput.h"
 #include "DerivedDataPayload.h"
 #include "Memory/SharedBuffer.h"
 #include "Misc/CommandLine.h"
@@ -50,112 +52,113 @@ static FSharedBuffer LoadFile(const FString& Path)
 	return Buffer;
 }
 
-static FIoHash HashAndWriteToCompressedBufferFile(const FString& Directory, const void* InData, uint64 InDataSize)
-{
-	FIoHash DataHash = FIoHash::HashBuffer(InData, InDataSize);
-	TStringBuilder<128> DataHashStringBuilder;
-	DataHashStringBuilder << DataHash;
-
-	FCompressedBuffer CompressedBufferContents = FCompressedBuffer::Compress(NAME_None, FSharedBuffer::MakeView(InData, InDataSize));
-	if (TUniquePtr<FArchive> FileAr{IFileManager::Get().CreateFileWriter(*(Directory / *DataHashStringBuilder), FILEWRITE_NoReplaceExisting)})
-	{
-		*FileAr << CompressedBufferContents;
-	}
-	return DataHash;
-}
-
 class FWorkerBuildContext : public FBuildContext
 {
 public:
 	FWorkerBuildContext(const FBuildLoop::FBuildActionRecord& InBuildActionRecord)
-	: BuildActionRecord(InBuildActionRecord)
+	: OutputBuilder(GetDerivedDataBuildRef().CreateOutput(InBuildActionRecord.OutputFilePath, InBuildActionRecord.BuildAction.GetFunction()))
+	, BuildActionRecord(InBuildActionRecord)
 	{
-		BuildWriter.BeginObject("BuildOutput");
-		BuildWriter.BeginArray("Payloads");
 	}
 
 	virtual ~FWorkerBuildContext()
 	{
-		BuildWriter.EndArray();
+		FBuildOutput Output = OutputBuilder.Build();
 
-		BuildWriter.EndObject();
+		if (Output.HasError())
+		{
+			UE_LOG(LogDerivedDataBuildLoop, Error, TEXT("Error detected in build output '%s', no payloads will be written."), *WriteToString<64>(Output.GetName()));
+		}
+		else
+		{
+			for (const FPayload& Payload : Output.GetPayloads())
+			{
+				check(Payload.IsValid());
+
+				if (Payload.HasData())
+				{
+					TStringBuilder<128> DataHashStringBuilder;
+					DataHashStringBuilder << Payload.GetRawHash();
+
+					if (TUniquePtr<FArchive> FileAr{IFileManager::Get().CreateFileWriter(*(BuildActionRecord.OutputPath / *DataHashStringBuilder), FILEWRITE_NoReplaceExisting)})
+					{
+						*FileAr << const_cast<FCompressedBuffer&>(Payload.GetData());
+					}
+				}
+			}
+		}
 
 		if (TUniquePtr<FArchive> FileAr{IFileManager::Get().CreateFileWriter(*BuildActionRecord.OutputFilePath)})
 		{
-			BuildWriter.Save(*FileAr);
+			FCbWriter OutputWriter;
+			Output.Save(OutputWriter);
+			OutputWriter.Save(*FileAr);
 		}
 	}
 
 	virtual FCbObject GetConstant(FStringView Key) const override
 	{
-		TStringConversion<FTCHARToUTF8_Convert> ConvertedKey(Key.GetData(), Key.Len());
-		return BuildActionRecord.BuildAction["Constants"].AsObject()[ConvertedKey].AsObject();
+		FCbObject FoundVal;
+		BuildActionRecord.BuildAction.IterateConstants([&Key, &FoundVal] (FStringView ConstantKey, FCbObject&& Value)
+			{
+				if (Key == ConstantKey)
+					FoundVal = MoveTemp(Value);
+			});
+		return FoundVal;
 	}
 
 	virtual FSharedBuffer GetInput(FStringView Key) const override
 	{
-		TStringConversion<FTCHARToUTF8_Convert> ConvertedKey(Key.GetData(), Key.Len());
-		FIoHash InputHash(BuildActionRecord.BuildAction["Inputs"].AsObject()[ConvertedKey].AsBinaryAttachment());
-		if (!InputHash.IsZero())
-		{
-			TStringBuilder<256> Path;
-			Path << BuildActionRecord.InputPath << TEXT('/') << InputHash;
-			return FCompressedBuffer::FromCompressed(LoadFile(*Path)).Decompress();
-		}
-		return FSharedBuffer();
+		FSharedBuffer FoundVal;
+		BuildActionRecord.BuildAction.IterateInputs([this, &Key, &FoundVal] (FStringView InputKey, const FIoHash& RawHash, uint64 RawSize)
+			{
+				if (!RawHash.IsZero() && (InputKey == Key))
+				{
+					TStringBuilder<256> Path;
+					Path << BuildActionRecord.InputPath << TEXT('/') << RawHash;
+					FoundVal = FCompressedBuffer::FromCompressed(LoadFile(*Path)).Decompress();
+				}
+			});
+		return FoundVal;
 	}
 
 	virtual void AddPayload(const FPayload& Payload) override
 	{
-		AddPayload(Payload.GetId(), Payload.GetData());
+		OutputBuilder.AddPayload(Payload);
 	}
 	virtual void AddPayload(const FPayloadId& Id, const FCompressedBuffer& Buffer)
 	{
-		BuildWriter.BeginObject();
-		BuildWriter.AddObjectId("Id", FCbObjectId(Id.GetView()));
-		BuildWriter.AddInteger("RawSize", Buffer.GetRawSize());
-
-		FIoHash IoHash(Buffer.GetRawHash());
-		TStringBuilder<128> DataHashStringBuilder;
-		DataHashStringBuilder << IoHash;
-		if (TUniquePtr<FArchive> FileAr{IFileManager::Get().CreateFileWriter(*(BuildActionRecord.OutputPath / *DataHashStringBuilder), FILEWRITE_NoReplaceExisting)})
-		{
-			*FileAr << const_cast<FCompressedBuffer&>(Buffer);
-		}
-		BuildWriter.AddBinaryAttachment("RawHash", IoHash);
-		BuildWriter.EndObject();
+		FPayload Payload(Id, Buffer);
+		OutputBuilder.AddPayload(Payload);
 	}
 	virtual void AddPayload(const FPayloadId& Id, const FSharedBuffer& Buffer) override
 	{
-		BuildWriter.BeginObject();
-		BuildWriter.AddObjectId("Id", FCbObjectId(Id.GetView()));
-		BuildWriter.AddInteger("RawSize", Buffer.GetSize());
-		FIoHash Hash = HashAndWriteToCompressedBufferFile(BuildActionRecord.OutputPath, Buffer.GetData(), Buffer.GetSize());
-		BuildWriter.AddBinaryAttachment("RawHash", Hash);
-		BuildWriter.EndObject();
+		FPayload Payload(Id, FCompressedBuffer::Compress(NAME_Default, Buffer));
+		OutputBuilder.AddPayload(Payload);
 	}
 	virtual void AddPayload(const FPayloadId& Id, const FCbObject& Object) override
 	{
+		// TODO: Should be changed to a 'GetSerializedBuffer' call to handle cases where the Object is a UniformObject.
 		AddPayload(Id, Object.GetBuffer());
 	}
 
 	virtual void SetCachePolicy(ECachePolicy Policy) override
 	{
-		checkNoEntry();
+		unimplemented();
 	}
 
 	virtual void BeginAsyncBuild() override
 	{
-		check(false);
+		unimplemented();
 	}
 
 	virtual void EndAsyncBuild() override
 	{
-		check(false);
+		unimplemented();
 	}
 
 private:
-	FCbWriter BuildWriter;
+	FBuildOutputBuilder OutputBuilder;
 	const FBuildLoop::FBuildActionRecord& BuildActionRecord;
 };
 
@@ -164,7 +167,7 @@ FBuildLoop::FBuildActionRecord::FBuildActionRecord(const FString& InSourceFilePa
 , OutputFilePath(FPaths::ChangeExtension(InSourceFilePath, TEXT("uddbo")))
 , InputPath(InCommonInputPath.IsEmpty() ? FPaths::GetPath(InSourceFilePath) / TEXT("Inputs") : InCommonInputPath)
 , OutputPath(InCommonOutputPath.IsEmpty() ? FPaths::GetPath(InSourceFilePath) / TEXT("Outputs") : InCommonOutputPath)
-, BuildAction(MoveTemp(InSharedBuffer))
+, BuildAction(GetDerivedDataBuildRef().LoadAction(InSourceFilePath, FCbObject(MoveTemp(InSharedBuffer))).Get())
 {
 }
 
@@ -303,7 +306,7 @@ void FBuildLoop::PerformBuilds(const FBuildFunctionCallback& BuildFunctionCallba
 	for (const FBuildActionRecord& BuildActionRecord : BuildActionRecords)
 	{
 		FWorkerBuildContext Context(BuildActionRecord);
-		BuildFunctionCallback(FName(BuildActionRecord.BuildAction["Function"].AsObject()["Name"].AsString()), Context);
+		BuildFunctionCallback(FName(BuildActionRecord.BuildAction.GetFunction()), Context);
 	}
 }
 
diff --git a/Engine/Source/Programs/DerivedDataBuildWorker/Private/DerivedDataBuildLoop.h b/Engine/Source/Programs/DerivedDataBuildWorker/Private/DerivedDataBuildLoop.h
--- a/Engine/Source/Programs/DerivedDataBuildWorker/Private/DerivedDataBuildLoop.h
+++ b/Engine/Source/Programs/DerivedDataBuildWorker/Private/DerivedDataBuildLoop.h
@@ -4,6 +4,7 @@
 
 #include "CoreMinimal.h"
 
+#include "DerivedDataBuildAction.h"
 #include "DerivedDataBuildFunction.h"
 #include "DerivedDataBuildWorkerInterface.h"
 #include "Serialization/CompactBinary.h"
@@ -34,7 +35,7 @@ private:
 		const FString OutputFilePath;
 		const FString InputPath;
 		const FString OutputPath;
-		const FCbObject BuildAction;
+		const FBuildAction BuildAction;
 
 		FBuildActionRecord(const FString& InSourceFilePath, const FString& InCommonInputPath, const FString& InCommonOutputPath, FSharedBuffer&& InSharedBuffer);
 	};
diff --git a/Engine/Source/Runtime/Engine/Private/TextureDerivedDataBuildExporter.cpp b/Engine/Source/Runtime/Engine/Private/TextureDerivedDataBuildExporter.cpp
--- a/Engine/Source/Runtime/Engine/Private/TextureDerivedDataBuildExporter.cpp
+++ b/Engine/Source/Runtime/Engine/Private/TextureDerivedDataBuildExporter.cpp
@@ -4,12 +4,16 @@
 
 #if WITH_EDITOR
 
+#include "DerivedDataBuild.h"
+#include "DerivedDataBuildAction.h"
+#include "DerivedDataBuildOutput.h"
 #include "DerivedDataPayload.h"
 #include "HAL/FileManager.h"
 #include "HAL/IConsoleManager.h"
 #include "HAL/PlatformFileManager.h"
 #include "Interfaces/ITextureFormat.h"
 #include "Interfaces/ITextureFormatManagerModule.h"
+#include "Misc/CoreMisc.h"
 #include "Misc/FileHelper.h"
 #include "Misc/StringBuilder.h"
 #include "Serialization/CompactBinaryPackage.h"
@@ -87,9 +91,10 @@ static void WriteCbField(FCbWriter& Writer, const FAnsiStringView& Name, const F
 	Writer.EndArray();
 }
 
-static void WriteBuildSettingsToCompactBinary(FCbWriter& Writer, const FAnsiStringView& Name, const FTextureBuildSettings& BuildSettings, const ITextureFormat* TextureFormat)
+static FCbObject WriteBuildSettingsToCompactBinary(const FTextureBuildSettings& BuildSettings, const ITextureFormat* TextureFormat)
 {
-	Writer.BeginObject(Name);
+	FCbWriter Writer;
+	Writer.BeginObject();
 
 	if (BuildSettings.FormatConfigOverride)
 	{
@@ -168,11 +173,14 @@ static void WriteBuildSettingsToCompactBinary(FCbWriter& Writer, const FAnsiStri
 	Writer.AddBool("bHasEditorOnlyData", BuildSettings.bHasEditorOnlyData);
 
 	Writer.EndObject();
+	return Writer.Save().AsObject();
 }
 
-static void WriteOutputSettingsToCompactBinary(FCbWriter& Writer, const FAnsiStringView& Name, int32 NumInlineMips, const FString& KeySuffix)
+static FCbObject WriteOutputSettingsToCompactBinary(int32 NumInlineMips, const FString& KeySuffix)
 {
-	Writer.BeginObject(Name);
+	FCbWriter Writer;
+	Writer.BeginObject();
+
 	Writer.AddInteger("NumInlineMips", NumInlineMips);
 	
 	FString MipDerivedDataKey;
@@ -187,11 +195,13 @@ static void WriteOutputSettingsToCompactBinary(FCbWriter& Writer, const FAnsiStr
 	Writer.AddString("MipKeyPrefix",*MipDerivedDataKey);
 
 	Writer.EndObject();
+	return Writer.Save().AsObject();
 }
 
-static void WriteTextureSourceToCompactBinary(FCbWriter& Writer, const FAnsiStringView& Name, const FTextureSource& TextureSource, EGammaSpace GammaSpace)
+static FCbObject WriteTextureSourceToCompactBinary(const FTextureSource& TextureSource, EGammaSpace GammaSpace)
 {
-	Writer.BeginObject(Name);
+	FCbWriter Writer;
+	Writer.BeginObject();
 
 	Writer.AddString("Input", TextureSource.GetIdString());
 	Writer.AddInteger("CompressionFormat", TextureSource.GetSourceCompression());
@@ -215,6 +225,7 @@ static void WriteTextureSourceToCompactBinary(FCbWriter& Writer, const FAnsiStri
 	Writer.EndArray();
 
 	Writer.EndObject();
+	return Writer.Save().AsObject();
 }
 
 static FIoHash HashAndWriteToCompressedBufferFile(const FString& Directory, const void* InData, uint64 InDataSize)
@@ -223,7 +234,7 @@ static FIoHash HashAndWriteToCompressedBufferFile(const FString& Directory, cons
 	TStringBuilder<128> DataHashStringBuilder;
 	DataHashStringBuilder << DataHash;
 
-	FCompressedBuffer CompressedBufferContents = FCompressedBuffer::Compress(NAME_None, FSharedBuffer::MakeView(InData, InDataSize));
+	FCompressedBuffer CompressedBufferContents = FCompressedBuffer::Compress(NAME_Default, FSharedBuffer::MakeView(InData, InDataSize));
 	if (TUniquePtr<FArchive> FileAr{IFileManager::Get().CreateFileWriter(*(Directory / *DataHashStringBuilder), FILEWRITE_NoReplaceExisting)})
 	{
 		*FileAr << CompressedBufferContents;
@@ -231,7 +242,7 @@ static FIoHash HashAndWriteToCompressedBufferFile(const FString& Directory, cons
 	return DataHash;
 }
 
-static FIoHash ExportTextureBulkDataAttachment(const FString& ExportRoot, FTextureSource& TextureSource)
+static FIoHash ExportTextureBulkDataAttachment(const FString& ExportRoot, FTextureSource& TextureSource, uint64& OutSize)
 {
 	const FString BuildInputPath = ExportRoot / TEXT("Inputs");
 
@@ -240,12 +251,14 @@ static FIoHash ExportTextureBulkDataAttachment(const FString& ExportRoot, FTextu
 
 	if (!bExporting)
 	{
+		OutSize = 0;
 		return FIoHash();
 	}
 
 	FIoHash BulkDataHash;
-	TextureSource.OperateOnLoadedBulkData([&BuildInputPath, &BulkDataHash] (const FSharedBuffer& BulkDataBuffer) 
+	TextureSource.OperateOnLoadedBulkData([&BuildInputPath, &BulkDataHash, &OutSize] (const FSharedBuffer& BulkDataBuffer) 
 	{
+		OutSize = BulkDataBuffer.GetSize();
 		BulkDataHash = HashAndWriteToCompressedBufferFile(BuildInputPath, BulkDataBuffer.GetData(), BulkDataBuffer.GetSize());
 	});
 
@@ -261,6 +274,13 @@ void FTextureDerivedDataBuildExporter::Init(const FString& InKeySuffix)
 		return;
 	}
 
+	DerivedDataBuild = GetDerivedDataBuild();
+	if (!DerivedDataBuild)
+	{
+		bEnabled = false;
+		return;
+	}
+
 	KeySuffix = InKeySuffix;
 
 	FString DerivedDataKeyLong;
@@ -274,7 +294,7 @@ void FTextureDerivedDataBuildExporter::ExportTextureSourceBulkData(FTextureSourc
 {
 	if (bEnabled)
 	{
-		ExportedTextureBulkDataHash = ExportTextureBulkDataAttachment(ExportRoot, TextureSource);
+		ExportedTextureBulkDataHash = ExportTextureBulkDataAttachment(ExportRoot, TextureSource, ExportedTextureBulkDataSize);
 	}
 }
 
@@ -282,7 +302,7 @@ void FTextureDerivedDataBuildExporter::ExportCompositeTextureSourceBulkData(FTex
 {
 	if (bEnabled)
 	{
-		ExportedCompositeTextureBulkDataHash = ExportTextureBulkDataAttachment(ExportRoot, TextureSource);
+		ExportedCompositeTextureBulkDataHash = ExportTextureBulkDataAttachment(ExportRoot, TextureSource, ExportedCompositeTextureBulkDataSize);
 	}
 }
 
@@ -315,51 +335,43 @@ void FTextureDerivedDataBuildExporter::ExportTextureBuild(const UTexture& Textur
 	{
 		return;
 	}
-
-	FCbWriter BuildWriter;
-	BuildWriter.BeginObject("BuildAction");
-
-	BuildWriter.BeginObject("Function");
-
 	// Texture format modules are inconsistent in their naming.  eg: TextureFormatUncompressed, PS5TextureFormat
 	// We attempt to  unify the naming here when specifying build function names.
-	TStringBuilder<64> FormatModuleNameBuilder;
-	FormatModuleNameBuilder << TextureFormatModuleName.ToString().Replace(TEXT("TextureFormat"), TEXT(""));
-	FormatModuleNameBuilder << TEXT("Texture");
+	TStringBuilder<64> BuildFunctionNameBuilder;
+	BuildFunctionNameBuilder << TextureFormatModuleName.ToString().Replace(TEXT("TextureFormat"), TEXT(""));
+	BuildFunctionNameBuilder << TEXT("Texture");
 
-	BuildWriter.AddString("Name", FormatModuleNameBuilder);
-	BuildWriter.AddString("Version", "0"); // TODO: compute module version
-	BuildWriter.EndObject();
+	BuildFunctionName = BuildFunctionNameBuilder;
+	TexturePath = Texture.GetPathName();
 
-	BuildWriter.BeginObject("Constants");
-	WriteBuildSettingsToCompactBinary(BuildWriter, "TextureBuildSettings", BuildSettings, TextureFormat);
-	WriteOutputSettingsToCompactBinary(BuildWriter, "TextureOutputSettings", NumInlineMips, KeySuffix);
+	UE::DerivedData::FBuildActionBuilder ActionBuilder = DerivedDataBuild->CreateAction(TexturePath, BuildFunctionName);
+
+	ActionBuilder.AddConstant(TEXT("TextureBuildSettings"), WriteBuildSettingsToCompactBinary(BuildSettings, TextureFormat));
+	ActionBuilder.AddConstant(TEXT("TextureOutputSettings"), WriteOutputSettingsToCompactBinary(NumInlineMips, KeySuffix));
 
 	FTextureFormatSettings TextureFormatSettings;
 	Texture.GetLayerFormatSettings(LayerIndex, TextureFormatSettings);
 	EGammaSpace TextureGammaSpace = TextureFormatSettings.SRGB ? (Texture.bUseLegacyGamma ? EGammaSpace::Pow22 : EGammaSpace::sRGB) : EGammaSpace::Linear;
-	WriteTextureSourceToCompactBinary(BuildWriter, "TextureSource", Texture.Source, TextureGammaSpace);
+	ActionBuilder.AddConstant(TEXT("TextureSource"), WriteTextureSourceToCompactBinary(Texture.Source, TextureGammaSpace));
 	if ((bool)Texture.CompositeTexture && !ExportedCompositeTextureBulkDataHash.IsZero())
 	{
 		FTextureFormatSettings CompositeTextureFormatSettings;
 		Texture.CompositeTexture->GetLayerFormatSettings(LayerIndex, CompositeTextureFormatSettings);
 		EGammaSpace CompositeTextureGammaSpace = CompositeTextureFormatSettings.SRGB ? (Texture.CompositeTexture->bUseLegacyGamma ? EGammaSpace::Pow22 : EGammaSpace::sRGB) : EGammaSpace::Linear;
 
-		WriteTextureSourceToCompactBinary(BuildWriter, "CompositeTextureSource", Texture.CompositeTexture->Source, CompositeTextureGammaSpace);
+		ActionBuilder.AddConstant(TEXT("CompositeTextureSource"), WriteTextureSourceToCompactBinary(Texture.CompositeTexture->Source, CompositeTextureGammaSpace));
 	}
-	BuildWriter.EndObject();
 
-	BuildWriter.BeginObject("Inputs");
-	BuildWriter.AddBinaryAttachment(TCHAR_TO_UTF8(*Texture.Source.GetIdString()), ExportedTextureBulkDataHash);
+	ActionBuilder.AddInput(Texture.Source.GetIdString(), ExportedTextureBulkDataHash, ExportedTextureBulkDataSize);
 	if ((bool)Texture.CompositeTexture && !ExportedCompositeTextureBulkDataHash.IsZero())
 	{
-		BuildWriter.AddBinaryAttachment(TCHAR_TO_UTF8(*Texture.CompositeTexture->Source.GetIdString()), ExportedCompositeTextureBulkDataHash);
+		ActionBuilder.AddInput(Texture.CompositeTexture->Source.GetIdString(), ExportedCompositeTextureBulkDataHash, ExportedCompositeTextureBulkDataSize);
 	}
-	BuildWriter.EndObject();
 
-	BuildWriter.EndObject();
 	if (TUniquePtr<FArchive> Ar{IFileManager::Get().CreateFileWriter(*(ExportRoot / TEXT("build.uddba")))})
 	{
+		FCbWriter BuildWriter;
+		ActionBuilder.Build().Save(BuildWriter);
 		BuildWriter.Save(*Ar);
 	}
 }
@@ -377,9 +389,7 @@ void FTextureDerivedDataBuildExporter::ExportTextureOutput(FTexturePlatformData&
 	const FString OutputPath = ExportRoot / TEXT("ReferenceOutputs");
 	TArray<TPair<FString,FString>> DDCReferences;
 
-	FCbWriter BuildWriter;
-	BuildWriter.BeginObject("BuildOutput");
-	BuildWriter.BeginArray("Payloads");
+	UE::DerivedData::FBuildOutputBuilder OutputBuilder = DerivedDataBuild->CreateOutput(TexturePath, BuildFunctionName);
 
 	const int32 MipCount = PlatformData.Mips.Num();
 	const int32 FirstInlineMip = bForceAllMipsToBeInlined ? 0 : FMath::Max(0, MipCount - FMath::Max((int32)NUM_INLINE_DERIVED_MIPS, (int32)PlatformData.GetNumMipsInTail()));
@@ -402,12 +412,10 @@ void FTextureDerivedDataBuildExporter::ExportTextureOutput(FTexturePlatformData&
 		TStringBuilder<32> PayloadName;
 		PayloadName << "Mip" << MipIndex;
 
-		BuildWriter.BeginObject();
-		BuildWriter.AddObjectId("Id", FCbObjectId(UE::DerivedData::FPayloadId::FromName(*PayloadName).GetView()));
-		BuildWriter.AddInteger("RawSize", DerivedData.Num());
 		FIoHash DerivedDataHash = HashAndWriteToCompressedBufferFile(OutputPath, DerivedData.GetData(), DerivedData.Num());
-		BuildWriter.AddBinaryAttachment("RawHash", DerivedDataHash);
-		BuildWriter.EndObject();
+		UE::DerivedData::FPayload MipPayload(UE::DerivedData::FPayloadId::FromName(PayloadName), DerivedDataHash, DerivedData.Num());
+
+		OutputBuilder.AddPayload(MipPayload);
 
 		check(Mip.DerivedDataKey.IsEmpty());
 
@@ -430,26 +438,22 @@ void FTextureDerivedDataBuildExporter::ExportTextureOutput(FTexturePlatformData&
 		PlatformData.Mips[MipIndex].DerivedDataKey.Empty();
 	}
 
-	BuildWriter.BeginObject();
-	BuildWriter.AddObjectId("Id", FCbObjectId(UE::DerivedData::FPayloadId::FromName("Texture").GetView()));
-	BuildWriter.AddInteger("RawSize", RawDerivedData.Num());
 	FIoHash DerivedDataHash = HashAndWriteToCompressedBufferFile(OutputPath, RawDerivedData.GetData(), RawDerivedData.Num());
-	BuildWriter.AddBinaryAttachment("RawHash", DerivedDataHash);
-	BuildWriter.EndObject();
+	UE::DerivedData::FPayload TexturePayload(UE::DerivedData::FPayloadId::FromName("Texture"), DerivedDataHash, RawDerivedData.Num());
 
-	FString DerivedDataKey;
-	ShortenKey(*DerivedDataKeyLong, DerivedDataKey);
-	DDCReferences.Emplace(TEXT("Texture"), DerivedDataKey);
-
-	BuildWriter.EndArray();
-
-	BuildWriter.EndObject();
+	OutputBuilder.AddPayload(TexturePayload);
 
 	if (TUniquePtr<FArchive> FileAr{IFileManager::Get().CreateFileWriter(*(ExportRoot / TEXT("ReferenceOutput.uddbo")))})
 	{
-		BuildWriter.Save(*FileAr);
+		FCbWriter OutputWriter;
+		OutputBuilder.Build().Save(OutputWriter);
+		OutputWriter.Save(*FileAr);
 	}
 
+	FString DerivedDataKey;
+	ShortenKey(*DerivedDataKeyLong, DerivedDataKey);
+	DDCReferences.Emplace(TEXT("Texture"), DerivedDataKey);
+
 	TArray<FString> DDCRefStringArray;
 	FString Separator(TEXT(","));
 	for (TPair<FString, FString>& DDCReference : DDCReferences)
diff --git a/Engine/Source/Runtime/Engine/Private/TextureDerivedDataBuildExporter.h b/Engine/Source/Runtime/Engine/Private/TextureDerivedDataBuildExporter.h
--- a/Engine/Source/Runtime/Engine/Private/TextureDerivedDataBuildExporter.h
+++ b/Engine/Source/Runtime/Engine/Private/TextureDerivedDataBuildExporter.h
@@ -12,6 +12,7 @@ struct FTextureBuildSettings;
 struct FTexturePlatformData;
 class UTexture;
 
+namespace UE::DerivedData { class IBuild; }
 
 /** Collects input/output about the build of textures (i.e. compression) and exports the data to a file for testing of remote texture compression */
 class FTextureDerivedDataBuildExporter
@@ -25,10 +26,15 @@ public:
 	void ExportTextureBuild(const UTexture& Texture, const FTextureBuildSettings& BuildSettings, int32 LayerIndex, int32 MaxInlineMips);
 	void ExportTextureOutput(FTexturePlatformData& PlatformData, const FTextureBuildSettings& BuildSettings);
 private:
+	UE::DerivedData::IBuild* DerivedDataBuild;
 	FString KeySuffix;
+	FString TexturePath;
 	FString ExportRoot;
+	FString BuildFunctionName;
 	FIoHash ExportedTextureBulkDataHash;
+	uint64 ExportedTextureBulkDataSize {0};
 	FIoHash ExportedCompositeTextureBulkDataHash;
+	uint64 ExportedCompositeTextureBulkDataSize {0};
 	bool bEnabled = false;
 };
 
