U5 - adding missing parallel GT tag

diff --git a/Engine/Source/Runtime/Engine/Private/ShadowMap.cpp b/Engine/Source/Runtime/Engine/Private/ShadowMap.cpp
--- a/Engine/Source/Runtime/Engine/Private/ShadowMap.cpp
+++ b/Engine/Source/Runtime/Engine/Private/ShadowMap.cpp
@@ -310,6 +310,8 @@ void FShadowMapPendingTexture::CreateUObjects()
 
 void FShadowMapPendingTexture::StartEncoding(ULevel* LightingScenario, ITextureCompressorModule* Compressor)
 {
+	FOptionalTaskTagScope Scope(ETaskTag::EParallelGameThread);
+
 	// Create the shadow-map texture.
 	CreateUObjects();
 
