WorldPartition: DebugDraw filter per SceneView.ViewActor to avoid overlapping debug draw (Multiplayer PIE) + Add ClientOnly string to grids

diff --git a/Engine/Source/Runtime/Engine/Private/WorldPartition/WorldPartitionRuntimeSpatialHash.cpp b/Engine/Source/Runtime/Engine/Private/WorldPartition/WorldPartitionRuntimeSpatialHash.cpp
--- a/Engine/Source/Runtime/Engine/Private/WorldPartition/WorldPartitionRuntimeSpatialHash.cpp
+++ b/Engine/Source/Runtime/Engine/Private/WorldPartition/WorldPartitionRuntimeSpatialHash.cpp
@@ -1176,6 +1176,10 @@ void UWorldPartitionRuntimeSpatialHash::Draw2D(UCanvas* Canvas, const TArray<FWo
 		// Draw Grid name, loading range
 		{
 			FString GridInfoText = FString::Printf(TEXT("%s | %d m"), *StreamingGrid->GridName.ToString(), int32(StreamingGrid->GetLoadingRange() * 0.01f));
+			if (StreamingGrid->bClientOnlyVisible)
+			{
+				GridInfoText += "| Client Only";
+			}
 			Canvas->SetDrawColor(255, 255, 0);
 			Canvas->DrawText(GEngine->GetTinyFont(), GridInfoText, GridInfoPos.X, GridInfoPos.Y);
 		}
diff --git a/Engine/Source/Runtime/Engine/Private/WorldPartition/WorldPartitionSubsystem.cpp b/Engine/Source/Runtime/Engine/Private/WorldPartition/WorldPartitionSubsystem.cpp
--- a/Engine/Source/Runtime/Engine/Private/WorldPartition/WorldPartitionSubsystem.cpp
+++ b/Engine/Source/Runtime/Engine/Private/WorldPartition/WorldPartitionSubsystem.cpp
@@ -228,6 +228,13 @@ void UWorldPartitionSubsystem::Draw(UCanvas* Canvas, class APlayerController* PC
 		return;
 	}
 
+	// Filter out views that don't match our world
+	if (Canvas->SceneView->ViewActor != nullptr && 
+		Canvas->SceneView->ViewActor->GetWorld() != GetWorld())
+	{
+		return;
+	}
+
 	const FVector2D CanvasTopLeftPadding(10.f, 10.f);
 
 	FVector2D CurrentOffset(CanvasTopLeftPadding);
