change r.VT.EnableCompressZlib to off by default in code
leave on in config for some projects as needed

diff --git a/Engine/Source/Runtime/Engine/Private/VT/VirtualTextureBuildSettings.cpp b/Engine/Source/Runtime/Engine/Private/VT/VirtualTextureBuildSettings.cpp
--- a/Engine/Source/Runtime/Engine/Private/VT/VirtualTextureBuildSettings.cpp
+++ b/Engine/Source/Runtime/Engine/Private/VT/VirtualTextureBuildSettings.cpp
@@ -12,8 +12,8 @@ static TAutoConsoleVariable<int32> CVarVTCompressCrunch(
 
 static TAutoConsoleVariable<int32> CVarVTCompressZlib(
 	TEXT("r.VT.EnableCompressZlib"),
-	1,
-	TEXT("Enables Zlib compression for virtual textures, if no compression is enabled/supported")
+	0,
+	TEXT("Enables Zlib compression for virtual textures, instead of iostore compression")
 );
 
 static TAutoConsoleVariable<int32> CVarVTTileSize(
