support Mobile full precision floats only for MaterialExpressions, excluding shader codes in .ush/.usf

diff --git a/Engine/Shaders/Private/Common.ush b/Engine/Shaders/Private/Common.ush
--- a/Engine/Shaders/Private/Common.ush
+++ b/Engine/Shaders/Private/Common.ush
@@ -8,7 +8,7 @@
 
 #include "/Engine/Public/Platform.ush"
 // These types are used for material translator generated code, or any functions the translated code can call
-#if PIXELSHADER
+#if PIXELSHADER && !FORCE_MATERIAL_FLOAT_FULL_PRECISION
 	#define MaterialFloat half
 	#define MaterialFloat2 half2
 	#define MaterialFloat3 half3
diff --git a/Engine/Shaders/Private/Definitions.usf b/Engine/Shaders/Private/Definitions.usf
--- a/Engine/Shaders/Private/Definitions.usf
+++ b/Engine/Shaders/Private/Definitions.usf
@@ -159,3 +159,7 @@
 #ifndef SUPPORTS_ANISOTROPIC_MATERIALS
 #define SUPPORTS_ANISOTROPIC_MATERIALS					0
 #endif
+
+#ifndef FORCE_MATERIAL_FLOAT_FULL_PRECISION
+#define FORCE_MATERIAL_FLOAT_FULL_PRECISION             0
+#endif
\ No newline at end of file
diff --git a/Engine/Source/Editor/MaterialEditor/Private/MaterialEditorDetailCustomization.cpp b/Engine/Source/Editor/MaterialEditor/Private/MaterialEditorDetailCustomization.cpp
--- a/Engine/Source/Editor/MaterialEditor/Private/MaterialEditorDetailCustomization.cpp
+++ b/Engine/Source/Editor/MaterialEditor/Private/MaterialEditorDetailCustomization.cpp
@@ -828,7 +828,7 @@ void FMaterialDetailCustomization::CustomizeDetails( IDetailLayoutBuilder& Detai
 				FProperty* Property = PropertyHandle->GetProperty();
 				FName PropertyName = Property->GetFName();
 
-				if (PropertyName != GET_MEMBER_NAME_CHECKED(UMaterial, bUseFullPrecision)) 
+				if (PropertyName != GET_MEMBER_NAME_CHECKED(UMaterial, FloatPrecisionMode))
 				{
 					DetailLayout.HideProperty(PropertyHandle);
 				}
diff --git a/Engine/Source/Runtime/Engine/Classes/Engine/EngineTypes.h b/Engine/Source/Runtime/Engine/Classes/Engine/EngineTypes.h
--- a/Engine/Source/Runtime/Engine/Classes/Engine/EngineTypes.h
+++ b/Engine/Source/Runtime/Engine/Classes/Engine/EngineTypes.h
@@ -250,6 +250,19 @@ enum EBlendMode
 	BLEND_MAX,
 };
 
+/** The default float precision for material's pixel shaders on mobile devices*/
+UENUM()
+enum EMaterialFloatPrecisionMode
+{
+	/** Half precision, except explict 'float' in .ush/.usf*/
+	MFPM_Half UMETA(DisplayName = "Half-precision"),
+	/** Force full-precision for MaterialFloat only, no effect on shader codes in .ush/.usf*/
+	MFPM_Full_MaterialExpressionOnly UMETA(DisplayName = "Full-precision for MaterialExpressions only"),
+	/** All the floats are full-precision */
+	MFPM_Full UMETA(DisplayName = "Full-precision for every float"),
+	MFPM_MAX,
+};
+
 /** Controls where the sampler for different texture lookups comes from */
 UENUM()
 enum ESamplerSourceMode
diff --git a/Engine/Source/Runtime/Engine/Classes/Materials/Material.h b/Engine/Source/Runtime/Engine/Classes/Materials/Material.h
--- a/Engine/Source/Runtime/Engine/Classes/Materials/Material.h
+++ b/Engine/Source/Runtime/Engine/Classes/Materials/Material.h
@@ -762,13 +762,23 @@ public:
 	UPROPERTY(EditAnywhere, BlueprintReadOnly, Category=Material, AdvancedDisplay)
 	uint8 bFullyRough : 1;
 
-	/** 
-	 *	Forces this material to use full (highp) precision in the pixel shader.
+	/**
+	 *	Deprecated. Use FloatPrecisionMode instead.
+	 *  Forces this material to use full (highp) precision in the pixel shader.
 	 *	This is slower than the default (mediump) but can be used to work around precision-related rendering errors.
-	 *	This setting has no effect on older mobile devices that do not support high precision. 
+	 *	This setting has no effect on older mobile devices that do not support high precision.
+	 */
+	UPROPERTY()
+	uint8 bUseFullPrecision_DEPRECATED : 1;
+
+	/**
+	 *	How to use full (highp) precision in the pixel shader.
+	 *	highp is slower than the default (mediump) but can be used to work around precision-related rendering errors.
+	 *  Use 'Full-precision for MaterialExpressions only' if you still want to keep the precision of the halfs in .ush/.usf
+	 *	This setting has no effect on older mobile devices that do not support high precision.
 	 */
 	UPROPERTY(EditAnywhere, BlueprintReadOnly, Category = Mobile)
-	uint8 bUseFullPrecision : 1;
+	TEnumAsByte<enum EMaterialFloatPrecisionMode> FloatPrecisionMode;
 
 	/* Use lightmap directionality and per pixel normals. If disabled, lighting from lightmaps will be flat but cheaper. */
 	UPROPERTY(EditAnywhere, BlueprintReadOnly, Category=Mobile)
diff --git a/Engine/Source/Runtime/Engine/Private/Materials/Material.cpp b/Engine/Source/Runtime/Engine/Private/Materials/Material.cpp
--- a/Engine/Source/Runtime/Engine/Private/Materials/Material.cpp
+++ b/Engine/Source/Runtime/Engine/Private/Materials/Material.cpp
@@ -1017,6 +1017,8 @@ UMaterial::UMaterial(const FObjectInitializer& ObjectInitializer)
 	PhysMaterialMask = nullptr;
 
 	bLoadedCachedExpressionData = false;
+
+	FloatPrecisionMode = EMaterialFloatPrecisionMode::MFPM_Half;
 }
 
 void UMaterial::PreSave(const class ITargetPlatform* TargetPlatform)
@@ -4141,6 +4143,12 @@ void UMaterial::PostLoad()
 		DecalBlendMode = DBM_Translucent;
 	}
 
+	if(bUseFullPrecision_DEPRECATED && FloatPrecisionMode == EMaterialFloatPrecisionMode::MFPM_Half)
+	{
+		FloatPrecisionMode = EMaterialFloatPrecisionMode::MFPM_Full;
+		bUseFullPrecision_DEPRECATED = false;
+	}
+
 #if WITH_EDITOR
 	// Create exec flow expressions, if needed
 	CreateExecutionFlowExpressions();
diff --git a/Engine/Source/Runtime/Engine/Private/Materials/MaterialShared.cpp b/Engine/Source/Runtime/Engine/Private/Materials/MaterialShared.cpp
--- a/Engine/Source/Runtime/Engine/Private/Materials/MaterialShared.cpp
+++ b/Engine/Source/Runtime/Engine/Private/Materials/MaterialShared.cpp
@@ -1447,9 +1447,9 @@ bool FMaterialResource::UseNormalCurvatureToRoughness() const
 	return Material->bNormalCurvatureToRoughness;
 }
 
-bool FMaterialResource::IsUsingFullPrecision() const
+EMaterialFloatPrecisionMode FMaterialResource::GetMaterialFloatPrecisionMode() const
 {
-	return Material->bUseFullPrecision;
+	return Material->FloatPrecisionMode;
 }
 
 bool FMaterialResource::IsUsingAlphaToCoverage() const
@@ -2023,10 +2023,15 @@ void FMaterial::SetupMaterialEnvironment(
 	OutEnvironment.SetDefine(TEXT("TRANSLUCENT_SHADOW_WITH_MASKED_OPACITY"), GetCastDynamicShadowAsMasked());
 	OutEnvironment.SetDefine(TEXT("MATERIAL_USE_ALPHA_TO_COVERAGE"), IsUsingAlphaToCoverage());
 
-	if (IsUsingFullPrecision())
+	EMaterialFloatPrecisionMode FloatPrecisionMode = GetMaterialFloatPrecisionMode();
+	if (FloatPrecisionMode == EMaterialFloatPrecisionMode::MFPM_Full)
 	{
 		OutEnvironment.CompilerFlags.Add(CFLAG_UseFullPrecisionInPS);
 	}
+	else if (FloatPrecisionMode == EMaterialFloatPrecisionMode::MFPM_Full_MaterialExpressionOnly)
+	{
+		OutEnvironment.SetDefine(TEXT("FORCE_MATERIAL_FLOAT_FULL_PRECISION"), TEXT("1"));
+	}
 
 	switch(GetMaterialDomain())
 	{
diff --git a/Engine/Source/Runtime/Engine/Public/MaterialShared.h b/Engine/Source/Runtime/Engine/Public/MaterialShared.h
--- a/Engine/Source/Runtime/Engine/Public/MaterialShared.h
+++ b/Engine/Source/Runtime/Engine/Public/MaterialShared.h
@@ -1660,7 +1660,7 @@ public:
 	virtual bool IsUsedWithVirtualHeightfieldMesh() const { return false; }
 	virtual bool IsFullyRough() const { return false; }
 	virtual bool UseNormalCurvatureToRoughness() const { return false; }
-	virtual bool IsUsingFullPrecision() const { return false; }
+	virtual enum EMaterialFloatPrecisionMode GetMaterialFloatPrecisionMode() const { return EMaterialFloatPrecisionMode::MFPM_Half; };
 	virtual bool IsUsingAlphaToCoverage() const { return false; }
 	virtual bool IsUsingPreintegratedGFForSimpleIBL() const { return false; }
 	virtual bool IsUsingHQForwardReflections() const { return false; }
@@ -2542,7 +2542,7 @@ public:
 	ENGINE_API virtual bool IsUsedWithVirtualHeightfieldMesh() const override;
 	ENGINE_API virtual bool IsFullyRough() const override;
 	ENGINE_API virtual bool UseNormalCurvatureToRoughness() const override;
-	ENGINE_API virtual bool IsUsingFullPrecision() const override;
+	ENGINE_API virtual enum EMaterialFloatPrecisionMode GetMaterialFloatPrecisionMode() const override;
 	ENGINE_API virtual bool IsUsingAlphaToCoverage() const override;
 	ENGINE_API virtual bool IsUsingPreintegratedGFForSimpleIBL() const override;
 	ENGINE_API virtual bool IsUsingHQForwardReflections() const override;
