Fix to not crash when out of bounds bone index in SkeletalMeshComponentPhysics

diff --git a/Engine/Source/Runtime/Engine/Private/SkeletalMeshComponentPhysics.cpp b/Engine/Source/Runtime/Engine/Private/SkeletalMeshComponentPhysics.cpp
--- a/Engine/Source/Runtime/Engine/Private/SkeletalMeshComponentPhysics.cpp
+++ b/Engine/Source/Runtime/Engine/Private/SkeletalMeshComponentPhysics.cpp
@@ -1591,6 +1591,11 @@ void USkeletalMeshComponent::SendRenderDebugPhysics(FPrimitiveSceneProxy* Overri
 			if (BI && BI->IsValidBodyInstance())
 			{
 				const int32 BoneIndex = BI->InstanceBoneIndex;
+				if (BoneIndex >= GetComponentSpaceTransforms().Num())
+				{
+					UE_LOG(LogSkeletalMesh, Log, TEXT("SkeletalMeshComponent : (%d) Bone index out of bounds"), BoneIndex);
+					continue;
+				}
 				DebugMassData.AddDefaulted();
 				FPrimitiveSceneProxy::FDebugMassData& MassData = DebugMassData.Last();
 				const FTransform MassToWorld = BI->GetMassSpaceToWorldSpace();
