Added initial breadcrumbs for diagnosing a UBlueprintGeneratedClass-related crash during reference collection.

We've run into several crash reports where the UberGraphFunction appears to dereference a deleted UStruct. The initial breadcrumbs added here help diagnose the exact BP and property that was serialized.

diff --git a/Engine/Source/Runtime/Engine/Private/BlueprintGeneratedClass.cpp b/Engine/Source/Runtime/Engine/Private/BlueprintGeneratedClass.cpp
--- a/Engine/Source/Runtime/Engine/Private/BlueprintGeneratedClass.cpp
+++ b/Engine/Source/Runtime/Engine/Private/BlueprintGeneratedClass.cpp
@@ -24,6 +24,7 @@
 #include "Net/Core/PushModel/PushModel.h"
 #include "UObject/CoreObjectVersion.h"
 #include "UObject/UE5MainStreamObjectVersion.h"
+#include "GenericPlatform/GenericPlatformCrashContext.h"
 
 #if WITH_EDITOR
 #include "Kismet2/BlueprintEditorUtils.h"
@@ -53,6 +54,57 @@ static FAutoConsoleVariableRef CVarBlueprintComponentInstancingFastPathDisabled(
 	ECVF_Default
 );
 
+#if WITH_ADDITIONAL_CRASH_CONTEXTS
+struct BPGCBreadcrumbsParams
+{
+	UObject& Object;
+	UBlueprintGeneratedClass& BPGC;
+	FArchive& Ar;
+	uint32 ThreadId;
+};
+
+// Called during a crash: dynamic memory allocations may not be reliable here.
+void WriteBPGCBreadcrumbs(FCrashContextExtendedWriter& Writer, const BPGCBreadcrumbsParams& Params)
+{
+	constexpr int32 MAX_DATA_SIZE = 1024;
+	constexpr int32 MAX_FILENAME_SIZE = 64;
+	constexpr int32 MAX_THREADS_TO_LOG = 16;
+
+	static int32 ThreadCount = 0;
+
+	// In the unlikely case that there are too many threads that reported a BPGC-related crash, we'll just ignore the excess.
+	// Theoretically, there should be enough information written by the remaining threads.
+	if (ThreadCount < MAX_THREADS_TO_LOG)
+	{
+		// Note: TStringBuilder *can* potentially allocate dynamic memory if its inline storage is exceeded.
+		// In practice, we stay under the current threshold by only recording the minimum amount data that we need.
+		TStringBuilder<MAX_DATA_SIZE> Builder;
+
+		Params.Object.GetPathName(nullptr, Builder);
+		Builder.Append(TCHAR('\n'));
+		Params.BPGC.GetPathName(nullptr, Builder);
+		Builder.Append(TCHAR('\n'));
+
+		if (Params.Ar.GetSerializedProperty() != nullptr)
+		{
+			Params.Ar.GetSerializedProperty()->GetPathName(nullptr, Builder);
+		}
+		else
+		{
+			Builder.Append(TEXT("Serialized property was null!"));
+		}
+
+		TStringBuilder<MAX_FILENAME_SIZE> Filename;
+
+		Filename.Appendf(TEXT("BPGCBreadcrumb_%u"), Params.ThreadId);
+
+		Writer.AddString(Filename.ToString(), Builder.ToString());
+	}
+
+	++ThreadCount;
+}
+#endif // WITH_ADDITIONAL_CRASH_CONTEXTS
+
 UBlueprintGeneratedClass::UBlueprintGeneratedClass(const FObjectInitializer& ObjectInitializer)
 	: Super(ObjectInitializer)
 #if VALIDATE_UBER_GRAPH_PERSISTENT_FRAME
@@ -1743,14 +1795,30 @@ void UBlueprintGeneratedClass::AddReferencedObjectsInUbergraphFrame(UObject* InT
 					);
 #endif//VALIDATE_UBER_GRAPH_PERSISTENT_FRAME
 
-					checkSlow(BPGC->UberGraphFunction);
-					FVerySlowReferenceCollectorArchiveScope CollectorScope(
-						Collector.GetInternalPersistentFrameReferenceCollectorArchive(),
-						BPGC->UberGraphFunction,
-						BPGC->UberGraphFramePointerProperty,
-						InThis,
-						PointerToUberGraphFrame->RawPointer);
-					BPGC->UberGraphFunction->SerializeBin(CollectorScope.GetArchive(), PointerToUberGraphFrame->RawPointer);
+					{
+						checkSlow(BPGC->UberGraphFunction);
+						FVerySlowReferenceCollectorArchiveScope CollectorScope(
+							Collector.GetInternalPersistentFrameReferenceCollectorArchive(),
+							BPGC->UberGraphFunction,
+							BPGC->UberGraphFramePointerProperty,
+							InThis,
+							PointerToUberGraphFrame->RawPointer);
+
+#if WITH_ADDITIONAL_CRASH_CONTEXTS
+						BPGCBreadcrumbsParams Params =
+						{
+							*InThis,
+							*BPGC,
+							CollectorScope.GetArchive(),
+							FPlatformTLS::GetCurrentThreadId()
+						};
+
+						UE_ADD_CRASH_CONTEXT_SCOPE([&](FCrashContextExtendedWriter& Writer) { WriteBPGCBreadcrumbs(Writer, Params); });
+#endif // WITH_ADDITIONAL_CRASH_CONTEXTS
+
+						BPGC->UberGraphFunction->SerializeBin(CollectorScope.GetArchive(), PointerToUberGraphFrame->RawPointer);
+					}
+
 				}
 			}
 #endif // USE_UBER_GRAPH_PERSISTENT_FRAME
