Support phi nodes with loops in new HLSL translator

diff --git a/Engine/Source/Editor/MaterialEditor/Private/MaterialExpressionClasses.cpp b/Engine/Source/Editor/MaterialEditor/Private/MaterialExpressionClasses.cpp
--- a/Engine/Source/Editor/MaterialEditor/Private/MaterialExpressionClasses.cpp
+++ b/Engine/Source/Editor/MaterialEditor/Private/MaterialExpressionClasses.cpp
@@ -85,7 +85,10 @@ void MaterialExpressionClasses::InitMaterialExpressionClasses()
 	if(!bInitialized)
 	{
 		static const auto CVarMaterialEnableControlFlow = IConsoleManager::Get().FindTConsoleVariableDataInt(TEXT("r.MaterialEnableControlFlow"));
+		static const auto CVarMaterialEnableNewHLSLGenerator = IConsoleManager::Get().FindTConsoleVariableDataInt(TEXT("r.MaterialEnableNewHLSLGenerator"));
+		
 		const bool bEnableControlFlow = CVarMaterialEnableControlFlow->GetValueOnAnyThread() != 0;
+		const bool bEnableNewHLSLGenerator = CVarMaterialEnableNewHLSLGenerator->GetValueOnAnyThread() != 0;
 
 		UMaterialEditorOptions* TempEditorOptions = NewObject<UMaterialEditorOptions>();
 		UClass* BaseType = UMaterialExpression::StaticClass();
@@ -114,6 +117,11 @@ void MaterialExpressionClasses::InitMaterialExpressionClasses()
 							continue;
 						}
 
+						if (!bEnableNewHLSLGenerator && Class->HasMetaData("MaterialNewHLSLGenerator"))
+						{
+							continue;
+						}
+
 						// Exclude comments from the expression list, as well as the base parameter expression and local variable usage, as they should not be used directly
 						if ( Class != UMaterialExpressionComment::StaticClass() 
 							&& Class != UMaterialExpressionParameter::StaticClass()
diff --git a/Engine/Source/Runtime/Engine/Classes/Materials/MaterialExpression.h b/Engine/Source/Runtime/Engine/Classes/Materials/MaterialExpression.h
--- a/Engine/Source/Runtime/Engine/Classes/Materials/MaterialExpression.h
+++ b/Engine/Source/Runtime/Engine/Classes/Materials/MaterialExpression.h
@@ -19,6 +19,7 @@ struct FPropertyChangedEvent;
 struct FMaterialShadingModelField;
 
 class FMaterialHLSLGenerator;
+enum class EMaterialNewScopeFlag : uint8;
 
 namespace UE
 {
@@ -129,6 +130,7 @@ public:
 
 	/** Creates a new scope, and populates it with the expression connected to this input */
 	ENGINE_API UE::HLSLTree::FScope* NewScopeWithStatements(FMaterialHLSLGenerator& Generator, UE::HLSLTree::FScope& Scope) const;
+	ENGINE_API UE::HLSLTree::FScope* NewScopeWithStatements(FMaterialHLSLGenerator& Generator, UE::HLSLTree::FScope& Scope, EMaterialNewScopeFlag Flags) const;
 #endif // WITH_EDITOR
 private:
 	UPROPERTY()
diff --git a/Engine/Source/Runtime/Engine/Classes/Materials/MaterialExpressionBinaryOp.h b/Engine/Source/Runtime/Engine/Classes/Materials/MaterialExpressionBinaryOp.h
new file mode 100644
--- /dev/null
+++ b/Engine/Source/Runtime/Engine/Classes/Materials/MaterialExpressionBinaryOp.h
@@ -0,0 +1,50 @@
+// Copyright Epic Games, Inc. All Rights Reserved.
+
+
+#pragma once
+
+#include "CoreMinimal.h"
+#include "UObject/ObjectMacros.h"
+#include "MaterialExpressionIO.h"
+#include "Materials/MaterialExpression.h"
+#include "HLSLTree/HLSLTreeCommon.h"
+#include "MaterialExpressionBinaryOp.generated.h"
+
+UCLASS(abstract, MinimalAPI)
+class UMaterialExpressionBinaryOp : public UMaterialExpression
+{
+	GENERATED_UCLASS_BODY()
+
+	UPROPERTY(meta = (RequiredInput = "false", ToolTip = "Defaults to 'ConstA' if not specified"))
+	FExpressionInput A;
+
+	UPROPERTY(meta = (RequiredInput = "false", ToolTip = "Defaults to 'ConstB' if not specified"))
+	FExpressionInput B;
+
+	/** only used if A is not hooked up */
+	UPROPERTY(EditAnywhere, Category = MaterialExpressionAdd, meta = (OverridingInputProperty = "A"))
+	float ConstA;
+
+	/** only used if B is not hooked up */
+	UPROPERTY(EditAnywhere, Category = MaterialExpressionAdd, meta = (OverridingInputProperty = "B"))
+	float ConstB;
+
+	virtual UE::HLSLTree::EBinaryOp GetBinaryOp() const PURE_VIRTUAL(UMaterialExpressionBinaryOp::GetBinaryOp, return UE::HLSLTree::EBinaryOp::None;);
+
+	//~ Begin UMaterialExpression Interface
+#if WITH_EDITOR
+	virtual FText GetKeywords() const override;
+	virtual void GetCaption(TArray<FString>& OutCaptions) const override;
+	virtual EMaterialGenerateHLSLStatus GenerateHLSLExpression(FMaterialHLSLGenerator& Generator, UE::HLSLTree::FScope& Scope, int32 OutputIndex, UE::HLSLTree::FExpression*& OutExpression) override;
+	virtual int32 Compile(class FMaterialCompiler* Compiler, int32 OutputIndex) override;
+#endif // WITH_EDITOR
+	//~ End UMaterialExpression Interface
+
+};
+
+UCLASS(MinimalAPI, meta = (MaterialNewHLSLGenerator))
+class UMaterialExpressionLess : public UMaterialExpressionBinaryOp
+{
+	GENERATED_UCLASS_BODY()
+	virtual UE::HLSLTree::EBinaryOp GetBinaryOp() const override { return UE::HLSLTree::EBinaryOp::Less; }
+};
diff --git a/Engine/Source/Runtime/Engine/Classes/Materials/MaterialExpressionWhileLoop.h b/Engine/Source/Runtime/Engine/Classes/Materials/MaterialExpressionWhileLoop.h
new file mode 100644
--- /dev/null
+++ b/Engine/Source/Runtime/Engine/Classes/Materials/MaterialExpressionWhileLoop.h
@@ -0,0 +1,35 @@
+// Copyright Epic Games, Inc. All Rights Reserved.
+
+
+#pragma once
+
+#include "CoreMinimal.h"
+#include "UObject/ObjectMacros.h"
+#include "MaterialExpressionIO.h"
+#include "Materials/MaterialExpression.h"
+#include "MaterialExpressionWhileLoop.generated.h"
+
+UCLASS(collapsecategories, hidecategories = Object, MinimalAPI, meta = (MaterialControlFlow, MaterialNewHLSLGenerator))
+class UMaterialExpressionWhileLoop : public UMaterialExpression
+{
+	GENERATED_UCLASS_BODY()
+
+	UPROPERTY()
+	FExpressionExecOutput LoopBody;
+
+	UPROPERTY()
+	FExpressionExecOutput Completed;
+
+	UPROPERTY()
+	FExpressionInput Condition;
+
+	//~ Begin UMaterialExpression Interface
+#if WITH_EDITOR
+	virtual int32 Compile(class FMaterialCompiler* Compiler, int32 OutputIndex) override;
+	virtual void GetCaption(TArray<FString>& OutCaptions) const override;
+	virtual bool HasExecInput() override { return true; }
+
+	virtual EMaterialGenerateHLSLStatus GenerateHLSLStatements(FMaterialHLSLGenerator& Generator, UE::HLSLTree::FScope& Scope) override;
+#endif
+	//~ End UMaterialExpression Interface
+};
diff --git a/Engine/Source/Runtime/Engine/Private/HLSLTree/HLSLTree.cpp b/Engine/Source/Runtime/Engine/Private/HLSLTree/HLSLTree.cpp
--- a/Engine/Source/Runtime/Engine/Private/HLSLTree/HLSLTree.cpp
+++ b/Engine/Source/Runtime/Engine/Private/HLSLTree/HLSLTree.cpp
@@ -165,10 +165,53 @@ UE::HLSLTree::FEmitContext::~FEmitContext()
 	}
 }
 
-UE::HLSLTree::FEmitContext::FScopeEntry* UE::HLSLTree::FEmitContext::FindScopeEntry(const FScope* Scope)
+UE::HLSLTree::FEmitScope* UE::HLSLTree::FEmitContext::FindScope(const FScope& Scope)
 {
-	FEmitContext::FScopeEntry** FoundEntry = ScopeMap.Find(Scope);
-	return FoundEntry ? *FoundEntry : nullptr;
+	FEmitScope** FoundScope = ScopeMap.Find(&Scope);
+	return FoundScope ? *FoundScope : nullptr;
+}
+
+UE::HLSLTree::FEmitScope* UE::HLSLTree::FEmitContext::AcquireScope(const FScope& Scope)
+{
+	FEmitScope** FoundScope = ScopeMap.Find(&Scope);
+	FEmitScope* EmitScope = nullptr;
+	if (FoundScope)
+	{
+		EmitScope = *FoundScope;
+	}
+	else
+	{
+		EmitScope = new(*Allocator) FEmitScope();
+		EmitScope->SourceScope = &Scope;
+		if (Scope.ParentScope)
+		{
+			EmitScope->ParentScope = AcquireScope(*Scope.ParentScope);
+		}
+		ScopeMap.Add(&Scope, EmitScope);
+	}
+
+	return EmitScope;
+}
+
+void UE::HLSLTree::FEmitContext::FinalizeScope(FEmitScope& EmitScope)
+{
+	check(!EmitScope.bFinalized);
+
+	FEmitAssignment::TConstIterator It(EmitScope.FirstAssignment);
+	while (It)
+	{
+		const FEmitAssignment& EmitAssignment = *It;
+		const FEmitValue* Value = AcquireValue(EmitAssignment.Expression);
+		WriteStatementToScopef(EmitScope, TEXT("%s = %s;"), EmitAssignment.Declaration, GetCode(Value));
+		It.Next();
+	}
+
+	EmitScope.bFinalized = true;
+}
+
+UE::HLSLTree::FEmitScope& UE::HLSLTree::FEmitContext::GetCurrentScope()
+{
+	return *ScopeStack.Last();
 }
 
 void UE::HLSLTree::FExpressionEmitResult::ForwardValue(FEmitContext& Context, const FEmitValue* InValue)
@@ -233,39 +276,35 @@ const UE::HLSLTree::FEmitValue* UE::HLSLTree::FEmitContext::AcquireValue(FExpres
 				}
 				else
 				{
-					FScopeEntry* ScopeEntry = FindScopeEntry(Expression->ParentScope);
-					check(ScopeEntry);
+					FEmitScope* EmitScope = FindScope(*Expression->ParentScope);
+					check(EmitScope);
 
 					// Check to see if we've already generated code for an equivalent expression in either this scope, or any outer visible scope
 					const FSHAHash Hash = LocalWriter.GetCodeHash();
 					const TCHAR* Declaration = nullptr;
 
-					FScopeEntry* CheckScopeEntry = ScopeEntry;
-					while (CheckScopeEntry)
+					FEmitScope* CheckScope = EmitScope;
+					while (CheckScope)
 					{
-						const TCHAR** FoundDeclaration = CheckScopeEntry->ExpressionMap.Find(Hash);
+						const TCHAR** FoundDeclaration = CheckScope->ExpressionMap.Find(Hash);
 						if (FoundDeclaration)
 						{
 							// Re-use results from previous matching expression
 							Declaration = *FoundDeclaration;
 							break;
 						}
-						CheckScopeEntry = CheckScopeEntry->ParentEntry;
+						CheckScope = CheckScope->ParentScope;
 					}
 
 					if (!Declaration)
 					{
 						const Shader::FValueTypeDescription& TypeDesc = Shader::GetValueTypeDescription(EmitResult.Type);
 						Declaration = AcquireLocalDeclarationCode();
-						/*ScopeStack[ScopeIndex].ExpressionCodeWriter->WriteLinef(TEXT("const %s %s = %s;"),
-							TypeDesc.Name,
-							Declaration,
-							LocalWriter.GetStringBuilder().ToString());*/
-						AppendHLSLf(*ScopeEntry->ScopeCode, TEXT("const %s %s = %s;"),
+						WriteStatementToScopef(*EmitScope, TEXT("const %s %s = %s;"),
 							TypeDesc.Name,
 							Declaration,
 							LocalWriter.GetStringBuilder().ToString());
-						ScopeEntry->ExpressionMap.Add(Hash, Declaration);
+						EmitScope->ExpressionMap.Add(Hash, Declaration);
 					}
 					Entry->Value.Code = Declaration;
 				}
@@ -283,8 +322,8 @@ const UE::HLSLTree::FEmitValue* UE::HLSLTree::FEmitContext::AcquireValue(FFuncti
 	FFunctionCallEntry*& Entry = FunctionStack.Last().FunctionCallMap.FindOrAdd(FunctionCall);
 	if (!Entry)
 	{
-		FScopeEntry* ParentScopeEntry = FindScopeEntry(FunctionCall->ParentScope);
-		check(ParentScopeEntry);
+		FEmitScope* EmitScope = FindScope(*FunctionCall->ParentScope);
+		check(EmitScope);
 
 		{
 			FFunctionStackEntry& StackEntry = FunctionStack.AddDefaulted_GetRef();
@@ -292,10 +331,8 @@ const UE::HLSLTree::FEmitValue* UE::HLSLTree::FEmitContext::AcquireValue(FFuncti
 		}
 
 		// Link the function scope to our parent scope
-		ScopeMap.Add(FunctionCall->FunctionScope, ParentScopeEntry);
-
-		// Emit the function's HLSL into the current scope (allows sharing expressions across functions called from the same scope)
-		FunctionCall->FunctionScope->EmitUnscopedHLSL(*this, *ParentScopeEntry->ScopeCode);
+		ScopeMap.Add(FunctionCall->FunctionScope, EmitScope);
+		FunctionCall->FunctionScope->EmitHLSL(*this, *EmitScope);
 
 		// Assign function outputs
 		FEmitValue* OutputValues = new(*Allocator) FEmitValue[FunctionCall->NumOutputs];
@@ -407,101 +444,200 @@ void UE::HLSLTree::FEmitContext::AppendPreshader(const FEmitValue* Value, Shader
 	}
 }
 
-UE::HLSLTree::FEmitCode* UE::HLSLTree::FEmitContext::AppendHLSL(FEmitCode& Scope, const TCHAR* Code, int32 InLength)
+FStringView UE::HLSLTree::FEmitContext::InternalAcquireInternedString(const TCHAR* InString, int32 InLength)
 {
-	const int32 Length = InLength > 0 ? InLength : FCString::Strlen(Code);
-	TCHAR* InternedCode = new(*Allocator) TCHAR[Length + 1];
-	FMemory::Memcpy(InternedCode, Code, Length * sizeof(TCHAR));
-	InternedCode[Length] = 0;
+	TCHAR* InternedString = nullptr;
+	int32 Length = 0;
+	if (InString)
+	{
+		Length = InLength > 0 ? InLength : FCString::Strlen(InString);
+		InternedString = new(*Allocator) TCHAR[Length + 1];
+		FMemory::Memcpy(InternedString, InString, Length * sizeof(TCHAR));
+		InternedString[Length] = 0;
+	}
+	return FStringView(InternedString, Length);
+}
 
-	FEmitCode* EmitCode = new(*Allocator) FEmitCode();
-	EmitCode->CodeLength = Length;
-	EmitCode->Code = InternedCode;
-	EmitCode->ParentScope = &Scope;
+FStringView UE::HLSLTree::FEmitContext::AcquireInternedString(const TCHAR* Format, ...)
+{
+	TCHAR Buffer[1024];
+
+	va_list va;
+	va_start(va, Format);
+	const int32 Length = FPlatformString::GetVarArgs(Buffer, UE_ARRAY_COUNT(Buffer), Format, va);
+	va_end(va);
 
-	if (Scope.LastScopeCode)
+	check(Length >= 0 && Length < UE_ARRAY_COUNT(Buffer));
+
+	return InternalAcquireInternedString(Buffer, Length);
+}
+
+void UE::HLSLTree::FEmitContext::InternalWriteStatementToScope(FEmitScope& EmitScope, FStringView InternedCode)
+{
+	FEmitStatement* EmitStatement = new(*Allocator) FEmitStatement();
+	EmitStatement->Code = InternedCode;
+	if (EmitScope.LastStatement)
+	{
+		EmitStatement->LinkAfter(EmitScope.LastStatement);
+	}
+	else
 	{
-		Scope.LastScopeCode->NextCode = EmitCode;
+		EmitScope.FirstStatement = EmitStatement;
 	}
-	EmitCode->PrevCode = Scope.LastScopeCode;
-	Scope.LastScopeCode = EmitCode;
-	if (!Scope.FirstScopeCode)
+	EmitScope.LastStatement = EmitStatement;
+}
+
+bool UE::HLSLTree::FEmitContext::InternalWriteScope(const FScope& Scope, FStringView InternedCode)
+{
+	FEmitScope* EmitScope = AcquireScope(Scope);
+	if (!Scope.EmitHLSL(*this, *EmitScope))
 	{
-		Scope.FirstScopeCode = EmitCode;
+		return false;
 	}
 
-	TotalCodeLength += Length + 1; // include newline
+	FinalizeScope(*EmitScope);
 
-	return EmitCode;
+	FEmitScopeLink* Link = InternalWriteScopeLink(InternedCode);
+	Link->NextScope = EmitScope;
+	return true;
 }
 
-UE::HLSLTree::FEmitCode* UE::HLSLTree::FEmitContext::InsertHLSL(FEmitCode& Scope, FEmitCode* InsertAfter, const TCHAR* Code, int32 InLength)
+UE::HLSLTree::FEmitScopeLink* UE::HLSLTree::FEmitContext::InternalWriteScopeLink(FStringView InternedCode)
 {
-	const int32 Length = InLength > 0 ? InLength : FCString::Strlen(Code);
-	TCHAR* InternedCode = new(*Allocator) TCHAR[Length + 1];
-	FMemory::Memcpy(InternedCode, Code, Length * sizeof(TCHAR));
-	InternedCode[Length] = 0;
-
-	FEmitCode* EmitCode = new(*Allocator) FEmitCode();
-	EmitCode->CodeLength = Length;
-	EmitCode->Code = InternedCode;
-	EmitCode->ParentScope = &Scope;
-	EmitCode->PrevCode = InsertAfter;
-	if (InsertAfter)
+	FEmitScope& EmitScope = GetCurrentScope();
+	FEmitScopeLink* EmitLink = new(*Allocator) FEmitScopeLink();
+	EmitLink->Code = InternedCode;
+	if (EmitScope.LastLink)
 	{
-		check(InsertAfter->ParentScope == &Scope);
-
-		EmitCode->NextCode = InsertAfter;
-		if (InsertAfter->NextCode)
-		{
-			InsertAfter->NextCode->PrevCode = EmitCode;
-		}
-		InsertAfter->NextCode = EmitCode;
+		EmitLink->LinkAfter(EmitScope.LastLink);
 	}
 	else
 	{
-		EmitCode->NextCode = Scope.FirstScopeCode;
-		if (Scope.FirstScopeCode)
+		EmitScope.FirstLink = EmitLink;
+	}
+	EmitScope.LastLink = EmitLink;
+	return EmitLink;
+}
+
+void UE::HLSLTree::FEmitContext::WriteDeclaration(FEmitScope& EmitScope, Shader::EValueType Type, const TCHAR* Declaration, const TCHAR* Value)
+{
+	FEmitDeclaration* EmitDeclaration = new(*Allocator) FEmitDeclaration();
+	EmitDeclaration->Type = Type;
+	EmitDeclaration->Declaration = Declaration;
+	EmitDeclaration->Value = Value;
+	EmitDeclaration->LinkHead(EmitScope.FirstDeclaration);
+}
+
+bool UE::HLSLTree::FEmitContext::WriteAssignment(FEmitScope& EmitScope, const TCHAR* Declaration, FExpression* Expression, Shader::EValueType& InOutType)
+{
+	if (EmitScope.bFinalized)
+	{
+		const FEmitValue* Value = AcquireValue(Expression);
+		const Shader::EValueType ValueType = Value->GetExpressionType();
+		if (InOutType == Shader::EValueType::Void)
 		{
-			Scope.FirstScopeCode->PrevCode = EmitCode;
+			InOutType = ValueType;
+		}
+		else if (InOutType != ValueType)
+		{
+			return false;
 		}
-		Scope.FirstScopeCode = EmitCode;
-	}
 
-	if (Scope.LastScopeCode == InsertAfter)
+		WriteStatementToScopef(EmitScope, TEXT("%s = %s;"), Declaration, GetCode(Value));
+	}
+	else
 	{
-		Scope.LastScopeCode = EmitCode;
+		FEmitAssignment* EmitAssignment = new(*Allocator) FEmitAssignment();
+		EmitAssignment->Declaration = Declaration;
+		EmitAssignment->Expression = Expression;
+		EmitAssignment->LinkHead(EmitScope.FirstAssignment);
 	}
 
-	TotalCodeLength += Length + 1; // include newline
-
-	return EmitCode;
+	return true;
 }
 
-UE::HLSLTree::FEmitCode* UE::HLSLTree::FEmitContext::AppendHLSLf(FEmitCode& Scope, const TCHAR* Format, ...)
+void UE::HLSLTree::FEmitContext::Finalize()
 {
-	TCHAR Code[1024];
-
-	va_list va;
-	va_start(va, Format);
-	const int32 Length = FCString::GetVarArgs(Code, UE_ARRAY_COUNT(Code), Format, va);
-	check(Length >= 0 && Length < UE_ARRAY_COUNT(Code));
-	va_end(va);
+}
 
-	return AppendHLSL(Scope, Code, Length);
+UE::HLSLTree::FScope* UE::HLSLTree::FScope::FindSharedParent(FScope* Lhs, FScope* Rhs)
+{
+	FScope* Scope0 = Lhs;
+	FScope* Scope1 = Rhs;
+	if (Scope1)
+	{
+		while (Scope0 != Scope1)
+		{
+			if (Scope0->NestedLevel > Scope1->NestedLevel)
+			{
+				check(Scope0->ParentScope);
+				Scope0 = Scope0->ParentScope;
+			}
+			else
+			{
+				check(Scope1->ParentScope);
+				Scope1 = Scope1->ParentScope;
+			}
+		}
+	}
+	return Scope0;
 }
 
-UE::HLSLTree::FEmitCode* UE::HLSLTree::FEmitContext::InsertHLSLf(FEmitCode& Scope, FEmitCode* InsertAfter, const TCHAR* Format, ...)
+bool UE::HLSLTree::FExpressionLocalPHI::EmitCode(FEmitContext& Context, FExpressionEmitResult& OutResult) const
 {
-	TCHAR Code[1024];
+	FEmitScope* EmitScopes[MaxNumPreviousScopes] = { nullptr };
+	const FEmitValue* EmitValues[MaxNumPreviousScopes] = { nullptr };
+	Shader::EValueType CombinedValueType = Shader::EValueType::Void;
 
-	va_list va;
-	va_start(va, Format);
-	const int32 Length = FCString::GetVarArgs(Code, UE_ARRAY_COUNT(Code), Format, va);
-	check(Length >= 0 && Length < UE_ARRAY_COUNT(Code));
-	va_end(va);
+	const TCHAR* Declaration = Context.AcquireLocalDeclarationCode();
+
+	// Find the outermost scope to declare our local variable
+	FScope* DeclarationScope = ParentScope;
+	for (int32 i = 0; i < NumValues; ++i)
+	{
+		DeclarationScope = FScope::FindSharedParent(DeclarationScope, Scopes[i]);
+		if (!DeclarationScope)
+		{
+			return false;
+		}
+	}
+
+	bool bNeedToAddDeclaration = true;
+	for (int32 i = 0; i < NumValues; ++i)
+	{
+		EmitScopes[i] = Context.AcquireScope(*Scopes[i]);
+		if (Scopes[i] == DeclarationScope)
+		{
+			const FEmitValue* Value = Context.AcquireValue(Values[i]);
+			CombinedValueType = Value->GetExpressionType();
+			Context.WriteDeclaration(*EmitScopes[i], CombinedValueType, Declaration, Context.GetCode(Value));
+			bNeedToAddDeclaration = false;
+		}
+		else
+		{
+			if (!Context.WriteAssignment(*EmitScopes[i], Declaration, Values[i], CombinedValueType))
+			{
+				return false;
+			}
+		}
+	}
+
+	if (CombinedValueType == Shader::EValueType::Void)
+	{
+		return false;
+	}
+
+	if (bNeedToAddDeclaration)
+	{
+		FEmitScope* EmitDeclarationScope = Context.FindScope(*DeclarationScope);
+		Context.WriteDeclaration(*EmitDeclarationScope, CombinedValueType, Declaration);
+	}
 
-	return InsertHLSL(Scope, InsertAfter, Code, Length);
+	OutResult.EvaluationType = EExpressionEvaluationType::Shader;
+	OutResult.Type = CombinedValueType;
+	OutResult.bInline = true;
+	OutResult.Writer.Writef(TEXT("%s"), Declaration);
+	return true;
 }
 
 void UE::HLSLTree::FNodeVisitor::VisitNode(FNode* Node)
@@ -561,21 +697,16 @@ UE::HLSLTree::ENodeVisitResult UE::HLSLTree::FScope::Visit(FNodeVisitor& Visitor
 	return Result;
 }
 
-bool UE::HLSLTree::FScope::EmitHLSL(FEmitContext& Context, FEmitCode& Scope) const
+bool UE::HLSLTree::FScope::EmitHLSL(FEmitContext& Context, FEmitScope& Scope) const
 {
-	FEmitContext::FScopeEntry* ScopeEntry = new(*Context.Allocator) FEmitContext::FScopeEntry();
-	ScopeEntry->ParentEntry = Context.CurrentScopeEntry;
-	ScopeEntry->Scope = this;
-	ScopeEntry->ScopeCode = &Scope;
-
-	Context.ScopeMap.Add(this, ScopeEntry);
-	Context.CurrentScopeEntry = ScopeEntry;
-
 	const FStatement* Statement = FirstStatement;
 	bool bResult = true;
+	
+	Context.ScopeStack.Add(&Scope);
+	
 	while (Statement)
 	{
-		if (!Statement->EmitHLSL(Context, Scope))
+		if (!Statement->EmitHLSL(Context))
 		{
 			bResult = false;
 			break;
@@ -583,25 +714,15 @@ bool UE::HLSLTree::FScope::EmitHLSL(FEmitContext& Context, FEmitCode& Scope) con
 		Statement = Statement->NextStatement;
 	}
 
-	Context.CurrentScopeEntry = ScopeEntry->ParentEntry;
+	verify(Context.ScopeStack.Pop(false) == &Scope);
 
 	return bResult;
 }
 
-bool UE::HLSLTree::FScope::EmitUnscopedHLSL(FEmitContext& Context, FEmitCode& Scope) const
+void UE::HLSLTree::FScope::AddPreviousScope(FScope& Scope)
 {
-	const FStatement* Statement = FirstStatement;
-	bool bResult = true;
-	while (Statement)
-	{
-		if (!Statement->EmitHLSL(Context, Scope))
-		{
-			bResult = false;
-			break;
-		}
-		Statement = Statement->NextStatement;
-	}
-	return bResult;
+	check(NumPreviousScopes < MaxNumPreviousScopes);
+	PreviousScope[NumPreviousScopes++] = &Scope;
 }
 
 void UE::HLSLTree::FScope::AddExpression(FExpression* Expression)
@@ -612,25 +733,7 @@ void UE::HLSLTree::FScope::AddExpression(FExpression* Expression)
 
 void UE::HLSLTree::FScope::UseNode(FNode* Node)
 {
-	FScope* Scope0 = this;
-	FScope* Scope1 = Node->ParentScope;
-	if (Scope1)
-	{
-		while (Scope0 != Scope1)
-		{
-			if (Scope0->NestedLevel > Scope1->NestedLevel)
-			{
-				check(Scope0->ParentScope);
-				Scope0 = Scope0->ParentScope;
-			}
-			else
-			{
-				check(Scope1->ParentScope);
-				Scope1 = Scope1->ParentScope;
-			}
-		}
-	}
-	Node->ParentScope = Scope0;
+	Node->ParentScope = FScope::FindSharedParent(this, Node->ParentScope);
 }
 
 namespace UE
@@ -706,6 +809,22 @@ UE::HLSLTree::FTree* UE::HLSLTree::FTree::Create(FMemStackBase& Allocator)
 	return Tree;
 }
 
+void UE::HLSLTree::FTree::Destroy(FTree* Tree)
+{
+	if (Tree)
+	{
+		FNode* Node = Tree->Nodes;
+		while (Node)
+		{
+			FNode* Next = Node->NextNode;
+			Node->~FNode();
+			Node = Next;
+		}
+		Tree->~FTree();
+		FMemory::Memzero(*Tree);
+	}
+}
+
 //
 static void WriteIndent(int32 IndentLevel, FStringBuilderBase& InOutString)
 {
@@ -715,50 +834,86 @@ static void WriteIndent(int32 IndentLevel, FStringBuilderBase& InOutString)
 	}
 }
 
-static void WriteScope(const UE::HLSLTree::FEmitCode& Code, int32 IndentLevel, FStringBuilderBase& InOutString)
+static void WriteScope(const UE::HLSLTree::FEmitScope& EmitScope, int32 IndentLevel, FStringBuilderBase& InOutString)
 {
-	const UE::HLSLTree::FEmitCode* EmitCode = &Code;
-	while (EmitCode)
-	{
-		WriteIndent(IndentLevel, InOutString);
-		InOutString.Append(EmitCode->Code, EmitCode->CodeLength);
-		InOutString.Append(TCHAR('\n'));
+	check(EmitScope.bFinalized);
 
-		if (EmitCode->FirstScopeCode)
+	{
+		UE::HLSLTree::FEmitDeclaration::TConstIterator It(EmitScope.FirstDeclaration);
+		while (It)
 		{
+			const UE::HLSLTree::FEmitDeclaration& EmitDeclaration = *It;
+			const UE::Shader::FValueTypeDescription TypeDesc = UE::Shader::GetValueTypeDescription(EmitDeclaration.Type);
 			WriteIndent(IndentLevel, InOutString);
-			InOutString.Append(TEXT("{\n"));
-			WriteScope(*EmitCode->FirstScopeCode, IndentLevel + 1, InOutString);
+			if (EmitDeclaration.Value)
+			{
+				InOutString.Appendf(TEXT("%s %s = %s;\n"), TypeDesc.Name, EmitDeclaration.Declaration, EmitDeclaration.Value);
+			}
+			else
+			{
+				InOutString.Appendf(TEXT("%s %s;\n"), TypeDesc.Name, EmitDeclaration.Declaration);
+			}
+			It.Next();
+		}
+	}
+
+	{
+		UE::HLSLTree::FEmitStatement::TConstIterator It(EmitScope.FirstStatement);
+		while (It)
+		{
+			const UE::HLSLTree::FEmitStatement& EmitStatement = *It;
 			WriteIndent(IndentLevel, InOutString);
-			InOutString.Append(TEXT("}\n"));
+			InOutString.Append(EmitStatement.Code);
+			InOutString.Append(TCHAR('\n'));
+			It.Next();
+		}
+	}
+
+	{
+		UE::HLSLTree::FEmitScopeLink::TConstIterator It(EmitScope.FirstLink);
+		while (It)
+		{
+			const UE::HLSLTree::FEmitScopeLink& EmitLink = *It;
+			if (EmitLink.Code.Len() > 0)
+			{
+				WriteIndent(IndentLevel, InOutString);
+				InOutString.Append(EmitLink.Code);
+				InOutString.Append(TCHAR('\n'));
+			}
+
+			if (EmitLink.NextScope)
+			{
+				WriteIndent(IndentLevel, InOutString);
+				InOutString.Append(TEXT("{\n"));
+				WriteScope(*EmitLink.NextScope, IndentLevel + 1, InOutString);
+				WriteIndent(IndentLevel, InOutString);
+				InOutString.Append(TEXT("}\n"));
+			}
+
+			It.Next();
 		}
-		EmitCode = EmitCode->NextCode;
 	}
 }
 
 bool UE::HLSLTree::FTree::EmitHLSL(UE::HLSLTree::FEmitContext& Context, FCodeWriter& Writer) const
 {
-	FEmitCode RootCode;
-	if (RootScope->EmitHLSL(Context, RootCode))
+	FEmitScope* EmitRootScope = Context.AcquireScope(*RootScope);
+	if (RootScope->EmitHLSL(Context, *EmitRootScope))
 	{
-		WriteScope(*RootCode.FirstScopeCode, 0, *Writer.StringBuilder);
+		Context.FinalizeScope(*EmitRootScope);
+		Context.Finalize();
+		WriteScope(*EmitRootScope, 0, *Writer.StringBuilder);
 		return true;
 	}
-
 	return false;
 }
 
-UE::HLSLTree::FScope* UE::HLSLTree::FTree::NewScope(FScope& Scope, FScope** PreviousScopes, int32 NumPreviousScopes)
+UE::HLSLTree::FScope* UE::HLSLTree::FTree::NewScope(FScope& Scope)
 {
 	FScope* NewScope = NewNode<FScope>();
 	NewScope->ParentScope = &Scope;
 	NewScope->NestedLevel = Scope.NestedLevel + 1;
-	NewScope->NumPreviousScopes = NumPreviousScopes;
-	if (NumPreviousScopes > 0)
-	{
-		check(PreviousScopes);
-		FMemory::Memcpy(NewScope->PreviousScope, PreviousScopes, sizeof(FScope*) * NumPreviousScopes);
-	}
+	NewScope->NumPreviousScopes = 0;
 	return NewScope;
 }
 
diff --git a/Engine/Source/Runtime/Engine/Private/HLSLTree/HLSLTreeCommon.cpp b/Engine/Source/Runtime/Engine/Private/HLSLTree/HLSLTreeCommon.cpp
--- a/Engine/Source/Runtime/Engine/Private/HLSLTree/HLSLTreeCommon.cpp
+++ b/Engine/Source/Runtime/Engine/Private/HLSLTree/HLSLTreeCommon.cpp
@@ -4,6 +4,20 @@
 #include "MaterialShared.h"
 #include "Engine/Texture.h"
 
+UE::HLSLTree::FBinaryOpDescription UE::HLSLTree::GetBinaryOpDesription(EBinaryOp Op)
+{
+	switch (Op)
+	{
+	case EBinaryOp::None: return FBinaryOpDescription(TEXT("None"), TEXT("")); break;
+	case EBinaryOp::Add: return FBinaryOpDescription(TEXT("Add"), TEXT("+")); break;
+	case EBinaryOp::Sub: return FBinaryOpDescription(TEXT("Subtract"), TEXT("-")); break;
+	case EBinaryOp::Mul: return FBinaryOpDescription(TEXT("Multiply"), TEXT("*")); break;
+	case EBinaryOp::Div: return FBinaryOpDescription(TEXT("Divide"), TEXT("/")); break;
+	case EBinaryOp::Less: return FBinaryOpDescription(TEXT("Less"), TEXT("<")); break;
+	default: checkNoEntry(); return FBinaryOpDescription();
+	}
+}
+
 UE::HLSLTree::FSwizzleParameters::FSwizzleParameters(int8 InR, int8 InG, int8 InB, int8 InA) : NumComponents(0)
 {
 	ComponentIndex[0] = InR;
@@ -104,55 +118,6 @@ bool UE::HLSLTree::FExpressionParameter::EmitCode(FEmitContext& Context, FExpres
 	return true;
 }
 
-bool UE::HLSLTree::FExpressionLocalPHI::EmitCode(FEmitContext& Context, FExpressionEmitResult& OutResult) const
-{
-	FEmitContext::FScopeEntry* ScopeEntries[MaxNumPreviousScopes] = { nullptr };
-	const FEmitValue* EmitValues[MaxNumPreviousScopes] = { nullptr };
-	Shader::EValueType CombinedValueType = Shader::EValueType::Void;
-	for (int32 i = 0; i < NumValues; ++i)
-	{
-		const FEmitValue* Value = Context.AcquireValue(Values[i]);
-		if (Value)
-		{
-			const Shader::EValueType ValueType = Value->GetExpressionType();
-			if (CombinedValueType == Shader::EValueType::Void)
-			{
-				CombinedValueType = ValueType;
-			}
-			else if (ValueType != CombinedValueType)
-			{
-				return false;
-			}
-
-			EmitValues[i] = Value;
-		}
-		ScopeEntries[i] = Context.FindScopeEntry(Scopes[i]);
-		check(ScopeEntries[i]);
-	}
-
-	const Shader::FValueTypeDescription TypeDesc = Shader::GetValueTypeDescription(CombinedValueType);
-	const TCHAR* Declaration = Context.AcquireLocalDeclarationCode();
-
-	// Declare the local before the scopes
-	FEmitCode& ScopeCode = *ScopeEntries[0]->ScopeCode;
-	Context.InsertHLSLf(*ScopeCode.ParentScope, ScopeCode.PrevCode, TEXT("%s %s;"), TypeDesc.Name, Declaration);
-
-	// Emit code to assign the value within each scope
-	for (int32 i = 0; i < NumValues; ++i)
-	{
-		if (EmitValues[i])
-		{
-			Context.AppendHLSLf(*ScopeEntries[i]->ScopeCode, TEXT("%s = %s;"), Declaration, Context.GetCode(EmitValues[i]));
-		}
-	}
-
-	OutResult.EvaluationType = EExpressionEvaluationType::Shader;
-	OutResult.Type = CombinedValueType;
-	OutResult.bInline = true;
-	OutResult.Writer.Writef(TEXT("%s"), Declaration);
-	return true;
-}
-
 bool UE::HLSLTree::FExpressionExternalInput::EmitCode(FEmitContext& Context, FExpressionEmitResult& OutResult) const
 {
 	const int32 TypeIndex = (int32)InputType;
@@ -382,24 +347,45 @@ bool UE::HLSLTree::FExpressionBinaryOp::EmitCode(FEmitContext& Context, FExpress
 	}
 
 	FString ErrorMessage;
-
-	OutResult.EvaluationType = CombineEvaluationTypes(LhsValue->GetEvaluationType(), RhsValue->GetEvaluationType());
-	OutResult.Type = MakeArithmeticResultType(LhsValue->GetExpressionType(), RhsValue->GetExpressionType(), ErrorMessage);
-	if (OutResult.EvaluationType == EExpressionEvaluationType::Shader)
+	switch (Op)
 	{
-		OutResult.Writer.Writef(TEXT("(%s + %s)"), Context.GetCode(LhsValue), Context.GetCode(RhsValue));
+	case EBinaryOp::Less:
+		OutResult.Type = Shader::MakeComparisonResultType(LhsValue->GetExpressionType(), RhsValue->GetExpressionType(), ErrorMessage);
+		break;
+	default:
+		OutResult.Type = Shader::MakeArithmeticResultType(LhsValue->GetExpressionType(), RhsValue->GetExpressionType(), ErrorMessage);
+		break;
 	}
-	else
+
+	const EExpressionEvaluationType EvaluationType = CombineEvaluationTypes(LhsValue->GetEvaluationType(), RhsValue->GetEvaluationType());
+	Shader::EPreshaderOpcode Opcode = Shader::EPreshaderOpcode::Nop;
+	if (EvaluationType != EExpressionEvaluationType::Shader)
 	{
-		Shader::EPreshaderOpcode Opcode = Shader::EPreshaderOpcode::Nop;
 		switch (Op)
 		{
 		case UE::HLSLTree::EBinaryOp::Add: Opcode = Shader::EPreshaderOpcode::Add; break;
-		default: checkNoEntry(); break;
+		default: break;
 		}
+	}
 
+	if (Opcode == Shader::EPreshaderOpcode::Nop)
+	{
+		// Either preshader not supported for our inputs, or not for this operation type
+		OutResult.EvaluationType = EExpressionEvaluationType::Shader;
+		const TCHAR* LhsCode = Context.GetCode(LhsValue);
+		const TCHAR* RhsCode = Context.GetCode(RhsValue);
+		switch (Op)
+		{
+		case EBinaryOp::Add: OutResult.Writer.Writef(TEXT("(%s + %s)"), LhsCode, RhsCode); break;
+		case EBinaryOp::Less: OutResult.Writer.Writef(TEXT("(%s < %s)"), LhsCode, RhsCode); break;
+		default: checkNoEntry(); break;
+		}
+	}
+	else
+	{
 		Context.AppendPreshader(LhsValue, OutResult.Preshader);
 		Context.AppendPreshader(RhsValue, OutResult.Preshader);
+		OutResult.EvaluationType = EvaluationType;
 		OutResult.Preshader.WriteOpcode(Opcode);
 	}
 	return true;
@@ -542,12 +528,18 @@ bool UE::HLSLTree::FExpressionFunctionOutput::EmitCode(FEmitContext& Context, FE
 	return true;
 }
 
-bool UE::HLSLTree::FStatementJump::EmitHLSL(FEmitContext& Context, FEmitCode& Scope) const
+bool UE::HLSLTree::FStatementNestedScope::EmitHLSL(FEmitContext& Context) const
 {
-	return TargetScope->EmitHLSL(Context, Scope);
+	return Context.WriteScope(*NextScope);
+}
+
+bool UE::HLSLTree::FStatementBreak::EmitHLSL(FEmitContext& Context) const
+{
+	Context.WriteScopeTerminatorf(TEXT("break;"));
+	return true;
 }
 
-bool UE::HLSLTree::FStatementReturn::EmitHLSL(FEmitContext& Context, FEmitCode& Scope) const
+bool UE::HLSLTree::FStatementReturn::EmitHLSL(FEmitContext& Context) const
 {
 	const FEmitValue* Value = Context.AcquireValue(Expression);
 	if (!Value)
@@ -555,11 +547,11 @@ bool UE::HLSLTree::FStatementReturn::EmitHLSL(FEmitContext& Context, FEmitCode&
 		return false;
 	}
 
-	Context.AppendHLSLf(Scope, TEXT("return %s;"), Context.GetCode(Value));
+	Context.WriteScopeTerminatorf(TEXT("return %s;"), Context.GetCode(Value));
 	return true;
 }
 
-bool UE::HLSLTree::FStatementIf::EmitHLSL(FEmitContext& Context, FEmitCode& Scope) const
+bool UE::HLSLTree::FStatementIf::EmitHLSL(FEmitContext& Context) const
 {
 	const FEmitValue* ConditionValue = Context.AcquireValue(ConditionExpression);
 	if (!ConditionValue)
@@ -567,19 +559,27 @@ bool UE::HLSLTree::FStatementIf::EmitHLSL(FEmitContext& Context, FEmitCode& Scop
 		return false;
 	}
 
-	FEmitCode* ThenScopeCode = Context.AppendHLSLf(Scope, TEXT("if (%s)"), Context.GetCode(ConditionValue));
-	ThenScope->EmitHLSL(Context, *ThenScopeCode);
+	Context.WriteScopef(*ThenScope, TEXT("if (%s)"), Context.GetCode(ConditionValue));
 	if (ElseScope)
 	{
-		FEmitCode* ElseScopeCode = Context.AppendHLSLf(Scope, TEXT("else"));
-		ElseScope->EmitHLSL(Context, *ElseScopeCode);
+		Context.WriteScopef(*ElseScope, TEXT("else"));
+	}
+	if (NextScope)
+	{
+		Context.WriteScope(*NextScope);
 	}
 	return true;
 }
 
-bool UE::HLSLTree::FStatementFor::EmitHLSL(FEmitContext& Context, FEmitCode& Scope) const
+bool UE::HLSLTree::FStatementLoop::EmitHLSL(FEmitContext& Context) const
+{
+	Context.WriteScopef(*LoopScope, TEXT("while (true)"));
+	return Context.WriteScope(*NextScope);
+}
+
+bool UE::HLSLTree::FStatementFor::EmitHLSL(FEmitContext& Context) const
 {
-	const FEmitValue* StartExpressionValue = Context.AcquireValue(StartExpression);
+	/*const FEmitValue* StartExpressionValue = Context.AcquireValue(StartExpression);
 	const FEmitValue* EndExpressionValue = Context.AcquireValue(StartExpression);
 	if (!!StartExpressionValue || !EndExpressionValue)
 	{
@@ -595,5 +595,6 @@ bool UE::HLSLTree::FStatementFor::EmitHLSL(FEmitContext& Context, FEmitCode& Sco
 		Context.GetCode(EndExpressionValue),
 		DeclarationCode);
 	LoopScope->EmitHLSL(Context, *LoopScopeCode);
-	return true;
+	return NextScope->EmitHLSL(Context, Scope);*/
+	return false;
 }
diff --git a/Engine/Source/Runtime/Engine/Private/Materials/MaterialExpressionHLSL.cpp b/Engine/Source/Runtime/Engine/Private/Materials/MaterialExpressionHLSL.cpp
--- a/Engine/Source/Runtime/Engine/Private/Materials/MaterialExpressionHLSL.cpp
+++ b/Engine/Source/Runtime/Engine/Private/Materials/MaterialExpressionHLSL.cpp
@@ -28,11 +28,13 @@
 #include "Materials/MaterialExpressionStaticSwitch.h"
 #include "Materials/MaterialExpressionGetLocal.h"
 #include "Materials/MaterialExpressionAdd.h"
+#include "Materials/MaterialExpressionBinaryOp.h"
 #include "Materials/MaterialExpressionSetMaterialAttributes.h"
 #include "Materials/MaterialExpressionReflectionVectorWS.h"
 #include "Materials/MaterialExpressionSetLocal.h"
 #include "Materials/MaterialExpressionIfThenElse.h"
 #include "Materials/MaterialExpressionForLoop.h"
+#include "Materials/MaterialExpressionWhileLoop.h"
 #include "Materials/MaterialFunctionInterface.h"
 
 EMaterialGenerateHLSLStatus UMaterialExpression::GenerateHLSLExpression(FMaterialHLSLGenerator& Generator, UE::HLSLTree::FScope& Scope, int32 OutputIndex, UE::HLSLTree::FExpression*& OutExpression)
@@ -192,6 +194,19 @@ EMaterialGenerateHLSLStatus UMaterialExpressionTextureSampleParameter::GenerateH
 	return GenerateHLSLExpressionBase(Generator, Scope, TextureDeclaration, OutExpression);
 }
 
+EMaterialGenerateHLSLStatus UMaterialExpressionBinaryOp::GenerateHLSLExpression(FMaterialHLSLGenerator& Generator, UE::HLSLTree::FScope& Scope, int32 OutputIndex, UE::HLSLTree::FExpression*& OutExpression)
+{
+	UE::HLSLTree::FExpression* Lhs = A.GetTracedInput().Expression ? A.AcquireHLSLExpression(Generator, Scope) : Generator.NewConstant(Scope, ConstA);
+	UE::HLSLTree::FExpression* Rhs = B.GetTracedInput().Expression ? B.AcquireHLSLExpression(Generator, Scope) : Generator.NewConstant(Scope, ConstB);
+	if (!Lhs || !Rhs)
+	{
+		return EMaterialGenerateHLSLStatus::Error;
+	}
+
+	OutExpression = Generator.GetTree().NewExpression<UE::HLSLTree::FExpressionBinaryOp>(Scope, GetBinaryOp(), Lhs, Rhs);
+	return EMaterialGenerateHLSLStatus::Success;
+}
+
 EMaterialGenerateHLSLStatus UMaterialExpressionAdd::GenerateHLSLExpression(FMaterialHLSLGenerator& Generator, UE::HLSLTree::FScope& Scope, int32 OutputIndex, UE::HLSLTree::FExpression*& OutExpression)
 {
 	UE::HLSLTree::FExpression* Lhs = A.GetTracedInput().Expression ? A.AcquireHLSLExpression(Generator, Scope) : Generator.NewConstant(Scope, ConstA);
@@ -324,20 +339,46 @@ EMaterialGenerateHLSLStatus UMaterialExpressionIfThenElse::GenerateHLSLStatement
 		return Generator.Error(TEXT("Missing condition connection"));
 	}
 
-	UE::HLSLTree::FScope* ThenScope = Then.NewScopeWithStatements(Generator, Scope);
-	if (!ThenScope)
-	{
-		return Generator.Error(TEXT("Missing Then connection"));
-	}
-
 	UE::HLSLTree::FStatementIf* IfStatement = Generator.GetTree().NewStatement<UE::HLSLTree::FStatementIf>(Scope);
 	IfStatement->ConditionExpression = ConditionExpression;
-	IfStatement->ThenScope = ThenScope;
+	IfStatement->NextScope = Generator.NewJoinedScope(Scope);
+	IfStatement->ThenScope = Then.NewScopeWithStatements(Generator, Scope);
 	IfStatement->ElseScope = Else.NewScopeWithStatements(Generator, Scope);
 
 	return EMaterialGenerateHLSLStatus::Success;
 }
 
+EMaterialGenerateHLSLStatus UMaterialExpressionWhileLoop::GenerateHLSLStatements(FMaterialHLSLGenerator& Generator, UE::HLSLTree::FScope& Scope)
+{
+	if (!Condition.IsConnected())
+	{
+		return Generator.Error(TEXT("Missing condition connection"));
+	}
+
+	if (!LoopBody.GetExpression())
+	{
+		return Generator.Error(TEXT("Missing LoopBody connection"));
+	}
+
+	UE::HLSLTree::FStatementLoop* LoopStatement = Generator.GetTree().NewStatement<UE::HLSLTree::FStatementLoop>(Scope);
+	LoopStatement->LoopScope = Generator.NewScope(Scope);
+
+	UE::HLSLTree::FStatementIf* IfStatement = Generator.GetTree().NewStatement<UE::HLSLTree::FStatementIf>(*LoopStatement->LoopScope);
+	IfStatement->ThenScope = Generator.NewScope(*LoopStatement->LoopScope);
+	IfStatement->ElseScope = Generator.NewScope(*LoopStatement->LoopScope);
+	LoopStatement->NextScope = Generator.NewScope(Scope, EMaterialNewScopeFlag::NoPreviousScope);
+	LoopStatement->LoopScope->AddPreviousScope(*IfStatement->ThenScope);
+	LoopStatement->NextScope->AddPreviousScope(*IfStatement->ElseScope);
+
+	Generator.GetTree().NewStatement<UE::HLSLTree::FStatementBreak>(*IfStatement->ElseScope);
+
+	IfStatement->ConditionExpression = Condition.AcquireHLSLExpression(Generator, *LoopStatement->LoopScope);
+	LoopBody.GenerateHLSLStatements(Generator, *IfStatement->ThenScope);
+	Completed.GenerateHLSLStatements(Generator, *LoopStatement->NextScope);
+
+	return EMaterialGenerateHLSLStatus::Success;
+}
+
 EMaterialGenerateHLSLStatus UMaterialExpressionForLoop::GenerateHLSLExpression(FMaterialHLSLGenerator& Generator, UE::HLSLTree::FScope& Scope, int32 OutputIndex, UE::HLSLTree::FExpression*& OutExpression)
 {
 	return EMaterialGenerateHLSLStatus::Success;
@@ -345,6 +386,7 @@ EMaterialGenerateHLSLStatus UMaterialExpressionForLoop::GenerateHLSLExpression(F
 
 EMaterialGenerateHLSLStatus UMaterialExpressionForLoop::GenerateHLSLStatements(FMaterialHLSLGenerator& Generator, UE::HLSLTree::FScope& Scope)
 {
+#if 0
 	if (!LoopBody.GetExpression())
 	{
 		return Generator.Error(TEXT("Missing LoopBody connection"));
@@ -371,7 +413,7 @@ EMaterialGenerateHLSLStatus UMaterialExpressionForLoop::GenerateHLSLStatements(F
 	ForStatement->LoopScope = LoopScope;
 
 	Completed.GenerateHLSLStatements(Generator, Scope);
-
+#endif
 	return EMaterialGenerateHLSLStatus::Success;
 }
 
diff --git a/Engine/Source/Runtime/Engine/Private/Materials/MaterialExpressions.cpp b/Engine/Source/Runtime/Engine/Private/Materials/MaterialExpressions.cpp
--- a/Engine/Source/Runtime/Engine/Private/Materials/MaterialExpressions.cpp
+++ b/Engine/Source/Runtime/Engine/Private/Materials/MaterialExpressions.cpp
@@ -257,8 +257,10 @@
 #include "Materials/MaterialExpressionExecEnd.h"
 #include "Materials/MaterialExpressionIfThenElse.h"
 #include "Materials/MaterialExpressionForLoop.h"
+#include "Materials/MaterialExpressionWhileLoop.h"
 #include "Materials/MaterialExpressionSetLocal.h"
 #include "Materials/MaterialExpressionGetLocal.h"
+#include "Materials/MaterialExpressionBinaryOp.h"
 #include "EditorSupportDelegates.h"
 #include "MaterialCompiler.h"
 #if WITH_EDITOR
@@ -21553,7 +21555,88 @@ uint32 UMaterialExpressionForLoop::GetOutputType(int32 OutputIndex)
 	default: checkNoEntry(); return 0u;
 	}
 }
+#endif // WITH_EDITOR
+
+UMaterialExpressionWhileLoop::UMaterialExpressionWhileLoop(const FObjectInitializer& ObjectInitializer)
+	: Super(ObjectInitializer)
+{
+	struct FConstructorStatics
+	{
+		FText NAME_Execution;
+		FConstructorStatics() : NAME_Execution(LOCTEXT("Execution", "Execution")) { }
+	};
+	static FConstructorStatics ConstructorStatics;
+#if WITH_EDITORONLY_DATA
+	MenuCategories.Add(ConstructorStatics.NAME_Execution);
+	Outputs.Reset();
+	bShowOutputNameOnPin = true;
+	bHidePreviewWindow = true;
+#endif
+}
+
+#if WITH_EDITOR
+int32 UMaterialExpressionWhileLoop::Compile(class FMaterialCompiler* Compiler, int32 OutputIndex)
+{
+	return Compiler->Error(TEXT("Not supported with legacy compiler"));
+}
+
+void UMaterialExpressionWhileLoop::GetCaption(TArray<FString>& OutCaptions) const
+{
+	OutCaptions.Add(TEXT("While"));
+}
+
+#endif // WITH_EDITOR
+
+UMaterialExpressionBinaryOp::UMaterialExpressionBinaryOp(const FObjectInitializer& ObjectInitializer)
+{
+	// Structure to hold one-time initialization
+	struct FConstructorStatics
+	{
+		FText NAME_Math;
+		FConstructorStatics() : NAME_Math(LOCTEXT("Math", "Math"))
+		{
+		}
+	};
+	static FConstructorStatics ConstructorStatics;
+
+	ConstA = 0.0f;
+	ConstB = 1.0f;
+
+#if WITH_EDITORONLY_DATA
+	MenuCategories.Add(ConstructorStatics.NAME_Math);
 #endif
+}
+
+UMaterialExpressionLess::UMaterialExpressionLess(const FObjectInitializer& ObjectInitializer) {}
+
+#if WITH_EDITOR
+FText UMaterialExpressionBinaryOp::GetKeywords() const
+{
+	const UE::HLSLTree::FBinaryOpDescription Description = UE::HLSLTree::GetBinaryOpDesription(GetBinaryOp());
+	return FText::FromString(Description.Operator);
+}
+
+void UMaterialExpressionBinaryOp::GetCaption(TArray<FString>& OutCaptions) const
+{
+	const UE::HLSLTree::FBinaryOpDescription Description = UE::HLSLTree::GetBinaryOpDesription(GetBinaryOp());
+	FString ret = Description.Name;
+	FExpressionInput ATraced = A.GetTracedInput();
+	FExpressionInput BTraced = B.GetTracedInput();
+	if (!ATraced.Expression || !BTraced.Expression)
+	{
+		ret += TEXT("(");
+		ret += ATraced.Expression ? TEXT(",") : FString::Printf(TEXT("%.4g,"), ConstA);
+		ret += BTraced.Expression ? TEXT(")") : FString::Printf(TEXT("%.4g)"), ConstB);
+	}
+
+	OutCaptions.Add(ret);
+}
+
+int32 UMaterialExpressionBinaryOp::Compile(class FMaterialCompiler* Compiler, int32 OutputIndex)
+{
+	return Compiler->Error(TEXT("Not supported with legacy compiler"));
+}
+#endif // WITH_EDITOR
 
 UMaterialExpressionGetLocal::UMaterialExpressionGetLocal(const FObjectInitializer& ObjectInitializer)
 	: Super(ObjectInitializer)
diff --git a/Engine/Source/Runtime/Engine/Private/Materials/MaterialHLSLGenerator.cpp b/Engine/Source/Runtime/Engine/Private/Materials/MaterialHLSLGenerator.cpp
--- a/Engine/Source/Runtime/Engine/Private/Materials/MaterialHLSLGenerator.cpp
+++ b/Engine/Source/Runtime/Engine/Private/Materials/MaterialHLSLGenerator.cpp
@@ -59,6 +59,21 @@ bool FMaterialHLSLGenerator::Finalize()
 		}
 	}
 
+	if (JoinedScopeStack.Num() != 0)
+	{
+		Error(TEXT("Invalid control flow"));
+		return false;
+	}
+
+	// Resolve values for any PHI nodes that were generated
+	for (UE::HLSLTree::FExpressionLocalPHI* Expression : PHIExpressions)
+	{
+		for (int32 i = 0; i < Expression->NumValues; ++i)
+		{
+			Expression->Values[i] = InternalAcquireLocalValue(*Expression->Scopes[i], Expression->LocalName);
+		}
+	}
+
 	return true;
 }
 
@@ -204,6 +219,24 @@ bool FMaterialHLSLGenerator::GenerateResult(UE::HLSLTree::FScope& Scope)
 	return bResult;
 }
 
+UE::HLSLTree::FScope* FMaterialHLSLGenerator::NewScope(UE::HLSLTree::FScope& Scope, EMaterialNewScopeFlag Flags)
+{
+	UE::HLSLTree::FScope* NewScope = HLSLTree->NewScope(Scope);
+	if (!EnumHasAllFlags(Flags, EMaterialNewScopeFlag::NoPreviousScope))
+	{
+		NewScope->AddPreviousScope(Scope);
+	}
+
+	return NewScope;
+}
+
+UE::HLSLTree::FScope* FMaterialHLSLGenerator::NewJoinedScope(UE::HLSLTree::FScope& Scope)
+{
+	UE::HLSLTree::FScope* NewScope = HLSLTree->NewScope(Scope);
+	JoinedScopeStack.Add(NewScope);
+	return NewScope;
+}
+
 UE::HLSLTree::FExpressionConstant* FMaterialHLSLGenerator::NewConstant(UE::HLSLTree::FScope& Scope, const UE::Shader::FValue& Value)
 {
 	UE::HLSLTree::FExpressionConstant* Expression = HLSLTree->NewExpression<UE::HLSLTree::FExpressionConstant>(Scope, Value);
@@ -336,7 +369,7 @@ bool FMaterialHLSLGenerator::GenerateAssignLocal(UE::HLSLTree::FScope& Scope, co
 	return true;
 }
 
-UE::HLSLTree::FExpression* FMaterialHLSLGenerator::AcquireLocalValue(UE::HLSLTree::FScope& Scope, const FName& LocalName)
+UE::HLSLTree::FExpression* FMaterialHLSLGenerator::InternalAcquireLocalValue(UE::HLSLTree::FScope& Scope, const FName& LocalName)
 {
 	const FLocalKey Key(&Scope, LocalName);
 	UE::HLSLTree::FExpression** FoundExpression = LocalMap.Find(Key);
@@ -349,23 +382,30 @@ UE::HLSLTree::FExpression* FMaterialHLSLGenerator::AcquireLocalValue(UE::HLSLTre
 	if (PreviousScopes.Num() > 1)
 	{
 		UE::HLSLTree::FExpressionLocalPHI* Expression = HLSLTree->NewExpression<UE::HLSLTree::FExpressionLocalPHI>(Scope);
+		Expression->LocalName = LocalName;
 		Expression->NumValues = PreviousScopes.Num();
 		for (int32 i = 0; i < PreviousScopes.Num(); ++i)
 		{
 			Expression->Scopes[i] = PreviousScopes[i];
-			Expression->Values[i] = AcquireLocalValue(*PreviousScopes[i], LocalName);
 		}
+		PHIExpressions.Add(Expression);
 		LocalMap.Add(Key, Expression);
 		return Expression;
 	}
 
-	if (Scope.ParentScope)
+	if (PreviousScopes.Num() == 1)
 	{
-		return AcquireLocalValue(*Scope.ParentScope, LocalName);
+		return InternalAcquireLocalValue(*PreviousScopes[0], LocalName);
 	}
+
 	return nullptr;
 }
 
+UE::HLSLTree::FExpression* FMaterialHLSLGenerator::AcquireLocalValue(UE::HLSLTree::FScope& Scope, const FName& LocalName)
+{
+	return InternalAcquireLocalValue(Scope, LocalName);
+}
+
 UE::HLSLTree::FExpression* FMaterialHLSLGenerator::AcquireExpression(UE::HLSLTree::FScope& Scope, UMaterialExpression* MaterialExpression, int32 OutputIndex)
 {
 	const FExpressionKey Key(MaterialExpression, OutputIndex);
@@ -397,6 +437,7 @@ UE::HLSLTree::FTextureParameterDeclaration* FMaterialHLSLGenerator::AcquireTextu
 bool FMaterialHLSLGenerator::GenerateStatements(UE::HLSLTree::FScope& Scope, UMaterialExpression* MaterialExpression)
 {
 	FStatementEntry& Entry = StatementMap.FindOrAdd(MaterialExpression);
+
 	Entry.PreviousScope[Entry.NumInputs++] = &Scope;
 	check(Entry.NumInputs <= MaxNumPreviousScopes);
 	check(Entry.NumInputs <= MaterialExpression->NumExecutionInputs);
@@ -407,10 +448,17 @@ bool FMaterialHLSLGenerator::GenerateStatements(UE::HLSLTree::FScope& Scope, UMa
 		UE::HLSLTree::FScope* ScopeToUse = &Scope;
 		if (MaterialExpression->NumExecutionInputs > 1u)
 		{
-			ScopeToUse = HLSLTree->NewScope(*ScopeToUse->ParentScope, Entry.PreviousScope, Entry.NumInputs);
-			// Add a jump to link the previous scope to this new scope
-			UE::HLSLTree::FStatementJump* JumpStatement = HLSLTree->NewStatement<UE::HLSLTree::FStatementJump>(*Scope.ParentScope);
-			JumpStatement->TargetScope = ScopeToUse;
+			if (JoinedScopeStack.Num() == 0)
+			{
+				Error(TEXT("Bad control flow"));
+				return false;
+			}
+
+			ScopeToUse = JoinedScopeStack.Pop(false);
+			for (int32 i = 0; i < Entry.NumInputs; ++i)
+			{
+				ScopeToUse->AddPreviousScope(*Entry.PreviousScope[i]);
+			}
 		}
 
 		const FExpressionKey Key(MaterialExpression);
diff --git a/Engine/Source/Runtime/Engine/Private/Materials/MaterialHLSLTree.cpp b/Engine/Source/Runtime/Engine/Private/Materials/MaterialHLSLTree.cpp
--- a/Engine/Source/Runtime/Engine/Private/Materials/MaterialHLSLTree.cpp
+++ b/Engine/Source/Runtime/Engine/Private/Materials/MaterialHLSLTree.cpp
@@ -16,6 +16,12 @@ FMaterialHLSLTree::FMaterialHLSLTree()
 {
 }
 
+FMaterialHLSLTree::~FMaterialHLSLTree()
+{
+	UE::HLSLTree::FTree::Destroy(HLSLTree);
+}
+
+
 bool FMaterialHLSLTree::InitializeForMaterial(const FMaterialCompileTargetParameters& InCompilerTarget, FMaterial& InOutMaterial)
 {
 	UMaterialInterface* MaterialInterface = InOutMaterial.GetMaterialInterface();
diff --git a/Engine/Source/Runtime/Engine/Private/Materials/MaterialShared.cpp b/Engine/Source/Runtime/Engine/Private/Materials/MaterialShared.cpp
--- a/Engine/Source/Runtime/Engine/Private/Materials/MaterialShared.cpp
+++ b/Engine/Source/Runtime/Engine/Private/Materials/MaterialShared.cpp
@@ -357,12 +357,17 @@ bool FExpressionExecOutput::GenerateHLSLStatements(FMaterialHLSLGenerator& Gener
 }
 
 UE::HLSLTree::FScope* FExpressionExecOutput::NewScopeWithStatements(FMaterialHLSLGenerator& Generator, UE::HLSLTree::FScope& Scope) const
+{
+	return NewScopeWithStatements(Generator, Scope, EMaterialNewScopeFlag::None);
+}
+
+UE::HLSLTree::FScope* FExpressionExecOutput::NewScopeWithStatements(FMaterialHLSLGenerator& Generator, UE::HLSLTree::FScope& Scope, EMaterialNewScopeFlag Flags) const
 {
 	UE::HLSLTree::FScope* Result = nullptr;
 	if (Expression)
 	{
 		Expression->ValidateState();
-		Result = Generator.GetTree().NewScope(Scope); // Create a new scope for the statements
+		Result = Generator.NewScope(Scope, Flags); // Create a new scope for the statements
 		Generator.GenerateStatements(*Result, Expression);
 	}
 
diff --git a/Engine/Source/Runtime/Engine/Private/Shader/ShaderValue.cpp b/Engine/Source/Runtime/Engine/Private/Shader/ShaderValue.cpp
--- a/Engine/Source/Runtime/Engine/Private/Shader/ShaderValue.cpp
+++ b/Engine/Source/Runtime/Engine/Private/Shader/ShaderValue.cpp
@@ -225,6 +225,30 @@ UE::Shader::EValueType UE::Shader::MakeArithmeticResultType(EValueType Lhs, EVal
 	return EValueType::Void;
 }
 
+UE::Shader::EValueType UE::Shader::MakeComparisonResultType(EValueType Lhs, EValueType Rhs, FString& OutErrorMessage)
+{
+	const FValueTypeDescription LhsDesc = GetValueTypeDescription(Lhs);
+	const FValueTypeDescription RhsDesc = GetValueTypeDescription(Rhs);
+
+	if (Lhs == Rhs)
+	{
+		if (LhsDesc.NumComponents > 0)
+		{
+			return MakeValueType(EValueComponentType::Bool, LhsDesc.NumComponents);
+		}
+		else
+		{
+			OutErrorMessage = FString::Printf(TEXT("Attempting to perform comparison on non-numeric types: %s %s"), LhsDesc.Name, RhsDesc.Name);
+		}
+	}
+	else
+	{
+		OutErrorMessage = FString::Printf(TEXT("Comparison between types %s and %s are undefined"), LhsDesc.Name, RhsDesc.Name);
+	}
+
+	return EValueType::Void;
+}
+
 namespace UE
 {
 namespace Shader
diff --git a/Engine/Source/Runtime/Engine/Public/HLSLTree/HLSLTree.h b/Engine/Source/Runtime/Engine/Public/HLSLTree/HLSLTree.h
--- a/Engine/Source/Runtime/Engine/Public/HLSLTree/HLSLTree.h
+++ b/Engine/Source/Runtime/Engine/Public/HLSLTree/HLSLTree.h
@@ -4,6 +4,8 @@
 #include "CoreMinimal.h"
 #include "Engine/EngineTypes.h"
 #include "Containers/ArrayView.h"
+#include "Containers/StringView.h"
+#include "Containers/List.h"
 #include "Misc/StringBuilder.h"
 #include "Misc/MemStack.h"
 #include "HLSLTree/HLSLTreeTypes.h"
@@ -28,6 +30,8 @@ class FPreshaderData;
 namespace HLSLTree
 {
 
+class FExpressionLocalPHI;
+
 static constexpr int32 MaxNumPreviousScopes = 2;
 
 /** Allows building a string incrementally, with some additional features to support code generation, such as indent handling */
@@ -88,18 +92,50 @@ public:
 	int32 IndentLevel;
 };
 
-class FEmitCode
+class FScope;
+class FEmitScope;
+
+class FEmitDeclaration : public TIntrusiveLinkedList<FEmitDeclaration>
 {
 public:
-	const TCHAR* Code = nullptr;
-	FEmitCode* ParentScope = nullptr;
-	FEmitCode* NextCode = nullptr;
-	FEmitCode* PrevCode = nullptr;
-	FEmitCode* FirstScopeCode = nullptr;
-	FEmitCode* LastScopeCode = nullptr;
-	int32 CodeLength = 0;
+	const TCHAR* Declaration;
+	const TCHAR* Value;
+	Shader::EValueType Type;
+};
 
-	friend class FEmitContext;
+class FEmitAssignment : public TIntrusiveLinkedList<FEmitAssignment>
+{
+public:
+	const TCHAR* Declaration;
+	FExpression* Expression;
+};
+
+class FEmitStatement : public TIntrusiveLinkedList<FEmitStatement>
+{
+public:
+	FStringView Code;
+};
+
+class FEmitScopeLink : public TIntrusiveLinkedList<FEmitScopeLink>
+{
+public:
+	FEmitScope* NextScope;
+	FStringView Code;
+};
+
+class FEmitScope
+{
+public:
+	FEmitScope* ParentScope = nullptr;
+	FEmitDeclaration* FirstDeclaration = nullptr;
+	FEmitAssignment* FirstAssignment = nullptr;
+	FEmitStatement* FirstStatement = nullptr;
+	FEmitStatement* LastStatement = nullptr;
+	FEmitScopeLink* FirstLink = nullptr;
+	FEmitScopeLink* LastLink = nullptr;
+	const FScope* SourceScope = nullptr;
+	TMap<FSHAHash, const TCHAR*> ExpressionMap;
+	bool bFinalized = false;
 };
 
 class FEmitValue
@@ -135,25 +171,78 @@ public:
 	/** Get a unique local variable name */
 	const TCHAR* AcquireLocalDeclarationCode();
 
+	FStringView InternalAcquireInternedString(const TCHAR* String, int32 Length);
+
+	FStringView AcquireInternedString(const TCHAR* Format, ...);
+
+	template<int32 ArrayLength>
+	FStringView AcquireInternedString(TCHAR (&String)[ArrayLength])
+	{
+		return InternalAcquireInternedString(String, ArrayLength);
+	}
+
 	/** Gets HLSL code that references the given value */
 	const TCHAR* GetCode(const FEmitValue* Value) const;
 
 	/** Append preshader bytecode that represents the given value */
 	void AppendPreshader(const FEmitValue* Value, Shader::FPreshaderData& InOutPreshader) const;
 
-	FEmitCode* AppendHLSL(FEmitCode& Scope, const TCHAR* Code, int32 Length = 0);
-	FEmitCode* InsertHLSL(FEmitCode& Scope, FEmitCode* InsertAfter, const TCHAR* Code, int32 Length = 0);
+	FEmitScope* FindScope(const FScope& Scope);
+	FEmitScope* AcquireScope(const FScope& Scope);
+	void FinalizeScope(FEmitScope& EmitScope);
 
-	FEmitCode* AppendHLSLf(FEmitCode& Scope, const TCHAR* Format, ...);
-	FEmitCode* InsertHLSLf(FEmitCode& Scope, FEmitCode* InsertAfter, const TCHAR* Format, ...);
+	FEmitScope& GetCurrentScope();
 
-	struct FScopeEntry
+	void InternalWriteStatementToScope(FEmitScope& EmitScope, FStringView InternedCode);
+	bool InternalWriteScope(const FScope& Scope, FStringView InternedCode);
+	FEmitScopeLink* InternalWriteScopeLink(FStringView InternedCode);
+
+	template<typename StringType>
+	void WriteStatementToScope(FEmitScope& EmitScope, const StringType& String)
 	{
-		FScopeEntry* ParentEntry;
-		const FScope* Scope;
-		FEmitCode* ScopeCode;
-		TMap<FSHAHash, const TCHAR*> ExpressionMap;
-	};
+		InternalWriteStatementToScope(EmitScope, AcquireInternedString(String));
+	}
+
+	template<typename FormatType, typename... Types>
+	void WriteStatementToScopef(FEmitScope& EmitScope, const FormatType& Format, Types... Args)
+	{
+		InternalWriteStatementToScope(EmitScope, AcquireInternedString(Format, Forward<Types>(Args)...));
+	}
+
+	template<typename StringType>
+	void WriteStatement(const StringType& String)
+	{
+		InternalWriteStatementToScope(GetCurrentScope(), AcquireInternedString(String));
+	}
+
+	template<typename FormatType, typename... Types>
+	void WriteStatementf(const FormatType& Format, Types... Args)
+	{
+		InternalWriteStatementToScope(GetCurrentScope(), AcquireInternedString(Format, Forward<Types>(Args)...));
+	}
+
+	bool WriteScope(const FScope& Scope)
+	{
+		return InternalWriteScope(Scope, FStringView());
+	}
+
+	template<typename FormatType, typename... Types>
+	bool WriteScopef(const FScope& Scope, const FormatType& Format, Types... Args)
+	{
+		return InternalWriteScope(Scope, AcquireInternedString(Format, Forward<Types>(Args)...));
+	}
+
+	template<typename FormatType, typename... Types>
+	void WriteScopeTerminatorf(const FormatType& Format, Types... Args)
+	{
+		InternalWriteScopeLink(AcquireInternedString(Format, Forward<Types>(Args)...));
+	}
+
+	void WriteDeclaration(FEmitScope& EmitScope, Shader::EValueType Type, const TCHAR* Declaration, const TCHAR* Value = nullptr);
+
+	bool WriteAssignment(FEmitScope& EmitScope, const TCHAR* Declaration, FExpression* Expression, Shader::EValueType& InOutType);
+
+	void Finalize();
 
 	struct FDeclarationEntry
 	{
@@ -174,10 +263,8 @@ public:
 		TMap<FFunctionCall*, FFunctionCallEntry*> FunctionCallMap;
 	};
 
-	FScopeEntry* FindScopeEntry(const FScope* Scope);
-
-	TMap<const FScope*, FScopeEntry*> ScopeMap;
-	FScopeEntry* CurrentScopeEntry = nullptr;
+	TArray<FEmitScope*> ScopeStack;
+	TMap<const FScope*, FEmitScope*> ScopeMap;
 
 	TArray<FFunctionStackEntry> FunctionStack;
 	TArray<Shader::FPreshaderData*> TempPreshaders;
@@ -242,6 +329,7 @@ public:
 	virtual ENodeVisitResult Visit(FNodeVisitor& Visitor) = 0;
 
 	FScope* ParentScope = nullptr;
+	FNode* NextNode = nullptr;
 };
 
 /**
@@ -255,7 +343,7 @@ public:
 	virtual ENodeVisitResult Visit(FNodeVisitor& Visitor) override;
 
 	/** Emits HLSL code for the statement. The generated code should include any required semi-colons and newlines */
-	virtual bool EmitHLSL(FEmitContext& Context, FEmitCode& Scope) const = 0;
+	virtual bool EmitHLSL(FEmitContext& Context) const = 0;
 
 	FStatement* NextStatement = nullptr;
 };
@@ -276,6 +364,22 @@ public:
 	virtual bool EmitCode(FEmitContext& Context, FExpressionEmitResult& OutResult) const = 0;
 };
 
+/**
+ * Represents a phi node (see various topics on single static assignment)
+ * A phi node takes on a value based on the previous scope that was executed.
+ * In practice, this means the generated HLSL code will declare a local variable before all the previous scopes, then assign that variable the proper value from within each scope
+ */
+class FExpressionLocalPHI final : public FExpression
+{
+public:
+	virtual bool EmitCode(FEmitContext& Context, FExpressionEmitResult& OutResult) const override;
+
+	FName LocalName;
+	FScope* Scopes[MaxNumPreviousScopes];
+	FExpression* Values[MaxNumPreviousScopes];
+	int32 NumValues;
+};
+
 /**
  * Represents an HLSL parameter.  This is a uniform constant passed to the generated shader
  */
@@ -332,6 +436,8 @@ public:
 class FScope final : public FNode
 {
 public:
+	static FScope* FindSharedParent(FScope* Lhs, FScope* Rhs);
+
 	inline TArrayView<FScope*> GetPreviousScopes() const
 	{
 		// const_cast needed, otherwise type of array view is 'FScope*const' which doesn't make sense
@@ -339,9 +445,9 @@ public:
 	}
 
 	virtual ENodeVisitResult Visit(FNodeVisitor& Visitor) override;
-	bool EmitHLSL(FEmitContext& Context, FEmitCode& Scope) const;
+	bool EmitHLSL(FEmitContext& Context, FEmitScope& Scope) const;
 
-	bool EmitUnscopedHLSL(FEmitContext& Context, FEmitCode& Scope) const;
+	void AddPreviousScope(FScope& Scope);
 
 	void AddExpression(FExpression* Expression);
 	void AddStatement(FStatement* Statement);
@@ -368,6 +474,7 @@ class FTree
 {
 public:
 	static FTree* Create(FMemStackBase& Allocator);
+	static void Destroy(FTree* Tree);
 
 	bool EmitHLSL(FEmitContext& Context, FCodeWriter& Writer) const;
 
@@ -389,7 +496,8 @@ public:
 		return Statement;
 	}
 
-	FScope* NewScope(FScope& Scope, FScope** PreviousScopes = nullptr, int32 NumPreviousScopes = 0);
+	FScope* NewScope(FScope& Scope);
+
 	FParameterDeclaration* NewParameterDeclaration(FScope& Scope, const FName& Name, const Shader::FValue& DefaultValue);
 	FTextureParameterDeclaration* NewTextureParameterDeclaration(FScope& Scope, const FName& Name, const FTextureDescription& DefaultValue);
 
@@ -405,10 +513,13 @@ private:
 	inline T* NewNode(ArgTypes&&... Args)
 	{
 		T* Node = new(*Allocator) T(Forward<ArgTypes>(Args)...);
+		Node->NextNode = Nodes;
+		Nodes = Node;
 		return Node;
 	}
 
 	FMemStackBase* Allocator;
+	FNode* Nodes;
 	FScope* RootScope;
 };
 
diff --git a/Engine/Source/Runtime/Engine/Public/HLSLTree/HLSLTreeCommon.h b/Engine/Source/Runtime/Engine/Public/HLSLTree/HLSLTreeCommon.h
--- a/Engine/Source/Runtime/Engine/Public/HLSLTree/HLSLTreeCommon.h
+++ b/Engine/Source/Runtime/Engine/Public/HLSLTree/HLSLTreeCommon.h
@@ -8,15 +8,32 @@ namespace UE
 namespace HLSLTree
 {
 
-enum class EBinaryOp
+enum class EBinaryOp : uint8
 {
 	None,
 	Add,
 	Sub,
 	Mul,
 	Div,
+	Less,
 };
 
+struct FBinaryOpDescription
+{
+	FBinaryOpDescription()
+		: Name(nullptr), Operator(nullptr)
+	{}
+
+	FBinaryOpDescription(const TCHAR* InName, const TCHAR* InOperator)
+		: Name(InName), Operator(InOperator)
+	{}
+
+	const TCHAR* Name;
+	const TCHAR* Operator;
+};
+
+FBinaryOpDescription GetBinaryOpDesription(EBinaryOp Op);
+
 class FExpressionConstant : public FExpression
 {
 public:
@@ -51,21 +68,6 @@ public:
 	virtual bool EmitCode(FEmitContext& Context, FExpressionEmitResult& OutResult) const override;
 };
 
-/**
- * Represents a phi node (see various topics on single static assignment)
- * A phi node takes on a value based on the previous scope that was executed.
- * In practice, this means the generated HLSL code will declare a local variable before all the previous scopes, then assign that variable the proper value from within each scope
- */
-class FExpressionLocalPHI final : public FExpression
-{
-public:
-	virtual bool EmitCode(FEmitContext& Context, FExpressionEmitResult& OutResult) const override;
-
-	FScope* Scopes[MaxNumPreviousScopes];
-	FExpression* Values[MaxNumPreviousScopes];
-	int32 NumValues;
-};
-
 enum class EExternalInputType
 {
 	TexCoord0,
@@ -318,22 +320,22 @@ public:
 	virtual bool EmitCode(FEmitContext& Context, FExpressionEmitResult& OutResult) const override;
 };
 
-class FStatementJump : public FStatement
+class FStatementNestedScope : public FStatement
 {
 public:
-	FScope* TargetScope;
+	FScope* NextScope;
 
 	virtual ENodeVisitResult Visit(FNodeVisitor& Visitor) override
 	{
 		const ENodeVisitResult Result = FStatement::Visit(Visitor);
 		if (ShouldVisitDependentNodes(Result))
 		{
-			Visitor.VisitNode(TargetScope);
+			Visitor.VisitNode(NextScope);
 		}
 		return Result;
 	}
 
-	virtual bool EmitHLSL(FEmitContext& Context, FEmitCode& Scope) const override;
+	virtual bool EmitHLSL(FEmitContext& Context) const override;
 };
 
 class FStatementReturn : public FStatement
@@ -351,7 +353,13 @@ public:
 		return Result;
 	}
 
-	virtual bool EmitHLSL(FEmitContext& Context, FEmitCode& Scope) const override;
+	virtual bool EmitHLSL(FEmitContext& Context) const override;
+};
+
+class FStatementBreak : public FStatement
+{
+public:
+	virtual bool EmitHLSL(FEmitContext& Context) const override;
 };
 
 class FStatementIf : public FStatement
@@ -360,6 +368,7 @@ public:
 	FExpression* ConditionExpression;
 	FScope* ThenScope;
 	FScope* ElseScope;
+	FScope* NextScope;
 
 	virtual ENodeVisitResult Visit(FNodeVisitor& Visitor) override
 	{
@@ -369,11 +378,32 @@ public:
 			Visitor.VisitNode(ConditionExpression);
 			Visitor.VisitNode(ThenScope);
 			Visitor.VisitNode(ElseScope);
+			Visitor.VisitNode(NextScope);
 		}
 		return Result;
 	}
 
-	virtual bool EmitHLSL(FEmitContext& Context, FEmitCode& Scope) const override;
+	virtual bool EmitHLSL(FEmitContext& Context) const override;
+};
+
+class FStatementLoop : public FStatement
+{
+public:
+	FScope* LoopScope;
+	FScope* NextScope;
+
+	virtual ENodeVisitResult Visit(FNodeVisitor& Visitor) override
+	{
+		const ENodeVisitResult Result = FStatement::Visit(Visitor);
+		if (ShouldVisitDependentNodes(Result))
+		{
+			Visitor.VisitNode(LoopScope);
+			Visitor.VisitNode(NextScope);
+		}
+		return Result;
+	}
+
+	virtual bool EmitHLSL(FEmitContext& Context) const override;
 };
 
 class FStatementFor : public FStatement
@@ -382,6 +412,7 @@ public:
 	FExpression* StartExpression;
 	FExpression* EndExpression;
 	FScope* LoopScope;
+	FScope* NextScope;
 
 	virtual ENodeVisitResult Visit(FNodeVisitor& Visitor) override
 	{
@@ -391,12 +422,13 @@ public:
 			Visitor.VisitNode(StartExpression);
 			Visitor.VisitNode(EndExpression);
 			Visitor.VisitNode(LoopScope);
+			Visitor.VisitNode(NextScope);
 		}
 		return Result;
 	}
 
-	virtual bool EmitHLSL(FEmitContext& Context, FEmitCode& Scope) const override;
+	virtual bool EmitHLSL(FEmitContext& Context) const override;
 };
 
 }
-}
\ No newline at end of file
+}
diff --git a/Engine/Source/Runtime/Engine/Public/MaterialHLSLGenerator.h b/Engine/Source/Runtime/Engine/Public/MaterialHLSLGenerator.h
--- a/Engine/Source/Runtime/Engine/Public/MaterialHLSLGenerator.h
+++ b/Engine/Source/Runtime/Engine/Public/MaterialHLSLGenerator.h
@@ -7,6 +7,7 @@
 
 #include "Containers/Map.h"
 #include "Misc/MemStack.h"
+#include "Misc/EnumClassFlags.h"
 #include "Templates/RefCounting.h"
 #include "RHIDefinitions.h"
 #include "HLSLTree/HLSLTree.h"
@@ -33,9 +34,17 @@ class FExpressionConstant;
 class FExpressionExternalInput;
 class FExpressionSwizzle;
 class FExpressionCast;
+class FExpressionLocalPHI;
 }
 }
 
+enum class EMaterialNewScopeFlag : uint8
+{
+	None = 0u,
+	NoPreviousScope = (1u << 0),
+};
+ENUM_CLASS_FLAGS(EMaterialNewScopeFlag);
+
 /**
  * MaterialHLSLGenerator is a bridge between a material, and HLSLTree.  It facilitates generating HLSL source code for a given material, using HLSLTree
  */
@@ -66,6 +75,10 @@ public:
 
 	bool GenerateResult(UE::HLSLTree::FScope& Scope);
 
+	UE::HLSLTree::FScope* NewScope(UE::HLSLTree::FScope& Scope, EMaterialNewScopeFlag Flags = EMaterialNewScopeFlag::None);
+
+	UE::HLSLTree::FScope* NewJoinedScope(UE::HLSLTree::FScope& Scope);
+
 	UE::HLSLTree::FExpressionConstant* NewConstant(UE::HLSLTree::FScope& Scope, const UE::Shader::FValue& Value);
 	UE::HLSLTree::FExpressionExternalInput* NewTexCoord(UE::HLSLTree::FScope& Scope, int32 Index);
 	UE::HLSLTree::FExpressionSwizzle* NewSwizzle(UE::HLSLTree::FScope& Scope, const UE::HLSLTree::FSwizzleParameters& Params, UE::HLSLTree::FExpression* Input);
@@ -155,12 +168,16 @@ private:
 		int32 NumInputs = 0;
 	};
 
+	UE::HLSLTree::FExpression* InternalAcquireLocalValue(UE::HLSLTree::FScope& Scope, const FName& LocalName);
+
 	const FMaterialCompileTargetParameters& CompileTarget;
 	UMaterial* TargetMaterial;
 	UMaterialFunctionInterface* TargetMaterialFunction;
 
 	UE::HLSLTree::FTree* HLSLTree;
 	TArray<FExpressionKey> ExpressionStack;
+	TArray<UE::HLSLTree::FScope*> JoinedScopeStack;
+	TArray<UE::HLSLTree::FExpressionLocalPHI*> PHIExpressions;
 	TArray<FString> CompileErrors;
 	TArray<UMaterialExpression*> ErrorExpressions;
 	TMap<FName, UE::HLSLTree::FParameterDeclaration*> ParameterDeclarationMap;
diff --git a/Engine/Source/Runtime/Engine/Public/MaterialHLSLTree.h b/Engine/Source/Runtime/Engine/Public/MaterialHLSLTree.h
--- a/Engine/Source/Runtime/Engine/Public/MaterialHLSLTree.h
+++ b/Engine/Source/Runtime/Engine/Public/MaterialHLSLTree.h
@@ -25,6 +25,7 @@ class FMaterialHLSLTree
 {
 public:
 	FMaterialHLSLTree();
+	~FMaterialHLSLTree();
 
 	bool InitializeForMaterial(const FMaterialCompileTargetParameters& InCompilerTarget, FMaterial& InOutMaterial);
 
diff --git a/Engine/Source/Runtime/Engine/Public/Shader/ShaderTypes.h b/Engine/Source/Runtime/Engine/Public/Shader/ShaderTypes.h
--- a/Engine/Source/Runtime/Engine/Public/Shader/ShaderTypes.h
+++ b/Engine/Source/Runtime/Engine/Public/Shader/ShaderTypes.h
@@ -53,6 +53,7 @@ FValueTypeDescription GetValueTypeDescription(EValueType Type);
 EValueType MakeValueType(EValueComponentType ComponentType, int32 NumComponents);
 EValueType MakeValueType(EValueType BaseType, int32 NumComponents);
 EValueType MakeArithmeticResultType(EValueType Lhs, EValueType Rhs, FString& OutErrorMessage);
+EValueType MakeComparisonResultType(EValueType Lhs, EValueType Rhs, FString& OutErrorMessage);
 
 template<typename T>
 struct TValue
