Display the index stride in the rhi debug name of index buffers.

diff --git a/Engine/Source/Runtime/Engine/Private/RawIndexBuffer.cpp b/Engine/Source/Runtime/Engine/Private/RawIndexBuffer.cpp
--- a/Engine/Source/Runtime/Engine/Private/RawIndexBuffer.cpp
+++ b/Engine/Source/Runtime/Engine/Private/RawIndexBuffer.cpp
@@ -330,7 +330,7 @@ FBufferRHIRef FRawStaticIndexBuffer::CreateRHIBuffer_Internal()
 		uint32 BufferFlags = BUF_Static | (bSRV ? BUF_ShaderResource : BUF_None);
 
 		// Create the index buffer.
-		FRHIResourceCreateInfo CreateInfo(TEXT("FRawStaticIndexBuffer"), &IndexStorage);
+		FRHIResourceCreateInfo CreateInfo(Is32Bit() ? TEXT("FRawStaticIndexBuffer32") : TEXT("FRawStaticIndexBuffer16"), &IndexStorage);
 		CreateInfo.bWithoutNativeResource = !SizeInBytes;
 		if (bRenderThread)
 		{
diff --git a/Engine/Source/Runtime/Engine/Public/RawIndexBuffer.h b/Engine/Source/Runtime/Engine/Public/RawIndexBuffer.h
--- a/Engine/Source/Runtime/Engine/Public/RawIndexBuffer.h
+++ b/Engine/Source/Runtime/Engine/Public/RawIndexBuffer.h
@@ -375,6 +375,7 @@ public:
 		: Indices(InNeedsCPUAccess)
 		, CachedNumIndices(0)
 	{
+		static_assert(sizeof(INDEX_TYPE) == 2 || sizeof(INDEX_TYPE) == 4, "FRawStaticIndexBuffer16or32 must have a stride of 2 or 4 bytes.");
 	}
 
 	/**
@@ -531,7 +532,7 @@ private:
 		if (CachedNumIndices)
 		{
 			// Create the index buffer.
-			FRHIResourceCreateInfo CreateInfo(TEXT("FRawStaticIndexBuffer"), &Indices);
+			FRHIResourceCreateInfo CreateInfo(sizeof(INDEX_TYPE) == 4 ? TEXT("FRawStaticIndexBuffer32") : TEXT("FRawStaticIndexBuffer16"), &Indices);
 			EBufferUsageFlags Flags = BUF_Static;
 
 			if (IsSRVNeeded())
@@ -557,4 +558,4 @@ private:
 		}
 		return nullptr;
 	}
-};
+};
\ No newline at end of file
