Fix for Modulation & Submix content intermittently failing to load.

- Refactor logic querying/loading the default SoundClass.  Prior logic loaded them extremely early in the engine boot cycle, which indirectly could initiate load attempts on assets that were of plugin-supplied types prior to plugins supplying their given type being loaded.

diff --git a/Engine/Source/Runtime/AudioMixer/Private/Components/SynthComponent.cpp b/Engine/Source/Runtime/AudioMixer/Private/Components/SynthComponent.cpp
--- a/Engine/Source/Runtime/AudioMixer/Private/Components/SynthComponent.cpp
+++ b/Engine/Source/Runtime/AudioMixer/Private/Components/SynthComponent.cpp
@@ -141,10 +141,6 @@ USynthComponent::USynthComponent(const FObjectInitializer& ObjectInitializer)
 	bAlwaysPlay = true;
 	Synth = nullptr;
 
-	// Set the default sound class
-	SoundClass = USoundBase::DefaultSoundClassObject;
-	Synth = nullptr;
-
 	PreferredBufferLength = DEFAULT_PROCEDURAL_SOUNDWAVE_BUFFER_SIZE;
 
 #if WITH_EDITORONLY_DATA
diff --git a/Engine/Source/Runtime/AudioMixer/Public/Components/SynthComponent.h b/Engine/Source/Runtime/AudioMixer/Public/Components/SynthComponent.h
--- a/Engine/Source/Runtime/AudioMixer/Public/Components/SynthComponent.h
+++ b/Engine/Source/Runtime/AudioMixer/Public/Components/SynthComponent.h
@@ -253,6 +253,15 @@ public:
 	void AddAudioBufferListener(IAudioBufferListener* InAudioBufferListener);
 	void RemoveAudioBufferListener(IAudioBufferListener* InAudioBufferListener);
 
+	virtual USoundClass* GetSoundClass()
+	{
+		if (SoundClass)
+		{
+			return SoundClass;
+		}
+
+		return USoundBase::GetDefaultSoundClass();
+	}
 
 	virtual void BeginDestroy() override;
 
diff --git a/Engine/Source/Runtime/Engine/Classes/Sound/SoundBase.h b/Engine/Source/Runtime/Engine/Classes/Sound/SoundBase.h
--- a/Engine/Source/Runtime/Engine/Classes/Sound/SoundBase.h
+++ b/Engine/Source/Runtime/Engine/Classes/Sound/SoundBase.h
@@ -8,6 +8,8 @@
 #include "Audio.h"
 #include "CoreMinimal.h"
 #include "IAudioExtensionPlugin.h"
+#include "Sound/AudioSettings.h"
+#include "Sound/SoundClass.h"
 #include "SoundConcurrency.h"
 #include "SoundModulationDestination.h"
 #include "SoundSourceBusSend.h"
@@ -58,9 +60,6 @@ class ENGINE_API USoundBase : public UObject, public IInterface_AssetUserData
 	GENERATED_UCLASS_BODY()
 
 public:
-	static USoundClass* DefaultSoundClassObject;
-	static USoundConcurrency* DefaultSoundConcurrencyObject;
-
 	/** Sound class this sound belongs to */
 	UPROPERTY(EditAnywhere, Category = Sound, meta = (DisplayName = "Class"), AssetRegistrySearchable)
 	TObjectPtr<USoundClass> SoundClassObject;
@@ -189,9 +188,10 @@ public:
 	TArray<TObjectPtr<UAssetUserData>> AssetUserData;
 
 public:
+	static USoundClass* GetDefaultSoundClass();
+	static USoundConcurrency* GetDefaultSoundConcurrency();
 
 	//~ Begin UObject Interface.
-	virtual void PostInitProperties() override;
 #if WITH_EDITORONLY_DATA
 	virtual void PostLoad() override;
 #endif
diff --git a/Engine/Source/Runtime/Engine/Private/ActiveSound.cpp b/Engine/Source/Runtime/Engine/Private/ActiveSound.cpp
--- a/Engine/Source/Runtime/Engine/Private/ActiveSound.cpp
+++ b/Engine/Source/Runtime/Engine/Private/ActiveSound.cpp
@@ -479,20 +479,24 @@ void FActiveSound::GetConcurrencyHandles(TArray<FConcurrencyHandle>& OutConcurre
 {
 	OutConcurrencyHandles.Reset();
 
-	if (!ConcurrencySet.Num() && Sound)
-	{
-		Sound->GetConcurrencyHandles(OutConcurrencyHandles);
-	}
-	else
+	if (ConcurrencySet.Num() > 0)
 	{
 		for (const USoundConcurrency* Concurrency : ConcurrencySet)
 		{
 			if (Concurrency)
 			{
-				OutConcurrencyHandles.Emplace(FConcurrencyHandle(*Concurrency));
+				OutConcurrencyHandles.Emplace(*Concurrency);
 			}
 		}
 	}
+	else if (Sound)
+	{
+		Sound->GetConcurrencyHandles(OutConcurrencyHandles);
+	}
+	else if (USoundConcurrency* DefaultConcurrency = USoundBase::GetDefaultSoundConcurrency())
+	{
+		OutConcurrencyHandles.Emplace(*DefaultConcurrency);
+	}
 }
 
 bool FActiveSound::GetConcurrencyFadeDuration(float& OutFadeDuration) const
diff --git a/Engine/Source/Runtime/Engine/Private/SoundBase.cpp b/Engine/Source/Runtime/Engine/Private/SoundBase.cpp
--- a/Engine/Source/Runtime/Engine/Private/SoundBase.cpp
+++ b/Engine/Source/Runtime/Engine/Private/SoundBase.cpp
@@ -10,9 +10,6 @@
 #include "Engine/AssetUserData.h"
 
 
-USoundClass* USoundBase::DefaultSoundClassObject = nullptr;
-USoundConcurrency* USoundBase::DefaultSoundConcurrencyObject = nullptr;
-
 USoundBase::USoundBase(const FObjectInitializer& ObjectInitializer)
 	: Super(ObjectInitializer)
 	, VirtualizationMode(EVirtualizationMode::Restart)
@@ -29,35 +26,40 @@ USoundBase::USoundBase(const FObjectInitializer& ObjectInitializer)
 	bEnableSubmixSends = true;
 }
 
-void USoundBase::PostInitProperties()
+USoundClass* USoundBase::GetDefaultSoundClass()
 {
-	Super::PostInitProperties();
-
-	if (USoundBase::DefaultSoundClassObject == nullptr)
+	const FSoftObjectPath DefaultPath = GetDefault<UAudioSettings>()->DefaultSoundClassName;
+	if (DefaultPath.IsValid())
 	{
-		const FSoftObjectPath DefaultSoundClassName = GetDefault<UAudioSettings>()->DefaultSoundClassName;
-		if (DefaultSoundClassName.IsValid())
+		UObject* DefaultObject = DefaultPath.TryLoad();
+		if (DefaultObject)
 		{
-			SCOPED_BOOT_TIMING("USoundBase::LoadSoundClass");
-			USoundBase::DefaultSoundClassObject = LoadObject<USoundClass>(nullptr, *DefaultSoundClassName.ToString());
+			if (USoundClass* DefaultSoundClass = CastChecked<USoundClass>(DefaultObject))
+			{
+				return DefaultSoundClass;
+			}
 		}
 	}
-	SoundClassObject = USoundBase::DefaultSoundClassObject;
 
-	if (USoundBase::DefaultSoundConcurrencyObject == nullptr)
+	return nullptr;
+}
+
+USoundConcurrency* USoundBase::GetDefaultSoundConcurrency()
+{
+	const FSoftObjectPath DefaultPath = GetDefault<UAudioSettings>()->DefaultSoundConcurrencyName;
+	if (DefaultPath.IsValid())
 	{
-		const FSoftObjectPath DefaultSoundConcurrencyName = GetDefault<UAudioSettings>()->DefaultSoundConcurrencyName;
-		if (DefaultSoundConcurrencyName.IsValid())
+		UObject* DefaultObject = DefaultPath.TryLoad();
+		if (DefaultObject)
 		{
-			SCOPED_BOOT_TIMING("USoundBase::LoadSoundConcurrency");
-			USoundBase::DefaultSoundConcurrencyObject = LoadObject<USoundConcurrency>(nullptr, *DefaultSoundConcurrencyName.ToString());
+			if (USoundConcurrency* DefaultSoundConcurrency = CastChecked<USoundConcurrency>(DefaultObject))
+			{
+				return DefaultSoundConcurrency;
+			}
 		}
 	}
 
-	if (USoundBase::DefaultSoundConcurrencyObject != nullptr)
-	{
-		ConcurrencySet.Add(USoundBase::DefaultSoundConcurrencyObject);
-	}
+	return nullptr;
 }
 
 bool USoundBase::IsPlayable() const
@@ -135,12 +137,18 @@ bool USoundBase::IsLooping()
 
 bool USoundBase::ShouldApplyInteriorVolumes()
 {
-	return (SoundClassObject && SoundClassObject->Properties.bApplyAmbientVolumes);
+	USoundClass* SoundClass = GetSoundClass();
+	return SoundClass && SoundClass->Properties.bApplyAmbientVolumes;
 }
 
 USoundClass* USoundBase::GetSoundClass() const
 {
-	return SoundClassObject;
+	if (SoundClassObject)
+	{
+		return SoundClassObject;
+	}
+
+	return GetDefaultSoundClass();
 }
 
 USoundSubmixBase* USoundBase::GetSoundSubmix() const
@@ -174,13 +182,20 @@ void USoundBase::GetConcurrencyHandles(TArray<FConcurrencyHandle>& OutConcurrenc
 	}
 	else
 	{
-		for (const USoundConcurrency* Concurrency : ConcurrencySet)
+		if (ConcurrencySet.Num() > 0)
 		{
-			if (Concurrency)
+			for (const USoundConcurrency* Concurrency : ConcurrencySet)
 			{
-				OutConcurrencyHandles.Emplace(*Concurrency);
+				if (Concurrency)
+				{
+					OutConcurrencyHandles.Emplace(*Concurrency);
+				}
 			}
 		}
+		else if (const USoundConcurrency* DefaultConcurrency = GetDefaultSoundConcurrency())
+		{
+			OutConcurrencyHandles.Emplace(*DefaultConcurrency);
+		}
 	}
 }
 
diff --git a/Engine/Source/Runtime/Engine/Private/SoundCue.cpp b/Engine/Source/Runtime/Engine/Private/SoundCue.cpp
--- a/Engine/Source/Runtime/Engine/Private/SoundCue.cpp
+++ b/Engine/Source/Runtime/Engine/Private/SoundCue.cpp
@@ -234,7 +234,7 @@ void USoundCue::PostLoad()
 	
 	ESoundWaveLoadingBehavior SoundClassLoadingBehavior = ESoundWaveLoadingBehavior::Inherited;
 
-	USoundClass* CurrentSoundClass = SoundClassObject;
+	USoundClass* CurrentSoundClass = GetSoundClass();
 
 	// Recurse through this sound class's parents until we find an override.
 	while (SoundClassLoadingBehavior == ESoundWaveLoadingBehavior::Inherited && CurrentSoundClass != nullptr)
diff --git a/Engine/Source/Runtime/Engine/Private/SoundWave.cpp b/Engine/Source/Runtime/Engine/Private/SoundWave.cpp
--- a/Engine/Source/Runtime/Engine/Private/SoundWave.cpp
+++ b/Engine/Source/Runtime/Engine/Private/SoundWave.cpp
@@ -3051,7 +3051,7 @@ void USoundWave::CacheInheritedLoadingBehavior() const
 		// if this is true then the behavior should not be Inherited here
 		check(!*bLoadingBehaviorOverriddenPtr);
 
-		USoundClass* CurrentSoundClass = SoundClassObject;
+		USoundClass* CurrentSoundClass = GetSoundClass();
 		ESoundWaveLoadingBehavior SoundClassLoadingBehavior = ESoundWaveLoadingBehavior::Inherited;
 
 		// Recurse through this sound class's parents until we find an override.
diff --git a/Engine/Source/Runtime/MediaAssets/Private/Assets/MediaSoundComponent.cpp b/Engine/Source/Runtime/MediaAssets/Private/Assets/MediaSoundComponent.cpp
--- a/Engine/Source/Runtime/MediaAssets/Private/Assets/MediaSoundComponent.cpp
+++ b/Engine/Source/Runtime/MediaAssets/Private/Assets/MediaSoundComponent.cpp
@@ -26,11 +26,6 @@ CSV_DECLARE_CATEGORY_MODULE_EXTERN(MEDIA_API, MediaStreaming);
 
 static const int32 MaxAudioInputSamples = 8;	// accept at most these many samples into our input queue
 
-/* Static initialization
- *****************************************************************************/
-
-USoundClass* UMediaSoundComponent::DefaultMediaSoundClassObject = nullptr;
-
 
 /* UMediaSoundComponent structors
  *****************************************************************************/
@@ -257,27 +252,6 @@ void UMediaSoundComponent::Deactivate()
 /* UObject interface
  *****************************************************************************/
 
-void UMediaSoundComponent::PostInitProperties()
-{
-	Super::PostInitProperties();
-
-	if (UMediaSoundComponent::DefaultMediaSoundClassObject == nullptr)
-	{
-		const FSoftObjectPath DefaultMediaSoundClassName = GetDefault<UAudioSettings>()->DefaultMediaSoundClassName;
-
-		if (DefaultMediaSoundClassName.IsValid())
-		{
-			UMediaSoundComponent::DefaultMediaSoundClassObject = LoadObject<USoundClass>(nullptr, *DefaultMediaSoundClassName.ToString());
-		}
-	}
-
-	// We have a different default sound class object for media sound components
-	if (SoundClass == USoundBase::DefaultSoundClassObject || SoundClass == nullptr)
-	{
-		SoundClass = UMediaSoundComponent::DefaultMediaSoundClassObject;
-	}
-}
-
 
 void UMediaSoundComponent::PostLoad()
 {
diff --git a/Engine/Source/Runtime/MediaAssets/Public/MediaSoundComponent.h b/Engine/Source/Runtime/MediaAssets/Public/MediaSoundComponent.h
--- a/Engine/Source/Runtime/MediaAssets/Public/MediaSoundComponent.h
+++ b/Engine/Source/Runtime/MediaAssets/Public/MediaSoundComponent.h
@@ -135,6 +135,29 @@ public:
 	UFUNCTION(BlueprintCallable, Category="Media|MediaSoundComponent")
 	UMediaPlayer* GetMediaPlayer() const;
 
+	virtual USoundClass* GetSoundClass() override
+	{
+		if (SoundClass)
+		{
+			return SoundClass;
+		}
+
+		const FSoftObjectPath DefaultPath = GetDefault<UAudioSettings>()->DefaultMediaSoundClassName;
+		if (DefaultPath.IsValid())
+		{
+			UObject* DefaultObject = DefaultPath.TryLoad();
+			if (DefaultObject)
+			{
+				if (USoundClass* DefaultSoundClass = CastChecked<USoundClass>(DefaultObject))
+				{
+					return DefaultSoundClass;
+				}
+			}
+		}
+
+		return USoundBase::GetDefaultSoundClass();
+	}
+
 	/**
 	 * Set the media player that provides the audio samples.
 	 *
@@ -210,7 +233,6 @@ public:
 public:
 
 	//~ UObject interface
-	virtual void PostInitProperties() override;
 	virtual void PostLoad() override;
 
 #if WITH_EDITOR
@@ -302,8 +324,4 @@ private:
 
 	/** Whether or not envelope follower settings changed. */
 	bool bEnvelopeFollowerSettingsChanged;
-
-private:
-
-	static USoundClass* DefaultMediaSoundClassObject;
 };
