Clean up of strata horizontal mix background/foreground parameters order for concistency (and  my sanity)

diff --git a/Engine/Shaders/Private/Strata/Strata.ush b/Engine/Shaders/Private/Strata/Strata.ush
--- a/Engine/Shaders/Private/Strata/Strata.ush
+++ b/Engine/Shaders/Private/Strata/Strata.ush
@@ -898,7 +898,7 @@ FStrataData StrataWeight(FStrataData A, float Weight)
 	return StrataData;
 }
 
-FStrataData StrataHorizontalMixing(FStrataData Foreground, FStrataData Background, float Mix)
+FStrataData StrataHorizontalMixing(FStrataData Background, FStrataData Foreground, float Mix)
 {
 	const float ForegroundMixFactor = saturate(Mix);
 	const float BackgroundMixFactor = 1.0 - ForegroundMixFactor;
@@ -945,7 +945,7 @@ float HorizontalMixingParameterBlendingBSDFCoverageToNormalMix(FStrataData Backg
 //  - Slab node with SSSprofile cannot be blendend. Those will have to be specifically handled by the user and not mixed with any other slab if it can be demoted to parameter blending
 
 // ==> NOTE: Always pair with the compiler behavior in StrataCompilationInfoHorizontalMixingParamBlend
-FStrataData StrataHorizontalMixingParameterBlending(FStrataData Foreground, FStrataData Background, float HorizontalMixValue, float NormalMixCodeChunk, uint NewNormalIndex)
+FStrataData StrataHorizontalMixingParameterBlending(FStrataData Background, FStrataData Foreground, float HorizontalMixValue, float NormalMixCodeChunk, uint NewNormalIndex)
 {
 	const float ForegroundLightingMixFactor = saturate(HorizontalMixValue);	// Only the horizontal mixing value
 	const float ForegroundOtherMixFactor = saturate(NormalMixCodeChunk);	// Horizontal mixing value combined with coverage
diff --git a/Engine/Source/Runtime/Engine/Private/Materials/HLSLMaterialTranslator.cpp b/Engine/Source/Runtime/Engine/Private/Materials/HLSLMaterialTranslator.cpp
--- a/Engine/Source/Runtime/Engine/Private/Materials/HLSLMaterialTranslator.cpp
+++ b/Engine/Source/Runtime/Engine/Private/Materials/HLSLMaterialTranslator.cpp
@@ -9109,7 +9109,7 @@ int32 FHLSLMaterialTranslator::StrataSingleLayerWaterBSDF(
 	);
 }
 
-int32 FHLSLMaterialTranslator::StrataHorizontalMixing(int32 Foreground, int32 Background, int32 Mix)
+int32 FHLSLMaterialTranslator::StrataHorizontalMixing(int32 Background, int32 Foreground, int32 Mix)
 {
 	if (Foreground == INDEX_NONE || Background == INDEX_NONE || Mix == INDEX_NONE)
 	{
@@ -9117,13 +9117,13 @@ int32 FHLSLMaterialTranslator::StrataHorizontalMixing(int32 Foreground, int32 Ba
 	}
 	return AddCodeChunk(
 		MCT_Strata, TEXT("StrataHorizontalMixing(%s, %s, %s)"),
-		*GetParameterCode(Foreground),
 		*GetParameterCode(Background),
+		*GetParameterCode(Foreground),
 		*GetParameterCode(Mix)
 	);
 }
 
-int32 FHLSLMaterialTranslator::StrataHorizontalMixingParameterBlending(int32 Foreground, int32 Background, int32 HorizontalMixCodeChunk, int32 NormalMixCodeChunk, const FString& SharedLocalBasisIndexMacro)
+int32 FHLSLMaterialTranslator::StrataHorizontalMixingParameterBlending(int32 Background, int32 Foreground, int32 HorizontalMixCodeChunk, int32 NormalMixCodeChunk, const FString& SharedLocalBasisIndexMacro)
 {
 	if (Foreground == INDEX_NONE || Background == INDEX_NONE || HorizontalMixCodeChunk == INDEX_NONE || NormalMixCodeChunk == INDEX_NONE)
 	{
@@ -9131,8 +9131,8 @@ int32 FHLSLMaterialTranslator::StrataHorizontalMixingParameterBlending(int32 For
 	}
 	return AddCodeChunk(
 		MCT_Strata, TEXT("StrataHorizontalMixingParameterBlending(%s, %s, %s, %s, %s)"),
-		*GetParameterCode(Foreground),
 		*GetParameterCode(Background),
+		*GetParameterCode(Foreground),
 		*GetParameterCode(HorizontalMixCodeChunk),
 		*GetParameterCode(NormalMixCodeChunk),
 		*SharedLocalBasisIndexMacro
diff --git a/Engine/Source/Runtime/Engine/Private/Materials/HLSLMaterialTranslator.h b/Engine/Source/Runtime/Engine/Private/Materials/HLSLMaterialTranslator.h
--- a/Engine/Source/Runtime/Engine/Private/Materials/HLSLMaterialTranslator.h
+++ b/Engine/Source/Runtime/Engine/Private/Materials/HLSLMaterialTranslator.h
@@ -996,8 +996,8 @@ protected:
 		int32 BaseColor, int32 Metallic, int32 Specular, int32 Roughness, 
 		int32 EmissiveColor, int32 TopMaterialOpacity, int32 WaterAlbedo, int32 WaterExtinction, int32 WaterPhaseG, 
 		int32 ColorScaleBehindWater, int32 Normal, const FString& SharedLocalBasisIndexMacro) override;
-	virtual int32 StrataHorizontalMixing(int32 Foreground, int32 Background, int32 Mix) override;
-	virtual int32 StrataHorizontalMixingParameterBlending(int32 Foreground, int32 Background, int32 HorizontalMixCodeChunk, int32 NormalMixCodeChunk, const FString& SharedLocalBasisIndexMacro) override;
+	virtual int32 StrataHorizontalMixing(int32 Background, int32 Foreground, int32 Mix) override;
+	virtual int32 StrataHorizontalMixingParameterBlending(int32 Background, int32 Foreground, int32 HorizontalMixCodeChunk, int32 NormalMixCodeChunk, const FString& SharedLocalBasisIndexMacro) override;
 	virtual int32 StrataVerticalLayering(int32 Top, int32 Base) override;
 	virtual int32 StrataVerticalLayeringParameterBlending(int32 Top, int32 Base, const FString& SharedLocalBasisIndexMacro, int32 TopBSDFNormalCodeChunk) override;
 	virtual int32 StrataAdd(int32 A, int32 B) override;
diff --git a/Engine/Source/Runtime/Engine/Private/Materials/MaterialExpressions.cpp b/Engine/Source/Runtime/Engine/Private/Materials/MaterialExpressions.cpp
--- a/Engine/Source/Runtime/Engine/Private/Materials/MaterialExpressions.cpp
+++ b/Engine/Source/Runtime/Engine/Private/Materials/MaterialExpressions.cpp
@@ -626,7 +626,7 @@ int32 CompileStrataBlendFunction(FMaterialCompiler* Compiler, const int32 A, con
 	{
 		return INDEX_NONE;
 	}
-	int32 OutputCodeChunk = Compiler->StrataHorizontalMixing(B, A, Alpha); // B is foreground (fully visible when Alpha=1), and A is background (fully visible when Alpha=0)
+	int32 OutputCodeChunk = Compiler->StrataHorizontalMixing(A, B, Alpha); // A is background (fully visible when Alpha=0) and B is foreground (fully visible when Alpha=1)
 
 #if WITH_EDITOR
 	// Also register the horizontal mixing operation with the compiler.
@@ -634,7 +634,7 @@ int32 CompileStrataBlendFunction(FMaterialCompiler* Compiler, const int32 A, con
 	{
 		return Compiler->Errorf(TEXT("Could not find A or B code chunk in CompileStrataBlendFunction"));
 	}
-	FStrataMaterialCompilationInfo StrataInfo = StrataCompilationInfoHorizontalMixing(Compiler, Compiler->GetStrataCompilationInfo(B), Compiler->GetStrataCompilationInfo(A));
+	FStrataMaterialCompilationInfo StrataInfo = StrataCompilationInfoHorizontalMixing(Compiler, Compiler->GetStrataCompilationInfo(A), Compiler->GetStrataCompilationInfo(B));
 	Compiler->StrataCompilationInfoRegisterCodeChunk(OutputCodeChunk, StrataInfo);
 #endif
 
@@ -20875,21 +20875,21 @@ int32 UMaterialExpressionStrataHorizontalMixing::Compile(class FMaterialCompiler
 		}
 		const FStrataRegisteredSharedLocalBasis NewRegisteredSharedLocalBasis = StrataCompilationInfoCreateSharedLocalBasis(Compiler, NewNormalCodeChunk, NewTangentCodeChunk);
 
-		OutputCodeChunk = Compiler->StrataHorizontalMixingParameterBlending(ForegroundCodeChunk, BackgroundCodeChunk, HorizontalMixCodeChunk, NormalMixCodeChunk, Compiler->GetStrataSharedLocalBasisIndexMacro(NewRegisteredSharedLocalBasis));
+		OutputCodeChunk = Compiler->StrataHorizontalMixingParameterBlending(BackgroundCodeChunk, ForegroundCodeChunk, HorizontalMixCodeChunk, NormalMixCodeChunk, Compiler->GetStrataSharedLocalBasisIndexMacro(NewRegisteredSharedLocalBasis));
 
 		StrataCompilationInfoCreateSingleBSDFMaterial(Compiler, OutputCodeChunk, NewRegisteredSharedLocalBasis, STRATA_BSDF_TYPE_SLAB);
-		FStrataMaterialCompilationInfo StrataInfo = StrataCompilationInfoHorizontalMixingParamBlend(Compiler, ForegroundStrataData, BackgroundStrataData, NewRegisteredSharedLocalBasis);
+		FStrataMaterialCompilationInfo StrataInfo = StrataCompilationInfoHorizontalMixingParamBlend(Compiler, BackgroundStrataData, ForegroundStrataData, NewRegisteredSharedLocalBasis);
 
 		Compiler->StrataCompilationInfoRegisterCodeChunk(OutputCodeChunk, StrataInfo);
 	}
 	else
 	{
 		OutputCodeChunk = Compiler->StrataHorizontalMixing(
-			ForegroundCodeChunk,
 			BackgroundCodeChunk,
+			ForegroundCodeChunk,
 			HorizontalMixCodeChunk);
 
-		FStrataMaterialCompilationInfo StrataInfo = StrataCompilationInfoHorizontalMixing(Compiler, ForegroundStrataData, BackgroundStrataData);
+		FStrataMaterialCompilationInfo StrataInfo = StrataCompilationInfoHorizontalMixing(Compiler, BackgroundStrataData, ForegroundStrataData);
 
 		Compiler->StrataCompilationInfoRegisterCodeChunk(OutputCodeChunk, StrataInfo);
 	}
diff --git a/Engine/Source/Runtime/Engine/Private/Materials/StrataMaterial.cpp b/Engine/Source/Runtime/Engine/Private/Materials/StrataMaterial.cpp
--- a/Engine/Source/Runtime/Engine/Private/Materials/StrataMaterial.cpp
+++ b/Engine/Source/Runtime/Engine/Private/Materials/StrataMaterial.cpp
@@ -152,21 +152,21 @@ FStrataMaterialCompilationInfo StrataCompilationInfoAddParamBlend(FMaterialCompi
 }
 
 
-FStrataMaterialCompilationInfo StrataCompilationInfoHorizontalMixing(FMaterialCompiler* Compiler, const FStrataMaterialCompilationInfo& A, const FStrataMaterialCompilationInfo& B)
+FStrataMaterialCompilationInfo StrataCompilationInfoHorizontalMixing(FMaterialCompiler* Compiler, const FStrataMaterialCompilationInfo& Background, const FStrataMaterialCompilationInfo& Foreground)
 {
-	return StrataCompilationInfoAdd(Compiler, A, B); // Mixing is a similar operation to Add when it comes to bsdf count
+	return StrataCompilationInfoAdd(Compiler, Background, Foreground); // Mixing is a similar operation to Add when it comes to bsdf count
 }
 
 
 // ==> NOTE: Always pair with the shader behavior in StrataHorizontalMixingParameterBlending
-FStrataMaterialCompilationInfo StrataCompilationInfoHorizontalMixingParamBlend(FMaterialCompiler* Compiler, const FStrataMaterialCompilationInfo& A, const FStrataMaterialCompilationInfo& B, const FStrataRegisteredSharedLocalBasis& RegisteredSharedLocalBasis)
+FStrataMaterialCompilationInfo StrataCompilationInfoHorizontalMixingParamBlend(FMaterialCompiler* Compiler, const FStrataMaterialCompilationInfo& Background, const FStrataMaterialCompilationInfo& Foreground, const FStrataRegisteredSharedLocalBasis& RegisteredSharedLocalBasis)
 {
-	check(A.TotalBSDFCount == 1);
-	check(B.TotalBSDFCount == 1);
+	check(Background.TotalBSDFCount == 1);
+	check(Foreground.TotalBSDFCount == 1);
 
-	FStrataMaterialCompilationInfo StrataInfo = A;
+	FStrataMaterialCompilationInfo StrataInfo = Background;
 	FStrataMaterialCompilationInfo::FBSDF& NewBSDF = StrataInfo.Layers[0].BSDFs[0];
-	const FStrataMaterialCompilationInfo::FBSDF& OtherBSDF = B.Layers[0].BSDFs[0];
+	const FStrataMaterialCompilationInfo::FBSDF& OtherBSDF = Foreground.Layers[0].BSDFs[0];
 
 	NewBSDF.RegisteredSharedLocalBasis = RegisteredSharedLocalBasis;
 
diff --git a/Engine/Source/Runtime/Engine/Private/Materials/StrataMaterial.h b/Engine/Source/Runtime/Engine/Private/Materials/StrataMaterial.h
--- a/Engine/Source/Runtime/Engine/Private/Materials/StrataMaterial.h
+++ b/Engine/Source/Runtime/Engine/Private/Materials/StrataMaterial.h
@@ -61,8 +61,8 @@ void  StrataCompilationInfoCreateSingleBSDFMaterial(FMaterialCompiler* Compiler,
 FStrataMaterialCompilationInfo StrataCompilationInfoWeight(FMaterialCompiler* Compiler, const FStrataMaterialCompilationInfo& A);
 FStrataMaterialCompilationInfo StrataCompilationInfoAdd(FMaterialCompiler* Compiler, const FStrataMaterialCompilationInfo& A, const FStrataMaterialCompilationInfo& B);
 FStrataMaterialCompilationInfo StrataCompilationInfoAddParamBlend(FMaterialCompiler* Compiler, const FStrataMaterialCompilationInfo& A, const FStrataMaterialCompilationInfo& B, const FStrataRegisteredSharedLocalBasis& RegisteredSharedLocalBasis);
-FStrataMaterialCompilationInfo StrataCompilationInfoHorizontalMixing(FMaterialCompiler* Compiler, const FStrataMaterialCompilationInfo& A, const FStrataMaterialCompilationInfo& B);
-FStrataMaterialCompilationInfo StrataCompilationInfoHorizontalMixingParamBlend(FMaterialCompiler* Compiler, const FStrataMaterialCompilationInfo& A, const FStrataMaterialCompilationInfo& B, const FStrataRegisteredSharedLocalBasis& RegisteredSharedLocalBasis);
+FStrataMaterialCompilationInfo StrataCompilationInfoHorizontalMixing(FMaterialCompiler* Compiler, const FStrataMaterialCompilationInfo& Background, const FStrataMaterialCompilationInfo& Foreground);
+FStrataMaterialCompilationInfo StrataCompilationInfoHorizontalMixingParamBlend(FMaterialCompiler* Compiler, const FStrataMaterialCompilationInfo& Background, const FStrataMaterialCompilationInfo& Foreground, const FStrataRegisteredSharedLocalBasis& RegisteredSharedLocalBasis);
 FStrataMaterialCompilationInfo StrataCompilationInfoVerticalLayering(FMaterialCompiler* Compiler, const FStrataMaterialCompilationInfo& Top, const FStrataMaterialCompilationInfo& Base);
 FStrataMaterialCompilationInfo StrataCompilationInfoVerticalLayeringParamBlend(FMaterialCompiler* Compiler, const FStrataMaterialCompilationInfo& Top, const FStrataMaterialCompilationInfo& Base, const FStrataRegisteredSharedLocalBasis& RegisteredSharedLocalBasis);
 
diff --git a/Engine/Source/Runtime/Engine/Public/MaterialCompiler.h b/Engine/Source/Runtime/Engine/Public/MaterialCompiler.h
--- a/Engine/Source/Runtime/Engine/Public/MaterialCompiler.h
+++ b/Engine/Source/Runtime/Engine/Public/MaterialCompiler.h
@@ -439,8 +439,8 @@ public:
 	virtual int32 StrataSingleLayerWaterBSDF(
 		int32 BaseColor, int32 Metallic, int32 Specular, int32 Roughness, int32 EmissiveColor, int32 TopMaterialOpacity, 
 		int32 WaterAlbedo, int32 WaterExtinction, int32 WaterPhaseG, int32 ColorScaleBehindWater, int32 Normal, const FString& SharedLocalBasisIndexMacro) = 0;
-	virtual int32 StrataHorizontalMixing(int32 Foreground, int32 Background, int32 Mix) = 0;
-	virtual int32 StrataHorizontalMixingParameterBlending(int32 Foreground, int32 Background, int32 HorizontalMixCodeChunk, int32 NormalMixCodeChunk, const FString& SharedLocalBasisIndexMacro) = 0;
+	virtual int32 StrataHorizontalMixing(int32 Background, int32 Foreground, int32 Mix) = 0;
+	virtual int32 StrataHorizontalMixingParameterBlending(int32 Background, int32 Foreground, int32 HorizontalMixCodeChunk, int32 NormalMixCodeChunk, const FString& SharedLocalBasisIndexMacro) = 0;
 	virtual int32 StrataVerticalLayering(int32 Top, int32 Base) = 0;
 	virtual int32 StrataVerticalLayeringParameterBlending(int32 Top, int32 Base, const FString& SharedLocalBasisIndexMacro, int32 TopBSDFNormalCodeChunk) = 0;
 	virtual int32 StrataAdd(int32 A, int32 B) = 0;
@@ -979,14 +979,14 @@ public:
 			WaterAlbedo, WaterExtinction, WaterPhaseG, ColorScaleBehindWater, Normal, SharedLocalBasisIndexMacro);
 	}
 
-	virtual int32 StrataHorizontalMixing(int32 Foreground, int32 Background, int32 Mix) override
+	virtual int32 StrataHorizontalMixing(int32 Background, int32 Foreground, int32 Mix) override
 	{
-		return Compiler->StrataHorizontalMixing(Foreground, Background, Mix);
+		return Compiler->StrataHorizontalMixing(Background, Foreground, Mix);
 	}
 
-	virtual int32 StrataHorizontalMixingParameterBlending(int32 Foreground, int32 Background, int32 HorizontalMixCodeChunk, int32 NormalMixCodeChunk, const FString& SharedLocalBasisIndexMacro) override
+	virtual int32 StrataHorizontalMixingParameterBlending(int32 Background, int32 Foreground, int32 HorizontalMixCodeChunk, int32 NormalMixCodeChunk, const FString& SharedLocalBasisIndexMacro) override
 	{
-		return Compiler->StrataHorizontalMixingParameterBlending(Foreground, Background, HorizontalMixCodeChunk, NormalMixCodeChunk, SharedLocalBasisIndexMacro);
+		return Compiler->StrataHorizontalMixingParameterBlending(Background, Foreground, HorizontalMixCodeChunk, NormalMixCodeChunk, SharedLocalBasisIndexMacro);
 	}
 
 	virtual int32 StrataVerticalLayering(int32 Top, int32 Base) override
