Fixed editor crash when strata horizontal blending mix parameter is not specified.

diff --git a/Engine/Source/Runtime/Engine/Private/Materials/MaterialExpressions.cpp b/Engine/Source/Runtime/Engine/Private/Materials/MaterialExpressions.cpp
--- a/Engine/Source/Runtime/Engine/Private/Materials/MaterialExpressions.cpp
+++ b/Engine/Source/Runtime/Engine/Private/Materials/MaterialExpressions.cpp
@@ -20830,7 +20830,7 @@ int32 UMaterialExpressionStrataHorizontalMixing::Compile(class FMaterialCompiler
 	const FStrataMaterialCompilationInfo& ForegroundStrataData = Compiler->GetStrataCompilationInfo(ForegroundCodeChunk);
 	const FStrataMaterialCompilationInfo& BackgroundStrataData = Compiler->GetStrataCompilationInfo(BackgroundCodeChunk);
 
-	const int32 HorizontalMixCodeChunk = Mix.Compile(Compiler);
+	const int32 HorizontalMixCodeChunk = CompileWithDefaultFloat1(Compiler, Mix, 0.5f);
 
 	int32 OutputCodeChunk = INDEX_NONE;
 	if (bUseParameterBlending)
