Use the full asset path when passing the debug name to texture cubes

diff --git a/Engine/Source/Runtime/Engine/Private/TextureCube.cpp b/Engine/Source/Runtime/Engine/Private/TextureCube.cpp
--- a/Engine/Source/Runtime/Engine/Private/TextureCube.cpp
+++ b/Engine/Source/Runtime/Engine/Private/TextureCube.cpp
@@ -424,7 +424,7 @@ public:
 
 		// Create the RHI texture.
 		ETextureCreateFlags TexCreateFlags = (Owner->SRGB ? TexCreate_SRGB : TexCreate_None)  | (Owner->bNotOfflineProcessed ? TexCreate_None : TexCreate_OfflineProcessed);
-		FString Name = Owner->GetName();
+		FString Name = Owner->GetPathName();
 		FRHIResourceCreateInfo CreateInfo(*Name);
 		CreateInfo.ExtData = Owner->PlatformData ? Owner->PlatformData->GetExtData() : 0;
 		TextureCubeRHI = RHICreateTextureCube( Owner->GetSizeX(), Owner->GetPixelFormat(), Owner->GetNumMips(), TexCreateFlags, CreateInfo );
