Crash opening ShooterGame in -game
#fix reverting back to CreateDefaultSubobject vs EditorOnly version to deal with -game editor builds

diff --git a/Engine/Source/Runtime/Engine/Private/Animation/AnimSequenceBase.cpp b/Engine/Source/Runtime/Engine/Private/Animation/AnimSequenceBase.cpp
--- a/Engine/Source/Runtime/Engine/Private/Animation/AnimSequenceBase.cpp
+++ b/Engine/Source/Runtime/Engine/Private/Animation/AnimSequenceBase.cpp
@@ -1268,7 +1268,7 @@ void UAnimSequenceBase::CreateModel()
 {	
 	if (DataModel == nullptr)
 	{
-		DataModel = CreateEditorOnlyDefaultSubobject<UAnimDataModel>(FName(TEXT("AnimationDataModel")));
+		DataModel = CreateDefaultSubobject<UAnimDataModel>(FName(TEXT("AnimationDataModel")));
 	}
 	
 	DataModel->GetModifiedEvent().RemoveAll(this);
