Edigrundo load order fix to unblock stage builds. 1/2

diff --git a/Engine/Source/Runtime/AudioMixer/Private/Components/SynthComponent.cpp b/Engine/Source/Runtime/AudioMixer/Private/Components/SynthComponent.cpp
--- a/Engine/Source/Runtime/AudioMixer/Private/Components/SynthComponent.cpp
+++ b/Engine/Source/Runtime/AudioMixer/Private/Components/SynthComponent.cpp
@@ -141,6 +141,10 @@ USynthComponent::USynthComponent(const FObjectInitializer& ObjectInitializer)
 	bAlwaysPlay = true;
 	Synth = nullptr;
 
+	// Set the default sound class
+	SoundClass = USoundBase::DefaultSoundClassObject;
+	Synth = nullptr;
+
 	PreferredBufferLength = DEFAULT_PROCEDURAL_SOUNDWAVE_BUFFER_SIZE;
 
 #if WITH_EDITORONLY_DATA
diff --git a/Engine/Source/Runtime/AudioMixer/Public/Components/SynthComponent.h b/Engine/Source/Runtime/AudioMixer/Public/Components/SynthComponent.h
--- a/Engine/Source/Runtime/AudioMixer/Public/Components/SynthComponent.h
+++ b/Engine/Source/Runtime/AudioMixer/Public/Components/SynthComponent.h
@@ -253,15 +253,6 @@ public:
 	void AddAudioBufferListener(IAudioBufferListener* InAudioBufferListener);
 	void RemoveAudioBufferListener(IAudioBufferListener* InAudioBufferListener);
 
-	virtual USoundClass* GetSoundClass()
-	{
-		if (SoundClass)
-		{
-			return SoundClass;
-		}
-
-		return USoundBase::GetDefaultSoundClass();
-	}
 
 	virtual void BeginDestroy() override;
 
diff --git a/Engine/Source/Runtime/Engine/Classes/Sound/SoundBase.h b/Engine/Source/Runtime/Engine/Classes/Sound/SoundBase.h
--- a/Engine/Source/Runtime/Engine/Classes/Sound/SoundBase.h
+++ b/Engine/Source/Runtime/Engine/Classes/Sound/SoundBase.h
@@ -8,8 +8,6 @@
 #include "Audio.h"
 #include "CoreMinimal.h"
 #include "IAudioExtensionPlugin.h"
-#include "Sound/AudioSettings.h"
-#include "Sound/SoundClass.h"
 #include "SoundConcurrency.h"
 #include "SoundModulationDestination.h"
 #include "SoundSourceBusSend.h"
@@ -60,6 +58,9 @@ class ENGINE_API USoundBase : public UObject, public IInterface_AssetUserData
 	GENERATED_UCLASS_BODY()
 
 public:
+	static USoundClass* DefaultSoundClassObject;
+	static USoundConcurrency* DefaultSoundConcurrencyObject;
+
 	/** Sound class this sound belongs to */
 	UPROPERTY(EditAnywhere, Category = Sound, meta = (DisplayName = "Class"), AssetRegistrySearchable)
 	TObjectPtr<USoundClass> SoundClassObject;
@@ -188,10 +189,9 @@ public:
 	TArray<TObjectPtr<UAssetUserData>> AssetUserData;
 
 public:
-	static USoundClass* GetDefaultSoundClass();
-	static USoundConcurrency* GetDefaultSoundConcurrency();
 
 	//~ Begin UObject Interface.
+	virtual void PostInitProperties() override;
 #if WITH_EDITORONLY_DATA
 	virtual void PostLoad() override;
 #endif
diff --git a/Engine/Source/Runtime/Engine/Private/ActiveSound.cpp b/Engine/Source/Runtime/Engine/Private/ActiveSound.cpp
--- a/Engine/Source/Runtime/Engine/Private/ActiveSound.cpp
+++ b/Engine/Source/Runtime/Engine/Private/ActiveSound.cpp
@@ -479,24 +479,20 @@ void FActiveSound::GetConcurrencyHandles(TArray<FConcurrencyHandle>& OutConcurre
 {
 	OutConcurrencyHandles.Reset();
 
-	if (ConcurrencySet.Num() > 0)
+	if (!ConcurrencySet.Num() && Sound)
+	{
+		Sound->GetConcurrencyHandles(OutConcurrencyHandles);
+	}
+	else
 	{
 		for (const USoundConcurrency* Concurrency : ConcurrencySet)
 		{
 			if (Concurrency)
 			{
-				OutConcurrencyHandles.Emplace(*Concurrency);
+				OutConcurrencyHandles.Emplace(FConcurrencyHandle(*Concurrency));
 			}
 		}
 	}
-	else if (Sound)
-	{
-		Sound->GetConcurrencyHandles(OutConcurrencyHandles);
-	}
-	else if (USoundConcurrency* DefaultConcurrency = USoundBase::GetDefaultSoundConcurrency())
-	{
-		OutConcurrencyHandles.Emplace(*DefaultConcurrency);
-	}
 }
 
 bool FActiveSound::GetConcurrencyFadeDuration(float& OutFadeDuration) const
diff --git a/Engine/Source/Runtime/Engine/Private/SoundBase.cpp b/Engine/Source/Runtime/Engine/Private/SoundBase.cpp
--- a/Engine/Source/Runtime/Engine/Private/SoundBase.cpp
+++ b/Engine/Source/Runtime/Engine/Private/SoundBase.cpp
@@ -10,6 +10,9 @@
 #include "Engine/AssetUserData.h"
 
 
+USoundClass* USoundBase::DefaultSoundClassObject = nullptr;
+USoundConcurrency* USoundBase::DefaultSoundConcurrencyObject = nullptr;
+
 USoundBase::USoundBase(const FObjectInitializer& ObjectInitializer)
 	: Super(ObjectInitializer)
 	, VirtualizationMode(EVirtualizationMode::Restart)
@@ -26,40 +29,35 @@ USoundBase::USoundBase(const FObjectInitializer& ObjectInitializer)
 	bEnableSubmixSends = true;
 }
 
-USoundClass* USoundBase::GetDefaultSoundClass()
+void USoundBase::PostInitProperties()
 {
-	const FSoftObjectPath DefaultPath = GetDefault<UAudioSettings>()->DefaultSoundClassName;
-	if (DefaultPath.IsValid())
+	Super::PostInitProperties();
+
+	if (USoundBase::DefaultSoundClassObject == nullptr)
 	{
-		UObject* DefaultObject = DefaultPath.TryLoad();
-		if (DefaultObject)
+		const FSoftObjectPath DefaultSoundClassName = GetDefault<UAudioSettings>()->DefaultSoundClassName;
+		if (DefaultSoundClassName.IsValid())
 		{
-			if (USoundClass* DefaultSoundClass = CastChecked<USoundClass>(DefaultObject))
-			{
-				return DefaultSoundClass;
-			}
+			SCOPED_BOOT_TIMING("USoundBase::LoadSoundClass");
+			USoundBase::DefaultSoundClassObject = LoadObject<USoundClass>(nullptr, *DefaultSoundClassName.ToString());
 		}
 	}
+	SoundClassObject = USoundBase::DefaultSoundClassObject;
 
-	return nullptr;
-}
-
-USoundConcurrency* USoundBase::GetDefaultSoundConcurrency()
-{
-	const FSoftObjectPath DefaultPath = GetDefault<UAudioSettings>()->DefaultSoundConcurrencyName;
-	if (DefaultPath.IsValid())
+	if (USoundBase::DefaultSoundConcurrencyObject == nullptr)
 	{
-		UObject* DefaultObject = DefaultPath.TryLoad();
-		if (DefaultObject)
+		const FSoftObjectPath DefaultSoundConcurrencyName = GetDefault<UAudioSettings>()->DefaultSoundConcurrencyName;
+		if (DefaultSoundConcurrencyName.IsValid())
 		{
-			if (USoundConcurrency* DefaultSoundConcurrency = CastChecked<USoundConcurrency>(DefaultObject))
-			{
-				return DefaultSoundConcurrency;
-			}
+			SCOPED_BOOT_TIMING("USoundBase::LoadSoundConcurrency");
+			USoundBase::DefaultSoundConcurrencyObject = LoadObject<USoundConcurrency>(nullptr, *DefaultSoundConcurrencyName.ToString());
 		}
 	}
 
-	return nullptr;
+	if (USoundBase::DefaultSoundConcurrencyObject != nullptr)
+	{
+		ConcurrencySet.Add(USoundBase::DefaultSoundConcurrencyObject);
+	}
 }
 
 bool USoundBase::IsPlayable() const
@@ -137,18 +135,12 @@ bool USoundBase::IsLooping()
 
 bool USoundBase::ShouldApplyInteriorVolumes()
 {
-	USoundClass* SoundClass = GetSoundClass();
-	return SoundClass && SoundClass->Properties.bApplyAmbientVolumes;
+	return (SoundClassObject && SoundClassObject->Properties.bApplyAmbientVolumes);
 }
 
 USoundClass* USoundBase::GetSoundClass() const
 {
-	if (SoundClassObject)
-	{
-		return SoundClassObject;
-	}
-
-	return GetDefaultSoundClass();
+	return SoundClassObject;
 }
 
 USoundSubmixBase* USoundBase::GetSoundSubmix() const
@@ -182,20 +174,13 @@ void USoundBase::GetConcurrencyHandles(TArray<FConcurrencyHandle>& OutConcurrenc
 	}
 	else
 	{
-		if (ConcurrencySet.Num() > 0)
+		for (const USoundConcurrency* Concurrency : ConcurrencySet)
 		{
-			for (const USoundConcurrency* Concurrency : ConcurrencySet)
+			if (Concurrency)
 			{
-				if (Concurrency)
-				{
-					OutConcurrencyHandles.Emplace(*Concurrency);
-				}
+				OutConcurrencyHandles.Emplace(*Concurrency);
 			}
 		}
-		else if (const USoundConcurrency* DefaultConcurrency = GetDefaultSoundConcurrency())
-		{
-			OutConcurrencyHandles.Emplace(*DefaultConcurrency);
-		}
 	}
 }
 
diff --git a/Engine/Source/Runtime/Engine/Private/SoundCue.cpp b/Engine/Source/Runtime/Engine/Private/SoundCue.cpp
--- a/Engine/Source/Runtime/Engine/Private/SoundCue.cpp
+++ b/Engine/Source/Runtime/Engine/Private/SoundCue.cpp
@@ -234,7 +234,7 @@ void USoundCue::PostLoad()
 	
 	ESoundWaveLoadingBehavior SoundClassLoadingBehavior = ESoundWaveLoadingBehavior::Inherited;
 
-	USoundClass* CurrentSoundClass = GetSoundClass();
+	USoundClass* CurrentSoundClass = SoundClassObject;
 
 	// Recurse through this sound class's parents until we find an override.
 	while (SoundClassLoadingBehavior == ESoundWaveLoadingBehavior::Inherited && CurrentSoundClass != nullptr)
diff --git a/Engine/Source/Runtime/Engine/Private/SoundWave.cpp b/Engine/Source/Runtime/Engine/Private/SoundWave.cpp
--- a/Engine/Source/Runtime/Engine/Private/SoundWave.cpp
+++ b/Engine/Source/Runtime/Engine/Private/SoundWave.cpp
@@ -3051,7 +3051,7 @@ void USoundWave::CacheInheritedLoadingBehavior() const
 		// if this is true then the behavior should not be Inherited here
 		check(!*bLoadingBehaviorOverriddenPtr);
 
-		USoundClass* CurrentSoundClass = GetSoundClass();
+		USoundClass* CurrentSoundClass = SoundClassObject;
 		ESoundWaveLoadingBehavior SoundClassLoadingBehavior = ESoundWaveLoadingBehavior::Inherited;
 
 		// Recurse through this sound class's parents until we find an override.
diff --git a/Engine/Source/Runtime/MediaAssets/Private/Assets/MediaSoundComponent.cpp b/Engine/Source/Runtime/MediaAssets/Private/Assets/MediaSoundComponent.cpp
--- a/Engine/Source/Runtime/MediaAssets/Private/Assets/MediaSoundComponent.cpp
+++ b/Engine/Source/Runtime/MediaAssets/Private/Assets/MediaSoundComponent.cpp
@@ -26,6 +26,11 @@ CSV_DECLARE_CATEGORY_MODULE_EXTERN(MEDIA_API, MediaStreaming);
 
 static const int32 MaxAudioInputSamples = 8;	// accept at most these many samples into our input queue
 
+/* Static initialization
+ *****************************************************************************/
+
+USoundClass* UMediaSoundComponent::DefaultMediaSoundClassObject = nullptr;
+
 
 /* UMediaSoundComponent structors
  *****************************************************************************/
@@ -252,6 +257,27 @@ void UMediaSoundComponent::Deactivate()
 /* UObject interface
  *****************************************************************************/
 
+void UMediaSoundComponent::PostInitProperties()
+{
+	Super::PostInitProperties();
+
+	if (UMediaSoundComponent::DefaultMediaSoundClassObject == nullptr)
+	{
+		const FSoftObjectPath DefaultMediaSoundClassName = GetDefault<UAudioSettings>()->DefaultMediaSoundClassName;
+
+		if (DefaultMediaSoundClassName.IsValid())
+		{
+			UMediaSoundComponent::DefaultMediaSoundClassObject = LoadObject<USoundClass>(nullptr, *DefaultMediaSoundClassName.ToString());
+		}
+	}
+
+	// We have a different default sound class object for media sound components
+	if (SoundClass == USoundBase::DefaultSoundClassObject || SoundClass == nullptr)
+	{
+		SoundClass = UMediaSoundComponent::DefaultMediaSoundClassObject;
+	}
+}
+
 
 void UMediaSoundComponent::PostLoad()
 {
diff --git a/Engine/Source/Runtime/MediaAssets/Public/MediaSoundComponent.h b/Engine/Source/Runtime/MediaAssets/Public/MediaSoundComponent.h
--- a/Engine/Source/Runtime/MediaAssets/Public/MediaSoundComponent.h
+++ b/Engine/Source/Runtime/MediaAssets/Public/MediaSoundComponent.h
@@ -135,29 +135,6 @@ public:
 	UFUNCTION(BlueprintCallable, Category="Media|MediaSoundComponent")
 	UMediaPlayer* GetMediaPlayer() const;
 
-	virtual USoundClass* GetSoundClass() override
-	{
-		if (SoundClass)
-		{
-			return SoundClass;
-		}
-
-		const FSoftObjectPath DefaultPath = GetDefault<UAudioSettings>()->DefaultMediaSoundClassName;
-		if (DefaultPath.IsValid())
-		{
-			UObject* DefaultObject = DefaultPath.TryLoad();
-			if (DefaultObject)
-			{
-				if (USoundClass* DefaultSoundClass = CastChecked<USoundClass>(DefaultObject))
-				{
-					return DefaultSoundClass;
-				}
-			}
-		}
-
-		return USoundBase::GetDefaultSoundClass();
-	}
-
 	/**
 	 * Set the media player that provides the audio samples.
 	 *
@@ -233,6 +210,7 @@ public:
 public:
 
 	//~ UObject interface
+	virtual void PostInitProperties() override;
 	virtual void PostLoad() override;
 
 #if WITH_EDITOR
@@ -324,4 +302,8 @@ private:
 
 	/** Whether or not envelope follower settings changed. */
 	bool bEnvelopeFollowerSettingsChanged;
+
+private:
+
+	static USoundClass* DefaultMediaSoundClassObject;
 };
