Speculative fix streaming buffer size crash

diff --git a/Engine/Source/Runtime/Engine/Private/Rendering/StreamableTextureResource.cpp b/Engine/Source/Runtime/Engine/Private/Rendering/StreamableTextureResource.cpp
--- a/Engine/Source/Runtime/Engine/Private/Rendering/StreamableTextureResource.cpp
+++ b/Engine/Source/Runtime/Engine/Private/Rendering/StreamableTextureResource.cpp
@@ -36,7 +36,7 @@ bool CanCreateWithPartiallyResidentMips(ETextureCreateFlags TexCreateFlags)
 	const ETextureCreateFlags iRequiredFlags =
 		TexCreate_OfflineProcessed;
 
-	return EnumHasAllFlags(TexCreateFlags & iDisableFlags, iRequiredFlags) && CVarVirtualTextureEnabled.GetValueOnAnyThread();
+	return ((TexCreateFlags & (iDisableFlags | iRequiredFlags)) == iRequiredFlags) && CVarVirtualTextureEnabled.GetValueOnAnyThread();
 	
 #else
 	return false;
