Fixed a test that needs to be ordered differently to be able to plug MakeMaterialAttribute node with strata material into it, straight into a root node.

diff --git a/Engine/Source/Runtime/Engine/Private/Materials/MaterialExpressions.cpp b/Engine/Source/Runtime/Engine/Private/Materials/MaterialExpressions.cpp
--- a/Engine/Source/Runtime/Engine/Private/Materials/MaterialExpressions.cpp
+++ b/Engine/Source/Runtime/Engine/Private/Materials/MaterialExpressions.cpp
@@ -1169,13 +1169,15 @@ uint32 UMaterialExpression::GetInputType(int32 InputIndex)
 uint32 UMaterialExpression::GetOutputType(int32 OutputIndex)
 {
 	// different outputs should be defined by sub classed expressions 
-	if (IsResultStrataMaterial(OutputIndex))
+
+	// Material attributes need to be tested first to work when plugged in main root node (to not return MCT_Strata when Strata mateiral is fed)
+	if (IsResultMaterialAttributes(OutputIndex))
 	{
-		return MCT_Strata;
+		return MCT_MaterialAttributes;
 	}
-	else if (IsResultMaterialAttributes(OutputIndex))
+	else if (IsResultStrataMaterial(OutputIndex))
 	{
-		return MCT_MaterialAttributes;
+		return MCT_Strata;
 	}
 	else
 	{
