Move the member under the private scope of the class
Fix non-unity build

#rnx

diff --git a/Engine/Source/Runtime/Engine/Private/ShaderCompiler/ShaderCompiler.cpp b/Engine/Source/Runtime/Engine/Private/ShaderCompiler/ShaderCompiler.cpp
--- a/Engine/Source/Runtime/Engine/Private/ShaderCompiler/ShaderCompiler.cpp
+++ b/Engine/Source/Runtime/Engine/Private/ShaderCompiler/ShaderCompiler.cpp
@@ -31,6 +31,7 @@
 #include "Misc/CoreDelegates.h"
 #include "EditorSupportDelegates.h"
 #include "GlobalShader.h"
+#include "SceneManagement.h"
 #include "Interfaces/IShaderFormat.h"
 #include "Interfaces/IShaderFormatModule.h"
 #include "RendererInterface.h"
diff --git a/Engine/Source/Runtime/Engine/Public/ShaderCompiler.h b/Engine/Source/Runtime/Engine/Public/ShaderCompiler.h
--- a/Engine/Source/Runtime/Engine/Public/ShaderCompiler.h
+++ b/Engine/Source/Runtime/Engine/Public/ShaderCompiler.h
@@ -934,6 +934,9 @@ private:
 	/** Opt out of material shader compilation and instead place an empty shader map. */
 	bool bNoShaderCompilation;
 
+	/** Used to show a notification accompanying progress. */
+	FAsyncCompilationNotification Notification;
+
 	/** Launches the worker, returns the launched process handle. */
 	FProcHandle LaunchWorker(const FString& WorkingDirectory, uint32 ProcessId, uint32 ThreadId, const FString& WorkerInputFile, const FString& WorkerOutputFile);
 
@@ -1119,8 +1122,6 @@ public:
 	 * Returns true if the given shader compile worker is still running.
 	 */
 	static bool IsShaderCompilerWorkerRunning(FProcHandle & WorkerHandle);
-
-	FAsyncCompilationNotification Notification;
 };
 
 /** The global shader compiling thread manager. */
