Fix level streaming code to ignore existing level packages that are empty.

diff --git a/Engine/Source/Runtime/Engine/Private/LevelStreaming.cpp b/Engine/Source/Runtime/Engine/Private/LevelStreaming.cpp
--- a/Engine/Source/Runtime/Engine/Private/LevelStreaming.cpp
+++ b/Engine/Source/Runtime/Engine/Private/LevelStreaming.cpp
@@ -1033,7 +1033,21 @@ bool ULevelStreaming::RequestLevel(UWorld* PersistentWorld, bool bAllowLevelLoad
 	int32 PIEInstanceID = INDEX_NONE;
 
 	// Try to find the [to be] loaded package.
+	UWorld* World = nullptr;
 	UPackage* LevelPackage = (UPackage*)StaticFindObjectFast(UPackage::StaticClass(), nullptr, DesiredPackageName, 0, 0, RF_NoFlags, EInternalObjectFlags::PendingKill);
+	
+	if (LevelPackage)
+	{
+		// Find world object and use its PersistentLevel pointer.
+		World = UWorld::FindWorldInPackage(LevelPackage);
+
+		// Check for a redirector. Follow it, if found.
+		if (!World)
+		{
+			World = UWorld::FollowWorldRedirectorInPackage(LevelPackage);
+			LevelPackage = World ? World->GetOutermost() : nullptr;
+		}
+	}
 
 	// copy streaming level on demand if we are in PIE
 	// (the world is already loaded for the editor, just find it and copy it)
@@ -1083,19 +1097,6 @@ bool ULevelStreaming::RequestLevel(UWorld* PersistentWorld, bool bAllowLevelLoad
 	// Package is already or still loaded.
 	if (LevelPackage)
 	{
-		// Find world object and use its PersistentLevel pointer.
-		UWorld* World = UWorld::FindWorldInPackage(LevelPackage);
-
-		// Check for a redirector. Follow it, if found.
-		if (!World)
-		{
-			World = UWorld::FollowWorldRedirectorInPackage(LevelPackage);
-			if (World)
-			{
-				LevelPackage = World->GetOutermost();
-			}
-		}
-
 		if (World != nullptr)
 		{
 			if (World->IsPendingKill())
@@ -1177,7 +1178,6 @@ bool ULevelStreaming::RequestLevel(UWorld* PersistentWorld, bool bAllowLevelLoad
 
 				// When loading an instanced package we need to build an instancing context in case non external actors part of the level are 
 				// pulling on external actors.
-				UWorld* World = GetTypedOuter<UWorld>();
 				FString ExternalActorsPath = ULevel::GetExternalActorsPath(PackagePath.GetPackageName());
 				TArray<FString> ActorPackageNames = ULevel::GetOnDiskExternalActorPackages(ExternalActorsPath);
 								
