Add compute.RebuildComputeGraphs console command for live updating compute framework.

diff --git a/Engine/Source/Runtime/Engine/Private/ComputeFramework/ComputeFramework.cpp b/Engine/Source/Runtime/Engine/Private/ComputeFramework/ComputeFramework.cpp
--- a/Engine/Source/Runtime/Engine/Private/ComputeFramework/ComputeFramework.cpp
+++ b/Engine/Source/Runtime/Engine/Private/ComputeFramework/ComputeFramework.cpp
@@ -2,6 +2,11 @@
 
 #include "ComputeFramework/ComputeFramework.h"
 
+#include "ComputeFramework/ComputeGraph.h"
+#include "ComputeFramework/ComputeKernelFromText.h"
+#include "ShaderCore.h"
+#include "UObject/UObjectIterator.h"
+
 DEFINE_LOG_CATEGORY(LogComputeFramework);
 
 static int32 GComputeFrameworkMode = 1;
@@ -14,9 +19,33 @@ static FAutoConsoleVariableRef CVarComputeFrameworkMode(
 	ECVF_Scalability | ECVF_RenderThreadSafe
 	);
 
-bool SupportsComputeFramework(ERHIFeatureLevel::Type FeatureLevel, EShaderPlatform ShaderPlatform)
+FAutoConsoleCommand CmdRebuildComputeGraphs(
+	TEXT("compute.RebuildComputeGraphs"),
+	TEXT("Force all loaded UComputeGraph objects to rebuild."),
+	FConsoleCommandDelegate::CreateStatic(ComputeFramework::RebuildComputeGraphs)
+);
+
+namespace ComputeFramework
 {
-	return GComputeFrameworkMode > 0
-		&& FeatureLevel >= ERHIFeatureLevel::SM5
-		&& FDataDrivenShaderPlatformInfo::GetSupportsComputeFramework(ShaderPlatform);
+	bool IsEnabled(ERHIFeatureLevel::Type FeatureLevel, EShaderPlatform ShaderPlatform)
+	{
+		return GComputeFrameworkMode > 0
+			&& FDataDrivenShaderPlatformInfo::GetSupportsComputeFramework(ShaderPlatform);
+	}
+
+	void RebuildComputeGraphs()
+	{
+#if WITH_EDITOR
+		FlushShaderFileCache();
+
+		for (TObjectIterator<UComputeKernelFromText> It; It; ++It)
+		{
+			It->ReparseKernelSourceText();
+		}
+		for (TObjectIterator<UComputeGraph> It; It; ++It)
+		{
+			It->UpdateResources();
+		}
+#endif // WITH_EDITOR
+	}
 }
diff --git a/Engine/Source/Runtime/Engine/Private/ComputeFramework/ComputeGraph.cpp b/Engine/Source/Runtime/Engine/Private/ComputeFramework/ComputeGraph.cpp
--- a/Engine/Source/Runtime/Engine/Private/ComputeFramework/ComputeGraph.cpp
+++ b/Engine/Source/Runtime/Engine/Private/ComputeFramework/ComputeGraph.cpp
@@ -126,8 +126,11 @@ FShaderParametersMetadata* UComputeGraph::BuildKernelShaderMetadata(int32 Kernel
 	for (int32 DataProviderIndex : DataProviderIndices)
 	{
 		UComputeDataInterface* DataInterface = DataInterfaces[DataProviderIndex];
-		TCHAR const* UID = GetDataInterfaceUID(DataProviderIndex);
-		DataInterface->GetShaderParameters(UID, Builder);
+		if (DataInterface != nullptr)
+		{
+			TCHAR const* UID = GetDataInterfaceUID(DataProviderIndex);
+			DataInterface->GetShaderParameters(UID, Builder);
+		}
 	}
 
 	return Builder.Build(FShaderParametersMetadata::EUseCase::ShaderParameterStruct, *GetName());
@@ -218,12 +221,14 @@ FString UComputeGraph::BuildKernelSource(int32 KernelIndex) const
 			for (int32 DataProviderIndex : DataProviderIndices)
 			{
 				UComputeDataInterface* DataInterface = DataInterfaces[DataProviderIndex];
-
-				// Add a unique prefix to generate unique names in the data interface shader code.
-				TCHAR const* UID = GetDataInterfaceUID(DataProviderIndex);
-				HLSL += FString::Printf(TEXT("#define DI_UID %s_\n"), UID);
-				DataInterface->GetHLSL(HLSL);
-				HLSL += TEXT("#undef DI_UID\n");
+				if (DataInterface != nullptr)
+				{
+					// Add a unique prefix to generate unique names in the data interface shader code.
+					TCHAR const* UID = GetDataInterfaceUID(DataProviderIndex);
+					HLSL += FString::Printf(TEXT("#define DI_UID %s_\n"), UID);
+					DataInterface->GetHLSL(HLSL);
+					HLSL += TEXT("#undef DI_UID\n");
+				}
 			}
 
 			// Bind every external kernel function to the associated data input function.
@@ -273,7 +278,7 @@ void UComputeGraph::CacheResourceShadersForRendering(uint32 CompilationFlags)
 
 			FString ShaderEntryPoint = Kernel->KernelSource->GetEntryPoint();
 			FString ShaderSource = BuildKernelSource(KernelIndex);
-			uint64 ShaderSourceHash = FCrc::TypeCrc32(*ShaderSource, 0);
+			uint64 ShaderSourceHash = FCrc::TypeCrc32(*ShaderSource, Kernel->KernelSource->GetSourceCodeHash());
 			FShaderParametersMetadata* ShaderMetadata = BuildKernelShaderMetadata(KernelIndex);
 
 			const ERHIFeatureLevel::Type CacheFeatureLevel = GMaxRHIFeatureLevel;
diff --git a/Engine/Source/Runtime/Engine/Public/ComputeFramework/ComputeFramework.h b/Engine/Source/Runtime/Engine/Public/ComputeFramework/ComputeFramework.h
--- a/Engine/Source/Runtime/Engine/Public/ComputeFramework/ComputeFramework.h
+++ b/Engine/Source/Runtime/Engine/Public/ComputeFramework/ComputeFramework.h
@@ -7,5 +7,11 @@
 
 DECLARE_LOG_CATEGORY_EXTERN(LogComputeFramework, Log, All);
 
-/** Returns true if ComputeFramework is supported. */
-extern ENGINE_API bool SupportsComputeFramework(ERHIFeatureLevel::Type FeatureLevel, EShaderPlatform ShaderPlatform);
+namespace ComputeFramework
+{
+	/** Returns true if ComputeFramework is supported. */
+	ENGINE_API bool IsEnabled(ERHIFeatureLevel::Type FeatureLevel, EShaderPlatform ShaderPlatform);
+
+	/** Rebuild all compute graphs. */
+	ENGINE_API void RebuildComputeGraphs();
+}
diff --git a/Engine/Source/Runtime/Engine/Public/ComputeFramework/ComputeKernelFromText.h b/Engine/Source/Runtime/Engine/Public/ComputeFramework/ComputeKernelFromText.h
--- a/Engine/Source/Runtime/Engine/Public/ComputeFramework/ComputeKernelFromText.h
+++ b/Engine/Source/Runtime/Engine/Public/ComputeFramework/ComputeKernelFromText.h
@@ -33,6 +33,11 @@ public:
 	UPROPERTY()
 	uint64 SourceHash = 0;
 
+#if WITH_EDITOR
+	/** Parse the kernel source to get the kernel external functions and other data. */
+	void ReparseKernelSourceText();
+#endif
+
 protected:
 	//~ Begin UComputeKernelSource Interface.
 	FString GetEntryPoint() const override
@@ -65,9 +70,6 @@ protected:
 
 private:
 #if WITH_EDITOR
-	/** Parse the kernel source to get the kernel external functions and other data. */
-	void ReparseKernelSourceText();
-
 	FString KernelSourceText;
 	FFilePath PrevSourceFile;
 #endif
diff --git a/Engine/Source/Runtime/Renderer/Private/RendererScene.cpp b/Engine/Source/Runtime/Renderer/Private/RendererScene.cpp
--- a/Engine/Source/Runtime/Renderer/Private/RendererScene.cpp
+++ b/Engine/Source/Runtime/Renderer/Private/RendererScene.cpp
@@ -1083,7 +1083,7 @@ FScene::FScene(UWorld* InWorld, bool bInRequiresHitProxies, bool bInIsEditorScen
 		GPUSkinCache = new FGPUSkinCache(InFeatureLevel, bRequiresMemoryLimit);
 	}
 
-	if (SupportsComputeFramework(InFeatureLevel, GetFeatureLevelShaderPlatform(InFeatureLevel)))
+	if (ComputeFramework::IsEnabled(InFeatureLevel, GetFeatureLevelShaderPlatform(InFeatureLevel)))
 	{
 		ComputeGraphScheduler = new FComputeGraphScheduler();
 	}
