Fix mistaken change of bIncludeOnlyOnDiskAssets from false to true in UAssetRegistryImpl::GetCachedAssetDataForObjectPath. This was causing out of date assetregistry information to be used for PrimaryAssets modified in the editor.
#rnx

diff --git a/Engine/Source/Runtime/Engine/Private/AssetManager.cpp b/Engine/Source/Runtime/Engine/Private/AssetManager.cpp
--- a/Engine/Source/Runtime/Engine/Private/AssetManager.cpp
+++ b/Engine/Source/Runtime/Engine/Private/AssetManager.cpp
@@ -1265,7 +1265,7 @@ bool UAssetManager::GetPrimaryAssetData(const FPrimaryAssetId& PrimaryAssetId, F
 
 	if (NameData)
 	{
-		FAssetData CachedAssetData = GetAssetRegistry().GetAssetByObjectPath(NameData->AssetDataPath, true /* bIncludeOnlyOnDiskAssets */);
+		FAssetData CachedAssetData = GetAssetRegistry().GetAssetByObjectPath(NameData->AssetDataPath, false /* bIncludeOnlyOnDiskAssets */);
 
 		if (CachedAssetData.IsValid())
 		{
@@ -1288,7 +1288,7 @@ bool UAssetManager::GetPrimaryAssetDataList(FPrimaryAssetType PrimaryAssetType,
 
 		for (const TPair<FName, FPrimaryAssetData>& Pair : TypeData.AssetMap)
 		{
-			FAssetData CachedAssetData = Registry.GetAssetByObjectPath(Pair.Value.AssetDataPath, true /* bIncludeOnlyOnDiskAssets */);
+			FAssetData CachedAssetData = Registry.GetAssetByObjectPath(Pair.Value.AssetDataPath, false /* bIncludeOnlyOnDiskAssets */);
 
 			if (CachedAssetData.IsValid())
 			{
