AudioThread replaced by a pipe, with an ability to switch back to audio thread by conditional compilation (UE_AUDIO_THREAD_AS_PIPE)

diff --git a/Engine/Source/Runtime/Core/Private/Async/TaskGraph.cpp b/Engine/Source/Runtime/Core/Private/Async/TaskGraph.cpp
--- a/Engine/Source/Runtime/Core/Private/Async/TaskGraph.cpp
+++ b/Engine/Source/Runtime/Core/Private/Async/TaskGraph.cpp
@@ -57,6 +57,10 @@ static int32 GNumWorkerThreadsToIgnore = 0;
 	#define YIELD_BETWEEN_TASKS 0
 #endif
 
+#if !defined(UE_TASKGRAPH_BUSYWAIT)
+	#define UE_TASKGRAPH_BUSYWAIT 0
+#endif
+
 namespace ENamedThreads
 {
 	CORE_API TAtomic<Type> FRenderThreadStatics::RenderThread(ENamedThreads::GameThread); // defaults to game and is set and reset by the render thread itself
@@ -825,7 +829,9 @@ private:
 		{
 			return TEXT("RHI Thread");
 		}
+PRAGMA_DISABLE_DEPRECATION_WARNINGS
 		else if (ThreadId == ENamedThreads::AudioThread)
+PRAGMA_ENABLE_DEPRECATION_WARNINGS
 		{
 			return TEXT("Audio Thread");
 		}
@@ -1407,6 +1413,30 @@ PRAGMA_ENABLE_DEPRECATION_WARNINGS
 #endif
 	}
 
+#if UE_AUDIO_THREAD_AS_PIPE
+
+	// AudioThread was removed and replaced by a pipe. While all engine code was converted where's a tiny possibilty that 
+	// an external code is using it directly. Redirect all audio tasks to the pipe
+	static void RedirectAudioTasksToPipe(FBaseGraphTask* Task)
+	{
+PRAGMA_DISABLE_DEPRECATION_WARNINGS
+		if (ENamedThreads::GetThreadIndex(Task->ThreadToExecuteOn) == ENamedThreads::AudioThread)
+PRAGMA_ENABLE_DEPRECATION_WARNINGS
+		{
+			extern CORE_API UE::Tasks::FPipe GAudioPipe;
+			GAudioPipe.Launch(UE_SOURCE_LOCATION,
+				[Task]
+				{
+					TArray<FBaseGraphTask*> Dummy;
+					Task->Execute(Dummy, ENamedThreads::AnyThread, true);
+				}
+			);
+			return;
+		}
+	}
+
+#endif
+
 	// API inherited from FTaskGraphInterface
 
 	/** 
@@ -1420,6 +1450,9 @@ PRAGMA_ENABLE_DEPRECATION_WARNINGS
 		TASKGRAPH_SCOPE_CYCLE_COUNTER(2, STAT_TaskGraph_QueueTask);
 
 		RedirectStatsTasksToPipe(Task);
+#if UE_AUDIO_THREAD_AS_PIPE
+		RedirectAudioTasksToPipe(Task);
+#endif
 
 		if (ENamedThreads::GetThreadIndex(ThreadToExecuteOn) == ENamedThreads::AnyThread)
 		{
@@ -2077,6 +2110,7 @@ private:
 			TGraphTask<FReturnGraphTask>::CreateTask(&Tasks, CurrentThread).ConstructAndDispatchWhenReady(CurrentThread);
 			ProcessThreadUntilRequestReturn(CurrentThread);
 		}
+#if UE_TASKGRAPH_BUSYWAIT
 		else if (LowLevelTasks::FScheduler::Get().IsWorkerThread())
 		{
 			LowLevelTasks::BusyWaitUntil([Index(0), &Tasks]() mutable
@@ -2092,6 +2126,7 @@ private:
 				return true;
 			});
 		}
+#endif
 		else
 		{
 			if (!FTaskGraphInterface::IsMultithread())
@@ -2571,10 +2606,12 @@ PRAGMA_ENABLE_DEPRECATION_WARNINGS
 	{
 		Tasks.Add(TGraphTask<FBroadcastTask>::CreateTask().ConstructAndDispatchWhenReady(Callback, StartTime, TEXT("RT"), ENamedThreads::SetTaskPriority(RenderThread, ENamedThreads::HighTaskPriority), nullptr, nullptr, nullptr));
 	}
+PRAGMA_DISABLE_DEPRECATION_WARNINGS
 	if (FTaskGraphInterface::Get().IsThreadProcessingTasks(ENamedThreads::AudioThread))
 	{
 		Tasks.Add(TGraphTask<FBroadcastTask>::CreateTask().ConstructAndDispatchWhenReady(Callback, StartTime, TEXT("AudioT"), ENamedThreads::SetTaskPriority(ENamedThreads::AudioThread, ENamedThreads::HighTaskPriority), nullptr, nullptr, nullptr));
 	}
+PRAGMA_ENABLE_DEPRECATION_WARNINGS
 
 	Callback(ENamedThreads::GameThread_Local);
 
diff --git a/Engine/Source/Runtime/Core/Private/HAL/ThreadingBase.cpp b/Engine/Source/Runtime/Core/Private/HAL/ThreadingBase.cpp
--- a/Engine/Source/Runtime/Core/Private/HAL/ThreadingBase.cpp
+++ b/Engine/Source/Runtime/Core/Private/HAL/ThreadingBase.cpp
@@ -13,6 +13,7 @@
 #include "HAL/PlatformStackWalk.h"
 #include "ProfilingDebugging/MiscTrace.h"
 #include "Async/Fundamental/Scheduler.h"
+#include "Tasks/Pipe.h"
 
 #include <atomic>
 
@@ -193,34 +194,48 @@ CORE_API bool IsInSlateThread()
 	return newValue;
 }
 
-CORE_API TAtomic<bool> GIsAudioThreadSuspended(false);
+// tasks pipe that arranges audio tasks execution one after another so no synchronisation between them is required
+CORE_API UE::Tasks::FPipe GAudioPipe{ TEXT("AudioPipe") };
+// True if async audio processing is enabled and started
+CORE_API std::atomic<bool> GIsAudioThreadRunning{ false };
 
+CORE_API std::atomic<bool> GIsAudioThreadSuspended{ false };
+
+#if !UE_AUDIO_THREAD_AS_PIPE
 CORE_API FRunnableThread* GAudioThread = nullptr;
+#endif
+
 
 CORE_API bool IsAudioThreadRunning()
 {
-PRAGMA_DISABLE_DEPRECATION_WARNINGS
-	return (GAudioThread != nullptr) && !GIsAudioThreadSuspended.Load(EMemoryOrder::Relaxed);
-PRAGMA_ENABLE_DEPRECATION_WARNINGS
+	return GIsAudioThreadRunning.load(std::memory_order_acquire) && !GIsAudioThreadSuspended.load(std::memory_order_acquire);
 }
 
 CORE_API bool IsInAudioThread()
 {
+#if UE_AUDIO_THREAD_AS_PIPE
+
+	return GIsAudioThreadRunning.load(std::memory_order_acquire) && !GIsAudioThreadSuspended.load(std::memory_order_acquire) ? GAudioPipe.IsInContext() : IsInGameThread();
+
+#else
+
 PRAGMA_DISABLE_DEPRECATION_WARNINGS
 	// True if this is the audio thread or if there is no audio thread, then if it is the game thread
-	bool newValue = (nullptr == GAudioThread || GIsAudioThreadSuspended.Load(EMemoryOrder::Relaxed))
+	bool newValue = (nullptr == GAudioThread || GIsAudioThreadSuspended.load(std::memory_order_relaxed))
 		? FTaskTagScope::IsCurrentTag(ETaskTag::EGameThread)
 		: FTaskTagScope::IsCurrentTag(ETaskTag::EAudioThread);
 #if !UE_BUILD_SHIPPING && !UE_BUILD_TEST
 	if (!LowLevelTasks::FScheduler::IsBusyWaiting())
 	{
-		bool oldValue = FPlatformTLS::GetCurrentThreadId() == ((nullptr == GAudioThread || GIsAudioThreadSuspended.Load(EMemoryOrder::Relaxed)) ? GGameThreadId : GAudioThread->GetThreadID());
+		bool oldValue = FPlatformTLS::GetCurrentThreadId() == ((nullptr == GAudioThread || GIsAudioThreadSuspended.load(std::memory_order_relaxed)) ? GGameThreadId : GAudioThread->GetThreadID());
 		ensureMsgf(oldValue == newValue, TEXT("oldValue(%i) newValue(%i) If this check fails make sure that there is a FTaskTagScope(ETaskTag::EAudioThread) as deep as possible on the current callstack, you can see the current value in ActiveNamedThreads(%x)"), oldValue, newValue, FTaskTagScope::GetCurrentTag());
 		newValue = oldValue;
 	}
 #endif
 	return newValue;
 PRAGMA_ENABLE_DEPRECATION_WARNINGS
+
+#endif
 }
 
 CORE_API TAtomic<int32> GIsRenderingThreadSuspended(0);
diff --git a/Engine/Source/Runtime/Core/Private/Tests/Async/TaskGraphTest.cpp b/Engine/Source/Runtime/Core/Private/Tests/Async/TaskGraphTest.cpp
--- a/Engine/Source/Runtime/Core/Private/Tests/Async/TaskGraphTest.cpp
+++ b/Engine/Source/Runtime/Core/Private/Tests/Async/TaskGraphTest.cpp
@@ -696,9 +696,9 @@ namespace OldTaskGraphTests
 	);
 
 
-	IMPLEMENT_SIMPLE_AUTOMATION_TEST(FOldBenchmark, "System.Core.Async.TaskGraph.OldBenchmark", EAutomationTestFlags::EditorContext | EAutomationTestFlags::ClientContext | EAutomationTestFlags::ServerContext | EAutomationTestFlags::EngineFilter);
+	IMPLEMENT_SIMPLE_AUTOMATION_TEST(FTaskGraphOldBenchmark, "System.Core.Async.TaskGraph.OldBenchmark", EAutomationTestFlags::EditorContext | EAutomationTestFlags::ClientContext | EAutomationTestFlags::ServerContext | EAutomationTestFlags::EngineFilter);
 
-	bool FOldBenchmark::RunTest(const FString& Parameters)
+	bool FTaskGraphOldBenchmark::RunTest(const FString& Parameters)
 	{
 		TArray<FString> Args;
 		TaskGraphBenchmark(Args);
@@ -716,9 +716,9 @@ namespace OldTaskGraphTests
 
 namespace TaskGraphTests
 {
-	IMPLEMENT_SIMPLE_AUTOMATION_TEST(FGraphEventTest, "System.Core.Async.TaskGraph.GraphEventTest", EAutomationTestFlags::ApplicationContextMask | EAutomationTestFlags::EngineFilter);
+	IMPLEMENT_SIMPLE_AUTOMATION_TEST(FTaskGraphGraphEventTest, "System.Core.Async.TaskGraph.GraphEventTest", EAutomationTestFlags::ApplicationContextMask | EAutomationTestFlags::EngineFilter);
 
-	bool FGraphEventTest::RunTest(const FString& Parameters)
+	bool FTaskGraphGraphEventTest::RunTest(const FString& Parameters)
 	{
 		{	// task completes before it's waited for
 			FGraphEventRef Event = FFunctionGraphTask::CreateAndDispatchWhenReady(
@@ -914,7 +914,9 @@ PRAGMA_ENABLE_DEPRECATION_WARNINGS
 		if (IsAudioThreadRunning())
 		{	// AudioThread
 			bool bExecuted = false;
+PRAGMA_DISABLE_DEPRECATION_WARNINGS
 			FFunctionGraphTask::CreateAndDispatchWhenReady([&bExecuted] { bExecuted = true; }, TStatId{}, nullptr, ENamedThreads::AudioThread)->Wait();
+PRAGMA_ENABLE_DEPRECATION_WARNINGS
 			check(bExecuted);
 		}
 
diff --git a/Engine/Source/Runtime/Core/Public/Async/ParallelFor.h b/Engine/Source/Runtime/Core/Public/Async/ParallelFor.h
--- a/Engine/Source/Runtime/Core/Public/Async/ParallelFor.h
+++ b/Engine/Source/Runtime/Core/Public/Async/ParallelFor.h
@@ -231,7 +231,9 @@ namespace ParallelForImpl
 				ETaskTag::EStaticInit | 
 				ETaskTag::EGameThread | 
 				ETaskTag::ESlateThread | 
+#if !UE_AUDIO_THREAD_AS_PIPE
 				ETaskTag::EAudioThread | 
+#endif
 				ETaskTag::ERenderingThread | 
 				ETaskTag::ERhiThread;
 
diff --git a/Engine/Source/Runtime/Core/Public/Async/TaskGraphInterfaces.h b/Engine/Source/Runtime/Core/Public/Async/TaskGraphInterfaces.h
--- a/Engine/Source/Runtime/Core/Public/Async/TaskGraphInterfaces.h
+++ b/Engine/Source/Runtime/Core/Public/Async/TaskGraphInterfaces.h
@@ -57,10 +57,15 @@ namespace ENamedThreads
 		UnusedAnchor = -1,
 		/** The always-present, named threads are listed next **/
 #if STATS
-		StatsThread UE_DEPRECATED(5.0, "`StatsThread` has been removed. Stats system should be used by its public API in `Stats2.h`"),
+		StatsThread UE_DEPRECATED(5.0, "`StatsThread` has been removed. Stats system should be used by Stats public API"),
 #endif
 		RHIThread,
-		AudioThread,
+#if UE_AUDIO_THREAD_AS_PIPE
+		AudioThread UE_DEPRECATED(5.0, "`AudioThread` has been removed. Please use `FAudioThread` API")
+#else
+		AudioThread
+#endif
+		,
 		GameThread,
 		// The render thread is sometimes the game thread and is sometimes the actual rendering thread
 		ActualRenderingThread = GameThread + 1,
diff --git a/Engine/Source/Runtime/Core/Public/CoreGlobals.h b/Engine/Source/Runtime/Core/Public/CoreGlobals.h
--- a/Engine/Source/Runtime/Core/Public/CoreGlobals.h
+++ b/Engine/Source/Runtime/Core/Public/CoreGlobals.h
@@ -10,6 +10,8 @@
 #include "Templates/Atomic.h"
 #include "ProfilingDebugging/CpuProfilerTrace.h"
 
+#include <atomic>
+
 class Error;
 class FConfigCacheIni;
 class FFixedUObjectArray;
@@ -477,13 +479,6 @@ extern CORE_API uint32 GRenderThreadId;
 /** Thread ID of the slate thread, if any */
 extern CORE_API uint32 GSlateLoadingThreadId;
 
-/** Thread ID of the audio thread, if any */
-UE_DEPRECATED(4.26, "Please use `IsAudioThreadRunning()` or `IsInAudioThread()`")
-extern CORE_API uint32 GAudioThreadId;
-
-/** Whether the audio thread is suspended */
-extern CORE_API TAtomic<bool> GIsAudioThreadSuspended;
-
 /** Has GGameThreadId been set yet? */
 extern CORE_API bool GIsGameThreadIdInitialized;
 
@@ -543,7 +538,11 @@ enum class ETaskTag : int32
 	EStaticInit					= 1 << 0,
 	EGameThread					= 1 << 1,
 	ESlateThread				= 1 << 2,
-	EAudioThread				= 1 << 3,
+#if UE_AUDIO_THREAD_AS_PIPE	
+	EAudioThread UE_DEPRECATED(5.0, "AudioThread was removed and ETaskTag::EAudioThread is not used anymore. Please remove it.") = 1 << 3,
+#else
+	EAudioThread = 1 << 3,
+#endif
 	ERenderingThread			= 1 << 4,
 	ERhiThread					= 1 << 5,
 	EAsyncLoadingThread			= 1 << 6,
@@ -640,10 +639,14 @@ extern CORE_API bool IsAudioThreadRunning();
 /** @return True if called from the audio thread, and not merely a thread calling audio functions. */
 extern CORE_API bool IsInAudioThread();
 
+#if !UE_AUDIO_THREAD_AS_PIPE
+
 /** Thread used for audio */
 UE_DEPRECATED(4.26, "Please use `IsAudioThreadRunning()` or `IsInAudioThread()`")
 extern CORE_API FRunnableThread* GAudioThread;
 
+#endif
+
 /** @return True if called from the slate thread, and not merely a thread calling slate functions. */
 extern CORE_API bool IsInSlateThread();
 
diff --git a/Engine/Source/Runtime/Core/Public/Misc/CoreDefines.h b/Engine/Source/Runtime/Core/Public/Misc/CoreDefines.h
--- a/Engine/Source/Runtime/Core/Public/Misc/CoreDefines.h
+++ b/Engine/Source/Runtime/Core/Public/Misc/CoreDefines.h
@@ -65,3 +65,8 @@ Can be overridden by setting to 1 or 0 in the project's .Target.cs files.
 #ifndef NO_CVARS
 	#define NO_CVARS 0
 #endif
+
+#if !defined(UE_AUDIO_THREAD_AS_PIPE)
+	#define UE_AUDIO_THREAD_AS_PIPE 1
+#endif
+
diff --git a/Engine/Source/Runtime/Engine/Private/AudioThread.cpp b/Engine/Source/Runtime/Engine/Private/AudioThread.cpp
--- a/Engine/Source/Runtime/Engine/Private/AudioThread.cpp
+++ b/Engine/Source/Runtime/Engine/Private/AudioThread.cpp
@@ -12,11 +12,17 @@
 #include "UObject/UObjectGlobals.h"
 #include "Audio.h"
 #include "HAL/LowLevelMemTracker.h"
+#include "Async/Async.h"
+#include "Tasks/Pipe.h"
 
 //
 // Globals
 //
 
+extern CORE_API UE::Tasks::FPipe GAudioPipe;
+extern CORE_API std::atomic<bool> GIsAudioThreadRunning;
+extern CORE_API std::atomic<bool> GIsAudioThreadSuspended;
+
 int32 GCVarSuspendAudioThread = 0;
 TAutoConsoleVariable<int32> CVarSuspendAudioThread(TEXT("AudioThread.SuspendAudioThread"), GCVarSuspendAudioThread, TEXT("0=Resume, 1=Suspend"), ECVF_Cheat);
 
@@ -41,7 +47,6 @@ static FAutoConsoleVariableRef CVarBatchAudioAsyncBatchSize(
 	TEXT("When AudioThread.EnableBatchProcessing = 1, controls the number of audio commands grouped together for threading.")
 );
 
-
 static int32 GAudioCommandFenceWaitTimeMs = 35;
 TAutoConsoleVariable<int32> CVarAudioCommandFenceWaitTimeMs(
 	TEXT("AudioCommand.FenceWaitTimeMs"),
@@ -82,12 +87,79 @@ struct FAudioThreadInteractor
 static FAutoConsoleVariableSink CVarUseAudioThreadSink(FConsoleCommandDelegate::CreateStatic(&FAudioThreadInteractor::UseAudioThreadCVarSinkFunction));
 
 bool FAudioThread::bUseThreadedAudio = false;
-FRunnable* FAudioThread::AudioThreadRunnable = nullptr;
+
 FCriticalSection FAudioThread::CurrentAudioThreadStatIdCS;
 TStatId FAudioThread::CurrentAudioThreadStatId;
 TStatId FAudioThread::LongestAudioThreadStatId;
 double FAudioThread::LongestAudioThreadTimeMsec = 0.0;
 
+#if UE_AUDIO_THREAD_AS_PIPE
+
+UE::Tasks::FTaskEvent FAudioThread::ResumeEvent{ UE_SOURCE_LOCATION };
+int32 FAudioThread::SuspendCount{ 0 };
+
+void FAudioThread::SuspendAudioThread()
+{
+	check(IsInGameThread()); // thread-safe version would be much more complicated
+
+	if (!GIsAudioThreadRunning)
+	{
+		return; // nothing to suspend
+	}
+
+	if (++SuspendCount != 1)
+	{
+		return; // recursive scope
+	}
+
+	check(!GIsAudioThreadSuspended.load(std::memory_order_relaxed));
+
+	FEventRef SuspendEvent;
+
+	FAudioThread::RunCommandOnAudioThread(
+		[&SuspendEvent]
+		{
+			GIsAudioThreadSuspended.store(true, std::memory_order_release);
+			SuspendEvent->Trigger();
+
+			ResumeEvent.Wait();
+			ResumeEvent = UE::Tasks::FTaskEvent{ UE_SOURCE_LOCATION}; // thread-safe because happens "inside" a suspend call and 
+			// can't happen concurrently with a resume call
+		}
+	);
+
+	// release batch processing so the task above will be executed
+	FAudioThread::ProcessAllCommands();
+
+	// wait for the command above to block audio processing
+	SuspendEvent->Wait();
+}
+
+void FAudioThread::ResumeAudioThread()
+{
+	check(IsInGameThread());
+
+	if (!GIsAudioThreadRunning)
+	{
+		return; // nothing to resume
+	}
+
+	if (--SuspendCount != 0)
+	{
+		return; // recursive scope
+	}
+
+	check(GIsAudioThreadSuspended.load(std::memory_order_relaxed));
+	GIsAudioThreadSuspended.store(false, std::memory_order_release);
+
+	check(!ResumeEvent.IsCompleted());
+	ResumeEvent.Trigger();
+}
+
+#else // UE_AUDIO_THREAD_AS_PIPE
+
+FRunnable* FAudioThread::AudioThreadRunnable = nullptr;
+
 /** The audio thread main loop */
 void AudioThreadMain( FEvent* TaskGraphBoundSyncEvent )
 {
@@ -132,7 +204,7 @@ static int32 AudioThreadSuspendCount = 0;
 void FAudioThread::SuspendAudioThread()
 {
 	check(FPlatformTLS::GetCurrentThreadId() == GGameThreadId);
-	check(!GIsAudioThreadSuspended.Load() || CVarSuspendAudioThread.GetValueOnGameThread() != 0);
+	check(!GIsAudioThreadSuspended || CVarSuspendAudioThread.GetValueOnGameThread() != 0);
 	if (IsAudioThreadRunning())
 	{
 		// Make GC wait on the audio thread finishing processing
@@ -151,7 +223,7 @@ void FAudioThread::ResumeAudioThread()
 PRAGMA_DISABLE_DEPRECATION_WARNINGS
 
 	check(FPlatformTLS::GetCurrentThreadId() == GGameThreadId);
-	if (GIsAudioThreadSuspended.Load() && CVarSuspendAudioThread.GetValueOnGameThread() == 0)
+	if (GIsAudioThreadSuspended && CVarSuspendAudioThread.GetValueOnGameThread() == 0)
 	{
 		GIsAudioThreadSuspended = false;
 		FPlatformMisc::MemoryBarrier();
@@ -179,22 +251,21 @@ void FAudioThread::OnPostGarbageCollect()
 	}
 }
 
+#if !UE_AUDIO_THREAD_AS_PIPE
+
 bool FAudioThread::Init()
 { 
-PRAGMA_DISABLE_DEPRECATION_WARNINGS
-	GAudioThreadId = FPlatformTLS::GetCurrentThreadId();
-PRAGMA_ENABLE_DEPRECATION_WARNINGS
+	GIsAudioThreadRunning.store(true, std::memory_order_release);
 	return true;
 }
 
 void FAudioThread::Exit()
 {
-PRAGMA_DISABLE_DEPRECATION_WARNINGS
-	GAudioThreadId = 0;
-PRAGMA_ENABLE_DEPRECATION_WARNINGS
-	FPlatformProcess::TeardownAudioThread();
+	GIsAudioThreadRunning.store(false, std::memory_order_release);
 }
 
+#endif
+
 uint32 FAudioThread::Run()
 {
 	LLM_SCOPE(ELLMTag::AudioMisc);
@@ -206,6 +277,8 @@ uint32 FAudioThread::Run()
 	return 0;
 }
 
+#endif // UE_AUDIO_THREAD_AS_PIPE
+
 void FAudioThread::SetUseThreadedAudio(const bool bInUseThreadedAudio)
 {
 	if (IsAudioThreadRunning() && !bInUseThreadedAudio)
@@ -223,6 +296,61 @@ bool FAudioThread::IsUsingThreadedAudio()
 	return bUseThreadedAudio;
 }
 
+#if UE_AUDIO_THREAD_AS_PIPE
+
+// batching audio commands allows to avoid the overhead of launching a task per command when resources are limited.
+// We assume that resources are limited if the previous batch is not completed yet (potentially waiting for execution due to the CPU being busy
+// with something else). Otherwise we don't wait until we collect a full batch.
+struct FAudioAsyncBatcher
+{
+	using FWork = TUniqueFunction<void()>;
+	TArray<FWork> WorkItems;
+
+	UE::Tasks::FTask LastBatch;
+
+	void Add(FWork&& Work)
+	{
+		check(IsInGameThread());
+
+#if !WITH_EDITOR
+		if (GCVarEnableBatchProcessing)
+		{
+			if (WorkItems.Num() >= GBatchAudioAsyncBatchSize // collected enough work
+				|| LastBatch.IsCompleted()) // the previous batch is completed, kick off a new one
+			{
+				Flush();
+			}
+			WorkItems.Add(Forward<FWork>(Work));
+		}
+#else
+		LastBatch = GAudioPipe.Launch(UE_SOURCE_LOCATION, MoveTemp(Work));
+#endif
+	}
+
+	void Flush()
+	{
+		check(IsInGameThread());
+
+		if (WorkItems.IsEmpty())
+		{
+			return;
+		}
+
+		LastBatch = GAudioPipe.Launch(TEXT("AudioBatch"),
+			[WorkItems = MoveTemp(WorkItems)]() mutable
+			{
+				for (FWork& Work : WorkItems)
+				{
+					Work();
+				}
+			}
+		);
+		WorkItems.Reset();
+	}
+};
+
+#else // UE_AUDIO_THREAD_AS_PIPE
+
 struct FAudioAsyncBatcher
 {
 	FGraphEventArray DispatchEvent;
@@ -273,15 +401,19 @@ struct FAudioAsyncBatcher
 
 };
 
+#endif // UE_AUDIO_THREAD_AS_PIPE
+
 static FAudioAsyncBatcher GAudioAsyncBatcher;
 
 TUniqueFunction<void()> FAudioThread::GetCommandWrapper(TUniqueFunction<void()> InFunction, const TStatId InStatId)
 {
 	if (GCVarEnableAudioCommandLogging == 1)
 	{
-		return[Function = MoveTemp(InFunction), InStatId]()
+		return [Function = MoveTemp(InFunction), InStatId]()
 		{
+#if !UE_AUDIO_THREAD_AS_PIPE
 			FTaskTagScope Scope(ETaskTag::EAudioThread);
+#endif
 			FScopeCycleCounter ScopeCycleCounter(InStatId);
 			FAudioThread::SetCurrentAudioThreadStatId(InStatId);
 
@@ -301,9 +433,11 @@ TUniqueFunction<void()> FAudioThread::GetCommandWrapper(TUniqueFunction<void()>
 	}
 	else
 	{
-		return[Function = MoveTemp(InFunction), InStatId]()
+		return [Function = MoveTemp(InFunction), InStatId]()
 		{
+#if !UE_AUDIO_THREAD_AS_PIPE
 			FTaskTagScope Scope(ETaskTag::EAudioThread);
+#endif
 			FScopeCycleCounter ScopeCycleCounter(InStatId);
 			Function();
 		};
@@ -312,6 +446,32 @@ TUniqueFunction<void()> FAudioThread::GetCommandWrapper(TUniqueFunction<void()>
 
 void FAudioThread::RunCommandOnAudioThread(TUniqueFunction<void()> InFunction, const TStatId InStatId)
 {
+#if UE_AUDIO_THREAD_AS_PIPE
+
+	TUniqueFunction<void()> CommandWrapper{ GetCommandWrapper(MoveTemp(InFunction), InStatId) };
+	if (IsInAudioThread())
+	{
+		// it's audio-thread-safe so execute the command in-place
+		CommandWrapper();
+	}
+	else if (IsInGameThread())
+	{
+		// batch commands to minimise game thread overhead
+		GAudioAsyncBatcher.Add(MoveTemp(CommandWrapper));
+	}
+	// we are on an unknown thread
+	else if (IsUsingThreadedAudio())
+	{
+		GAudioPipe.Launch(TEXT("AudioCommand"), MoveTemp(CommandWrapper));
+	}
+	else
+	{
+		// the command must be executed on the game thread
+		AsyncTask(ENamedThreads::GameThread, MoveTemp(CommandWrapper));
+	}
+
+#else
+
 	if (IsInAudioThread())
 	{
 		// it's audio-thread-safe so execute the command in-place
@@ -336,6 +496,8 @@ void FAudioThread::RunCommandOnAudioThread(TUniqueFunction<void()> InFunction, c
 		// the command must be executed on the game thread
 		FFunctionGraphTask::CreateAndDispatchWhenReady(MoveTemp(CommandWrapper), TStatId{}, nullptr, ENamedThreads::GameThread);
 	}
+
+#endif
 }
 
 void FAudioThread::SetCurrentAudioThreadStatId(TStatId InStatId)
@@ -387,7 +549,11 @@ void FAudioThread::ProcessAllCommands()
 	}
 	else
 	{
-		check(!GAudioAsyncBatcher.NumBatched);
+#if UE_AUDIO_THREAD_AS_PIPE
+		check(GAudioAsyncBatcher.WorkItems.IsEmpty());
+#else
+		check(FPlatformTLS::GetCurrentThreadId() == GGameThreadId);
+#endif
 	}
 }
 
@@ -400,18 +566,102 @@ void FAudioThread::RunCommandOnGameThread(TUniqueFunction<void()> InFunction, co
 	}
 	else
 	{
-		check(FPlatformTLS::GetCurrentThreadId() == GGameThreadId);
+		check(IsInGameThread());
 		FScopeCycleCounter ScopeCycleCounter(InStatId);
 		InFunction();
 	}
 }
 
+#if UE_AUDIO_THREAD_AS_PIPE
+
+FDelegateHandle FAudioThread::PreGC;
+FDelegateHandle FAudioThread::PostGC;
+FDelegateHandle FAudioThread::PreGCDestroy;
+FDelegateHandle FAudioThread::PostGCDestroy;
+
+void FAudioThread::StartAudioThread()
+{
+	check(IsInGameThread());
+
+	check(!GIsAudioThreadRunning.load(std::memory_order_relaxed));
+
+	if (!bUseThreadedAudio)
+	{
+		return;
+	}
+
+	PreGC = FCoreUObjectDelegates::GetPreGarbageCollectDelegate().AddStatic(&FAudioThread::SuspendAudioThread);
+	PostGC = FCoreUObjectDelegates::GetPostGarbageCollect().AddStatic(&FAudioThread::ResumeAudioThread);
+
+	PreGCDestroy = FCoreUObjectDelegates::PreGarbageCollectConditionalBeginDestroy.AddStatic(&FAudioThread::SuspendAudioThread);
+	PostGCDestroy = FCoreUObjectDelegates::PostGarbageCollectConditionalBeginDestroy.AddStatic(&FAudioThread::ResumeAudioThread);
+
+	GIsAudioThreadRunning.store(true, std::memory_order_release);
+}
+
+void FAudioThread::StopAudioThread()
+{
+	if (!bUseThreadedAudio)
+	{
+		return;
+	}
+
+	check(IsAudioThreadRunning());
+
+	GAudioAsyncBatcher.Flush();
+	GAudioAsyncBatcher.LastBatch.Wait();
+
+	FCoreUObjectDelegates::GetPreGarbageCollectDelegate().Remove(PreGC);
+	FCoreUObjectDelegates::GetPostGarbageCollect().Remove(PostGC);
+	FCoreUObjectDelegates::PreGarbageCollectConditionalBeginDestroy.Remove(PreGCDestroy);
+	FCoreUObjectDelegates::PostGarbageCollectConditionalBeginDestroy.Remove(PostGCDestroy);
+
+	GIsAudioThreadRunning.store(false, std::memory_order_release);
+}
+
+FAudioCommandFence::~FAudioCommandFence()
+{
+	check(IsInGameThread());
+	Fence.Wait();
+}
+
+void FAudioCommandFence::BeginFence()
+{
+	check(IsInGameThread());
+
+	if (!IsAudioThreadRunning())
+	{
+		return;
+	}
+
+	Fence.Wait();
+	FAudioThread::ProcessAllCommands();
+	Fence = GAudioAsyncBatcher.LastBatch;
+}
+
+bool FAudioCommandFence::IsFenceComplete() const
+{
+	check(IsInGameThread());
+	return Fence.IsCompleted();
+}
+
+/**
+ * Waits for pending fence commands to retire.
+ */
+void FAudioCommandFence::Wait(bool bProcessGameThreadTasks) const
+{
+	check(IsInGameThread());
+	Fence.Wait();
+}
+
+#else // UE_AUDIO_THREAD_AS_PIPE
+
 void FAudioThread::StartAudioThread()
 {
 	check(FPlatformTLS::GetCurrentThreadId() == GGameThreadId);
 
 	check(!IsAudioThreadRunning());
-	check(!GIsAudioThreadSuspended.Load());
+	check(!GIsAudioThreadSuspended);
 	if (bUseThreadedAudio)
 	{
 PRAGMA_DISABLE_DEPRECATION_WARNINGS
@@ -449,7 +699,7 @@ PRAGMA_ENABLE_DEPRECATION_WARNINGS
 void FAudioThread::StopAudioThread()
 {
 	check(FPlatformTLS::GetCurrentThreadId() == GGameThreadId);
-	check(!GIsAudioThreadSuspended.Load() || CVarSuspendAudioThread.GetValueOnGameThread() != 0);
+	check(!GIsAudioThreadSuspended || CVarSuspendAudioThread.GetValueOnGameThread() != 0);
 
 	if (!IsAudioThreadRunning())
 	{
@@ -597,3 +847,4 @@ void FAudioCommandFence::Wait(bool bProcessGameThreadTasks) const
 	}
 }
 
+#endif // UE_AUDIO_THREAD_AS_PIPE
\ No newline at end of file
diff --git a/Engine/Source/Runtime/Engine/Public/AudioThread.h b/Engine/Source/Runtime/Engine/Public/AudioThread.h
--- a/Engine/Source/Runtime/Engine/Public/AudioThread.h
+++ b/Engine/Source/Runtime/Engine/Public/AudioThread.h
@@ -8,24 +8,47 @@
 
 #include "CoreMinimal.h"
 #include "Stats/Stats.h"
+
+#include "Tasks/Pipe.h"
+
+
+#if !UE_AUDIO_THREAD_AS_PIPE
+
 #include "Async/TaskGraphInterfaces.h"
 #include "HAL/Runnable.h"
 #include "Containers/Queue.h"
 
+#endif
+
 ////////////////////////////////////
 // Audio thread API
+// (the naming is outdated as the thread was replaced by piped tasks)
 ////////////////////////////////////
 
 DECLARE_STATS_GROUP(TEXT("Audio Thread Commands"), STATGROUP_AudioThreadCommands, STATCAT_Advanced);
 
-class FAudioThread : public FRunnable
+class FAudioThread
+#if !UE_AUDIO_THREAD_AS_PIPE
+	: public FRunnable
+#endif
 {
 private:
-	/**
-	* Whether to run with an audio thread
-	*/
+	/** Whether to run with an audio thread */
 	static bool bUseThreadedAudio;
 
+#if UE_AUDIO_THREAD_AS_PIPE
+
+	static UE::Tasks::FTaskEvent ResumeEvent;
+	static int32 SuspendCount;  // accessed only from GT
+
+	// GC callback handles
+	static FDelegateHandle PreGC;
+	static FDelegateHandle PostGC;
+	static FDelegateHandle PreGCDestroy;
+	static FDelegateHandle PostGCDestroy;
+
+#else
+
 	/** 
 	* Sync event to make sure that audio thread is bound to the task graph before main thread queues work against it.
 	*/
@@ -43,6 +66,8 @@ private:
 	void OnPreGarbageCollect();
 	void OnPostGarbageCollect();
 
+#endif
+
 	/** Stat id of the currently executing audio thread command. */
 	static TStatId CurrentAudioThreadStatId;
 	static TStatId LongestAudioThreadStatId;
@@ -59,6 +84,8 @@ private:
 
 public:
 
+#if !UE_AUDIO_THREAD_AS_PIPE
+
 	FAudioThread();
 	virtual ~FAudioThread();
 
@@ -67,13 +94,15 @@ public:
 	virtual void Exit() override;
 	virtual uint32 Run() override;
 
+#endif
+
 	/** Starts the audio thread. */
 	static ENGINE_API void StartAudioThread();
 
 	/** Stops the audio thread. */
 	static ENGINE_API void StopAudioThread();
 
-	/** Execute a command on the audio thread. If GIsAudioThreadRunning is false the command will execute immediately */
+	/** Execute a command on the audio thread. If it's safe the command will execute immediately. */
 	static ENGINE_API void RunCommandOnAudioThread(TUniqueFunction<void()> InFunction, const TStatId InStatId = TStatId());
 
 	/** Processes all enqueued audio thread commands. */
@@ -121,7 +150,9 @@ struct FAudioThreadSuspendContext
 class ENGINE_API FAudioCommandFence
 {
 public:
+#if !UE_AUDIO_THREAD_AS_PIPE
 	FAudioCommandFence();
+#endif
 	~FAudioCommandFence();
 
 	/**
@@ -141,9 +172,19 @@ public:
 	bool IsFenceComplete() const;
 
 private:
+
+#if UE_AUDIO_THREAD_AS_PIPE
+
+	/** The last audio batch task **/
+	mutable UE::Tasks::FTask Fence;
+
+#else
+
 	/** Graph event that represents completion of this fence **/
 	mutable FGraphEventRef CompletionEvent;
 	/** Event that fires when CompletionEvent is done. **/
 	mutable FEvent* FenceDoneEvent;
+
+#endif
 };
 
