LWC: Custom net serialized FVector make/breaks to allow blueprints to modify them without double -> float casts. TODO: Not necessary if we fixed basic FVector make/break to be double.

diff --git a/Engine/Source/Runtime/Engine/Classes/Engine/NetSerialization.h b/Engine/Source/Runtime/Engine/Classes/Engine/NetSerialization.h
--- a/Engine/Source/Runtime/Engine/Classes/Engine/NetSerialization.h
+++ b/Engine/Source/Runtime/Engine/Classes/Engine/NetSerialization.h
@@ -420,7 +420,7 @@ bool SerializeFixedVector(FVector3d& Vector, FArchive& Ar)
  *	Note: this is the historical UE format for vector net serialization
  *
  */
-USTRUCT()
+USTRUCT(meta = (HasNativeMake = "Engine.KismetMathLibrary.MakeVector_NetQuantize", HasNativeBreak = "Engine.KismetMathLibrary.BreakVector_NetQuantize"))
 struct FVector_NetQuantize : public FVector
 {
 	GENERATED_USTRUCT_BODY()
@@ -466,7 +466,7 @@ struct TStructOpsTypeTraits< FVector_NetQuantize > : public TStructOpsTypeTraits
  *	Valid range: 2^24 / 10 = +/- 1,677,721.6
  *
  */
-USTRUCT()
+USTRUCT(meta = (HasNativeMake = "Engine.KismetMathLibrary.MakeVector_NetQuantize10", HasNativeBreak = "Engine.KismetMathLibrary.BreakVector_NetQuantize10"))
 struct FVector_NetQuantize10 : public FVector
 {
 	GENERATED_USTRUCT_BODY()
@@ -512,7 +512,7 @@ struct TStructOpsTypeTraits< FVector_NetQuantize10 > : public TStructOpsTypeTrai
  *	Valid range: 2^30 / 100 = +/- 10,737,418.24
  *
  */
-USTRUCT()
+USTRUCT(meta = (HasNativeMake = "Engine.KismetMathLibrary.MakeVector_NetQuantize100", HasNativeBreak = "Engine.KismetMathLibrary.BreakVector_NetQuantize100"))
 struct FVector_NetQuantize100 : public FVector
 {
 	GENERATED_USTRUCT_BODY()
@@ -535,7 +535,7 @@ struct FVector_NetQuantize100 : public FVector
 	{
 		bOutSuccess = SerializePackedVector<100, 30>(*this, Ar);
 		return true;
-	}	
+	}
 };
 
 template<>
@@ -554,7 +554,7 @@ struct TStructOpsTypeTraits< FVector_NetQuantize100 > : public TStructOpsTypeTra
  *	16 bits per component
  *	Valid range: -1..+1 inclusive
  */
-USTRUCT()
+USTRUCT(meta = (HasNativeMake = "Engine.KismetMathLibrary.MakeVector_NetQuantizeNormal", HasNativeBreak = "Engine.KismetMathLibrary.BreakVector_NetQuantizeNormal"))
 struct FVector_NetQuantizeNormal : public FVector
 {
 	GENERATED_USTRUCT_BODY()
diff --git a/Engine/Source/Runtime/Engine/Classes/Kismet/KismetMathLibrary.h b/Engine/Source/Runtime/Engine/Classes/Kismet/KismetMathLibrary.h
--- a/Engine/Source/Runtime/Engine/Classes/Kismet/KismetMathLibrary.h
+++ b/Engine/Source/Runtime/Engine/Classes/Kismet/KismetMathLibrary.h
@@ -12,6 +12,7 @@
 #include "UObject/ScriptMacros.h"
 #include "Kismet/BlueprintFunctionLibrary.h"
 #include "Misc/QualifiedFrameTime.h"
+#include "Engine/NetSerialization.h"
 
 #include "KismetMathLibrary.generated.h"
 
@@ -4141,6 +4142,27 @@ class ENGINE_API UKismetMathLibrary : public UBlueprintFunctionLibrary
 	UFUNCTION(BlueprintPure, Category="Math|Smoothing", meta=(DisplayName="Dynamic Weighted Moving Average Rotator"))
 	static FRotator DynamicWeightedMovingAverage_FRotator(FRotator CurrentSample, FRotator PreviousSample, float MaxDistance, float MinWeight, float MaxWeight);
 
+
+	// NetQuantized vector make/breaks
+	// LWC_TODO: These can be removed once blueprints support FReal
+	UFUNCTION(BlueprintPure, Category="Math|Vector", meta=(NativeMakeFunc))
+	static FVector_NetQuantize MakeVector_NetQuantize(float X, float Y, float Z) { return FVector_NetQuantize(X, Y, Z); }
+	UFUNCTION(BlueprintPure, Category="Math|Vector", meta=(NativeMakeFunc))
+	static FVector_NetQuantize10 MakeVector_NetQuantize10(float X, float Y, float Z) { return FVector_NetQuantize10(X, Y, Z); }
+	UFUNCTION(BlueprintPure, Category="Math|Vector", meta=(NativeMakeFunc))
+	static FVector_NetQuantize100 MakeVector_NetQuantize100(float X, float Y, float Z) { return FVector_NetQuantize100(X, Y, Z); }
+	UFUNCTION(BlueprintPure, Category="Math|Vector", meta=(NativeMakeFunc))
+	static FVector_NetQuantizeNormal MakeVector_NetQuantizeNormal(float X, float Y, float Z) { return FVector_NetQuantizeNormal(X, Y, Z); }
+	UFUNCTION(BlueprintPure, Category="Math|Vector", meta=(NativeBreakFunc))
+	static void BreakVector_NetQuantize(FVector_NetQuantize InVec, float& X, float& Y, float& Z) { BreakVector((FVector)InVec, X, Y, Z); }
+	UFUNCTION(BlueprintPure, Category="Math|Vector", meta=(NativeBreakFunc))
+	static void BreakVector_NetQuantize10(FVector_NetQuantize10 InVec, float& X, float& Y, float& Z) { BreakVector((FVector)InVec, X, Y, Z); }
+	UFUNCTION(BlueprintPure, Category="Math|Vector", meta=(NativeBreakFunc))
+	static void BreakVector_NetQuantize100(FVector_NetQuantize100 InVec, float& X, float& Y, float& Z) { BreakVector((FVector)InVec, X, Y, Z); }
+	UFUNCTION(BlueprintPure, Category="Math|Vector", meta=(NativeBreakFunc))
+	static void BreakVector_NetQuantizeNormal(FVector_NetQuantizeNormal InVec, float& X, float& Y, float& Z) { BreakVector((FVector)InVec, X, Y, Z); }
+
+	
 private:
 
 	static void ReportError_Divide_ByteByte();
