Add conversion of material attribute root node to starta.
This is a basic first version, we will need to handle shading models

diff --git a/Engine/Source/Runtime/Engine/Private/Materials/Material.cpp b/Engine/Source/Runtime/Engine/Private/Materials/Material.cpp
--- a/Engine/Source/Runtime/Engine/Private/Materials/Material.cpp
+++ b/Engine/Source/Runtime/Engine/Private/Materials/Material.cpp
@@ -90,6 +90,8 @@
 #include "Materials/MaterialExpressionThinTranslucentMaterialOutput.h"
 #include "Materials/MaterialExpressionClearCoatNormalCustomOutput.h"
 #include "Materials/MaterialExpressionConstant.h"
+#include "Materials/MaterialExpressionBreakMaterialAttributes.h"
+
 
 #if WITH_EDITOR
 #include "Logging/TokenizedMessage.h"
@@ -3654,33 +3656,63 @@ void UMaterial::ConvertMaterialToStrataMaterial()
 {
 #if WITH_EDITOR
 	static const auto CVarStrata = IConsoleManager::Get().FindTConsoleVariableDataInt(TEXT("r.Strata"));
-	const bool bStrata = CVarStrata ? CVarStrata->GetValueOnAnyThread() > 0 : false;
+	const bool bStrataEnabled = CVarStrata ? CVarStrata->GetValueOnAnyThread() > 0 : false;
 
-	if (!bStrata)
+	if (!bStrataEnabled)
 	{
 		return;
 	}
 
-	static const auto CVar = IConsoleManager::Get().FindTConsoleVariableDataInt(TEXT("r.Strata"));
-	const bool bStrataEnabled = CVar && CVar->GetValueOnAnyThread() > 0;
-	if (bStrataEnabled && !FrontMaterial.IsConnected())
+	auto MoveConnectionTo = [](auto& OldNodeInput, UMaterialExpression* NewNode, uint32 NewInputIndex)
 	{
-		auto MoveConnectionTo = [](auto& OldNodeInput, UMaterialExpression* NewNode, uint32 NewInputIndex)
+		if (OldNodeInput.IsConnected())
 		{
-			if (OldNodeInput.IsConnected())
-			{
-				NewNode->GetInput(NewInputIndex)->Connect(OldNodeInput.OutputIndex, OldNodeInput.Expression);
-				OldNodeInput.Expression = nullptr;
-			}
-		};
-		auto CopyConnectionTo = [](auto& OldNodeInput, UMaterialExpression* NewNode, uint32 NewInputIndex)
+			NewNode->GetInput(NewInputIndex)->Connect(OldNodeInput.OutputIndex, OldNodeInput.Expression);
+			OldNodeInput.Expression = nullptr;
+		}
+	};
+	auto CopyConnectionTo = [](auto& OldNodeInput, UMaterialExpression* NewNode, uint32 NewInputIndex)
+	{
+		if (OldNodeInput.IsConnected())
 		{
-			if (OldNodeInput.IsConnected())
-			{
-				NewNode->GetInput(NewInputIndex)->Connect(OldNodeInput.OutputIndex, OldNodeInput.Expression);
-			}
-		};
+			NewNode->GetInput(NewInputIndex)->Connect(OldNodeInput.OutputIndex, OldNodeInput.Expression);
+		}
+	};
 
+	if (bUseMaterialAttributes && MaterialAttributes.Expression && !MaterialAttributes.Expression->IsResultStrataMaterial(MaterialAttributes.OutputIndex)) // M_Rifle cause issues there
+	{
+		UMaterialExpressionBreakMaterialAttributes* BreakMatAtt = NewObject<UMaterialExpressionBreakMaterialAttributes>(this);
+		MoveConnectionTo(MaterialAttributes, BreakMatAtt, 0);
+
+		// STRATA_TODO plug more than just the basic default material
+		// STARTA_TODO handle shading models
+		UMaterialExpressionStrataSlabBSDF* SlabBSDF = NewObject<UMaterialExpressionStrataSlabBSDF>(this);
+		SlabBSDF->BaseColor.Connect(0, BreakMatAtt);	// 0- BaseColor
+		SlabBSDF->Metallic.Connect(1, BreakMatAtt);		// 1- Metallic
+		SlabBSDF->Specular.Connect(2, BreakMatAtt);		// 2- Specular
+		SlabBSDF->Roughness.Connect(3, BreakMatAtt);	// 3- Roughness
+		SlabBSDF->Anisotropy.Connect(4, BreakMatAtt);	// 4- Anisotropy
+		SlabBSDF->EmissiveColor.Connect(5, BreakMatAtt);// 5- EmissiveColor
+		// 6- Opacity									// We should not have to do anything, it should be forwarded to the root node correctly.
+		// 7- Opacity Mask								// Idem
+		SlabBSDF->Normal.Connect(8, BreakMatAtt);		// 8- Normal
+		SlabBSDF->Tangent.Connect(9, BreakMatAtt);		// 9- Tangent
+		// 10- WorldPositionOffset
+		// 11- SubsurfaceColor
+		// 12- ClearCoat
+		// 13- ClearCoatRoughness
+		// 14- AmbiantOcclusion
+		// 15- Refraction
+		// 16- PixelDepthOffset
+		// 17- ShadingModel
+
+		UMaterialExpressionMakeMaterialAttributes* MakeMatAtt = NewObject<UMaterialExpressionMakeMaterialAttributes>(this);
+		MakeMatAtt->FrontMaterial.Connect(0, SlabBSDF);
+
+		MaterialAttributes.Connect(0, MakeMatAtt);
+	}
+	else if (!bUseMaterialAttributes && !FrontMaterial.IsConnected())
+	{
 		// STRATA_TODO for material conversion
 		//  - WorldPositionOffset can remain on the end point node
 		//  - ShadingModel
