Allow inclusion of editor only actors/components in HLOD
Exclude visualization components

diff --git a/Engine/Source/Runtime/Engine/Private/Actor.cpp b/Engine/Source/Runtime/Engine/Private/Actor.cpp
--- a/Engine/Source/Runtime/Engine/Private/Actor.cpp
+++ b/Engine/Source/Runtime/Engine/Private/Actor.cpp
@@ -5515,11 +5515,6 @@ bool AActor::IsHLODRelevant() const
 		return false;
 	}
 
-	if (IsEditorOnly())
-	{
-		return false;
-	}
-
 	if (IsTemplate())
 	{
 		return false;
diff --git a/Engine/Source/Runtime/Engine/Private/Components/StaticMeshComponent.cpp b/Engine/Source/Runtime/Engine/Private/Components/StaticMeshComponent.cpp
--- a/Engine/Source/Runtime/Engine/Private/Components/StaticMeshComponent.cpp
+++ b/Engine/Source/Runtime/Engine/Private/Components/StaticMeshComponent.cpp
@@ -2520,7 +2520,7 @@ bool UStaticMeshComponent::IsHLODRelevant() const
 		return false;
 	}
 
-	if (bIsEditorOnly)
+	if (IsVisualizationComponent())
 	{
 		return false;
 	}
