Build fix for CL16484321

diff --git a/Engine/Source/Runtime/Engine/Private/Components/StaticMeshComponent.cpp b/Engine/Source/Runtime/Engine/Private/Components/StaticMeshComponent.cpp
--- a/Engine/Source/Runtime/Engine/Private/Components/StaticMeshComponent.cpp
+++ b/Engine/Source/Runtime/Engine/Private/Components/StaticMeshComponent.cpp
@@ -2520,17 +2520,17 @@ bool UStaticMeshComponent::IsHLODRelevant() const
 		return false;
 	}
 
-	if (IsVisualizationComponent())
+	if (Mobility == EComponentMobility::Movable)
 	{
 		return false;
 	}
 
-	if (Mobility == EComponentMobility::Movable)
+#if WITH_EDITORONLY_DATA
+	if (IsVisualizationComponent())
 	{
 		return false;
 	}
 
-#if WITH_EDITORONLY_DATA
 	if (!bEnableAutoLODGeneration)
 	{
 		return false;
