Removed left behind stratasheen node functions because it has been integrated on the slab node.

diff --git a/Engine/Source/Runtime/Engine/Private/Materials/HLSLMaterialTranslator.cpp b/Engine/Source/Runtime/Engine/Private/Materials/HLSLMaterialTranslator.cpp
--- a/Engine/Source/Runtime/Engine/Private/Materials/HLSLMaterialTranslator.cpp
+++ b/Engine/Source/Runtime/Engine/Private/Materials/HLSLMaterialTranslator.cpp
@@ -9020,17 +9020,6 @@ int32 FHLSLMaterialTranslator::StrataSlabBSDF(
 	);
 }
 
-int32 FHLSLMaterialTranslator::StrataSheenBSDF(int32 BaseColor, int32 Roughness, int32 Normal, const FString& SharedLocalBasisIndexMacro)
-{
-	return AddCodeChunk(
-		MCT_Strata, TEXT("GetStrataSheenBSDF(%s, %s, %s) /* %s */"),
-		*GetParameterCode(BaseColor),
-		*GetParameterCode(Roughness),
-		*SharedLocalBasisIndexMacro,
-		*GetParameterCode(Normal)
-	);
-}
-
 int32 FHLSLMaterialTranslator::StrataVolumetricFogCloudBSDF(int32 Albedo, int32 Extinction, int32 EmissiveColor, int32 AmbientOcclusion)
 {
 	return AddCodeChunk(
diff --git a/Engine/Source/Runtime/Engine/Private/Materials/HLSLMaterialTranslator.h b/Engine/Source/Runtime/Engine/Private/Materials/HLSLMaterialTranslator.h
--- a/Engine/Source/Runtime/Engine/Private/Materials/HLSLMaterialTranslator.h
+++ b/Engine/Source/Runtime/Engine/Private/Materials/HLSLMaterialTranslator.h
@@ -988,7 +988,6 @@ protected:
 		int32 FuzzAmount, int32 FuzzColor, 
 		int32 Thickness,
 		int32 Normal, int32 Tangent, const FString& SharedLocalBasisIndexMacro) override;
-	virtual int32 StrataSheenBSDF(int32 BaseColor, int32 Roughness, int32 Normal, const FString& SharedLocalBasisIndexMacro) override;
 	virtual int32 StrataVolumetricFogCloudBSDF(int32 Albedo, int32 Extinction, int32 EmissiveColor, int32 AmbientOcclusion) override;
 	virtual int32 StrataUnlitBSDF(int32 EmissiveColor, int32 TransmittanceColor) override;
 	virtual int32 StrataHairBSDF(int32 BaseColor, int32 Scatter, int32 Specular, int32 Roughness, int32 Backlit, int32 EmissiveColor, int32 Tangent, const FString& SharedLocalBasisIndexMacro) override;
diff --git a/Engine/Source/Runtime/Engine/Public/MaterialCompiler.h b/Engine/Source/Runtime/Engine/Public/MaterialCompiler.h
--- a/Engine/Source/Runtime/Engine/Public/MaterialCompiler.h
+++ b/Engine/Source/Runtime/Engine/Public/MaterialCompiler.h
@@ -432,7 +432,6 @@ public:
 		int32 FuzzAmount, int32 FuzzColor,
 		int32 Thickness,
 		int32 Normal, int32 Tangent, const FString& SharedLocalBasisIndexMacro) = 0;
-	virtual int32 StrataSheenBSDF(int32 BaseColor, int32 Roughness, int32 Normal, const FString& SharedLocalBasisIndexMacro) = 0;
 	virtual int32 StrataVolumetricFogCloudBSDF(int32 Albedo, int32 Extinction, int32 EmissiveColor, int32 AmbientOcclusion) = 0;
 	virtual int32 StrataUnlitBSDF(int32 EmissiveColor, int32 TransmittanceColor) = 0;
 	virtual int32 StrataHairBSDF(int32 BaseColor, int32 Scatter, int32 Specular, int32 Roughness, int32 Backlit, int32 EmissiveColor, int32 Tangent, const FString& SharedLocalBasisIndexMacro) = 0;
@@ -950,11 +949,6 @@ public:
 			Normal, Tangent, SharedLocalBasisIndexMacro);
 	}
 
-	virtual int32 StrataSheenBSDF(int32 BaseColor, int32 Roughness, int32 Normal, const FString& SharedLocalBasisIndexMacro) override
-	{
-		return Compiler->StrataSheenBSDF(BaseColor, Roughness, Normal, SharedLocalBasisIndexMacro);
-	}
-
 	virtual int32 StrataVolumetricFogCloudBSDF(int32 Albedo, int32 Extinction, int32 EmissiveColor, int32 AmbientOcclusion) override
 	{
 		return Compiler->StrataVolumetricFogCloudBSDF(Albedo, Extinction, EmissiveColor, AmbientOcclusion);
