Pass FMaterialPixelParameters via 'inout' to various functions that can compute additional material outputs, needed to ensure that any VT feedback generated while computing these is properly output

diff --git a/Engine/Shaders/Private/MaterialTemplate.ush b/Engine/Shaders/Private/MaterialTemplate.ush
--- a/Engine/Shaders/Private/MaterialTemplate.ush
+++ b/Engine/Shaders/Private/MaterialTemplate.ush
@@ -2431,12 +2431,12 @@ half4 GetMaterialSubsurfaceData(FPixelMaterialInputs PixelMaterialInputs)
 	return OutSubsurface;
 }
 
-half GetMaterialCustomData0(FMaterialPixelParameters Parameters)
+half GetMaterialCustomData0(in out FMaterialPixelParameters Parameters)
 {
 %s;
 }
 
-half GetMaterialCustomData1(FMaterialPixelParameters Parameters)
+half GetMaterialCustomData1(in out FMaterialPixelParameters Parameters)
 {
 %s;
 }
@@ -2502,7 +2502,7 @@ float3 TransformTangentNormalToWorld(MaterialFloat3x3 TangentToWorld, float3 Tan
 
 #endif //DECAL_PRIMITIVE
 
-float3 CalculateAnisotropyTangent(FMaterialPixelParameters Parameters, FPixelMaterialInputs PixelMaterialInputs)
+float3 CalculateAnisotropyTangent(in out FMaterialPixelParameters Parameters, FPixelMaterialInputs PixelMaterialInputs)
 {
 	float3 Normal = Parameters.WorldNormal;
 
diff --git a/Engine/Shaders/Private/ShadingModelsMaterial.ush b/Engine/Shaders/Private/ShadingModelsMaterial.ush
--- a/Engine/Shaders/Private/ShadingModelsMaterial.ush
+++ b/Engine/Shaders/Private/ShadingModelsMaterial.ush
@@ -11,7 +11,7 @@
 
 void SetGBufferForShadingModel(
 	in out FGBufferData GBuffer, 
-	in const FMaterialPixelParameters MaterialParameters,
+	in out FMaterialPixelParameters MaterialParameters,
 	const float Opacity,
 	const half3 BaseColor,
 	const half  Metallic,
diff --git a/Engine/Source/Runtime/Engine/Private/Materials/HLSLMaterialTranslator.cpp b/Engine/Source/Runtime/Engine/Private/Materials/HLSLMaterialTranslator.cpp
--- a/Engine/Source/Runtime/Engine/Private/Materials/HLSLMaterialTranslator.cpp
+++ b/Engine/Source/Runtime/Engine/Private/Materials/HLSLMaterialTranslator.cpp
@@ -9566,7 +9566,7 @@ int32 FHLSLMaterialTranslator::CustomOutput(class UMaterialExpressionCustomOutpu
 		GetFixedParameterCode(OutputCode, *CurrentScopeChunks, Definitions, Body, CompiledPDV_FiniteDifferences);
 	}
 
-	FString ImplementationCodeFinite = FString::Printf(TEXT("%s %s%d(FMaterial%sParameters Parameters)\r\n{\r\n%s return %s;\r\n}\r\n"), *OutputTypeString, *Custom->GetFunctionName(), OutputIndex, ShaderFrequency == SF_Vertex ? TEXT("Vertex") : TEXT("Pixel"), *Definitions, *Body);
+	FString ImplementationCodeFinite = FString::Printf(TEXT("%s %s%d(in out FMaterial%sParameters Parameters)\r\n{\r\n%s return %s;\r\n}\r\n"), *OutputTypeString, *Custom->GetFunctionName(), OutputIndex, ShaderFrequency == SF_Vertex ? TEXT("Vertex") : TEXT("Pixel"), *Definitions, *Body);
 	DerivativeVariations[CompiledPDV_FiniteDifferences].CustomOutputImplementations.Add(ImplementationCodeFinite);
 
 	// FIXEM: for now just skipping these, as they create invalid code.
