Allow DDC2 TextureBuildFunction to be used when building textures in the editor process.  Still does DDC1 cache get/put operations around the DDC2 ones, but that will be changed in the future as we shift from use of build actions to build definitions.

diff --git a/Engine/Source/Developer/TextureBuild/Private/TextureBuildFunction.cpp b/Engine/Source/Developer/TextureBuild/Private/TextureBuildFunction.cpp
--- a/Engine/Source/Developer/TextureBuild/Private/TextureBuildFunction.cpp
+++ b/Engine/Source/Developer/TextureBuild/Private/TextureBuildFunction.cpp
@@ -18,7 +18,7 @@
 #include "Serialization/MemoryWriter.h"
 #include "TextureCompressorModule.h"
 
-DEFINE_LOG_CATEGORY_STATIC(LogTextureBuildWorker, Log, All);
+DEFINE_LOG_CATEGORY_STATIC(LogTextureBuildFunction, Log, All);
 
 #ifndef CASE_ENUM_TO_TEXT
 #define CASE_ENUM_TO_TEXT(txt) case txt: return TEXT(#txt);
@@ -253,7 +253,7 @@ void FTextureBuildFunction::Build(UE::DerivedData::FBuildContext& Context) const
 		ReadTextureSourceFromCompactBinary(AssociatedNormalSourceMipsCbObj, Context, AssociatedNormalSourceMips);
 	}
 
-	UE_LOG(LogTextureBuildWorker, Display, TEXT("Compressing %d source mip(s) (%dx%d) to %s..."), SourceMips.Num(), SourceMips[0].SizeX, SourceMips[0].SizeY, *BuildSettings.TextureFormatName.ToString());
+	UE_LOG(LogTextureBuildFunction, Display, TEXT("Compressing %d source mip(s) (%dx%d) to %s..."), SourceMips.Num(), SourceMips[0].SizeX, SourceMips[0].SizeY, *BuildSettings.TextureFormatName.ToString());
 
 	TArray<FCompressedImage2D> CompressedMips;
 	uint32 NumMipsInTail;
@@ -387,6 +387,6 @@ void FTextureBuildFunction::Build(UE::DerivedData::FBuildContext& Context) const
 
 
 		FCompositeBuffer CompositeResult(OrderedBuffers);
-		Context.AddPayload(UE::DerivedData::FPayloadId::FromName("Texture"), FCompressedBuffer::Compress(NAME_Default, CompositeResult));
+		Context.AddPayload(UE::DerivedData::FPayloadId::FromName(TEXT("Texture")), FCompressedBuffer::Compress(NAME_Default, CompositeResult));
 	}
 }
diff --git a/Engine/Source/Runtime/Engine/Private/TextureDerivedDataBuildExporter.cpp b/Engine/Source/Runtime/Engine/Private/TextureDerivedDataBuildExporter.cpp
--- a/Engine/Source/Runtime/Engine/Private/TextureDerivedDataBuildExporter.cpp
+++ b/Engine/Source/Runtime/Engine/Private/TextureDerivedDataBuildExporter.cpp
@@ -18,6 +18,7 @@
 #include "Misc/StringBuilder.h"
 #include "Serialization/CompactBinaryPackage.h"
 #include "Serialization/CompactBinaryWriter.h"
+#include "TextureDerivedDataBuildUtils.h"
 #include "TextureDerivedDataTask.h"
 #include "TextureFormatManager.h"
 
@@ -62,172 +63,6 @@ static bool ShortenKey(const TCHAR* CacheKey, FString& Result)
 	return true;
 }
 
-static void WriteCbField(FCbWriter& Writer, const FAnsiStringView& Name, const FColor& Color)
-{
-	// Choosing the big endian ordering
-	Writer.BeginArray(Name);
-	Writer.AddInteger(Color.A);
-	Writer.AddInteger(Color.R);
-	Writer.AddInteger(Color.G);
-	Writer.AddInteger(Color.B);
-	Writer.EndArray();
-}
-
-static void WriteCbField(FCbWriter& Writer, const FAnsiStringView& Name, const FVector4& Vec4)
-{
-	Writer.BeginArray(Name);
-	Writer.AddFloat(Vec4.X);
-	Writer.AddFloat(Vec4.Y);
-	Writer.AddFloat(Vec4.Z);
-	Writer.AddFloat(Vec4.W);
-	Writer.EndArray();
-}
-
-static void WriteCbField(FCbWriter& Writer, const FAnsiStringView& Name, const FIntPoint& IntPoint)
-{
-	Writer.BeginArray(Name);
-	Writer.AddInteger(IntPoint.X);
-	Writer.AddInteger(IntPoint.Y);
-	Writer.EndArray();
-}
-
-static FCbObject WriteBuildSettingsToCompactBinary(const FTextureBuildSettings& BuildSettings, const ITextureFormat* TextureFormat)
-{
-	FCbWriter Writer;
-	Writer.BeginObject();
-
-	if (BuildSettings.FormatConfigOverride)
-	{
-		Writer.AddObject("FormatConfigOverride", BuildSettings.FormatConfigOverride);
-	}
-	else if (FCbObject TextureFormatConfig = TextureFormat->ExportGlobalFormatConfig(BuildSettings))
-	{
-		Writer.AddObject("FormatConfigOverride", TextureFormatConfig);
-	}
-	Writer.BeginObject("ColorAdjustment");
-	Writer.AddFloat("AdjustBrightness", BuildSettings.ColorAdjustment.AdjustBrightness);
-	Writer.AddFloat("AdjustBrightnessCurve", BuildSettings.ColorAdjustment.AdjustBrightnessCurve);
-	Writer.AddFloat("AdjustSaturation", BuildSettings.ColorAdjustment.AdjustSaturation);
-	Writer.AddFloat("AdjustVibrance", BuildSettings.ColorAdjustment.AdjustVibrance);
-	Writer.AddFloat("AdjustRGBCurve", BuildSettings.ColorAdjustment.AdjustRGBCurve);
-	Writer.AddFloat("AdjustHue", BuildSettings.ColorAdjustment.AdjustHue);
-	Writer.AddFloat("AdjustMinAlpha", BuildSettings.ColorAdjustment.AdjustMinAlpha);
-	Writer.AddFloat("AdjustMaxAlpha", BuildSettings.ColorAdjustment.AdjustMaxAlpha);
-	Writer.EndObject();
-
-	WriteCbField(Writer, "AlphaCoverageThresholds", BuildSettings.AlphaCoverageThresholds);
-
-	Writer.AddFloat("MipSharpening", BuildSettings.MipSharpening);
-	Writer.AddInteger("DiffuseConvolveMipLevel", BuildSettings.DiffuseConvolveMipLevel);
-	Writer.AddInteger("SharpenMipKernelSize", BuildSettings.SharpenMipKernelSize);
-	Writer.AddInteger("MaxTextureResolution", BuildSettings.MaxTextureResolution);
-	Writer.AddString("TextureFormatName", WriteToString<64>(BuildSettings.TextureFormatName));
-	Writer.AddBool("bHDRSource", BuildSettings.bHDRSource);
-	Writer.AddInteger("MipGenSettings", BuildSettings.MipGenSettings);
-	Writer.AddBool("bCubemap", BuildSettings.bCubemap);
-	Writer.AddBool("bTextureArray", BuildSettings.bTextureArray);
-	Writer.AddBool("bVolume", BuildSettings.bVolume);
-	Writer.AddBool("bLongLatSource", BuildSettings.bLongLatSource);
-	Writer.AddBool("bSRGB", BuildSettings.bSRGB);
-	Writer.AddBool("bUseLegacyGamma", BuildSettings.bUseLegacyGamma);
-	Writer.AddBool("bPreserveBorder", BuildSettings.bPreserveBorder);
-	Writer.AddBool("bForceNoAlphaChannel", BuildSettings.bForceNoAlphaChannel);
-	Writer.AddBool("bForceAlphaChannel", BuildSettings.bForceAlphaChannel);
-	Writer.AddBool("bDitherMipMapAlpha", BuildSettings.bDitherMipMapAlpha);
-	Writer.AddBool("bComputeBokehAlpha", BuildSettings.bComputeBokehAlpha);
-	Writer.AddBool("bReplicateRed", BuildSettings.bReplicateRed);
-	Writer.AddBool("bReplicateAlpha", BuildSettings.bReplicateAlpha);
-	Writer.AddBool("bDownsampleWithAverage", BuildSettings.bDownsampleWithAverage);
-	Writer.AddBool("bSharpenWithoutColorShift", BuildSettings.bSharpenWithoutColorShift);
-	Writer.AddBool("bBorderColorBlack", BuildSettings.bBorderColorBlack);
-	Writer.AddBool("bFlipGreenChannel", BuildSettings.bFlipGreenChannel);
-	Writer.AddBool("bApplyYCoCgBlockScale", BuildSettings.bApplyYCoCgBlockScale);
-	Writer.AddBool("bApplyKernelToTopMip", BuildSettings.bApplyKernelToTopMip);
-	Writer.AddBool("bRenormalizeTopMip", BuildSettings.bRenormalizeTopMip);
-	Writer.AddInteger("CompositeTextureMode", BuildSettings.CompositeTextureMode);
-	Writer.AddFloat("CompositePower", BuildSettings.CompositePower);
-	Writer.AddInteger("LODBias", BuildSettings.LODBias);
-	Writer.AddInteger("LODBiasWithCinematicMips", BuildSettings.LODBiasWithCinematicMips);
-
-	WriteCbField(Writer, "TopMipSize", BuildSettings.TopMipSize);
-
-	Writer.AddInteger("VolumeSizeZ", BuildSettings.VolumeSizeZ);
-	Writer.AddInteger("ArraySlices", BuildSettings.ArraySlices);
-	Writer.AddBool("bStreamable", BuildSettings.bStreamable);
-	Writer.AddBool("bVirtualStreamable", BuildSettings.bVirtualStreamable);
-	Writer.AddBool("bChromaKeyTexture", BuildSettings.bChromaKeyTexture);
-	Writer.AddInteger("PowerOfTwoMode", BuildSettings.PowerOfTwoMode);
-	WriteCbField(Writer, "PaddingColor", BuildSettings.PaddingColor);
-	WriteCbField(Writer, "ChromaKeyColor", BuildSettings.ChromaKeyColor);
-	Writer.AddFloat("ChromaKeyThreshold", BuildSettings.ChromaKeyThreshold);
-	Writer.AddInteger("CompressionQuality", BuildSettings.CompressionQuality);
-	Writer.AddInteger("LossyCompressionAmount", BuildSettings.LossyCompressionAmount);
-	Writer.AddFloat("Downscale", BuildSettings.Downscale);
-	Writer.AddInteger("DownscaleOptions", BuildSettings.DownscaleOptions);
-	Writer.AddInteger("VirtualAddressingModeX", BuildSettings.VirtualAddressingModeX);
-	Writer.AddInteger("VirtualAddressingModeY", BuildSettings.VirtualAddressingModeY);
-	Writer.AddInteger("VirtualTextureTileSize", BuildSettings.VirtualTextureTileSize);
-	Writer.AddInteger("VirtualTextureBorderSize", BuildSettings.VirtualTextureBorderSize);
-	Writer.AddBool("bVirtualTextureEnableCompressZlib", BuildSettings.bVirtualTextureEnableCompressZlib);
-	Writer.AddBool("bVirtualTextureEnableCompressCrunch", BuildSettings.bVirtualTextureEnableCompressCrunch);
-	Writer.AddBool("bHasEditorOnlyData", BuildSettings.bHasEditorOnlyData);
-
-	Writer.EndObject();
-	return Writer.Save().AsObject();
-}
-
-static FCbObject WriteOutputSettingsToCompactBinary(int32 NumInlineMips, const FString& KeySuffix)
-{
-	FCbWriter Writer;
-	Writer.BeginObject();
-
-	Writer.AddInteger("NumInlineMips", NumInlineMips);
-	
-	FString MipDerivedDataKey;
-	FTexture2DMipMap DummyMip;
-	DummyMip.SizeX = 0;
-	DummyMip.SizeY = 0;
-	GetTextureDerivedMipKey(0, DummyMip, KeySuffix, MipDerivedDataKey);
-	int32 PrefixEndIndex = MipDerivedDataKey.Find(TEXT("_MIP0_"), ESearchCase::CaseSensitive);
-	check(PrefixEndIndex != -1);
-	MipDerivedDataKey.LeftInline(PrefixEndIndex);
-	check(!MipDerivedDataKey.IsEmpty());
-	Writer.AddString("MipKeyPrefix",*MipDerivedDataKey);
-
-	Writer.EndObject();
-	return Writer.Save().AsObject();
-}
-
-static FCbObject WriteTextureSourceToCompactBinary(const FTextureSource& TextureSource, EGammaSpace GammaSpace)
-{
-	FCbWriter Writer;
-	Writer.BeginObject();
-
-	Writer.AddString("Input", TextureSource.GetIdString());
-	Writer.AddInteger("CompressionFormat", TextureSource.GetSourceCompression());
-	Writer.AddInteger("SourceFormat", TextureSource.GetFormat());
-	Writer.AddInteger("GammaSpace", (int32)GammaSpace);
-	Writer.AddInteger("NumSlices", TextureSource.GetNumSlices());
-	Writer.AddInteger("SizeX", TextureSource.GetSizeX());
-	Writer.AddInteger("SizeY", TextureSource.GetSizeY());
-	Writer.BeginArray("Mips");
-	const int32 NumMips = TextureSource.GetNumMips();
-	int64 Offset = 0;
-	for (int32 MipIndex = 0; MipIndex < NumMips; ++MipIndex)
-	{
-		Writer.BeginObject();
-		Writer.AddInteger("Offset", Offset);
-		const int64 MipSize = TextureSource.CalcMipSize(MipIndex);
-		Writer.AddInteger("Size", MipSize);
-		Offset += MipSize;
-		Writer.EndObject();
-	}
-	Writer.EndArray();
-
-	Writer.EndObject();
-	return Writer.Save().AsObject();
-}
-
 static FIoHash HashAndWriteToCompressedBufferFile(const FString& Directory, const void* InData, uint64 InDataSize)
 {
 	FIoHash DataHash = FIoHash::HashBuffer(InData, InDataSize);
@@ -315,18 +150,10 @@ void FTextureDerivedDataBuildExporter::ExportTextureBuild(const UTexture& Textur
 	
 	IPlatformFile& PlatformFile = FPlatformFileManager::Get().GetPlatformFile();
 	bool bExporting = PlatformFile.CreateDirectoryTree(*ExportRoot);
-	
-	const ITextureFormat* TextureFormat = nullptr;
-	FName TextureFormatModuleName;
-	ITextureFormatModule* TextureFormatModule = nullptr;
 
-	ITextureFormatManagerModule* TFM = GetTextureFormatManager();
-	if (TFM)
-	{
-		TextureFormat = TFM->FindTextureFormatAndModule(BuildSettings.TextureFormatName, TextureFormatModuleName, TextureFormatModule);
-	}
+	BuildFunctionName = GetTextureBuildFunctionName(BuildSettings);
 
-	if (TextureFormat == nullptr)
+	if (BuildFunctionName.IsEmpty())
 	{
 		bExporting = false;
 	}
@@ -335,32 +162,16 @@ void FTextureDerivedDataBuildExporter::ExportTextureBuild(const UTexture& Textur
 	{
 		return;
 	}
-	// Texture format modules are inconsistent in their naming.  eg: TextureFormatUncompressed, PS5TextureFormat
-	// We attempt to  unify the naming here when specifying build function names.
-	TStringBuilder<64> BuildFunctionNameBuilder;
-	BuildFunctionNameBuilder << TextureFormatModuleName.ToString().Replace(TEXT("TextureFormat"), TEXT(""));
-	BuildFunctionNameBuilder << TEXT("Texture");
 
-	BuildFunctionName = BuildFunctionNameBuilder;
 	TexturePath = Texture.GetPathName();
 
 	UE::DerivedData::FBuildActionBuilder ActionBuilder = DerivedDataBuild->CreateAction(TexturePath, BuildFunctionName);
 
-	ActionBuilder.AddConstant(TEXT("TextureBuildSettings"), WriteBuildSettingsToCompactBinary(BuildSettings, TextureFormat));
-	ActionBuilder.AddConstant(TEXT("TextureOutputSettings"), WriteOutputSettingsToCompactBinary(NumInlineMips, KeySuffix));
-
-	FTextureFormatSettings TextureFormatSettings;
-	Texture.GetLayerFormatSettings(LayerIndex, TextureFormatSettings);
-	EGammaSpace TextureGammaSpace = TextureFormatSettings.SRGB ? (Texture.bUseLegacyGamma ? EGammaSpace::Pow22 : EGammaSpace::sRGB) : EGammaSpace::Linear;
-	ActionBuilder.AddConstant(TEXT("TextureSource"), WriteTextureSourceToCompactBinary(Texture.Source, TextureGammaSpace));
-	if ((bool)Texture.CompositeTexture && !ExportedCompositeTextureBulkDataHash.IsZero())
-	{
-		FTextureFormatSettings CompositeTextureFormatSettings;
-		Texture.CompositeTexture->GetLayerFormatSettings(LayerIndex, CompositeTextureFormatSettings);
-		EGammaSpace CompositeTextureGammaSpace = CompositeTextureFormatSettings.SRGB ? (Texture.CompositeTexture->bUseLegacyGamma ? EGammaSpace::Pow22 : EGammaSpace::sRGB) : EGammaSpace::Linear;
-
-		ActionBuilder.AddConstant(TEXT("CompositeTextureSource"), WriteTextureSourceToCompactBinary(Texture.CompositeTexture->Source, CompositeTextureGammaSpace));
-	}
+	ComposeTextureBuildFunctionConstants(KeySuffix, Texture, BuildSettings, LayerIndex, NumInlineMips, 
+		[&ActionBuilder] (FStringView Key, const FCbObject& Value)
+		{
+			ActionBuilder.AddConstant(Key, Value);
+		});
 
 	ActionBuilder.AddInput(Texture.Source.GetIdString(), ExportedTextureBulkDataHash, ExportedTextureBulkDataSize);
 	if ((bool)Texture.CompositeTexture && !ExportedCompositeTextureBulkDataHash.IsZero())
diff --git a/Engine/Source/Runtime/Engine/Private/TextureDerivedDataBuildUtils.cpp b/Engine/Source/Runtime/Engine/Private/TextureDerivedDataBuildUtils.cpp
new file mode 100644
--- /dev/null
+++ b/Engine/Source/Runtime/Engine/Private/TextureDerivedDataBuildUtils.cpp
@@ -0,0 +1,246 @@
+// Copyright Epic Games, Inc. All Rights Reserved.
+
+#include "TextureDerivedDataBuildUtils.h"
+
+#if WITH_EDITOR
+#include "Engine/Texture.h"
+#include "Interfaces/ITextureFormat.h"
+#include "Interfaces/ITextureFormat.h"
+#include "Interfaces/ITextureFormatManagerModule.h"
+#include "Interfaces/ITextureFormatModule.h"
+#include "Serialization/CompactBinary.h"
+#include "Serialization/CompactBinaryWriter.h"
+#include "TextureCompressorModule.h"
+#include "TextureFormatManager.h"
+#include "TextureResource.h"
+
+void GetTextureDerivedDataKeyFromSuffix(const FString& KeySuffix, FString& OutKey);
+void GetTextureDerivedMipKey(int32 MipIndex, const FTexture2DMipMap& Mip, const FString& KeySuffix, FString& OutKey);
+
+static void WriteCbField(FCbWriter& Writer, const FAnsiStringView& Name, const FColor& Color)
+{
+	// Choosing the big endian ordering
+	Writer.BeginArray(Name);
+	Writer.AddInteger(Color.A);
+	Writer.AddInteger(Color.R);
+	Writer.AddInteger(Color.G);
+	Writer.AddInteger(Color.B);
+	Writer.EndArray();
+}
+
+static void WriteCbField(FCbWriter& Writer, const FAnsiStringView& Name, const FVector4& Vec4)
+{
+	Writer.BeginArray(Name);
+	Writer.AddFloat(Vec4.X);
+	Writer.AddFloat(Vec4.Y);
+	Writer.AddFloat(Vec4.Z);
+	Writer.AddFloat(Vec4.W);
+	Writer.EndArray();
+}
+
+static void WriteCbField(FCbWriter& Writer, const FAnsiStringView& Name, const FIntPoint& IntPoint)
+{
+	Writer.BeginArray(Name);
+	Writer.AddInteger(IntPoint.X);
+	Writer.AddInteger(IntPoint.Y);
+	Writer.EndArray();
+}
+
+static FCbObject WriteBuildSettingsToCompactBinary(const FTextureBuildSettings& BuildSettings, const ITextureFormat* TextureFormat)
+{
+	FCbWriter Writer;
+	Writer.BeginObject();
+
+	if (BuildSettings.FormatConfigOverride)
+	{
+		Writer.AddObject("FormatConfigOverride", BuildSettings.FormatConfigOverride);
+	}
+	else if (FCbObject TextureFormatConfig = TextureFormat->ExportGlobalFormatConfig(BuildSettings))
+	{
+		Writer.AddObject("FormatConfigOverride", TextureFormatConfig);
+	}
+	Writer.BeginObject("ColorAdjustment");
+	Writer.AddFloat("AdjustBrightness", BuildSettings.ColorAdjustment.AdjustBrightness);
+	Writer.AddFloat("AdjustBrightnessCurve", BuildSettings.ColorAdjustment.AdjustBrightnessCurve);
+	Writer.AddFloat("AdjustSaturation", BuildSettings.ColorAdjustment.AdjustSaturation);
+	Writer.AddFloat("AdjustVibrance", BuildSettings.ColorAdjustment.AdjustVibrance);
+	Writer.AddFloat("AdjustRGBCurve", BuildSettings.ColorAdjustment.AdjustRGBCurve);
+	Writer.AddFloat("AdjustHue", BuildSettings.ColorAdjustment.AdjustHue);
+	Writer.AddFloat("AdjustMinAlpha", BuildSettings.ColorAdjustment.AdjustMinAlpha);
+	Writer.AddFloat("AdjustMaxAlpha", BuildSettings.ColorAdjustment.AdjustMaxAlpha);
+	Writer.EndObject();
+
+	WriteCbField(Writer, "AlphaCoverageThresholds", BuildSettings.AlphaCoverageThresholds);
+
+	Writer.AddFloat("MipSharpening", BuildSettings.MipSharpening);
+	Writer.AddInteger("DiffuseConvolveMipLevel", BuildSettings.DiffuseConvolveMipLevel);
+	Writer.AddInteger("SharpenMipKernelSize", BuildSettings.SharpenMipKernelSize);
+	Writer.AddInteger("MaxTextureResolution", BuildSettings.MaxTextureResolution);
+	Writer.AddString("TextureFormatName", WriteToString<64>(BuildSettings.TextureFormatName));
+	Writer.AddBool("bHDRSource", BuildSettings.bHDRSource);
+	Writer.AddInteger("MipGenSettings", BuildSettings.MipGenSettings);
+	Writer.AddBool("bCubemap", BuildSettings.bCubemap);
+	Writer.AddBool("bTextureArray", BuildSettings.bTextureArray);
+	Writer.AddBool("bVolume", BuildSettings.bVolume);
+	Writer.AddBool("bLongLatSource", BuildSettings.bLongLatSource);
+	Writer.AddBool("bSRGB", BuildSettings.bSRGB);
+	Writer.AddBool("bUseLegacyGamma", BuildSettings.bUseLegacyGamma);
+	Writer.AddBool("bPreserveBorder", BuildSettings.bPreserveBorder);
+	Writer.AddBool("bForceNoAlphaChannel", BuildSettings.bForceNoAlphaChannel);
+	Writer.AddBool("bForceAlphaChannel", BuildSettings.bForceAlphaChannel);
+	Writer.AddBool("bDitherMipMapAlpha", BuildSettings.bDitherMipMapAlpha);
+	Writer.AddBool("bComputeBokehAlpha", BuildSettings.bComputeBokehAlpha);
+	Writer.AddBool("bReplicateRed", BuildSettings.bReplicateRed);
+	Writer.AddBool("bReplicateAlpha", BuildSettings.bReplicateAlpha);
+	Writer.AddBool("bDownsampleWithAverage", BuildSettings.bDownsampleWithAverage);
+	Writer.AddBool("bSharpenWithoutColorShift", BuildSettings.bSharpenWithoutColorShift);
+	Writer.AddBool("bBorderColorBlack", BuildSettings.bBorderColorBlack);
+	Writer.AddBool("bFlipGreenChannel", BuildSettings.bFlipGreenChannel);
+	Writer.AddBool("bApplyYCoCgBlockScale", BuildSettings.bApplyYCoCgBlockScale);
+	Writer.AddBool("bApplyKernelToTopMip", BuildSettings.bApplyKernelToTopMip);
+	Writer.AddBool("bRenormalizeTopMip", BuildSettings.bRenormalizeTopMip);
+	Writer.AddInteger("CompositeTextureMode", BuildSettings.CompositeTextureMode);
+	Writer.AddFloat("CompositePower", BuildSettings.CompositePower);
+	Writer.AddInteger("LODBias", BuildSettings.LODBias);
+	Writer.AddInteger("LODBiasWithCinematicMips", BuildSettings.LODBiasWithCinematicMips);
+
+	WriteCbField(Writer, "TopMipSize", BuildSettings.TopMipSize);
+
+	Writer.AddInteger("VolumeSizeZ", BuildSettings.VolumeSizeZ);
+	Writer.AddInteger("ArraySlices", BuildSettings.ArraySlices);
+	Writer.AddBool("bStreamable", BuildSettings.bStreamable);
+	Writer.AddBool("bVirtualStreamable", BuildSettings.bVirtualStreamable);
+	Writer.AddBool("bChromaKeyTexture", BuildSettings.bChromaKeyTexture);
+	Writer.AddInteger("PowerOfTwoMode", BuildSettings.PowerOfTwoMode);
+	WriteCbField(Writer, "PaddingColor", BuildSettings.PaddingColor);
+	WriteCbField(Writer, "ChromaKeyColor", BuildSettings.ChromaKeyColor);
+	Writer.AddFloat("ChromaKeyThreshold", BuildSettings.ChromaKeyThreshold);
+	Writer.AddInteger("CompressionQuality", BuildSettings.CompressionQuality);
+	Writer.AddInteger("LossyCompressionAmount", BuildSettings.LossyCompressionAmount);
+	Writer.AddFloat("Downscale", BuildSettings.Downscale);
+	Writer.AddInteger("DownscaleOptions", BuildSettings.DownscaleOptions);
+	Writer.AddInteger("VirtualAddressingModeX", BuildSettings.VirtualAddressingModeX);
+	Writer.AddInteger("VirtualAddressingModeY", BuildSettings.VirtualAddressingModeY);
+	Writer.AddInteger("VirtualTextureTileSize", BuildSettings.VirtualTextureTileSize);
+	Writer.AddInteger("VirtualTextureBorderSize", BuildSettings.VirtualTextureBorderSize);
+	Writer.AddBool("bVirtualTextureEnableCompressZlib", BuildSettings.bVirtualTextureEnableCompressZlib);
+	Writer.AddBool("bVirtualTextureEnableCompressCrunch", BuildSettings.bVirtualTextureEnableCompressCrunch);
+	Writer.AddBool("bHasEditorOnlyData", BuildSettings.bHasEditorOnlyData);
+
+	Writer.EndObject();
+	return Writer.Save().AsObject();
+}
+
+static FCbObject WriteOutputSettingsToCompactBinary(int32 NumInlineMips, const FString& KeySuffix)
+{
+	FCbWriter Writer;
+	Writer.BeginObject();
+
+	Writer.AddInteger("NumInlineMips", NumInlineMips);
+	
+	FString MipDerivedDataKey;
+	FTexture2DMipMap DummyMip;
+	DummyMip.SizeX = 0;
+	DummyMip.SizeY = 0;
+	GetTextureDerivedMipKey(0, DummyMip, KeySuffix, MipDerivedDataKey);
+	int32 PrefixEndIndex = MipDerivedDataKey.Find(TEXT("_MIP0_"), ESearchCase::CaseSensitive);
+	check(PrefixEndIndex != -1);
+	MipDerivedDataKey.LeftInline(PrefixEndIndex);
+	check(!MipDerivedDataKey.IsEmpty());
+	Writer.AddString("MipKeyPrefix",*MipDerivedDataKey);
+
+	Writer.EndObject();
+	return Writer.Save().AsObject();
+}
+
+static FCbObject WriteTextureSourceToCompactBinary(const FTextureSource& TextureSource, EGammaSpace GammaSpace)
+{
+	FCbWriter Writer;
+	Writer.BeginObject();
+
+	Writer.AddString("Input", TextureSource.GetIdString());
+	Writer.AddInteger("CompressionFormat", TextureSource.GetSourceCompression());
+	Writer.AddInteger("SourceFormat", TextureSource.GetFormat());
+	Writer.AddInteger("GammaSpace", (int32)GammaSpace);
+	Writer.AddInteger("NumSlices", TextureSource.GetNumSlices());
+	Writer.AddInteger("SizeX", TextureSource.GetSizeX());
+	Writer.AddInteger("SizeY", TextureSource.GetSizeY());
+	Writer.BeginArray("Mips");
+	const int32 NumMips = TextureSource.GetNumMips();
+	int64 Offset = 0;
+	for (int32 MipIndex = 0; MipIndex < NumMips; ++MipIndex)
+	{
+		Writer.BeginObject();
+		Writer.AddInteger("Offset", Offset);
+		const int64 MipSize = TextureSource.CalcMipSize(MipIndex);
+		Writer.AddInteger("Size", MipSize);
+		Offset += MipSize;
+		Writer.EndObject();
+	}
+	Writer.EndArray();
+
+	Writer.EndObject();
+	return Writer.Save().AsObject();
+}
+
+FString GetTextureBuildFunctionName(const FTextureBuildSettings& BuildSettings)
+{
+	const ITextureFormat* TextureFormat = nullptr;
+	FName TextureFormatModuleName;
+	ITextureFormatModule* TextureFormatModule = nullptr;
+
+	ITextureFormatManagerModule* TFM = GetTextureFormatManager();
+	if (TFM)
+	{
+		TextureFormat = TFM->FindTextureFormatAndModule(BuildSettings.TextureFormatName, TextureFormatModuleName, TextureFormatModule);
+	}
+
+	if (TextureFormat == nullptr)
+	{
+		return FString();
+	}
+
+	// Texture format modules are inconsistent in their naming.  eg: TextureFormatUncompressed, PS5TextureFormat
+	// We attempt to  unify the naming here when specifying build function names.
+	TStringBuilder<64> BuildFunctionNameBuilder;
+	BuildFunctionNameBuilder << TextureFormatModuleName.ToString().Replace(TEXT("TextureFormat"), TEXT(""));
+	BuildFunctionNameBuilder << TEXT("Texture");
+
+	return BuildFunctionNameBuilder.ToString();
+}
+
+void ComposeTextureBuildFunctionConstants(const FString& KeySuffix, const UTexture& Texture, const FTextureBuildSettings& BuildSettings, int32 LayerIndex, int32 NumInlineMips, FTextureConstantOperator Operator)
+{
+	const ITextureFormat* TextureFormat = nullptr;
+	FName TextureFormatModuleName;
+	ITextureFormatModule* TextureFormatModule = nullptr;
+
+	ITextureFormatManagerModule* TFM = GetTextureFormatManager();
+	if (TFM)
+	{
+		TextureFormat = TFM->FindTextureFormatAndModule(BuildSettings.TextureFormatName, TextureFormatModuleName, TextureFormatModule);
+	}
+
+	if (TextureFormat == nullptr)
+	{
+		return;
+	}
+
+	Operator(TEXT("TextureBuildSettings"), WriteBuildSettingsToCompactBinary(BuildSettings, TextureFormat));
+	Operator(TEXT("TextureOutputSettings"), WriteOutputSettingsToCompactBinary(NumInlineMips, KeySuffix));
+
+	FTextureFormatSettings TextureFormatSettings;
+	Texture.GetLayerFormatSettings(LayerIndex, TextureFormatSettings);
+	EGammaSpace TextureGammaSpace = TextureFormatSettings.SRGB ? (Texture.bUseLegacyGamma ? EGammaSpace::Pow22 : EGammaSpace::sRGB) : EGammaSpace::Linear;
+	Operator(TEXT("TextureSource"), WriteTextureSourceToCompactBinary(Texture.Source, TextureGammaSpace));
+	if ((bool)Texture.CompositeTexture)
+	{
+		FTextureFormatSettings CompositeTextureFormatSettings;
+		Texture.CompositeTexture->GetLayerFormatSettings(LayerIndex, CompositeTextureFormatSettings);
+		EGammaSpace CompositeTextureGammaSpace = CompositeTextureFormatSettings.SRGB ? (Texture.CompositeTexture->bUseLegacyGamma ? EGammaSpace::Pow22 : EGammaSpace::sRGB) : EGammaSpace::Linear;
+
+		Operator(TEXT("CompositeTextureSource"), WriteTextureSourceToCompactBinary(Texture.CompositeTexture->Source, CompositeTextureGammaSpace));
+	}
+}
+
+#endif // WITH_EDITOR
diff --git a/Engine/Source/Runtime/Engine/Private/TextureDerivedDataBuildUtils.h b/Engine/Source/Runtime/Engine/Private/TextureDerivedDataBuildUtils.h
new file mode 100644
--- /dev/null
+++ b/Engine/Source/Runtime/Engine/Private/TextureDerivedDataBuildUtils.h
@@ -0,0 +1,23 @@
+// Copyright Epic Games, Inc. All Rights Reserved.
+
+#pragma once
+
+#if WITH_EDITOR
+
+#include "CoreMinimal.h"
+#include "Containers/StringView.h"
+#include "Templates/Function.h"
+
+namespace UE::DerivedData { class FBuildActionBuilder; }
+
+class FCbObject;
+class UTexture;
+struct FTextureBuildSettings;
+
+using FTextureConstantOperator = TFunctionRef<void (FStringView Key, const FCbObject& Value)>;
+
+// Suitable for use with an action builder or in the future a definition builder
+FString GetTextureBuildFunctionName(const FTextureBuildSettings& BuildSettings);
+void ComposeTextureBuildFunctionConstants(const FString& KeySuffix, const UTexture& Texture, const FTextureBuildSettings& BuildSettings, int32 LayerIndex, int32 NumInlineMips, FTextureConstantOperator Operator);
+
+#endif // WITH_EDITOR
diff --git a/Engine/Source/Runtime/Engine/Private/TextureDerivedDataTask.cpp b/Engine/Source/Runtime/Engine/Private/TextureDerivedDataTask.cpp
--- a/Engine/Source/Runtime/Engine/Private/TextureDerivedDataTask.cpp
+++ b/Engine/Source/Runtime/Engine/Private/TextureDerivedDataTask.cpp
@@ -23,7 +23,12 @@
 
 #if WITH_EDITOR
 
+#include "DerivedDataBuild.h"
+#include "DerivedDataBuildAction.h"
+#include "DerivedDataBuildOutput.h"
+#include "DerivedDataBuildSession.h"
 #include "DerivedDataCacheInterface.h"
+#include "DerivedDataPayload.h"
 #include "Engine/TextureCube.h"
 #include "Engine/VolumeTexture.h"
 #include "GenericPlatform/GenericPlatformMath.h"
@@ -32,7 +37,9 @@
 #include "Interfaces/ITextureFormat.h"
 #include "Math/NumericLimits.h"
 #include "Math/UnrealMathUtility.h"
+#include "Misc/FeedbackContext.h"
 #include "ProfilingDebugging/CookStats.h"
+#include "TextureDerivedDataBuildUtils.h"
 #include "VT/VirtualTextureDataBuilder.h"
 
 static TAutoConsoleVariable<int32> CVarVTValidateCompressionOnLoad(
@@ -49,6 +56,8 @@ static TAutoConsoleVariable<int32> CVarVTValidateCompressionOnSave(
 	TEXT("This is slow, but allows debugging corrupt VT data")
 );
 
+void GetTextureDerivedDataKeyFromSuffix(const FString& KeySuffix, FString& OutKey);
+
 class FTextureStatusMessageContext : public FScopedSlowTask
 {
 public:
@@ -246,6 +255,78 @@ void FTextureSourceData::GetAsyncSourceMips(IImageWrapperModule* InImageWrapper)
 		}
 	}
 }
+void FTextureCacheDerivedDataWorker::ConsumeBuildFunctionOutput(const UE::DerivedData::FBuildOutput& BuildOutput, const FString& TexturePath, bool bReplaceExistingDDC)
+{
+	using namespace UE::DerivedData;
+
+	FSharedBuffer PrimaryData;
+	{
+		FPayloadId PrimaryPayloadId = FPayloadId::FromName(TEXT("Texture"));
+		const FPayload& PrimaryPayload = BuildOutput.GetPayload(PrimaryPayloadId);
+		if (PrimaryPayload.IsNull())
+		{
+			UE_LOG(LogTexture, Warning, TEXT("Texture build function missing primary output payload when building %s derived data for %s"), *BuildSettingsPerLayer[0].TextureFormatName.GetPlainNameString(), *Texture.GetPathName());
+			return;
+		}
+		PrimaryData = PrimaryPayload.GetData().Decompress();
+		TArrayView<const uint8> PrimaryDataView((const uint8*)PrimaryData.GetData(), PrimaryData.GetSize());
+
+		FMemoryReaderView Ar(PrimaryDataView, /*bIsPersistent=*/ true);
+		DerivedData->Serialize(Ar, NULL);
+	}
+
+	if (DerivedData->Mips.Num())
+	{
+		const bool bInlineMips = (CacheFlags & ETextureCacheFlags::InlineMips) != 0;
+
+		for (int32 MipIndex = 0; MipIndex < DerivedData->Mips.Num(); ++MipIndex)
+		{
+			FTexture2DMipMap& Mip = DerivedData->Mips[MipIndex];
+			if (Mip.DerivedDataKey.IsEmpty())
+			{
+				break;
+			}
+
+			TStringBuilder<32> PayloadName;
+			PayloadName << "Mip" << MipIndex;
+
+			FPayloadId MipPayloadId = FPayloadId::FromName(*PayloadName);
+			const FPayload& MipPayload = BuildOutput.GetPayload(MipPayloadId);
+			if (MipPayload.IsNull())
+			{
+				UE_LOG(LogTexture, Warning, TEXT("Texture build function missing Mip%d output payload when building %s derived data for %s"), MipIndex, *BuildSettingsPerLayer[0].TextureFormatName.GetPlainNameString(), *Texture.GetPathName());
+				return;
+			}
+			FSharedBuffer MipData = MipPayload.GetData().Decompress();
+			TArrayView<const uint8> MipDataView((const uint8*)MipData.GetData(), MipData.GetSize());
+			GetDerivedDataCacheRef().Put(*Mip.DerivedDataKey, MipDataView, TexturePath, bReplaceExistingDDC);
+
+			if (bInlineMips && (MipIndex >= (int32)BuildSettingsPerLayer[0].LODBiasWithCinematicMips))
+			{
+				int32 MipSize = 0;
+				FMemoryReaderView Ar(MipDataView, /*bIsPersistent=*/ true);
+				Ar << MipSize;
+
+				Mip.BulkData.Lock(LOCK_READ_WRITE);
+				void* MipAllocData = Mip.BulkData.Realloc(MipSize);
+				Ar.Serialize(MipAllocData, MipSize);
+				Mip.BulkData.Unlock();
+				Mip.DerivedDataKey.Empty();
+			}
+		}
+
+		FString DerivedDataKey;
+		GetTextureDerivedDataKeyFromSuffix(KeySuffix, DerivedDataKey);
+		TArrayView<const uint8> PrimaryDataView((const uint8*)PrimaryData.GetData(), PrimaryData.GetSize());
+		GetDerivedDataCacheRef().Put(*DerivedDataKey, PrimaryDataView, TexturePath, bReplaceExistingDDC);
+
+		bSucceeded = true;
+	}
+	else
+	{
+		UE_LOG(LogTexture, Warning, TEXT("Failed to build %s derived data for %s"), *BuildSettingsPerLayer[0].TextureFormatName.GetPlainNameString(), *Texture.GetPathName());
+	}
+}
 
 void FTextureCacheDerivedDataWorker::BuildTexture(bool bReplaceExistingDDC)
 {
@@ -345,72 +426,117 @@ void FTextureCacheDerivedDataWorker::BuildTexture(bool bReplaceExistingDDC)
 
 		BuildExporter.ExportTextureBuild(Texture, BuildSettingsPerLayer[0], 0, NUM_INLINE_DERIVED_MIPS);
 
-		// Compress the texture.
-		TArray<FCompressedImage2D> CompressedMips;
-		if (Compressor->BuildTexture(TextureData.Blocks[0].MipsPerLayer[0],
-			((bool)Texture.CompositeTexture && CompositeTextureData.Blocks.Num() && CompositeTextureData.Blocks[0].MipsPerLayer.Num()) ? CompositeTextureData.Blocks[0].MipsPerLayer[0] : TArray<FImage>(),
-			BuildSettingsPerLayer[0],
-			CompressedMips,
-			OptData.NumMipsInTail,
-			OptData.ExtData))
+		if (!BuildFunctionName.IsEmpty())
 		{
-			check(CompressedMips.Num());
+			using namespace UE::DerivedData;
+			// Compress the texture by calling the DDC2 build function.  May be executed locally or remotely.
+			// The FBuildAction built here will be superceded by an FBuildDefinition in the future that will
+			// replace both the cache query/fill as well as execution.
+			IBuild& Build = GetDerivedDataBuildRef();
+			FString TexturePath = Texture.GetPathName();
+			FBuildActionBuilder ActionBuilder = Build.CreateAction(TexturePath, BuildFunctionName);
+
+			ComposeTextureBuildFunctionConstants(KeySuffix, Texture, BuildSettingsPerLayer[0], 0, NUM_INLINE_DERIVED_MIPS, 
+				[&ActionBuilder] (FStringView Key, const FCbObject& Value)
+				{
+					ActionBuilder.AddConstant(Key, Value);
+				});
 
-			// Build the derived data.
-			const int32 MipCount = CompressedMips.Num();
-			for (int32 MipIndex = 0; MipIndex < MipCount; ++MipIndex)
+			for (const FBuildInputRecord& InputRecord : BuildInputRecords)
 			{
-				const FCompressedImage2D& CompressedImage = CompressedMips[MipIndex];
-				FTexture2DMipMap* NewMip = new FTexture2DMipMap();
-				DerivedData->Mips.Add(NewMip);
-				NewMip->SizeX = CompressedImage.SizeX;
-				NewMip->SizeY = CompressedImage.SizeY;
-				NewMip->SizeZ = CompressedImage.SizeZ;
-				NewMip->FileRegionType = FFileRegion::SelectType(EPixelFormat(CompressedImage.PixelFormat));
-				check(NewMip->SizeZ == 1 || BuildSettingsPerLayer[0].bVolume || BuildSettingsPerLayer[0].bTextureArray); // Only volume & arrays can have SizeZ != 1
-				NewMip->BulkData.Lock(LOCK_READ_WRITE);
-				check(CompressedImage.RawData.GetTypeSize() == 1);
-				void* NewMipData = NewMip->BulkData.Realloc(CompressedImage.RawData.Num());
-				FMemory::Memcpy(NewMipData, CompressedImage.RawData.GetData(), CompressedImage.RawData.Num());
-				NewMip->BulkData.Unlock();
-
-				if (MipIndex == 0)
+				ActionBuilder.AddInput(InputRecord.Id, InputRecord.RawHash, InputRecord.RawSize);
+			}
+
+			FBuildSession Session = Build.CreateSession(TexturePath, nullptr);
+			Session.BuildAction(ActionBuilder.Build(), InputsBuilder.Build(), EBuildPolicy::Default, EPriority::Blocking,
+				[this, &TexturePath, bReplaceExistingDDC] (FBuildActionCompleteParams&& Params)
 				{
-					DerivedData->SizeX = CompressedImage.SizeX;
-					DerivedData->SizeY = CompressedImage.SizeY;
-					DerivedData->PixelFormat = (EPixelFormat)CompressedImage.PixelFormat;
-					DerivedData->SetNumSlices(BuildSettingsPerLayer[0].bCubemap ? 6 : (BuildSettingsPerLayer[0].bVolume || BuildSettingsPerLayer[0].bTextureArray) ? CompressedImage.SizeZ : 1);
-					DerivedData->SetIsCubemap(BuildSettingsPerLayer[0].bCubemap);
-				}
-				else
+					#if !NO_LOGGING
+					Params.Output.IterateDiagnostics([this, &TexturePath] (const FBuildDiagnostic& Diagnostic)
+						{
+							if (GWarn)
+							{
+								FName CategoryName(Diagnostic.Category);
+								GWarn->Log(CategoryName, Diagnostic.Level == EBuildDiagnosticLevel::Error ? ELogVerbosity::Error : ELogVerbosity::Warning, *FString(Diagnostic.Message));
+							}
+						});
+					#endif
+
+					if (Params.Status == EStatus::Ok)
+					{
+						ConsumeBuildFunctionOutput(Params.Output, TexturePath, bReplaceExistingDDC);
+					}
+				}).Wait();
+		}
+		else
+		{
+			// Compress the texture by calling texture compressor directly.
+			TArray<FCompressedImage2D> CompressedMips;
+			if (Compressor->BuildTexture(TextureData.Blocks[0].MipsPerLayer[0],
+				((bool)Texture.CompositeTexture && CompositeTextureData.Blocks.Num() && CompositeTextureData.Blocks[0].MipsPerLayer.Num()) ? CompositeTextureData.Blocks[0].MipsPerLayer[0] : TArray<FImage>(),
+				BuildSettingsPerLayer[0],
+				CompressedMips,
+				OptData.NumMipsInTail,
+				OptData.ExtData))
+			{
+				check(CompressedMips.Num());
+
+				// Build the derived data.
+				const int32 MipCount = CompressedMips.Num();
+				for (int32 MipIndex = 0; MipIndex < MipCount; ++MipIndex)
 				{
-					check(CompressedImage.PixelFormat == DerivedData->PixelFormat);
+					const FCompressedImage2D& CompressedImage = CompressedMips[MipIndex];
+					FTexture2DMipMap* NewMip = new FTexture2DMipMap();
+					DerivedData->Mips.Add(NewMip);
+					NewMip->SizeX = CompressedImage.SizeX;
+					NewMip->SizeY = CompressedImage.SizeY;
+					NewMip->SizeZ = CompressedImage.SizeZ;
+					NewMip->FileRegionType = FFileRegion::SelectType(EPixelFormat(CompressedImage.PixelFormat));
+					check(NewMip->SizeZ == 1 || BuildSettingsPerLayer[0].bVolume || BuildSettingsPerLayer[0].bTextureArray); // Only volume & arrays can have SizeZ != 1
+					NewMip->BulkData.Lock(LOCK_READ_WRITE);
+					check(CompressedImage.RawData.GetTypeSize() == 1);
+					void* NewMipData = NewMip->BulkData.Realloc(CompressedImage.RawData.Num());
+					FMemory::Memcpy(NewMipData, CompressedImage.RawData.GetData(), CompressedImage.RawData.Num());
+					NewMip->BulkData.Unlock();
+
+					if (MipIndex == 0)
+					{
+						DerivedData->SizeX = CompressedImage.SizeX;
+						DerivedData->SizeY = CompressedImage.SizeY;
+						DerivedData->PixelFormat = (EPixelFormat)CompressedImage.PixelFormat;
+						DerivedData->SetNumSlices(BuildSettingsPerLayer[0].bCubemap ? 6 : (BuildSettingsPerLayer[0].bVolume || BuildSettingsPerLayer[0].bTextureArray) ? CompressedImage.SizeZ : 1);
+						DerivedData->SetIsCubemap(BuildSettingsPerLayer[0].bCubemap);
+					}
+					else
+					{
+						check(CompressedImage.PixelFormat == DerivedData->PixelFormat);
+					}
 				}
-			}
 
-			DerivedData->SetOptData(OptData);
+				DerivedData->SetOptData(OptData);
 
-			BuildExporter.ExportTextureOutput(*DerivedData, BuildSettingsPerLayer[0]);
-			
-			// Store it in the cache.
-			// @todo: This will remove the streaming bulk data, which we immediately reload below!
-			// Should ideally avoid this redundant work, but it only happens when we actually have 
-			// to build the texture, which should only ever be once.
-			this->BytesCached = PutDerivedDataInCache(DerivedData, KeySuffix, Texture.GetPathName(), BuildSettingsPerLayer[0].bCubemap || (BuildSettingsPerLayer[0].bVolume && !GSupportsVolumeTextureStreaming) || (BuildSettingsPerLayer[0].bTextureArray && !GSupportsTexture2DArrayStreaming), bReplaceExistingDDC);
-		}
+				BuildExporter.ExportTextureOutput(*DerivedData, BuildSettingsPerLayer[0]);
+				
+				// Store it in the cache.
+				// @todo: This will remove the streaming bulk data, which we immediately reload below!
+				// Should ideally avoid this redundant work, but it only happens when we actually have 
+				// to build the texture, which should only ever be once.
+				this->BytesCached = PutDerivedDataInCache(DerivedData, KeySuffix, Texture.GetPathName(), BuildSettingsPerLayer[0].bCubemap || (BuildSettingsPerLayer[0].bVolume && !GSupportsVolumeTextureStreaming) || (BuildSettingsPerLayer[0].bTextureArray && !GSupportsTexture2DArrayStreaming), bReplaceExistingDDC);
+			}
 
-		if (DerivedData->Mips.Num())
-		{
-			const bool bInlineMips = (CacheFlags & ETextureCacheFlags::InlineMips) != 0;
-			bSucceeded = !bInlineMips || DerivedData->TryInlineMipData(BuildSettingsPerLayer[0].LODBiasWithCinematicMips, &Texture);
-			if (!bSucceeded)
+			if (DerivedData->Mips.Num())
 			{
-				UE_LOG(LogTexture, Display, TEXT("Failed to put and then read back mipmap data from DDC for %s"), *Texture.GetPathName());
+				const bool bInlineMips = (CacheFlags & ETextureCacheFlags::InlineMips) != 0;
+				bSucceeded = !bInlineMips || DerivedData->TryInlineMipData(BuildSettingsPerLayer[0].LODBiasWithCinematicMips, &Texture);
+				if (!bSucceeded)
+				{
+					UE_LOG(LogTexture, Display, TEXT("Failed to put and then read back mipmap data from DDC for %s"), *Texture.GetPathName());
+				}
+			}
+			else
+			{
+				UE_LOG(LogTexture, Warning, TEXT("Failed to build %s derived data for %s"), *BuildSettingsPerLayer[0].TextureFormatName.GetPlainNameString(), *Texture.GetPathName());
 			}
-		}
-		else
-		{
-			UE_LOG(LogTexture, Warning, TEXT("Failed to build %s derived data for %s"), *BuildSettingsPerLayer[0].TextureFormatName.GetPlainNameString(), *Texture.GetPathName());
 		}
 	}
 }
@@ -426,6 +552,7 @@ FTextureCacheDerivedDataWorker::FTextureCacheDerivedDataWorker(
 	, ImageWrapper(nullptr)
 	, DerivedData(InDerivedData)
 	, Texture(*InTexture)
+	, InputsBuilder(GetDerivedDataBuildRef().CreateInputs(InTexture->GetPathName()))
 	, CacheFlags(InCacheFlags)
 	, RequiredMemoryEstimate(InTexture->GetBuildRequiredMemory())
 	, bSucceeded(false)
@@ -461,6 +588,12 @@ FTextureCacheDerivedDataWorker::FTextureCacheDerivedDataWorker(
 	const bool bAllowAsyncLoading = (CacheFlags & ETextureCacheFlags::AllowAsyncLoading) != 0;
 	const bool bForVirtualTextureStreamingBuild = (CacheFlags & ETextureCacheFlags::ForVirtualTextureStreamingBuild) != 0;
 
+	static const bool bBuildFunctionEnabled = FParse::Param(FCommandLine::Get(), TEXT("DDC2TextureBuilds"));
+	if (bBuildFunctionEnabled && !bForVirtualTextureStreamingBuild && (BuildSettingsPerLayer.Num() == 1))
+	{
+		BuildFunctionName = GetTextureBuildFunctionName(BuildSettingsPerLayer[0]);
+	}
+
 	// FVirtualTextureDataBuilder always wants to load ImageWrapper module
 	// This is not strictly necessary, used only for debug output, but seems simpler to just always load this here, doesn't seem like it should be too expensive
 	if (bAllowAsyncLoading || bForVirtualTextureStreamingBuild)
@@ -627,6 +760,7 @@ void FTextureCacheDerivedDataWorker::DoWork()
 		{
 			TextureData.GetSourceMips(Texture.Source, ImageWrapper);
 			BuildExporter.ExportTextureSourceBulkData(Texture.Source);
+			ConditionalAddBuildFunctionInput(Texture.Source);
 			bHasTextureSourceMips = true;
 		}
 
@@ -635,6 +769,7 @@ void FTextureCacheDerivedDataWorker::DoWork()
 		{
 			CompositeTextureData.GetSourceMips(Texture.CompositeTexture->Source, ImageWrapper);
 			BuildExporter.ExportCompositeTextureSourceBulkData(Texture.CompositeTexture->Source);
+			ConditionalAddBuildFunctionInput(Texture.CompositeTexture->Source);
 			bHasCompositeTextureSourceMips = true;
 		}
 
@@ -642,6 +777,7 @@ void FTextureCacheDerivedDataWorker::DoWork()
 		{
 			TextureData.GetAsyncSourceMips(ImageWrapper);
 			BuildExporter.ExportTextureSourceBulkData(TextureData.AsyncSource);
+			ConditionalAddBuildFunctionInput(TextureData.AsyncSource);
 			TextureData.AsyncSource.RemoveBulkData();
 		}
 
@@ -649,9 +785,10 @@ void FTextureCacheDerivedDataWorker::DoWork()
 		{
 			CompositeTextureData.GetAsyncSourceMips(ImageWrapper);
 			if ((bool)Texture.CompositeTexture)
-		{
-			BuildExporter.ExportCompositeTextureSourceBulkData(CompositeTextureData.AsyncSource);
-		}
+			{
+				BuildExporter.ExportCompositeTextureSourceBulkData(CompositeTextureData.AsyncSource);
+				ConditionalAddBuildFunctionInput(CompositeTextureData.AsyncSource);
+			}
 			CompositeTextureData.AsyncSource.RemoveBulkData();
 		}
 
diff --git a/Engine/Source/Runtime/Engine/Private/TextureDerivedDataTask.h b/Engine/Source/Runtime/Engine/Private/TextureDerivedDataTask.h
--- a/Engine/Source/Runtime/Engine/Private/TextureDerivedDataTask.h
+++ b/Engine/Source/Runtime/Engine/Private/TextureDerivedDataTask.h
@@ -11,10 +11,11 @@
 #if WITH_EDITOR
 
 #include "Async/AsyncWork.h"
+#include "DerivedDataBuildInputs.h"
+#include "Engine/Texture2D.h"
+#include "IImageWrapperModule.h"
 #include "ImageCore.h"
 #include "TextureCompressorModule.h"
-#include "IImageWrapperModule.h"
-#include "Engine/Texture2D.h"
 #include "TextureDerivedDataBuildExporter.h"
 
 #endif // WITH_EDITOR
@@ -27,6 +28,8 @@ enum
 
 #if WITH_EDITOR
 
+namespace UE::DerivedData { class FBuildOutput; }
+
 void GetTextureDerivedDataKeySuffix(const UTexture& Texture, const FTextureBuildSettings* BuildSettingsPerLayer, FString& OutKeySuffix);
 uint32 PutDerivedDataInCache(FTexturePlatformData* DerivedData, const FString& DerivedDataKeySuffix, const FStringView& TextureName, bool bForceAllMipsToBeInlined, bool bReplaceExistingDDC);
 
@@ -106,6 +109,13 @@ struct FTextureSourceData
  */
 class FTextureCacheDerivedDataWorker : public FNonAbandonableTask
 {
+	struct FBuildInputRecord
+	{
+		FString Id;
+		FIoHash RawHash;
+		uint64 RawSize;
+	};
+
 	/** Texture compressor module, must be loaded in the game thread. see FModuleManager::WarnIfItWasntSafeToLoadHere() */
 	ITextureCompressorModule* Compressor;
 	/** Image wrapper module, must be loaded in the game thread. see FModuleManager::WarnIfItWasntSafeToLoadHere() */
@@ -122,6 +132,12 @@ class FTextureCacheDerivedDataWorker : public FNonAbandonableTask
 	FTextureSourceData TextureData;
 	/** Source mip images of the composite texture (e.g. normal map for compute roughness). Not necessarily in RGBA32F, usually only top mip as other mips need to be generated first */
 	FTextureSourceData CompositeTextureData;
+	/** DDC2 build function name to use to build this texture (if DDC2 is enabled and the target texture type has a DDC2 build function, empty otherwise) */
+	FString BuildFunctionName;
+	/** DDC2 build input recrods for putting on the action (if DDC2 is enabled and the target texture type has a DDC2 build function, empty otherwise) */
+	TArray<FBuildInputRecord> BuildInputRecords;
+	/** DDC2 build inputs builder that accumulates input blobs and their associated key (if DDC2 is enabled and the target texture type has a DDC2 build function, empty otherwise) */
+	UE::DerivedData::FBuildInputsBuilder InputsBuilder;
 	/** Exporter that can optionally write out build actions and reference outputs for the purpose of testing remote execution of builds */
 	FTextureDerivedDataBuildExporter BuildExporter;
 	/** Texture cache flags. */
@@ -139,6 +155,23 @@ class FTextureCacheDerivedDataWorker : public FNonAbandonableTask
 	/** Build the texture. This function is safe to call from any thread. */
 	void BuildTexture(bool bReplaceExistingDDC = false);
 
+	void ConditionalAddBuildFunctionInput(FTextureSource& Source)
+	{
+		if (!BuildFunctionName.IsEmpty())
+		{
+			Source.OperateOnLoadedBulkData([this, &Source] (const FSharedBuffer& BulkDataBuffer) 
+			{
+				FCompressedBuffer InputBuffer = FCompressedBuffer::Compress(NAME_Default, BulkDataBuffer);
+				FBuildInputRecord& NewRef = BuildInputRecords.AddDefaulted_GetRef();
+				NewRef.Id = Source.GetIdString();
+				NewRef.RawHash = InputBuffer.GetRawHash();
+				NewRef.RawSize = InputBuffer.GetRawSize();
+				InputsBuilder.AddInput(Source.GetIdString(), InputBuffer);
+			});
+		}
+	}
+
+	void ConsumeBuildFunctionOutput(const UE::DerivedData::FBuildOutput& BuildOutput, const FString& TexturePath, bool bReplaceExistingDDC);
 public:
 
 	/** Initialization constructor. */
