Fixed: vulkan crash on skeletal mesh streaming

diff --git a/Engine/Source/Runtime/Engine/Private/Rendering/SkinWeightVertexBuffer.cpp b/Engine/Source/Runtime/Engine/Private/Rendering/SkinWeightVertexBuffer.cpp
--- a/Engine/Source/Runtime/Engine/Private/Rendering/SkinWeightVertexBuffer.cpp
+++ b/Engine/Source/Runtime/Engine/Private/Rendering/SkinWeightVertexBuffer.cpp
@@ -139,7 +139,7 @@ void FSkinWeightLookupVertexBuffer::InitRHI()
 
 		if (bSRV)
 		{
-			SRVValue = RHICreateShaderResourceView(FShaderResourceViewInitializer(LookupData ? VertexBufferRHI : nullptr, PixelFormat));
+			SRVValue = RHICreateShaderResourceView(FShaderResourceViewInitializer(VertexBufferRHI, PixelFormat));
 		}
 	}
 }
@@ -432,16 +432,18 @@ void FSkinWeightDataVertexBuffer::InitRHI()
 
 	// BUF_ShaderResource is needed for support of the SkinCache (we could make is dependent on GEnableGPUSkinCacheShaders or are there other users?)
 	VertexBufferRHI = CreateRHIBuffer_RenderThread();
-
-	bool bSRV = VertexBufferRHI && GSupportsResourceView && GPixelFormats[GetPixelFormat()].Supported;
-	// When bAllowCPUAccess is true, the meshes is likely going to be used for Niagara to spawn particles on mesh surface.
-	// And it can be the case for CPU *and* GPU access: no differenciation today. That is why we create a SRV in this case.
-	// This also avoid setting lots of states on all the members of all the different buffers used by meshes. Follow up: https://jira.it.epicgames.net/browse/UE-69376.
-	bSRV |= GetNeedsCPUAccess();
-
-	if (bSRV)
+	if (VertexBufferRHI)
 	{
-		SRVValue = RHICreateShaderResourceView(FShaderResourceViewInitializer(WeightData ? VertexBufferRHI : nullptr, GetPixelFormat()));
+		bool bSRV = GSupportsResourceView && GPixelFormats[GetPixelFormat()].Supported;
+		// When bAllowCPUAccess is true, the meshes is likely going to be used for Niagara to spawn particles on mesh surface.
+		// And it can be the case for CPU *and* GPU access: no differenciation today. That is why we create a SRV in this case.
+		// This also avoid setting lots of states on all the members of all the different buffers used by meshes. Follow up: https://jira.it.epicgames.net/browse/UE-69376.
+		bSRV |= GetNeedsCPUAccess();
+
+		if (bSRV)
+		{
+			SRVValue = RHICreateShaderResourceView(FShaderResourceViewInitializer(VertexBufferRHI, GetPixelFormat()));
+		}
 	}
 }
 
diff --git a/Engine/Source/Runtime/VulkanRHI/Private/VulkanBuffer.cpp b/Engine/Source/Runtime/VulkanRHI/Private/VulkanBuffer.cpp
--- a/Engine/Source/Runtime/VulkanRHI/Private/VulkanBuffer.cpp
+++ b/Engine/Source/Runtime/VulkanRHI/Private/VulkanBuffer.cpp
@@ -437,7 +437,14 @@ void FVulkanResourceMultiBuffer::Unlock(bool bFromRenderingThread)
 void FVulkanResourceMultiBuffer::Swap(FVulkanResourceMultiBuffer& Other)
 {
 	FRHIBuffer::Swap(Other);
-	::Swap(*this, Other);
+	
+	::Swap(UEUsage, Other.UEUsage);
+	::Swap(BufferUsageFlags, Other.BufferUsageFlags);
+	::Swap(NumBuffers, Other.NumBuffers);
+	::Swap(DynamicBufferIndex, Other.DynamicBufferIndex);
+	::Swap(Buffers, Other.Buffers);
+	::Swap(Current, Other.Current);
+	::Swap(VolatileLockInfo, Other.VolatileLockInfo);
 }
 
 FBufferRHIRef FVulkanDynamicRHI::RHICreateBuffer(uint32 Size, EBufferUsageFlags Usage, uint32 Stride, ERHIAccess ResourceState, FRHIResourceCreateInfo& CreateInfo)
