Minor refactoring of light weight instance manager classes to give subclasses more flexibility.

diff --git a/Engine/Source/Runtime/Engine/Classes/GameFramework/LightWeightInstanceManager.h b/Engine/Source/Runtime/Engine/Classes/GameFramework/LightWeightInstanceManager.h
--- a/Engine/Source/Runtime/Engine/Classes/GameFramework/LightWeightInstanceManager.h
+++ b/Engine/Source/Runtime/Engine/Classes/GameFramework/LightWeightInstanceManager.h
@@ -128,6 +128,9 @@ public:
 	}
 
 protected:
+	// Creates an actor to replace the instance specified by Handle
+	AActor* ConvertInstanceToActor(const FActorInstanceHandle& Handle);
+
 	// Takes a polymorphic struct to set the initial data for a new instance
 	int32 AddNewInstance(FLWIData* InitData);
 
@@ -147,6 +150,9 @@ protected:
 	// Sets the parameters for actor spawning.
 	virtual void SetSpawnParameters(FActorSpawnParameters& SpawnParams);
 
+	// Returns the class to use when spawning a new actor
+	virtual UClass* GetActorClassToSpawn() const;
+
 	// Called after spawning a new actor from a light weight instance
 	virtual void PostActorSpawn(const FActorInstanceHandle& Handle);
 
diff --git a/Engine/Source/Runtime/Engine/Classes/GameFramework/LightWeightInstanceStaticMeshManager.h b/Engine/Source/Runtime/Engine/Classes/GameFramework/LightWeightInstanceStaticMeshManager.h
--- a/Engine/Source/Runtime/Engine/Classes/GameFramework/LightWeightInstanceStaticMeshManager.h
+++ b/Engine/Source/Runtime/Engine/Classes/GameFramework/LightWeightInstanceStaticMeshManager.h
@@ -61,7 +61,7 @@ protected:
 
 protected:
 
-	UPROPERTY(ReplicatedUsing = OnRep_StaticMesh)
+	UPROPERTY(EditInstanceOnly, BlueprintReadOnly, Category = Debug, ReplicatedUsing = OnRep_StaticMesh)
 	TSoftObjectPtr<UStaticMesh> StaticMesh;
 	UFUNCTION()
 	void OnRep_StaticMesh();
diff --git a/Engine/Source/Runtime/Engine/Private/LightWeightInstanceManager.cpp b/Engine/Source/Runtime/Engine/Private/LightWeightInstanceManager.cpp
--- a/Engine/Source/Runtime/Engine/Private/LightWeightInstanceManager.cpp
+++ b/Engine/Source/Runtime/Engine/Private/LightWeightInstanceManager.cpp
@@ -50,16 +50,7 @@ AActor* ALightWeightInstanceManager::FetchActorFromHandle(const FActorInstanceHa
 		if ((FoundActor = Actors.Find(Handle.GetInstanceIndex())) == nullptr)
 		{
 			// spawn a new actor
-			FActorSpawnParameters SpawnParams;
-			SetSpawnParameters(SpawnParams);
-
-			AActor* Actor = GetLevel()->GetWorld()->SpawnActor<AActor>(RepresentedClass, InstanceTransforms[Handle.GetInstanceIndex()], SpawnParams);
-			check(Actor);
-
-			Handle.Actor = Actor;
-			Actors.Add(Handle.GetInstanceIndex(), Actor);
-
-			PostActorSpawn(Handle);
+			ConvertInstanceToActor(Handle);
 		}
 		else
 		{
@@ -71,6 +62,28 @@ AActor* ALightWeightInstanceManager::FetchActorFromHandle(const FActorInstanceHa
 	return Handle.Actor.Get();
 }
 
+AActor* ALightWeightInstanceManager::ConvertInstanceToActor(const FActorInstanceHandle& Handle)
+{
+	// we shouldn't be calling this on indices that already have an actor representing them
+	if (Actors.Contains(Handle.GetInstanceIndex()) && Actors[Handle.GetInstanceIndex()] != nullptr)
+	{
+		return Actors[Handle.GetInstanceIndex()];
+	}
+
+	FActorSpawnParameters SpawnParams;
+	SetSpawnParameters(SpawnParams);
+
+	AActor* NewActor = GetLevel()->GetWorld()->SpawnActor<AActor>(GetActorClassToSpawn(), InstanceTransforms[Handle.GetInstanceIndex()], SpawnParams);
+	check(NewActor);
+
+	Handle.Actor = NewActor;
+	Actors.Add(Handle.GetInstanceIndex(), NewActor);
+
+	PostActorSpawn(Handle);
+
+	return NewActor;
+}
+
 int32 ALightWeightInstanceManager::FindIndexForActor(const AActor* InActor) const
 {
 	for (const TPair<int32, AActor*>& ActorPair : Actors)
@@ -163,6 +176,11 @@ void ALightWeightInstanceManager::SetSpawnParameters(FActorSpawnParameters& Spaw
 	SpawnParams.ObjectFlags = RF_Transactional;
 }
 
+UClass* ALightWeightInstanceManager::GetActorClassToSpawn() const
+{
+	return RepresentedClass;
+}
+
 void ALightWeightInstanceManager::PostActorSpawn(const FActorInstanceHandle& Handle)
 {
 	// do nothing
