Release position/tangent buffer handles when invalidating passthrough skinned mesh VF. This fixes the issue when skin cache releases an entry, the corresponding passthrough VF hangs onto the buffer handles, causing the reference count doesn't reach 0 to release the buffers.

diff --git a/Engine/Source/Runtime/Engine/Private/GPUSkinCache.cpp b/Engine/Source/Runtime/Engine/Private/GPUSkinCache.cpp
--- a/Engine/Source/Runtime/Engine/Private/GPUSkinCache.cpp
+++ b/Engine/Source/Runtime/Engine/Private/GPUSkinCache.cpp
@@ -2174,6 +2174,12 @@ void FGPUSkinCache::ReleaseSkinCacheEntry(FGPUSkinCacheEntry* SkinCacheEntry)
 #if RHI_RAYTRACING
 	SkinCache->RemoveRayTracingGeometryUpdate(&SkinCacheEntry->GPUSkin->RayTracingGeometry);
 #endif // RHI_RAYTRACING
+
+	for (FGPUSkinCacheEntry::FSectionDispatchData& SectionData : SkinCacheEntry->GetDispatchData())
+	{
+		SectionData.TargetVertexFactory->InvalidateStreams();
+	}
+
 	FRWBuffersAllocation* PositionAllocation = SkinCacheEntry->PositionAllocation;
 	if (PositionAllocation)
 	{
diff --git a/Engine/Source/Runtime/Engine/Public/GPUSkinVertexFactory.h b/Engine/Source/Runtime/Engine/Public/GPUSkinVertexFactory.h
--- a/Engine/Source/Runtime/Engine/Public/GPUSkinVertexFactory.h
+++ b/Engine/Source/Runtime/Engine/Public/GPUSkinVertexFactory.h
@@ -521,6 +521,8 @@ public:
 	{
 		PositionStreamIndex = -1;
 		TangentStreamIndex = -1;
+		PositionVBAlias.ReleaseRHI();
+		TangentVBAlias.ReleaseRHI();
 	}
 
 	virtual void ReleaseRHI() override
