Fix for blueprint errors. Adding back the AtmosphericFog component blueprint funtcions (doing nothing)

diff --git a/Engine/Source/Runtime/Engine/Classes/Atmosphere/AtmosphericFogComponent.h b/Engine/Source/Runtime/Engine/Classes/Atmosphere/AtmosphericFogComponent.h
--- a/Engine/Source/Runtime/Engine/Classes/Atmosphere/AtmosphericFogComponent.h
+++ b/Engine/Source/Runtime/Engine/Classes/Atmosphere/AtmosphericFogComponent.h
@@ -20,6 +20,58 @@ class UE_DEPRECATED(4.26, "Please use the SkyAtmosphere component instead.") UAt
 
 	~UAtmosphericFogComponent();
 
+	/** Deprecated */
+	UFUNCTION(BlueprintCallable, Category = "Rendering|Components|AtmosphericFog")
+	ENGINE_API void SetDefaultBrightness(float NewBrightness) {}
+
+	/** Deprecated */
+	UFUNCTION(BlueprintCallable, Category = "Rendering|Components|AtmosphericFog")
+	ENGINE_API void SetDefaultLightColor(FLinearColor NewLightColor) {}
+
+	/** Deprecated */
+	UFUNCTION(BlueprintCallable, Category = "Rendering|Components|AtmosphericFog")
+	ENGINE_API void SetSunMultiplier(float NewSunMultiplier) {}
+
+	/** Deprecated */
+	UFUNCTION(BlueprintCallable, Category = "Rendering|Components|AtmosphericFog")
+	ENGINE_API void SetFogMultiplier(float NewFogMultiplier) {}
+
+	/** Deprecated */
+	UFUNCTION(BlueprintCallable, Category = "Rendering|Components|AtmosphericFog")
+	ENGINE_API void SetDensityMultiplier(float NewDensityMultiplier) {}
+
+	/** Deprecated */
+	UFUNCTION(BlueprintCallable, Category = "Rendering|Components|AtmosphericFog")
+	ENGINE_API void SetDensityOffset(float NewDensityOffset) {}
+
+	/** Deprecated */
+	UFUNCTION(BlueprintCallable, Category = "Rendering|Components|AtmosphericFog")
+	ENGINE_API void SetDistanceScale(float NewDistanceScale) {}
+
+	/** Deprecated */
+	UFUNCTION(BlueprintCallable, Category = "Rendering|Components|AtmosphericFog")
+	ENGINE_API void SetAltitudeScale(float NewAltitudeScale) {}
+
+	/** Deprecated */
+	UFUNCTION(BlueprintCallable, Category = "Rendering|Components|AtmosphericFog")
+	ENGINE_API void SetStartDistance(float NewStartDistance) {}
+
+	/** Deprecated */
+	UFUNCTION(BlueprintCallable, Category = "Rendering|Components|AtmosphericFog")
+	ENGINE_API void SetDistanceOffset(float NewDistanceOffset) {}
+
+	/** Deprecated */
+	UFUNCTION(BlueprintCallable, Category = "Rendering|Components|AtmosphericFog")
+	ENGINE_API void DisableSunDisk(bool NewSunDisk) {}
+
+	/** Deprecated */
+	UFUNCTION(BlueprintCallable, Category = "Rendering|Components|AtmosphericFog")
+	ENGINE_API void DisableGroundScattering(bool NewGroundScattering) {}
+
+	/** Deprecated */
+	UFUNCTION(BlueprintCallable, Category = "Rendering|Components|AtmosphericFog")
+	ENGINE_API void SetPrecomputeParams(float DensityHeight, int32 MaxScatteringOrder, int32 InscatterAltitudeSampleNum) {}
+
 protected:
 
 public:
