Fixed crash when toggleing metalness workflow on slab node with empty pins.

diff --git a/Engine/Source/Runtime/Engine/Private/Materials/MaterialExpressions.cpp b/Engine/Source/Runtime/Engine/Private/Materials/MaterialExpressions.cpp
--- a/Engine/Source/Runtime/Engine/Private/Materials/MaterialExpressions.cpp
+++ b/Engine/Source/Runtime/Engine/Private/Materials/MaterialExpressions.cpp
@@ -20078,12 +20078,24 @@ void UMaterialExpressionStrataSlabBSDF::PostEditChangeProperty(FPropertyChangedE
 {
 	if (PropertyChangedEvent.Property && (PropertyChangedEvent.Property->GetName() == TEXT("bUseMetalness")))
 	{
+		auto SafeConnect = [](FExpressionInput& Input, int32 OutputIndex, class UMaterialExpression* Expression)
+		{
+			if (Expression)
+			{
+				Input.Connect(OutputIndex, Expression);
+			}
+			else
+			{
+				Input.Expression = nullptr;
+			}
+		};
+
 		// Relink to make the node looks somewhat ok when toggling Metalness
 		if (bUseMetalness)
 		{
-			BaseColor.Connect(DiffuseAlbedo.OutputIndex, DiffuseAlbedo.Expression);
-			Specular.Connect(F0.OutputIndex, F0.Expression);
-			EdgeColor.Connect(F90.OutputIndex, F90.Expression);
+			SafeConnect(BaseColor, DiffuseAlbedo.OutputIndex, DiffuseAlbedo.Expression);
+			SafeConnect(Specular, F0.OutputIndex, F0.Expression);
+			SafeConnect(EdgeColor, F90.OutputIndex, F90.Expression);
 
 			DiffuseAlbedo.Expression = nullptr;
 			F0.Expression = nullptr;
@@ -20091,9 +20103,9 @@ void UMaterialExpressionStrataSlabBSDF::PostEditChangeProperty(FPropertyChangedE
 		}
 		else
 		{
-			DiffuseAlbedo.Connect(BaseColor.OutputIndex, BaseColor.Expression);
-			F0.Connect(Specular.OutputIndex, Specular.Expression);
-			F90.Connect(EdgeColor.OutputIndex, EdgeColor.Expression);
+			SafeConnect(DiffuseAlbedo, BaseColor.OutputIndex, BaseColor.Expression);
+			SafeConnect(F0, Specular.OutputIndex, Specular.Expression);
+			SafeConnect(F90, EdgeColor.OutputIndex, EdgeColor.Expression);
 
 			BaseColor.Expression = nullptr;
 			EdgeColor.Expression = nullptr;
