Don't filter memory stat groups in a detailed memreport. Just log everything.

diff --git a/Engine/Source/Runtime/Engine/Private/UnrealEngine.cpp b/Engine/Source/Runtime/Engine/Private/UnrealEngine.cpp
--- a/Engine/Source/Runtime/Engine/Private/UnrealEngine.cpp
+++ b/Engine/Source/Runtime/Engine/Private/UnrealEngine.cpp
@@ -6951,8 +6951,7 @@ bool UEngine::HandleMemCommand( const TCHAR* Cmd, FOutputDevice& Ar )
 		for (int32 Index = 0; Index < Stats.Num(); Index++)
 		{
 			FStatMessage const& Meta = Stats[Index];
-			FName LastGroup = Meta.NameAndInfo.GetGroupName();
-			if ((LastGroup == NAME_STATGROUP_SceneMemory || LastGroup == NAME_STATGROUP_Memory || LastGroup == NAME_STATGROUP_TextureGroup || LastGroup == NAME_STATGROUP_RHI)  && Meta.NameAndInfo.GetFlag(EStatMetaFlags::IsMemory))
+			if (Meta.NameAndInfo.GetFlag(EStatMetaFlags::IsMemory))
 			{
 				Ar.CategorizedLogf( CategoryName, ELogVerbosity::Log, TEXT("%s"), *FStatsUtils::DebugPrint(Meta));
 			}
