Round 2: Fix for Modulation & Submix content intermittently failing to load.

- Refactor logic querying/loading the default SoundClass.  Prior logic loaded them extremely early in the engine boot cycle, which indirectly could initiate load attempts on assets that were of plugin-supplied types prior to plugins supplying their given type being loaded.

diff --git a/Engine/Source/Runtime/AudioMixer/Private/Components/SynthComponent.cpp b/Engine/Source/Runtime/AudioMixer/Private/Components/SynthComponent.cpp
--- a/Engine/Source/Runtime/AudioMixer/Private/Components/SynthComponent.cpp
+++ b/Engine/Source/Runtime/AudioMixer/Private/Components/SynthComponent.cpp
@@ -1,8 +1,11 @@
 // Copyright Epic Games, Inc. All Rights Reserved.
 
 #include "Components/SynthComponent.h"
+
 #include "AudioDevice.h"
 #include "AudioMixerLog.h"
+#include "Sound/AudioSettings.h"
+
 
 USynthSound::USynthSound(const FObjectInitializer& ObjectInitializer)
 	: Super(ObjectInitializer)
@@ -141,10 +144,6 @@ USynthComponent::USynthComponent(const FObjectInitializer& ObjectInitializer)
 	bAlwaysPlay = true;
 	Synth = nullptr;
 
-	// Set the default sound class
-	SoundClass = USoundBase::DefaultSoundClassObject;
-	Synth = nullptr;
-
 	PreferredBufferLength = DEFAULT_PROCEDURAL_SOUNDWAVE_BUFFER_SIZE;
 
 #if WITH_EDITORONLY_DATA
@@ -350,6 +349,21 @@ void USynthComponent::OnUnregister()
 	}
 }
 
+USoundClass* USynthComponent::GetSoundClass()
+{
+	if (SoundClass)
+	{
+		return SoundClass;
+	}
+
+	if (const UAudioSettings* AudioSettings = GetDefault<UAudioSettings>())
+	{
+		return AudioSettings->GetDefaultSoundClass();
+	}
+
+	return nullptr;
+}
+
 bool USynthComponent::IsReadyForOwnerToAutoDestroy() const
 {
 	const bool bIsAudioComponentReadyForDestroy = !AudioComponent || (AudioComponent && !AudioComponent->IsPlaying());
@@ -402,6 +416,19 @@ void USynthComponent::Serialize(FArchive& Ar)
 		}
 	}
 #endif // WITH_EDITORONLY_DATA
+
+	if (Ar.IsCooking())
+	{
+		if (!SoundClass)
+		{
+			SoundClass = GetDefault<UAudioSettings>()->GetDefaultSoundClass();
+		}
+
+		if (ConcurrencySet.IsEmpty())
+		{
+			ConcurrencySet.Add(GetDefault<UAudioSettings>()->GetDefaultSoundConcurrency());
+		}
+	}
 }
 
 void USynthComponent::PumpPendingMessages()
diff --git a/Engine/Source/Runtime/AudioMixer/Public/Components/SynthComponent.h b/Engine/Source/Runtime/AudioMixer/Public/Components/SynthComponent.h
--- a/Engine/Source/Runtime/AudioMixer/Public/Components/SynthComponent.h
+++ b/Engine/Source/Runtime/AudioMixer/Public/Components/SynthComponent.h
@@ -253,6 +253,7 @@ public:
 	void AddAudioBufferListener(IAudioBufferListener* InAudioBufferListener);
 	void RemoveAudioBufferListener(IAudioBufferListener* InAudioBufferListener);
 
+	virtual USoundClass* GetSoundClass();
 
 	virtual void BeginDestroy() override;
 
diff --git a/Engine/Source/Runtime/Engine/Classes/Sound/AudioSettings.h b/Engine/Source/Runtime/Engine/Classes/Sound/AudioSettings.h
--- a/Engine/Source/Runtime/Engine/Classes/Sound/AudioSettings.h
+++ b/Engine/Source/Runtime/Engine/Classes/Sound/AudioSettings.h
@@ -11,6 +11,9 @@
 
 struct FPropertyChangedChainEvent;
 
+class USoundClass;
+class USoundConcurrency;
+
 // Enumeration for what our options are for sample rates used for VOIP.
 UENUM()
 enum class EVoiceSampleRate : int32
@@ -220,7 +223,23 @@ class ENGINE_API UAudioSettings : public UDeveloperSettings
 	FAudioSettingsChanged AudioSettingsChanged;
 #endif // WITH_EDITOR
 
+private:
+	UPROPERTY(Transient)
+	USoundClass* DefaultSoundClass;
+
+	UPROPERTY(Transient)
+	USoundClass* DefaultMediaSoundClass;
+
+	UPROPERTY(Transient)
+	USoundConcurrency* DefaultSoundConcurrency;
+
 public:
+	void LoadDefaultObjects();
+
+	USoundClass* GetDefaultSoundClass() const;
+	USoundClass* GetDefaultMediaSoundClass() const;
+	USoundConcurrency* GetDefaultSoundConcurrency() const;
+
 	// Get the quality level settings at the provided level index
 	const FAudioQualitySettings& GetQualityLevelSettings(int32 QualityLevel) const;
 	
@@ -249,6 +268,7 @@ private:
 	TArray<FAudioQualitySettings> CachedQualityLevels;
 	FSoftObjectPath CachedAmbisonicSubmix;
 	FSoftObjectPath CachedMasterSubmix;
+	FSoftObjectPath CachedSoundClass;
 #endif // WITH_EDITOR
 
 	void AddDefaultSettings();
diff --git a/Engine/Source/Runtime/Engine/Classes/Sound/SoundBase.h b/Engine/Source/Runtime/Engine/Classes/Sound/SoundBase.h
--- a/Engine/Source/Runtime/Engine/Classes/Sound/SoundBase.h
+++ b/Engine/Source/Runtime/Engine/Classes/Sound/SoundBase.h
@@ -8,6 +8,8 @@
 #include "Audio.h"
 #include "CoreMinimal.h"
 #include "IAudioExtensionPlugin.h"
+#include "Sound/AudioSettings.h"
+#include "Sound/SoundClass.h"
 #include "SoundConcurrency.h"
 #include "SoundModulationDestination.h"
 #include "SoundSourceBusSend.h"
@@ -58,9 +60,6 @@ class ENGINE_API USoundBase : public UObject, public IInterface_AssetUserData
 	GENERATED_UCLASS_BODY()
 
 public:
-	static USoundClass* DefaultSoundClassObject;
-	static USoundConcurrency* DefaultSoundConcurrencyObject;
-
 	/** Sound class this sound belongs to */
 	UPROPERTY(EditAnywhere, Category = Sound, meta = (DisplayName = "Class"), AssetRegistrySearchable)
 	TObjectPtr<USoundClass> SoundClassObject;
@@ -189,9 +188,7 @@ public:
 	TArray<TObjectPtr<UAssetUserData>> AssetUserData;
 
 public:
-
 	//~ Begin UObject Interface.
-	virtual void PostInitProperties() override;
 #if WITH_EDITORONLY_DATA
 	virtual void PostLoad() override;
 #endif
diff --git a/Engine/Source/Runtime/Engine/Private/ActiveSound.cpp b/Engine/Source/Runtime/Engine/Private/ActiveSound.cpp
--- a/Engine/Source/Runtime/Engine/Private/ActiveSound.cpp
+++ b/Engine/Source/Runtime/Engine/Private/ActiveSound.cpp
@@ -479,20 +479,27 @@ void FActiveSound::GetConcurrencyHandles(TArray<FConcurrencyHandle>& OutConcurre
 {
 	OutConcurrencyHandles.Reset();
 
-	if (!ConcurrencySet.Num() && Sound)
-	{
-		Sound->GetConcurrencyHandles(OutConcurrencyHandles);
-	}
-	else
+	const UAudioSettings* AudioSettings = GetDefault<UAudioSettings>();
+	check(AudioSettings);
+
+	if (ConcurrencySet.Num() > 0)
 	{
 		for (const USoundConcurrency* Concurrency : ConcurrencySet)
 		{
 			if (Concurrency)
 			{
-				OutConcurrencyHandles.Emplace(FConcurrencyHandle(*Concurrency));
+				OutConcurrencyHandles.Emplace(*Concurrency);
 			}
 		}
 	}
+	else if (Sound)
+	{
+		Sound->GetConcurrencyHandles(OutConcurrencyHandles);
+	}
+	else if (const USoundConcurrency* DefaultConcurrency = AudioSettings->GetDefaultSoundConcurrency())
+	{
+		OutConcurrencyHandles.Emplace(*DefaultConcurrency);
+	}
 }
 
 bool FActiveSound::GetConcurrencyFadeDuration(float& OutFadeDuration) const
diff --git a/Engine/Source/Runtime/Engine/Private/AudioDeviceManager.cpp b/Engine/Source/Runtime/Engine/Private/AudioDeviceManager.cpp
--- a/Engine/Source/Runtime/Engine/Private/AudioDeviceManager.cpp
+++ b/Engine/Source/Runtime/Engine/Private/AudioDeviceManager.cpp
@@ -538,8 +538,13 @@ bool FAudioDeviceManager::InitializeManager()
 	{
 		check(AudioDeviceModule);
 
+		UAudioSettings* AudioSettings = GetMutableDefault<UAudioSettings>();
+		check(AudioSettings);
+
+		AudioSettings->LoadDefaultObjects();
+
 		const bool bIsAudioMixerEnabled = AudioDeviceModule->IsAudioMixerModule();
-		GetMutableDefault<UAudioSettings>()->SetAudioMixerEnabled(bIsAudioMixerEnabled);
+		AudioSettings->SetAudioMixerEnabled(bIsAudioMixerEnabled);
 
 #if WITH_EDITOR
 		if (bIsAudioMixerEnabled)
@@ -610,7 +615,7 @@ bool FAudioDeviceManager::LoadDefaultAudioDeviceModule()
 		bUsingAudioMixer = false;
 	}
 	else if(bForceAudioMixer)
-	{
+	{ 
 		GConfig->GetString(TEXT("Audio"), TEXT("AudioMixerModuleName"), AudioMixerModuleName, GEngineIni);
 	}
 
diff --git a/Engine/Source/Runtime/Engine/Private/AudioSettings.cpp b/Engine/Source/Runtime/Engine/Private/AudioSettings.cpp
--- a/Engine/Source/Runtime/Engine/Private/AudioSettings.cpp
+++ b/Engine/Source/Runtime/Engine/Private/AudioSettings.cpp
@@ -7,6 +7,8 @@
 #include "Misc/ConfigCacheIni.h"
 #include "Misc/Paths.h"
 #include "Sound/SoundBase.h"
+#include "Sound/SoundClass.h"
+#include "Sound/SoundConcurrency.h"
 #include "Sound/SoundNodeQualityLevel.h"
 #include "Sound/SoundSubmix.h"
 #include "UObject/UObjectHash.h"
@@ -46,6 +48,8 @@ void UAudioSettings::AddDefaultSettings()
 #if WITH_EDITOR
 void UAudioSettings::PreEditChange(FProperty* PropertyAboutToChange)
 {
+	CachedSoundClass = DefaultSoundClassName;
+
 	// Cache master submix in case user tries to set to submix that isn't a top-level submix
 	CachedMasterSubmix = MasterSubmix;
 
@@ -60,6 +64,9 @@ void UAudioSettings::PostEditChangeChainProperty(FPropertyChangedChainEvent& Pro
 		bool bReconcileQualityNodes = false;
 		bool bPromptRestartRequired = false;
 		FName PropertyName = PropertyChangedEvent.Property->GetFName();
+
+		LoadDefaultObjects();
+
 		if (PropertyName == GET_MEMBER_NAME_CHECKED(UAudioSettings, MasterSubmix))
 		{
 			if (USoundSubmix* NewSubmix = Cast<USoundSubmix>(MasterSubmix.TryLoad()))
@@ -172,6 +179,80 @@ int32 UAudioSettings::GetQualityLevelSettingsNum() const
 	return QualityLevels.Num();
 }
 
+void UAudioSettings::LoadDefaultObjects()
+{
+	UE_LOG(LogAudio, Display, TEXT("Loading Default Audio Settings Objects..."));
+
+	// TODO: Move all soft object paths to load here (SoundMixes, Submixes, etc.)
+	static const FString EngineSoundsDir = TEXT("/Engine/EngineSounds");
+
+	DefaultSoundClass = nullptr;
+	if (UObject* SoundClassObject = DefaultSoundClassName.TryLoad())
+	{
+		DefaultSoundClass = CastChecked<USoundClass>(SoundClassObject);
+	}
+
+#if WITH_EDITOR
+	if (!DefaultSoundClass)
+	{
+		DefaultSoundClassName = CachedSoundClass;
+		if (UObject* SoundClassObject = DefaultSoundClassName.TryLoad())
+		{
+			DefaultSoundClass = CastChecked<USoundClass>(SoundClassObject);
+		}
+	}
+#endif // WITH_EDITOR
+
+	if (!DefaultSoundClass)
+	{
+		UE_LOG(LogAudio, Warning, TEXT("Failed to load Default SoundClassObject from path '%s'.  Attempting to fall back to engine default."), *DefaultSoundClassName.GetAssetPathString());
+		DefaultSoundClassName.SetPath(EngineSoundsDir / TEXT("Master"));
+		if (UObject* SoundClassObject = DefaultSoundClassName.TryLoad())
+		{
+			DefaultSoundClass = CastChecked<USoundClass>(SoundClassObject);
+		}
+	}
+
+	if (!DefaultSoundClass)
+	{
+		UE_LOG(LogAudio, Error, TEXT("Failed to load Default SoundClassObject from path '%s'."), *DefaultSoundClassName.GetAssetPathString());
+	}
+
+	DefaultMediaSoundClass = nullptr;
+	if (UObject* MediaSoundClassObject = DefaultMediaSoundClassName.TryLoad())
+	{
+		DefaultMediaSoundClass = CastChecked<USoundClass>(MediaSoundClassObject);
+	}
+	else
+	{
+		UE_LOG(LogAudio, Display, TEXT("No default MediaSoundClassObject specified (or failed to load)."));
+	}
+
+	DefaultSoundConcurrency = nullptr;
+	if (UObject* SoundConcurrencyObject = DefaultSoundConcurrencyName.TryLoad())
+	{
+		DefaultSoundConcurrency = CastChecked<USoundConcurrency>(SoundConcurrencyObject);
+	}
+	else
+	{
+		UE_LOG(LogAudio, Display, TEXT("No default SoundConcurrencyObject specified (or failed to load)."));
+	}
+}
+
+USoundClass* UAudioSettings::GetDefaultMediaSoundClass() const
+{
+	return DefaultMediaSoundClass;
+}
+
+USoundClass* UAudioSettings::GetDefaultSoundClass() const
+{
+	return DefaultSoundClass;
+}
+
+USoundConcurrency* UAudioSettings::GetDefaultSoundConcurrency() const
+{
+	return DefaultSoundConcurrency;
+}
 
 void UAudioSettings::SetAudioMixerEnabled(const bool bInAudioMixerEnabled)
 {
diff --git a/Engine/Source/Runtime/Engine/Private/SoundBase.cpp b/Engine/Source/Runtime/Engine/Private/SoundBase.cpp
--- a/Engine/Source/Runtime/Engine/Private/SoundBase.cpp
+++ b/Engine/Source/Runtime/Engine/Private/SoundBase.cpp
@@ -10,9 +10,6 @@
 #include "Engine/AssetUserData.h"
 
 
-USoundClass* USoundBase::DefaultSoundClassObject = nullptr;
-USoundConcurrency* USoundBase::DefaultSoundConcurrencyObject = nullptr;
-
 USoundBase::USoundBase(const FObjectInitializer& ObjectInitializer)
 	: Super(ObjectInitializer)
 	, VirtualizationMode(EVirtualizationMode::Restart)
@@ -29,37 +26,6 @@ USoundBase::USoundBase(const FObjectInitializer& ObjectInitializer)
 	bEnableSubmixSends = true;
 }
 
-void USoundBase::PostInitProperties()
-{
-	Super::PostInitProperties();
-
-	if (USoundBase::DefaultSoundClassObject == nullptr)
-	{
-		const FSoftObjectPath DefaultSoundClassName = GetDefault<UAudioSettings>()->DefaultSoundClassName;
-		if (DefaultSoundClassName.IsValid())
-		{
-			SCOPED_BOOT_TIMING("USoundBase::LoadSoundClass");
-			USoundBase::DefaultSoundClassObject = LoadObject<USoundClass>(nullptr, *DefaultSoundClassName.ToString());
-		}
-	}
-	SoundClassObject = USoundBase::DefaultSoundClassObject;
-
-	if (USoundBase::DefaultSoundConcurrencyObject == nullptr)
-	{
-		const FSoftObjectPath DefaultSoundConcurrencyName = GetDefault<UAudioSettings>()->DefaultSoundConcurrencyName;
-		if (DefaultSoundConcurrencyName.IsValid())
-		{
-			SCOPED_BOOT_TIMING("USoundBase::LoadSoundConcurrency");
-			USoundBase::DefaultSoundConcurrencyObject = LoadObject<USoundConcurrency>(nullptr, *DefaultSoundConcurrencyName.ToString());
-		}
-	}
-
-	if (USoundBase::DefaultSoundConcurrencyObject != nullptr)
-	{
-		ConcurrencySet.Add(USoundBase::DefaultSoundConcurrencyObject);
-	}
-}
-
 bool USoundBase::IsPlayable() const
 {
 	return false;
@@ -135,12 +101,26 @@ bool USoundBase::IsLooping()
 
 bool USoundBase::ShouldApplyInteriorVolumes()
 {
-	return (SoundClassObject && SoundClassObject->Properties.bApplyAmbientVolumes);
+	USoundClass* SoundClass = GetSoundClass();
+	return SoundClass && SoundClass->Properties.bApplyAmbientVolumes;
 }
 
 USoundClass* USoundBase::GetSoundClass() const
 {
-	return SoundClassObject;
+	if (SoundClassObject)
+	{
+		return SoundClassObject;
+	}
+
+	if (const UAudioSettings* AudioSettings = GetDefault<UAudioSettings>())
+	{
+		if (USoundClass* DefaultSoundClass = AudioSettings->GetDefaultSoundClass())
+		{
+			return DefaultSoundClass;
+		}
+	}
+
+	return nullptr;
 }
 
 USoundSubmixBase* USoundBase::GetSoundSubmix() const
@@ -167,6 +147,9 @@ void USoundBase::GetSoundSourceBusSends(EBusSendType BusSendType, TArray<FSoundS
 
 void USoundBase::GetConcurrencyHandles(TArray<FConcurrencyHandle>& OutConcurrencyHandles) const
 {
+	const UAudioSettings* AudioSettings = GetDefault<UAudioSettings>();
+	check(AudioSettings);
+
 	OutConcurrencyHandles.Reset();
 	if (bOverrideConcurrency)
 	{
@@ -174,13 +157,20 @@ void USoundBase::GetConcurrencyHandles(TArray<FConcurrencyHandle>& OutConcurrenc
 	}
 	else
 	{
-		for (const USoundConcurrency* Concurrency : ConcurrencySet)
+		if (ConcurrencySet.Num() > 0)
 		{
-			if (Concurrency)
+			for (const USoundConcurrency* Concurrency : ConcurrencySet)
 			{
-				OutConcurrencyHandles.Emplace(*Concurrency);
+				if (Concurrency)
+				{
+					OutConcurrencyHandles.Emplace(*Concurrency);
+				}
 			}
 		}
+		else if (const USoundConcurrency* DefaultConcurrency = AudioSettings->GetDefaultSoundConcurrency())
+		{
+			OutConcurrencyHandles.Emplace(*DefaultConcurrency);
+		}
 	}
 }
 
@@ -243,6 +233,19 @@ void USoundBase::Serialize(FArchive& Ar)
 		}
 	}
 #endif // WITH_EDITORONLY_DATA
+
+	if (Ar.IsCooking())
+	{
+		if (!SoundClassObject)
+		{
+			SoundClassObject = GetDefault<UAudioSettings>()->GetDefaultSoundClass();
+		}
+
+		if (ConcurrencySet.IsEmpty())
+		{
+			ConcurrencySet.Add(GetDefault<UAudioSettings>()->GetDefaultSoundConcurrency());
+		}
+	}
 }
 
 void USoundBase::AddAssetUserData(UAssetUserData* InUserData)
diff --git a/Engine/Source/Runtime/Engine/Private/SoundCue.cpp b/Engine/Source/Runtime/Engine/Private/SoundCue.cpp
--- a/Engine/Source/Runtime/Engine/Private/SoundCue.cpp
+++ b/Engine/Source/Runtime/Engine/Private/SoundCue.cpp
@@ -234,7 +234,7 @@ void USoundCue::PostLoad()
 	
 	ESoundWaveLoadingBehavior SoundClassLoadingBehavior = ESoundWaveLoadingBehavior::Inherited;
 
-	USoundClass* CurrentSoundClass = SoundClassObject;
+	USoundClass* CurrentSoundClass = GetSoundClass();
 
 	// Recurse through this sound class's parents until we find an override.
 	while (SoundClassLoadingBehavior == ESoundWaveLoadingBehavior::Inherited && CurrentSoundClass != nullptr)
diff --git a/Engine/Source/Runtime/Engine/Private/SoundWave.cpp b/Engine/Source/Runtime/Engine/Private/SoundWave.cpp
--- a/Engine/Source/Runtime/Engine/Private/SoundWave.cpp
+++ b/Engine/Source/Runtime/Engine/Private/SoundWave.cpp
@@ -3051,7 +3051,7 @@ void USoundWave::CacheInheritedLoadingBehavior() const
 		// if this is true then the behavior should not be Inherited here
 		check(!*bLoadingBehaviorOverriddenPtr);
 
-		USoundClass* CurrentSoundClass = SoundClassObject;
+		USoundClass* CurrentSoundClass = GetSoundClass();
 		ESoundWaveLoadingBehavior SoundClassLoadingBehavior = ESoundWaveLoadingBehavior::Inherited;
 
 		// Recurse through this sound class's parents until we find an override.
diff --git a/Engine/Source/Runtime/MediaAssets/Private/Assets/MediaSoundComponent.cpp b/Engine/Source/Runtime/MediaAssets/Private/Assets/MediaSoundComponent.cpp
--- a/Engine/Source/Runtime/MediaAssets/Private/Assets/MediaSoundComponent.cpp
+++ b/Engine/Source/Runtime/MediaAssets/Private/Assets/MediaSoundComponent.cpp
@@ -26,11 +26,6 @@ CSV_DECLARE_CATEGORY_MODULE_EXTERN(MEDIA_API, MediaStreaming);
 
 static const int32 MaxAudioInputSamples = 8;	// accept at most these many samples into our input queue
 
-/* Static initialization
- *****************************************************************************/
-
-USoundClass* UMediaSoundComponent::DefaultMediaSoundClassObject = nullptr;
-
 
 /* UMediaSoundComponent structors
  *****************************************************************************/
@@ -257,27 +252,6 @@ void UMediaSoundComponent::Deactivate()
 /* UObject interface
  *****************************************************************************/
 
-void UMediaSoundComponent::PostInitProperties()
-{
-	Super::PostInitProperties();
-
-	if (UMediaSoundComponent::DefaultMediaSoundClassObject == nullptr)
-	{
-		const FSoftObjectPath DefaultMediaSoundClassName = GetDefault<UAudioSettings>()->DefaultMediaSoundClassName;
-
-		if (DefaultMediaSoundClassName.IsValid())
-		{
-			UMediaSoundComponent::DefaultMediaSoundClassObject = LoadObject<USoundClass>(nullptr, *DefaultMediaSoundClassName.ToString());
-		}
-	}
-
-	// We have a different default sound class object for media sound components
-	if (SoundClass == USoundBase::DefaultSoundClassObject || SoundClass == nullptr)
-	{
-		SoundClass = UMediaSoundComponent::DefaultMediaSoundClassObject;
-	}
-}
-
 
 void UMediaSoundComponent::PostLoad()
 {
diff --git a/Engine/Source/Runtime/MediaAssets/Public/MediaSoundComponent.h b/Engine/Source/Runtime/MediaAssets/Public/MediaSoundComponent.h
--- a/Engine/Source/Runtime/MediaAssets/Public/MediaSoundComponent.h
+++ b/Engine/Source/Runtime/MediaAssets/Public/MediaSoundComponent.h
@@ -135,6 +135,29 @@ public:
 	UFUNCTION(BlueprintCallable, Category="Media|MediaSoundComponent")
 	UMediaPlayer* GetMediaPlayer() const;
 
+	virtual USoundClass* GetSoundClass() override
+	{
+		if (SoundClass)
+		{
+			return SoundClass;
+		}
+
+		if (const UAudioSettings* AudioSettings = GetDefault<UAudioSettings>())
+		{
+			if (USoundClass* DefaultSoundClass = AudioSettings->GetDefaultMediaSoundClass())
+			{
+				return DefaultSoundClass;
+			}
+
+			if (USoundClass* DefaultSoundClass = AudioSettings->GetDefaultSoundClass())
+			{
+				return DefaultSoundClass;
+			}
+		}
+
+		return nullptr;
+	}
+
 	/**
 	 * Set the media player that provides the audio samples.
 	 *
@@ -210,7 +233,6 @@ public:
 public:
 
 	//~ UObject interface
-	virtual void PostInitProperties() override;
 	virtual void PostLoad() override;
 
 #if WITH_EDITOR
@@ -302,8 +324,4 @@ private:
 
 	/** Whether or not envelope follower settings changed. */
 	bool bEnvelopeFollowerSettingsChanged;
-
-private:
-
-	static USoundClass* DefaultMediaSoundClassObject;
 };
