Added support for extended blend in settings to dynamic montages. New functions:
- UAnimInstance::PlaySlotAnimationAsDynamicMontage_WithBlendInArgs
- UAnimInstance::PlaySlotAnimationAsDynamicMontage_WithBlendInSettings

diff --git a/Engine/Source/Runtime/Engine/Classes/Animation/AnimInstance.h b/Engine/Source/Runtime/Engine/Classes/Animation/AnimInstance.h
--- a/Engine/Source/Runtime/Engine/Classes/Animation/AnimInstance.h
+++ b/Engine/Source/Runtime/Engine/Classes/Animation/AnimInstance.h
@@ -608,6 +608,14 @@ public:
 	UFUNCTION(BlueprintCallable, Category="Animation")
 	UAnimMontage* PlaySlotAnimationAsDynamicMontage(UAnimSequenceBase* Asset, FName SlotNodeName, float BlendInTime = 0.25f, float BlendOutTime = 0.25f, float InPlayRate = 1.f, int32 LoopCount = 1, float BlendOutTriggerTime = -1.f, float InTimeToStartMontageAt = 0.f);
 
+	/** Play normal animation asset on the slot node by creating a dynamic UAnimMontage with blend in arguments. You can only play one asset (whether montage or animsequence) at a time per SlotGroup. */
+	UFUNCTION(BlueprintCallable, Category = "Animation")
+	UAnimMontage* PlaySlotAnimationAsDynamicMontage_WithBlendArgs(UAnimSequenceBase* Asset, FName SlotNodeName, const FAlphaBlendArgs& BlendIn, const FAlphaBlendArgs& BlendOut, float InPlayRate = 1.f, int32 LoopCount = 1, float BlendOutTriggerTime = -1.f, float InTimeToStartMontageAt = 0.f);
+
+	/** Play normal animation asset on the slot node by creating a dynamic UAnimMontage with blend in settings. You can only play one asset (whether montage or animsequence) at a time per SlotGroup. */
+	UFUNCTION(BlueprintCallable, Category = "Animation")
+	UAnimMontage* PlaySlotAnimationAsDynamicMontage_WithBlendSettings(UAnimSequenceBase* Asset, FName SlotNodeName, const FMontageBlendSettings& BlendInSettings, const FMontageBlendSettings& BlendOutSettings, float InPlayRate = 1.f, int32 LoopCount = 1, float BlendOutTriggerTime = -1.f, float InTimeToStartMontageAt = 0.f);
+
 	/** Stops currently playing slot animation slot or all*/
 	UFUNCTION(BlueprintCallable, Category="Animation")
 	void StopSlotAnimation(float InBlendOutTime = 0.25f, FName SlotNodeName = NAME_None);
diff --git a/Engine/Source/Runtime/Engine/Classes/Animation/AnimMontage.h b/Engine/Source/Runtime/Engine/Classes/Animation/AnimMontage.h
--- a/Engine/Source/Runtime/Engine/Classes/Animation/AnimMontage.h
+++ b/Engine/Source/Runtime/Engine/Classes/Animation/AnimMontage.h
@@ -299,6 +299,8 @@ struct ENGINE_API FMontageBlendSettings
 	GENERATED_BODY()
 
 	FMontageBlendSettings();
+	FMontageBlendSettings(float BlendTime);
+	FMontageBlendSettings(const FAlphaBlendArgs& BlendArgs);
 
 	/** Blend Profile to use for this blend */
 	UPROPERTY(EditAnywhere, BlueprintReadWrite, Category = "Blend", meta = (DisplayAfter = "Blend"))
@@ -929,6 +931,10 @@ public:
 	/** Utility function to create dynamic montage from AnimSequence */
 	ENGINE_API static UAnimMontage* CreateSlotAnimationAsDynamicMontage(UAnimSequenceBase* Asset, FName SlotNodeName, float BlendInTime = 0.25f, float BlendOutTime = 0.25f, float InPlayRate = 1.f, int32 LoopCount = 1, float BlendOutTriggerTime = -1.f, float InTimeToStartMontageAt = 0.f);
 
+	/** Utility function to create dynamic montage from AnimSequence with blend in settings */
+	UFUNCTION(BlueprintCallable, Category = "Animation")
+	static ENGINE_API UAnimMontage* CreateSlotAnimationAsDynamicMontage_WithBlendSettings(UAnimSequenceBase* Asset, FName SlotNodeName, const FMontageBlendSettings& BlendInSettings, const FMontageBlendSettings& BlendOutSettings, float InPlayRate = 1.f, int32 LoopCount = 1, float InBlendOutTriggerTime = -1.f);
+
 	//~Begin Time Stretch Curve
 public:
 
diff --git a/Engine/Source/Runtime/Engine/Private/AlphaBlend.cpp b/Engine/Source/Runtime/Engine/Private/AlphaBlend.cpp
--- a/Engine/Source/Runtime/Engine/Private/AlphaBlend.cpp
+++ b/Engine/Source/Runtime/Engine/Private/AlphaBlend.cpp
@@ -235,6 +235,14 @@ FAlphaBlendArgs::FAlphaBlendArgs()
 
 }
 
+FAlphaBlendArgs::FAlphaBlendArgs(float InBlendTime)
+	: CustomCurve(nullptr)
+	, BlendTime(InBlendTime)
+	, BlendOption(EAlphaBlendOption::Linear)
+{
+
+}
+
 FAlphaBlendArgs::FAlphaBlendArgs(const struct FAlphaBlend& InAlphaBlend)
 	: CustomCurve(InAlphaBlend.GetCustomCurve())
 	, BlendTime(InAlphaBlend.GetBlendTime())
diff --git a/Engine/Source/Runtime/Engine/Private/Animation/AnimInstance.cpp b/Engine/Source/Runtime/Engine/Private/Animation/AnimInstance.cpp
--- a/Engine/Source/Runtime/Engine/Private/Animation/AnimInstance.cpp
+++ b/Engine/Source/Runtime/Engine/Private/Animation/AnimInstance.cpp
@@ -1839,11 +1839,25 @@ void UAnimInstance::TriggerQueuedMontageEvents()
 }
 
 UAnimMontage* UAnimInstance::PlaySlotAnimationAsDynamicMontage(UAnimSequenceBase* Asset, FName SlotNodeName, float BlendInTime, float BlendOutTime, float InPlayRate, int32 LoopCount, float BlendOutTriggerTime, float InTimeToStartMontageAt)
+{
+	FMontageBlendSettings BlendInSettings(BlendInTime);
+	FMontageBlendSettings BlendOutSettings(BlendOutTime);
+	return PlaySlotAnimationAsDynamicMontage_WithBlendSettings(Asset, SlotNodeName, BlendInSettings, BlendOutSettings, InPlayRate, LoopCount, BlendOutTriggerTime, InTimeToStartMontageAt);
+}
+
+UAnimMontage* UAnimInstance::PlaySlotAnimationAsDynamicMontage_WithBlendArgs(UAnimSequenceBase* Asset, FName SlotNodeName, const FAlphaBlendArgs& BlendIn, const FAlphaBlendArgs& BlendOut, float InPlayRate, int32 LoopCount, float BlendOutTriggerTime, float InTimeToStartMontageAt)
+{
+	FMontageBlendSettings BlendInSettings(BlendIn);
+	FMontageBlendSettings BlendOutSettings(BlendOut);
+	return PlaySlotAnimationAsDynamicMontage_WithBlendSettings(Asset, SlotNodeName, BlendInSettings, BlendOutSettings, InPlayRate, LoopCount, BlendOutTriggerTime, InTimeToStartMontageAt);
+}
+
+UAnimMontage* UAnimInstance::PlaySlotAnimationAsDynamicMontage_WithBlendSettings(UAnimSequenceBase* Asset, FName SlotNodeName, const FMontageBlendSettings& BlendInSettings, const FMontageBlendSettings& BlendOutSettings, float InPlayRate, int32 LoopCount, float BlendOutTriggerTime, float InTimeToStartMontageAt)
 {
 	if (Asset && CurrentSkeleton->IsCompatible(Asset->GetSkeleton()))
 	{
 		// create asset using the information
-		UAnimMontage* NewMontage = UAnimMontage::CreateSlotAnimationAsDynamicMontage(Asset, SlotNodeName, BlendInTime, BlendOutTime, InPlayRate, LoopCount, BlendOutTriggerTime, InTimeToStartMontageAt);
+		UAnimMontage* NewMontage = UAnimMontage::CreateSlotAnimationAsDynamicMontage_WithBlendSettings(Asset, SlotNodeName, BlendInSettings, BlendOutSettings, InPlayRate, LoopCount, BlendOutTriggerTime);
 
 		if (NewMontage)
 		{
diff --git a/Engine/Source/Runtime/Engine/Private/Animation/AnimMontage.cpp b/Engine/Source/Runtime/Engine/Private/Animation/AnimMontage.cpp
--- a/Engine/Source/Runtime/Engine/Private/Animation/AnimMontage.cpp
+++ b/Engine/Source/Runtime/Engine/Private/Animation/AnimMontage.cpp
@@ -3009,6 +3009,15 @@ UAnimMontage* FAnimMontageInstance::PreviewSequencerMontagePosition(FName SlotNa
 }
 
 UAnimMontage* UAnimMontage::CreateSlotAnimationAsDynamicMontage(UAnimSequenceBase* Asset, FName SlotNodeName, float BlendInTime, float BlendOutTime, float InPlayRate, int32 LoopCount, float BlendOutTriggerTime, float InTimeToStartMontageAt)
+{
+	FMontageBlendSettings BlendInSettings(BlendInTime);
+	FMontageBlendSettings BlendOutSettings(BlendOutTime);
+
+	// InTimeToStartMontageAt is an unused argument. Keeping it to avoid changing public api.
+	return CreateSlotAnimationAsDynamicMontage_WithBlendSettings(Asset, SlotNodeName, BlendInSettings, BlendOutSettings, InPlayRate, LoopCount, BlendOutTriggerTime);
+}
+
+UAnimMontage* UAnimMontage::CreateSlotAnimationAsDynamicMontage_WithBlendSettings(UAnimSequenceBase* Asset, FName SlotNodeName, const FMontageBlendSettings& BlendInSettings, const FMontageBlendSettings& BlendOutSettings, float InPlayRate, int32 LoopCount, float InBlendOutTriggerTime)
 {
 	// create temporary montage and play
 	bool bValidAsset = Asset && !Asset->IsA(UAnimMontage::StaticClass());
@@ -3059,9 +3068,16 @@ UAnimMontage* UAnimMontage::CreateSlotAnimationAsDynamicMontage(UAnimSequenceBas
 
 	// add new section
 	NewMontage->CompositeSections.Add(NewSection);
-	NewMontage->BlendIn.SetBlendTime(BlendInTime);
-	NewMontage->BlendOut.SetBlendTime(BlendOutTime);
-	NewMontage->BlendOutTriggerTime = BlendOutTriggerTime;
+
+	NewMontage->BlendIn = FAlphaBlend(BlendInSettings.Blend);
+	NewMontage->BlendModeIn = BlendInSettings.BlendMode;
+	NewMontage->BlendProfileIn = BlendInSettings.BlendProfile;
+
+	NewMontage->BlendOut = FAlphaBlend(BlendOutSettings.Blend);
+	NewMontage->BlendModeOut = BlendOutSettings.BlendMode;
+	NewMontage->BlendProfileOut = BlendOutSettings.BlendProfile;
+
+	NewMontage->BlendOutTriggerTime = InBlendOutTriggerTime;
 	return NewMontage;
 }
 
@@ -3103,8 +3119,17 @@ void UAnimMontage::BakeTimeStretchCurve()
 
 FMontageBlendSettings::FMontageBlendSettings()
 	: BlendProfile(nullptr)
-	, Blend(FAlphaBlendArgs())
 	, BlendMode(EMontageBlendMode::Standard)
-{
+{}
+
+FMontageBlendSettings::FMontageBlendSettings(float BlendTime)
+	: BlendProfile(nullptr)
+	, Blend(BlendTime)
+	, BlendMode(EMontageBlendMode::Standard)
+{}
 
-}
\ No newline at end of file
+FMontageBlendSettings::FMontageBlendSettings(const FAlphaBlendArgs& BlendArgs)
+	: BlendProfile(nullptr)
+	, Blend(BlendArgs)
+	, BlendMode(EMontageBlendMode::Standard)
+{}
diff --git a/Engine/Source/Runtime/Engine/Public/AlphaBlend.h b/Engine/Source/Runtime/Engine/Public/AlphaBlend.h
--- a/Engine/Source/Runtime/Engine/Public/AlphaBlend.h
+++ b/Engine/Source/Runtime/Engine/Public/AlphaBlend.h
@@ -52,6 +52,7 @@ struct ENGINE_API FAlphaBlendArgs
 	GENERATED_BODY()
 
 	FAlphaBlendArgs();
+	FAlphaBlendArgs(float InBlendTime);
 	FAlphaBlendArgs(const struct FAlphaBlend& InAlphaBlend);
 
 	/** If you're using Custom BlendOption, you can specify curve */
