Remove cook Class/Concurrency injection moved to Serialize from PostInitProperties. No longer required and breaking server builds.

diff --git a/Engine/Source/Runtime/AudioMixer/Private/Components/SynthComponent.cpp b/Engine/Source/Runtime/AudioMixer/Private/Components/SynthComponent.cpp
--- a/Engine/Source/Runtime/AudioMixer/Private/Components/SynthComponent.cpp
+++ b/Engine/Source/Runtime/AudioMixer/Private/Components/SynthComponent.cpp
@@ -356,7 +356,8 @@ USoundClass* USynthComponent::GetSoundClass()
 		return SoundClass;
 	}
 
-	if (const UAudioSettings* AudioSettings = GetDefault<UAudioSettings>())
+	const UAudioSettings* AudioSettings = GetDefault<UAudioSettings>();
+	if (ensure(AudioSettings))
 	{
 		return AudioSettings->GetDefaultSoundClass();
 	}
@@ -417,18 +418,6 @@ void USynthComponent::Serialize(FArchive& Ar)
 	}
 #endif // WITH_EDITORONLY_DATA
 
-	if (Ar.IsCooking())
-	{
-		if (!SoundClass)
-		{
-			SoundClass = GetDefault<UAudioSettings>()->GetDefaultSoundClass();
-		}
-
-		if (ConcurrencySet.IsEmpty())
-		{
-			ConcurrencySet.Add(GetDefault<UAudioSettings>()->GetDefaultSoundConcurrency());
-		}
-	}
 }
 
 void USynthComponent::PumpPendingMessages()
diff --git a/Engine/Source/Runtime/Engine/Private/SoundBase.cpp b/Engine/Source/Runtime/Engine/Private/SoundBase.cpp
--- a/Engine/Source/Runtime/Engine/Private/SoundBase.cpp
+++ b/Engine/Source/Runtime/Engine/Private/SoundBase.cpp
@@ -112,7 +112,8 @@ USoundClass* USoundBase::GetSoundClass() const
 		return SoundClassObject;
 	}
 
-	if (const UAudioSettings* AudioSettings = GetDefault<UAudioSettings>())
+	const UAudioSettings* AudioSettings = GetDefault<UAudioSettings>();
+	if (ensure(AudioSettings))
 	{
 		if (USoundClass* DefaultSoundClass = AudioSettings->GetDefaultSoundClass())
 		{
@@ -148,29 +149,28 @@ void USoundBase::GetSoundSourceBusSends(EBusSendType BusSendType, TArray<FSoundS
 void USoundBase::GetConcurrencyHandles(TArray<FConcurrencyHandle>& OutConcurrencyHandles) const
 {
 	const UAudioSettings* AudioSettings = GetDefault<UAudioSettings>();
-	check(AudioSettings);
 
 	OutConcurrencyHandles.Reset();
 	if (bOverrideConcurrency)
 	{
 		OutConcurrencyHandles.Add(ConcurrencyOverrides);
 	}
-	else
+	else if (!ConcurrencySet.IsEmpty())
 	{
-		if (ConcurrencySet.Num() > 0)
+		for (const USoundConcurrency* Concurrency : ConcurrencySet)
 		{
-			for (const USoundConcurrency* Concurrency : ConcurrencySet)
+			if (Concurrency)
 			{
-				if (Concurrency)
-				{
-					OutConcurrencyHandles.Emplace(*Concurrency);
-				}
+				OutConcurrencyHandles.Emplace(*Concurrency);
 			}
 		}
-		else if (const USoundConcurrency* DefaultConcurrency = AudioSettings->GetDefaultSoundConcurrency())
+	}
+	else if (ensure(AudioSettings))
+	{
+		if (const USoundConcurrency* DefaultConcurrency = AudioSettings->GetDefaultSoundConcurrency())
 		{
 			OutConcurrencyHandles.Emplace(*DefaultConcurrency);
-		}
+		}	
 	}
 }
 
@@ -233,19 +233,6 @@ void USoundBase::Serialize(FArchive& Ar)
 		}
 	}
 #endif // WITH_EDITORONLY_DATA
-
-	if (Ar.IsCooking())
-	{
-		if (!SoundClassObject)
-		{
-			SoundClassObject = GetDefault<UAudioSettings>()->GetDefaultSoundClass();
-		}
-
-		if (ConcurrencySet.IsEmpty())
-		{
-			ConcurrencySet.Add(GetDefault<UAudioSettings>()->GetDefaultSoundConcurrency());
-		}
-	}
 }
 
 void USoundBase::AddAssetUserData(UAssetUserData* InUserData)
