Disabled DXC for Android Vulkan as some specific shaders cause PSO compilation failures on Qualcomm devices
Looks like driver does not like array of vertex attributes

diff --git a/Engine/Source/Runtime/Engine/Private/ShaderCompiler/ShaderCompiler.cpp b/Engine/Source/Runtime/Engine/Private/ShaderCompiler/ShaderCompiler.cpp
--- a/Engine/Source/Runtime/Engine/Private/ShaderCompiler/ShaderCompiler.cpp
+++ b/Engine/Source/Runtime/Engine/Private/ShaderCompiler/ShaderCompiler.cpp
@@ -926,7 +926,7 @@ static TAutoConsoleVariable<int32> CVarOpenGLForceDXC(
 
 static TAutoConsoleVariable<int32> CVarVulkanForceDXC(
 	TEXT("r.Vulkan.ForceDXC"),
-	3,
+	1, 
 	TEXT("Forces DirectX Shader Compiler (DXC) to be used for all Vulkan shaders instead of hlslcc.\n")
 	TEXT(" 0: Disable (hlslcc/glslang)\n")
 	TEXT(" 1: Enabled on desktop platforms only (default)\n")
