Strata - fixed opacity not beeing forwarded correctly when root node is in material attribute mode.

diff --git a/Engine/Source/Runtime/Engine/Private/Materials/Material.cpp b/Engine/Source/Runtime/Engine/Private/Materials/Material.cpp
--- a/Engine/Source/Runtime/Engine/Private/Materials/Material.cpp
+++ b/Engine/Source/Runtime/Engine/Private/Materials/Material.cpp
@@ -3675,8 +3675,8 @@ void UMaterial::ConvertMaterialToStrataMaterial()
 		SlabBSDF->Roughness.Connect(3, BreakMatAtt);	// 3- Roughness
 		SlabBSDF->Anisotropy.Connect(4, BreakMatAtt);	// 4- Anisotropy
 		SlabBSDF->EmissiveColor.Connect(5, BreakMatAtt);// 5- EmissiveColor
-		// 6- Opacity									// We should not have to do anything, it should be forwarded to the root node correctly.
-		// 7- Opacity Mask								// Idem
+		// 6- Opacity									// See below
+		// 7- Opacity Mask								// See below
 		SlabBSDF->Normal.Connect(8, BreakMatAtt);		// 8- Normal
 		SlabBSDF->Tangent.Connect(9, BreakMatAtt);		// 9- Tangent
 		// 10- WorldPositionOffset
@@ -3690,6 +3690,9 @@ void UMaterial::ConvertMaterialToStrataMaterial()
 
 		UMaterialExpressionMakeMaterialAttributes* MakeMatAtt = NewObject<UMaterialExpressionMakeMaterialAttributes>(this);
 		MakeMatAtt->FrontMaterial.Connect(0, SlabBSDF);
+		// Forward opacity and opacity mask
+		MakeMatAtt->Opacity.Connect(6, BreakMatAtt);
+		MakeMatAtt->OpacityMask.Connect(7, BreakMatAtt);
 
 		MaterialAttributes.Connect(0, MakeMatAtt);
 	}
