WorldPartition: Replay support

diff --git a/Engine/Source/Runtime/Engine/Private/ReplayHelper.cpp b/Engine/Source/Runtime/Engine/Private/ReplayHelper.cpp
--- a/Engine/Source/Runtime/Engine/Private/ReplayHelper.cpp
+++ b/Engine/Source/Runtime/Engine/Private/ReplayHelper.cpp
@@ -138,7 +138,8 @@ void FReplayHelper::StartRecording(UNetConnection* Connection)
 {
 	World = Connection ? Connection->GetWorld() : nullptr;
 
-	bHasLevelStreamingFixes = !!CVarWithLevelStreamingFixes.GetValueOnAnyThread();
+	// World Partition needs bHasLevelStreamingFixes to be true
+	bHasLevelStreamingFixes = !!CVarWithLevelStreamingFixes.GetValueOnAnyThread() || (World->GetWorldPartition() != nullptr);
 	bHasDeltaCheckpoints = !!CVarWithDeltaCheckpoints.GetValueOnAnyThread() && ReplayStreamer->IsCheckpointTypeSupported(EReplayCheckpointType::Delta);
 	bHasGameSpecificFrameData = !!CVarWithGameSpecificFrameData.GetValueOnAnyThread();
 
diff --git a/Engine/Source/Runtime/Engine/Private/WorldPartition/DataLayer/WorldDataLayers.cpp b/Engine/Source/Runtime/Engine/Private/WorldPartition/DataLayer/WorldDataLayers.cpp
--- a/Engine/Source/Runtime/Engine/Private/WorldPartition/DataLayer/WorldDataLayers.cpp
+++ b/Engine/Source/Runtime/Engine/Private/WorldPartition/DataLayer/WorldDataLayers.cpp
@@ -22,6 +22,10 @@ AWorldDataLayers::AWorldDataLayers(const FObjectInitializer& ObjectInitializer)
 {
 	bAlwaysRelevant = true;
 	bReplicates = true;
+
+	// Avoid actor from being Destroyed/Recreated when scrubbing a replay
+	// instead AWorldDataLayers::RewindForReplay() gets called to reset this actors state
+	bReplayRewindable = true;
 }
 
 void AWorldDataLayers::GetLifetimeReplicatedProps(TArray< FLifetimeProperty >& OutLifetimeProps) const
@@ -32,8 +36,30 @@ void AWorldDataLayers::GetLifetimeReplicatedProps(TArray< FLifetimeProperty >& O
 	DOREPLIFETIME(AWorldDataLayers, RepActiveDataLayerNames);
 }
 
+void AWorldDataLayers::BeginPlay()
+{
+	Super::BeginPlay();
+
+	// When running a Replay we want to reset our state to CDO (empty) and rely on the Replay/Replication.
+	// Unfortunately this can't be tested in the PostLoad as the World doesn't have a demo driver yet.
+	if (GetWorld()->IsPlayingReplay())
+	{
+		ResetDataLayerStates();
+	}
+}
+
+void AWorldDataLayers::RewindForReplay()
+{
+	Super::RewindForReplay();
+
+	// Same as BeginPlay when rewinding we want to reset our state to CDO (empty) and rely on Replay/Replication.
+	ResetDataLayerStates();
+}
+
 void AWorldDataLayers::InitializeDataLayerStates()
 {
+	check(ActiveDataLayerNames.IsEmpty() && LoadedDataLayerNames.IsEmpty());
+
 	if (GetWorld()->IsGameWorld())
 	{
 		ForEachDataLayer([this](class UDataLayer* DataLayer)
@@ -57,6 +83,14 @@ void AWorldDataLayers::InitializeDataLayerStates()
 	}
 }
 
+void AWorldDataLayers::ResetDataLayerStates()
+{
+	ActiveDataLayerNames.Reset();
+	LoadedDataLayerNames.Reset();
+	RepActiveDataLayerNames.Reset();
+	RepLoadedDataLayerNames.Reset();
+}
+
 void AWorldDataLayers::SetDataLayerState(FActorDataLayer InDataLayer, EDataLayerState InState)
 {
 	if (ensure(GetLocalRole() == ROLE_Authority))
diff --git a/Engine/Source/Runtime/Engine/Private/WorldPartition/WorldPartitionStreamingPolicy.cpp b/Engine/Source/Runtime/Engine/Private/WorldPartition/WorldPartitionStreamingPolicy.cpp
--- a/Engine/Source/Runtime/Engine/Private/WorldPartition/WorldPartitionStreamingPolicy.cpp
+++ b/Engine/Source/Runtime/Engine/Private/WorldPartition/WorldPartitionStreamingPolicy.cpp
@@ -87,16 +87,11 @@ void UWorldPartitionStreamingPolicy::UpdateStreamingSources()
 					FVector ViewLocation;
 					FRotator ViewRotation;
 					Player->PlayerController->GetPlayerViewPoint(ViewLocation, ViewRotation);
-
-					//@todo_ow: this test is to cover cases where GetPlayerViewPoint returns (0,0,0) when invalid. It should probably return a bool
-					//          to indicate that the returned position is invalid.
-					if (!ViewLocation.IsZero())
-					{
-						// Transform to Local
-						ViewLocation = WorldToLocal.TransformPosition(ViewLocation);
-						ViewRotation = WorldToLocal.TransformRotation(ViewRotation.Quaternion()).Rotator();
-						StreamingSources.Add(FWorldPartitionStreamingSource(Player->GetFName(), ViewLocation, ViewRotation, EStreamingSourceTargetState::Activated));
-					}
+					
+					// Transform to Local
+					ViewLocation = WorldToLocal.TransformPosition(ViewLocation);
+					ViewRotation = WorldToLocal.TransformRotation(ViewRotation.Quaternion()).Rotator();
+					StreamingSources.Add(FWorldPartitionStreamingSource(Player->GetFName(), ViewLocation, ViewRotation, EStreamingSourceTargetState::Activated));
 				}
 			}
 		}
diff --git a/Engine/Source/Runtime/Engine/Public/WorldPartition/DataLayer/WorldDataLayers.h b/Engine/Source/Runtime/Engine/Public/WorldPartition/DataLayer/WorldDataLayers.h
--- a/Engine/Source/Runtime/Engine/Public/WorldPartition/DataLayer/WorldDataLayers.h
+++ b/Engine/Source/Runtime/Engine/Public/WorldPartition/DataLayer/WorldDataLayers.h
@@ -20,7 +20,8 @@ class ENGINE_API AWorldDataLayers : public AInfo
 
 public:
 	virtual void PostLoad() override;
-	
+	virtual void RewindForReplay() override;
+	virtual void BeginPlay() override;
 #if WITH_EDITOR
 	static AWorldDataLayers* Create(UWorld* World);
 	UDataLayer* CreateDataLayer(FName InName = TEXT("DataLayer"), EObjectFlags InObjectFlags = RF_NoFlags);
@@ -54,6 +55,7 @@ public:
 
 protected:
 	void InitializeDataLayerStates();
+	void ResetDataLayerStates();
 
 	UFUNCTION()
 	void OnRep_ActiveDataLayerNames();
