Revert changes to skin SRV initialization. They are no longer needed after CL# 16508177

diff --git a/Engine/Source/Runtime/Engine/Private/Rendering/SkinWeightVertexBuffer.cpp b/Engine/Source/Runtime/Engine/Private/Rendering/SkinWeightVertexBuffer.cpp
--- a/Engine/Source/Runtime/Engine/Private/Rendering/SkinWeightVertexBuffer.cpp
+++ b/Engine/Source/Runtime/Engine/Private/Rendering/SkinWeightVertexBuffer.cpp
@@ -139,7 +139,7 @@ void FSkinWeightLookupVertexBuffer::InitRHI()
 
 		if (bSRV)
 		{
-			SRVValue = RHICreateShaderResourceView(FShaderResourceViewInitializer(VertexBufferRHI, PixelFormat));
+			SRVValue = RHICreateShaderResourceView(FShaderResourceViewInitializer(LookupData ? VertexBufferRHI : nullptr, PixelFormat));
 		}
 	}
 }
@@ -432,18 +432,16 @@ void FSkinWeightDataVertexBuffer::InitRHI()
 
 	// BUF_ShaderResource is needed for support of the SkinCache (we could make is dependent on GEnableGPUSkinCacheShaders or are there other users?)
 	VertexBufferRHI = CreateRHIBuffer_RenderThread();
-	if (VertexBufferRHI)
-	{
-		bool bSRV = GSupportsResourceView && GPixelFormats[GetPixelFormat()].Supported;
-		// When bAllowCPUAccess is true, the meshes is likely going to be used for Niagara to spawn particles on mesh surface.
-		// And it can be the case for CPU *and* GPU access: no differenciation today. That is why we create a SRV in this case.
-		// This also avoid setting lots of states on all the members of all the different buffers used by meshes. Follow up: https://jira.it.epicgames.net/browse/UE-69376.
-		bSRV |= GetNeedsCPUAccess();
+	
+	bool bSRV = VertexBufferRHI && GSupportsResourceView && GPixelFormats[GetPixelFormat()].Supported;
+	// When bAllowCPUAccess is true, the meshes is likely going to be used for Niagara to spawn particles on mesh surface.
+	// And it can be the case for CPU *and* GPU access: no differenciation today. That is why we create a SRV in this case.
+	// This also avoid setting lots of states on all the members of all the different buffers used by meshes. Follow up: https://jira.it.epicgames.net/browse/UE-69376.
+	bSRV |= GetNeedsCPUAccess();
 
-		if (bSRV)
-		{
-			SRVValue = RHICreateShaderResourceView(FShaderResourceViewInitializer(VertexBufferRHI, GetPixelFormat()));
-		}
+	if (bSRV)
+	{
+		SRVValue = RHICreateShaderResourceView(FShaderResourceViewInitializer(WeightData ? VertexBufferRHI : nullptr, GetPixelFormat()));
 	}
 }
 
