LevelInstance: Add time log to ULevelInstanceSubsystem::BreakLevelInstance

diff --git a/Engine/Source/Runtime/Engine/Private/LevelInstance/LevelInstanceSubsystem.cpp b/Engine/Source/Runtime/Engine/Private/LevelInstance/LevelInstanceSubsystem.cpp
--- a/Engine/Source/Runtime/Engine/Private/LevelInstance/LevelInstanceSubsystem.cpp
+++ b/Engine/Source/Runtime/Engine/Private/LevelInstance/LevelInstanceSubsystem.cpp
@@ -866,6 +866,8 @@ ALevelInstance* ULevelInstanceSubsystem::CreateLevelInstanceFrom(const TArray<AA
 
 bool ULevelInstanceSubsystem::BreakLevelInstance(ALevelInstance* LevelInstanceActor, uint32 Levels /* = 1 */, TArray<AActor*>* OutMovedActors /* = nullptr */)
 {
+	const double StartTime = FPlatformTime::Seconds();
+
 	TArray<AActor*> MovedActors;
 	BreakLevelInstance_Impl(LevelInstanceActor, Levels, MovedActors);
 
@@ -878,6 +880,10 @@ bool ULevelInstanceSubsystem::BreakLevelInstance(ALevelInstance* LevelInstanceAc
 	ActorSelection->EndBatchSelectOperation(false);
 
 	bool bStatus = MovedActors.Num() > 0;
+
+	const double ElapsedTime = FPlatformTime::Seconds() - StartTime;
+	UE_LOG(LogLevelInstance, Log, TEXT("Break took %s seconds (%s actors)"), *FText::AsNumber(ElapsedTime).ToString(), *FText::AsNumber(MovedActors.Num()).ToString());
+	
 	if (OutMovedActors)
 	{
 		*OutMovedActors = MoveTemp(MovedActors);
