Fix missing override of the previous local to world transofmr in primitive unifor buffer.

diff --git a/Engine/Source/Runtime/Engine/Public/PrimitiveUniformShaderParameters.h b/Engine/Source/Runtime/Engine/Public/PrimitiveUniformShaderParameters.h
--- a/Engine/Source/Runtime/Engine/Public/PrimitiveUniformShaderParameters.h
+++ b/Engine/Source/Runtime/Engine/Public/PrimitiveUniformShaderParameters.h
@@ -119,7 +119,6 @@ public:
 
 	PRIMITIVE_UNIFORM_BUILDER_METHOD(const FVector3f&,	ActorWorldPosition);
 	PRIMITIVE_UNIFORM_BUILDER_METHOD(const FMatrix44f&,	LocalToWorld);
-	PRIMITIVE_UNIFORM_BUILDER_METHOD(const FMatrix44f&,	PreviousLocalToWorld);
 
 #undef PRIMITIVE_UNIFORM_BUILDER_FLAG_METHOD
 #undef PRIMITIVE_UNIFORM_BUILDER_METHOD
@@ -162,6 +161,13 @@ public:
 		return *this;
 	}
 
+	inline FPrimitiveUniformShaderParametersBuilder& PreviousLocalToWorld(const FMatrix44f& InPreviousLocalToWorld)
+	{
+		bHasPreviousLocalToWorld = true;
+		Parameters.PreviousLocalToWorld = InPreviousLocalToWorld;
+		return *this;
+	}
+
 	inline FPrimitiveUniformShaderParametersBuilder& CustomPrimitiveData(const FCustomPrimitiveData* InCustomPrimitiveData)
 	{
 		// If this primitive has custom primitive data, set it
