Change default behavior for whether to force routing actor initialize phase in its own frame or not, disabled by default.
Can be controlled by s.ForceRouteActorInitializeNextFrame cvar.

diff --git a/Engine/Source/Runtime/Engine/Classes/Engine/CoreSettings.h b/Engine/Source/Runtime/Engine/Classes/Engine/CoreSettings.h
--- a/Engine/Source/Runtime/Engine/Classes/Engine/CoreSettings.h
+++ b/Engine/Source/Runtime/Engine/Classes/Engine/CoreSettings.h
@@ -146,7 +146,8 @@ extern ENGINE_API int32 GLevelStreamingAllowLevelRequestsWhileAsyncLoadingInMatc
 extern ENGINE_API int32 GLevelStreamingMaxLevelRequestsAtOnceWhileInMatch;
 /** Whether to force a verification of objects residing in a GC'ed level package (ignored in shipping builds). */
 extern ENGINE_API int32 GLevelStreamingForceVerifyLevelsGotRemovedByGC;
-
+/** Whether to force routing actor initialize phase in its own frame. */
+extern ENGINE_API int32 GLevelStreamingForceRouteActorInitializeNextFrame;
 
 /**
 * Implements the settings for garbage collection.
diff --git a/Engine/Source/Runtime/Engine/Private/CoreSettings.cpp b/Engine/Source/Runtime/Engine/Private/CoreSettings.cpp
--- a/Engine/Source/Runtime/Engine/Private/CoreSettings.cpp
+++ b/Engine/Source/Runtime/Engine/Private/CoreSettings.cpp
@@ -24,6 +24,7 @@ int32 GLevelStreamingContinuouslyIncrementalGCWhileLevelsPendingPurge = 1;
 int32 GLevelStreamingAllowLevelRequestsWhileAsyncLoadingInMatch = 1;
 int32 GLevelStreamingMaxLevelRequestsAtOnceWhileInMatch = 0;
 int32 GLevelStreamingForceVerifyLevelsGotRemovedByGC = 0;
+int32 GLevelStreamingForceRouteActorInitializeNextFrame = 0;
 
 static FAutoConsoleVariableRef CVarUseBackgroundLevelStreaming(
 	TEXT("s.UseBackgroundLevelStreaming"),
@@ -102,6 +103,13 @@ static FAutoConsoleVariableRef CVarForceVerifyLevelsGotRemovedByGC(
 	ECVF_Default
 );
 
+static FAutoConsoleVariableRef CVarForceRouteActorInitializeNextFrame(
+	TEXT("s.ForceRouteActorInitializeNextFrame"),
+	GLevelStreamingForceRouteActorInitializeNextFrame,
+	TEXT("Whether to force routing actor initialize phase in its own frame."),
+	ECVF_Default
+);
+
 static FAutoConsoleVariableRef CVarContinuouslyIncrementalGCWhileLevelsPendingPurge(
 	TEXT("s.ContinuouslyIncrementalGCWhileLevelsPendingPurge"),
 	GLevelStreamingContinuouslyIncrementalGCWhileLevelsPendingPurge,
diff --git a/Engine/Source/Runtime/Engine/Private/World.cpp b/Engine/Source/Runtime/Engine/Private/World.cpp
--- a/Engine/Source/Runtime/Engine/Private/World.cpp
+++ b/Engine/Source/Runtime/Engine/Private/World.cpp
@@ -2728,7 +2728,7 @@ void UWorld::AddToWorld( ULevel* Level, const FTransform& LevelTransform, bool b
 				Level->ClearActorsSeamlessTraveledFlag();
 			}
 
-			const float PreventNextStepTimeLimit = 0.0; // We will always run route actor initialize in its own frame if we are using a time limit
+			const float PreventNextStepTimeLimit = GLevelStreamingForceRouteActorInitializeNextFrame ? 0.0f : TimeLimit;
 			bExecuteNextStep = (!bConsiderTimeLimit || !IsTimeLimitExceeded( TEXT("initializing network actors"), StartTime, Level, PreventNextStepTimeLimit ));
 		}
 
