Workaround for adding manually loaded AudioSettings objects to root to avoid GC CDO mutation warning
#rnx

diff --git a/Engine/Source/Runtime/Engine/Private/AudioSettings.cpp b/Engine/Source/Runtime/Engine/Private/AudioSettings.cpp
--- a/Engine/Source/Runtime/Engine/Private/AudioSettings.cpp
+++ b/Engine/Source/Runtime/Engine/Private/AudioSettings.cpp
@@ -186,10 +186,16 @@ void UAudioSettings::LoadDefaultObjects()
 	// TODO: Move all soft object paths to load here (SoundMixes, Submixes, etc.)
 	static const FString EngineSoundsDir = TEXT("/Engine/EngineSounds");
 
-	DefaultSoundClass = nullptr;
+	if (DefaultSoundClass)
+	{
+		DefaultSoundClass->RemoveFromRoot();
+		DefaultSoundClass = nullptr;
+	}
+
 	if (UObject* SoundClassObject = DefaultSoundClassName.TryLoad())
 	{
 		DefaultSoundClass = CastChecked<USoundClass>(SoundClassObject);
+		DefaultSoundClass->AddToRoot();
 	}
 
 #if WITH_EDITOR
@@ -199,6 +205,7 @@ void UAudioSettings::LoadDefaultObjects()
 		if (UObject* SoundClassObject = DefaultSoundClassName.TryLoad())
 		{
 			DefaultSoundClass = CastChecked<USoundClass>(SoundClassObject);
+			DefaultSoundClass->AddToRoot();
 		}
 	}
 #endif // WITH_EDITOR
@@ -210,6 +217,7 @@ void UAudioSettings::LoadDefaultObjects()
 		if (UObject* SoundClassObject = DefaultSoundClassName.TryLoad())
 		{
 			DefaultSoundClass = CastChecked<USoundClass>(SoundClassObject);
+			DefaultSoundClass->AddToRoot();
 		}
 	}
 
@@ -218,7 +226,12 @@ void UAudioSettings::LoadDefaultObjects()
 		UE_LOG(LogAudio, Error, TEXT("Failed to load Default SoundClassObject from path '%s'."), *DefaultSoundClassName.GetAssetPathString());
 	}
 
-	DefaultMediaSoundClass = nullptr;
+	if (DefaultMediaSoundClass)
+	{
+		DefaultMediaSoundClass->RemoveFromRoot();
+		DefaultMediaSoundClass = nullptr;
+	}
+
 	if (UObject* MediaSoundClassObject = DefaultMediaSoundClassName.TryLoad())
 	{
 		DefaultMediaSoundClass = CastChecked<USoundClass>(MediaSoundClassObject);
@@ -228,10 +241,16 @@ void UAudioSettings::LoadDefaultObjects()
 		UE_LOG(LogAudio, Display, TEXT("No default MediaSoundClassObject specified (or failed to load)."));
 	}
 
-	DefaultSoundConcurrency = nullptr;
+	if (DefaultSoundConcurrency)
+	{
+		DefaultSoundConcurrency->RemoveFromRoot();
+		DefaultSoundConcurrency = nullptr;
+	}
+
 	if (UObject* SoundConcurrencyObject = DefaultSoundConcurrencyName.TryLoad())
 	{
 		DefaultSoundConcurrency = CastChecked<USoundConcurrency>(SoundConcurrencyObject);
+		DefaultSoundConcurrency->AddToRoot();
 	}
 	else
 	{
