Dump preprocessed shader source in log on D3DCompile exception.

#rnx

diff --git a/Engine/Source/Developer/Windows/ShaderFormatD3D/Private/D3DShaderCompiler.cpp b/Engine/Source/Developer/Windows/ShaderFormatD3D/Private/D3DShaderCompiler.cpp
--- a/Engine/Source/Developer/Windows/ShaderFormatD3D/Private/D3DShaderCompiler.cpp
+++ b/Engine/Source/Developer/Windows/ShaderFormatD3D/Private/D3DShaderCompiler.cpp
@@ -650,6 +650,14 @@ bool CompileAndProcessD3DShaderFXC(FString& PreprocessedShaderSource, const FStr
 					ExceptionInfo.Base
 				)
 			);
+
+			// Dump input shader source on exception to be able to investigate issue through logs on CIS servers
+			FString DumpedSource;
+			DumpedSource = PreprocessedShaderSource;
+			DumpedSource += TEXT("\n#if 0 /*DIRECT COMPILE*/\n");
+			DumpedSource += CreateShaderCompilerWorkerDirectCommandLine(Input);
+			DumpedSource += TEXT("\n#endif /*DIRECT COMPILE*/\n");
+			Output.OptionalPreprocessedShaderSource = MoveTemp(DumpedSource);
 		}
 	}
 	else
diff --git a/Engine/Source/Runtime/Engine/Private/ShaderCompiler/ShaderCompiler.cpp b/Engine/Source/Runtime/Engine/Private/ShaderCompiler/ShaderCompiler.cpp
--- a/Engine/Source/Runtime/Engine/Private/ShaderCompiler/ShaderCompiler.cpp
+++ b/Engine/Source/Runtime/Engine/Private/ShaderCompiler/ShaderCompiler.cpp
@@ -125,6 +125,9 @@ static FAutoConsoleVariableRef CVarShaderCompilerAllowDistributedCompilation(
 	ECVF_Default
 );
 
+/** Maximum number of preprocessed shaders to dump to the log on a crash. Replace with CVar if needed. */
+static constexpr uint32 GMaxNumDumpedShaderSources = 10;
+
 /** Helper functions for logging more debug info */
 namespace ShaderCompiler
 {
@@ -1709,47 +1712,52 @@ void FShaderCompileUtilities::DoReadTaskResults(const TArray<FShaderCommonCompil
 	}
 }
 
-static bool CheckSingleJob(FShaderCompileJob* SingleJob, TArray<FString>& Errors)
+static bool CheckSingleJob(const FShaderCompileJob& SingleJob, TArray<FString>& OutErrors, FString* OutDumpedSource)
 {
-	if (SingleJob->bSucceeded)
+	if (SingleJob.bSucceeded)
 	{
-		check(SingleJob->Output.ShaderCode.GetShaderCodeSize() > 0);
+		check(SingleJob.Output.ShaderCode.GetShaderCodeSize() > 0);
 	}
 
-	if (GShowShaderWarnings || !SingleJob->bSucceeded)
+	if (GShowShaderWarnings || !SingleJob.bSucceeded)
 	{
-		for (int32 ErrorIndex = 0; ErrorIndex < SingleJob->Output.Errors.Num(); ErrorIndex++)
+		for (int32 ErrorIndex = 0; ErrorIndex < SingleJob.Output.Errors.Num(); ErrorIndex++)
 		{
-			const FShaderCompilerError& InError = SingleJob->Output.Errors[ErrorIndex];
-			Errors.AddUnique(InError.GetErrorStringWithLineMarker());
+			const FShaderCompilerError& InError = SingleJob.Output.Errors[ErrorIndex];
+			OutErrors.AddUnique(InError.GetErrorStringWithLineMarker());
 		}
 	}
 
-	bool bSucceeded = SingleJob->bSucceeded;
+	bool bSucceeded = SingleJob.bSucceeded;
 
-	if (SingleJob->Key.ShaderType)
+	if (SingleJob.Key.ShaderType)
 	{
 		// Allow the shader validation to fail the compile if it sees any parameters bound that aren't supported.
-		const bool bValidationResult = SingleJob->Key.ShaderType->ValidateCompiledResult(
-			(EShaderPlatform)SingleJob->Input.Target.Platform,
-			SingleJob->Output.ParameterMap,
-			Errors);
+		const bool bValidationResult = SingleJob.Key.ShaderType->ValidateCompiledResult(
+			(EShaderPlatform)SingleJob.Input.Target.Platform,
+			SingleJob.Output.ParameterMap,
+			OutErrors);
 		bSucceeded = bValidationResult && bSucceeded;
 	}
 
-	if (SingleJob->Key.VFType)
+	if (SingleJob.Key.VFType)
 	{
-		const int32 OriginalNumErrors = Errors.Num();
+		const int32 OriginalNumErrors = OutErrors.Num();
 
 		// Allow the vertex factory to fail the compile if it sees any parameters bound that aren't supported
-		SingleJob->Key.VFType->ValidateCompiledResult((EShaderPlatform)SingleJob->Input.Target.Platform, SingleJob->Output.ParameterMap, Errors);
+		SingleJob.Key.VFType->ValidateCompiledResult((EShaderPlatform)SingleJob.Input.Target.Platform, SingleJob.Output.ParameterMap, OutErrors);
 
-		if (Errors.Num() > OriginalNumErrors)
+		if (OutErrors.Num() > OriginalNumErrors)
 		{
 			bSucceeded = false;
 		}
 	}
 
+	if (!bSucceeded && OutDumpedSource != nullptr)
+	{
+		*OutDumpedSource = SingleJob.Output.OptionalPreprocessedShaderSource;
+	}
+
 	return bSucceeded;
 };
 
@@ -2813,6 +2821,7 @@ FShaderCompilingManager::FShaderCompilingManager() :
 	bCompilingDuringGame(false),
 	NumExternalJobs(0),
 	NumSingleThreadedRunsBeforeRetry(GSingleThreadedRunsIdle),
+	NumDumpedShaderSources(0),
 #if PLATFORM_MAC
 	ShaderCompileWorkerName(FPaths::EngineDir() / TEXT("Binaries/Mac/ShaderCompileWorker")),
 #elif PLATFORM_LINUX
@@ -3709,6 +3718,7 @@ void FShaderCompilingManager::ProcessCompiledShaderMaps(
 			TArray<TRefCountPtr<FMaterial>>& MaterialDependencies = CompilingShaderMap->CompilingMaterialDependencies;
 
 			TArray<FString> Errors;
+			FString DumpedSource;
 			TArray<FShaderCommonCompileJobPtr>& ResultArray = CompileResults.FinishedJobs;
 
 			bool bSuccess = true;
@@ -3716,18 +3726,17 @@ void FShaderCompilingManager::ProcessCompiledShaderMaps(
 			{
 				FShaderCommonCompileJob& CurrentJob = *ResultArray[JobIndex];
 
-				auto* SingleJob = CurrentJob.GetSingleShaderJob();
-				if (SingleJob)
+				if (FShaderCompileJob* SingleJob = CurrentJob.GetSingleShaderJob())
 				{
-					const bool bCheckSucceeded = CheckSingleJob(SingleJob, Errors);
+					const bool bCheckSucceeded = CheckSingleJob(*SingleJob, Errors, (NumDumpedShaderSources < GMaxNumDumpedShaderSources ? &DumpedSource : nullptr));
 					bSuccess = bCheckSucceeded && bSuccess;
 				}
 				else
 				{
-					auto* PipelineJob = CurrentJob.GetShaderPipelineJob();
+					FShaderPipelineCompileJob* PipelineJob = CurrentJob.GetShaderPipelineJob();
 					for (int32 Index = 0; Index < PipelineJob->StageJobs.Num(); ++Index)
 					{
-						const bool bCheckSucceeded = CheckSingleJob(PipelineJob->StageJobs[Index], Errors);
+						const bool bCheckSucceeded = CheckSingleJob(*PipelineJob->StageJobs[Index], Errors, (NumDumpedShaderSources < GMaxNumDumpedShaderSources ? &DumpedSource : nullptr));
 						bSuccess = PipelineJob->StageJobs[Index]->bSucceeded && bCheckSucceeded && bSuccess;
 					}
 				}
@@ -3817,6 +3826,24 @@ void FShaderCompilingManager::ProcessCompiledShaderMaps(
 
 					MaterialsToUpdate.Add(Material, nullptr);
 
+					if (!DumpedSource.IsEmpty() && NumDumpedShaderSources < GMaxNumDumpedShaderSources)
+					{
+						// Log dumped shader source line by line as message lengths in UE_LOG are limited.
+						TArray<FString> DumpedSourceLines;
+						DumpedSource.ParseIntoArrayLines(DumpedSourceLines, /*bCullEmpty:*/ false);
+
+						TGuardValue<ELogTimes::Type> DisableLogTimes(GPrintLogTimes, ELogTimes::None);
+						GLog->Serialize(TEXT("\n======================= DUMPED SHADER BEGIN =======================\n"), ELogVerbosity::Log, NAME_None);
+						for (const FString& Line : DumpedSourceLines)
+						{
+							GLog->Serialize(*Line, ELogVerbosity::Log, NAME_None);
+						}
+						GLog->Serialize(TEXT("\n======================= DUMPED SHADER END =======================\n"), ELogVerbosity::Log, NAME_None);
+
+						// Limit number of preprocessed shaders to dump to the log as they are quite large
+						++NumDumpedShaderSources;
+					}
+
 					if (Material->IsDefaultMaterial())
 					{
 						// Log the errors unsuppressed before the fatal error, so it's always obvious from the log what the compile error was
diff --git a/Engine/Source/Runtime/Engine/Public/ShaderCompiler.h b/Engine/Source/Runtime/Engine/Public/ShaderCompiler.h
--- a/Engine/Source/Runtime/Engine/Public/ShaderCompiler.h
+++ b/Engine/Source/Runtime/Engine/Public/ShaderCompiler.h
@@ -799,6 +799,8 @@ private:
 	int32 MaxShaderJobBatchSize;
 	/** Number of runs through single-threaded compiling before we can retry to compile through workers. -1 if not used. */
 	int32 NumSingleThreadedRunsBeforeRetry;
+	/** Number of preprocessed shader sources that are dumped due to a crash of the shader compiler. Only used for D3DCompile exceptions at the moment. */
+	uint32 NumDumpedShaderSources;
 	/** Process Id of UE. */
 	uint32 ProcessId;
 	/** Whether to allow compiling shaders through the worker application, which allows multiple cores to be used. */
diff --git a/Engine/Source/Runtime/RenderCore/Public/ShaderCompilerCore.h b/Engine/Source/Runtime/RenderCore/Public/ShaderCompilerCore.h
--- a/Engine/Source/Runtime/RenderCore/Public/ShaderCompilerCore.h
+++ b/Engine/Source/Runtime/RenderCore/Public/ShaderCompilerCore.h
@@ -18,7 +18,7 @@ class Error;
 // this is for the protocol, not the data, bump if FShaderCompilerInput or ProcessInputFromArchive changes.
 const int32 ShaderCompileWorkerInputVersion = 14;
 // this is for the protocol, not the data, bump if FShaderCompilerOutput or WriteToOutputArchive changes.
-const int32 ShaderCompileWorkerOutputVersion = 6;
+const int32 ShaderCompileWorkerOutputVersion = 7;
 // this is for the protocol, not the data.
 const int32 ShaderCompileWorkerSingleJobHeader = 'S';
 // this is for the protocol, not the data.
@@ -448,6 +448,7 @@ struct FShaderCompilerOutput
 	bool bUsedHLSLccCompiler;
 	TArray<FString> UsedAttributes;
 
+	FString OptionalPreprocessedShaderSource;
 	FString OptionalFinalShaderSource;
 
 	TArray<uint8> PlatformDebugData;
@@ -464,6 +465,7 @@ struct FShaderCompilerOutput
 		Ar << Output.ParameterMap << Output.Errors << Output.Target << Output.ShaderCode << Output.OutputHash << Output.NumInstructions << Output.NumTextureSamplers << Output.bSucceeded;
 		Ar << Output.bFailedRemovingUnused << Output.bSupportsQueryingUsedAttributes << Output.UsedAttributes;
 		Ar << Output.CompileTime;
+		Ar << Output.OptionalPreprocessedShaderSource;
 		Ar << Output.OptionalFinalShaderSource;
 		Ar << Output.PlatformDebugData;
 
