Fixed root motion extraction to account for offset in the ref pose.
Approved for Early Access by [REDACTED_PERSON] and Laurent Delayen

diff --git a/Engine/Source/Runtime/Engine/Private/Animation/AnimSequence.cpp b/Engine/Source/Runtime/Engine/Private/Animation/AnimSequence.cpp
--- a/Engine/Source/Runtime/Engine/Private/Animation/AnimSequence.cpp
+++ b/Engine/Source/Runtime/Engine/Private/Animation/AnimSequence.cpp
@@ -1321,7 +1321,16 @@ FTransform UAnimSequence::ExtractRootMotionFromRange(float StartTrackPosition, f
 {
 	const FVector DefaultScale(1.f);
 
-	FTransform InitialTransform = ExtractRootTrackTransform(0.f, NULL);
+	FTransform RootTransformRefPose = FTransform::Identity;
+	if (const USkeleton* MySkeleton = GetSkeleton())
+	{
+		const FReferenceSkeleton& RefSkeleton = MySkeleton->GetReferenceSkeleton();
+		if (RefSkeleton.GetNum() > 0)
+		{
+			RootTransformRefPose = RefSkeleton.GetRefBonePose()[0];
+		}
+	}
+
 	FTransform StartTransform = ExtractRootTrackTransform(StartTrackPosition, NULL);
 	FTransform EndTransform = ExtractRootTrackTransform(EndTrackPosition, NULL);
 
@@ -1341,9 +1350,10 @@ FTransform UAnimSequence::ExtractRootMotionFromRange(float StartTrackPosition, f
 		}
 	}
 
-	// Transform to Component Space Rotation (inverse root transform from first frame)
-	StartTransform = StartTransform.GetRelativeTransform(InitialTransform);
-    EndTransform = EndTransform.GetRelativeTransform(InitialTransform);
+	// Transform to Component Space
+	const FTransform RootToComponent = RootTransformRefPose.Inverse();
+	StartTransform = RootToComponent * StartTransform;
+	EndTransform = RootToComponent * EndTransform;
 
 	return EndTransform.GetRelativeTransform(StartTransform);
 }
