Improve hardware ray tracing related project settings

Don't allow features that depend on ray tracing to be enabled until the RT checkbox is also enabled.

diff --git a/Engine/Source/Runtime/Engine/Classes/Engine/RendererSettings.h b/Engine/Source/Runtime/Engine/Classes/Engine/RendererSettings.h
--- a/Engine/Source/Runtime/Engine/Classes/Engine/RendererSettings.h
+++ b/Engine/Source/Runtime/Engine/Classes/Engine/RendererSettings.h
@@ -394,7 +394,7 @@ class ENGINE_API URendererSettings : public UDeveloperSettings
 
 	UPROPERTY(config, EditAnywhere, Category = HardwareRayTracing, meta = (
 		ConsoleVariable = "r.RayTracing.Shadows", DisplayName = "Ray Tracing Shadows",
-		ToolTip = "Enables hardware ray-traced shadows."))
+		ToolTip = "Enables hardware ray-traced shadows. Requires Hardware Ray Tracing to be enabled."))
 		uint32 bEnableRayTracingShadows : 1;
 
 	UPROPERTY(config, EditAnywhere, Category = HardwareRayTracing, meta = (
@@ -405,7 +405,7 @@ class ENGINE_API URendererSettings : public UDeveloperSettings
 
 	UPROPERTY(config, EditAnywhere, Category = HardwareRayTracing, meta = (
 		ConsoleVariable = "r.PathTracing", DisplayName = "Path Tracing",
-		ToolTip = "Support Path Tracing renderer. This enables additional material permutations.",
+		ToolTip = "Enables the Path Tracing renderer. This enables additional material permutations. Requires Hardware Ray Tracing to be enabled.",
 		ConfigRestartRequired = true))
 		uint32 bEnablePathTracing : 1;
 
diff --git a/Engine/Source/Runtime/Engine/Private/RendererSettings.cpp b/Engine/Source/Runtime/Engine/Private/RendererSettings.cpp
--- a/Engine/Source/Runtime/Engine/Private/RendererSettings.cpp
+++ b/Engine/Source/Runtime/Engine/Private/RendererSettings.cpp
@@ -74,6 +74,7 @@ URendererSettings::URendererSettings(const FObjectInitializer& ObjectInitializer
 	bEnableRayTracing = 0;
 	bUseHardwareRayTracingForLumen = 0;
 	bEnableRayTracingShadows = 0;
+	bEnablePathTracing = 0;
 	bEnableRayTracingTextureLOD = 0;
 	MaxSkinBones = FGPUBaseSkinVertexFactory::GHardwareMaxGPUSkinBones;
 }
@@ -250,20 +251,31 @@ void URendererSettings::PostEditChangeProperty(FPropertyChangedEvent& PropertyCh
 
 bool URendererSettings::CanEditChange(const FProperty* InProperty) const
 {
-	const bool ParentVal = Super::CanEditChange(InProperty);
+	if (!Super::CanEditChange(InProperty))
+	{
+		return false;
+	}
 
 	if ((InProperty->GetFName() == GET_MEMBER_NAME_CHECKED(URendererSettings, bSupportSkinCacheShaders)))
 	{
 		//only allow DISABLE of skincache shaders if raytracing is also disabled as skincache is a dependency of raytracing.
-		return ParentVal && (!bSupportSkinCacheShaders || !bEnableRayTracing);
+		return !bSupportSkinCacheShaders || !bEnableRayTracing;
+	}
+
+	// the following settings can only be edited if ray tracing is enabled
+	if ((InProperty->GetFName() == GET_MEMBER_NAME_CHECKED(URendererSettings, bEnablePathTracing)) ||
+		(InProperty->GetFName() == GET_MEMBER_NAME_CHECKED(URendererSettings, bEnableRayTracingShadows)) ||
+		(InProperty->GetFName() == GET_MEMBER_NAME_CHECKED(URendererSettings, bEnableRayTracingTextureLOD)))
+	{
+		return bEnableRayTracing;
 	}
 
 	if ((InProperty->GetFName() == GET_MEMBER_NAME_CHECKED(URendererSettings, bSupportSkyAtmosphereAffectsHeightFog)))
 	{
-		return ParentVal && bSupportSkyAtmosphere;
+		return bSupportSkyAtmosphere;
 	}
 
-	return ParentVal;
+	return true;
 }
 #endif // #if WITH_EDITOR
 
