Work around for no unity error caused by compiler crash.
#rnx

diff --git a/Engine/Source/Runtime/Engine/Private/Particles/ParticleModules_Velocity.cpp b/Engine/Source/Runtime/Engine/Private/Particles/ParticleModules_Velocity.cpp
--- a/Engine/Source/Runtime/Engine/Private/Particles/ParticleModules_Velocity.cpp
+++ b/Engine/Source/Runtime/Engine/Private/Particles/ParticleModules_Velocity.cpp
@@ -430,7 +430,11 @@ void UParticleModuleVelocityCone::SpawnEx(FParticleEmitterInstance* Owner, int32
 		const FVector DefaultSpawnDirection = DefaultDirectionRotation.TransformVector(DefaultDirection);
 
 		// Orientate the cone along the direction vector		
-		const FVector ForwardDirection = (Direction != FVector::ZeroVector)? Direction.GetSafeNormal(): DefaultDirection;
+		FVector ForwardDirection = DefaultDirection;
+		if (Direction != FVector::ZeroVector)
+		{
+			ForwardDirection = Direction.GetSafeNormal();
+		}
 		FVector UpDirection(0.0f, 0.0f, 1.0f);
 		FVector RightDirection(1.0f, 0.0f, 0.0f);
 
