Add GeneratedCppIncludes in PCH as it is one of the most included files from vcperf profile runs
Add missing CoreMinimal.h when compiling core
  - 251s -> 245s for QAGame Win64 Development with -NoXGE on AMD TR 3970X

diff --git a/Engine/Source/Editor/UnrealEd/Private/UnrealEdPrivatePCH.h b/Engine/Source/Editor/UnrealEd/Private/UnrealEdPrivatePCH.h
--- a/Engine/Source/Editor/UnrealEd/Private/UnrealEdPrivatePCH.h
+++ b/Engine/Source/Editor/UnrealEd/Private/UnrealEdPrivatePCH.h
@@ -272,6 +272,7 @@
 #include "UObject/UObjectMarks.h"
 #include "UObject/Class.h"
 #include "UObject/GarbageCollection.h"
+#include "UObject/GeneratedCppIncludes.h"
 #include "Serialization/ArchiveUObject.h"
 #include "UObject/CoreNative.h"
 #include "UObject/WeakObjectPtr.h"
diff --git a/Engine/Source/Editor/UnrealEd/Public/UnrealEdSharedPCH.h b/Engine/Source/Editor/UnrealEd/Public/UnrealEdSharedPCH.h
--- a/Engine/Source/Editor/UnrealEd/Public/UnrealEdSharedPCH.h
+++ b/Engine/Source/Editor/UnrealEd/Public/UnrealEdSharedPCH.h
@@ -263,6 +263,7 @@
 #include "Serialization/ArchiveUObject.h"
 #include "UObject/Class.h"
 #include "UObject/GarbageCollection.h"
+#include "UObject/GeneratedCppIncludes.h"
 #include "UObject/CoreNative.h"
 #include "UObject/WeakObjectPtr.h"
 #include "UObject/Interface.h"
diff --git a/Engine/Source/Runtime/Core/Private/CorePrivatePCH.h b/Engine/Source/Runtime/Core/Private/CorePrivatePCH.h
--- a/Engine/Source/Runtime/Core/Private/CorePrivatePCH.h
+++ b/Engine/Source/Runtime/Core/Private/CorePrivatePCH.h
@@ -159,6 +159,7 @@
 #include "Math/Quat.h"
 #include "Internationalization/Culture.h"
 #include "Math/Transform.h"
+#include "CoreMinimal.h"
 #include "Math/IntRect.h"
 #include "Misc/MemStack.h"
 #include "Containers/LockFreeFixedSizeAllocator.h"
diff --git a/Engine/Source/Runtime/CoreUObject/Private/CoreUObjectPrivatePCH.h b/Engine/Source/Runtime/CoreUObject/Private/CoreUObjectPrivatePCH.h
--- a/Engine/Source/Runtime/CoreUObject/Private/CoreUObjectPrivatePCH.h
+++ b/Engine/Source/Runtime/CoreUObject/Private/CoreUObjectPrivatePCH.h
@@ -224,6 +224,7 @@
 #include "Serialization/ArchiveUObject.h"
 #include "UObject/Class.h"
 #include "UObject/GarbageCollection.h"
+#include "UObject/GeneratedCppIncludes.h"
 #include "UObject/CoreNative.h"
 #include "Templates/Casts.h"
 #include "UObject/WeakObjectPtr.h"
diff --git a/Engine/Source/Runtime/CoreUObject/Public/CoreUObjectSharedPCH.h b/Engine/Source/Runtime/CoreUObject/Public/CoreUObjectSharedPCH.h
--- a/Engine/Source/Runtime/CoreUObject/Public/CoreUObjectSharedPCH.h
+++ b/Engine/Source/Runtime/CoreUObject/Public/CoreUObjectSharedPCH.h
@@ -225,6 +225,7 @@
 #include "UObject/UObjectBaseUtility.h"
 #include "UObject/UObjectMarks.h"
 #include "UObject/GarbageCollection.h"
+#include "UObject/GeneratedCppIncludes.h"
 #include "Serialization/ArchiveUObject.h"
 #include "UObject/Class.h"
 #include "UObject/CoreNative.h"
diff --git a/Engine/Source/Runtime/Engine/Private/EnginePrivatePCH.h b/Engine/Source/Runtime/Engine/Private/EnginePrivatePCH.h
--- a/Engine/Source/Runtime/Engine/Private/EnginePrivatePCH.h
+++ b/Engine/Source/Runtime/Engine/Private/EnginePrivatePCH.h
@@ -255,6 +255,7 @@
 #include "UObject/UObjectMarks.h"
 #include "UObject/Class.h"
 #include "UObject/GarbageCollection.h"
+#include "UObject/GeneratedCppIncludes.h"
 #include "Serialization/ArchiveUObject.h"
 #include "UObject/CoreNative.h"
 #include "UObject/WeakObjectPtr.h"
@@ -445,6 +446,7 @@
 #include "GameFramework/Volume.h"
 #include "EngineGlobals.h"
 #include "Audio.h"
+#include "AudioDevice.h"
 #include "Sound/SoundAttenuation.h"
 #include "IAudioExtensionPlugin.h"
 #include "Engine/MeshMerging.h"
diff --git a/Engine/Source/Runtime/Engine/Public/EngineSharedPCH.h b/Engine/Source/Runtime/Engine/Public/EngineSharedPCH.h
--- a/Engine/Source/Runtime/Engine/Public/EngineSharedPCH.h
+++ b/Engine/Source/Runtime/Engine/Public/EngineSharedPCH.h
@@ -283,6 +283,7 @@
 #include "UObject/UObjectMarks.h"
 #include "UObject/Class.h"
 #include "UObject/GarbageCollection.h"
+#include "UObject/GeneratedCppIncludes.h"
 #include "Serialization/ArchiveUObject.h"
 #include "UObject/CoreNative.h"
 #include "UObject/WeakObjectPtr.h"
