Added s.AsyncLoadingTimeLimitPIEMultiplier to tweak the current async loading time limit in PIE, which resets when PIE ends.

diff --git a/Engine/Source/Editor/UnrealEd/Private/EditorEngine.cpp b/Engine/Source/Editor/UnrealEd/Private/EditorEngine.cpp
--- a/Engine/Source/Editor/UnrealEd/Private/EditorEngine.cpp
+++ b/Engine/Source/Editor/UnrealEd/Private/EditorEngine.cpp
@@ -1395,8 +1395,15 @@ void UEditorEngine::Tick( float DeltaSeconds, bool bIdleMode )
 
 	// Update subsystems.
 	{
-		// This assumes that UObject::StaticTick only calls ProcessAsyncLoading.	
-		StaticTick(DeltaSeconds, !!GAsyncLoadingUseFullTimeLimit, GAsyncLoadingTimeLimit / 1000.f);
+		// This assumes that UObject::StaticTick only calls ProcessAsyncLoading.
+		float AsyncLoadingTimeLimit = GAsyncLoadingTimeLimit;
+
+		if (PlayWorld)
+		{
+			AsyncLoadingTimeLimit *= GAsyncLoadingTimeLimitPIEMultiplier;
+		}
+
+		StaticTick(DeltaSeconds, !!GAsyncLoadingUseFullTimeLimit, AsyncLoadingTimeLimit / 1000.f);
 	}
 
 	FEngineAnalytics::Tick(DeltaSeconds);
diff --git a/Engine/Source/Editor/UnrealEd/Private/PlayLevel.cpp b/Engine/Source/Editor/UnrealEd/Private/PlayLevel.cpp
--- a/Engine/Source/Editor/UnrealEd/Private/PlayLevel.cpp
+++ b/Engine/Source/Editor/UnrealEd/Private/PlayLevel.cpp
@@ -35,6 +35,7 @@
 #include "Async/TaskGraphInterfaces.h"
 #include "GameFramework/Actor.h"
 #include "Engine/Blueprint.h"
+#include "Engine/CoreSettings.h"
 #include "Engine/GameViewportClient.h"
 #include "Engine/GameInstance.h"
 #include "Engine/RendererSettings.h"
@@ -567,6 +568,9 @@ void UEditorEngine::EndPlayMap()
 		bIsSimulatingInEditor = false;
 		PRAGMA_ENABLE_DEPRECATION_WARNINGS
 	}
+
+	// reset PIE multiplier to none
+	GAsyncLoadingTimeLimitPIEMultiplier = 1.0f;
 }
 
 void UEditorEngine::CleanupPIEOnlineSessions(TArray<FName> OnlineIdentifiers)
diff --git a/Engine/Source/Runtime/Engine/Classes/Engine/CoreSettings.h b/Engine/Source/Runtime/Engine/Classes/Engine/CoreSettings.h
--- a/Engine/Source/Runtime/Engine/Classes/Engine/CoreSettings.h
+++ b/Engine/Source/Runtime/Engine/Classes/Engine/CoreSettings.h
@@ -122,6 +122,8 @@ protected:
 extern ENGINE_API int32 GUseBackgroundLevelStreaming;
 /** Maximum amount of time to spend doing asynchronous loading (ms per frame). */
 extern ENGINE_API float GAsyncLoadingTimeLimit;
+/** Multiplifer for the maximum amount of time to spend doing asynchronous loading when in PIE. */
+extern ENGINE_API float GAsyncLoadingTimeLimitPIEMultiplier;
 /** Whether to use the entire time limit even if blocked on I/O. */
 extern ENGINE_API int32 GAsyncLoadingUseFullTimeLimit;
 /** Additional time to spend asynchronous loading during a high priority load. */
diff --git a/Engine/Source/Runtime/Engine/Private/CoreSettings.cpp b/Engine/Source/Runtime/Engine/Private/CoreSettings.cpp
--- a/Engine/Source/Runtime/Engine/Private/CoreSettings.cpp
+++ b/Engine/Source/Runtime/Engine/Private/CoreSettings.cpp
@@ -12,6 +12,7 @@ DEFINE_LOG_CATEGORY_STATIC(LogCoreSettings, Log, All);
 
 int32 GUseBackgroundLevelStreaming = 1;
 float GAsyncLoadingTimeLimit = 5.0f;
+float GAsyncLoadingTimeLimitPIEMultiplier = 1.0f;
 int32 GAsyncLoadingUseFullTimeLimit = 0;
 float GPriorityAsyncLoadingExtraTime = 15.0f;
 float GLevelStreamingActorsUpdateTimeLimit = 5.0f;
@@ -40,6 +41,13 @@ static FAutoConsoleVariableRef CVarAsyncLoadingTimeLimit(
 	ECVF_Default
 	);
 
+static FAutoConsoleVariableRef CVarAsyncLoadingTimeLimitPIEMultiplier(
+	TEXT("s.AsyncLoadingTimeLimitPIEMultiplier"),
+	GAsyncLoadingTimeLimitPIEMultiplier,
+	TEXT("Multiplifer for the maximum amount of time to spend doing asynchronous loading when in PIE."),
+	ECVF_Default
+	);
+
 static FAutoConsoleVariableRef CVarAsyncLoadingUseFullTimeLimit(
 	TEXT("s.AsyncLoadingUseFullTimeLimit"),
 	GAsyncLoadingUseFullTimeLimit,
