PushSelectionToProxy should not MarkRenderStateDirty. SceneProxy->SetSelection_GameThread routes the selection information to the render thread itself and all components tested responded properly without needing their render states recreated. Recreating render states can be quite expensive if done on a large amount of components at once.

diff --git a/Engine/Source/Runtime/Engine/Private/Components/PrimitiveComponent.cpp b/Engine/Source/Runtime/Engine/Private/Components/PrimitiveComponent.cpp
--- a/Engine/Source/Runtime/Engine/Private/Components/PrimitiveComponent.cpp
+++ b/Engine/Source/Runtime/Engine/Private/Components/PrimitiveComponent.cpp
@@ -1537,7 +1537,6 @@ void UPrimitiveComponent::PushSelectionToProxy()
 	if (SceneProxy)
 	{
 		SceneProxy->SetSelection_GameThread(ShouldRenderSelected(), IsComponentIndividuallySelected());
-		MarkRenderStateDirty();
 	}
 }
 
