Put MarkRenderStateDirty back in UPrimitiveComponent::PushSelectionToProxy for now, need to fix up Scene->PrimitivesSelected to not break selection skipping occlusion

diff --git a/Engine/Source/Runtime/Engine/Private/Components/PrimitiveComponent.cpp b/Engine/Source/Runtime/Engine/Private/Components/PrimitiveComponent.cpp
--- a/Engine/Source/Runtime/Engine/Private/Components/PrimitiveComponent.cpp
+++ b/Engine/Source/Runtime/Engine/Private/Components/PrimitiveComponent.cpp
@@ -1537,6 +1537,7 @@ void UPrimitiveComponent::PushSelectionToProxy()
 	if (SceneProxy)
 	{
 		SceneProxy->SetSelection_GameThread(ShouldRenderSelected(), IsComponentIndividuallySelected());
+		MarkRenderStateDirty();
 	}
 }
 
