Add support for for-loops to new material HLSL generator
- Built with same infinite-loop primitive as while-loops
- Improved/tuned support for local variables, used internally when generating for-loops
- Added ability for material expressions to register custom data when generating HLSL, will allow user-defined expressions to communicate between outputs (like for-loops need to)

diff --git a/Engine/Source/Runtime/Engine/Classes/Materials/MaterialExpressionAppendVector.h b/Engine/Source/Runtime/Engine/Classes/Materials/MaterialExpressionAppendVector.h
--- a/Engine/Source/Runtime/Engine/Classes/Materials/MaterialExpressionAppendVector.h
+++ b/Engine/Source/Runtime/Engine/Classes/Materials/MaterialExpressionAppendVector.h
@@ -25,6 +25,7 @@ class UMaterialExpressionAppendVector : public UMaterialExpression
 #if WITH_EDITOR
 	virtual int32 Compile(class FMaterialCompiler* Compiler, int32 OutputIndex) override;
 	virtual void GetCaption(TArray<FString>& OutCaptions) const override;
+	virtual EMaterialGenerateHLSLStatus GenerateHLSLExpression(FMaterialHLSLGenerator& Generator, UE::HLSLTree::FScope& Scope, int32 OutputIndex, UE::HLSLTree::FExpression*& OutExpression) override;
 #endif
 	//~ End UMaterialExpression Interface
 
diff --git a/Engine/Source/Runtime/Engine/Classes/Materials/MaterialExpressionMultiply.h b/Engine/Source/Runtime/Engine/Classes/Materials/MaterialExpressionMultiply.h
--- a/Engine/Source/Runtime/Engine/Classes/Materials/MaterialExpressionMultiply.h
+++ b/Engine/Source/Runtime/Engine/Classes/Materials/MaterialExpressionMultiply.h
@@ -34,6 +34,7 @@ class UMaterialExpressionMultiply : public UMaterialExpression
 	virtual int32 Compile(class FMaterialCompiler* Compiler, int32 OutputIndex) override;
 	virtual void GetCaption(TArray<FString>& OutCaptions) const override;
 	virtual FText GetKeywords() const override {return FText::FromString(TEXT("*"));}
+	virtual EMaterialGenerateHLSLStatus GenerateHLSLExpression(FMaterialHLSLGenerator& Generator, UE::HLSLTree::FScope& Scope, int32 OutputIndex, UE::HLSLTree::FExpression*& OutExpression) override;
 #endif // WITH_EDITOR
 	//~ End UMaterialExpression Interface
 };
diff --git a/Engine/Source/Runtime/Engine/Classes/Materials/MaterialExpressionTextureCoordinate.h b/Engine/Source/Runtime/Engine/Classes/Materials/MaterialExpressionTextureCoordinate.h
--- a/Engine/Source/Runtime/Engine/Classes/Materials/MaterialExpressionTextureCoordinate.h
+++ b/Engine/Source/Runtime/Engine/Classes/Materials/MaterialExpressionTextureCoordinate.h
@@ -40,6 +40,7 @@ class UMaterialExpressionTextureCoordinate : public UMaterialExpression
 	virtual void GetCaption(TArray<FString>& OutCaptions) const override;
 	virtual FText GetKeywords() const override { return FText::FromString(TEXT("texcoord")); }
 	virtual bool MatchesSearchQuery(const TCHAR* SearchQuery) override;
+	virtual EMaterialGenerateHLSLStatus GenerateHLSLExpression(FMaterialHLSLGenerator& Generator, UE::HLSLTree::FScope& Scope, int32 OutputIndex, UE::HLSLTree::FExpression*& OutExpression) override;
 #endif
 	//~ End UMaterialExpression Interface
 };
diff --git a/Engine/Source/Runtime/Engine/Private/HLSLTree/HLSLTree.cpp b/Engine/Source/Runtime/Engine/Private/HLSLTree/HLSLTree.cpp
--- a/Engine/Source/Runtime/Engine/Private/HLSLTree/HLSLTree.cpp
+++ b/Engine/Source/Runtime/Engine/Private/HLSLTree/HLSLTree.cpp
@@ -193,22 +193,6 @@ UE::HLSLTree::FEmitScope* UE::HLSLTree::FEmitContext::AcquireScope(const FScope&
 	return EmitScope;
 }
 
-void UE::HLSLTree::FEmitContext::FinalizeScope(FEmitScope& EmitScope)
-{
-	check(!EmitScope.bFinalized);
-
-	FEmitAssignment::TConstIterator It(EmitScope.FirstAssignment);
-	while (It)
-	{
-		const FEmitAssignment& EmitAssignment = *It;
-		const FEmitValue* Value = AcquireValue(EmitAssignment.Expression);
-		WriteStatementToScopef(EmitScope, TEXT("%s = %s;"), EmitAssignment.Declaration, GetCode(Value));
-		It.Next();
-	}
-
-	EmitScope.bFinalized = true;
-}
-
 UE::HLSLTree::FEmitScope& UE::HLSLTree::FEmitContext::GetCurrentScope()
 {
 	return *ScopeStack.Last();
@@ -245,7 +229,19 @@ const UE::HLSLTree::FEmitValue* UE::HLSLTree::FEmitContext::AcquireValue(FExpres
 		FLocalHLSLCodeWriter LocalWriter;
 		Shader::FPreshaderData LocalPreshader;
 		FExpressionEmitResult EmitResult(LocalWriter, LocalPreshader);
-		if (Expression->EmitCode(*this, EmitResult))
+
+		bool bEmitResult = false;
+		{
+			bool bAlreadyPending = false;
+			const FSetElementId Id = PendingEmitValueExpressions.Add(Expression, &bAlreadyPending);
+			if (!bAlreadyPending)
+			{
+				bEmitResult = Expression->EmitCode(*this, EmitResult);
+				PendingEmitValueExpressions.Remove(Id);
+			}
+		}
+
+		if (bEmitResult)
 		{
 			check(EmitResult.EvaluationType != EExpressionEvaluationType::None);
 			check(EmitResult.Type != Shader::EValueType::Void);
@@ -495,8 +491,6 @@ bool UE::HLSLTree::FEmitContext::InternalWriteScope(const FScope& Scope, FString
 		return false;
 	}
 
-	FinalizeScope(*EmitScope);
-
 	FEmitScopeLink* Link = InternalWriteScopeLink(InternedCode);
 	Link->NextScope = EmitScope;
 	return true;
@@ -530,9 +524,9 @@ void UE::HLSLTree::FEmitContext::WriteDeclaration(FEmitScope& EmitScope, Shader:
 
 bool UE::HLSLTree::FEmitContext::WriteAssignment(FEmitScope& EmitScope, const TCHAR* Declaration, FExpression* Expression, Shader::EValueType& InOutType)
 {
-	if (EmitScope.bFinalized)
+	const FEmitValue* Value = AcquireValue(Expression);
+	if (Value)
 	{
-		const FEmitValue* Value = AcquireValue(Expression);
 		const Shader::EValueType ValueType = Value->GetExpressionType();
 		if (InOutType == Shader::EValueType::Void)
 		{
@@ -542,7 +536,6 @@ bool UE::HLSLTree::FEmitContext::WriteAssignment(FEmitScope& EmitScope, const TC
 		{
 			return false;
 		}
-
 		WriteStatementToScopef(EmitScope, TEXT("%s = %s;"), Declaration, GetCode(Value));
 	}
 	else
@@ -556,8 +549,40 @@ bool UE::HLSLTree::FEmitContext::WriteAssignment(FEmitScope& EmitScope, const TC
 	return true;
 }
 
-void UE::HLSLTree::FEmitContext::Finalize()
+bool UE::HLSLTree::FEmitContext::FinalizeScope(FEmitScope& EmitScope)
 {
+	if (EmitScope.FirstAssignment)
+	{
+		// Scope has pending assigmnets, flush all of them, then restart finalization
+		// (Since flushing pending assignments may generate additional pending assignments)
+		FEmitAssignment::TConstIterator It(EmitScope.FirstAssignment);
+		while (It)
+		{
+			const FEmitAssignment& EmitAssignment = *It;
+			const FEmitValue* Value = AcquireValue(EmitAssignment.Expression);
+			WriteStatementToScopef(EmitScope, TEXT("%s = %s;"), EmitAssignment.Declaration, GetCode(Value));
+			It.Next();
+		}
+		EmitScope.FirstAssignment = nullptr;
+		return false;
+	}
+
+	UE::HLSLTree::FEmitScopeLink::TConstIterator It(EmitScope.FirstLink);
+	while (It)
+	{
+		const UE::HLSLTree::FEmitScopeLink& EmitLink = *It;
+		if (EmitLink.NextScope)
+		{
+			if (!FinalizeScope(*EmitLink.NextScope))
+			{
+				return false;
+			}
+		}
+
+		It.Next();
+	}
+
+	return true;
 }
 
 UE::HLSLTree::FScope* UE::HLSLTree::FScope::FindSharedParent(FScope* Lhs, FScope* Rhs)
@@ -719,6 +744,20 @@ bool UE::HLSLTree::FScope::EmitHLSL(FEmitContext& Context, FEmitScope& Scope) co
 	return bResult;
 }
 
+bool UE::HLSLTree::FScope::HasParentScope(const FScope& InParentScope) const
+{
+	const FScope* CurrentScope = this;
+	while (CurrentScope)
+	{
+		if (CurrentScope == &InParentScope)
+		{
+			return true;
+		}
+		CurrentScope = CurrentScope->ParentScope;
+	}
+	return false;
+}
+
 void UE::HLSLTree::FScope::AddPreviousScope(FScope& Scope)
 {
 	check(NumPreviousScopes < MaxNumPreviousScopes);
@@ -836,8 +875,6 @@ static void WriteIndent(int32 IndentLevel, FStringBuilderBase& InOutString)
 
 static void WriteScope(const UE::HLSLTree::FEmitScope& EmitScope, int32 IndentLevel, FStringBuilderBase& InOutString)
 {
-	check(EmitScope.bFinalized);
-
 	{
 		UE::HLSLTree::FEmitDeclaration::TConstIterator It(EmitScope.FirstDeclaration);
 		while (It)
@@ -900,8 +937,13 @@ bool UE::HLSLTree::FTree::EmitHLSL(UE::HLSLTree::FEmitContext& Context, FCodeWri
 	FEmitScope* EmitRootScope = Context.AcquireScope(*RootScope);
 	if (RootScope->EmitHLSL(Context, *EmitRootScope))
 	{
-		Context.FinalizeScope(*EmitRootScope);
-		Context.Finalize();
+		// Need to continue iterating FinalizeScope until it succeeds
+		bool bFinalizeResult = false;
+		while (!bFinalizeResult)
+		{
+			bFinalizeResult = Context.FinalizeScope(*EmitRootScope);
+		}
+
 		WriteScope(*EmitRootScope, 0, *Writer.StringBuilder);
 		return true;
 	}
diff --git a/Engine/Source/Runtime/Engine/Private/HLSLTree/HLSLTreeCommon.cpp b/Engine/Source/Runtime/Engine/Private/HLSLTree/HLSLTreeCommon.cpp
--- a/Engine/Source/Runtime/Engine/Private/HLSLTree/HLSLTreeCommon.cpp
+++ b/Engine/Source/Runtime/Engine/Private/HLSLTree/HLSLTreeCommon.cpp
@@ -364,6 +364,7 @@ bool UE::HLSLTree::FExpressionBinaryOp::EmitCode(FEmitContext& Context, FExpress
 		switch (Op)
 		{
 		case UE::HLSLTree::EBinaryOp::Add: Opcode = Shader::EPreshaderOpcode::Add; break;
+		case UE::HLSLTree::EBinaryOp::Mul: Opcode = Shader::EPreshaderOpcode::Mul; break;
 		default: break;
 		}
 	}
@@ -377,6 +378,7 @@ bool UE::HLSLTree::FExpressionBinaryOp::EmitCode(FEmitContext& Context, FExpress
 		switch (Op)
 		{
 		case EBinaryOp::Add: OutResult.Writer.Writef(TEXT("(%s + %s)"), LhsCode, RhsCode); break;
+		case EBinaryOp::Mul: OutResult.Writer.Writef(TEXT("(%s * %s)"), LhsCode, RhsCode); break;
 		case EBinaryOp::Less: OutResult.Writer.Writef(TEXT("(%s < %s)"), LhsCode, RhsCode); break;
 		default: checkNoEntry(); break;
 		}
@@ -428,6 +430,44 @@ bool UE::HLSLTree::FExpressionSwizzle::EmitCode(FEmitContext& Context, FExpressi
 	return true;
 }
 
+bool UE::HLSLTree::FExpressionAppend::EmitCode(FEmitContext& Context, FExpressionEmitResult& OutResult) const
+{
+	const FEmitValue* LhsValue = Context.AcquireValue(Lhs);
+	const FEmitValue* RhsValue = Context.AcquireValue(Rhs);
+	if (!LhsValue || !RhsValue)
+	{
+		return false;
+	}
+
+	const Shader::FValueTypeDescription LhsTypeDesc = Shader::GetValueTypeDescription(LhsValue->GetExpressionType());
+	const Shader::FValueTypeDescription RhsTypeDesc = Shader::GetValueTypeDescription(RhsValue->GetExpressionType());
+	if (LhsTypeDesc.ComponentType != RhsTypeDesc.ComponentType)
+	{
+		return false;
+	}
+
+	const int32 NumComponents = LhsTypeDesc.NumComponents + RhsTypeDesc.NumComponents;
+	if (NumComponents > 4)
+	{
+		return false;
+	}
+
+	OutResult.EvaluationType = CombineEvaluationTypes(LhsValue->GetEvaluationType(), RhsValue->GetEvaluationType());
+	OutResult.Type = Shader::MakeValueType(LhsTypeDesc.ComponentType, NumComponents);
+	if (OutResult.EvaluationType == EExpressionEvaluationType::Shader)
+	{
+		const Shader::FValueTypeDescription ResultTypeDesc = Shader::GetValueTypeDescription(OutResult.Type);
+		OutResult.Writer.Writef(TEXT("%s(%s, %s)"), ResultTypeDesc.Name, Context.GetCode(LhsValue), Context.GetCode(RhsValue));
+	}
+	else
+	{
+		Context.AppendPreshader(LhsValue, OutResult.Preshader);
+		Context.AppendPreshader(RhsValue, OutResult.Preshader);
+		OutResult.Preshader.WriteOpcode(Shader::EPreshaderOpcode::AppendVector);
+	}
+	return true;
+}
+
 bool UE::HLSLTree::FExpressionCast::EmitCode(FEmitContext& Context, FExpressionEmitResult& OutResult) const
 {
 	const FEmitValue* InputValue = Context.AcquireValue(Input);
@@ -576,25 +616,3 @@ bool UE::HLSLTree::FStatementLoop::EmitHLSL(FEmitContext& Context) const
 	Context.WriteScopef(*LoopScope, TEXT("while (true)"));
 	return Context.WriteScope(*NextScope);
 }
-
-bool UE::HLSLTree::FStatementFor::EmitHLSL(FEmitContext& Context) const
-{
-	/*const FEmitValue* StartExpressionValue = Context.AcquireValue(StartExpression);
-	const FEmitValue* EndExpressionValue = Context.AcquireValue(StartExpression);
-	if (!!StartExpressionValue || !EndExpressionValue)
-	{
-		return false;
-	}
-
-	const TCHAR* DeclarationCode = TEXT("");
-
-	FEmitCode* LoopScopeCode = Context.AppendHLSLf(Scope, TEXT("for (%s = %s; %s < %s; %s++"),
-		DeclarationCode,
-		Context.GetCode(StartExpressionValue),
-		DeclarationCode,
-		Context.GetCode(EndExpressionValue),
-		DeclarationCode);
-	LoopScope->EmitHLSL(Context, *LoopScopeCode);
-	return NextScope->EmitHLSL(Context, Scope);*/
-	return false;
-}
diff --git a/Engine/Source/Runtime/Engine/Private/Materials/MaterialExpressionHLSL.cpp b/Engine/Source/Runtime/Engine/Private/Materials/MaterialExpressionHLSL.cpp
--- a/Engine/Source/Runtime/Engine/Private/Materials/MaterialExpressionHLSL.cpp
+++ b/Engine/Source/Runtime/Engine/Private/Materials/MaterialExpressionHLSL.cpp
@@ -19,6 +19,7 @@
 #include "Materials/MaterialExpressionVectorParameter.h"
 #include "Materials/MaterialExpressionScalarParameter.h"
 #include "Materials/MaterialExpressionStaticBoolParameter.h"
+#include "Materials/MaterialExpressionTextureCoordinate.h"
 #include "Materials/MaterialExpressionTextureSample.h"
 #include "Materials/MaterialExpressionTextureSampleParameter.h"
 #include "Materials/MaterialExpressionTextureObject.h"
@@ -28,7 +29,9 @@
 #include "Materials/MaterialExpressionStaticSwitch.h"
 #include "Materials/MaterialExpressionGetLocal.h"
 #include "Materials/MaterialExpressionAdd.h"
+#include "Materials/MaterialExpressionMultiply.h"
 #include "Materials/MaterialExpressionBinaryOp.h"
+#include "Materials/MaterialExpressionAppendVector.h"
 #include "Materials/MaterialExpressionSetMaterialAttributes.h"
 #include "Materials/MaterialExpressionReflectionVectorWS.h"
 #include "Materials/MaterialExpressionSetLocal.h"
@@ -64,37 +67,37 @@ FMaterialHLSLTree& UMaterialFunctionInterface::AcquireHLSLTree(FMaterialHLSLGene
 
 EMaterialGenerateHLSLStatus UMaterialExpressionConstant::GenerateHLSLExpression(FMaterialHLSLGenerator& Generator, UE::HLSLTree::FScope& Scope, int32 OutputIndex, UE::HLSLTree::FExpression*& OutExpression)
 {
-	OutExpression = Generator.NewConstant(Scope, R);
+	OutExpression = Generator.NewConstant(R);
 	return EMaterialGenerateHLSLStatus::Success;
 }
 
 EMaterialGenerateHLSLStatus UMaterialExpressionConstant2Vector::GenerateHLSLExpression(FMaterialHLSLGenerator& Generator, UE::HLSLTree::FScope& Scope, int32 OutputIndex, UE::HLSLTree::FExpression*& OutExpression)
 {
-	OutExpression = Generator.NewConstant(Scope, UE::Shader::FValue(R, G));
+	OutExpression = Generator.NewConstant(UE::Shader::FValue(R, G));
 	return EMaterialGenerateHLSLStatus::Success;
 }
 
 EMaterialGenerateHLSLStatus UMaterialExpressionConstant3Vector::GenerateHLSLExpression(FMaterialHLSLGenerator& Generator, UE::HLSLTree::FScope& Scope, int32 OutputIndex, UE::HLSLTree::FExpression*& OutExpression)
 {
-	OutExpression = Generator.NewConstant(Scope, UE::Shader::FValue(Constant.R, Constant.G, Constant.B));
+	OutExpression = Generator.NewConstant(UE::Shader::FValue(Constant.R, Constant.G, Constant.B));
 	return EMaterialGenerateHLSLStatus::Success;
 }
 
 EMaterialGenerateHLSLStatus UMaterialExpressionConstant4Vector::GenerateHLSLExpression(FMaterialHLSLGenerator& Generator, UE::HLSLTree::FScope& Scope, int32 OutputIndex, UE::HLSLTree::FExpression*& OutExpression)
 {
-	OutExpression = Generator.NewConstant(Scope, UE::Shader::FValue(Constant.R, Constant.G, Constant.B, Constant.A));
+	OutExpression = Generator.NewConstant(UE::Shader::FValue(Constant.R, Constant.G, Constant.B, Constant.A));
 	return EMaterialGenerateHLSLStatus::Success;
 }
 
 EMaterialGenerateHLSLStatus UMaterialExpressionStaticBool::GenerateHLSLExpression(FMaterialHLSLGenerator& Generator, UE::HLSLTree::FScope& Scope, int32 OutputIndex, UE::HLSLTree::FExpression*& OutExpression)
 {
-	OutExpression = Generator.NewConstant(Scope, (bool)Value);
+	OutExpression = Generator.NewConstant((bool)Value);
 	return EMaterialGenerateHLSLStatus::Success;
 }
 
 EMaterialGenerateHLSLStatus UMaterialExpressionStaticSwitch::GenerateHLSLExpression(FMaterialHLSLGenerator& Generator, UE::HLSLTree::FScope& Scope, int32 OutputIndex, UE::HLSLTree::FExpression*& OutExpression)
 {
-	UE::HLSLTree::FExpression* ConditionExpression = Value.GetTracedInput().Expression ? Value.AcquireHLSLExpression(Generator, Scope) : Generator.NewConstant(Scope, (bool)DefaultValue);
+	UE::HLSLTree::FExpression* ConditionExpression = Value.GetTracedInput().Expression ? Value.AcquireHLSLExpression(Generator, Scope) : Generator.NewConstant((bool)DefaultValue);
 	UE::HLSLTree::FExpression* TrueExpression = A.AcquireHLSLExpression(Generator, Scope);
 	UE::HLSLTree::FExpression* FalseExpression = B.AcquireHLSLExpression(Generator, Scope);
 
@@ -133,6 +136,13 @@ EMaterialGenerateHLSLStatus UMaterialExpressionStaticBoolParameter::GenerateHLSL
 	return EMaterialGenerateHLSLStatus::Success;
 }
 
+EMaterialGenerateHLSLStatus UMaterialExpressionTextureCoordinate::GenerateHLSLExpression(FMaterialHLSLGenerator& Generator, UE::HLSLTree::FScope& Scope, int32 OutputIndex, UE::HLSLTree::FExpression*& OutExpression)
+{
+	// TODO - tiling, mirroring
+	OutExpression = Generator.NewTexCoord(Scope, CoordinateIndex);
+	return EMaterialGenerateHLSLStatus::Success;
+}
+
 EMaterialGenerateHLSLStatus UMaterialExpressionTextureObject::GenerateHLSLTexture(FMaterialHLSLGenerator& Generator, UE::HLSLTree::FScope& Scope, int32 OutputIndex, UE::HLSLTree::FTextureParameterDeclaration*& OutTexture)
 {
 	const UE::HLSLTree::FTextureDescription TextureDesc(Texture, SamplerType);
@@ -196,8 +206,8 @@ EMaterialGenerateHLSLStatus UMaterialExpressionTextureSampleParameter::GenerateH
 
 EMaterialGenerateHLSLStatus UMaterialExpressionBinaryOp::GenerateHLSLExpression(FMaterialHLSLGenerator& Generator, UE::HLSLTree::FScope& Scope, int32 OutputIndex, UE::HLSLTree::FExpression*& OutExpression)
 {
-	UE::HLSLTree::FExpression* Lhs = A.GetTracedInput().Expression ? A.AcquireHLSLExpression(Generator, Scope) : Generator.NewConstant(Scope, ConstA);
-	UE::HLSLTree::FExpression* Rhs = B.GetTracedInput().Expression ? B.AcquireHLSLExpression(Generator, Scope) : Generator.NewConstant(Scope, ConstB);
+	UE::HLSLTree::FExpression* Lhs = A.GetTracedInput().Expression ? A.AcquireHLSLExpression(Generator, Scope) : Generator.NewConstant(ConstA);
+	UE::HLSLTree::FExpression* Rhs = B.GetTracedInput().Expression ? B.AcquireHLSLExpression(Generator, Scope) : Generator.NewConstant(ConstB);
 	if (!Lhs || !Rhs)
 	{
 		return EMaterialGenerateHLSLStatus::Error;
@@ -209,21 +219,37 @@ EMaterialGenerateHLSLStatus UMaterialExpressionBinaryOp::GenerateHLSLExpression(
 
 EMaterialGenerateHLSLStatus UMaterialExpressionAdd::GenerateHLSLExpression(FMaterialHLSLGenerator& Generator, UE::HLSLTree::FScope& Scope, int32 OutputIndex, UE::HLSLTree::FExpression*& OutExpression)
 {
-	UE::HLSLTree::FExpression* Lhs = A.GetTracedInput().Expression ? A.AcquireHLSLExpression(Generator, Scope) : Generator.NewConstant(Scope, ConstA);
-	UE::HLSLTree::FExpression* Rhs = B.GetTracedInput().Expression ? B.AcquireHLSLExpression(Generator, Scope) : Generator.NewConstant(Scope, ConstB);
+	UE::HLSLTree::FExpression* Lhs = A.GetTracedInput().Expression ? A.AcquireHLSLExpression(Generator, Scope) : Generator.NewConstant(ConstA);
+	UE::HLSLTree::FExpression* Rhs = B.GetTracedInput().Expression ? B.AcquireHLSLExpression(Generator, Scope) : Generator.NewConstant(ConstB);
 	if (!Lhs || !Rhs)
 	{
 		return EMaterialGenerateHLSLStatus::Error;
 	}
+	OutExpression = Generator.GetTree().NewExpression<UE::HLSLTree::FExpressionBinaryOp>(Scope, UE::HLSLTree::EBinaryOp::Add, Lhs, Rhs);
+	return EMaterialGenerateHLSLStatus::Success;
+}
 
-	/*FString ErrorMessage;
-	const UE::Shader::EValueType ResultType = UE::HLSLTree::MakeArithmeticResultType(Lhs->Type, Rhs->Type, ErrorMessage);
-	if (ResultType == UE::Shader::EValueType::Void)
+EMaterialGenerateHLSLStatus UMaterialExpressionMultiply::GenerateHLSLExpression(FMaterialHLSLGenerator& Generator, UE::HLSLTree::FScope& Scope, int32 OutputIndex, UE::HLSLTree::FExpression*& OutExpression)
+{
+	UE::HLSLTree::FExpression* Lhs = A.GetTracedInput().Expression ? A.AcquireHLSLExpression(Generator, Scope) : Generator.NewConstant(ConstA);
+	UE::HLSLTree::FExpression* Rhs = B.GetTracedInput().Expression ? B.AcquireHLSLExpression(Generator, Scope) : Generator.NewConstant(ConstB);
+	if (!Lhs || !Rhs)
 	{
-		return Generator.Errorf(TEXT("%s"), *ErrorMessage);
-	}*/
+		return EMaterialGenerateHLSLStatus::Error;
+	}
+	OutExpression = Generator.GetTree().NewExpression<UE::HLSLTree::FExpressionBinaryOp>(Scope, UE::HLSLTree::EBinaryOp::Mul, Lhs, Rhs);
+	return EMaterialGenerateHLSLStatus::Success;
+}
 
-	OutExpression = Generator.GetTree().NewExpression<UE::HLSLTree::FExpressionBinaryOp>(Scope, UE::HLSLTree::EBinaryOp::Add, Lhs, Rhs);
+EMaterialGenerateHLSLStatus UMaterialExpressionAppendVector::GenerateHLSLExpression(FMaterialHLSLGenerator& Generator, UE::HLSLTree::FScope& Scope, int32 OutputIndex, UE::HLSLTree::FExpression*& OutExpression)
+{
+	UE::HLSLTree::FExpression* Lhs = A.AcquireHLSLExpression(Generator, Scope);
+	UE::HLSLTree::FExpression* Rhs = B.AcquireHLSLExpression(Generator, Scope);
+	if (!Lhs || !Rhs)
+	{
+		return EMaterialGenerateHLSLStatus::Error;
+	}
+	OutExpression = Generator.GetTree().NewExpression<UE::HLSLTree::FExpressionAppend>(Scope, Lhs, Rhs);
 	return EMaterialGenerateHLSLStatus::Success;
 }
 
@@ -379,14 +405,33 @@ EMaterialGenerateHLSLStatus UMaterialExpressionWhileLoop::GenerateHLSLStatements
 	return EMaterialGenerateHLSLStatus::Success;
 }
 
+struct FGlobalExpressionDataForLoop
+{
+	int32 NumLoops = 0;
+};
+DECLARE_MATERIAL_HLSLGENERATOR_DATA(FGlobalExpressionDataForLoop);
+
+struct FExpressionDataForLoop
+{
+	UE::HLSLTree::FScope* LoopScope = nullptr;
+	FName LocalName;
+};
+DECLARE_MATERIAL_HLSLGENERATOR_DATA(FExpressionDataForLoop);
+
 EMaterialGenerateHLSLStatus UMaterialExpressionForLoop::GenerateHLSLExpression(FMaterialHLSLGenerator& Generator, UE::HLSLTree::FScope& Scope, int32 OutputIndex, UE::HLSLTree::FExpression*& OutExpression)
 {
+	FExpressionDataForLoop* ExpressionData = Generator.FindExpressionData<FExpressionDataForLoop>(this);
+	if (!ExpressionData || !Scope.HasParentScope(*ExpressionData->LoopScope))
+	{
+		return Generator.Error(TEXT("For loop index accessed outside loop scope"));
+	}
+
+	OutExpression = Generator.AcquireLocalValue(Scope, ExpressionData->LocalName);
 	return EMaterialGenerateHLSLStatus::Success;
 }
 
 EMaterialGenerateHLSLStatus UMaterialExpressionForLoop::GenerateHLSLStatements(FMaterialHLSLGenerator& Generator, UE::HLSLTree::FScope& Scope)
 {
-#if 0
 	if (!LoopBody.GetExpression())
 	{
 		return Generator.Error(TEXT("Missing LoopBody connection"));
@@ -404,16 +449,37 @@ EMaterialGenerateHLSLStatus UMaterialExpressionForLoop::GenerateHLSLStatements(F
 		return Generator.Error(TEXT("Missing EndIndex connection"));
 	}
 
-	UE::HLSLTree::FScope* LoopScope = Generator.GetTree().NewScope(Scope);
-	LoopBody.GenerateHLSLStatements(Generator, *LoopScope);
+	FGlobalExpressionDataForLoop* GlobalData = Generator.AcquireGlobalData<FGlobalExpressionDataForLoop>();
+	FExpressionDataForLoop* ExpressionData = Generator.NewExpressionData<FExpressionDataForLoop>(this);
+	ExpressionData->LocalName = *FString::Printf(TEXT("ForLoopControl%d"), GlobalData->NumLoops++);
+
+	UE::HLSLTree::FExpression* StepExpression = IndexStep.GetTracedInput().Expression ? IndexStep.AcquireHLSLExpressionWithCast(Generator, Scope, UE::Shader::EValueType::Int1) : Generator.NewConstant(1);
+
+	Generator.GenerateAssignLocal(Scope, ExpressionData->LocalName, StartExpression);
+
+	UE::HLSLTree::FStatementLoop* LoopStatement = Generator.GetTree().NewStatement<UE::HLSLTree::FStatementLoop>(Scope);
+	LoopStatement->LoopScope = Generator.NewScope(Scope);
+	ExpressionData->LoopScope = LoopStatement->LoopScope;
+
+	UE::HLSLTree::FStatementIf* IfStatement = Generator.GetTree().NewStatement<UE::HLSLTree::FStatementIf>(*LoopStatement->LoopScope);
+	IfStatement->ThenScope = Generator.NewScope(*LoopStatement->LoopScope);
+	IfStatement->ElseScope = Generator.NewScope(*LoopStatement->LoopScope);
+	LoopStatement->NextScope = Generator.NewScope(Scope, EMaterialNewScopeFlag::NoPreviousScope);
+	LoopStatement->LoopScope->AddPreviousScope(*IfStatement->ThenScope);
+	LoopStatement->NextScope->AddPreviousScope(*IfStatement->ElseScope);
+
+	Generator.GetTree().NewStatement<UE::HLSLTree::FStatementBreak>(*IfStatement->ElseScope);
 
-	UE::HLSLTree::FStatementFor* ForStatement = Generator.GetTree().NewStatement<UE::HLSLTree::FStatementFor>(Scope);
-	ForStatement->StartExpression = StartExpression;
-	ForStatement->EndExpression = EndExpression;
-	ForStatement->LoopScope = LoopScope;
+	UE::HLSLTree::FExpression* LocalExpression = Generator.AcquireLocalValue(*LoopStatement->LoopScope, ExpressionData->LocalName);
+
+	IfStatement->ConditionExpression = Generator.GetTree().NewExpression<UE::HLSLTree::FExpressionBinaryOp>(*LoopStatement->LoopScope, UE::HLSLTree::EBinaryOp::Less, LocalExpression, EndExpression);
+	LoopBody.GenerateHLSLStatements(Generator, *IfStatement->ThenScope);
+
+	UE::HLSLTree::FExpression* NewLocalExpression = Generator.GetTree().NewExpression<UE::HLSLTree::FExpressionBinaryOp>(*IfStatement->ThenScope, UE::HLSLTree::EBinaryOp::Add, Generator.AcquireLocalValue(*IfStatement->ThenScope, ExpressionData->LocalName), StepExpression);
+	Generator.GenerateAssignLocal(*IfStatement->ThenScope, ExpressionData->LocalName, NewLocalExpression);
+
+	Completed.GenerateHLSLStatements(Generator, *LoopStatement->NextScope);
 
-	Completed.GenerateHLSLStatements(Generator, Scope);
-#endif
 	return EMaterialGenerateHLSLStatus::Success;
 }
 
diff --git a/Engine/Source/Runtime/Engine/Private/Materials/MaterialHLSLGenerator.cpp b/Engine/Source/Runtime/Engine/Private/Materials/MaterialHLSLGenerator.cpp
--- a/Engine/Source/Runtime/Engine/Private/Materials/MaterialHLSLGenerator.cpp
+++ b/Engine/Source/Runtime/Engine/Private/Materials/MaterialHLSLGenerator.cpp
@@ -66,8 +66,10 @@ bool FMaterialHLSLGenerator::Finalize()
 	}
 
 	// Resolve values for any PHI nodes that were generated
-	for (UE::HLSLTree::FExpressionLocalPHI* Expression : PHIExpressions)
+	// Resolving a PHI may produce additional PHIs
+	while(PHIExpressions.Num() > 0)
 	{
+		UE::HLSLTree::FExpressionLocalPHI* Expression = PHIExpressions.Pop(false);
 		for (int32 i = 0; i < Expression->NumValues; ++i)
 		{
 			Expression->Values[i] = InternalAcquireLocalValue(*Expression->Scopes[i], Expression->LocalName);
@@ -134,7 +136,7 @@ static UE::HLSLTree::FExpression* CompileMaterialInput(FMaterialHLSLGenerator& G
 				DefaultValue.NumComponents = UE::Shader::GetValueTypeDescription(InputDescription.Type).NumComponents;
 				if (InputDescription.ConstantValue != DefaultValue)
 				{
-					Expression = Generator.NewConstant(Scope, InputDescription.ConstantValue);
+					Expression = Generator.NewConstant(InputDescription.ConstantValue);
 				}
 			}
 			else
@@ -237,9 +239,9 @@ UE::HLSLTree::FScope* FMaterialHLSLGenerator::NewJoinedScope(UE::HLSLTree::FScop
 	return NewScope;
 }
 
-UE::HLSLTree::FExpressionConstant* FMaterialHLSLGenerator::NewConstant(UE::HLSLTree::FScope& Scope, const UE::Shader::FValue& Value)
+UE::HLSLTree::FExpressionConstant* FMaterialHLSLGenerator::NewConstant(const UE::Shader::FValue& Value)
 {
-	UE::HLSLTree::FExpressionConstant* Expression = HLSLTree->NewExpression<UE::HLSLTree::FExpressionConstant>(Scope, Value);
+	UE::HLSLTree::FExpressionConstant* Expression = HLSLTree->NewExpression<UE::HLSLTree::FExpressionConstant>(HLSLTree->GetRootScope(), Value);
 	return Expression;
 }
 
@@ -470,4 +472,18 @@ bool FMaterialHLSLGenerator::GenerateStatements(UE::HLSLTree::FScope& Scope, UMa
 	return bResult;
 }
 
+void FMaterialHLSLGenerator::InternalRegisterExpressionData(const FName& Type, const UMaterialExpression* MaterialExpression, void* Data)
+{
+	const FExpressionDataKey Key(Type, MaterialExpression);
+	check(!ExpressionDataMap.Contains(Key));
+	ExpressionDataMap.Add(Key, Data);
+}
+
+void* FMaterialHLSLGenerator::InternalFindExpressionData(const FName& Type, const UMaterialExpression* MaterialExpression)
+{
+	const FExpressionDataKey Key(Type, MaterialExpression);
+	void** Result = ExpressionDataMap.Find(Key);
+	return Result ? *Result : nullptr;
+}
+
 #endif // WITH_EDITOR
diff --git a/Engine/Source/Runtime/Engine/Private/Shader/ShaderValue.cpp b/Engine/Source/Runtime/Engine/Private/Shader/ShaderValue.cpp
--- a/Engine/Source/Runtime/Engine/Private/Shader/ShaderValue.cpp
+++ b/Engine/Source/Runtime/Engine/Private/Shader/ShaderValue.cpp
@@ -207,14 +207,34 @@ UE::Shader::EValueType UE::Shader::MakeArithmeticResultType(EValueType Lhs, EVal
 		{
 			return Lhs;
 		}
+
+		EValueComponentType ComponentType = EValueComponentType::Void;
 		if (LhsDesc.ComponentType == RhsDesc.ComponentType)
+		{
+			ComponentType = LhsDesc.ComponentType;
+		}
+		else if (LhsDesc.ComponentType == EValueComponentType::Float || RhsDesc.ComponentType == EValueComponentType::Float)
+		{
+			ComponentType = EValueComponentType::Float;
+		}
+		else
+		{
+			ComponentType = EValueComponentType::Int;
+		}
+
+		if (ComponentType != EValueComponentType::Void)
 		{
 			if (LhsDesc.NumComponents == 1 || RhsDesc.NumComponents == 1)
 			{
 				// single component type is valid to combine with other type
-				return MakeValueType(LhsDesc.ComponentType, FMath::Max(LhsDesc.NumComponents, RhsDesc.NumComponents));
+				return MakeValueType(ComponentType, FMath::Max(LhsDesc.NumComponents, RhsDesc.NumComponents));
+			}
+			else if (LhsDesc.NumComponents == RhsDesc.NumComponents)
+			{
+				return MakeValueType(ComponentType, LhsDesc.NumComponents);
 			}
 		}
+
 		OutErrorMessage = FString::Printf(TEXT("Arithmetic between types %s and %s are undefined"), LhsDesc.Name, RhsDesc.Name);
 	}
 	else
diff --git a/Engine/Source/Runtime/Engine/Public/HLSLTree/HLSLTree.h b/Engine/Source/Runtime/Engine/Public/HLSLTree/HLSLTree.h
--- a/Engine/Source/Runtime/Engine/Public/HLSLTree/HLSLTree.h
+++ b/Engine/Source/Runtime/Engine/Public/HLSLTree/HLSLTree.h
@@ -135,7 +135,6 @@ public:
 	FEmitScopeLink* LastLink = nullptr;
 	const FScope* SourceScope = nullptr;
 	TMap<FSHAHash, const TCHAR*> ExpressionMap;
-	bool bFinalized = false;
 };
 
 class FEmitValue
@@ -189,7 +188,6 @@ public:
 
 	FEmitScope* FindScope(const FScope& Scope);
 	FEmitScope* AcquireScope(const FScope& Scope);
-	void FinalizeScope(FEmitScope& EmitScope);
 
 	FEmitScope& GetCurrentScope();
 
@@ -242,7 +240,7 @@ public:
 
 	bool WriteAssignment(FEmitScope& EmitScope, const TCHAR* Declaration, FExpression* Expression, Shader::EValueType& InOutType);
 
-	void Finalize();
+	bool FinalizeScope(FEmitScope& EmitScope);
 
 	struct FDeclarationEntry
 	{
@@ -268,6 +266,7 @@ public:
 
 	TArray<FFunctionStackEntry> FunctionStack;
 	TArray<Shader::FPreshaderData*> TempPreshaders;
+	TSet<const FExpression*> PendingEmitValueExpressions;
 	FMemStackBase* Allocator = nullptr;
 	const FMaterial* Material = nullptr; // TODO - remove preshader material dependency
 	const FStaticParameterSet* StaticParameters = nullptr;
@@ -444,6 +443,8 @@ public:
 		return MakeArrayView(const_cast<FScope*>(this)->PreviousScope, NumPreviousScopes);
 	}
 
+	bool HasParentScope(const FScope& ParentScope) const;
+
 	virtual ENodeVisitResult Visit(FNodeVisitor& Visitor) override;
 	bool EmitHLSL(FEmitContext& Context, FEmitScope& Scope) const;
 
@@ -497,7 +498,6 @@ public:
 	}
 
 	FScope* NewScope(FScope& Scope);
-
 	FParameterDeclaration* NewParameterDeclaration(FScope& Scope, const FName& Name, const Shader::FValue& DefaultValue);
 	FTextureParameterDeclaration* NewTextureParameterDeclaration(FScope& Scope, const FName& Name, const FTextureDescription& DefaultValue);
 
diff --git a/Engine/Source/Runtime/Engine/Public/HLSLTree/HLSLTreeCommon.h b/Engine/Source/Runtime/Engine/Public/HLSLTree/HLSLTreeCommon.h
--- a/Engine/Source/Runtime/Engine/Public/HLSLTree/HLSLTreeCommon.h
+++ b/Engine/Source/Runtime/Engine/Public/HLSLTree/HLSLTreeCommon.h
@@ -49,7 +49,7 @@ public:
 class FExpressionParameter : public FExpression
 {
 public:
-	FExpressionParameter(FParameterDeclaration* InDeclaration)
+	explicit FExpressionParameter(FParameterDeclaration* InDeclaration)
 		: Declaration(InDeclaration)
 	{}
 
@@ -248,6 +248,31 @@ public:
 	virtual bool EmitCode(FEmitContext& Context, FExpressionEmitResult& OutResult) const override;
 };
 
+class FExpressionAppend : public FExpression
+{
+public:
+	FExpressionAppend(FExpression* InLhs, FExpression* InRhs)
+		: Lhs(InLhs)
+		, Rhs(InRhs)
+	{}
+
+	FExpression* Lhs;
+	FExpression* Rhs;
+
+	virtual ENodeVisitResult Visit(FNodeVisitor& Visitor) override
+	{
+		const ENodeVisitResult Result = FExpression::Visit(Visitor);
+		if (ShouldVisitDependentNodes(Result))
+		{
+			Visitor.VisitNode(Lhs);
+			Visitor.VisitNode(Rhs);
+		}
+		return Result;
+	}
+
+	virtual bool EmitCode(FEmitContext& Context, FExpressionEmitResult& OutResult) const override;
+};
+
 class FExpressionCast : public FExpression
 {
 public:
@@ -406,29 +431,5 @@ public:
 	virtual bool EmitHLSL(FEmitContext& Context) const override;
 };
 
-class FStatementFor : public FStatement
-{
-public:
-	FExpression* StartExpression;
-	FExpression* EndExpression;
-	FScope* LoopScope;
-	FScope* NextScope;
-
-	virtual ENodeVisitResult Visit(FNodeVisitor& Visitor) override
-	{
-		const ENodeVisitResult Result = FStatement::Visit(Visitor);
-		if (ShouldVisitDependentNodes(Result))
-		{
-			Visitor.VisitNode(StartExpression);
-			Visitor.VisitNode(EndExpression);
-			Visitor.VisitNode(LoopScope);
-			Visitor.VisitNode(NextScope);
-		}
-		return Result;
-	}
-
-	virtual bool EmitHLSL(FEmitContext& Context) const override;
-};
-
 }
 }
diff --git a/Engine/Source/Runtime/Engine/Public/MaterialHLSLGenerator.h b/Engine/Source/Runtime/Engine/Public/MaterialHLSLGenerator.h
--- a/Engine/Source/Runtime/Engine/Public/MaterialHLSLGenerator.h
+++ b/Engine/Source/Runtime/Engine/Public/MaterialHLSLGenerator.h
@@ -45,6 +45,12 @@ enum class EMaterialNewScopeFlag : uint8
 };
 ENUM_CLASS_FLAGS(EMaterialNewScopeFlag);
 
+template<typename T>
+struct TMaterialHLSLGeneratorType;
+
+#define DECLARE_MATERIAL_HLSLGENERATOR_DATA(T) \
+	template<> struct TMaterialHLSLGeneratorType<T> { static const FName& GetTypeName() { static const FName Name(TEXT(#T)); return Name; } }
+
 /**
  * MaterialHLSLGenerator is a bridge between a material, and HLSLTree.  It facilitates generating HLSL source code for a given material, using HLSLTree
  */
@@ -79,7 +85,7 @@ public:
 
 	UE::HLSLTree::FScope* NewJoinedScope(UE::HLSLTree::FScope& Scope);
 
-	UE::HLSLTree::FExpressionConstant* NewConstant(UE::HLSLTree::FScope& Scope, const UE::Shader::FValue& Value);
+	UE::HLSLTree::FExpressionConstant* NewConstant(const UE::Shader::FValue& Value);
 	UE::HLSLTree::FExpressionExternalInput* NewTexCoord(UE::HLSLTree::FScope& Scope, int32 Index);
 	UE::HLSLTree::FExpressionSwizzle* NewSwizzle(UE::HLSLTree::FScope& Scope, const UE::HLSLTree::FSwizzleParameters& Params, UE::HLSLTree::FExpression* Input);
 	UE::HLSLTree::FExpression* NewCast(UE::HLSLTree::FScope& Scope, UE::Shader::EValueType Type, UE::HLSLTree::FExpression* Input, UE::HLSLTree::ECastFlags Flags = UE::HLSLTree::ECastFlags::None);
@@ -108,6 +114,32 @@ public:
 	UE::HLSLTree::FTextureParameterDeclaration* AcquireTextureDeclaration(UE::HLSLTree::FScope& Scope, UMaterialExpression* MaterialExpression, int32 OutputIndex);
 	bool GenerateStatements(UE::HLSLTree::FScope& Scope, UMaterialExpression* MaterialExpression);
 
+	template<typename T, typename... ArgTypes>
+	T* NewExpressionData(const UMaterialExpression* MaterialExpression, ArgTypes... Args)
+	{
+		T* Data = new T(Forward<ArgTypes>(Args)...);
+		InternalRegisterExpressionData(TMaterialHLSLGeneratorType<T>::GetTypeName(), MaterialExpression, Data);
+		return Data;
+	}
+
+	template<typename T>
+	T* FindExpressionData(const UMaterialExpression* MaterialExpression)
+	{
+		return static_cast<T*>(InternalFindExpressionData(TMaterialHLSLGeneratorType<T>::GetTypeName(), MaterialExpression));
+	}
+
+	template<typename T>
+	T* AcquireGlobalData()
+	{
+		T* Data = static_cast<T*>(InternalFindExpressionData(TMaterialHLSLGeneratorType<T>::GetTypeName(), nullptr));
+		if (!Data)
+		{
+			Data = new T();
+			InternalRegisterExpressionData(TMaterialHLSLGeneratorType<T>::GetTypeName(), nullptr, Data);
+		}
+		return Data;
+	}
+
 private:
 	static constexpr int32 MaxNumPreviousScopes = UE::HLSLTree::MaxNumPreviousScopes;
 	struct FExpressionKey
@@ -162,6 +194,24 @@ private:
 		}
 	};
 
+	struct FExpressionDataKey
+	{
+		FExpressionDataKey(const FName& InTypeName, const UMaterialExpression* InMaterialExpression) : MaterialExpression(InMaterialExpression), TypeName(InTypeName) {}
+
+		const UMaterialExpression* MaterialExpression;
+		FName TypeName;
+
+		friend inline uint32 GetTypeHash(const FExpressionDataKey& Value)
+		{
+			return HashCombine(GetTypeHash(Value.MaterialExpression), GetTypeHash(Value.TypeName));
+		}
+
+		friend inline bool operator==(const FExpressionDataKey& Lhs, const FExpressionDataKey& Rhs)
+		{
+			return Lhs.MaterialExpression == Rhs.MaterialExpression && Lhs.TypeName == Rhs.TypeName;
+		}
+	};
+
 	struct FStatementEntry
 	{
 		UE::HLSLTree::FScope* PreviousScope[MaxNumPreviousScopes];
@@ -170,6 +220,9 @@ private:
 
 	UE::HLSLTree::FExpression* InternalAcquireLocalValue(UE::HLSLTree::FScope& Scope, const FName& LocalName);
 
+	void InternalRegisterExpressionData(const FName& Type, const UMaterialExpression* MaterialExpression, void* Data);
+	void* InternalFindExpressionData(const FName& Type, const UMaterialExpression* MaterialExpression);
+
 	const FMaterialCompileTargetParameters& CompileTarget;
 	UMaterial* TargetMaterial;
 	UMaterialFunctionInterface* TargetMaterialFunction;
@@ -187,6 +240,7 @@ private:
 	TMap<FLocalKey, UE::HLSLTree::FExpression*> LocalMap;
 	TMap<FExpressionKey, UE::HLSLTree::FExpression*> ExpressionMap;
 	TMap<UMaterialExpression*, FStatementEntry> StatementMap;
+	TMap<FExpressionDataKey, void*> ExpressionDataMap;
 	bool bGeneratedResult;
 };
 
