Missing file from previous commit

diff --git a/Engine/Source/Runtime/Engine/Public/Shader/ShaderTypes.h b/Engine/Source/Runtime/Engine/Public/Shader/ShaderTypes.h
--- a/Engine/Source/Runtime/Engine/Public/Shader/ShaderTypes.h
+++ b/Engine/Source/Runtime/Engine/Public/Shader/ShaderTypes.h
@@ -151,6 +151,11 @@ struct FValue
 		Component[0].Bool = v;
 	}
 
+	inline FValue(int32 v) : ComponentType(EValueComponentType::Int), NumComponents(1)
+	{
+		Component[0].Int = v;
+	}
+
 	inline EValueType GetType() const { return MakeValueType(ComponentType, NumComponents); }
 
 	inline const FValueComponent& GetComponent(int32 i) const { check(i >= 0 && i < NumComponents); return Component[i]; }
