Guarding against introduction of more than 1 LevelScriptActor in a level

diff --git a/Engine/Source/Runtime/Engine/Private/LevelScriptActor.cpp b/Engine/Source/Runtime/Engine/Private/LevelScriptActor.cpp
--- a/Engine/Source/Runtime/Engine/Private/LevelScriptActor.cpp
+++ b/Engine/Source/Runtime/Engine/Private/LevelScriptActor.cpp
@@ -29,6 +29,27 @@ ALevelScriptActor::ALevelScriptActor(const FObjectInitializer& ObjectInitializer
 	bReplicates = true;
 	bAlwaysRelevant = true;
 	bReplayRewindable = true;
+
+#if WITH_EDITOR
+	// this is intended as an early detection of cases where more than one LevelScriptActor is introduced into a single map
+	const UObject* ThisOuter = GetOuter();
+	if (!Cast<UPackage>(ThisOuter))
+	{
+		TArray<UObject*> AllSiblingObjects;
+		GetObjectsWithOuter(ThisOuter, AllSiblingObjects, false, RF_NoFlags, EInternalObjectFlags::PendingKill);
+
+		for (const UObject* Sibling : AllSiblingObjects)
+		{
+			bool bIsNotAnLSA           = !Cast<ALevelScriptActor>(Sibling);
+			bool bIsTheSameObject      = Sibling == this;
+			bool bHasNewerClassVersion = Sibling->GetClass()->HasAnyClassFlags(CLASS_NewerVersionExists);
+
+			ensureMsgf((bIsNotAnLSA || bIsTheSameObject || bHasNewerClassVersion),
+				TEXT("Detected the creation of more than one LevelScriptActor (%s, %s) within the same outer (%s). This can lead to duplicate level blueprint operations during play."),
+				*GetName(), *Sibling->GetName(), *ThisOuter->GetName());
+		}
+	}
+#endif // WITH_EDITOR
 }
 
 void ALevelScriptActor::PreInitializeComponents()
