- Allow Volumes to create static textures from PF_FloatRGBA

diff --git a/Engine/Source/Runtime/Engine/Private/TextureRenderTargetVolume.cpp b/Engine/Source/Runtime/Engine/Private/TextureRenderTargetVolume.cpp
--- a/Engine/Source/Runtime/Engine/Private/TextureRenderTargetVolume.cpp
+++ b/Engine/Source/Runtime/Engine/Private/TextureRenderTargetVolume.cpp
@@ -151,6 +151,7 @@ UVolumeTexture* UTextureRenderTargetVolume::ConstructTextureVolume(UObject* ObjO
 	switch (PixelFormat)
 	{
 		case PF_R32_FLOAT:
+		case PF_FloatRGBA:
 			TextureFormat = TSF_RGBA16F;
 			break;
 	}
