WorldPartition: UWorldPartitionStreamingPolicy::GetMaxCellsToLoad not limited for Servers (fixes failed MP sessions because clients where making some levels visible before server had knowledge of them)

Repro originates from pr #8093

diff --git a/Engine/Source/Runtime/Engine/Private/WorldPartition/WorldPartitionStreamingPolicy.cpp b/Engine/Source/Runtime/Engine/Private/WorldPartition/WorldPartitionStreamingPolicy.cpp
--- a/Engine/Source/Runtime/Engine/Private/WorldPartition/WorldPartitionStreamingPolicy.cpp
+++ b/Engine/Source/Runtime/Engine/Private/WorldPartition/WorldPartitionStreamingPolicy.cpp
@@ -240,7 +240,14 @@ int32 UWorldPartitionStreamingPolicy::GetMaxCellsToLoad() const
 {
 	// This policy limits the number of concurrent loading streaming cells, except if match hasn't started
 	UWorld* World = WorldPartition->GetWorld();
-	return World->bMatchStarted ? (GMaxLoadingStreamingCells - GetCellLoadingCount()) : MAX_int32;
+	const ENetMode NetMode = World->GetNetMode();
+	if (NetMode == NM_Standalone || NetMode == NM_Client)
+	{
+		return World->bMatchStarted ? (GMaxLoadingStreamingCells - GetCellLoadingCount()) : MAX_int32;
+	}
+
+	// Always allow max on server to make sure StreamingLevels are added before clients update the visibility
+	return MAX_int32;
 }
 
 void UWorldPartitionStreamingPolicy::SetTargetStateForCells(EWorldPartitionRuntimeCellState TargetState, const TSet<const UWorldPartitionRuntimeCell*>& Cells)
