Leverage reverse lookups to reduce time spent in async mesh compilation during PIE game tick
  - 13.8ms -> 5.9us when PIEing on special project

#rnx

diff --git a/Engine/Source/Runtime/Engine/Private/StaticMeshCompiler.cpp b/Engine/Source/Runtime/Engine/Private/StaticMeshCompiler.cpp
--- a/Engine/Source/Runtime/Engine/Private/StaticMeshCompiler.cpp
+++ b/Engine/Source/Runtime/Engine/Private/StaticMeshCompiler.cpp
@@ -465,69 +465,83 @@ void FStaticMeshCompilingManager::FinishCompilationsForGame()
 			
 			TSet<UStaticMesh*> StaticMeshToCompile;
 			TArray<FBoxSphereBounds, TInlineAllocator<16>> ActorsBounds;
-			for (const UStaticMeshComponent* Component : ObjectCacheScope.GetContext().GetStaticMeshComponents())
+			for (TWeakObjectPtr<UStaticMesh>& StaticMeshPtr : RegisteredStaticMesh)
 			{
-				if (Component->IsRegistered() &&
-					PIEWorlds.Contains(Component->GetWorld()) &&
-					Component->GetStaticMesh() != nullptr &&
-					RegisteredStaticMesh.Contains(Component->GetStaticMesh()) &&
-					(PlayInEditorMode == 0 || Component->GetCollisionEnabled() != ECollisionEnabled::NoCollision || Component->IsNavigationRelevant() || Component->bAlwaysCreatePhysicsState || Component->CanCharacterStepUpOn != ECB_No))
+				if (UStaticMesh* StaticMesh = StaticMeshPtr.Get())
 				{
-					const FBoxSphereBounds ComponentBounds = Component->Bounds.GetBox();
-					const UWorld* ComponentWorld = Component->GetWorld();
-
-					if (PlayInEditorMode == 2)
+					for (const UStaticMeshComponent* Component : ObjectCacheScope.GetContext().GetStaticMeshComponents(StaticMesh))
 					{
-						ActorsBounds.Reset();
-						WorldActors.MultiFind(ComponentWorld, ActorsBounds);
-						
-						bool bStaticMeshComponentCollided = false;
-						if (ActorsBounds.Num())
+						if (Component->IsRegistered() &&
+							PIEWorlds.Contains(Component->GetWorld()) &&
+							(PlayInEditorMode == 0 || Component->GetCollisionEnabled() != ECollisionEnabled::NoCollision || Component->IsNavigationRelevant() || Component->bAlwaysCreatePhysicsState || Component->CanCharacterStepUpOn != ECB_No))
 						{
-							for (const FBoxSphereBounds& ActorBounds : ActorsBounds)
+							if (PlayInEditorMode == 2)
 							{
-								if (FMath::SphereAABBIntersection(ActorBounds.Origin, ActorBounds.SphereRadius * ActorBounds.SphereRadius, ComponentBounds.GetBox()))
+								const FBoxSphereBounds ComponentBounds = Component->Bounds.GetBox();
+								const UWorld* ComponentWorld = Component->GetWorld();
+
+								ActorsBounds.Reset();
+								WorldActors.MultiFind(ComponentWorld, ActorsBounds);
+						
+								bool bStaticMeshComponentCollided = false;
+								if (ActorsBounds.Num())
 								{
-									if (bShowDebugDraw)
+									for (const FBoxSphereBounds& ActorBounds : ActorsBounds)
 									{
-										DrawDebugBox(ComponentWorld, ComponentBounds.Origin, ComponentBounds.BoxExtent, FColor::Red, false, 10.0f);
-									}
+										if (FMath::SphereAABBIntersection(ActorBounds.Origin, ActorBounds.SphereRadius * ActorBounds.SphereRadius, ComponentBounds.GetBox()))
+										{
+											if (bShowDebugDraw)
+											{
+												DrawDebugBox(ComponentWorld, ComponentBounds.Origin, ComponentBounds.BoxExtent, FColor::Red, false, 10.0f);
+											}
 								
-									bool bIsAlreadyInSet = false;
-									StaticMeshToCompile.Add(Component->GetStaticMesh(), &bIsAlreadyInSet);
-									if (!bIsAlreadyInSet)
-									{
-										UE_LOG(
-											LogStaticMesh,
-											Display,
-											TEXT("Waiting on static mesh %s being ready because it affects collision/navigation and is near a player/bot"),
-											*Component->GetStaticMesh()->GetFullName()
-										);
+											bool bIsAlreadyInSet = false;
+											StaticMeshToCompile.Add(Component->GetStaticMesh(), &bIsAlreadyInSet);
+											if (!bIsAlreadyInSet)
+											{
+												UE_LOG(
+													LogStaticMesh,
+													Display,
+													TEXT("Waiting on static mesh %s being ready because it affects collision/navigation and is near a player/bot"),
+													*Component->GetStaticMesh()->GetFullName()
+												);
+											}
+											bStaticMeshComponentCollided = true;
+											break;
+										}
 									}
-									bStaticMeshComponentCollided = true;
-									break;
 								}
-							}
-						}
 
-						if (bShowDebugDraw && !bStaticMeshComponentCollided)
-						{
-							DrawDebugBox(ComponentWorld, ComponentBounds.Origin, ComponentBounds.BoxExtent, FColor::Green);
-						}
-					}
-					else 
-					{
-						bool bIsAlreadyInSet = false;
-						StaticMeshToCompile.Add(Component->GetStaticMesh(), &bIsAlreadyInSet);
-						if (!bIsAlreadyInSet)
-						{
-							if (PlayInEditorMode == 0)
-							{
-								UE_LOG(LogStaticMesh, Display, TEXT("Waiting on static mesh %s being ready before playing"), *Component->GetStaticMesh()->GetFullName());
+								if (bShowDebugDraw && !bStaticMeshComponentCollided)
+								{
+									DrawDebugBox(ComponentWorld, ComponentBounds.Origin, ComponentBounds.BoxExtent, FColor::Green);
+								}
+
+								// No need to iterate throught all components once we have found one that requires the static mesh to finish compilation 
+								// unless bShowDebugDraw is activated.
+								if (!bShowDebugDraw)
+								{
+									break;
+								}
 							}
-							else
+							else 
 							{
-								UE_LOG(LogStaticMesh, Display, TEXT("Waiting on static mesh %s being ready because it affects collision/navigation"), *Component->GetStaticMesh()->GetFullName());
+								bool bIsAlreadyInSet = false;
+								StaticMeshToCompile.Add(StaticMesh, &bIsAlreadyInSet);
+								if (!bIsAlreadyInSet)
+								{
+									if (PlayInEditorMode == 0)
+									{
+										UE_LOG(LogStaticMesh, Display, TEXT("Waiting on static mesh %s being ready before playing"), *StaticMesh->GetFullName());
+									}
+									else
+									{
+										UE_LOG(LogStaticMesh, Display, TEXT("Waiting on static mesh %s being ready because it affects collision/navigation"), *StaticMesh->GetFullName());
+									}
+								}
+
+								// No need to iterate throught all components once we have found one that requires the static mesh to finish compilation.
+								break;
 							}
 						}
 					}
