Reduce memory allocations in FPerInstanceRenderData::UpdateBoundsTransforms.

diff --git a/Engine/Source/Runtime/Engine/Private/InstancedStaticMesh.cpp b/Engine/Source/Runtime/Engine/Private/InstancedStaticMesh.cpp
--- a/Engine/Source/Runtime/Engine/Private/InstancedStaticMesh.cpp
+++ b/Engine/Source/Runtime/Engine/Private/InstancedStaticMesh.cpp
@@ -1024,13 +1024,13 @@ void FPerInstanceRenderData::UpdateBoundsTransforms_Concurrent()
 				UE_LOG(LogStaticMesh, Warning, TEXT("Unconsumed ISM bounds/transforms update task, we did more work than necessary"));
 			}
 
-	UpdateBoundsTask = FFunctionGraphTask::CreateAndDispatchWhenReady(
-		[this]()
-	{
-					UpdateBoundsTransforms();
-				},
-				TStatId(),
-				&Prerequisites
+			UpdateBoundsTask = FFunctionGraphTask::CreateAndDispatchWhenReady(
+			[this]()
+			{
+				UpdateBoundsTransforms();
+			},
+			TStatId(),
+			&Prerequisites
 			);
 		}
 	);
@@ -1039,15 +1039,14 @@ void FPerInstanceRenderData::UpdateBoundsTransforms_Concurrent()
 void FPerInstanceRenderData::UpdateBoundsTransforms()
 {
 		const int32 InstanceCount = InstanceBuffer.GetNumInstances();
-			FBoxSphereBounds LocalBounds;
-			if (bTrackBounds)
-			{
-				LocalBounds = FBoxSphereBounds(InstanceLocalBounds);
-		PerInstanceBounds.Empty();
-		PerInstanceBounds.Reserve(InstanceCount);
-			}
-		PerInstanceTransforms.Empty();
-		PerInstanceTransforms.Reserve(InstanceCount);
+		FBoxSphereBounds LocalBounds;
+		if (bTrackBounds)
+		{
+			LocalBounds = FBoxSphereBounds(InstanceLocalBounds);
+			PerInstanceBounds.Empty(InstanceCount);
+		}
+
+		PerInstanceTransforms.Empty(InstanceCount);
 
 		for (int InstanceIndex = 0; InstanceIndex < InstanceCount; InstanceIndex++)
 		{
@@ -1056,11 +1055,11 @@ void FPerInstanceRenderData::UpdateBoundsTransforms()
 			InstanceBuffer.GetInstanceTransform(InstanceIndex, InstTransform);
 			InstTransform.M[3][3] = 1.0f;
 
-				if (bTrackBounds)
-				{
-			FBoxSphereBounds TransformedBounds = LocalBounds.TransformBy(InstTransform);
-			PerInstanceBounds.Add(FVector4(TransformedBounds.Origin, TransformedBounds.SphereRadius));
-				}
+			if (bTrackBounds)
+			{
+				FBoxSphereBounds TransformedBounds = LocalBounds.TransformBy(InstTransform);
+				PerInstanceBounds.Add(FVector4(TransformedBounds.Origin, TransformedBounds.SphereRadius));
+			}
 			PerInstanceTransforms.Add(InstTransform);
 		}
 }
