UPrimitiveComponent::PushSelectionToProxy() should not call MarkRenderStateDirty() for perf reasons. Fix up FScene::PrimitivesSelected array that came in with 15937150

diff --git a/Engine/Source/Runtime/Engine/Private/Components/PrimitiveComponent.cpp b/Engine/Source/Runtime/Engine/Private/Components/PrimitiveComponent.cpp
--- a/Engine/Source/Runtime/Engine/Private/Components/PrimitiveComponent.cpp
+++ b/Engine/Source/Runtime/Engine/Private/Components/PrimitiveComponent.cpp
@@ -1537,7 +1537,6 @@ void UPrimitiveComponent::PushSelectionToProxy()
 	if (SceneProxy)
 	{
 		SceneProxy->SetSelection_GameThread(ShouldRenderSelected(), IsComponentIndividuallySelected());
-		MarkRenderStateDirty();
 	}
 }
 
diff --git a/Engine/Source/Runtime/Engine/Private/PrimitiveSceneProxy.cpp b/Engine/Source/Runtime/Engine/Private/PrimitiveSceneProxy.cpp
--- a/Engine/Source/Runtime/Engine/Private/PrimitiveSceneProxy.cpp
+++ b/Engine/Source/Runtime/Engine/Private/PrimitiveSceneProxy.cpp
@@ -530,8 +530,18 @@ void FPrimitiveSceneProxy::SetDebugMassData(const TArray<FDebugMassData>& InDebu
 void FPrimitiveSceneProxy::SetSelection_RenderThread(const bool bInParentSelected, const bool bInIndividuallySelected)
 {
 	check(IsInRenderingThread());
+
+	const bool bWasSelected = IsSelected();
 	bParentSelected = bInParentSelected;
 	bIndividuallySelected = bInIndividuallySelected;
+	const bool bIsSelected = IsSelected();
+
+	// The renderer may have cached the selected state, let it know that this primitive is updated
+	if ((bWasSelected && !bIsSelected) || 
+		(bIsSelected && !bWasSelected))
+	{
+		GetScene().UpdatePrimitiveSelectedState_RenderThread(GetPrimitiveSceneInfo(), bIsSelected);
+	}
 }
 
 /**
diff --git a/Engine/Source/Runtime/Engine/Public/SceneInterface.h b/Engine/Source/Runtime/Engine/Public/SceneInterface.h
--- a/Engine/Source/Runtime/Engine/Public/SceneInterface.h
+++ b/Engine/Source/Runtime/Engine/Public/SceneInterface.h
@@ -231,6 +231,9 @@ public:
 
 	/** Update render states that possibly cached inside renderer, like mesh draw commands. More lightweight than re-registering the scene proxy. */
 	virtual void UpdateCachedRenderStates(class FPrimitiveSceneProxy* SceneProxy) {}
+	
+	/** Updates the selected state values that might be cached inside the renderer */
+	virtual void UpdatePrimitiveSelectedState_RenderThread(const FPrimitiveSceneInfo* PrimitiveSceneInfo, bool bIsSelected) {}
 
 	/** 
 	 * Adds a new exponential height fog component to the scene
diff --git a/Engine/Source/Runtime/Renderer/Private/RendererScene.cpp b/Engine/Source/Runtime/Renderer/Private/RendererScene.cpp
--- a/Engine/Source/Runtime/Renderer/Private/RendererScene.cpp
+++ b/Engine/Source/Runtime/Renderer/Private/RendererScene.cpp
@@ -1422,6 +1422,19 @@ void FScene::UpdatePrimitiveTransform(UPrimitiveComponent* Primitive)
 	}
 }
 
+void FScene::UpdatePrimitiveSelectedState_RenderThread(const FPrimitiveSceneInfo* PrimitiveSceneInfo, bool bIsSelected)
+{
+	check(IsInRenderingThread());
+
+	if (PrimitiveSceneInfo)
+	{
+		if (PrimitiveSceneInfo->GetIndex() != INDEX_NONE)
+		{
+			PrimitivesSelected[PrimitiveSceneInfo->GetIndex()] = bIsSelected;
+		}
+	}
+}
+
 void FScene::UpdatePrimitiveLightingAttachmentRoot(UPrimitiveComponent* Primitive)
 {
 	const UPrimitiveComponent* NewLightingAttachmentRoot = Primitive->GetLightingAttachmentRoot();
diff --git a/Engine/Source/Runtime/Renderer/Private/ScenePrivate.h b/Engine/Source/Runtime/Renderer/Private/ScenePrivate.h
--- a/Engine/Source/Runtime/Renderer/Private/ScenePrivate.h
+++ b/Engine/Source/Runtime/Renderer/Private/ScenePrivate.h
@@ -3033,6 +3033,9 @@ public:
 	/** Update render states that possibly cached inside renderer, like mesh draw commands. More lightweight than re-registering the scene proxy. */
 	virtual void UpdateCachedRenderStates(FPrimitiveSceneProxy* SceneProxy) override;
 
+	/** Updates PrimitivesSelected array for this PrimitiveSceneInfo */
+	virtual void UpdatePrimitiveSelectedState_RenderThread(const FPrimitiveSceneInfo* PrimitiveSceneInfo, bool bIsSelected) override;
+
 	virtual void Release() override;
 	virtual UWorld* GetWorld() const override { return World; }
 
