Texture: Use build definitions instead of build actions for the build with the new DDC

#rnx

diff --git a/Engine/Source/Runtime/Engine/Private/TextureDerivedDataBuildExporter.cpp b/Engine/Source/Runtime/Engine/Private/TextureDerivedDataBuildExporter.cpp
--- a/Engine/Source/Runtime/Engine/Private/TextureDerivedDataBuildExporter.cpp
+++ b/Engine/Source/Runtime/Engine/Private/TextureDerivedDataBuildExporter.cpp
@@ -173,10 +173,10 @@ void FTextureDerivedDataBuildExporter::ExportTextureBuild(const UTexture& Textur
 			ActionBuilder.AddConstant(Key, Value);
 		});
 
-	ActionBuilder.AddInput(Texture.Source.GetIdString(), ExportedTextureBulkDataHash, ExportedTextureBulkDataSize);
+	ActionBuilder.AddInput(Texture.Source.GetId().ToString(), ExportedTextureBulkDataHash, ExportedTextureBulkDataSize);
 	if ((bool)Texture.CompositeTexture && !ExportedCompositeTextureBulkDataHash.IsZero())
 	{
-		ActionBuilder.AddInput(Texture.CompositeTexture->Source.GetIdString(), ExportedCompositeTextureBulkDataHash, ExportedCompositeTextureBulkDataSize);
+		ActionBuilder.AddInput(Texture.CompositeTexture->Source.GetId().ToString(), ExportedCompositeTextureBulkDataHash, ExportedCompositeTextureBulkDataSize);
 	}
 
 	if (TUniquePtr<FArchive> Ar{IFileManager::Get().CreateFileWriter(*(ExportRoot / TEXT("build.uddba")))})
diff --git a/Engine/Source/Runtime/Engine/Private/TextureDerivedDataBuildUtils.cpp b/Engine/Source/Runtime/Engine/Private/TextureDerivedDataBuildUtils.cpp
--- a/Engine/Source/Runtime/Engine/Private/TextureDerivedDataBuildUtils.cpp
+++ b/Engine/Source/Runtime/Engine/Private/TextureDerivedDataBuildUtils.cpp
@@ -158,7 +158,7 @@ static FCbObject WriteTextureSourceToCompactBinary(const FTextureSource& Texture
 	FCbWriter Writer;
 	Writer.BeginObject();
 
-	Writer.AddString("Input", TextureSource.GetIdString());
+	Writer.AddString("Input", TextureSource.GetId().ToString());
 	Writer.AddInteger("CompressionFormat", TextureSource.GetSourceCompression());
 	Writer.AddInteger("SourceFormat", TextureSource.GetFormat());
 	Writer.AddInteger("GammaSpace", (int32)GammaSpace);
diff --git a/Engine/Source/Runtime/Engine/Private/TextureDerivedDataTask.cpp b/Engine/Source/Runtime/Engine/Private/TextureDerivedDataTask.cpp
--- a/Engine/Source/Runtime/Engine/Private/TextureDerivedDataTask.cpp
+++ b/Engine/Source/Runtime/Engine/Private/TextureDerivedDataTask.cpp
@@ -25,6 +25,8 @@
 
 #include "DerivedDataBuild.h"
 #include "DerivedDataBuildAction.h"
+#include "DerivedDataBuildInputResolver.h"
+#include "DerivedDataBuildInputs.h"
 #include "DerivedDataBuildOutput.h"
 #include "DerivedDataBuildSession.h"
 #include "DerivedDataCacheInterface.h"
@@ -430,26 +432,90 @@ void FTextureCacheDerivedDataWorker::BuildTexture(bool bReplaceExistingDDC)
 		{
 			using namespace UE::DerivedData;
 			// Compress the texture by calling the DDC2 build function.  May be executed locally or remotely.
-			// The FBuildAction built here will be superceded by an FBuildDefinition in the future that will
-			// replace both the cache query/fill as well as execution.
+			// The FBuildDefinition in the future will replace both the cache get/put as well as execution.
 			IBuild& Build = GetDerivedDataBuildRef();
 			FString TexturePath = Texture.GetPathName();
-			FBuildActionBuilder ActionBuilder = Build.CreateAction(TexturePath, BuildFunctionName);
+			FBuildDefinitionBuilder DefinitionBuilder = Build.CreateDefinition(TexturePath, BuildFunctionName);
 
 			ComposeTextureBuildFunctionConstants(KeySuffix, Texture, BuildSettingsPerLayer[0], 0, NUM_INLINE_DERIVED_MIPS, 
-				[&ActionBuilder] (FStringView Key, const FCbObject& Value)
+				[&DefinitionBuilder] (FStringView Key, const FCbObject& Value)
 				{
-					ActionBuilder.AddConstant(Key, Value);
+					DefinitionBuilder.AddConstant(Key, Value);
 				});
 
 			for (const FBuildInputRecord& InputRecord : BuildInputRecords)
 			{
-				ActionBuilder.AddInput(InputRecord.Id, InputRecord.RawHash, InputRecord.RawSize);
+				DefinitionBuilder.AddInputBulkData(InputRecord.Id.ToString(), InputRecord.Id);
 			}
 
-			FBuildSession Session = Build.CreateSession(TexturePath, nullptr);
-			Session.BuildAction(ActionBuilder.Build(), InputsBuilder.Build(), EBuildPolicy::Default, EPriority::Blocking,
-				[this, &TexturePath, bReplaceExistingDDC] (FBuildActionCompleteParams&& Params)
+			class FTextureBuildInputResolver final : public IBuildInputResolver
+			{
+			public:
+				explicit FTextureBuildInputResolver(TConstArrayView<FBuildInputRecord> InRecords)
+					: Records(InRecords)
+				{
+				}
+
+				FRequest ResolveInputMeta(
+					const FBuildDefinition& Definition,
+					EPriority Priority,
+					FOnBuildInputMetaResolved&& OnResolved) final
+				{
+					EStatus Status = EStatus::Ok;
+					TArray<FString> InputKeys;
+					TArray<FBuildInputMetaByKey> Inputs;
+					Definition.IterateInputBulkData([this, &Status, &InputKeys, &Inputs](FStringView Key, const FGuid& BulkDataId)
+					{
+						if (const FBuildInputRecord* Match = Algo::FindBy(Records, BulkDataId, &FBuildInputRecord::Id))
+						{
+							InputKeys.Emplace(Key);
+							Inputs.Add({InputKeys.Last(), Match->Data.GetRawHash(), Match->Data.GetRawSize()});
+						}
+						else
+						{
+							Status = EStatus::Error;
+						}
+					});
+					OnResolved({Inputs, Status});
+					return FRequest();
+				}
+
+				FRequest ResolveInputData(
+					const FBuildDefinition& Definition,
+					EPriority Priority,
+					FOnBuildInputDataResolved&& OnResolved,
+					FBuildInputFilter&& Filter) final
+				{
+					EStatus Status = EStatus::Ok;
+					TArray<FString> InputKeys;
+					TArray<FBuildInputDataByKey> Inputs;
+					Definition.IterateInputBulkData([this, &Filter, &Status, &InputKeys, &Inputs](FStringView Key, const FGuid& BulkDataId)
+					{
+						if (!Filter || Filter(Key))
+						{
+							if (const FBuildInputRecord* Match = Algo::FindBy(Records, BulkDataId, &FBuildInputRecord::Id))
+							{
+								InputKeys.Emplace(Key);
+								Inputs.Add({InputKeys.Last(), Match->Data});
+							}
+							else
+							{
+								Status = EStatus::Error;
+							}
+						}
+					});
+					OnResolved({Inputs, Status});
+					return FRequest();
+				}
+
+			private:
+				TConstArrayView<FBuildInputRecord> Records;
+			};
+
+			FTextureBuildInputResolver InputResolver(BuildInputRecords);
+			FBuildSession Session = Build.CreateSession(TexturePath, &InputResolver);
+			Session.Build(DefinitionBuilder.Build(), EBuildPolicy::Default, EPriority::Blocking,
+				[this, &TexturePath, bReplaceExistingDDC] (FBuildCompleteParams&& Params)
 				{
 					#if !NO_LOGGING
 					Params.Output.IterateDiagnostics([this, &TexturePath] (const FBuildDiagnostic& Diagnostic)
@@ -552,7 +618,6 @@ FTextureCacheDerivedDataWorker::FTextureCacheDerivedDataWorker(
 	, ImageWrapper(nullptr)
 	, DerivedData(InDerivedData)
 	, Texture(*InTexture)
-	, InputsBuilder(GetDerivedDataBuildRef().CreateInputs(InTexture->GetPathName()))
 	, CacheFlags(InCacheFlags)
 	, RequiredMemoryEstimate(InTexture->GetBuildRequiredMemory())
 	, bSucceeded(false)
diff --git a/Engine/Source/Runtime/Engine/Private/TextureDerivedDataTask.h b/Engine/Source/Runtime/Engine/Private/TextureDerivedDataTask.h
--- a/Engine/Source/Runtime/Engine/Private/TextureDerivedDataTask.h
+++ b/Engine/Source/Runtime/Engine/Private/TextureDerivedDataTask.h
@@ -11,7 +11,6 @@
 #if WITH_EDITOR
 
 #include "Async/AsyncWork.h"
-#include "DerivedDataBuildInputs.h"
 #include "Engine/Texture2D.h"
 #include "IImageWrapperModule.h"
 #include "ImageCore.h"
@@ -111,9 +110,8 @@ class FTextureCacheDerivedDataWorker : public FNonAbandonableTask
 {
 	struct FBuildInputRecord
 	{
-		FString Id;
-		FIoHash RawHash;
-		uint64 RawSize;
+		FGuid Id;
+		FCompressedBuffer Data;
 	};
 
 	/** Texture compressor module, must be loaded in the game thread. see FModuleManager::WarnIfItWasntSafeToLoadHere() */
@@ -134,10 +132,8 @@ class FTextureCacheDerivedDataWorker : public FNonAbandonableTask
 	FTextureSourceData CompositeTextureData;
 	/** DDC2 build function name to use to build this texture (if DDC2 is enabled and the target texture type has a DDC2 build function, empty otherwise) */
 	FString BuildFunctionName;
-	/** DDC2 build input recrods for putting on the action (if DDC2 is enabled and the target texture type has a DDC2 build function, empty otherwise) */
+	/** DDC2 build input records for putting on the action (if DDC2 is enabled and the target texture type has a DDC2 build function, empty otherwise) */
 	TArray<FBuildInputRecord> BuildInputRecords;
-	/** DDC2 build inputs builder that accumulates input blobs and their associated key (if DDC2 is enabled and the target texture type has a DDC2 build function, empty otherwise) */
-	UE::DerivedData::FBuildInputsBuilder InputsBuilder;
 	/** Exporter that can optionally write out build actions and reference outputs for the purpose of testing remote execution of builds */
 	FTextureDerivedDataBuildExporter BuildExporter;
 	/** Texture cache flags. */
@@ -159,19 +155,17 @@ class FTextureCacheDerivedDataWorker : public FNonAbandonableTask
 	{
 		if (!BuildFunctionName.IsEmpty())
 		{
-			Source.OperateOnLoadedBulkData([this, &Source] (const FSharedBuffer& BulkDataBuffer) 
+			FBuildInputRecord& NewRef = BuildInputRecords.AddDefaulted_GetRef();
+			NewRef.Id = Source.GetId();
+			Source.OperateOnLoadedBulkData([&NewRef](const FSharedBuffer& BulkDataBuffer)
 			{
-				FCompressedBuffer InputBuffer = FCompressedBuffer::Compress(NAME_Default, BulkDataBuffer);
-				FBuildInputRecord& NewRef = BuildInputRecords.AddDefaulted_GetRef();
-				NewRef.Id = Source.GetIdString();
-				NewRef.RawHash = InputBuffer.GetRawHash();
-				NewRef.RawSize = InputBuffer.GetRawSize();
-				InputsBuilder.AddInput(Source.GetIdString(), InputBuffer);
+				NewRef.Data = FCompressedBuffer::Compress(NAME_None, BulkDataBuffer);
 			});
 		}
 	}
 
 	void ConsumeBuildFunctionOutput(const UE::DerivedData::FBuildOutput& BuildOutput, const FString& TexturePath, bool bReplaceExistingDDC);
+
 public:
 
 	/** Initialization constructor. */
