HLSL scopes now only contain a single statement, various related cleanup

diff --git a/Engine/Source/Runtime/Engine/Private/HLSLTree/HLSLTree.cpp b/Engine/Source/Runtime/Engine/Private/HLSLTree/HLSLTree.cpp
--- a/Engine/Source/Runtime/Engine/Private/HLSLTree/HLSLTree.cpp
+++ b/Engine/Source/Runtime/Engine/Private/HLSLTree/HLSLTree.cpp
@@ -182,10 +182,9 @@ UE::HLSLTree::FEmitScope* UE::HLSLTree::FEmitContext::AcquireScope(const FScope&
 	else
 	{
 		EmitScope = new(*Allocator) FEmitScope();
-		EmitScope->SourceScope = &Scope;
-		if (Scope.ParentScope)
+		if (Scope.GetParentScope())
 		{
-			EmitScope->ParentScope = AcquireScope(*Scope.ParentScope);
+			EmitScope->ParentScope = AcquireScope(*Scope.GetParentScope());
 		}
 		ScopeMap.Add(&Scope, EmitScope);
 	}
@@ -712,35 +711,20 @@ UE::HLSLTree::ENodeVisitResult UE::HLSLTree::FScope::Visit(FNodeVisitor& Visitor
 	const ENodeVisitResult Result = Visitor.OnScope(*this);
 	if (ShouldVisitDependentNodes(Result))
 	{
-		FStatement* Statement = FirstStatement;
-		while (Statement)
-		{
-			Visitor.VisitNode(Statement);
-			Statement = Statement->NextStatement;
-		}
+		Visitor.VisitNode(Statement);
 	}
 	return Result;
 }
 
 bool UE::HLSLTree::FScope::EmitHLSL(FEmitContext& Context, FEmitScope& Scope) const
 {
-	const FStatement* Statement = FirstStatement;
 	bool bResult = true;
-	
-	Context.ScopeStack.Add(&Scope);
-	
-	while (Statement)
+	if (Statement)
 	{
-		if (!Statement->EmitHLSL(Context))
-		{
-			bResult = false;
-			break;
-		}
-		Statement = Statement->NextStatement;
+		Context.ScopeStack.Add(&Scope);
+		bResult = Statement->EmitHLSL(Context);
+		verify(Context.ScopeStack.Pop(false) == &Scope);
 	}
-
-	verify(Context.ScopeStack.Pop(false) == &Scope);
-
 	return bResult;
 }
 
@@ -764,17 +748,6 @@ void UE::HLSLTree::FScope::AddPreviousScope(FScope& Scope)
 	PreviousScope[NumPreviousScopes++] = &Scope;
 }
 
-void UE::HLSLTree::FScope::AddExpression(FExpression* Expression)
-{
-	check(!Expression->ParentScope);
-	Expression->ParentScope = this;
-}
-
-void UE::HLSLTree::FScope::UseNode(FNode* Node)
-{
-	Node->ParentScope = FScope::FindSharedParent(this, Node->ParentScope);
-}
-
 namespace UE
 {
 namespace HLSLTree
@@ -793,13 +766,13 @@ public:
 
 	virtual ENodeVisitResult OnExpression(FExpression& InExpression) override
 	{
-		Scope->UseNode(&InExpression);
+		InExpression.ParentScope = FScope::FindSharedParent(Scope, InExpression.ParentScope);
 		return ENodeVisitResult::VisitDependentNodes;
 	}
 
 	virtual ENodeVisitResult OnFunctionCall(FFunctionCall& InFunctionCall) override
 	{
-		Scope->UseNode(&InFunctionCall);
+		InFunctionCall.ParentScope = FScope::FindSharedParent(Scope, InFunctionCall.ParentScope);
 		return ENodeVisitResult::VisitDependentNodes;
 	}
 
@@ -821,25 +794,6 @@ void UE::HLSLTree::FScope::UseFunctionCall(FFunctionCall* FunctionCall)
 	Visitor.VisitNode(FunctionCall);
 }
 
-void UE::HLSLTree::FScope::AddStatement(FStatement* Statement)
-{
-	check(!Statement->ParentScope);
-	check(!Statement->NextStatement);
-
-	Statement->ParentScope = this;
-	if (!FirstStatement)
-	{
-		check(!LastStatement);
-		FirstStatement = Statement;
-		LastStatement = Statement;
-	}
-	else
-	{
-		check(LastStatement);
-		LastStatement->NextStatement = Statement;
-	}
-}
-
 UE::HLSLTree::FTree* UE::HLSLTree::FTree::Create(FMemStackBase& Allocator)
 {
 	FTree* Tree = new(Allocator) FTree();
@@ -950,6 +904,20 @@ bool UE::HLSLTree::FTree::EmitHLSL(UE::HLSLTree::FEmitContext& Context, FCodeWri
 	return false;
 }
 
+void UE::HLSLTree::FTree::RegisterExpression(FScope& Scope, FExpression* Expression)
+{
+	check(!Expression->ParentScope);
+	Expression->ParentScope = &Scope;
+}
+
+void UE::HLSLTree::FTree::RegisterStatement(FScope& Scope, FStatement* Statement)
+{
+	check(!Scope.Statement)
+	check(!Statement->ParentScope);
+	Statement->ParentScope = &Scope;
+	Scope.Statement = Statement;
+}
+
 UE::HLSLTree::FScope* UE::HLSLTree::FTree::NewScope(FScope& Scope)
 {
 	FScope* NewScope = NewNode<FScope>();
@@ -959,17 +927,15 @@ UE::HLSLTree::FScope* UE::HLSLTree::FTree::NewScope(FScope& Scope)
 	return NewScope;
 }
 
-UE::HLSLTree::FParameterDeclaration* UE::HLSLTree::FTree::NewParameterDeclaration(FScope& Scope, const FName& Name, const Shader::FValue& DefaultValue)
+UE::HLSLTree::FParameterDeclaration* UE::HLSLTree::FTree::NewParameterDeclaration(const FName& Name, const Shader::FValue& DefaultValue)
 {
 	FParameterDeclaration* Declaration = NewNode<FParameterDeclaration>(Name, DefaultValue);
-	Declaration->ParentScope = &Scope;
 	return Declaration;
 }
 
-UE::HLSLTree::FTextureParameterDeclaration* UE::HLSLTree::FTree::NewTextureParameterDeclaration(FScope& Scope, const FName& Name, const FTextureDescription& DefaultValue)
+UE::HLSLTree::FTextureParameterDeclaration* UE::HLSLTree::FTree::NewTextureParameterDeclaration(const FName& Name, const FTextureDescription& DefaultValue)
 {
 	FTextureParameterDeclaration* Declaration = NewNode<FTextureParameterDeclaration>(Name, DefaultValue);
-	Declaration->ParentScope = &Scope;
 	return Declaration;
 }
 
diff --git a/Engine/Source/Runtime/Engine/Private/HLSLTree/HLSLTreeCommon.cpp b/Engine/Source/Runtime/Engine/Private/HLSLTree/HLSLTreeCommon.cpp
--- a/Engine/Source/Runtime/Engine/Private/HLSLTree/HLSLTreeCommon.cpp
+++ b/Engine/Source/Runtime/Engine/Private/HLSLTree/HLSLTreeCommon.cpp
@@ -568,11 +568,6 @@ bool UE::HLSLTree::FExpressionFunctionOutput::EmitCode(FEmitContext& Context, FE
 	return true;
 }
 
-bool UE::HLSLTree::FStatementNestedScope::EmitHLSL(FEmitContext& Context) const
-{
-	return Context.WriteScope(*NextScope);
-}
-
 bool UE::HLSLTree::FStatementBreak::EmitHLSL(FEmitContext& Context) const
 {
 	Context.WriteScopeTerminatorf(TEXT("break;"));
diff --git a/Engine/Source/Runtime/Engine/Private/Materials/MaterialExpressionHLSL.cpp b/Engine/Source/Runtime/Engine/Private/Materials/MaterialExpressionHLSL.cpp
--- a/Engine/Source/Runtime/Engine/Private/Materials/MaterialExpressionHLSL.cpp
+++ b/Engine/Source/Runtime/Engine/Private/Materials/MaterialExpressionHLSL.cpp
@@ -117,21 +117,21 @@ EMaterialGenerateHLSLStatus UMaterialExpressionGetLocal::GenerateHLSLExpression(
 
 EMaterialGenerateHLSLStatus UMaterialExpressionVectorParameter::GenerateHLSLExpression(FMaterialHLSLGenerator& Generator, UE::HLSLTree::FScope& Scope, int32 OutputIndex, UE::HLSLTree::FExpression*& OutExpression)
 {
-	UE::HLSLTree::FParameterDeclaration* Declaration = Generator.AcquireParameterDeclaration(Scope, ParameterName, DefaultValue);
+	UE::HLSLTree::FParameterDeclaration* Declaration = Generator.AcquireParameterDeclaration(ParameterName, DefaultValue);
 	OutExpression = Generator.GetTree().NewExpression<UE::HLSLTree::FExpressionParameter>(Scope, Declaration);
 	return EMaterialGenerateHLSLStatus::Success;
 }
 
 EMaterialGenerateHLSLStatus UMaterialExpressionScalarParameter::GenerateHLSLExpression(FMaterialHLSLGenerator& Generator, UE::HLSLTree::FScope& Scope, int32 OutputIndex, UE::HLSLTree::FExpression*& OutExpression)
 {
-	UE::HLSLTree::FParameterDeclaration* Declaration = Generator.AcquireParameterDeclaration(Scope, ParameterName, DefaultValue);
+	UE::HLSLTree::FParameterDeclaration* Declaration = Generator.AcquireParameterDeclaration(ParameterName, DefaultValue);
 	OutExpression = Generator.GetTree().NewExpression<UE::HLSLTree::FExpressionParameter>(Scope, Declaration);
 	return EMaterialGenerateHLSLStatus::Success;
 }
 
 EMaterialGenerateHLSLStatus UMaterialExpressionStaticBoolParameter::GenerateHLSLExpression(FMaterialHLSLGenerator& Generator, UE::HLSLTree::FScope& Scope, int32 OutputIndex, UE::HLSLTree::FExpression*& OutExpression)
 {
-	UE::HLSLTree::FParameterDeclaration* Declaration = Generator.AcquireParameterDeclaration(Scope, ParameterName, (bool)DefaultValue);
+	UE::HLSLTree::FParameterDeclaration* Declaration = Generator.AcquireParameterDeclaration(ParameterName, (bool)DefaultValue);
 	OutExpression = Generator.GetTree().NewExpression<UE::HLSLTree::FExpressionParameter>(Scope, Declaration);
 	return EMaterialGenerateHLSLStatus::Success;
 }
@@ -146,14 +146,14 @@ EMaterialGenerateHLSLStatus UMaterialExpressionTextureCoordinate::GenerateHLSLEx
 EMaterialGenerateHLSLStatus UMaterialExpressionTextureObject::GenerateHLSLTexture(FMaterialHLSLGenerator& Generator, UE::HLSLTree::FScope& Scope, int32 OutputIndex, UE::HLSLTree::FTextureParameterDeclaration*& OutTexture)
 {
 	const UE::HLSLTree::FTextureDescription TextureDesc(Texture, SamplerType);
-	OutTexture = Generator.AcquireTextureDeclaration(Scope, TextureDesc);
+	OutTexture = Generator.AcquireTextureDeclaration(TextureDesc);
 	return EMaterialGenerateHLSLStatus::Success;
 }
 
 EMaterialGenerateHLSLStatus UMaterialExpressionTextureObjectParameter::GenerateHLSLTexture(FMaterialHLSLGenerator& Generator, UE::HLSLTree::FScope& Scope, int32 OutputIndex, UE::HLSLTree::FTextureParameterDeclaration*& OutTexture)
 {
 	const UE::HLSLTree::FTextureDescription TextureDesc(Texture, SamplerType);
-	OutTexture = Generator.AcquireTextureParameterDeclaration(Scope, ParameterName, TextureDesc);
+	OutTexture = Generator.AcquireTextureParameterDeclaration(ParameterName, TextureDesc);
 	return EMaterialGenerateHLSLStatus::Success;
 }
 
@@ -182,7 +182,7 @@ EMaterialGenerateHLSLStatus UMaterialExpressionTextureSample::GenerateHLSLExpres
 	else if (Texture)
 	{
 		const UE::HLSLTree::FTextureDescription TextureDesc(Texture, SamplerType);
-		TextureDeclaration = Generator.AcquireTextureDeclaration(Scope, TextureDesc);
+		TextureDeclaration = Generator.AcquireTextureDeclaration(TextureDesc);
 	}
 
 	return GenerateHLSLExpressionBase(Generator, Scope, TextureDeclaration, OutExpression);
@@ -198,7 +198,7 @@ EMaterialGenerateHLSLStatus UMaterialExpressionTextureSampleParameter::GenerateH
 	else if (Texture)
 	{
 		const UE::HLSLTree::FTextureDescription TextureDesc(Texture, SamplerType);
-		TextureDeclaration = Generator.AcquireTextureParameterDeclaration(Scope, ParameterName, TextureDesc);
+		TextureDeclaration = Generator.AcquireTextureParameterDeclaration(ParameterName, TextureDesc);
 	}
 
 	return GenerateHLSLExpressionBase(Generator, Scope, TextureDeclaration, OutExpression);
diff --git a/Engine/Source/Runtime/Engine/Private/Materials/MaterialHLSLGenerator.cpp b/Engine/Source/Runtime/Engine/Private/Materials/MaterialHLSLGenerator.cpp
--- a/Engine/Source/Runtime/Engine/Private/Materials/MaterialHLSLGenerator.cpp
+++ b/Engine/Source/Runtime/Engine/Private/Materials/MaterialHLSLGenerator.cpp
@@ -291,17 +291,17 @@ UE::HLSLTree::FExpression* FMaterialHLSLGenerator::NewFunctionInput(UE::HLSLTree
 	return Expression;
 }
 
-UE::HLSLTree::FParameterDeclaration* FMaterialHLSLGenerator::AcquireParameterDeclaration(UE::HLSLTree::FScope& Scope, const FName& Name, const UE::Shader::FValue& DefaultValue)
+UE::HLSLTree::FParameterDeclaration* FMaterialHLSLGenerator::AcquireParameterDeclaration(const FName& Name, const UE::Shader::FValue& DefaultValue)
 {
 	UE::HLSLTree::FParameterDeclaration*& Declaration = ParameterDeclarationMap.FindOrAdd(Name);
 	if (!Declaration)
 	{
-		Declaration = HLSLTree->NewParameterDeclaration(Scope, Name, DefaultValue);
+		Declaration = HLSLTree->NewParameterDeclaration(Name, DefaultValue);
 	}
 	return Declaration;
 }
 
-UE::HLSLTree::FTextureParameterDeclaration* FMaterialHLSLGenerator::AcquireTextureDeclaration(UE::HLSLTree::FScope& Scope, const UE::HLSLTree::FTextureDescription& Value)
+UE::HLSLTree::FTextureParameterDeclaration* FMaterialHLSLGenerator::AcquireTextureDeclaration(const UE::HLSLTree::FTextureDescription& Value)
 {
 	FString SamplerTypeError;
 	if (!UMaterialExpressionTextureBase::VerifySamplerType(CompileTarget.FeatureLevel, CompileTarget.TargetPlatform, Value.Texture, Value.SamplerType, SamplerTypeError))
@@ -313,12 +313,12 @@ UE::HLSLTree::FTextureParameterDeclaration* FMaterialHLSLGenerator::AcquireTextu
 	UE::HLSLTree::FTextureParameterDeclaration*& Declaration = TextureDeclarationMap.FindOrAdd(Value);
 	if (!Declaration)
 	{
-		Declaration = HLSLTree->NewTextureParameterDeclaration(Scope, FName(), Value);
+		Declaration = HLSLTree->NewTextureParameterDeclaration(FName(), Value);
 	}
 	return Declaration;
 }
 
-UE::HLSLTree::FTextureParameterDeclaration* FMaterialHLSLGenerator::AcquireTextureParameterDeclaration(UE::HLSLTree::FScope& Scope, const FName& Name, const UE::HLSLTree::FTextureDescription& DefaultValue)
+UE::HLSLTree::FTextureParameterDeclaration* FMaterialHLSLGenerator::AcquireTextureParameterDeclaration(const FName& Name, const UE::HLSLTree::FTextureDescription& DefaultValue)
 {
 	FString SamplerTypeError;
 	if (!UMaterialExpressionTextureBase::VerifySamplerType(CompileTarget.FeatureLevel, CompileTarget.TargetPlatform, DefaultValue.Texture, DefaultValue.SamplerType, SamplerTypeError))
@@ -330,7 +330,7 @@ UE::HLSLTree::FTextureParameterDeclaration* FMaterialHLSLGenerator::AcquireTextu
 	UE::HLSLTree::FTextureParameterDeclaration*& Declaration = TextureParameterDeclarationMap.FindOrAdd(Name);
 	if (!Declaration)
 	{
-		Declaration = HLSLTree->NewTextureParameterDeclaration(Scope, Name, DefaultValue);
+		Declaration = HLSLTree->NewTextureParameterDeclaration(Name, DefaultValue);
 	}
 	return Declaration;
 }
diff --git a/Engine/Source/Runtime/Engine/Public/HLSLTree/HLSLTree.h b/Engine/Source/Runtime/Engine/Public/HLSLTree/HLSLTree.h
--- a/Engine/Source/Runtime/Engine/Public/HLSLTree/HLSLTree.h
+++ b/Engine/Source/Runtime/Engine/Public/HLSLTree/HLSLTree.h
@@ -133,7 +133,6 @@ public:
 	FEmitStatement* LastStatement = nullptr;
 	FEmitScopeLink* FirstLink = nullptr;
 	FEmitScopeLink* LastLink = nullptr;
-	const FScope* SourceScope = nullptr;
 	TMap<FSHAHash, const TCHAR*> ExpressionMap;
 };
 
@@ -327,7 +326,6 @@ public:
 
 	virtual ENodeVisitResult Visit(FNodeVisitor& Visitor) = 0;
 
-	FScope* ParentScope = nullptr;
 	FNode* NextNode = nullptr;
 };
 
@@ -344,7 +342,7 @@ public:
 	/** Emits HLSL code for the statement. The generated code should include any required semi-colons and newlines */
 	virtual bool EmitHLSL(FEmitContext& Context) const = 0;
 
-	FStatement* NextStatement = nullptr;
+	FScope* ParentScope = nullptr;
 };
 
 /**
@@ -361,6 +359,8 @@ public:
 
 	/** Emits code for the given expression, either HLSL code or preshader bytecode */
 	virtual bool EmitCode(FEmitContext& Context, FExpressionEmitResult& OutResult) const = 0;
+
+	FScope* ParentScope = nullptr;
 };
 
 /**
@@ -415,6 +415,8 @@ class FFunctionCall final : public FNode
 public:
 	virtual ENodeVisitResult Visit(FNodeVisitor& Visitor) override;
 
+	FScope* ParentScope = nullptr;
+
 	/** Root scope of the function to call. Note that this scope will be from a separate (external) tree */
 	const FScope* FunctionScope;
 
@@ -429,14 +431,15 @@ public:
 };
 
 /**
- * Represents an HLSL scope.  This is an ordered list of statements (enclosed by an {} pair in HLSL)
- * All HLSL nodes track which scopes they are accessed from
+ * Represents an HLSL scope.  A scope contains a single statement, along with any expressions required by that statement
  */
 class FScope final : public FNode
 {
 public:
 	static FScope* FindSharedParent(FScope* Lhs, FScope* Rhs);
 
+	inline FScope* GetParentScope() const { return ParentScope; }
+
 	inline TArrayView<FScope*> GetPreviousScopes() const
 	{
 		// const_cast needed, otherwise type of array view is 'FScope*const' which doesn't make sense
@@ -450,20 +453,16 @@ public:
 
 	void AddPreviousScope(FScope& Scope);
 
-	void AddExpression(FExpression* Expression);
-	void AddStatement(FStatement* Statement);
-
 	void UseFunctionCall(FFunctionCall* FunctionCall);
 	void UseExpression(FExpression* Expression);
+
 private:
 	friend class FTree;
 	friend class FNodeVisitor_MoveToScope;
 
-	void UseNode(FNode* Node);
-
+	FScope* ParentScope = nullptr;
+	FStatement* Statement = nullptr;
 	FScope* PreviousScope[MaxNumPreviousScopes];
-	FStatement* FirstStatement = nullptr;
-	FStatement* LastStatement = nullptr;
 	int32 NumPreviousScopes = 0;
 	int32 NestedLevel = 0;
 };
@@ -485,7 +484,7 @@ public:
 	inline T* NewExpression(FScope& Scope, ArgTypes&&... Args)
 	{
 		T* Expression = NewNode<T>(Forward<ArgTypes>(Args)...);
-		Scope.AddExpression(Expression);
+		RegisterExpression(Scope, Expression);
 		return Expression;
 	}
 
@@ -493,13 +492,13 @@ public:
 	inline T* NewStatement(FScope& Scope, ArgTypes&&... Args)
 	{
 		T* Statement = NewNode<T>(Forward<ArgTypes>(Args)...);
-		Scope.AddStatement(Statement);
+		RegisterStatement(Scope, Statement);
 		return Statement;
 	}
 
 	FScope* NewScope(FScope& Scope);
-	FParameterDeclaration* NewParameterDeclaration(FScope& Scope, const FName& Name, const Shader::FValue& DefaultValue);
-	FTextureParameterDeclaration* NewTextureParameterDeclaration(FScope& Scope, const FName& Name, const FTextureDescription& DefaultValue);
+	FParameterDeclaration* NewParameterDeclaration(const FName& Name, const Shader::FValue& DefaultValue);
+	FTextureParameterDeclaration* NewTextureParameterDeclaration(const FName& Name, const FTextureDescription& DefaultValue);
 
 	FFunctionCall* NewFunctionCall(FScope& Scope,
 		const FScope& FunctionScope,
@@ -518,6 +517,9 @@ private:
 		return Node;
 	}
 
+	void RegisterExpression(FScope& Scope, FExpression* Expression);
+	void RegisterStatement(FScope& Scope, FStatement* Statement);
+
 	FMemStackBase* Allocator;
 	FNode* Nodes;
 	FScope* RootScope;
diff --git a/Engine/Source/Runtime/Engine/Public/HLSLTree/HLSLTreeCommon.h b/Engine/Source/Runtime/Engine/Public/HLSLTree/HLSLTreeCommon.h
--- a/Engine/Source/Runtime/Engine/Public/HLSLTree/HLSLTreeCommon.h
+++ b/Engine/Source/Runtime/Engine/Public/HLSLTree/HLSLTreeCommon.h
@@ -345,24 +345,6 @@ public:
 	virtual bool EmitCode(FEmitContext& Context, FExpressionEmitResult& OutResult) const override;
 };
 
-class FStatementNestedScope : public FStatement
-{
-public:
-	FScope* NextScope;
-
-	virtual ENodeVisitResult Visit(FNodeVisitor& Visitor) override
-	{
-		const ENodeVisitResult Result = FStatement::Visit(Visitor);
-		if (ShouldVisitDependentNodes(Result))
-		{
-			Visitor.VisitNode(NextScope);
-		}
-		return Result;
-	}
-
-	virtual bool EmitHLSL(FEmitContext& Context) const override;
-};
-
 class FStatementReturn : public FStatement
 {
 public:
diff --git a/Engine/Source/Runtime/Engine/Public/MaterialHLSLGenerator.h b/Engine/Source/Runtime/Engine/Public/MaterialHLSLGenerator.h
--- a/Engine/Source/Runtime/Engine/Public/MaterialHLSLGenerator.h
+++ b/Engine/Source/Runtime/Engine/Public/MaterialHLSLGenerator.h
@@ -92,13 +92,13 @@ public:
 
 	UE::HLSLTree::FExpression* NewFunctionInput(UE::HLSLTree::FScope& Scope, int32 InputIndex, UMaterialExpressionFunctionInput* MaterialFunctionInput);
 
-	UE::HLSLTree::FParameterDeclaration* AcquireParameterDeclaration(UE::HLSLTree::FScope& Scope, const FName& Name, const UE::Shader::FValue& DefaultValue);
+	UE::HLSLTree::FParameterDeclaration* AcquireParameterDeclaration(const FName& Name, const UE::Shader::FValue& DefaultValue);
 
 	/** Returns a declaration to access the given texture, with no parameter */
-	UE::HLSLTree::FTextureParameterDeclaration* AcquireTextureDeclaration(UE::HLSLTree::FScope& Scope, const UE::HLSLTree::FTextureDescription& Value);
+	UE::HLSLTree::FTextureParameterDeclaration* AcquireTextureDeclaration(const UE::HLSLTree::FTextureDescription& Value);
 
 	/** Returns a declaration to access the given texture parameter */
-	UE::HLSLTree::FTextureParameterDeclaration* AcquireTextureParameterDeclaration(UE::HLSLTree::FScope& Scope, const FName& Name, const UE::HLSLTree::FTextureDescription& DefaultValue);
+	UE::HLSLTree::FTextureParameterDeclaration* AcquireTextureParameterDeclaration(const FName& Name, const UE::HLSLTree::FTextureDescription& DefaultValue);
 
 	UE::HLSLTree::FFunctionCall* AcquireFunctionCall(UE::HLSLTree::FScope& Scope, UMaterialFunctionInterface* Function, TArrayView<UE::HLSLTree::FExpression*> Inputs);
 
