Fixed CastCheck() warning when running ListStreamingTextures command.

diff --git a/Engine/Source/Runtime/Engine/Private/Streaming/StreamingManagerTexture.cpp b/Engine/Source/Runtime/Engine/Private/Streaming/StreamingManagerTexture.cpp
--- a/Engine/Source/Runtime/Engine/Private/Streaming/StreamingManagerTexture.cpp
+++ b/Engine/Source/Runtime/Engine/Private/Streaming/StreamingManagerTexture.cpp
@@ -2018,9 +2018,9 @@ bool FRenderAssetStreamingManager::HandleListStreamingRenderAssetsCommand( const
 		FTexturePlatformData** TexturePlatformData = Cast<UTexture>(RenderAsset) ? const_cast<UTexture*>(Cast<UTexture>(RenderAsset))->GetRunningPlatformData() : nullptr;
 		if (AssetType == EStreamableRenderAssetType::Texture && TexturePlatformData)
 		{
-			const UTexture2D* Texture2D = CastChecked<UTexture2D>(RenderAsset);
-			const UVolumeTexture* VolumeTexture = CastChecked<UVolumeTexture>(RenderAsset);
-			const UTexture2DArray* Texture2DArray = CastChecked<UTexture2DArray>(RenderAsset);
+			const UTexture2D* Texture2D = Cast<UTexture2D>(RenderAsset);
+			const UVolumeTexture* VolumeTexture = Cast<UVolumeTexture>(RenderAsset);
+			const UTexture2DArray* Texture2DArray = Cast<UTexture2DArray>(RenderAsset);
 			const TIndirectArray<struct FTexture2DMipMap>& TextureMips = (*TexturePlatformData)->Mips;
 
 			auto OutputMipsString = [&](int32 OutputIndex)->FString
