Adding DefaultMediaSoundClass in AudioSettings to root to avoid GC crash on boot if it is not the same object as DefaultSoundClass.

diff --git a/Engine/Source/Runtime/Engine/Private/AudioSettings.cpp b/Engine/Source/Runtime/Engine/Private/AudioSettings.cpp
--- a/Engine/Source/Runtime/Engine/Private/AudioSettings.cpp
+++ b/Engine/Source/Runtime/Engine/Private/AudioSettings.cpp
@@ -235,6 +235,7 @@ void UAudioSettings::LoadDefaultObjects()
 	if (UObject* MediaSoundClassObject = DefaultMediaSoundClassName.TryLoad())
 	{
 		DefaultMediaSoundClass = CastChecked<USoundClass>(MediaSoundClassObject);
+		DefaultMediaSoundClass->AddToRoot();
 	}
 	else
 	{
