AWorldDataLayers::OverwriteDataLayerStates to override initial DataLayer States

diff --git a/Engine/Source/Runtime/Engine/Private/WorldPartition/DataLayer/WorldDataLayers.cpp b/Engine/Source/Runtime/Engine/Private/WorldPartition/DataLayer/WorldDataLayers.cpp
--- a/Engine/Source/Runtime/Engine/Private/WorldPartition/DataLayer/WorldDataLayers.cpp
+++ b/Engine/Source/Runtime/Engine/Private/WorldPartition/DataLayer/WorldDataLayers.cpp
@@ -162,6 +162,34 @@ EDataLayerState AWorldDataLayers::GetDataLayerStateByName(FName InDataLayerName)
 }
 
 #if WITH_EDITOR
+void AWorldDataLayers::OverwriteDataLayerStates(TArray<FActorDataLayer>* InActiveDataLayers, TArray<FActorDataLayer>* InLoadedDataLayers)
+{
+	if (GetLocalRole() == ROLE_Authority)
+	{
+		// This should get called before game starts. It doesn't send out events
+		check(!GetWorld()->bMatchStarted);
+		if (InActiveDataLayers)
+		{
+			ActiveDataLayerNames.Empty(InActiveDataLayers->Num());
+			for (const FActorDataLayer& DataLayer : *InActiveDataLayers)
+			{
+				ActiveDataLayerNames.Add(DataLayer.Name);
+			}
+			RepActiveDataLayerNames = ActiveDataLayerNames.Array();
+		}
+
+		if (InLoadedDataLayers)
+		{
+			LoadedDataLayerNames.Empty(InLoadedDataLayers->Num());
+			for (const FActorDataLayer& DataLayer : *InLoadedDataLayers)
+			{
+				LoadedDataLayerNames.Add(DataLayer.Name);
+			}
+			RepLoadedDataLayerNames = LoadedDataLayerNames.Array();
+		}
+	}
+}
+
 AWorldDataLayers* AWorldDataLayers::Create(UWorld* World)
 {
 	check(World);
diff --git a/Engine/Source/Runtime/Engine/Public/WorldPartition/DataLayer/WorldDataLayers.h b/Engine/Source/Runtime/Engine/Public/WorldPartition/DataLayer/WorldDataLayers.h
--- a/Engine/Source/Runtime/Engine/Public/WorldPartition/DataLayer/WorldDataLayers.h
+++ b/Engine/Source/Runtime/Engine/Public/WorldPartition/DataLayer/WorldDataLayers.h
@@ -33,6 +33,9 @@ public:
 	TArray<const UDataLayer*> GetDataLayerObjects(const TArray<FActorDataLayer>& DataLayers) const;
 	TArray<FName> GetDataLayerNames(const TArray<FActorDataLayer>& DataLayers) const;
 	//~ End Helper Functions
+
+	// Allows overriding of DataLayers with PlayFromHere
+	void OverwriteDataLayerStates(TArray<FActorDataLayer>* InActiveDataLayers = nullptr, TArray<FActorDataLayer>* InLoadedDataLayers = nullptr);
 #endif
 	
 	bool ContainsDataLayer(const UDataLayer* InDataLayer) const;
@@ -42,7 +45,7 @@ public:
 	void ForEachDataLayer(TFunctionRef<bool(class UDataLayer*)> Func) const;
 
 	void SetDataLayerState(FActorDataLayer InDataLayer, EDataLayerState InState);
-
+	
 	UFUNCTION(NetMulticast, Reliable)
 	void OnDataLayerStateChanged(const UDataLayer* InDataLayer, EDataLayerState InState);
 
