PR #7734: Added UPrimitiveComponent::SetRenderInDepthPass function (Contributed by [REDACTED_PERSON])

diff --git a/Engine/Source/Runtime/Engine/Classes/Components/PrimitiveComponent.h b/Engine/Source/Runtime/Engine/Classes/Components/PrimitiveComponent.h
--- a/Engine/Source/Runtime/Engine/Classes/Components/PrimitiveComponent.h
+++ b/Engine/Source/Runtime/Engine/Classes/Components/PrimitiveComponent.h
@@ -1634,6 +1634,10 @@ public:
 	/** Sets bRenderInMainPass property and marks the render state dirty. */
 	UFUNCTION(BlueprintCallable, Category = "Rendering")
 	void SetRenderInMainPass(bool bValue);
+	
+	/** Sets bRenderInDepthPass property and marks the render state dirty. */
+	UFUNCTION(BlueprintCallable, Category = "Rendering")
+	void SetRenderInDepthPass(bool bValue);
 
 	/** Sets bVisibleInSceneCaptureOnly property and marks the render state dirty. */
 	UFUNCTION(BlueprintCallable, Category = "Rendering")
diff --git a/Engine/Source/Runtime/Engine/Private/Components/PrimitiveComponent.cpp b/Engine/Source/Runtime/Engine/Private/Components/PrimitiveComponent.cpp
--- a/Engine/Source/Runtime/Engine/Private/Components/PrimitiveComponent.cpp
+++ b/Engine/Source/Runtime/Engine/Private/Components/PrimitiveComponent.cpp
@@ -3785,6 +3785,15 @@ void UPrimitiveComponent::SetRenderInMainPass(bool bValue)
 	}
 }
 
+void UPrimitiveComponent::SetRenderInDepthPass(bool bValue)
+{
+	if (bRenderInDepthPass != bValue)
+	{
+		bRenderInDepthPass = bValue;
+		MarkRenderStateDirty();
+	}
+}
+
 void UPrimitiveComponent::SetVisibleInSceneCaptureOnly(bool bValue)
 {
 	if (bVisibleInSceneCaptureOnly != bValue)
