Fixed OrthonormalizeTransform when LWC is enabled

diff --git a/Engine/Source/Runtime/Engine/Public/InstanceUniformShaderParameters.h b/Engine/Source/Runtime/Engine/Public/InstanceUniformShaderParameters.h
--- a/Engine/Source/Runtime/Engine/Public/InstanceUniformShaderParameters.h
+++ b/Engine/Source/Runtime/Engine/Public/InstanceUniformShaderParameters.h
@@ -36,9 +36,9 @@ public:
 	}
 };
 
-FORCEINLINE FVector OrthonormalizeTransform(FMatrix& Matrix)
+FORCEINLINE FVector3f OrthonormalizeTransform(FMatrix44f& Matrix)
 {
-	FVector X, Y, Z, Origin;
+	FVector3f X, Y, Z, Origin;
 	Matrix.GetScaledAxes(X, Y, Z);
 	Origin = Matrix.GetOrigin();
 
@@ -47,10 +47,10 @@ FORCEINLINE FVector OrthonormalizeTransform(FMatrix& Matrix)
 	Z -= (Z | X) / (X | X) * X;
 	Z -= (Z | Y) / (Y | Y) * Y;
 
-	Matrix = FMatrix(X, Y, Z, Origin);
+	Matrix = FMatrix44f(X, Y, Z, Origin);
 
 	// Extract per axis scales
-	FVector Scale;
+	FVector3f Scale;
 	Scale.X = X.Size();
 	Scale.Y = Y.Size();
 	Scale.Z = Z.Size();
@@ -77,7 +77,7 @@ struct FPrimitiveInstance
 	FORCEINLINE void OrthonormalizeAndUpdateScale()
 	{
 		// Remove shear
-		const FVector Scale = OrthonormalizeTransform(LocalToWorld);
+		const FVector3f Scale = OrthonormalizeTransform(LocalToWorld);
 		OrthonormalizeTransform(PrevLocalToWorld);
 
 		NonUniformScale = FVector4(
