Texture: Remove unused #include from TextureDerivedDataBuildExporter

#rnx

diff --git a/Engine/Source/Runtime/Engine/Private/TextureDerivedDataBuildExporter.cpp b/Engine/Source/Runtime/Engine/Private/TextureDerivedDataBuildExporter.cpp
--- a/Engine/Source/Runtime/Engine/Private/TextureDerivedDataBuildExporter.cpp
+++ b/Engine/Source/Runtime/Engine/Private/TextureDerivedDataBuildExporter.cpp
@@ -16,7 +16,6 @@
 #include "Misc/CoreMisc.h"
 #include "Misc/FileHelper.h"
 #include "Misc/StringBuilder.h"
-#include "Serialization/CompactBinaryPackage.h"
 #include "Serialization/CompactBinaryWriter.h"
 #include "TextureDerivedDataBuildUtils.h"
 #include "TextureDerivedDataTask.h"
