PIE Lazy/Soft object ptr fixup optimizations:
- Don't try to fixup FMulticastDelegate properties.
- Don't try to fixup UInstancedStaticMeshComponent objects.

PIE fixup from ~150ms to ~1ms per streaming cell.

diff --git a/Engine/Source/Editor/UnrealEd/Classes/Editor/FixupLazyObjectPtrForPIEArchive.h b/Engine/Source/Editor/UnrealEd/Classes/Editor/FixupLazyObjectPtrForPIEArchive.h
--- a/Engine/Source/Editor/UnrealEd/Classes/Editor/FixupLazyObjectPtrForPIEArchive.h
+++ b/Engine/Source/Editor/UnrealEd/Classes/Editor/FixupLazyObjectPtrForPIEArchive.h
@@ -24,6 +24,11 @@ public:
 		ArIgnoreArchetypeRef = true;
 	}
 
+	virtual bool ShouldSkipProperty(const FProperty* InProperty) const override
+	{
+		return InProperty->IsA<FMulticastDelegateProperty>() || FArchiveUObject::ShouldSkipProperty(InProperty);
+	}
+
 	virtual FArchive& operator<<(FLazyObjectPtr& LazyObjectPtr) override
 	{
 		FArchive& Ar = *this;
@@ -44,7 +49,11 @@ public:
 			UPackage* Outermost = Object->GetOutermost();
 			if (!Outermost || Outermost->HasAnyPackageFlags(PKG_PlayInEditor))
 			{
-				Object->Serialize(*this);
+				// Skip instanced static mesh component as their impact on serialization is enormous and they don't contain lazy ptrs.
+				if (!Cast<UInstancedStaticMeshComponent>(Object))
+				{
+					Object->Serialize(*this);
+				}
 			}
 		}
 		return *this;
diff --git a/Engine/Source/Runtime/Engine/Private/Level.cpp b/Engine/Source/Runtime/Engine/Private/Level.cpp
--- a/Engine/Source/Runtime/Engine/Private/Level.cpp
+++ b/Engine/Source/Runtime/Engine/Private/Level.cpp
@@ -52,6 +52,7 @@ Level.cpp: Level-related functions
 #include "UnrealEngine.h"
 #include "Misc/ArchiveMD5.h"
 #if WITH_EDITOR
+#include "Components/InstancedStaticMeshComponent.h"
 #include "Kismet2/KismetEditorUtilities.h"
 #include "Kismet2/BlueprintEditorUtils.h"
 #include "Algo/AnyOf.h"
@@ -2713,6 +2714,11 @@ void ULevel::FixupForPIE(int32 InPIEInstanceID, TFunctionRef<void(int32, FSoftOb
 			this->ArShouldSkipBulkData = true;
 		}
 
+		virtual bool ShouldSkipProperty(const FProperty* InProperty) const override
+		{
+			return InProperty->IsA<FMulticastDelegateProperty>() || FArchiveUObject::ShouldSkipProperty(InProperty);
+		}
+
 		FArchive& operator<<(FSoftObjectPath& Value)
 		{
 			Value.FixupForPIE(PIEInstanceID, CustomFixupFunction);
@@ -2730,7 +2736,11 @@ void ULevel::FixupForPIE(int32 InPIEInstanceID, TFunctionRef<void(int32, FSoftOb
 
 	for (UObject* Object : SubObjects)
 	{
-		Object->Serialize(FixupSerializer);
+		// Skip instanced static mesh component as their impact on serialization is enormous and they don't contain soft ptrs.
+		if (!Cast<UInstancedStaticMeshComponent>(Object))
+		{
+			Object->Serialize(FixupSerializer);
+		}
 	}
 }
 
