Implement support for loading a shipping texture derived data from the cache.

diff --git a/Engine/Source/Runtime/Engine/Classes/Engine/Texture.h b/Engine/Source/Runtime/Engine/Classes/Engine/Texture.h
--- a/Engine/Source/Runtime/Engine/Classes/Engine/Texture.h
+++ b/Engine/Source/Runtime/Engine/Classes/Engine/Texture.h
@@ -596,6 +596,16 @@ struct FTexturePlatformData
 #if WITH_EDITORONLY_DATA
 	/** The key associated with this derived data. */
 	FString DerivedDataKey;
+	/** 
+	* The key associated with the derived data generated for the shipping build.
+	* In order to save developer's time when running a build in the Editor,
+	* cooking a texture for the Editor build normally involves less complicated processing
+	* than cooking a texture for the shipping build. However, in many cases a fully cooked
+	* texture is already available in the cache when the project is opened on the developer's machine.
+	* This key allows to load the texture generated for the shipping build from the cache, if available,
+	* instead of building a new texture locally using fast cooking.
+	*/
+	FString ShippingDerivedDataKey;
 	/** Async cache task if one is outstanding. */
 	struct FTextureAsyncCacheDerivedDataTask* AsyncTask;
 #endif
diff --git a/Engine/Source/Runtime/Engine/Private/TextureDerivedData.cpp b/Engine/Source/Runtime/Engine/Private/TextureDerivedData.cpp
--- a/Engine/Source/Runtime/Engine/Private/TextureDerivedData.cpp
+++ b/Engine/Source/Runtime/Engine/Private/TextureDerivedData.cpp
@@ -733,6 +733,23 @@ void FTexturePlatformData::Cache(
 	bool bAsync = !bForDDC && (Flags & ETextureCacheFlags::Async) != 0;
 	GetTextureDerivedDataKey(InTexture, InSettingsPerLayer, DerivedDataKey);
 
+	if (!bForDDC && !bForceRebuild && InSettingsPerLayer[0].bHasEditorOnlyData)
+	{
+		// If we are running in the Editor, generate a shipping derived data key, which can be used to load
+		// a texture cooked for the shipping build from the cache, if available.
+		// Specifically, FTextureBuildSettings::bHasEditorOnlyData controls the usage of Oodle RDO when building textures.
+		// In cases when an RDO-processed texture is present in the cache, we can use it instead of building or loading a non-RDO texture.
+		const int32 NumLayers = InTexture.Source.GetNumLayers();
+		TArray<FTextureBuildSettings> ShippingBuildSettingsPerLayer;
+		ShippingBuildSettingsPerLayer.SetNum(NumLayers);
+		for (int32 LayerIndex = 0; LayerIndex < NumLayers; ++LayerIndex)
+		{
+			ShippingBuildSettingsPerLayer[LayerIndex] = InSettingsPerLayer[LayerIndex];
+			ShippingBuildSettingsPerLayer[LayerIndex].bHasEditorOnlyData = false;
+		}
+		GetTextureDerivedDataKey(InTexture, ShippingBuildSettingsPerLayer.GetData(), ShippingDerivedDataKey);
+	}
+
 	if (!Compressor)
 	{
 		Compressor = &FModuleManager::LoadModuleChecked<ITextureCompressorModule>(TEXTURE_COMPRESSOR_MODULENAME);
diff --git a/Engine/Source/Runtime/Engine/Private/TextureDerivedDataTask.cpp b/Engine/Source/Runtime/Engine/Private/TextureDerivedDataTask.cpp
--- a/Engine/Source/Runtime/Engine/Private/TextureDerivedDataTask.cpp
+++ b/Engine/Source/Runtime/Engine/Private/TextureDerivedDataTask.cpp
@@ -719,7 +719,22 @@ void FTextureCacheDerivedDataWorker::DoWork()
 
 	TArray<uint8> RawDerivedData;
 
-	if (!bForceRebuild && GetDerivedDataCacheRef().GetSynchronous(*DerivedData->DerivedDataKey, RawDerivedData, Texture.GetPathName()))
+	if (!bForceRebuild)
+	{
+		// First try to load a texture generated for the shipping build from the cache.
+		// FTexturePlatformData::ShippingDerivedDataKey is set when we are running a build in the Editor.
+		// This allows to preview how the texture will look in the final build and avoid rebuilding texture locally using fast cooking.
+		if (!DerivedData->ShippingDerivedDataKey.IsEmpty() && DerivedData->ShippingDerivedDataKey != DerivedData->DerivedDataKey)
+		{
+			bLoadedFromDDC = GetDerivedDataCacheRef().GetSynchronous(*DerivedData->ShippingDerivedDataKey, RawDerivedData, Texture.GetPathName());
+		}
+		if (!bLoadedFromDDC)
+		{
+			bLoadedFromDDC = GetDerivedDataCacheRef().GetSynchronous(*DerivedData->DerivedDataKey, RawDerivedData, Texture.GetPathName());
+		}
+	}
+
+	if (bLoadedFromDDC)
 	{
 		const bool bInlineMips = (CacheFlags & ETextureCacheFlags::InlineMips) != 0;
 		const bool bForDDC = (CacheFlags & ETextureCacheFlags::ForDDCBuild) != 0;
@@ -793,7 +808,6 @@ void FTextureCacheDerivedDataWorker::DoWork()
 				}
 			}
 		}
-		bLoadedFromDDC = true;
 
 		if (bSucceeded && bForVirtualTextureStreamingBuild && CVarVTValidateCompressionOnLoad.GetValueOnAnyThread())
 		{
