Asserts when unable to load an actor package, even if it has valid metadata.
This usually happens when loading actors saved with a more recent engine.

diff --git a/Engine/Source/Runtime/Engine/Private/WorldPartition/WorldPartitionActorDesc.cpp b/Engine/Source/Runtime/Engine/Private/WorldPartition/WorldPartitionActorDesc.cpp
--- a/Engine/Source/Runtime/Engine/Private/WorldPartition/WorldPartitionActorDesc.cpp
+++ b/Engine/Source/Runtime/Engine/Private/WorldPartition/WorldPartitionActorDesc.cpp
@@ -269,12 +269,10 @@ AActor* FWorldPartitionActorDesc::Load() const
 		}
 
 		Package = LoadPackage(Package, *ActorPackage.ToString(), LOAD_None, nullptr, InstancingContext);
+		check(Package);
 
-		if (Package)
-		{
-			ActorPtr = FindObject<AActor>(nullptr, *ActorPath.ToString());
-			check(ActorPtr.IsValid());
-		}
+		ActorPtr = FindObject<AActor>(nullptr, *ActorPath.ToString());
+		check(ActorPtr.IsValid());
 	}
 
 	return ActorPtr.Get();
