Optimus: Export some compute graph object so that they're usable by the Optimus plugin.

#rnx

diff --git a/Engine/Source/Runtime/Engine/Private/ComputeFramework/ComputeGraph.cpp b/Engine/Source/Runtime/Engine/Private/ComputeFramework/ComputeGraph.cpp
--- a/Engine/Source/Runtime/Engine/Private/ComputeFramework/ComputeGraph.cpp
+++ b/Engine/Source/Runtime/Engine/Private/ComputeFramework/ComputeGraph.cpp
@@ -10,8 +10,10 @@
 #include "Interfaces/ITargetPlatform.h"
 #include "ShaderParameterMetadataBuilder.h"
 
+UComputeGraph::UComputeGraph() = default;
+
 UComputeGraph::UComputeGraph(const FObjectInitializer& ObjectInitializer)
-	: UObject(ObjectInitializer)
+	: Super(ObjectInitializer)
 {
 }
 
diff --git a/Engine/Source/Runtime/Engine/Public/ComputeFramework/ComputeGraph.h b/Engine/Source/Runtime/Engine/Public/ComputeFramework/ComputeGraph.h
--- a/Engine/Source/Runtime/Engine/Public/ComputeFramework/ComputeGraph.h
+++ b/Engine/Source/Runtime/Engine/Public/ComputeFramework/ComputeGraph.h
@@ -76,6 +76,7 @@ protected:
 	TArray<FComputeGraphEdge> GraphEdges;
 
 public:
+	UComputeGraph();
 	UComputeGraph(const FObjectInitializer& ObjectInitializer);
 	UComputeGraph(FVTableHelper& Helper);
 	virtual ~UComputeGraph();
diff --git a/Engine/Source/Runtime/Engine/Public/ComputeFramework/ComputeGraphComponent.h b/Engine/Source/Runtime/Engine/Public/ComputeFramework/ComputeGraphComponent.h
--- a/Engine/Source/Runtime/Engine/Public/ComputeFramework/ComputeGraphComponent.h
+++ b/Engine/Source/Runtime/Engine/Public/ComputeFramework/ComputeGraphComponent.h
@@ -13,7 +13,7 @@ class UComputeGraph;
  * This object binds a graph to its data providers, and queues the execution. 
  */
 UCLASS(Blueprintable, meta = (BlueprintSpawnableComponent))
-class UComputeGraphComponent : public UActorComponent
+class ENGINE_API UComputeGraphComponent : public UActorComponent
 {
 	GENERATED_BODY()
 
diff --git a/Engine/Source/Runtime/Engine/Public/ComputeFramework/ComputeKernelSource.h b/Engine/Source/Runtime/Engine/Public/ComputeFramework/ComputeKernelSource.h
--- a/Engine/Source/Runtime/Engine/Public/ComputeFramework/ComputeKernelSource.h
+++ b/Engine/Source/Runtime/Engine/Public/ComputeFramework/ComputeKernelSource.h
@@ -11,7 +11,7 @@
  * We derive from this for each authoring mechanism. (HLSL text, VPL graph, ML Meta Lang, etc.)
  */
 UCLASS(Abstract, EditInlineNew)
-class UComputeKernelSource : public UObject
+class ENGINE_API UComputeKernelSource : public UObject
 {
 	GENERATED_BODY()
 
diff --git a/Engine/Source/Runtime/Engine/Public/ComputeFramework/ShaderParamTypeDefinition.h b/Engine/Source/Runtime/Engine/Public/ComputeFramework/ShaderParamTypeDefinition.h
--- a/Engine/Source/Runtime/Engine/Public/ComputeFramework/ShaderParamTypeDefinition.h
+++ b/Engine/Source/Runtime/Engine/Public/ComputeFramework/ShaderParamTypeDefinition.h
@@ -177,11 +177,11 @@ private:
 * A hashing function to allow the FShaderValueType class to be used with hashing containers (e.g.
 * TSet or TMap).
 */
-uint32 GetTypeHash(const FShaderValueType& InShaderValueType);
+ENGINE_API uint32 GetTypeHash(const FShaderValueType& InShaderValueType);
 
 /* Fully describes the name and type of a parameter a shader exposes. */
 USTRUCT()
-struct FShaderParamTypeDefinition
+struct ENGINE_API FShaderParamTypeDefinition
 {
 	GENERATED_BODY()
 
@@ -294,6 +294,7 @@ public:
 		);
 };
 
+
 /* Describes a shader function signature. */
 USTRUCT()
 struct FShaderFunctionDefinition
