Added sky atmosphere fonctionality to recover the atmosphere transmittance.

diff --git a/Engine/Source/Runtime/Engine/Classes/Components/SkyAtmosphereComponent.h b/Engine/Source/Runtime/Engine/Classes/Components/SkyAtmosphereComponent.h
--- a/Engine/Source/Runtime/Engine/Classes/Components/SkyAtmosphereComponent.h
+++ b/Engine/Source/Runtime/Engine/Classes/Components/SkyAtmosphereComponent.h
@@ -199,6 +199,9 @@ class USkyAtmosphereComponent : public USceneComponent
 	UFUNCTION(BlueprintCallable, Category = "Rendering")
 	ENGINE_API void SetHeightFogContribution(float NewValue);
 
+	UFUNCTION(BlueprintCallable, Category = "Utilities", meta = (DisplayName = "Get Atmosphere Transmitance On Ground At Planet Top"))
+	ENGINE_API FLinearColor GetAtmosphereTransmitanceOnGroundAtPlanetTop(UDirectionalLightComponent* DirectionalLight);
+
 	// This is used to position the SkyAtmosphere similarly to the deprecated AtmosphericFog component
 	void SetPositionToMatchDeprecatedAtmosphericFog();
 
diff --git a/Engine/Source/Runtime/Engine/Private/Components/SkyAtmosphereComponent.cpp b/Engine/Source/Runtime/Engine/Private/Components/SkyAtmosphereComponent.cpp
--- a/Engine/Source/Runtime/Engine/Private/Components/SkyAtmosphereComponent.cpp
+++ b/Engine/Source/Runtime/Engine/Private/Components/SkyAtmosphereComponent.cpp
@@ -13,6 +13,7 @@
 #include "Rendering/SkyAtmosphereCommonData.h"
 #include "UObject/UObjectIterator.h"
 #include "UObject/ConstructorHelpers.h"
+#include "Components/DirectionalLightComponent.h"
 #include "UObject/UE5MainStreamObjectVersion.h"
 
 #if WITH_EDITOR
@@ -349,6 +350,17 @@ SKY_DECLARE_BLUEPRINT_SETFUNCTION_LINEARCOEFFICIENT(SkyLuminanceFactor);
 SKY_DECLARE_BLUEPRINT_SETFUNCTION(float, AerialPespectiveViewDistanceScale);
 SKY_DECLARE_BLUEPRINT_SETFUNCTION(float, HeightFogContribution);
 
+FLinearColor USkyAtmosphereComponent::GetAtmosphereTransmitanceOnGroundAtPlanetTop(UDirectionalLightComponent* DirectionalLight)
+{
+	if(DirectionalLight != nullptr)
+	{
+		FAtmosphereSetup AtmosphereSetup(*this);
+		const FLinearColor TransmittanceAtDirLight = AtmosphereSetup.GetTransmittanceAtGroundLevel(-DirectionalLight->GetDirection());
+		return TransmittanceAtDirLight;
+	}
+	return FLinearColor::White;
+}
+
 /*=============================================================================
 	ASkyAtmosphere implementation.
 =============================================================================*/
