Fixed leak in FTexture2DArrayResource (InitialMipData was always leaked after CreateTexture()).

diff --git a/Engine/Source/Runtime/Engine/Private/Rendering/Texture2DArrayResource.cpp b/Engine/Source/Runtime/Engine/Private/Rendering/Texture2DArrayResource.cpp
--- a/Engine/Source/Runtime/Engine/Private/Rendering/Texture2DArrayResource.cpp
+++ b/Engine/Source/Runtime/Engine/Private/Rendering/Texture2DArrayResource.cpp
@@ -95,7 +95,7 @@ void FTexture2DArrayResource::CreateTexture()
 	}
 		
 	SliceMipDataViews.Empty();
-	InitialMipData.Release();
+	InitialMipData.Reset();
 }
 
 void FTexture2DArrayResource::CreatePartiallyResidentTexture()
