*Fix for navmesh warning when having ISMC components containing a static mesh but no instances. IsNavigationrelevant will now check for the instance count as well.

Tested by QA

diff --git a/Engine/Source/Runtime/Engine/Classes/Components/InstancedStaticMeshComponent.h b/Engine/Source/Runtime/Engine/Classes/Components/InstancedStaticMeshComponent.h
--- a/Engine/Source/Runtime/Engine/Classes/Components/InstancedStaticMeshComponent.h
+++ b/Engine/Source/Runtime/Engine/Classes/Components/InstancedStaticMeshComponent.h
@@ -358,6 +358,7 @@ public:
 	//~ Begin UNavRelevantInterface Interface
 	virtual void GetNavigationData(FNavigationRelevantData& Data) const override;
 	virtual FBox GetNavigationBounds() const override;
+	virtual bool IsNavigationRelevant() const override;
 	//~ End UPrimitiveComponent Interface
 
 	//~ Begin UObject Interface
diff --git a/Engine/Source/Runtime/Engine/Private/InstancedStaticMesh.cpp b/Engine/Source/Runtime/Engine/Private/InstancedStaticMesh.cpp
--- a/Engine/Source/Runtime/Engine/Private/InstancedStaticMesh.cpp
+++ b/Engine/Source/Runtime/Engine/Private/InstancedStaticMesh.cpp
@@ -3712,6 +3712,11 @@ FBox UInstancedStaticMeshComponent::GetNavigationBounds() const
 	return CalcBounds(GetComponentTransform()).GetBox();
 }
 
+bool UInstancedStaticMeshComponent::IsNavigationRelevant() const
+{
+	return GetInstanceCount() > 0 && Super::IsNavigationRelevant();
+}
+
 void UInstancedStaticMeshComponent::GetNavigationPerInstanceTransforms(const FBox& AreaBox, TArray<FTransform>& InstanceData) const
 {
 	for (const auto& InstancedData : PerInstanceSMData)
