Don't allow runtime virtual texture samples in a material when writing to a runtime virtual texture.

diff --git a/Engine/Source/Runtime/Engine/Private/Materials/HLSLMaterialTranslator.cpp b/Engine/Source/Runtime/Engine/Private/Materials/HLSLMaterialTranslator.cpp
--- a/Engine/Source/Runtime/Engine/Private/Materials/HLSLMaterialTranslator.cpp
+++ b/Engine/Source/Runtime/Engine/Private/Materials/HLSLMaterialTranslator.cpp
@@ -695,31 +695,31 @@ bool FHLSLMaterialTranslator::Translate()
 		}
 		else
 		{
-			// Normal must always be compiled first; this will ensure its chunk calculations are the first to be added
-			{
-				// Verify that start chunk is 0
-				check(SharedPropertyCodeChunks[NormalShaderFrequency].Num() == 0);
-				Chunk[MP_Normal]					= Material->CompilePropertyAndSetMaterialProperty(MP_Normal					,this);
-				NormalCodeChunkEnd = SharedPropertyCodeChunks[NormalShaderFrequency].Num();
-			}
-
-			// Rest of properties
-			Chunk[MP_EmissiveColor]					= Material->CompilePropertyAndSetMaterialProperty(MP_EmissiveColor			,this);
-			Chunk[MP_DiffuseColor]					= Material->CompilePropertyAndSetMaterialProperty(MP_DiffuseColor			,this);
-			Chunk[MP_SpecularColor]					= Material->CompilePropertyAndSetMaterialProperty(MP_SpecularColor			,this);
-			Chunk[MP_BaseColor]						= Material->CompilePropertyAndSetMaterialProperty(MP_BaseColor				,this);
-			Chunk[MP_Metallic]						= Material->CompilePropertyAndSetMaterialProperty(MP_Metallic				,this);
-			Chunk[MP_Specular]						= Material->CompilePropertyAndSetMaterialProperty(MP_Specular				,this);
-			Chunk[MP_Roughness]						= Material->CompilePropertyAndSetMaterialProperty(MP_Roughness				,this);
-			Chunk[MP_Anisotropy]					= Material->CompilePropertyAndSetMaterialProperty(MP_Anisotropy				,this);
-			Chunk[MP_Opacity]						= Material->CompilePropertyAndSetMaterialProperty(MP_Opacity				,this);
-			Chunk[MP_OpacityMask]					= Material->CompilePropertyAndSetMaterialProperty(MP_OpacityMask			,this);
-			Chunk[MP_Tangent]						= Material->CompilePropertyAndSetMaterialProperty(MP_Tangent				,this);
-			Chunk[MP_WorldPositionOffset]			= Material->CompilePropertyAndSetMaterialProperty(MP_WorldPositionOffset	,this);
-
-			// Make sure to compile this property before using ShadingModelsFromCompilation
-			Chunk[MP_ShadingModel]					= Material->CompilePropertyAndSetMaterialProperty(MP_ShadingModel			,this);
-		
+		// Normal must always be compiled first; this will ensure its chunk calculations are the first to be added
+		{
+			// Verify that start chunk is 0
+			check(SharedPropertyCodeChunks[NormalShaderFrequency].Num() == 0);
+			Chunk[MP_Normal]					= Material->CompilePropertyAndSetMaterialProperty(MP_Normal					,this);
+			NormalCodeChunkEnd = SharedPropertyCodeChunks[NormalShaderFrequency].Num();
+		}
+
+		// Rest of properties
+		Chunk[MP_EmissiveColor]					= Material->CompilePropertyAndSetMaterialProperty(MP_EmissiveColor			,this);
+		Chunk[MP_DiffuseColor]					= Material->CompilePropertyAndSetMaterialProperty(MP_DiffuseColor			,this);
+		Chunk[MP_SpecularColor]					= Material->CompilePropertyAndSetMaterialProperty(MP_SpecularColor			,this);
+		Chunk[MP_BaseColor]						= Material->CompilePropertyAndSetMaterialProperty(MP_BaseColor				,this);
+		Chunk[MP_Metallic]						= Material->CompilePropertyAndSetMaterialProperty(MP_Metallic				,this);
+		Chunk[MP_Specular]						= Material->CompilePropertyAndSetMaterialProperty(MP_Specular				,this);
+		Chunk[MP_Roughness]						= Material->CompilePropertyAndSetMaterialProperty(MP_Roughness				,this);
+		Chunk[MP_Anisotropy]					= Material->CompilePropertyAndSetMaterialProperty(MP_Anisotropy				,this);
+		Chunk[MP_Opacity]						= Material->CompilePropertyAndSetMaterialProperty(MP_Opacity				,this);
+		Chunk[MP_OpacityMask]					= Material->CompilePropertyAndSetMaterialProperty(MP_OpacityMask			,this);
+		Chunk[MP_Tangent]						= Material->CompilePropertyAndSetMaterialProperty(MP_Tangent				,this);
+		Chunk[MP_WorldPositionOffset]			= Material->CompilePropertyAndSetMaterialProperty(MP_WorldPositionOffset	,this);
+
+		// Make sure to compile this property before using ShadingModelsFromCompilation
+		Chunk[MP_ShadingModel]					= Material->CompilePropertyAndSetMaterialProperty(MP_ShadingModel			,this);
+			
 			Chunk[MP_FrontMaterial]					= Material->CompilePropertyAndSetMaterialProperty(MP_FrontMaterial			,this);
 		}
 
@@ -737,46 +737,46 @@ bool FHLSLMaterialTranslator::Translate()
 		
 		if (!bEnableExecutionFlow)
 		{
-			if (Domain == MD_Volume || (Domain == MD_Surface && IsSubsurfaceShadingModel(MaterialShadingModels)))
-			{
-				// Note we don't test for the blend mode as you can have a translucent material using the subsurface shading model
+		if (Domain == MD_Volume || (Domain == MD_Surface && IsSubsurfaceShadingModel(MaterialShadingModels)))
+		{
+			// Note we don't test for the blend mode as you can have a translucent material using the subsurface shading model
 
-				// another ForceCast as CompilePropertyAndSetMaterialProperty() can return MCT_Float which we don't want here
-				int32 SubsurfaceColor = Material->CompilePropertyAndSetMaterialProperty(MP_SubsurfaceColor, this);
-				SubsurfaceColor = ForceCast(SubsurfaceColor, FMaterialAttributeDefinitionMap::GetValueType(MP_SubsurfaceColor), MFCF_ExactMatch | MFCF_ReplicateValue);
+			// another ForceCast as CompilePropertyAndSetMaterialProperty() can return MCT_Float which we don't want here
+			int32 SubsurfaceColor = Material->CompilePropertyAndSetMaterialProperty(MP_SubsurfaceColor, this);
+			SubsurfaceColor = ForceCast(SubsurfaceColor, FMaterialAttributeDefinitionMap::GetValueType(MP_SubsurfaceColor), MFCF_ExactMatch | MFCF_ReplicateValue);
 
-				static FName NameSubsurfaceProfile(TEXT("__SubsurfaceProfile"));
+			static FName NameSubsurfaceProfile(TEXT("__SubsurfaceProfile"));
 
-				// 1.0f is is a not used profile - later this gets replaced with the actual profile
-				int32 CodeSubsurfaceProfile = ForceCast(ScalarParameter(NameSubsurfaceProfile, 1.0f), MCT_Float1);
+			// 1.0f is is a not used profile - later this gets replaced with the actual profile
+			int32 CodeSubsurfaceProfile = ForceCast(ScalarParameter(NameSubsurfaceProfile, 1.0f), MCT_Float1);
 
 				Chunk[MP_SubsurfaceColor] = AppendVector(SubsurfaceColor, CodeSubsurfaceProfile);
-			}
+		}
 
 			Chunk[MP_CustomData0] = Material->CompilePropertyAndSetMaterialProperty(MP_CustomData0, this);
 			Chunk[MP_CustomData1] = Material->CompilePropertyAndSetMaterialProperty(MP_CustomData1, this);
 			Chunk[MP_AmbientOcclusion] = Material->CompilePropertyAndSetMaterialProperty(MP_AmbientOcclusion, this);
 
-			if (IsTranslucentBlendMode(BlendMode) || MaterialShadingModels.HasShadingModel(MSM_SingleLayerWater))
-			{
-				int32 UserRefraction = ForceCast(Material->CompilePropertyAndSetMaterialProperty(MP_Refraction, this), MCT_Float1);
-				int32 RefractionDepthBias = ForceCast(ScalarParameter(FName(TEXT("RefractionDepthBias")), Material->GetRefractionDepthBiasValue()), MCT_Float1);
+		if (IsTranslucentBlendMode(BlendMode) || MaterialShadingModels.HasShadingModel(MSM_SingleLayerWater))
+		{
+			int32 UserRefraction = ForceCast(Material->CompilePropertyAndSetMaterialProperty(MP_Refraction, this), MCT_Float1);
+			int32 RefractionDepthBias = ForceCast(ScalarParameter(FName(TEXT("RefractionDepthBias")), Material->GetRefractionDepthBiasValue()), MCT_Float1);
 
-				Chunk[MP_Refraction] = AppendVector(UserRefraction, RefractionDepthBias);
-			}
+			Chunk[MP_Refraction] = AppendVector(UserRefraction, RefractionDepthBias);
+		}
 
-			if (bCompileForComputeShader)
-			{
+		if (bCompileForComputeShader)
+		{
 				Chunk[CompiledMP_EmissiveColorCS] = Material->CompilePropertyAndSetMaterialProperty(MP_EmissiveColor, this, SF_Compute);
-			}
+		}
 
-			if (Chunk[MP_WorldPositionOffset] != INDEX_NONE)
-			{
-				// Only calculate previous WPO if there is a current WPO
-				Chunk[CompiledMP_PrevWorldPositionOffset] = Material->CompilePropertyAndSetMaterialProperty(MP_WorldPositionOffset, this, SF_Vertex, true);
-			}
+		if (Chunk[MP_WorldPositionOffset] != INDEX_NONE)
+		{
+			// Only calculate previous WPO if there is a current WPO
+			Chunk[CompiledMP_PrevWorldPositionOffset] = Material->CompilePropertyAndSetMaterialProperty(MP_WorldPositionOffset, this, SF_Vertex, true);
+		}
 
-			Chunk[MP_PixelDepthOffset] = Material->CompilePropertyAndSetMaterialProperty(MP_PixelDepthOffset, this);
+		Chunk[MP_PixelDepthOffset] = Material->CompilePropertyAndSetMaterialProperty(MP_PixelDepthOffset, this);
 		}
 
 		ResourcesString = TEXT("");
@@ -855,17 +855,17 @@ bool FHLSLMaterialTranslator::Translate()
 
 		if (!bEnableExecutionFlow)
 		{
-			for (uint32 CustomUVIndex = MP_CustomizedUVs0; CustomUVIndex <= MP_CustomizedUVs7; CustomUVIndex++)
+		for (uint32 CustomUVIndex = MP_CustomizedUVs0; CustomUVIndex <= MP_CustomizedUVs7; CustomUVIndex++)
+		{
+			// Only compile custom UV inputs for UV channels requested by the pixel shader inputs
+			// Any unconnected inputs will have a texcoord generated for them in Material->CompileProperty, which will pass through the vertex (uncustomized) texture coordinates
+			// Note: this is using NumUserTexCoords, which is set by translating all the pixel properties above
+			if (CustomUVIndex - MP_CustomizedUVs0 < SavedNumUserTexCoords)
 			{
-				// Only compile custom UV inputs for UV channels requested by the pixel shader inputs
-				// Any unconnected inputs will have a texcoord generated for them in Material->CompileProperty, which will pass through the vertex (uncustomized) texture coordinates
-				// Note: this is using NumUserTexCoords, which is set by translating all the pixel properties above
-				if (CustomUVIndex - MP_CustomizedUVs0 < SavedNumUserTexCoords)
-				{
-					Chunk[CustomUVIndex] = Material->CompilePropertyAndSetMaterialProperty((EMaterialProperty)CustomUVIndex, this);
-				}
+				Chunk[CustomUVIndex] = Material->CompilePropertyAndSetMaterialProperty((EMaterialProperty)CustomUVIndex, this);
 			}
 		}
+		}
 
 		// Output the implementation for any custom expressions we will call below.
 		for (int32 ExpressionIndex = 0; ExpressionIndex < CustomExpressions.Num(); ExpressionIndex++)
@@ -1036,7 +1036,7 @@ bool FHLSLMaterialTranslator::Translate()
 		if (Material->IsStrataMaterial())
 		{
 			if (Domain!= MD_Volume && !(BlendMode == BLEND_Translucent || BlendMode == BLEND_Opaque || BlendMode == BLEND_Masked))
-			{
+		{
 				Errorf(TEXT("Strata materials must be an opaque or translucent surface, or a volume."));
 			}
 		}
@@ -1086,7 +1086,7 @@ bool FHLSLMaterialTranslator::Translate()
 
 		MaterialCompilationOutput.NumUsedUVScalars = GetNumUserTexCoords() * 2;
 		MaterialCompilationOutput.NumUsedCustomInterpolatorScalars = CurrentCustomVertexInterpolatorOffset;
-		
+
 		if (bEnableExecutionFlow)
 		{
 			if (VertexAttributesChunk != INDEX_NONE)
@@ -1111,75 +1111,75 @@ bool FHLSLMaterialTranslator::Translate()
 			{
 				ECompiledPartialDerivativeVariation Variation = (ECompiledPartialDerivativeVariation)VariationIter;
 
-				// Do Normal Chunk first
-				{
-					GetFixedParameterCode(
-						0,
-						NormalCodeChunkEnd,
-						Chunk[MP_Normal],
-						SharedPropertyCodeChunks[NormalShaderFrequency],
+		// Do Normal Chunk first
+		{
+			GetFixedParameterCode(
+				0,
+				NormalCodeChunkEnd,
+				Chunk[MP_Normal],
+				SharedPropertyCodeChunks[NormalShaderFrequency],
 						DerivativeVariations[Variation].TranslatedCodeChunkDefinitions[MP_Normal],
 						DerivativeVariations[Variation].TranslatedCodeChunks[MP_Normal],
 						Variation);
 
-					// Always gather MP_Normal definitions as they can be shared by other properties
+			// Always gather MP_Normal definitions as they can be shared by other properties
 					if (DerivativeVariations[Variation].TranslatedCodeChunkDefinitions[MP_Normal].IsEmpty())
-					{
+			{
 						DerivativeVariations[Variation].TranslatedCodeChunkDefinitions[MP_Normal] = GetDefinitions(SharedPropertyCodeChunks[NormalShaderFrequency], 0, NormalCodeChunkEnd, Variation);
-					}
-				}
+			}
+		}
 
-				// Now the rest, skipping Normal
+		// Now the rest, skipping Normal
 				for (uint32 PropertyId = 0; PropertyId < MP_MAX; ++PropertyId)
-				{
-					if (PropertyId == MP_MaterialAttributes || PropertyId == MP_Normal || PropertyId == MP_CustomOutput)
-					{
-						continue;
-					}
+		{
+			if (PropertyId == MP_MaterialAttributes || PropertyId == MP_Normal || PropertyId == MP_CustomOutput)
+			{
+				continue;
+			}
 
-					const EShaderFrequency PropertyShaderFrequency = FMaterialAttributeDefinitionMap::GetShaderFrequency((EMaterialProperty)PropertyId);
+			const EShaderFrequency PropertyShaderFrequency = FMaterialAttributeDefinitionMap::GetShaderFrequency((EMaterialProperty)PropertyId);
 
-					int32 StartChunk = 0;
-					if (PropertyShaderFrequency == NormalShaderFrequency && SharedPixelProperties[PropertyId])
-					{
-						// When processing shared properties, do not generate the code before the Normal was generated as those are already handled
-						StartChunk = NormalCodeChunkEnd;
-					}
+			int32 StartChunk = 0;
+			if (PropertyShaderFrequency == NormalShaderFrequency && SharedPixelProperties[PropertyId])
+			{
+				// When processing shared properties, do not generate the code before the Normal was generated as those are already handled
+				StartChunk = NormalCodeChunkEnd;
+			}
 
-					GetFixedParameterCode(
-						StartChunk,
-						SharedPropertyCodeChunks[PropertyShaderFrequency].Num(),
-						Chunk[PropertyId],
-						SharedPropertyCodeChunks[PropertyShaderFrequency],
+			GetFixedParameterCode(
+				StartChunk,
+				SharedPropertyCodeChunks[PropertyShaderFrequency].Num(),
+				Chunk[PropertyId],
+				SharedPropertyCodeChunks[PropertyShaderFrequency],
 						DerivativeVariations[Variation].TranslatedCodeChunkDefinitions[PropertyId],
 						DerivativeVariations[Variation].TranslatedCodeChunks[PropertyId],
 						Variation);
 
-				}
+		}
 
 				for (uint32 PropertyId = MP_MAX; PropertyId < CompiledMP_MAX; ++PropertyId)
-				{
+		{
 					switch (PropertyId)
-					{
-					case CompiledMP_EmissiveColorCS:
-						if (bCompileForComputeShader)
-						{
+			{
+			case CompiledMP_EmissiveColorCS:
+			    if (bCompileForComputeShader)
+				{
 							GetFixedParameterCode(Chunk[PropertyId], SharedPropertyCodeChunks[SF_Compute], DerivativeVariations[Variation].TranslatedCodeChunkDefinitions[PropertyId], DerivativeVariations[Variation].TranslatedCodeChunks[PropertyId], Variation);
-						}
-						break;
-					case CompiledMP_PrevWorldPositionOffset:
-					{
+				}
+				break;
+			case CompiledMP_PrevWorldPositionOffset:
+				{
 						GetFixedParameterCode(Chunk[PropertyId], SharedPropertyCodeChunks[SF_Vertex], DerivativeVariations[Variation].TranslatedCodeChunkDefinitions[PropertyId], DerivativeVariations[Variation].TranslatedCodeChunks[PropertyId], Variation);
-					}
-					break;
-					default: check(0);
-						break;
-					}
 				}
+				break;
+			default: check(0);
+				break;
+			}
+		}
 
-				// Output the implementation for any custom output expressions
+		// Output the implementation for any custom output expressions
 				for (int32 ExpressionIndex = 0; ExpressionIndex < DerivativeVariations[Variation].CustomOutputImplementations.Num(); ExpressionIndex++)
-				{
+		{
 					ResourcesString += DerivativeVariations[Variation].CustomOutputImplementations[ExpressionIndex] + "\r\n\r\n";
 				}
 			}
@@ -1905,23 +1905,23 @@ void FHLSLMaterialTranslator::GetSharedInputsMaterialCode(FString& PixelMembersD
 
 			if (!bEnableExecutionFlow)
 			{
-				// Normal requires its own separate initializer
-				if (Property == MP_Normal)
-				{
+			// Normal requires its own separate initializer
+			if (Property == MP_Normal)
+			{
 					NormalInitializerValue = FString::Printf(TEXT("\tPixelMaterialInputs.%s = %s;\n"), *PropertyName, *DerivativeVariations[DerivativeVariation].TranslatedCodeChunks[Property]);
-				}
-				else
-				{
+			}
+			else
+			{
 					if (DerivativeVariations[DerivativeVariation].TranslatedCodeChunkDefinitions[Property].Len() > 0)
+				{
+					if (LastProperty >= 0)
 					{
-						if (LastProperty >= 0)
-						{
-							// Verify that all code chunks have the same contents
+						// Verify that all code chunks have the same contents
 							check(DerivativeVariations[DerivativeVariation].TranslatedCodeChunkDefinitions[Property].Len() == DerivativeVariations[DerivativeVariation].TranslatedCodeChunkDefinitions[LastProperty].Len());
-						}
-
-						LastProperty = Property;
 					}
+
+					LastProperty = Property;
+				}
 				}
 
 				PixelInputInitializerValues += FString::Printf(TEXT("\tPixelMaterialInputs.%s = %s;\n"), *PropertyName, *DerivativeVariations[DerivativeVariation].TranslatedCodeChunks[Property]);
@@ -2090,19 +2090,19 @@ FString FHLSLMaterialTranslator::GetMaterialShaderCode()
 		}
 		else
 		{
-			if (CustomUVIndex == 0)
-			{
+		if (CustomUVIndex == 0)
+		{
 				CustomUVAssignments += DerivativeVariations[BaseDerivativeVariation].TranslatedCodeChunkDefinitions[MP_CustomizedUVs0 + CustomUVIndex];
-			}
+		}
 
 			if (DerivativeVariations[BaseDerivativeVariation].TranslatedCodeChunkDefinitions[MP_CustomizedUVs0 + CustomUVIndex].Len() > 0)
+		{
+			if (LastProperty >= 0)
 			{
-				if (LastProperty >= 0)
-				{
 					check(DerivativeVariations[BaseDerivativeVariation].TranslatedCodeChunkDefinitions[LastProperty].Len() == DerivativeVariations[BaseDerivativeVariation].TranslatedCodeChunkDefinitions[MP_CustomizedUVs0 + CustomUVIndex].Len());
-				}
-				LastProperty = MP_CustomizedUVs0 + CustomUVIndex;
 			}
+			LastProperty = MP_CustomizedUVs0 + CustomUVIndex;
+		}
 			CustomUVAssignments += FString::Printf(TEXT("\tOutTexCoords[%u] = %s;") LINE_TERMINATOR, CustomUVIndex, *DerivativeVariations[BaseDerivativeVariation].TranslatedCodeChunks[MP_CustomizedUVs0 + CustomUVIndex]);
 		}
 	}
@@ -2209,10 +2209,10 @@ FString FHLSLMaterialTranslator::GetMaterialShaderCode()
 		{
 			ECompiledPartialDerivativeVariation Variation = (ECompiledPartialDerivativeVariation)Iter;
 
-			// Initializers required for Normal
+	// Initializers required for Normal
 			LazyPrintf.PushParam(*DerivativeVariations[Variation].TranslatedCodeChunkDefinitions[MP_Normal]);
 			LazyPrintf.PushParam(*NormalAssignment[Variation]);
-			// Finally the rest of common code followed by assignment into each input
+	// Finally the rest of common code followed by assignment into each input
 			LazyPrintf.PushParam(*PixelMembersSetupAndAssignments[Variation]);
 		}
 	}
@@ -2381,7 +2381,7 @@ FString FHLSLMaterialTranslator::GetParameterCodeRaw(int32 Index, const TCHAR* D
 			check(AccessedCodeChunk.SymbolName.Len() > 0);
 			return AccessedCodeChunk.SymbolName;
 		}
-
+			
 		ReferencedCodeChunks.AddUnique(Index);
 			
 		// Return the symbol used to reference this code chunk
@@ -2529,7 +2529,7 @@ void FHLSLMaterialTranslator::GetScopeCode(int32 IndentLevel, int32 ScopeChunkIn
 
 	// Add code for any dependencies of the current chunk
 	for (int32 ReferencedChunkIndex : CodeChunk.ReferencedCodeChunks)
-	{
+{
 		GetScopeCode(IndentLevel, ReferencedChunkIndex, CodeChunks, EmittedChunks, OutValue);
 	}
 
@@ -2732,13 +2732,13 @@ int32 FHLSLMaterialTranslator::AddCodeChunkInner(uint64 Hash, const TCHAR* Forma
 		if (CodeIndex == INDEX_NONE)
 		{
 			CodeIndex = CurrentScopeChunks->Num();
-			// Allocate a local variable name
-			const FString SymbolName = CreateSymbolName(TEXT("Local"));
-			// Construct the definition string which stores the result in a temporary and adds a newline for readability
+		// Allocate a local variable name
+		const FString SymbolName = CreateSymbolName(TEXT("Local"));
+		// Construct the definition string which stores the result in a temporary and adds a newline for readability
 			const FString LocalVariableDefinitionFinite = FString("	") + HLSLTypeString(Type) + TEXT(" ") + SymbolName + TEXT(" = ") + FormattedCode + TEXT(";") + LINE_TERMINATOR;
 			// Construct the definition string which stores the result in a temporary and adds a newline for readability
 			const FString LocalVariableDefinitionAnalytic = FString("	") + HLSLTypeString(Type) + TEXT(" ") + SymbolName + TEXT(" = ") + FormattedCode + TEXT(";") + LINE_TERMINATOR;
-			// Adding a code chunk that creates a local variable
+		// Adding a code chunk that creates a local variable
 			new(*CurrentScopeChunks) FShaderCodeChunk(Hash, *LocalVariableDefinitionFinite, *LocalVariableDefinitionAnalytic, SymbolName, Type, DerivativeStatus, false);
 		}
 	}
@@ -2758,7 +2758,7 @@ int32 FHLSLMaterialTranslator::AddCodeChunkInner(uint64 Hash, const TCHAR* Forma
 	AddCodeChunkToCurrentScope(CodeIndex);
 	return CodeIndex;
 }
-
+			
 static inline uint32 GetTCharStringBytes(const TCHAR* String)
 {
 	uint32 Length = 0u;
@@ -2792,19 +2792,19 @@ int32 FHLSLMaterialTranslator::AddCodeChunkInnerDeriv(const TCHAR* FormattedCode
 		CodeIndex = CurrentScopeChunks->Num();
 		// Adding an inline code chunk, the definition will be the code to inline
 		new(*CurrentScopeChunks) FShaderCodeChunk(Hash, FormattedCodeFinite, FormattedCodeAnalytic, TEXT(""), Type, DerivativeStatus, true);
-	}
+}
 	// Can only create temporaries for certain types
 	else if ((Type & (MCT_Float | MCT_VTPageTableResult)) || Type == MCT_ShadingModel || Type == MCT_MaterialAttributes || Type == MCT_Strata)
 	{
 		// Check for existing
 		for (int32 i = 0; i < CurrentScopeChunks->Num(); ++i)
-		{
+{
 			if ((*CurrentScopeChunks)[i].Hash == Hash)
-			{
+	{
 				CodeIndex = i;
 				break;
-			}
-		}
+	}
+}
 
 		if(CodeIndex == INDEX_NONE)
 		{
@@ -2822,12 +2822,12 @@ int32 FHLSLMaterialTranslator::AddCodeChunkInnerDeriv(const TCHAR* FormattedCode
 	else
 	{
 		if (Type & MCT_Texture)
-		{
+{
 			return Errorf(TEXT("Operation not supported on a Texture"));
 		}
 
 		if (Type == MCT_StaticBool)
-		{
+	{
 			return Errorf(TEXT("Operation not supported on a Static Bool"));
 		}
 
@@ -3477,8 +3477,8 @@ int32 FHLSLMaterialTranslator::CallExpression(FMaterialExpressionKey ExpressionK
 		ReferencedCodeChunks.Reset();
 	}
 
-	return Result;
-}
+		return Result;
+	}
 
 int32 FHLSLMaterialTranslator::CallExpressionExec(UMaterialExpression* Expression)
 {
@@ -3721,7 +3721,7 @@ int32 FHLSLMaterialTranslator::ValidCast(int32 Code, EMaterialValueType DestType
 					FString DDYCode = FString(HLSLTypeString(DestType)) + TEXT("(") + DerivStringDDY;
 
 					for (uint32 i = 0; i < NumPadComponents; i++)
-					{
+			{
 						DDXCode += TEXT(", ") + DerivStringDDX;
 						DDYCode += TEXT(", ") + DerivStringDDY;
 					}
@@ -4058,7 +4058,7 @@ int32 FHLSLMaterialTranslator::ViewProperty(EMaterialExposedViewProperty Propert
 
 	// Resolved templated code
 	Code.ReplaceInline(TEXT("<PREV>"), bCompilingPreviousFrame ? TEXT("Prev") : TEXT(""));
-
+		
 	if (InvProperty && !PropertyMeta.InvPropertyCode)
 	{
 		// fall back to compute the property's inverse from PropertyCode
@@ -4161,10 +4161,10 @@ int32 FHLSLMaterialTranslator::Sine(int32 X)
 		}
 		else
 		{
-			return AddCodeChunk(GetParameterType(X),TEXT("sin(%s)"),*GetParameterCode(X));
-		}
+		return AddCodeChunk(GetParameterType(X),TEXT("sin(%s)"),*GetParameterCode(X));
 	}
 }
+}
 
 int32 FHLSLMaterialTranslator::Cosine(int32 X)
 {
@@ -4185,10 +4185,10 @@ int32 FHLSLMaterialTranslator::Cosine(int32 X)
 		}
 		else
 		{
-			return AddCodeChunk(GetParameterType(X),TEXT("cos(%s)"),*GetParameterCode(X));
-		}
+		return AddCodeChunk(GetParameterType(X),TEXT("cos(%s)"),*GetParameterCode(X));
 	}
 }
+}
 
 int32 FHLSLMaterialTranslator::Tangent(int32 X)
 {
@@ -4209,10 +4209,10 @@ int32 FHLSLMaterialTranslator::Tangent(int32 X)
 		}
 		else
 		{
-			return AddCodeChunk(GetParameterType(X),TEXT("tan(%s)"),*GetParameterCode(X));
-		}
+		return AddCodeChunk(GetParameterType(X),TEXT("tan(%s)"),*GetParameterCode(X));
 	}
 }
+}
 
 int32 FHLSLMaterialTranslator::Arcsine(int32 X)
 {
@@ -4233,10 +4233,10 @@ int32 FHLSLMaterialTranslator::Arcsine(int32 X)
 		}
 		else
 		{
-			return AddCodeChunk(GetParameterType(X),TEXT("asin(%s)"),*GetParameterCode(X));
-		}
+		return AddCodeChunk(GetParameterType(X),TEXT("asin(%s)"),*GetParameterCode(X));
 	}
 }
+}
 
 int32 FHLSLMaterialTranslator::ArcsineFast(int32 X)
 {
@@ -4563,11 +4563,11 @@ int32 FHLSLMaterialTranslator::Abs(int32 X)
 	else
 	{
 		if (IsAnalyticDerivEnabled())
-		{
+	{
 			return DerivativeAutogen.GenerateExpressionFunc1(*this, FMaterialDerivativeAutogen::EFunc1::Abs, X);
-		}
-		else
-		{
+	}
+	else
+	{
 			return AddCodeChunk(GetParameterType(X), TEXT("abs(%s)"), *GetParameterCode(X));
 		}
 	}
@@ -5114,8 +5114,8 @@ int32 FHLSLMaterialTranslator::TextureCoordinate(uint32 CoordinateIndex, bool Un
 	}
 	else
 	{
-		// Note: inlining is important so that on GLES devices, where half precision is used in the pixel shader, 
-		// The UV does not get assigned to a half temporary in cases where the texture sample is done directly from interpolated UVs
+	// Note: inlining is important so that on GLES devices, where half precision is used in the pixel shader, 
+	// The UV does not get assigned to a half temporary in cases where the texture sample is done directly from interpolated UVs
 		return AddInlinedCodeChunk(MCT_Float2, *SampleCodeFinite, CoordinateIndex);
 	}
 }
@@ -5432,7 +5432,7 @@ int32 FHLSLMaterialTranslator::TextureSample(
 	{
 		if (Material->GetMaterialDomain() != MD_Surface && Material->GetMaterialDomain() != MD_DeferredDecal)
 		{
-			return Errorf(TEXT("Sampling a virtual texture is currently only supported inside surface and decal shaders."));
+			return Errorf(TEXT("Sampling a virtual texture is currently only supported for surface and decal materials."));
 		}
 	}
 
@@ -5489,7 +5489,7 @@ int32 FHLSLMaterialTranslator::TextureSample(
 	{
 		AutomaticViewMipBias = false;
 	}
-	
+
 	FString TextureName;
 	int32 VirtualTextureIndex = INDEX_NONE;
 	if (TextureType == MCT_TextureCube)
@@ -5501,13 +5501,13 @@ int32 FHLSLMaterialTranslator::TextureSample(
 		TextureName = CoerceParameter(TextureIndex, MCT_Texture2DArray);
 	}
 	else if (TextureType == MCT_VolumeTexture)
-	{
+		{
 		TextureName = CoerceParameter(TextureIndex, MCT_VolumeTexture);
-	}
+		}
 	else if (TextureType == MCT_TextureExternal)
-	{
+		{
 		TextureName = CoerceParameter(TextureIndex, MCT_TextureExternal);
-	}
+		}
 	else if (bVirtualTexture)
 	{
 		// Note, this does not really do anything (by design) other than adding it to the UniformExpressionSet
@@ -5522,25 +5522,25 @@ int32 FHLSLMaterialTranslator::TextureSample(
 		if (TextureUniformExpression == nullptr)
 		{
 			return Errorf(TEXT("The provided uniform expression is not a texture"));
-		}
+	}
 
 		VirtualTextureIndex = UniformTextureExpressions[(uint32)EMaterialTextureParameterType::Virtual].Find(TextureUniformExpression);
 		check(UniformTextureExpressions[(uint32)EMaterialTextureParameterType::Virtual].IsValidIndex(VirtualTextureIndex));
 
 		if (SamplerSource != SSM_FromTextureAsset)
-		{
+	{
 			// VT doesn't care if the shared sampler is wrap or clamp this is handled in the shader explicitly by our code so we still inherit this from the texture
 			TextureName += FString::Printf(TEXT("Material.VirtualTexturePhysical_%d, GetMaterialSharedSampler(Material.VirtualTexturePhysical_%dSampler, View.SharedBilinearClampedSampler)")
 				, VirtualTextureIndex, VirtualTextureIndex);
-		}
+	}
 		else
-		{
+	{
 			TextureName += FString::Printf(TEXT("Material.VirtualTexturePhysical_%d, Material.VirtualTexturePhysical_%dSampler")
 				, VirtualTextureIndex, VirtualTextureIndex);
 		}
 
 		NumVtSamples++;
- 	}
+	}
 	else // MCT_Texture2D
 	{
 		TextureName = CoerceParameter(TextureIndex, MCT_Texture2D);
@@ -5552,28 +5552,28 @@ int32 FHLSLMaterialTranslator::TextureSample(
 	if (!bVirtualTexture) //VT does not have explict samplers (and always requires manual view mip bias)
 	{
 		if (SamplerSource == SSM_FromTextureAsset)
-		{
+	{
 			SamplerStateCode = FString::Printf(TEXT("%sSampler"), *TextureName);
-		}
+	}
 		else if (SamplerSource == SSM_Wrap_WorldGroupSettings)
-		{
+	{
 			// Use the shared sampler to save sampler slots
 			SamplerStateCode = FString::Printf(	TEXT("GetMaterialSharedSampler(%sSampler,%s)"),
 												*TextureName, AutomaticViewMipBias ? TEXT("View.MaterialTextureBilinearWrapedSampler") : TEXT("Material.Wrap_WorldGroupSettings"));
 			RequiresManualViewMipBias = false;
-		}
+	}
 		else if (SamplerSource == SSM_Clamp_WorldGroupSettings)
-		{
+	{
 			// Use the shared sampler to save sampler slots
 			SamplerStateCode = FString::Printf(	TEXT("GetMaterialSharedSampler(%sSampler,%s)"),
 												*TextureName, AutomaticViewMipBias ? TEXT("View.MaterialTextureBilinearClampedSampler") : TEXT("Material.Clamp_WorldGroupSettings"));
 			RequiresManualViewMipBias = false;
 		}
 	}
-
+		
 	const EMaterialValueType UVsType = (TextureType == MCT_TextureCube || TextureType == MCT_Texture2DArray || TextureType == MCT_VolumeTexture) ? MCT_Float3 : MCT_Float2;
 	const EDerivativeStatus UvDerivativeStatus = GetDerivativeStatus(CoordinateIndex);
-
+	
 	if (RequiresManualViewMipBias)
 	{
 		if (MipValueMode == TMVM_Derivative)
@@ -5612,36 +5612,36 @@ int32 FHLSLMaterialTranslator::TextureSample(
 
 	FString TextureTypeName;
 	if (TextureType == MCT_TextureCube)
-	{
+		{
 		TextureTypeName = TEXT("TextureCubeSample");
-	}
+		}
 	else if (TextureType == MCT_Texture2DArray)
-	{
+		{
 		TextureTypeName = TEXT("Texture2DArraySample");
-	}
+		}
 	else if (TextureType == MCT_VolumeTexture)
-	{
+		{
 		TextureTypeName = TEXT("Texture3DSample");
-	}
+		}
 	else if (TextureType == MCT_TextureExternal)
-	{
+		{
 		TextureTypeName = TEXT("TextureExternalSample");
-	}
+		}
 	else if (bVirtualTexture)
-	{
+		{
 		TextureTypeName = TEXT("TextureVirtualSample");
-	}
+		}
 	else // MCT_Texture2D
-	{
+		{
 		TextureTypeName = TEXT("Texture2DSample");
-	}
+		}
 	
 	const FString UVs = CoerceParameter(CoordinateIndex, UVsType);
 	const bool bStoreTexCoordScales = ShaderFrequency == SF_Pixel && TextureReferenceIndex != INDEX_NONE;
 	const bool bDecal = ShaderFrequency == SF_Pixel && Material->GetMaterialDomain() == MD_DeferredDecal && MipValueMode == TMVM_None;
 
 	if (IsDebugTextureSampleEnabled() && !IsDerivativeValid(UvDerivativeStatus))
-	{
+		{
 		UE_LOG(LogMaterial, Warning, TEXT("Unknown derivatives: '%s'[%s]: %s"), *Material->GetDebugName(), *Material->GetAssetPath().ToString(), *AtParameterCodeChunk(CoordinateIndex).DefinitionAnalytic);
 	}
 
@@ -5663,24 +5663,24 @@ int32 FHLSLMaterialTranslator::TextureSample(
 			UV_Ddx = UVAnalytic + TEXT(".Ddx");
 			UV_Ddy = UVAnalytic + TEXT(".Ddy");
 			if (MipValueMode == TMVM_MipBias)
-			{
+		{
 				UV_Scale = FString::Printf(TEXT("exp2(%s)"), *MipValue0Code);
 				UV_Ddx = FString::Printf(TEXT("(%s)*exp2(%s)"), *UV_Ddx, *MipValue0Code);
 				UV_Ddy = FString::Printf(TEXT("(%s)*exp2(%s)"), *UV_Ddy, *MipValue0Code);
-			}
+		}
 		}
 		else
 		{
 			UV_Value = GetParameterCode(CoordinateIndex);
-		}
+ 	}
 	}
 
 	int32 SamplingCodeIndex = INDEX_NONE;
 	if (bVirtualTexture)
 	{
-		int32 VTStackIndex = INDEX_NONE;
-		int32 VTLayerIndex = INDEX_NONE;
-		int32 VTPageTableIndex = INDEX_NONE;
+	int32 VTStackIndex = INDEX_NONE;
+	int32 VTLayerIndex = INDEX_NONE;
+	int32 VTPageTableIndex = INDEX_NONE;
 
 		check(VirtualTextureIndex >= 0);
 
@@ -5759,7 +5759,7 @@ int32 FHLSLMaterialTranslator::TextureSample(
 		if (MipValueMode == TMVM_MipLevel)
 		{
 			TextureTypeName += TEXT("Level");
-		}
+	}
 
 		const FMaterialVTStackEntry& VTStackEntry = VTStacks[VTStackIndex];
 		const FString VTPageTableResult_Finite = GetParameterCode(VTStackEntry.CodeIndex);
@@ -5769,11 +5769,11 @@ int32 FHLSLMaterialTranslator::TextureSample(
 		SampleCodeFinite = ApplySamplerType(SampleCodeFinite, SamplerType);
 
 		if (IsAnalyticDerivEnabled() && IsDerivativeValid(UvDerivativeStatus))
-		{
+	{
 			const FString VTPageTableResult_Analytic = GetParameterCodeDeriv(VTStackEntry.CodeIndex, CompiledPDV_Analytic);
 			// 'Texture name/sampler', 'PageTableResult', 'LayerIndex', 'PackedUniform'
 			FString SampleCodeAnalytic;
-			
+
 			if(SamplerDebugSupported(TextureType, bVirtualTexture, bDecal))
 				SampleCodeAnalytic = FString::Printf(TEXT("Debug%s(%s, %s, %d, VTUniform_Unpack(Material.VTPackedUniform[%d]), %s)"), *TextureTypeName, *TextureName, *VTPageTableResult_Analytic, VTPageTableIndex, VirtualTextureIndex, *UV_Scale);
 			else
@@ -6024,7 +6024,7 @@ int32 FHLSLMaterialTranslator::SceneTextureLookup(int32 ViewportUV, uint32 InSce
 		// we can relax this later if needed
 		return NonPixelShaderExpressionError();
 	}
-
+		
 	// Special case for DBuffer normals.
 	// Note that full reprojection only happens with r.Decal.NormalReprojectionEnabled 1.
 	// Otherwise we use a depth buffer approximation.
@@ -6147,17 +6147,17 @@ void FHLSLMaterialTranslator::UseSceneTextureId(ESceneTextureId SceneTextureId,
 	MaterialCompilationOutput.SetIsSceneTextureUsed(SceneTextureId);
 
 	if (Material->GetMaterialDomain() == MD_DeferredDecal)
-	{
+			{
 		const bool bSceneTextureSupportsDecal = SceneTextureId == PPI_SceneDepth || SceneTextureId == PPI_WorldNormal || SceneTextureId == PPI_CustomDepth || SceneTextureId == PPI_CustomStencil;
 		if (!bSceneTextureSupportsDecal)
 		{
 			// Note: For DBuffer decals CustomDepth and CustomStencil are not available if r.CustomDepth.Order = 1
 			Errorf(TEXT("Decals can only access SceneDepth, CustomDepth, CustomStencil, and WorldNormal."));
-		}
+			}
 
 		const bool bSceneTextureRequiresSM5 = SceneTextureId == PPI_WorldNormal || SceneTextureId == PPI_CustomDepth || SceneTextureId == PPI_CustomStencil;
 		if (bSceneTextureRequiresSM5)
-		{
+			{
 			ErrorUnlessFeatureLevelSupported(ERHIFeatureLevel::SM5);
 		}
 
@@ -6755,7 +6755,7 @@ int32 FHLSLMaterialTranslator::VertexInterpolator(uint32 InterpolatorIndex)
 	{
 		return AddCodeChunk(Interpolator->InterpolatedType, *ValueCode);
 	}
-	
+
 }
 
 int32 FHLSLMaterialTranslator::Add(int32 A,int32 B)
@@ -6774,9 +6774,9 @@ int32 FHLSLMaterialTranslator::Add(int32 A,int32 B)
 		if (IsAnalyticDerivEnabled())
 		{
 			return DerivativeAutogen.GenerateExpressionFunc2(*this,FMaterialDerivativeAutogen::EFunc2::Add,A,B);
-		}
-		else
-		{
+	}
+	else
+	{
 			return AddCodeChunk(GetArithmeticResultType(A,B),TEXT("(%s + %s)"),*GetParameterCode(A),*GetParameterCode(B));
 		}
 	}
@@ -6798,9 +6798,9 @@ int32 FHLSLMaterialTranslator::Sub(int32 A,int32 B)
 		if (IsAnalyticDerivEnabled())
 		{
 			return DerivativeAutogen.GenerateExpressionFunc2(*this,FMaterialDerivativeAutogen::EFunc2::Sub,A,B);
-		}
-		else
-		{
+	}
+	else
+	{
 			return AddCodeChunk(GetArithmeticResultType(A,B),TEXT("(%s - %s)"),*GetParameterCode(A),*GetParameterCode(B));
 		}
 	}
@@ -6846,9 +6846,9 @@ int32 FHLSLMaterialTranslator::Div(int32 A,int32 B)
 		if (IsAnalyticDerivEnabled())
 		{
 			return DerivativeAutogen.GenerateExpressionFunc2(*this,FMaterialDerivativeAutogen::EFunc2::Div,A,B);
-		}
-		else
-		{
+	}
+	else
+	{
 			return AddCodeChunk(GetArithmeticResultType(A,B),TEXT("(%s / %s)"),*GetParameterCode(A),*GetParameterCode(B));
 		}
 	}
@@ -6900,18 +6900,18 @@ int32 FHLSLMaterialTranslator::Dot(int32 A,int32 B)
 		}
 		else
 		{
-			// Promote scalar (or truncate the bigger type)
-			if (TypeA == MCT_Float || (TypeB != MCT_Float && GetNumComponents(TypeA) > GetNumComponents(TypeB)))
-			{
-				return AddCodeChunk(MCT_Float,TEXT("dot(%s, %s)"), *CoerceParameter(A, TypeB), *GetParameterCode(B));
-			}
-			else
-			{
-				return AddCodeChunk(MCT_Float,TEXT("dot(%s, %s)"), *GetParameterCode(A), *CoerceParameter(B, TypeA));
-			}
+		// Promote scalar (or truncate the bigger type)
+		if (TypeA == MCT_Float || (TypeB != MCT_Float && GetNumComponents(TypeA) > GetNumComponents(TypeB)))
+		{
+			return AddCodeChunk(MCT_Float,TEXT("dot(%s, %s)"), *CoerceParameter(A, TypeB), *GetParameterCode(B));
+		}
+		else
+		{
+			return AddCodeChunk(MCT_Float,TEXT("dot(%s, %s)"), *GetParameterCode(A), *CoerceParameter(B, TypeA));
 		}
 	}
 }
+}
 
 int32 FHLSLMaterialTranslator::Cross(int32 A,int32 B)
 {
@@ -6934,7 +6934,7 @@ int32 FHLSLMaterialTranslator::Cross(int32 A,int32 B)
 		if(IsAnalyticDerivEnabled())
 			return DerivativeAutogen.GenerateExpressionFunc2(*this, FMaterialDerivativeAutogen::EFunc2::Cross, A, B);
 		else
-			return AddCodeChunk(MCT_Float3,TEXT("cross(%s,%s)"),*CoerceParameter(A,MCT_Float3),*CoerceParameter(B,MCT_Float3));
+		return AddCodeChunk(MCT_Float3,TEXT("cross(%s,%s)"),*CoerceParameter(A,MCT_Float3),*CoerceParameter(B,MCT_Float3));
 	}
 }
 
@@ -6951,9 +6951,9 @@ int32 FHLSLMaterialTranslator::Power(int32 Base,int32 Exponent)
 	}
 	else
 	{
-		// Clamp Pow input to >= 0 to help avoid common NaN cases
-		return AddCodeChunk(GetParameterType(Base),TEXT("PositiveClampedPow(%s,%s)"),*GetParameterCode(Base),*CoerceParameter(Exponent,MCT_Float));
-	}
+	// Clamp Pow input to >= 0 to help avoid common NaN cases
+	return AddCodeChunk(GetParameterType(Base),TEXT("PositiveClampedPow(%s,%s)"),*GetParameterCode(Base),*CoerceParameter(Exponent,MCT_Float));
+}
 }
 	
 int32 FHLSLMaterialTranslator::Logarithm2(int32 X)
@@ -7017,10 +7017,10 @@ int32 FHLSLMaterialTranslator::SquareRoot(int32 X)
 		}
 		else
 		{
-			return AddCodeChunk(GetParameterType(X),TEXT("sqrt(%s)"),*GetParameterCode(X));
-		}
+		return AddCodeChunk(GetParameterType(X),TEXT("sqrt(%s)"),*GetParameterCode(X));
 	}
 }
+}
 
 int32 FHLSLMaterialTranslator::Length(int32 X)
 {
@@ -7041,10 +7041,10 @@ int32 FHLSLMaterialTranslator::Length(int32 X)
 		}
 		else
 		{
-			return AddCodeChunk(MCT_Float,TEXT("length(%s)"),*GetParameterCode(X));
-		}
+		return AddCodeChunk(MCT_Float,TEXT("length(%s)"),*GetParameterCode(X));
 	}
 }
+}
 
 int32 FHLSLMaterialTranslator::Step(int32 Y, int32 X)
 {
@@ -7422,7 +7422,7 @@ int32 FHLSLMaterialTranslator::Lerp(int32 X,int32 Y,int32 A)
 			return Y;
 		}
 	}
-	
+
 	if (IsAnalyticDerivEnabled())
 	{
 		return DerivativeAutogen.GenerateLerpFunc(*this, X, Y, A);
@@ -7452,10 +7452,10 @@ int32 FHLSLMaterialTranslator::Min(int32 A,int32 B)
 		}
 		else
 		{
-			return AddCodeChunk(GetParameterType(A),TEXT("min(%s,%s)"),*GetParameterCode(A),*CoerceParameter(B,GetParameterType(A)));
-		}
+		return AddCodeChunk(GetParameterType(A),TEXT("min(%s,%s)"),*GetParameterCode(A),*CoerceParameter(B,GetParameterType(A)));
 	}
 }
+}
 
 int32 FHLSLMaterialTranslator::Max(int32 A,int32 B)
 {
@@ -7476,10 +7476,10 @@ int32 FHLSLMaterialTranslator::Max(int32 A,int32 B)
 		}
 		else
 		{
-			return AddCodeChunk(GetParameterType(A),TEXT("max(%s,%s)"),*GetParameterCode(A),*CoerceParameter(B,GetParameterType(A)));
-		}
+		return AddCodeChunk(GetParameterType(A),TEXT("max(%s,%s)"),*GetParameterCode(A),*CoerceParameter(B,GetParameterType(A)));
 	}
 }
+}
 
 int32 FHLSLMaterialTranslator::Clamp(int32 X,int32 A,int32 B)
 {
@@ -7502,10 +7502,10 @@ int32 FHLSLMaterialTranslator::Clamp(int32 X,int32 A,int32 B)
 		}
 		else
 		{
-			return AddCodeChunk(GetParameterType(X),TEXT("min(max(%s,%s),%s)"),*GetParameterCode(X),*CoerceParameter(A,GetParameterType(X)),*CoerceParameter(B,GetParameterType(X)));
-		}
+		return AddCodeChunk(GetParameterType(X),TEXT("min(max(%s,%s),%s)"),*GetParameterCode(X),*CoerceParameter(A,GetParameterType(X)),*CoerceParameter(B,GetParameterType(X)));
 	}
 }
+}
 
 int32 FHLSLMaterialTranslator::Saturate(int32 X)
 {
@@ -7526,10 +7526,10 @@ int32 FHLSLMaterialTranslator::Saturate(int32 X)
 		}
 		else
 		{
-			return AddCodeChunk(GetParameterType(X),TEXT("saturate(%s)"),*GetParameterCode(X));
-		}
+		return AddCodeChunk(GetParameterType(X),TEXT("saturate(%s)"),*GetParameterCode(X));
 	}
 }
+}
 
 int32 FHLSLMaterialTranslator::ComponentMask(int32 Vector,bool R,bool G,bool B,bool A)
 {
@@ -7599,7 +7599,7 @@ int32 FHLSLMaterialTranslator::ComponentMask(int32 Vector,bool R,bool G,bool B,b
 				FString::Printf(TEXT("%s.Ddx.%s"), *VectorDeriv, *MaskString),
 				FString::Printf(TEXT("%s.Ddy.%s"), *VectorDeriv, *MaskString),
 				GetDerivTypeIndex(ResultType)
-			);
+		);
 			return AddCodeChunkInnerDeriv(*CodeFinite, *CodeAnalytic, ResultType, false, EDerivativeStatus::Valid);
 		}
 		else
@@ -7730,8 +7730,8 @@ int32 FHLSLMaterialTranslator::TransformBase(EMaterialCommonBasis SourceCoordBas
 			check(AWComponent == 0);
 			if (DestCoordBasis == MCB_World)
 			{
-				CodeStr = TEXT("mul(<A>, <MATRIX>(Parameters.TangentToWorld))");
-			}
+					CodeStr = TEXT("mul(<A>, <MATRIX>(Parameters.TangentToWorld))");
+				}
 			// else use MCB_World as intermediary basis
 			break;
 		}
@@ -9696,12 +9696,12 @@ int32 FHLSLMaterialTranslator::RotateScaleOffsetTexCoords(int32 TexCoordCodeInde
 	}
 	else
 	{
-		return AddCodeChunk(MCT_Float2,
-			TEXT("RotateScaleOffsetTexCoords(%s, %s, %s.xy)"),
-			*GetParameterCode(TexCoordCodeIndex),
-			*GetParameterCode(RotationScale),
-			*GetParameterCode(Offset));
-	}
+	return AddCodeChunk(MCT_Float2,
+		TEXT("RotateScaleOffsetTexCoords(%s, %s, %s.xy)"),
+		*GetParameterCode(TexCoordCodeIndex),
+		*GetParameterCode(RotationScale),
+		*GetParameterCode(Offset));
+}
 }
 
 /**
diff --git a/Engine/Source/Runtime/Engine/Private/Materials/MaterialExpressions.cpp b/Engine/Source/Runtime/Engine/Private/Materials/MaterialExpressions.cpp
--- a/Engine/Source/Runtime/Engine/Private/Materials/MaterialExpressions.cpp
+++ b/Engine/Source/Runtime/Engine/Private/Materials/MaterialExpressions.cpp
@@ -2154,6 +2154,29 @@ UMaterialExpressionRuntimeVirtualTextureOutput::UMaterialExpressionRuntimeVirtua
 
 #if WITH_EDITOR
 
+bool ValidateRuntimeVirtualTextureOutput(FMaterialCompiler* Compiler, FExpressionInput& InExpressionInput)
+{
+	// Note that this doesn't traverse through material layers and functions.
+	if (InExpressionInput.IsConnected())
+	{
+		TArray<UMaterialExpression*> InputExpressions;
+		InExpressionInput.Expression->GetAllInputExpressions(InputExpressions);
+		for (UMaterialExpression* Expression : InputExpressions)
+		{
+			UMaterialExpressionRuntimeVirtualTextureSample* RuntimeVirtualTextureSampleExpression = Cast<UMaterialExpressionRuntimeVirtualTextureSample>(Expression);
+			if (RuntimeVirtualTextureSampleExpression != nullptr)
+			{
+				const FString VirtualTextureName = RuntimeVirtualTextureSampleExpression->VirtualTexture ? RuntimeVirtualTextureSampleExpression->VirtualTexture->GetName() : RuntimeVirtualTextureSampleExpression->GetName();
+
+				Compiler->Errorf(TEXT("Sampling a Runtime Virtual Texture (%s) when writing Runtime Virtual Texture output is not supported."), *VirtualTextureName);
+				return false;
+			}
+		}
+	}
+
+	return true;
+}
+
 int32 UMaterialExpressionRuntimeVirtualTextureOutput::Compile(class FMaterialCompiler* Compiler, int32 OutputIndex)
 {
 	int32 CodeInput = INDEX_NONE;
@@ -2165,35 +2188,42 @@ int32 UMaterialExpressionRuntimeVirtualTextureOutput::Compile(class FMaterialCom
 	{
 		CodeInput = BaseColor.IsConnected() ? BaseColor.Compile(Compiler) : Compiler->Constant3(0.f, 0.f, 0.f);
 		OutputAttributeMask |= BaseColor.IsConnected() ? (1 << (uint8)ERuntimeVirtualTextureAttributeType::BaseColor) : 0;
+		ValidateRuntimeVirtualTextureOutput(Compiler, BaseColor);
 	}
 	else if (OutputIndex == 1)
 	{
 		CodeInput = Specular.IsConnected() ? Specular.Compile(Compiler) : Compiler->Constant(0.5f);
 		OutputAttributeMask |= Specular.IsConnected() ? (1 << (uint8)ERuntimeVirtualTextureAttributeType::Specular) : 0;
+		ValidateRuntimeVirtualTextureOutput(Compiler, Specular);
 	}
 	else if (OutputIndex == 2)
 	{
 		CodeInput = Roughness.IsConnected() ? Roughness.Compile(Compiler) : Compiler->Constant(0.5f);
 		OutputAttributeMask |= Roughness.IsConnected() ? (1 << (uint8)ERuntimeVirtualTextureAttributeType::Roughness) : 0;
+		ValidateRuntimeVirtualTextureOutput(Compiler, Roughness);
 	}
 	else if (OutputIndex == 3)
 	{
 		CodeInput = Normal.IsConnected() ? Normal.Compile(Compiler) : Compiler->Constant3(0.f, 0.f, 1.f);
 		OutputAttributeMask |= Normal.IsConnected() ? (1 << (uint8)ERuntimeVirtualTextureAttributeType::Normal) : 0;
+		ValidateRuntimeVirtualTextureOutput(Compiler, Normal);
 	}
 	else if (OutputIndex == 4)
 	{
 		CodeInput = WorldHeight.IsConnected() ? WorldHeight.Compile(Compiler) : Compiler->Constant(0.f);
 		OutputAttributeMask |= WorldHeight.IsConnected() ? (1 << (uint8)ERuntimeVirtualTextureAttributeType::WorldHeight) : 0;
+		ValidateRuntimeVirtualTextureOutput(Compiler, WorldHeight);
 	}
 	else if (OutputIndex == 5)
 	{
 		CodeInput = Opacity.IsConnected() ? Opacity.Compile(Compiler) : Compiler->Constant(1.f);
+		ValidateRuntimeVirtualTextureOutput(Compiler, Opacity);
 	}
 	else if (OutputIndex == 6)
 	{
 		CodeInput = Mask.IsConnected() ? Mask.Compile(Compiler) : Compiler->Constant(1.f);
 		OutputAttributeMask |= Mask.IsConnected() ? (1 << (uint8)ERuntimeVirtualTextureAttributeType::Mask) : 0;
+		ValidateRuntimeVirtualTextureOutput(Compiler, Mask);
 	}
 
 	Compiler->VirtualTextureOutput(OutputAttributeMask);
