Added possibility to execute a first capture the skylight as soon as possible when created to get soem environment lighting.
A final capture is executed again when al lthe needed assets are loaded.
Feedback is also presented to the use on screen about what is happening.

diff --git a/Engine/Source/Runtime/Engine/Classes/Components/SkyLightComponent.h b/Engine/Source/Runtime/Engine/Classes/Components/SkyLightComponent.h
--- a/Engine/Source/Runtime/Engine/Classes/Components/SkyLightComponent.h
+++ b/Engine/Source/Runtime/Engine/Classes/Components/SkyLightComponent.h
@@ -88,6 +88,13 @@ enum ESkyLightSourceType
 	SLS_MAX,
 };
 
+enum ESkyLightCaptureStatus
+{
+	SLCS_Uninitialized,
+	SLCS_CapturedButIncomplete,
+	SLCS_CapturedAndComplete,
+};
+
 UCLASS(Blueprintable, ClassGroup=Lights, HideCategories=(Trigger,Activation,"Components|Activation",Physics), meta=(BlueprintSpawnableComponent))
 class ENGINE_API USkyLightComponent : public ULightComponentBase
 {
@@ -302,6 +309,12 @@ protected:
 
 	bool bHasEverCaptured;
 
+#if WITH_EDITOR
+	// In an attempt to get valid lighting data as soon as possible in a level, we will always trigger a skylight capture at creation.
+	// And then, a soon as the world is finally loaded, the final capture will be taken.
+	ESkyLightCaptureStatus CaptureStatus;
+#endif
+
 	TRefCountPtr<FSkyTextureCubeResource> ProcessedSkyTexture;
 	FSHVectorRGB3 IrradianceEnvironmentMap;
 	float AverageBrightness=0.0f;
diff --git a/Engine/Source/Runtime/Engine/Private/Components/SkyLightComponent.cpp b/Engine/Source/Runtime/Engine/Private/Components/SkyLightComponent.cpp
--- a/Engine/Source/Runtime/Engine/Private/Components/SkyLightComponent.cpp
+++ b/Engine/Source/Runtime/Engine/Private/Components/SkyLightComponent.cpp
@@ -221,6 +221,11 @@ FSkyLightSceneProxy::FSkyLightSceneProxy(const USkyLightComponent* InLightCompon
 	, CaptureCubeMapResolution(InLightComponent->CubemapResolution)
 	, LowerHemisphereColor(InLightComponent->LowerHemisphereColor)
 	, bLowerHemisphereIsSolidColor(InLightComponent->bLowerHemisphereIsBlack)
+#if WITH_EDITOR
+	, bCubemapSkyLightWaitingForCubeMapTexture(false)
+	, bCaptureSkyLightWaitingForShaders(false)
+	, bCaptureSkyLightWaitingForMeshesOrTextures(false)
+#endif
 	, LightColor(FLinearColor(InLightComponent->LightColor) * InLightComponent->Intensity)
 	, bMovable(InLightComponent->IsMovable())
 {
@@ -280,6 +285,10 @@ USkyLightComponent::USkyLightComponent(const FObjectInitializer& ObjectInitializ
 	CloudAmbientOcclusionStrength = 1.0f;
 	CloudAmbientOcclusionMapResolutionScale = 1.0f;
 	CloudAmbientOcclusionApertureScale = 0.05f;
+
+#if WITH_EDITOR
+	CaptureStatus = SLCS_Uninitialized;
+#endif
 }
 
 FSkyLightSceneProxy* USkyLightComponent::CreateSceneProxy() const
@@ -694,16 +703,38 @@ void USkyLightComponent::UpdateSkyCaptureContentsArray(UWorld* WorldToUpdate, TA
 			// We should process this texture as soon as possible so we can have a proper skylight.
 			FTextureCompilingManager::Get().RequestPriorityChange(CaptureComponent->Cubemap, EQueuedWorkPriority::Highest);
 		}
+
+		const bool bCubemapSkyLightWaitingForCubemapAsset	= CaptureComponent->SourceType == SLS_SpecifiedCubemap	&& bIsCubemapCompiling;
+		const bool bCaptureSkyLightWaitingCompiledShader	= CaptureComponent->SourceType == SLS_CapturedScene		&& bIsCompilingShaders;
+		const bool bCaptureSkyLightWaitingForMeshOrTexAssets= CaptureComponent->SourceType == SLS_CapturedScene		&& bSceneIsAsyncCompiling;
+		ENQUEUE_RENDER_COMMAND(FUpdateSkyLightProxyStatus)(
+			[CaptureComponent, bCubemapSkyLightWaitingForCubemapAsset, bCaptureSkyLightWaitingCompiledShader, bCaptureSkyLightWaitingForMeshOrTexAssets](FRHICommandList& RHICmdList)
+			{
+				FSkyLightSceneProxy* SkyLightSceneProxy = CaptureComponent->SceneProxy;
+				if (SkyLightSceneProxy)
+				{
+					SkyLightSceneProxy->bCubemapSkyLightWaitingForCubeMapTexture = bCubemapSkyLightWaitingForCubemapAsset;
+					SkyLightSceneProxy->bCaptureSkyLightWaitingForShaders = bCaptureSkyLightWaitingCompiledShader;
+					SkyLightSceneProxy->bCaptureSkyLightWaitingForMeshesOrTextures = bCaptureSkyLightWaitingForMeshOrTexAssets;
+				}
+			});
 #endif
 
 		if (((!Owner || !Owner->GetLevel() || Owner->GetLevel()->bIsVisible) && CaptureComponent->GetWorld() == WorldToUpdate)
 			// Only process sky capture requests once async texture and shader compiling completes, otherwise we will capture the scene with temporary shaders/textures
-			&& (((!bSceneIsAsyncCompiling) && (!bIsCompilingShaders)) || ((CaptureComponent->SourceType == SLS_SpecifiedCubemap) && (!bIsCubemapCompiling))))
+			&& (
+#if WITH_EDITOR
+				CaptureComponent->CaptureStatus == SLCS_Uninitialized 
+				|| 
+#endif
+				((!bSceneIsAsyncCompiling) && (!bIsCompilingShaders)) 
+				|| 
+				((CaptureComponent->SourceType == SLS_SpecifiedCubemap) && (!bIsCubemapCompiling)))
+			)
 		{
 			// Only capture valid sky light components
 			if (CaptureComponent->SourceType != SLS_SpecifiedCubemap || CaptureComponent->Cubemap)
 			{
-
 #if WITH_EDITOR
 				FStaticLightingSystemInterface::OnLightComponentUnregistered.Broadcast(CaptureComponent);
 #endif
@@ -748,9 +779,60 @@ void USkyLightComponent::UpdateSkyCaptureContentsArray(UWorld* WorldToUpdate, TA
 #endif
 			}
 
-			// Only remove queued update requests if we processed it for the right world
+#if WITH_EDITOR
+			switch (CaptureComponent->CaptureStatus)
+			{
+			case SLCS_Uninitialized:
+			{
+				if (!bCubemapSkyLightWaitingForCubemapAsset && !bCaptureSkyLightWaitingCompiledShader && !bCaptureSkyLightWaitingForMeshOrTexAssets)
+				{
+					// Do not recapture if the first forced capture was with a complete world (to avoid capturing twice each time a level is loaded or a skylight created))
+					CaptureComponent->CaptureStatus = SLCS_CapturedAndComplete;
+					ComponentArray.RemoveAt(CaptureIndex);
+				}
+				else
+				{
+					CaptureComponent->CaptureStatus = SLCS_CapturedButIncomplete;
+				}
+				break;
+			}
+			case SLCS_CapturedButIncomplete:
+			{
+				// Only remove queued update requests if we processed it for the a world with all meshes, textures and shaders.
+				ComponentArray.RemoveAt(CaptureIndex);
+				CaptureComponent->CaptureStatus = SLCS_CapturedAndComplete;
+				break;
+			}
+			case SLCS_CapturedAndComplete:
+			{
+				// It is valid to recapture a complete skylight.
+				ComponentArray.RemoveAt(CaptureIndex);
+				break;
+			}
+			default:
+			{
+				check(false);
+				break;
+			}
+			}
+#else
 			ComponentArray.RemoveAt(CaptureIndex);
+#endif
 		}
+
+#if WITH_EDITOR
+		ENQUEUE_RENDER_COMMAND(FUpdateSkyLightProxyStatusForcedCapture)(
+			[CaptureComponent, bCubemapSkyLightWaitingForCubemapAsset, bCaptureSkyLightWaitingCompiledShader, bCaptureSkyLightWaitingForMeshOrTexAssets](FRHICommandList& RHICmdList)
+			{
+				FSkyLightSceneProxy* SkyLightSceneProxy = CaptureComponent->SceneProxy;
+				if (SkyLightSceneProxy)
+				{
+					SkyLightSceneProxy->bCubemapSkyLightWaitingForCubeMapTexture = bCubemapSkyLightWaitingForCubemapAsset;
+					SkyLightSceneProxy->bCaptureSkyLightWaitingForShaders = bCaptureSkyLightWaitingCompiledShader;
+					SkyLightSceneProxy->bCaptureSkyLightWaitingForMeshesOrTextures = bCaptureSkyLightWaitingForMeshOrTexAssets;
+				}
+			});
+#endif
 	}
 }
 
diff --git a/Engine/Source/Runtime/Engine/Public/SceneManagement.h b/Engine/Source/Runtime/Engine/Public/SceneManagement.h
--- a/Engine/Source/Runtime/Engine/Public/SceneManagement.h
+++ b/Engine/Source/Runtime/Engine/Public/SceneManagement.h
@@ -1160,6 +1160,12 @@ public:
 	}
 	FLinearColor GetEffectiveLightColor() const;
 
+#if WITH_EDITOR
+	bool bCubemapSkyLightWaitingForCubeMapTexture;
+	bool bCaptureSkyLightWaitingForShaders;
+	bool bCaptureSkyLightWaitingForMeshesOrTextures;
+#endif
+
 private:
 	FLinearColor LightColor;
 	const uint8 bMovable : 1;
diff --git a/Engine/Source/Runtime/Renderer/Private/SceneRendering.cpp b/Engine/Source/Runtime/Renderer/Private/SceneRendering.cpp
--- a/Engine/Source/Runtime/Renderer/Private/SceneRendering.cpp
+++ b/Engine/Source/Runtime/Renderer/Private/SceneRendering.cpp
@@ -3007,13 +3007,22 @@ void FSceneRenderer::RenderFinish(FRDGBuilder& GraphBuilder, FRDGTextureRef View
 
 		const bool bSingleLayerWaterWarning = ShouldRenderSingleLayerWaterSkippedRenderEditorNotification(Views);
 
+		bool bShowWaitingSkylight = false;
+#if WITH_EDITOR
+		FSkyLightSceneProxy* SkyLight = Scene->SkyLight;
+		if (SkyLight && !SkyLight->bRealTimeCaptureEnabled)
+		{
+			bShowWaitingSkylight = SkyLight->bCubemapSkyLightWaitingForCubeMapTexture || SkyLight->bCaptureSkyLightWaitingForShaders || SkyLight->bCaptureSkyLightWaitingForMeshesOrTextures;
+		}
+#endif
+
 		FFXSystemInterface* FXInterface = Scene->GetFXSystem();
 		const bool bFxDebugDraw = FXInterface && FXInterface->ShouldDebugDraw_RenderThread();
 
 		const bool bAnyWarning = bShowPrecomputedVisibilityWarning || bShowDemotedLocalMemoryWarning || bShowGlobalClipPlaneWarning || bShowSkylightWarning || bShowPointLightWarning
 			|| bShowDFAODisabledWarning || bShowShadowedLightOverflowWarning || bShowMobileDynamicCSMWarning || bShowMobileLowQualityLightmapWarning || bShowMobileMovableDirectionalLightWarning
 			|| bMobileShowVertexFogWarning || bMobileMissingSkyMaterial || bShowSkinCacheOOM || bSingleLayerWaterWarning || bShowDFDisabledWarning || bShowNoSkyAtmosphereComponentWarning || bFxDebugDraw 
-			|| bLumenEnabledButNoSoftwareTracing || bNaniteEnabledButNoAtomics || bRealTimeSkyCaptureButNothingToCapture;
+			|| bLumenEnabledButNoSoftwareTracing || bNaniteEnabledButNoAtomics || bRealTimeSkyCaptureButNothingToCapture || bShowWaitingSkylight;
 
 		for(int32 ViewIndex = 0;ViewIndex < Views.Num();ViewIndex++)
 		{	
@@ -3036,7 +3045,7 @@ void FSceneRenderer::RenderFinish(FRDGBuilder& GraphBuilder, FRDGTextureRef View
 						bViewParentOrFrozen, bShowSkylightWarning, bShowPointLightWarning, bShowShadowedLightOverflowWarning,
 						bShowMobileLowQualityLightmapWarning, bShowMobileMovableDirectionalLightWarning, bShowMobileDynamicCSMWarning, bMobileShowVertexFogWarning, bMobileMissingSkyMaterial, 
 						bShowSkinCacheOOM, bSingleLayerWaterWarning, bShowNoSkyAtmosphereComponentWarning, bFxDebugDraw, FXInterface, 
-						bLumenEnabledButNoSoftwareTracing, bNaniteEnabledButNoAtomics, bNaniteRequireAtomics, bRealTimeSkyCaptureButNothingToCapture]
+						bLumenEnabledButNoSoftwareTracing, bNaniteEnabledButNoAtomics, bNaniteRequireAtomics, bRealTimeSkyCaptureButNothingToCapture, bShowWaitingSkylight]
 						(FCanvas& Canvas)
 					{
 						// so it can get the screen size
@@ -3208,6 +3217,36 @@ void FSceneRenderer::RenderFinish(FRDGBuilder& GraphBuilder, FRDGTextureRef View
 							Canvas.DrawShadowedText(10, Y, FText::FromString(String), GetStatsFont(), FLinearColor(1.0, 0.05, 0.05, 1.0));
 							Y += 14;
 						}
+
+#if WITH_EDITOR
+						FSkyLightSceneProxy* SkyLight = Scene->SkyLight;
+						if (bShowWaitingSkylight && SkyLight)
+						{
+							const FLinearColor OrangeColor = FColor::Orange;
+
+							FString String = FString::Printf(TEXT("Sky Light waiting on ["));
+							bool bAddComma = false;
+							if (SkyLight->bCubemapSkyLightWaitingForCubeMapTexture)
+							{
+								String += TEXT("CubeMap");
+								bAddComma = true;
+							}
+							if (SkyLight->bCaptureSkyLightWaitingForShaders)
+							{
+								String += bAddComma ? TEXT(", ") : TEXT("");
+								String += TEXT("Shaders");
+								bAddComma = true;
+							}
+							if (SkyLight->bCaptureSkyLightWaitingForMeshesOrTextures)
+							{
+								String += bAddComma ? TEXT(", ") : TEXT("");
+								String += TEXT("Meshes, Textures] for final capture");
+							}
+							String += FString::Printf(TEXT("] for final capture."));
+							Canvas.DrawShadowedText(10, Y, FText::FromString(String), GetStatsFont(), OrangeColor);
+							Y += 14;
+						}
+#endif
 					});
 					if (bFxDebugDraw)
 					{
