Light weight instance changes

Added GetTypeHash to FActorInstanceHandles so they can be used as keys in container classes.
GetActorClassToSpawn now takes a handle. This allows subclasses to potentially create more than one actor type.
Fixed a crash in ALightWeightInstanceStaticMeshManager.

diff --git a/Engine/Source/Runtime/Engine/Classes/Engine/EngineTypes.h b/Engine/Source/Runtime/Engine/Classes/Engine/EngineTypes.h
--- a/Engine/Source/Runtime/Engine/Classes/Engine/EngineTypes.h
+++ b/Engine/Source/Runtime/Engine/Classes/Engine/EngineTypes.h
@@ -2109,6 +2109,8 @@ struct ENGINE_API FActorInstanceHandle
 	bool operator==(const AActor* OtherActor) const;
 	bool operator!=(const AActor* OtherActor) const;
 
+	friend ENGINE_API uint32 GetTypeHash(const FActorInstanceHandle& Handle);
+
 	friend ENGINE_API FArchive& operator<<(FArchive& Ar, FActorInstanceHandle& Handle);
 
 private:
diff --git a/Engine/Source/Runtime/Engine/Classes/GameFramework/LightWeightInstanceManager.h b/Engine/Source/Runtime/Engine/Classes/GameFramework/LightWeightInstanceManager.h
--- a/Engine/Source/Runtime/Engine/Classes/GameFramework/LightWeightInstanceManager.h
+++ b/Engine/Source/Runtime/Engine/Classes/GameFramework/LightWeightInstanceManager.h
@@ -151,7 +151,7 @@ protected:
 	virtual void SetSpawnParameters(FActorSpawnParameters& SpawnParams);
 
 	// Returns the class to use when spawning a new actor
-	virtual UClass* GetActorClassToSpawn() const;
+	virtual UClass* GetActorClassToSpawn(const FActorInstanceHandle& Handle) const;
 
 	// Called after spawning a new actor from a light weight instance
 	virtual void PostActorSpawn(const FActorInstanceHandle& Handle);
diff --git a/Engine/Source/Runtime/Engine/Private/Components/CharacterMovementComponent.cpp b/Engine/Source/Runtime/Engine/Private/Components/CharacterMovementComponent.cpp
--- a/Engine/Source/Runtime/Engine/Private/Components/CharacterMovementComponent.cpp
+++ b/Engine/Source/Runtime/Engine/Private/Components/CharacterMovementComponent.cpp
@@ -6625,7 +6625,7 @@ bool UCharacterMovementComponent::CanStepUp(const FHitResult& Hit) const
 		 return true;
 	}
 
-	const AActor* HitActor = Hit.HitObjectHandle.FetchActor();
+	const AActor* HitActor = Hit.HitObjectHandle.GetManagingActor();
 	if (!HitActor->CanBeBaseForCharacter(CharacterOwner))
 	{
 		return false;
diff --git a/Engine/Source/Runtime/Engine/Private/EngineTypes.cpp b/Engine/Source/Runtime/Engine/Private/EngineTypes.cpp
--- a/Engine/Source/Runtime/Engine/Private/EngineTypes.cpp
+++ b/Engine/Source/Runtime/Engine/Private/EngineTypes.cpp
@@ -548,6 +548,19 @@ bool FActorInstanceHandle::operator!=(const AActor* OtherActor) const
 	return !(*this == OtherActor);
 }
 
+uint32 GetTypeHash(const FActorInstanceHandle& Handle)
+{
+	uint32 Hash = 0;
+	if (Handle.Actor.IsValid())
+	{
+		FCrc::StrCrc32(*(Handle.Actor->GetPathName()), Hash);
+	}
+	Hash = HashCombine(Hash, Handle.ManagerIndex);
+	Hash = HashCombine(Hash, Handle.InstanceIndex);
+
+	return Hash;
+}
+
 FArchive& operator<<(FArchive& Ar, FActorInstanceHandle& Handle)
 {
 	Ar << Handle.Actor;
diff --git a/Engine/Source/Runtime/Engine/Private/LightWeightInstanceManager.cpp b/Engine/Source/Runtime/Engine/Private/LightWeightInstanceManager.cpp
--- a/Engine/Source/Runtime/Engine/Private/LightWeightInstanceManager.cpp
+++ b/Engine/Source/Runtime/Engine/Private/LightWeightInstanceManager.cpp
@@ -73,7 +73,7 @@ AActor* ALightWeightInstanceManager::ConvertInstanceToActor(const FActorInstance
 	FActorSpawnParameters SpawnParams;
 	SetSpawnParameters(SpawnParams);
 
-	AActor* NewActor = GetLevel()->GetWorld()->SpawnActor<AActor>(GetActorClassToSpawn(), InstanceTransforms[Handle.GetInstanceIndex()], SpawnParams);
+	AActor* NewActor = GetLevel()->GetWorld()->SpawnActor<AActor>(GetActorClassToSpawn(Handle), InstanceTransforms[Handle.GetInstanceIndex()], SpawnParams);
 	check(NewActor);
 
 	Handle.Actor = NewActor;
@@ -176,7 +176,7 @@ void ALightWeightInstanceManager::SetSpawnParameters(FActorSpawnParameters& Spaw
 	SpawnParams.ObjectFlags = RF_Transactional;
 }
 
-UClass* ALightWeightInstanceManager::GetActorClassToSpawn() const
+UClass* ALightWeightInstanceManager::GetActorClassToSpawn(const FActorInstanceHandle& Handle) const
 {
 	return RepresentedClass;
 }
