Make sure to always keep actor labels when breaking a Level Instance.

diff --git a/Engine/Source/Runtime/Engine/Private/LevelInstance/LevelInstanceSubsystem.cpp b/Engine/Source/Runtime/Engine/Private/LevelInstance/LevelInstanceSubsystem.cpp
--- a/Engine/Source/Runtime/Engine/Private/LevelInstance/LevelInstanceSubsystem.cpp
+++ b/Engine/Source/Runtime/Engine/Private/LevelInstance/LevelInstanceSubsystem.cpp
@@ -15,6 +15,7 @@
 #include "Hash/CityHash.h"
 
 #if WITH_EDITOR
+#include "Settings/LevelEditorMiscSettings.h"
 #include "LevelInstance/LevelInstanceEditorLevelStreaming.h"
 #include "Misc/ScopedSlowTask.h"
 #include "Misc/ITransaction.h"
@@ -868,6 +869,10 @@ bool ULevelInstanceSubsystem::BreakLevelInstance(ALevelInstance* LevelInstanceAc
 {
 	const double StartTime = FPlatformTime::Seconds();
 
+	const uint32 bAvoidRelabelOnPasteSelected = GetMutableDefault<ULevelEditorMiscSettings>()->bAvoidRelabelOnPasteSelected;
+	ON_SCOPE_EXIT { GetMutableDefault<ULevelEditorMiscSettings>()->bAvoidRelabelOnPasteSelected = bAvoidRelabelOnPasteSelected; };
+	GetMutableDefault<ULevelEditorMiscSettings>()->bAvoidRelabelOnPasteSelected = 1;
+
 	TArray<AActor*> MovedActors;
 	BreakLevelInstance_Impl(LevelInstanceActor, Levels, MovedActors);
 
