Add some logging:
BlockTillLevelStreamingCompleted time,
WorldPartition pre-match started streaming states
Data Layer State Changes

diff --git a/Engine/Source/Runtime/Engine/Private/World.cpp b/Engine/Source/Runtime/Engine/Private/World.cpp
--- a/Engine/Source/Runtime/Engine/Private/World.cpp
+++ b/Engine/Source/Runtime/Engine/Private/World.cpp
@@ -3569,6 +3569,8 @@ void FStreamingLevelsToConsider::AddReferencedObjects(UObject* InThis, FReferenc
 
 void UWorld::BlockTillLevelStreamingCompleted()
 {
+	const double StartTime = FPlatformTime::Seconds();
+
 	bool bWorkToDo = false;
 	do
 	{
@@ -3595,6 +3597,8 @@ void UWorld::BlockTillLevelStreamingCompleted()
 			}
 		}
 	} while (bWorkToDo);
+	const double ElapsedTime = FPlatformTime::Seconds() - StartTime;
+	UE_LOG(LogWorld, Log, TEXT("BlockTillLevelStreamingCompleted took %s seconds (MatchStarted %d)"), *FText::AsNumber(ElapsedTime).ToString(), bMatchStarted);
 }
 
 void UWorld::InternalUpdateStreamingState()
diff --git a/Engine/Source/Runtime/Engine/Private/WorldPartition/DataLayer/WorldDataLayers.cpp b/Engine/Source/Runtime/Engine/Private/WorldPartition/DataLayer/WorldDataLayers.cpp
--- a/Engine/Source/Runtime/Engine/Private/WorldPartition/DataLayer/WorldDataLayers.cpp
+++ b/Engine/Source/Runtime/Engine/Private/WorldPartition/DataLayer/WorldDataLayers.cpp
@@ -9,6 +9,7 @@
 #include "WorldPartition/DataLayer/DataLayerSubsystem.h"
 #include "EngineUtils.h"
 #include "Net/UnrealNetwork.h"
+#include "WorldPartition/WorldPartition.h"
 #if WITH_EDITOR
 #include "WorldPartition/WorldPartitionEditorPerProjectUserSettings.h"
 #endif
@@ -17,6 +18,21 @@
 
 int32 AWorldDataLayers::DataLayersStateEpoch = 0;
 
+FString JoinDataLayerLabelsFromNames(AWorldDataLayers* InWorldDataLayers, const TArray<FName>& InDataLayerNames)
+{
+	check(InWorldDataLayers);
+	TArray<FString> DataLayerLabels;
+	DataLayerLabels.Reserve(InDataLayerNames.Num());
+	for (const FName& DataLayerName : InDataLayerNames)
+	{
+		if (const UDataLayer* DataLayer = InWorldDataLayers->GetDataLayerFromName(DataLayerName))
+		{
+			DataLayerLabels.Add(DataLayer->GetDataLayerLabel().ToString());
+		}
+	}
+	return FString::Join(DataLayerLabels, TEXT(","));
+}
+
 AWorldDataLayers::AWorldDataLayers(const FObjectInitializer& ObjectInitializer)
 	: Super(ObjectInitializer)
 {
@@ -80,6 +96,8 @@ void AWorldDataLayers::InitializeDataLayerStates()
 
 		RepActiveDataLayerNames = ActiveDataLayerNames.Array();
 		RepLoadedDataLayerNames = LoadedDataLayerNames.Array();
+
+		UE_LOG(LogWorldPartition, Log, TEXT("Initial Data Layer States Activated(%s) Loaded(%s)"), *JoinDataLayerLabelsFromNames(this, RepActiveDataLayerNames), *JoinDataLayerLabelsFromNames(this, RepLoadedDataLayerNames));
 	}
 }
 
@@ -122,6 +140,11 @@ void AWorldDataLayers::SetDataLayerState(FActorDataLayer InDataLayer, EDataLayer
 
 			++DataLayersStateEpoch;
 
+			UE_LOG(LogWorldPartition, Log, TEXT("Data Layer '%s' state changed: %s -> %s"), 
+				*DataLayer->GetDataLayerLabel().ToString(), 
+				*StaticEnum<EDataLayerState>()->GetDisplayNameTextByValue((int64)CurrentState).ToString(),
+				*StaticEnum<EDataLayerState>()->GetDisplayNameTextByValue((int64)InState).ToString());
+
 			OnDataLayerStateChanged(DataLayer, InState);
 		}
 	}
@@ -187,6 +210,8 @@ void AWorldDataLayers::OverwriteDataLayerStates(TArray<FActorDataLayer>* InActiv
 			}
 			RepLoadedDataLayerNames = LoadedDataLayerNames.Array();
 		}
+
+		UE_LOG(LogWorldPartition, Log, TEXT("Overwrite Data Layer States Activated(%s) Loaded(%s)"), *JoinDataLayerLabelsFromNames(this, RepActiveDataLayerNames), *JoinDataLayerLabelsFromNames(this, RepLoadedDataLayerNames));
 	}
 }
 
diff --git a/Engine/Source/Runtime/Engine/Private/WorldPartition/WorldPartitionStreamingPolicy.cpp b/Engine/Source/Runtime/Engine/Private/WorldPartition/WorldPartitionStreamingPolicy.cpp
--- a/Engine/Source/Runtime/Engine/Private/WorldPartition/WorldPartitionStreamingPolicy.cpp
+++ b/Engine/Source/Runtime/Engine/Private/WorldPartition/WorldPartitionStreamingPolicy.cpp
@@ -116,6 +116,22 @@ void UWorldPartitionStreamingPolicy::UpdateStreamingSources()
 	}
 }
 
+#define WORLDPARTITION_LOG_UPDATESTREAMINGSTATE(Verbosity)\
+UE_SUPPRESS(LogWorldPartition, Verbosity, \
+{ \
+	if (ToActivateCells.Num() > 0 || ToLoadCells.Num() > 0 || ToUnloadCells.Num() > 0) \
+	{ \
+		UE_LOG(LogWorldPartition, Verbosity, TEXT("UWorldPartitionStreamingPolicy: CellsToActivate(%d), CellsToLoad(%d), CellsToUnload(%d)"), ToActivateCells.Num(), ToLoadCells.Num(), ToUnloadCells.Num()); \
+		FTransform LocalToWorld = WorldPartition->GetInstanceTransform(); \
+		for (int i = 0; i < StreamingSources.Num(); ++i) \
+		{ \
+			FVector ViewLocation = LocalToWorld.TransformPosition(StreamingSources[i].Location); \
+			FRotator ViewRotation = LocalToWorld.TransformRotation(StreamingSources[i].Rotation.Quaternion()).Rotator(); \
+			UE_LOG(LogWorldPartition, Verbosity, TEXT("UWorldPartitionStreamingPolicy: Sources[%d] = %s,%s"), i, *ViewLocation.ToString(), *ViewRotation.ToString()); \
+		} \
+	} \
+}) \
+
 void UWorldPartitionStreamingPolicy::UpdateStreamingState()
 {
 	TRACE_CPUPROFILER_EVENT_SCOPE(UWorldPartitionStreamingPolicy::UpdateStreamingState);
@@ -199,22 +215,15 @@ void UWorldPartitionStreamingPolicy::UpdateStreamingState()
 	TSet<const UWorldPartitionRuntimeCell*> ToLoadCells = LoadStreamingCells.Difference(LoadedCells);
 	TSet<const UWorldPartitionRuntimeCell*> ToUnloadCells = ActivatedCells.Union(LoadedCells).Difference(ActivateStreamingCells.Union(LoadStreamingCells));
 
-	UE_SUPPRESS(LogWorldPartition, Verbose,
+	if(World->bMatchStarted)
 	{
-		if (ToActivateCells.Num() > 0 || ToLoadCells.Num() > 0 || ToUnloadCells.Num() > 0)
-		{
-			UE_LOG(LogWorldPartition, Verbose, TEXT("UWorldPartitionStreamingPolicy: CellsToActivate(%d), CellsToLoad(%d), CellsToUnload(%d)"), ToActivateCells.Num(), ToLoadCells.Num(), ToUnloadCells.Num());
-
-			FTransform LocalToWorld = WorldPartition->GetInstanceTransform();
-			for (int i = 0; i < StreamingSources.Num(); ++i)
-			{
-				FVector ViewLocation = LocalToWorld.TransformPosition(StreamingSources[i].Location);
-				FRotator ViewRotation = LocalToWorld.TransformRotation(StreamingSources[i].Rotation.Quaternion()).Rotator();
-				UE_LOG(LogWorldPartition, Verbose, TEXT("UWorldPartitionStreamingPolicy: Sources[%d] = %s,%s"), i, *ViewLocation.ToString(), *ViewRotation.ToString());
-			}
-		}
-	});
-
+		WORLDPARTITION_LOG_UPDATESTREAMINGSTATE(Verbose);
+	}
+	else
+	{
+		WORLDPARTITION_LOG_UPDATESTREAMINGSTATE(Log);
+	}
+	
 	if (ToUnloadCells.Num() > 0)
 	{
 		SetTargetStateForCells(EWorldPartitionRuntimeCellState::Unloaded, ToUnloadCells);
