- Added Trace tag for FRenderAssetStreamingManager::IncrementalUpdate.
- Replaced wrong usage of CastChecked for Cast calls when using "ListStreamingTextures" console command.

diff --git a/Engine/Source/Runtime/Engine/Private/Streaming/StreamingManagerTexture.cpp b/Engine/Source/Runtime/Engine/Private/Streaming/StreamingManagerTexture.cpp
--- a/Engine/Source/Runtime/Engine/Private/Streaming/StreamingManagerTexture.cpp
+++ b/Engine/Source/Runtime/Engine/Private/Streaming/StreamingManagerTexture.cpp
@@ -351,6 +351,7 @@ bool FRenderAssetStreamingManager::StreamOutRenderAssetData( int64 RequiredMemor
 
 void FRenderAssetStreamingManager::IncrementalUpdate(float Percentage, bool bUpdateDynamicComponents)
 {
+	TRACE_CPUPROFILER_EVENT_SCOPE(FRenderAssetStreamingManager::IncrementalUpdate);
 	QUICK_SCOPE_CYCLE_COUNTER(STAT_FRenderAssetStreamingManager_IncrementalUpdate);
 	FRemovedRenderAssetArray RemovedRenderAssets;
 
