Added ray tracing support for IsMaterialHidden flag for skeletal meshes

diff --git a/Engine/Source/Runtime/Engine/Private/GPUSkinCache.cpp b/Engine/Source/Runtime/Engine/Private/GPUSkinCache.cpp
--- a/Engine/Source/Runtime/Engine/Private/GPUSkinCache.cpp
+++ b/Engine/Source/Runtime/Engine/Private/GPUSkinCache.cpp
@@ -1751,6 +1751,8 @@ void FGPUSkinCache::ProcessRayTracingGeometryToUpdate(
 			Initializer.bAllowUpdate = true;
 
 			Initializer.Segments.Reserve(LODModel.RenderSections.Num());
+			const int32 LODIndex = SkinCacheEntry->GPUSkin->GetLOD();
+
 			for (const FSkelMeshRenderSection& Section : LODModel.RenderSections)
 			{
 				FRayTracingGeometrySegment Segment;
@@ -1761,7 +1763,9 @@ void FGPUSkinCache::ProcessRayTracingGeometryToUpdate(
 				Segment.MaxVertices = Section.GetNumVertices();
 				Segment.FirstPrimitive = Section.BaseIndex / 3;
 				Segment.NumPrimitives = Section.NumTriangles;
-				Segment.bEnabled = !Section.bDisabled;
+
+				// TODO: If we are at a dropped LOD, route material index through the LODMaterialMap in the LODInfo struct.
+				Segment.bEnabled = !SkinCacheEntry->GPUSkin->IsMaterialHidden(LODIndex, Section.MaterialIndex) && !Section.bDisabled;
 				Initializer.Segments.Add(Segment);
 			}
 
diff --git a/Engine/Source/Runtime/Engine/Private/SkeletalMesh.cpp b/Engine/Source/Runtime/Engine/Private/SkeletalMesh.cpp
--- a/Engine/Source/Runtime/Engine/Private/SkeletalMesh.cpp
+++ b/Engine/Source/Runtime/Engine/Private/SkeletalMesh.cpp
@@ -6274,14 +6274,19 @@ void FSkeletalMeshSceneProxy::GetDynamicRayTracingInstances(FRayTracingMaterialG
 			{
 				TArray<FRayTracingGeometrySegment> GeometrySections;
 				GeometrySections.Reserve(LODData.RenderSections.Num());
-				for (const FSkelMeshRenderSection& Section : LODData.RenderSections)
+
+				for (FSkeletalMeshSectionIter Iter(LODIndex, *MeshObject, LODData, LODSection); Iter; ++Iter)
 				{
+					const FSkelMeshRenderSection& Section = Iter.GetSection();
+					const FSectionElementInfo& SectionElementInfo = Iter.GetSectionElementInfo();
+
 					FRayTracingGeometrySegment Segment;
 					Segment.FirstPrimitive = Section.BaseIndex / 3;
 					Segment.NumPrimitives = Section.NumTriangles;
-					Segment.bEnabled = !Section.bDisabled;
+					Segment.bEnabled = !MeshObject->IsMaterialHidden(LODIndex, SectionElementInfo.UseMaterialIndex) && !Section.bDisabled;
 					GeometrySections.Add(Segment);
 				}
+
 				MeshObject->GetRayTracingGeometry()->Initializer.Segments = GeometrySections;
 
 				Context.DynamicRayTracingGeometriesToUpdate.Add(
diff --git a/Engine/Source/Runtime/Engine/Private/SkeletalRender.cpp b/Engine/Source/Runtime/Engine/Private/SkeletalRender.cpp
--- a/Engine/Source/Runtime/Engine/Private/SkeletalRender.cpp
+++ b/Engine/Source/Runtime/Engine/Private/SkeletalRender.cpp
@@ -40,6 +40,7 @@ FSkeletalMeshObject::FSkeletalMeshObject(USkinnedMeshComponent* InMeshComponent,
 ,   bHasBeenUpdatedAtLeastOnce(false)
 #if RHI_RAYTRACING
 , bRequireRecreatingRayTracingGeometry(false)
+, bHiddenMaterialVisibilityDirtyForRayTracing(false)
 , bSupportRayTracing(IsSkeletalMeshRayTracingSupported() && InMeshComponent->SkeletalMesh->bSupportRayTracing)
 #endif
 #if !UE_BUILD_SHIPPING
@@ -168,7 +169,8 @@ const TArray<FSkelMeshRenderSection>& FSkeletalMeshObject::GetRenderSections(int
 void FSkeletalMeshObject::SetHiddenMaterials(int32 InLODIndex,const TArray<bool>& HiddenMaterials)
 {
 	check(LODInfo.IsValidIndex(InLODIndex));
-	LODInfo[InLODIndex].HiddenMaterials = HiddenMaterials;		
+	bHiddenMaterialVisibilityDirtyForRayTracing = true;
+	LODInfo[InLODIndex].HiddenMaterials = HiddenMaterials;
 }
 
 /**
diff --git a/Engine/Source/Runtime/Engine/Private/SkeletalRenderGPUSkin.cpp b/Engine/Source/Runtime/Engine/Private/SkeletalRenderGPUSkin.cpp
--- a/Engine/Source/Runtime/Engine/Private/SkeletalRenderGPUSkin.cpp
+++ b/Engine/Source/Runtime/Engine/Private/SkeletalRenderGPUSkin.cpp
@@ -379,7 +379,10 @@ void FSkeletalMeshObjectGPUSkin::UpdateDynamicData_RenderThread(FGPUSkinCache* G
 	// So make sure to set bRequireRecreatingRayTracingGeometry to true to account for this, so RT geometry is updated correctly.
 	bRequireRecreatingRayTracingGeometry = (DynamicData == nullptr || RayTracingGeometry.Initializer.Segments.Num() == 0 || // Newly created
 		(DynamicData != nullptr && DynamicData->LODIndex != InDynamicData->LODIndex) || // LOD level changed
-		SkinCacheEntry == nullptr);
+		SkinCacheEntry == nullptr ||
+		bHiddenMaterialVisibilityDirtyForRayTracing);
+
+	bHiddenMaterialVisibilityDirtyForRayTracing = false;
 	
 	if (!bRequireRecreatingRayTracingGeometry)
 	{
diff --git a/Engine/Source/Runtime/Engine/Public/SkeletalRenderPublic.h b/Engine/Source/Runtime/Engine/Public/SkeletalRenderPublic.h
--- a/Engine/Source/Runtime/Engine/Public/SkeletalRenderPublic.h
+++ b/Engine/Source/Runtime/Engine/Public/SkeletalRenderPublic.h
@@ -244,6 +244,7 @@ public:
 #if RHI_RAYTRACING
 	bool bRequireRecreatingRayTracingGeometry;
 	bool bSupportRayTracing;
+	bool bHiddenMaterialVisibilityDirtyForRayTracing;
 #endif
 
 #if UE_BUILD_SHIPPING
