Reduce the impact of wrapped pointers from USkeletalMesh::GetMorphTargets.

diff --git a/Engine/Source/Runtime/Engine/Classes/Engine/SkeletalMesh.h b/Engine/Source/Runtime/Engine/Classes/Engine/SkeletalMesh.h
--- a/Engine/Source/Runtime/Engine/Classes/Engine/SkeletalMesh.h
+++ b/Engine/Source/Runtime/Engine/Classes/Engine/SkeletalMesh.h
@@ -1970,7 +1970,7 @@ public:
 		PRAGMA_ENABLE_DEPRECATION_WARNINGS
 	}
 
-	TArray<UMorphTarget*>& GetMorphTargets()
+	TArray<TObjectPtr<UMorphTarget>>& GetMorphTargets()
 	{
 		WaitUntilAsyncPropertyReleased(ESkeletalMeshAsyncProperties::MorphTargets);
 		PRAGMA_DISABLE_DEPRECATION_WARNINGS
