Reduce the impact of wrapped pointers by 75% from material code.
For example, texture streaming spikes drops from 26.7ms to 16.6ms in PIE.

diff --git a/Engine/Source/Developer/MaterialBaking/Private/ExportMaterialProxy.h b/Engine/Source/Developer/MaterialBaking/Private/ExportMaterialProxy.h
--- a/Engine/Source/Developer/MaterialBaking/Private/ExportMaterialProxy.h
+++ b/Engine/Source/Developer/MaterialBaking/Private/ExportMaterialProxy.h
@@ -315,7 +315,7 @@ public:
 		return bCorrectVertexFactory && bPCPlatform && bCorrectFrequency;
 	}
 
-	virtual TArrayView<UObject* const> GetReferencedTextures() const override
+	virtual TArrayView<const TObjectPtr<UObject>> GetReferencedTextures() const override
 	{
 		return ReferencedTextures;
 	}
@@ -611,7 +611,7 @@ private:
 	/** The material interface for this proxy */
 	UMaterialInterface* MaterialInterface;
 	UMaterial* Material;
-	TArray<UObject*> ReferencedTextures;
+	TArray<TObjectPtr<UObject>> ReferencedTextures;
 	/** The property to compile for rendering the sample */
 	EMaterialProperty PropertyToCompile;
 	/** The name of the specific custom output to compile for rendering the sample. Only used if PropertyToCompile is MP_CustomOutput */
diff --git a/Engine/Source/Developer/MaterialUtilities/Private/MaterialUtilities.cpp b/Engine/Source/Developer/MaterialUtilities/Private/MaterialUtilities.cpp
--- a/Engine/Source/Developer/MaterialUtilities/Private/MaterialUtilities.cpp
+++ b/Engine/Source/Developer/MaterialUtilities/Private/MaterialUtilities.cpp
@@ -484,7 +484,7 @@ public:
 		return true;
 	}
 
-	virtual TArrayView<UObject* const> GetReferencedTextures() const override
+	virtual TArrayView<const TObjectPtr<UObject>> GetReferencedTextures() const override
 	{
 		return ReferencedTextures;
 	}
@@ -740,7 +740,7 @@ private:
 	/** The material interface for this proxy */
 	UMaterialInterface* MaterialInterface;
 	UMaterial* Material;	
-	TArray<UObject*> ReferencedTextures;
+	TArray<TObjectPtr<UObject>> ReferencedTextures;
 	/** The property to compile for rendering the sample */
 	EMaterialProperty PropertyToCompile;
 	FGuid Id;
diff --git a/Engine/Source/Editor/MaterialEditor/Private/MaterialEditor.h b/Engine/Source/Editor/MaterialEditor/Private/MaterialEditor.h
--- a/Engine/Source/Editor/MaterialEditor/Private/MaterialEditor.h
+++ b/Engine/Source/Editor/MaterialEditor/Private/MaterialEditor.h
@@ -106,7 +106,7 @@ public:
 	 */
 	virtual bool ShouldCache(EShaderPlatform Platform, const FShaderType* ShaderType, const FVertexFactoryType* VertexFactoryType) const override;
 
-	virtual TArrayView<UObject* const> GetReferencedTextures() const override
+	virtual TArrayView<const TObjectPtr<UObject>> GetReferencedTextures() const override
 	{
 		return MakeArrayView(ReferencedTextures);
 	}
@@ -212,7 +212,7 @@ public:
 
 private:
 	TWeakObjectPtr<UMaterialExpression> Expression;
-	TArray<UObject*> ReferencedTextures;
+	TArray<TObjectPtr<UObject>> ReferencedTextures;
 	FGuid Id;
 };
 
diff --git a/Engine/Source/Editor/UnrealEd/Private/Lightmass/LightmassRender.cpp b/Engine/Source/Editor/UnrealEd/Private/Lightmass/LightmassRender.cpp
--- a/Engine/Source/Editor/UnrealEd/Private/Lightmass/LightmassRender.cpp
+++ b/Engine/Source/Editor/UnrealEd/Private/Lightmass/LightmassRender.cpp
@@ -275,7 +275,7 @@ public:
 		}
 	}
 
-	virtual TArrayView<UObject* const> GetReferencedTextures() const override
+	virtual TArrayView<const TObjectPtr<UObject>> GetReferencedTextures() const override
 	{
 		return ReferencedTextures;
 	}
@@ -850,7 +850,7 @@ private:
 	/** The material interface for this proxy */
 	UMaterialInterface* MaterialInterface;
 	UMaterial* Material;
-	TArray<UObject*> ReferencedTextures;
+	TArray<TObjectPtr<UObject>> ReferencedTextures;
 	/** The property to compile for rendering the sample */
 	EMaterialProperty PropertyToCompile;
 	/** Stores which exported attribute this proxy is compiling for. */
diff --git a/Engine/Source/Runtime/Engine/Classes/Materials/Material.h b/Engine/Source/Runtime/Engine/Classes/Materials/Material.h
--- a/Engine/Source/Runtime/Engine/Classes/Materials/Material.h
+++ b/Engine/Source/Runtime/Engine/Classes/Materials/Material.h
@@ -1750,7 +1750,7 @@ public:
 #endif // WITH_EDITOR
 
 	/** Appends textures referenced by expressions, including nested functions. */
-	ENGINE_API virtual TArrayView<UObject* const> GetReferencedTextures() const override final { return GetCachedExpressionData().ReferencedTextures; }
+	ENGINE_API virtual TArrayView<const TObjectPtr<UObject>> GetReferencedTextures() const override final { return GetCachedExpressionData().ReferencedTextures; }
 
 protected:
 
diff --git a/Engine/Source/Runtime/Engine/Classes/Materials/MaterialInstance.h b/Engine/Source/Runtime/Engine/Classes/Materials/MaterialInstance.h
--- a/Engine/Source/Runtime/Engine/Classes/Materials/MaterialInstance.h
+++ b/Engine/Source/Runtime/Engine/Classes/Materials/MaterialInstance.h
@@ -676,7 +676,7 @@ public:
 #endif // WITH_EDITOR
 
 	/** Appends textures referenced by expressions, including nested functions. */
-	ENGINE_API virtual TArrayView<UObject* const> GetReferencedTextures() const override final { return CachedData ? CachedData->ReferencedTextures : TArrayView<UObject* const>(); }
+	ENGINE_API virtual TArrayView<const TObjectPtr<UObject>> GetReferencedTextures() const override final { return CachedData ? CachedData->ReferencedTextures : TArrayView<const TObjectPtr<UObject>>(); }
 
 #if WITH_EDITOR
 	/** Add to the set any texture referenced by expressions, including nested functions, as well as any overrides from parameters. */
diff --git a/Engine/Source/Runtime/Engine/Classes/Materials/MaterialInterface.h b/Engine/Source/Runtime/Engine/Classes/Materials/MaterialInterface.h
--- a/Engine/Source/Runtime/Engine/Classes/Materials/MaterialInterface.h
+++ b/Engine/Source/Runtime/Engine/Classes/Materials/MaterialInterface.h
@@ -548,8 +548,8 @@ public:
 		PURE_VIRTUAL(UMaterialInterface::GetLayerParameterIndex, return INDEX_NONE;);
 
 	/** Get textures referenced by expressions, including nested functions. */
-	virtual TArrayView<UObject* const> GetReferencedTextures() const
-		PURE_VIRTUAL(UMaterialInterface::GetReferencedTextures,return TArrayView<UObject* const>(););
+	virtual TArrayView<const TObjectPtr<UObject>> GetReferencedTextures() const
+		PURE_VIRTUAL(UMaterialInterface::GetReferencedTextures,return TArrayView<const TObjectPtr<UObject>>(););
 
 	virtual void SaveShaderStableKeysInner(const class ITargetPlatform* TP, const struct FStableShaderKeyAndValue& SaveKeyVal)
 		PURE_VIRTUAL(UMaterialInterface::SaveShaderStableKeysInner, );
diff --git a/Engine/Source/Runtime/Engine/Private/Materials/MaterialShared.cpp b/Engine/Source/Runtime/Engine/Private/Materials/MaterialShared.cpp
--- a/Engine/Source/Runtime/Engine/Private/Materials/MaterialShared.cpp
+++ b/Engine/Source/Runtime/Engine/Private/Materials/MaterialShared.cpp
@@ -1226,11 +1226,11 @@ void FMaterialResource::LegacySerialize(FArchive& Ar)
 	}
 }
 
-TArrayView<UObject* const> FMaterialResource::GetReferencedTextures() const
+TArrayView<const TObjectPtr<UObject>> FMaterialResource::GetReferencedTextures() const
 {
 	if (MaterialInstance)
 	{
-		const TArrayView<UObject* const> Textures = MaterialInstance->GetReferencedTextures();
+		const TArrayView<const TObjectPtr<UObject>> Textures = MaterialInstance->GetReferencedTextures();
 		if (Textures.Num())
 		{
 			return Textures;
@@ -3633,7 +3633,7 @@ void FMaterial::GetReferencedTexturesHash(EShaderPlatform Platform, FSHAHash& Ou
 {
 	FSHA1 HashState;
 
-	const TArrayView<UObject* const> ReferencedTextures = GetReferencedTextures();
+	const TArrayView<const TObjectPtr<UObject>> ReferencedTextures = GetReferencedTextures();
 	// Hash the names of the uniform expression textures to capture changes in their order or values resulting from material compiler code changes
 	for (int32 TextureIndex = 0; TextureIndex < ReferencedTextures.Num(); TextureIndex++)
 	{
diff --git a/Engine/Source/Runtime/Engine/Private/Materials/MaterialUniformExpressions.h b/Engine/Source/Runtime/Engine/Private/Materials/MaterialUniformExpressions.h
--- a/Engine/Source/Runtime/Engine/Private/Materials/MaterialUniformExpressions.h
+++ b/Engine/Source/Runtime/Engine/Private/Materials/MaterialUniformExpressions.h
@@ -325,7 +325,7 @@ template<typename TextureType>
 static TextureType* GetIndexedTexture(const FMaterial& Material, int32 TextureIndex)
 {
 	UObject* IndexedTexture = nullptr;
-	const TArrayView<UObject* const> ReferencedTextures = Material.GetReferencedTextures();
+	const TArrayView<const TObjectPtr<UObject>> ReferencedTextures = Material.GetReferencedTextures();
 	if (ReferencedTextures.IsValidIndex(TextureIndex))
 	{
 		IndexedTexture = ReferencedTextures[TextureIndex];
diff --git a/Engine/Source/Runtime/Engine/Public/MaterialShared.h b/Engine/Source/Runtime/Engine/Public/MaterialShared.h
--- a/Engine/Source/Runtime/Engine/Public/MaterialShared.h
+++ b/Engine/Source/Runtime/Engine/Public/MaterialShared.h
@@ -1875,7 +1875,7 @@ public:
 
 	ENGINE_API virtual void AddReferencedObjects(FReferenceCollector& Collector);
 
-	virtual TArrayView<UObject* const> GetReferencedTextures() const = 0;
+	virtual TArrayView<const TObjectPtr<UObject>> GetReferencedTextures() const = 0;
 
 	/**
 	 * Finds the shader matching the template type and the passed in vertex factory, asserts if not found.
@@ -2622,7 +2622,7 @@ public:
 
 	ENGINE_API virtual void LegacySerialize(FArchive& Ar) override;
 
-	ENGINE_API virtual TArrayView<UObject* const> GetReferencedTextures() const override;
+	ENGINE_API virtual TArrayView<const TObjectPtr<UObject>> GetReferencedTextures() const override;
 
 	ENGINE_API virtual void AddReferencedObjects(FReferenceCollector& Collector) override;
 
