Fix UE.EditorAutomation (RunTest=System) Win64.

diff --git a/Engine/Source/Runtime/Engine/Private/WorldPartition/WorldPartitionActorDesc.cpp b/Engine/Source/Runtime/Engine/Private/WorldPartition/WorldPartitionActorDesc.cpp
--- a/Engine/Source/Runtime/Engine/Private/WorldPartition/WorldPartitionActorDesc.cpp
+++ b/Engine/Source/Runtime/Engine/Private/WorldPartition/WorldPartitionActorDesc.cpp
@@ -269,10 +269,13 @@ AActor* FWorldPartitionActorDesc::Load() const
 		}
 
 		Package = LoadPackage(Package, *ActorPackage.ToString(), LOAD_None, nullptr, InstancingContext);
-		check(Package);
+		check(Package || GIsAutomationTesting);
 
-		ActorPtr = FindObject<AActor>(nullptr, *ActorPath.ToString());
-		check(ActorPtr.IsValid());
+		if (Package)
+		{
+			ActorPtr = FindObject<AActor>(nullptr, *ActorPath.ToString());
+			check(ActorPtr.IsValid());
+		}
 	}
 
 	return ActorPtr.Get();
