Material vector input explicitly uses Vector3f

diff --git a/Engine/Source/Runtime/Engine/Classes/Materials/Material.h b/Engine/Source/Runtime/Engine/Classes/Materials/Material.h
--- a/Engine/Source/Runtime/Engine/Classes/Materials/Material.h
+++ b/Engine/Source/Runtime/Engine/Classes/Materials/Material.h
@@ -251,7 +251,7 @@ struct FVectorMaterialInput : public FMaterialInput
 	uint32 UseConstant : 1;
 
 	UPROPERTY()
-	FVector Constant;
+	FVector3f Constant;
 
 	FVectorMaterialInput:
 		UseConstant(0) {}
diff --git a/Engine/Source/Runtime/Engine/Private/Materials/MaterialShared.cpp b/Engine/Source/Runtime/Engine/Private/Materials/MaterialShared.cpp
--- a/Engine/Source/Runtime/Engine/Private/Materials/MaterialShared.cpp
+++ b/Engine/Source/Runtime/Engine/Private/Materials/MaterialShared.cpp
@@ -483,7 +483,7 @@ bool FStrataMaterialInput::Serialize(FArchive& Ar)
 
 bool FVectorMaterialInput::Serialize(FArchive& Ar)
 {
-	return SerializeMaterialInput<FVector>(Ar, *this);
+	return SerializeMaterialInput<FVector3f>(Ar, *this);
 }
 
 bool FVector2MaterialInput::Serialize(FArchive& Ar)
diff --git a/Engine/Source/Runtime/Engine/Public/MaterialExpressionIO.h b/Engine/Source/Runtime/Engine/Public/MaterialExpressionIO.h
--- a/Engine/Source/Runtime/Engine/Public/MaterialExpressionIO.h
+++ b/Engine/Source/Runtime/Engine/Public/MaterialExpressionIO.h
@@ -273,7 +273,7 @@ struct TStructOpsTypeTraits<FStrataMaterialInput>
 	};
 };
 
-struct FVectorMaterialInput : FMaterialInput<FVector>
+struct FVectorMaterialInput : FMaterialInput<FVector3f>
 {
 #if WITH_EDITOR
 	ENGINE_API int32 CompileWithDefault(class FMaterialCompiler* Compiler, EMaterialProperty Property);
diff --git a/Engine/Source/Runtime/Engine/Public/Shader/ShaderTypes.h b/Engine/Source/Runtime/Engine/Public/Shader/ShaderTypes.h
--- a/Engine/Source/Runtime/Engine/Public/Shader/ShaderTypes.h
+++ b/Engine/Source/Runtime/Engine/Public/Shader/ShaderTypes.h
@@ -131,7 +131,7 @@ struct FValue
 		Component[1].Float = Value.Y;
 	}
 
-	inline FValue(const FVector& Value) : ComponentType(EValueComponentType::Float), NumComponents(3)
+	inline FValue(const FVector3f& Value) : ComponentType(EValueComponentType::Float), NumComponents(3)
 	{
 		Component[0].Float = Value.X;
 		Component[1].Float = Value.Y;
