Kept track of the number of dirty packages the user has in the Editor analytics summary session to be able to estimate how often user lost work after a crash/abnormal shutdown.

diff --git a/Engine/Source/Runtime/Engine/Private/Analytics/EditorAnalyticsSessionSummary.cpp b/Engine/Source/Runtime/Engine/Private/Analytics/EditorAnalyticsSessionSummary.cpp
--- a/Engine/Source/Runtime/Engine/Private/Analytics/EditorAnalyticsSessionSummary.cpp
+++ b/Engine/Source/Runtime/Engine/Private/Analytics/EditorAnalyticsSessionSummary.cpp
@@ -11,6 +11,7 @@
 #include "Interfaces/IAnalyticsPropertyStore.h"
 #include "Kismet2/DebuggerCommands.h"
 #include "ProfilingDebugging/StallDetector.h"
+#include "UObject/Package.h"
 
 namespace EditorAnalyticsProperties
 {
@@ -21,6 +22,7 @@ namespace EditorAnalyticsProperties
 	static const TAnalyticsProperty<int32>   Idle30Min              = TEXT("30MinIdle");
 	static const TAnalyticsProperty<bool>    IsInPIE                = TEXT("IsInPIE");
 	static const TAnalyticsProperty<bool>    IsInVRMode             = TEXT("IsInVRMode");
+	static const TAnalyticsProperty<uint32>  DirtyPackageCount      = TEXT("DirtyPackageCount");
 
 	static const TAnalyticsProperty<uint32>  TotalStallCount        = TEXT("TotalStallCount");
 	static const TAnalyticsProperty<uint32>  TotalStallReported     = TEXT("TotalStallReported");
@@ -43,6 +45,7 @@ FEditorAnalyticsSessionSummary::FEditorAnalyticsSessionSummary(TSharedPtr<IAnaly
 	EditorAnalyticsProperties::Idle30Min.Set(GetStore(), 0);
 	EditorAnalyticsProperties::IsInVRMode.Set(GetStore(), IVREditorModule::Get().IsVREditorModeActive());
 	EditorAnalyticsProperties::IsInPIE.Set(GetStore(), FPlayWorldCommandCallbacks::IsInPIE());
+	EditorAnalyticsProperties::DirtyPackageCount.Set(GetStore(), 0);
 
 	EditorAnalyticsProperties::TopStallName.Set(GetStore(), TEXT(""), /*Capacity*/128);
 	EditorAnalyticsProperties::TopStallBudgetSeconds.Set(GetStore(), 0.0);
@@ -55,7 +58,8 @@ FEditorAnalyticsSessionSummary::FEditorAnalyticsSessionSummary(TSharedPtr<IAnaly
 	// The current summary revision number. Identify the key set used and its behavior. If we add/remove keys, we should increment the number. The number can
 	// also be incremented when we change some behaviors to be able to compare between versions. Mostly useful for 'dev' branches where users bypasses UGS to update/build.
 	//    - V3 -> Windows optimization for stall/ensure -> The engine only captures the responsible thread so CRC walks 1 thread rather than all threads.
-	EditorAnalyticsProperties::SummaryEventVersion.Set(GetStore(), 3);
+	//    - V4 -> Added DirtyPackageCount property.
+	EditorAnalyticsProperties::SummaryEventVersion.Set(GetStore(), 4);
 
 	// Persist the session to disk.
 	GetStore()->Flush();
@@ -64,6 +68,7 @@ FEditorAnalyticsSessionSummary::FEditorAnalyticsSessionSummary(TSharedPtr<IAnaly
 	FEditorDelegates::EndPIE.AddRaw(this, &FEditorAnalyticsSessionSummary::OnExitPIE);
 	FSlateApplication::Get().GetOnModalLoopTickEvent().AddRaw(this, &FEditorAnalyticsSessionSummary::Tick);
 	FSlateApplication::Get().GetLastUserInteractionTimeUpdateEvent().AddRaw(this, &FEditorAnalyticsSessionSummary::OnSlateUserInteraction);
+	UPackage::PackageDirtyStateChangedEvent.AddRaw(this, &FEditorAnalyticsSessionSummary::OnDirtyPackageStateChanged);
 }
 
 void FEditorAnalyticsSessionSummary::ShutdownInternal()
@@ -72,6 +77,7 @@ void FEditorAnalyticsSessionSummary::ShutdownInternal()
 	FEditorDelegates::EndPIE.RemoveAll(this);
 	FSlateApplication::Get().GetOnModalLoopTickEvent().RemoveAll(this);
 	FSlateApplication::Get().GetLastUserInteractionTimeUpdateEvent().RemoveAll(this);
+	UPackage::PackageDirtyStateChangedEvent.RemoveAll(this);
 }
 
 bool FEditorAnalyticsSessionSummary::UpdateSessionProgressInternal(bool bCrashing)
@@ -192,4 +198,34 @@ void FEditorAnalyticsSessionSummary::OnExitPIE(const bool /*bIsSimulating*/)
 	GetStore()->Flush(/*bAsync*/true, FTimespan::Zero());
 }
 
+void FEditorAnalyticsSessionSummary::OnDirtyPackageStateChanged(UPackage* InPackage)
+{
+	if (InPackage->HasAnyPackageFlags(PKG_CompiledIn))
+	{
+		return; // We are interested to know if user has unsaved work. Compiled in package are irrelevant for that.
+	}
+	else if (InPackage->IsDirty()) // Newly created + modified packages
+	{
+		EditorAnalyticsProperties::DirtyPackageCount.Update(GetStore(), [](uint32& Actual)
+		{
+			++Actual;
+			return true;
+		});
+	}
+	else
+	{
+		EditorAnalyticsProperties::DirtyPackageCount.Update(GetStore(), [](uint32& Actual)
+		{
+			if (Actual > 0) // In case the callback was register too late and missed some changes.
+			{
+				--Actual;
+				return true;
+			}
+			return false;
+		});
+	}
+
+	GetStore()->Flush(/*bAsync*/true, FTimespan::Zero());
+}
+
 #endif // WITH_EDITOR
diff --git a/Engine/Source/Runtime/Engine/Private/Analytics/EditorAnalyticsSessionSummary.h b/Engine/Source/Runtime/Engine/Private/Analytics/EditorAnalyticsSessionSummary.h
--- a/Engine/Source/Runtime/Engine/Private/Analytics/EditorAnalyticsSessionSummary.h
+++ b/Engine/Source/Runtime/Engine/Private/Analytics/EditorAnalyticsSessionSummary.h
@@ -9,6 +9,8 @@
 #include "Analytics/EngineAnalyticsSessionSummary.h"
 #include <atomic>
 
+class UPackage;
+
 /** Specializes the base engine analytics for the Editor, adding Editor only data to the summary. */
 class FEditorAnalyticsSessionSummary : public FEngineAnalyticsSessionSummary
 {
@@ -27,6 +29,7 @@ private:
 	void OnSlateUserInteraction(double CurrSlateInteractionTime);
 	void OnEnterPIE(const bool /*bIsSimulating*/);
 	void OnExitPIE(const bool /*bIsSimulating*/);
+	void OnDirtyPackageStateChanged(UPackage* package);
 
 private:
 	/** Last activity (user input, crash, terminate, shutdown) timestamp from FPlatformTime::Seconds() to track user inactivity. */
