[MorphTargets] - Discard CPU-side versions of the buffer to save memory on non-editor platforms.

- r.MorphTarget.Mode is a Read Only CVar to avoid the situation of not having CPU-side data if you were to toggle r.MorphTarget.Mode to 0 (CPU processing)

diff --git a/Engine/Source/Runtime/Engine/Classes/Animation/MorphTarget.h b/Engine/Source/Runtime/Engine/Classes/Animation/MorphTarget.h
--- a/Engine/Source/Runtime/Engine/Classes/Animation/MorphTarget.h
+++ b/Engine/Source/Runtime/Engine/Classes/Animation/MorphTarget.h
@@ -106,6 +106,12 @@ struct FMorphTargetLODModel
 		// this makes it clear to clear up later
 		bGeneratedByEngine = true;
 	}
+
+	void DiscardCPUBuffers()
+	{
+		Vertices.Empty();
+		SectionIndices.Empty();
+	}
 };
 
 
@@ -132,6 +138,9 @@ public:
 	/** return true if this morphtarget contains valid vertices */
 	ENGINE_API bool HasValidData() const;
 
+	/** Discard CPU Buffers after render resources have been created. */
+	ENGINE_API void DiscardCPUBuffers();
+
 #if WITH_EDITOR
 	/** Populates the given morph target LOD model with the provided deltas */
 	ENGINE_API void PopulateDeltas(const TArray<FMorphTargetDelta>& Deltas, const int32 LODIndex, const TArray<struct FSkelMeshSection>& Sections, const bool bCompareNormal = false, const bool bGeneratedByReductionSetting = false, const float PositionThreshold = THRESH_POINTS_ARE_NEAR);
diff --git a/Engine/Source/Runtime/Engine/Private/MorphTools.cpp b/Engine/Source/Runtime/Engine/Private/MorphTools.cpp
--- a/Engine/Source/Runtime/Engine/Private/MorphTools.cpp
+++ b/Engine/Source/Runtime/Engine/Private/MorphTools.cpp
@@ -52,6 +52,14 @@ bool UMorphTarget::HasValidData() const
 	return false;
 }
 
+void UMorphTarget::DiscardCPUBuffers()
+{
+	for (FMorphTargetLODModel& Model : MorphLODModels)
+	{
+		Model.DiscardCPUBuffers();
+	}
+}
+
 #if WITH_EDITOR
 
 void UMorphTarget::PopulateDeltas(const TArray<FMorphTargetDelta>& Deltas, const int32 LODIndex, const TArray<FSkelMeshSection>& Sections, const bool bCompareNormal, const bool bGeneratedByReductionSetting, const float PositionThreshold)
diff --git a/Engine/Source/Runtime/Engine/Private/SkeletalMeshLODRenderData.cpp b/Engine/Source/Runtime/Engine/Private/SkeletalMeshLODRenderData.cpp
--- a/Engine/Source/Runtime/Engine/Private/SkeletalMeshLODRenderData.cpp
+++ b/Engine/Source/Runtime/Engine/Private/SkeletalMeshLODRenderData.cpp
@@ -233,6 +233,15 @@ void FSkeletalMeshLODRenderData::InitResources(bool bNeedsVertexColors, int32 LO
 				MorphTargetSize = (MorphTargetSize > FMorphTargetVertexInfoBuffers::GetMaximumThreadGroupSize()) ? MorphTargetSize - FMorphTargetVertexInfoBuffers::GetMaximumThreadGroupSize() : 0;
 				StartOffset += FMorphTargetVertexInfoBuffers::GetMaximumThreadGroupSize();
 			} while (MorphTargetSize > 0);
+
+#if !WITH_EDITOR
+			if (NumSrcDeltas > 0)
+			{
+				// A CPU copy of the morph deltas has been made so it is safe to 
+				// discard the original data.  Keep CPU buffers when in the editor.
+				MorphTarget->DiscardCPUBuffers();
+			}
+#endif
 		}
 
 		check(MorphTargetVertexInfoBuffers.WorkItemsPerMorph.Num() == MorphTargetVertexInfoBuffers.StartOffsetPerMorph.Num());
diff --git a/Engine/Source/Runtime/Engine/Private/SkeletalRenderGPUSkin.cpp b/Engine/Source/Runtime/Engine/Private/SkeletalRenderGPUSkin.cpp
--- a/Engine/Source/Runtime/Engine/Private/SkeletalRenderGPUSkin.cpp
+++ b/Engine/Source/Runtime/Engine/Private/SkeletalRenderGPUSkin.cpp
@@ -43,7 +43,7 @@ static FAutoConsoleVariableRef CVarUseGPUMorphTargets(
 	TEXT("Use GPU for computing morph targets.\n")
 	TEXT(" 0: Use original CPU method (loop per morph then by vertex)\n")
 	TEXT(" 1: Enable GPU method (default)\n"),
-	ECVF_Default
+	ECVF_ReadOnly
 	);
 
 static int32 GForceUpdateMorphTargets = 0;
