Fixed compilation error

diff --git a/Engine/Source/Runtime/Engine/Private/SkeletalRender.cpp b/Engine/Source/Runtime/Engine/Private/SkeletalRender.cpp
--- a/Engine/Source/Runtime/Engine/Private/SkeletalRender.cpp
+++ b/Engine/Source/Runtime/Engine/Private/SkeletalRender.cpp
@@ -169,7 +169,9 @@ const TArray<FSkelMeshRenderSection>& FSkeletalMeshObject::GetRenderSections(int
 void FSkeletalMeshObject::SetHiddenMaterials(int32 InLODIndex,const TArray<bool>& HiddenMaterials)
 {
 	check(LODInfo.IsValidIndex(InLODIndex));
+#if RHI_RAYTRACING
 	bHiddenMaterialVisibilityDirtyForRayTracing = true;
+#endif
 	LODInfo[InLODIndex].HiddenMaterials = HiddenMaterials;
 }
 
