Register fallback billboard component only if the world is initialized

diff --git a/Engine/Source/Runtime/Engine/Private/ActorConstruction.cpp b/Engine/Source/Runtime/Engine/Private/ActorConstruction.cpp
--- a/Engine/Source/Runtime/Engine/Private/ActorConstruction.cpp
+++ b/Engine/Source/Runtime/Engine/Private/ActorConstruction.cpp
@@ -918,7 +918,10 @@ void AActor::ProcessUserConstructionScript()
 
 void AActor::FinishAndRegisterComponent(UActorComponent* Component)
 {
-	Component->RegisterComponent();
+	if (GetWorld()->bIsWorldInitialized)
+	{
+		Component->RegisterComponent();
+	}
 	BlueprintCreatedComponents.Add(Component);
 }
 
