Fixed crash when serializing legacy speedtree assets.

diff --git a/Engine/Source/Runtime/Engine/Private/StaticMesh.cpp b/Engine/Source/Runtime/Engine/Private/StaticMesh.cpp
--- a/Engine/Source/Runtime/Engine/Private/StaticMesh.cpp
+++ b/Engine/Source/Runtime/Engine/Private/StaticMesh.cpp
@@ -4970,8 +4970,12 @@ void UStaticMesh::Serialize(FArchive& Ar)
 			const int32 LODIndex = GetNumSourceModels() - 1;
 			FStaticMeshSourceModel& SourceModel = GetSourceModel(LODIndex);
 
+			// If we get here, it is a very old version which still serializes as RawMesh.
+			// Hence we can expect the RawMeshBulkData to be valid.
+			// At this point it will not have been converted to MeshDescription.
+			ensure(!SourceModel.RawMeshBulkData->IsEmpty());
 			FRawMesh RawMesh;
-			SourceModel.LoadRawMesh(RawMesh);
+			SourceModel.RawMeshBulkData->LoadRawMesh(RawMesh);
 
 			// Billboard LOD is made up out of quads so check for this
 			bool bQuadVertices = ((RawMesh.VertexPositions.Num() % 4) == 0);
