Fix primitive ID for invalid instance data as it must be 0xFFFFFFFF to indicate invalid.

diff --git a/Engine/Source/Runtime/Engine/Private/InstanceUniformShaderParameters.cpp b/Engine/Source/Runtime/Engine/Private/InstanceUniformShaderParameters.cpp
--- a/Engine/Source/Runtime/Engine/Private/InstanceUniformShaderParameters.cpp
+++ b/Engine/Source/Runtime/Engine/Private/InstanceUniformShaderParameters.cpp
@@ -63,6 +63,6 @@ ENGINE_API const FPrimitiveInstance& GetDummyPrimitiveInstance()
 
 ENGINE_API const FInstanceSceneShaderData& GetDummyInstanceSceneShaderData()
 {
-	static FInstanceSceneShaderData DummyShaderData = FInstanceSceneShaderData(GetDummyPrimitiveInstance(), 0 /* Primitive Id */);
+	static FInstanceSceneShaderData DummyShaderData = FInstanceSceneShaderData(GetDummyPrimitiveInstance(), 0xFFFFFFFFu /* Primitive Id */);
 	return DummyShaderData;
 }
