Fix RHI validation error "Attempted to access resource SkinCachePositions"

It happens when the skin cache skinning pass isn't called on the current frame's position buffer but the recompute tangent pass is executed on it. In this case the current position buffer is transitioned to ERHIAccess::SRVCompute but it isn't added to BuffersToTransitionToRead to be transitioned into ERHIAccess::VertexOrIndexBuffer, causing the validation error later on when rendering passes use itr for rendering.

diff --git a/Engine/Source/Runtime/Engine/Private/GPUSkinCache.cpp b/Engine/Source/Runtime/Engine/Private/GPUSkinCache.cpp
--- a/Engine/Source/Runtime/Engine/Private/GPUSkinCache.cpp
+++ b/Engine/Source/Runtime/Engine/Private/GPUSkinCache.cpp
@@ -1364,6 +1364,7 @@ void FGPUSkinCache::DoDispatch(FRHICommandListImmediate& RHICmdList)
 			{
 				IntermediateAccumulatedTangentBuffers.AddUnique(DispatchData.GetIntermediateAccumulatedTangentBuffer());
 			}
+			BuffersToTransitionToRead.Add(DispatchData.GetPositionRWBuffer());
 		}	
 	}
 	MakeBufferTransitions(RHICmdList, BuffersToSRVForRecomputeTangents, ERHIAccess::SRVCompute);
@@ -1456,6 +1457,7 @@ void FGPUSkinCache::DoDispatch(FRHICommandListImmediate& RHICmdList, FGPUSkinCac
 		{
 			RHICmdList.Transition(DispatchData.GetIntermediateAccumulatedTangentBuffer()->UpdateAccessState(ERHIAccess::UAVCompute));
 		}
+		BuffersToTransitionToRead.Add(DispatchData.GetPositionRWBuffer());
 
 		FSkinCacheRWBuffer* StagingBuffer = nullptr;
 		DispatchUpdateSkinTangents(RHICmdList, SkinCacheEntry, Section, StagingBuffer, true);
