Skylight now do an incomplete capture every 5 seconds even when the world is not yet ready (to progressively refine captured skylight)
Made an enum class and removed unecessary code.

diff --git a/Engine/Source/Runtime/Engine/Classes/Components/SkyLightComponent.h b/Engine/Source/Runtime/Engine/Classes/Components/SkyLightComponent.h
--- a/Engine/Source/Runtime/Engine/Classes/Components/SkyLightComponent.h
+++ b/Engine/Source/Runtime/Engine/Classes/Components/SkyLightComponent.h
@@ -88,7 +88,7 @@ enum ESkyLightSourceType
 	SLS_MAX,
 };
 
-enum ESkyLightCaptureStatus
+enum class ESkyLightCaptureStatus
 {
 	SLCS_Uninitialized,
 	SLCS_CapturedButIncomplete,
@@ -313,6 +313,7 @@ protected:
 	// In an attempt to get valid lighting data as soon as possible in a level, we will always trigger a skylight capture at creation.
 	// And then, a soon as the world is finally loaded, the final capture will be taken.
 	ESkyLightCaptureStatus CaptureStatus;
+	float SecondsSinceLastCapture;
 #endif
 
 	TRefCountPtr<FSkyTextureCubeResource> ProcessedSkyTexture;
diff --git a/Engine/Source/Runtime/Engine/Private/Components/SkyLightComponent.cpp b/Engine/Source/Runtime/Engine/Private/Components/SkyLightComponent.cpp
--- a/Engine/Source/Runtime/Engine/Private/Components/SkyLightComponent.cpp
+++ b/Engine/Source/Runtime/Engine/Private/Components/SkyLightComponent.cpp
@@ -222,6 +222,7 @@ FSkyLightSceneProxy::FSkyLightSceneProxy(const USkyLightComponent* InLightCompon
 	, LowerHemisphereColor(InLightComponent->LowerHemisphereColor)
 	, bLowerHemisphereIsSolidColor(InLightComponent->bLowerHemisphereIsBlack)
 #if WITH_EDITOR
+	, SecondsToNextIncompleteCapture(0.0f)
 	, bCubemapSkyLightWaitingForCubeMapTexture(false)
 	, bCaptureSkyLightWaitingForShaders(false)
 	, bCaptureSkyLightWaitingForMeshesOrTextures(false)
@@ -287,7 +288,8 @@ USkyLightComponent::USkyLightComponent(const FObjectInitializer& ObjectInitializ
 	CloudAmbientOcclusionApertureScale = 0.05f;
 
 #if WITH_EDITOR
-	CaptureStatus = SLCS_Uninitialized;
+	CaptureStatus = ESkyLightCaptureStatus::SLCS_Uninitialized;
+	SecondsSinceLastCapture = 0.0f;
 #endif
 }
 
@@ -704,28 +706,20 @@ void USkyLightComponent::UpdateSkyCaptureContentsArray(UWorld* WorldToUpdate, TA
 			FTextureCompilingManager::Get().RequestPriorityChange(CaptureComponent->Cubemap, EQueuedWorkPriority::Highest);
 		}
 
+		const float SecondsBetweenIncompleteCaptures = 5.0f;
 		const bool bCubemapSkyLightWaitingForCubemapAsset	= CaptureComponent->SourceType == SLS_SpecifiedCubemap	&& bIsCubemapCompiling;
 		const bool bCaptureSkyLightWaitingCompiledShader	= CaptureComponent->SourceType == SLS_CapturedScene		&& bIsCompilingShaders;
 		const bool bCaptureSkyLightWaitingForMeshOrTexAssets= CaptureComponent->SourceType == SLS_CapturedScene		&& bSceneIsAsyncCompiling;
-		ENQUEUE_RENDER_COMMAND(FUpdateSkyLightProxyStatus)(
-			[CaptureComponent, bCubemapSkyLightWaitingForCubemapAsset, bCaptureSkyLightWaitingCompiledShader, bCaptureSkyLightWaitingForMeshOrTexAssets](FRHICommandList& RHICmdList)
-			{
-				FSkyLightSceneProxy* SkyLightSceneProxy = CaptureComponent->SceneProxy;
-				if (SkyLightSceneProxy)
-				{
-					SkyLightSceneProxy->bCubemapSkyLightWaitingForCubeMapTexture = bCubemapSkyLightWaitingForCubemapAsset;
-					SkyLightSceneProxy->bCaptureSkyLightWaitingForShaders = bCaptureSkyLightWaitingCompiledShader;
-					SkyLightSceneProxy->bCaptureSkyLightWaitingForMeshesOrTextures = bCaptureSkyLightWaitingForMeshOrTexAssets;
-				}
-			});
 #endif
 
 		if (((!Owner || !Owner->GetLevel() || Owner->GetLevel()->bIsVisible) && CaptureComponent->GetWorld() == WorldToUpdate)
 			// Only process sky capture requests once async texture and shader compiling completes, otherwise we will capture the scene with temporary shaders/textures
 			&& (
 #if WITH_EDITOR
-				CaptureComponent->CaptureStatus == SLCS_Uninitialized 
+				CaptureComponent->CaptureStatus == ESkyLightCaptureStatus::SLCS_Uninitialized
 				|| 
+				(CaptureComponent->CaptureStatus == ESkyLightCaptureStatus::SLCS_CapturedButIncomplete && CaptureComponent->SecondsSinceLastCapture > SecondsBetweenIncompleteCaptures)
+				||
 #endif
 				((!bSceneIsAsyncCompiling) && (!bIsCompilingShaders)) 
 				|| 
@@ -780,33 +774,45 @@ void USkyLightComponent::UpdateSkyCaptureContentsArray(UWorld* WorldToUpdate, TA
 			}
 
 #if WITH_EDITOR
+			const bool bCaptureIsComplete = !bCubemapSkyLightWaitingForCubemapAsset && !bCaptureSkyLightWaitingCompiledShader && !bCaptureSkyLightWaitingForMeshOrTexAssets;
 			switch (CaptureComponent->CaptureStatus)
 			{
-			case SLCS_Uninitialized:
+			case ESkyLightCaptureStatus::SLCS_Uninitialized:
 			{
-				if (!bCubemapSkyLightWaitingForCubemapAsset && !bCaptureSkyLightWaitingCompiledShader && !bCaptureSkyLightWaitingForMeshOrTexAssets)
+				CaptureComponent->SecondsSinceLastCapture = 0.0f;
+
+				if (bCaptureIsComplete)
 				{
 					// Do not recapture if the first forced capture was with a complete world (to avoid capturing twice each time a level is loaded or a skylight created))
-					CaptureComponent->CaptureStatus = SLCS_CapturedAndComplete;
+					CaptureComponent->CaptureStatus = ESkyLightCaptureStatus::SLCS_CapturedAndComplete;
 					ComponentArray.RemoveAt(CaptureIndex);
 				}
 				else
 				{
-					CaptureComponent->CaptureStatus = SLCS_CapturedButIncomplete;
+					CaptureComponent->CaptureStatus = ESkyLightCaptureStatus::SLCS_CapturedButIncomplete;
 				}
 				break;
 			}
-			case SLCS_CapturedButIncomplete:
+			case ESkyLightCaptureStatus::SLCS_CapturedButIncomplete:
 			{
-				// Only remove queued update requests if we processed it for the a world with all meshes, textures and shaders.
-				ComponentArray.RemoveAt(CaptureIndex);
-				CaptureComponent->CaptureStatus = SLCS_CapturedAndComplete;
+				if (bCaptureIsComplete)
+				{
+					// Only remove queued update requests if we processed it for the a world with all meshes, textures and shaders.
+					ComponentArray.RemoveAt(CaptureIndex);
+					CaptureComponent->CaptureStatus = ESkyLightCaptureStatus::SLCS_CapturedAndComplete;
+				}
+				else if (CaptureComponent->SecondsSinceLastCapture > SecondsBetweenIncompleteCaptures)
+				{
+					// We have just executed another incomplete capture, so reset the timer for the next one.
+					CaptureComponent->SecondsSinceLastCapture = 0.0f;
+				}
 				break;
 			}
-			case SLCS_CapturedAndComplete:
+			case ESkyLightCaptureStatus::SLCS_CapturedAndComplete:
 			{
 				// It is valid to recapture a complete skylight.
 				ComponentArray.RemoveAt(CaptureIndex);
+				CaptureComponent->SecondsSinceLastCapture = 0.0f;
 				break;
 			}
 			default:
@@ -821,12 +827,15 @@ void USkyLightComponent::UpdateSkyCaptureContentsArray(UWorld* WorldToUpdate, TA
 		}
 
 #if WITH_EDITOR
+		CaptureComponent->SecondsSinceLastCapture += CaptureComponent->CaptureStatus == ESkyLightCaptureStatus::SLCS_CapturedButIncomplete ? WorldToUpdate->DeltaTimeSeconds : 0.0f;
+
 		ENQUEUE_RENDER_COMMAND(FUpdateSkyLightProxyStatusForcedCapture)(
-			[CaptureComponent, bCubemapSkyLightWaitingForCubemapAsset, bCaptureSkyLightWaitingCompiledShader, bCaptureSkyLightWaitingForMeshOrTexAssets](FRHICommandList& RHICmdList)
+			[CaptureComponent, SecondsBetweenIncompleteCaptures, bCubemapSkyLightWaitingForCubemapAsset, bCaptureSkyLightWaitingCompiledShader, bCaptureSkyLightWaitingForMeshOrTexAssets](FRHICommandList& RHICmdList)
 			{
 				FSkyLightSceneProxy* SkyLightSceneProxy = CaptureComponent->SceneProxy;
 				if (SkyLightSceneProxy)
 				{
+					SkyLightSceneProxy->SecondsToNextIncompleteCapture = FMath::Max(0.0f, SecondsBetweenIncompleteCaptures - CaptureComponent->SecondsSinceLastCapture);
 					SkyLightSceneProxy->bCubemapSkyLightWaitingForCubeMapTexture = bCubemapSkyLightWaitingForCubemapAsset;
 					SkyLightSceneProxy->bCaptureSkyLightWaitingForShaders = bCaptureSkyLightWaitingCompiledShader;
 					SkyLightSceneProxy->bCaptureSkyLightWaitingForMeshesOrTextures = bCaptureSkyLightWaitingForMeshOrTexAssets;
diff --git a/Engine/Source/Runtime/Engine/Public/SceneManagement.h b/Engine/Source/Runtime/Engine/Public/SceneManagement.h
--- a/Engine/Source/Runtime/Engine/Public/SceneManagement.h
+++ b/Engine/Source/Runtime/Engine/Public/SceneManagement.h
@@ -1161,6 +1161,7 @@ public:
 	FLinearColor GetEffectiveLightColor() const;
 
 #if WITH_EDITOR
+	float SecondsToNextIncompleteCapture;
 	bool bCubemapSkyLightWaitingForCubeMapTexture;
 	bool bCaptureSkyLightWaitingForShaders;
 	bool bCaptureSkyLightWaitingForMeshesOrTextures;
diff --git a/Engine/Source/Runtime/Renderer/Private/SceneRendering.cpp b/Engine/Source/Runtime/Renderer/Private/SceneRendering.cpp
--- a/Engine/Source/Runtime/Renderer/Private/SceneRendering.cpp
+++ b/Engine/Source/Runtime/Renderer/Private/SceneRendering.cpp
@@ -3224,7 +3224,7 @@ void FSceneRenderer::RenderFinish(FRDGBuilder& GraphBuilder, FRDGTextureRef View
 						{
 							const FLinearColor OrangeColor = FColor::Orange;
 
-							FString String = FString::Printf(TEXT("Sky Light waiting on ["));
+							FString String = TEXT("Sky Light waiting on [");
 							bool bAddComma = false;
 							if (SkyLight->bCubemapSkyLightWaitingForCubeMapTexture)
 							{
@@ -3242,7 +3242,11 @@ void FSceneRenderer::RenderFinish(FRDGBuilder& GraphBuilder, FRDGTextureRef View
 								String += bAddComma ? TEXT(", ") : TEXT("");
 								String += TEXT("Meshes, Textures] for final capture");
 							}
-							String += FString::Printf(TEXT("] for final capture."));
+							String += TEXT("] for final capture.");
+							Canvas.DrawShadowedText(10, Y, FText::FromString(String), GetStatsFont(), OrangeColor);
+							Y += 14;
+
+							String = FString::Printf(TEXT("Seconds to next incomplete capture = %.0f"), FMath::CeilToFloat(SkyLight->SecondsToNextIncompleteCapture));
 							Canvas.DrawShadowedText(10, Y, FText::FromString(String), GetStatsFont(), OrangeColor);
 							Y += 14;
 						}
