Fixed skylight state preventing continuous incomplete capture refininf when the type was changed while a capture was previously complete.
Now we simply mark the capture state as uninitialised when it is marked dirty.

diff --git a/Engine/Source/Runtime/Engine/Private/Components/SkyLightComponent.cpp b/Engine/Source/Runtime/Engine/Private/Components/SkyLightComponent.cpp
--- a/Engine/Source/Runtime/Engine/Private/Components/SkyLightComponent.cpp
+++ b/Engine/Source/Runtime/Engine/Private/Components/SkyLightComponent.cpp
@@ -309,6 +309,10 @@ void USkyLightComponent::SetCaptureIsDirty()
 	{
 		FScopeLock Lock(&SkyCapturesToUpdateLock);
 
+#if WITH_EDITOR
+		this->CaptureStatus = ESkyLightCaptureStatus::SLCS_Uninitialized;
+#endif
+
 		SkyCapturesToUpdate.AddUnique(this);
 
 		// Mark saved values as invalid, in case a sky recapture is requested in a construction script between a save / restore of sky capture state
