Set the FMeshBatch::bUseAsOccluder hint correctly for static mesh.
The value is read in FDepthPassMeshProcessor::AddMeshBatch which then re-evaluates the condition if the hint has been set to false.
This fixes issue that static mesh with bUseAsOccluder false renders in depth pass.

diff --git a/Engine/Source/Runtime/Engine/Private/StaticMeshRender.cpp b/Engine/Source/Runtime/Engine/Private/StaticMeshRender.cpp
--- a/Engine/Source/Runtime/Engine/Private/StaticMeshRender.cpp
+++ b/Engine/Source/Runtime/Engine/Private/StaticMeshRender.cpp
@@ -1095,11 +1095,13 @@ void FStaticMeshSceneProxy::DrawStaticElements(FStaticPrimitiveDrawInterface* PD
 		// Determine the DPG the primitive should be drawn in.
 		uint8 PrimitiveDPG = GetStaticDepthPriorityGroup();
 		int32 NumLODs = RenderData->LODResources.Num();
-		//Never use the dynamic path in this path, because only unselected elements will use DrawStaticElements
+		// Never use the dynamic path in this path, because only unselected elements will use DrawStaticElements
 		bool bIsMeshElementSelected = false;
 		const auto FeatureLevel = GetScene().GetFeatureLevel();
 		const bool IsMobile = IsMobilePlatform(GetScene().GetShaderPlatform());
 		const int32 NumRuntimeVirtualTextureTypes = RuntimeVirtualTextureMaterialTypes.Num();
+		// When this is false the final bUseAsOccluder state is re-evaluated in FDepthPassMeshProcessor::AddMeshBatch()
+		const bool bIsOccluderNonMovable = ShouldUseAsOccluder() && !IsMovable();
 
 		//check if a LOD is being forced
 		if (ForcedLodModel > 0) 
@@ -1193,9 +1195,8 @@ void FStaticMeshSceneProxy::DrawStaticElements(FStaticPrimitiveDrawInterface* PD
 					if (bSafeToUseUnifiedMesh)
 					{
 						bUseUnifiedMeshForShadow = bAllSectionsCastShadow;
-
-						// Depth pass is only used for deferred renderer. The other conditions are meant to match the logic in FDepthPassMeshProcessor::AddMeshBatch.
-						bUseUnifiedMeshForDepth = ShouldUseAsOccluder() && GetScene().GetShadingPath() == EShadingPath::Deferred && !IsMovable();
+						// Depth pass is only used for deferred renderer.
+						bUseUnifiedMeshForDepth = GetScene().GetShadingPath() == EShadingPath::Deferred;
 
 						if (bUseUnifiedMeshForShadow || bUseUnifiedMeshForDepth)
 						{
@@ -1213,7 +1214,7 @@ void FStaticMeshSceneProxy::DrawStaticElements(FStaticPrimitiveDrawInterface* PD
 
 									MeshBatch.CastShadow = bUseUnifiedMeshForShadow;
 									MeshBatch.bUseForDepthPass = bUseUnifiedMeshForDepth;
-									MeshBatch.bUseAsOccluder = bUseUnifiedMeshForDepth;
+									MeshBatch.bUseAsOccluder = bIsOccluderNonMovable;
 									MeshBatch.bUseForMaterial = false;
 
 									PDI->DrawMesh(MeshBatch, ScreenSize);
@@ -1262,8 +1263,8 @@ void FStaticMeshSceneProxy::DrawStaticElements(FStaticPrimitiveDrawInterface* PD
 								// If we have submitted an optimized shadow-only mesh, remaining mesh elements must not cast shadows.
 								FMeshBatch MeshBatch(BaseMeshBatch);
 								MeshBatch.CastShadow &= !bUseUnifiedMeshForShadow;
-								MeshBatch.bUseAsOccluder &= !bUseUnifiedMeshForDepth;
 								MeshBatch.bUseForDepthPass &= !bUseUnifiedMeshForDepth;
+								MeshBatch.bUseAsOccluder &= bIsOccluderNonMovable;
 								PDI->DrawMesh(MeshBatch, ScreenSize);
 							}
 						}
diff --git a/Engine/Source/Runtime/Renderer/Private/DepthRendering.cpp b/Engine/Source/Runtime/Renderer/Private/DepthRendering.cpp
--- a/Engine/Source/Runtime/Renderer/Private/DepthRendering.cpp
+++ b/Engine/Source/Runtime/Renderer/Private/DepthRendering.cpp
@@ -841,28 +841,25 @@ void FDepthPassMeshProcessor::AddMeshBatch(const FMeshBatch& RESTRICT MeshBatch,
 {
 	bool bDraw = MeshBatch.bUseForDepthPass;
 
-	// Filter by occluder flags and settings if required.
-	if (bDraw && bRespectUseAsOccluderFlag && !MeshBatch.bUseAsOccluder && EarlyZPassMode < DDM_AllOpaque)
+	// Filter draw calls if this is a partial depth prepass.
+	const bool bFilterDrawCalls = bRespectUseAsOccluderFlag && EarlyZPassMode < DDM_AllOpaque;
+
+	// If mesh batch hints that we are not an occluder then potentially remove.
+	if (bDraw && bFilterDrawCalls && !MeshBatch.bUseAsOccluder)
 	{
-		if (PrimitiveSceneProxy)
-		{
+		bDraw = PrimitiveSceneProxy != nullptr
 			// Only render primitives marked as occluders.
-			bDraw = PrimitiveSceneProxy->ShouldUseAsOccluder()
-				// Only render static objects unless movable are requested.
-				&& (!PrimitiveSceneProxy->IsMovable() || bEarlyZPassMovable);
+			&& PrimitiveSceneProxy->ShouldUseAsOccluder()
+			// Only render static objects unless movable are requested.
+			&& (!PrimitiveSceneProxy->IsMovable() || bEarlyZPassMovable);
+	}
 
-			// Filter dynamic mesh commands by screen size.
-			if (ViewIfDynamicMeshCommand)
-			{
-				extern float GMinScreenRadiusForDepthPrepass;
-				const float LODFactorDistanceSquared = (PrimitiveSceneProxy->GetBounds().Origin - ViewIfDynamicMeshCommand->ViewMatrices.GetViewOrigin()).SizeSquared() * FMath::Square(ViewIfDynamicMeshCommand->LODDistanceFactor);
-				bDraw = bDraw && FMath::Square(PrimitiveSceneProxy->GetBounds().SphereRadius) > GMinScreenRadiusForDepthPrepass * GMinScreenRadiusForDepthPrepass * LODFactorDistanceSquared;
-			}
-		}
-		else
-		{
-			bDraw = false;
-		}
+	// Filter dynamic mesh commands by size.
+	if (bDraw && bFilterDrawCalls && ViewIfDynamicMeshCommand != nullptr && PrimitiveSceneProxy != nullptr)
+	{
+		extern float GMinScreenRadiusForDepthPrepass;
+		const float LODFactorDistanceSquared = (PrimitiveSceneProxy->GetBounds().Origin - ViewIfDynamicMeshCommand->ViewMatrices.GetViewOrigin()).SizeSquared() * FMath::Square(ViewIfDynamicMeshCommand->LODDistanceFactor);
+		bDraw = FMath::Square(PrimitiveSceneProxy->GetBounds().SphereRadius) > GMinScreenRadiusForDepthPrepass * GMinScreenRadiusForDepthPrepass * LODFactorDistanceSquared;
 	}
 
 	// if we are skipping movable objects in early Z, which can happen in DDM_AllOpaqueNoVelocity
