Remove the ensure for world initialization in RegisterComponent and comment on why

diff --git a/Engine/Source/Runtime/Engine/Private/Components/ActorComponent.cpp b/Engine/Source/Runtime/Engine/Private/Components/ActorComponent.cpp
--- a/Engine/Source/Runtime/Engine/Private/Components/ActorComponent.cpp
+++ b/Engine/Source/Runtime/Engine/Private/Components/ActorComponent.cpp
@@ -1315,7 +1315,9 @@ void UActorComponent::RegisterComponent()
 	UWorld* MyOwnerWorld = (MyOwner ? MyOwner->GetWorld() : nullptr);
 	if (ensure(MyOwnerWorld))
 	{
-		ensure(MyOwnerWorld->bIsWorldInitialized);
+		//@note FH: world should be initialized when calling RegisterComponent or it should be handled gracefully as a no-op but that isn't currently the case
+		// however a lot of legacy code may end up calling RegisterComponent prior to world initialization hence why the ensure is currently commented
+		//ensure(MyOwnerWorld->bIsWorldInitialized);
 		RegisterComponentWithWorld(MyOwnerWorld);
 	}
 }
