Instance data optimizations - removed NonUniformScale and InvNonUniformScale from CPU, since they are already derived on the GPU from LocalToWorld as of CL 15961534

diff --git a/Engine/Plugins/Experimental/GPULightmass/Source/GPULightmass/Private/LightmapRenderer.cpp b/Engine/Plugins/Experimental/GPULightmass/Source/GPULightmass/Private/LightmapRenderer.cpp
--- a/Engine/Plugins/Experimental/GPULightmass/Source/GPULightmass/Private/LightmapRenderer.cpp
+++ b/Engine/Plugins/Experimental/GPULightmass/Source/GPULightmass/Private/LightmapRenderer.cpp
@@ -330,8 +330,6 @@ void FCachedRayTracingSceneData::SetupViewUniformBufferFromSceneRenderState(FSce
 						PrimitiveUniformShaderParameters.LocalToWorld,
 						PrimitiveUniformShaderParameters.PreviousLocalToWorld,
 						FRenderBounds(PrimitiveUniformShaderParameters.LocalObjectBoundsMin, PrimitiveUniformShaderParameters.LocalObjectBoundsMax),
-						PrimitiveUniformShaderParameters.NonUniformScale,
-						PrimitiveUniformShaderParameters.InvNonUniformScale,
 						FVector4(ForceInitToZero),
 						NANITE_INVALID_HIERARCHY_OFFSET,
 						InstanceSceneDataFlags,
@@ -476,8 +474,6 @@ void FCachedRayTracingSceneData::SetupViewUniformBufferFromSceneRenderState(FSce
 					PrimitiveUniformShaderParameters.LocalToWorld,
 					PrimitiveUniformShaderParameters.PreviousLocalToWorld,
 					FRenderBounds(PrimitiveUniformShaderParameters.LocalObjectBoundsMin, PrimitiveUniformShaderParameters.LocalObjectBoundsMax),
-					PrimitiveUniformShaderParameters.NonUniformScale,
-					PrimitiveUniformShaderParameters.InvNonUniformScale,
 					FVector4(ForceInitToZero),
 					NANITE_INVALID_HIERARCHY_OFFSET,
 					InstanceSceneDataFlags,
diff --git a/Engine/Source/Runtime/Engine/Private/InstanceUniformShaderParameters.cpp b/Engine/Source/Runtime/Engine/Private/InstanceUniformShaderParameters.cpp
--- a/Engine/Source/Runtime/Engine/Private/InstanceUniformShaderParameters.cpp
+++ b/Engine/Source/Runtime/Engine/Private/InstanceUniformShaderParameters.cpp
@@ -46,8 +46,6 @@ static FPrimitiveInstance DummyInstance =
 		FRenderTransform::Identity,
 		FRenderTransform::Identity,
 		FRenderBounds(FVector3f::ZeroVector, FVector3f::ZeroVector),
-		FVector4(1.0f, 1.0f, 1.0f, 1.0f),
-		FVector3f(1.0f, 1.0f, 1.0f),
 		FVector4(ForceInitToZero),
 		NANITE_INVALID_HIERARCHY_OFFSET,
 		0u, /* Instance Flags */
diff --git a/Engine/Source/Runtime/Engine/Public/InstanceUniformShaderParameters.h b/Engine/Source/Runtime/Engine/Public/InstanceUniformShaderParameters.h
--- a/Engine/Source/Runtime/Engine/Public/InstanceUniformShaderParameters.h
+++ b/Engine/Source/Runtime/Engine/Public/InstanceUniformShaderParameters.h
@@ -21,8 +21,6 @@ struct FPrimitiveInstance
 	FRenderTransform		PrevInstanceToLocal;
 	FRenderTransform		LocalToWorld;
 	FRenderTransform		PrevLocalToWorld;
-	FVector4				NonUniformScale;
-	FVector3f				InvNonUniformScale;
 	uint32					LastUpdateSceneFrameNumber;
 	FRenderBounds			LocalBounds;
 	float					PerInstanceRandom;
@@ -30,23 +28,12 @@ struct FPrimitiveInstance
 	uint32					NaniteHierarchyOffset;
 	uint32					Flags;
 
-	FORCEINLINE void OrthonormalizeAndUpdateScale()
+	FORCEINLINE void Orthonormalize()
 	{
 		// Remove shear
-		const FVector3f Scale = LocalToWorld.Orthonormalize();
+		LocalToWorld.Orthonormalize();
 		PrevLocalToWorld.Orthonormalize();
 
-		NonUniformScale = FVector4(
-			Scale.X, Scale.Y, Scale.Z,
-			FMath::Max3(FMath::Abs(Scale.X), FMath::Abs(Scale.Y), FMath::Abs(Scale.Z))
-		);
-
-		InvNonUniformScale = FVector3f(
-			Scale.X > KINDA_SMALL_NUMBER ? 1.0f / Scale.X : 0.0f,
-			Scale.Y > KINDA_SMALL_NUMBER ? 1.0f / Scale.Y : 0.0f,
-			Scale.Z > KINDA_SMALL_NUMBER ? 1.0f / Scale.Z : 0.0f
-		);
-
 		if (LocalToWorld.RotDeterminant() < 0.0f)
 		{
 			Flags |= INSTANCE_SCENE_DATA_FLAG_DETERMINANT_SIGN;
@@ -56,14 +43,31 @@ struct FPrimitiveInstance
 			Flags &= ~INSTANCE_SCENE_DATA_FLAG_DETERMINANT_SIGN;
 		}
 	}
+
+	FORCEINLINE FVector4 GetNonUniformScale() const
+	{
+		const FVector3f Scale = LocalToWorld.GetScale();
+		return FVector4(
+			Scale.X, Scale.Y, Scale.Z,
+			FMath::Max3(FMath::Abs(Scale.X), FMath::Abs(Scale.Y), FMath::Abs(Scale.Z))
+		);
+	}
+
+	FORCEINLINE FVector4 GetInvNonUniformScale() const
+	{
+		const FVector3f Scale = LocalToWorld.GetScale();
+		return FVector3f(
+			Scale.X > KINDA_SMALL_NUMBER ? 1.0f / Scale.X : 0.0f,
+			Scale.Y > KINDA_SMALL_NUMBER ? 1.0f / Scale.Y : 0.0f,
+			Scale.Z > KINDA_SMALL_NUMBER ? 1.0f / Scale.Z : 0.0f
+		);
+	}
 };
 
 FORCEINLINE FPrimitiveInstance ConstructPrimitiveInstance(
 	const FRenderTransform& LocalToWorld,
 	const FRenderTransform& PrevLocalToWorld,
 	const FRenderBounds& LocalBounds,
-	const FVector4& NonUniformScale,
-	const FVector3f& InvNonUniformScale,
 	const FVector4& LightMapAndShadowMapUVBias,
 	const uint32 NaniteHierarchyOffset,
 	uint32 Flags,
@@ -74,8 +78,6 @@ FORCEINLINE FPrimitiveInstance ConstructPrimitiveInstance(
 	FPrimitiveInstance Result;
 	Result.LocalToWorld							= LocalToWorld;
 	Result.PrevLocalToWorld						= PrevLocalToWorld;
-	Result.NonUniformScale						= NonUniformScale;
-	Result.InvNonUniformScale					= InvNonUniformScale;
 	Result.LightMapAndShadowMapUVBias			= LightMapAndShadowMapUVBias;
 	Result.LocalBounds							= LocalBounds;
 	Result.NaniteHierarchyOffset				= NaniteHierarchyOffset;
@@ -102,8 +104,6 @@ struct FInstanceSceneShaderData
 				FRenderTransform::Identity,
 				FRenderTransform::Identity,
 				FRenderBounds(FVector3f::ZeroVector, FVector3f::ZeroVector),
-				FVector4(1.0f, 1.0f, 1.0f, 1.0f),
-				FVector3f(1.0f, 1.0f, 1.0f),
 				FVector4(ForceInitToZero),
 				0xFFFFFFFFu, /* Nanite Hierarchy Offset */
 				0u, /* Instance Flags */
diff --git a/Engine/Source/Runtime/RenderCore/Public/RenderTransform.h b/Engine/Source/Runtime/RenderCore/Public/RenderTransform.h
--- a/Engine/Source/Runtime/RenderCore/Public/RenderTransform.h
+++ b/Engine/Source/Runtime/RenderCore/Public/RenderTransform.h
@@ -135,7 +135,17 @@ public:
 		return ToMatrix44f().InverseFast();
 	}
 
-	FORCEINLINE FVector3f Orthonormalize()
+	FORCEINLINE FVector3f GetScale() const
+	{
+		// Extract per axis scales
+		FVector3f Scale;
+		Scale.X = TransformRows[0].Size();
+		Scale.Y = TransformRows[1].Size();
+		Scale.Z = TransformRows[2].Size();
+		return Scale;
+	}
+
+	FORCEINLINE void Orthonormalize()
 	{
 		FVector3f X = TransformRows[0];
 		FVector3f Y = TransformRows[1];
@@ -149,14 +159,6 @@ public:
 		TransformRows[0] = X;
 		TransformRows[1] = Y;
 		TransformRows[2] = Z;
-
-		// Extract per axis scales
-		FVector3f Scale;
-		Scale.X = X.Size();
-		Scale.Y = Y.Size();
-		Scale.Z = Z.Size();
-
-		return Scale;
 	}
 
 	FORCEINLINE void SetIdentity()
diff --git a/Engine/Source/Runtime/Renderer/Private/GPUScene.cpp b/Engine/Source/Runtime/Renderer/Private/GPUScene.cpp
--- a/Engine/Source/Runtime/Renderer/Private/GPUScene.cpp
+++ b/Engine/Source/Runtime/Renderer/Private/GPUScene.cpp
@@ -206,7 +206,7 @@ inline void InitPrimitiveInstance(FPrimitiveInstance& PrimitiveInstance, const F
 		PrimitiveInstance.Flags &= ~INSTANCE_SCENE_DATA_FLAG_DETERMINANT_SIGN;
 	}
 
-	PrimitiveInstance.OrthonormalizeAndUpdateScale();
+	PrimitiveInstance.Orthonormalize();
 }
 
 inline void InitPrimitiveInstanceDummy(FPrimitiveInstance& DummyInstance, const FPrimitiveTransforms& PrimitiveTransforms, const FBoxSphereBounds& LocalBounds, uint32 SceneFrameNumber)
@@ -218,8 +218,6 @@ inline void InitPrimitiveInstanceDummy(FPrimitiveInstance& DummyInstance, const
 	DummyInstance.PrevInstanceToLocal.SetIdentity();
 	DummyInstance.LocalToWorld.SetIdentity();
 	DummyInstance.PrevLocalToWorld.SetIdentity();
-	DummyInstance.NonUniformScale = FVector4(1.0f, 1.0f, 1.0f, 1.0f);
-	DummyInstance.InvNonUniformScale = FVector3f(1.0f, 1.0f, 1.0f);
 	DummyInstance.LocalBounds = LocalBounds;
 
 	const bool bHasPreviousInstanceTransforms = false;
diff --git a/Engine/Source/Runtime/Renderer/Private/Renderer.cpp b/Engine/Source/Runtime/Renderer/Private/Renderer.cpp
--- a/Engine/Source/Runtime/Renderer/Private/Renderer.cpp
+++ b/Engine/Source/Runtime/Renderer/Private/Renderer.cpp
@@ -179,8 +179,6 @@ void FRendererModule::DrawTileMesh(FCanvasRenderContext& RenderContext, FMeshPas
 						PrimitiveParams.LocalToWorld,
 						PrimitiveParams.PreviousLocalToWorld,
 						FRenderBounds(PrimitiveParams.LocalObjectBoundsMin, PrimitiveParams.LocalObjectBoundsMax),
-						PrimitiveParams.NonUniformScale,
-						PrimitiveParams.InvNonUniformScale,
 						FVector4(ForceInitToZero),
 						NANITE_INVALID_HIERARCHY_OFFSET,
 						InstanceFlags,
