Edigrate of Quatrz fixes.

diff --git a/Engine/Source/Runtime/AudioMixer/Private/Quartz/AudioMixerClockManager.cpp b/Engine/Source/Runtime/AudioMixer/Private/Quartz/AudioMixerClockManager.cpp
--- a/Engine/Source/Runtime/AudioMixer/Private/Quartz/AudioMixerClockManager.cpp
+++ b/Engine/Source/Runtime/AudioMixer/Private/Quartz/AudioMixerClockManager.cpp
@@ -34,7 +34,7 @@ namespace Audio
 
 		for (auto& Clock : ActiveClocks)
 		{
-			Clock->Tick(DeltaTimeSeconds);
+			Clock->LowResolutionTick(DeltaTimeSeconds);
 		}
 	}
 
@@ -244,10 +244,7 @@ namespace Audio
 		}
 
 		FScopeLock Lock(&ActiveClockCritSec);
-		for (auto& Clock : ActiveClocks)
-		{
-			Clock->Shutdown();
-		}
+		ActiveClocks.Reset();
 	}
 
 	FQuartzQuantizedCommandHandle FQuartzClockManager::AddCommandToClock(FQuartzQuantizedCommandInitInfo& InQuantizationCommandInitInfo)
diff --git a/Engine/Source/Runtime/Engine/Classes/Sound/QuartzQuantizationUtilities.cpp b/Engine/Source/Runtime/Engine/Classes/Sound/QuartzQuantizationUtilities.cpp
--- a/Engine/Source/Runtime/Engine/Classes/Sound/QuartzQuantizationUtilities.cpp
+++ b/Engine/Source/Runtime/Engine/Classes/Sound/QuartzQuantizationUtilities.cpp
@@ -142,13 +142,19 @@ namespace Audio
 
 	void FQuartzClockTickRate::SetFramesPerTick(int32 InNewFramesPerTick)
 	{
+		if (InNewFramesPerTick < 1)
+		{
+			UE_LOG(LogAudioQuartz, Warning, TEXT("Quartz Metronme requires at least 1 frame per tick, clamping request"));
+			InNewFramesPerTick = 1;
+		}
+
 		FramesPerTick = InNewFramesPerTick;
 		RecalculateDurationsBasedOnFramesPerTick();
 	}
 
 	void FQuartzClockTickRate::SetMillisecondsPerTick(float InNewMillisecondsPerTick)
 	{
-		FramesPerTick = (InNewMillisecondsPerTick / 1000.f) * SampleRate;
+		FramesPerTick = FMath::Max(1.0f, (InNewMillisecondsPerTick * SampleRate) / 1000.f);
 		RecalculateDurationsBasedOnFramesPerTick();
 	}
 
@@ -156,7 +162,7 @@ namespace Audio
 	{
 		check(InNewThirtySecondNotesPerMinute > 0);
 
-		FramesPerTick = (60.f / InNewThirtySecondNotesPerMinute) * SampleRate;
+		FramesPerTick = FMath::Max(1.0f, (60.f * SampleRate ) / InNewThirtySecondNotesPerMinute);
 		RecalculateDurationsBasedOnFramesPerTick();
 	}
 
@@ -169,7 +175,7 @@ namespace Audio
 		// FramesPerTick = 1/8 * (60.f / (InNewBeatsPerMinute)) * SampleRate;
 		// (60.0 / 8.0) = 7.5f
 
-		FramesPerTick = (7.5f / InNewBeatsPerMinute) * SampleRate;
+		FramesPerTick = FMath::Max(1.0f, (7.5f * SampleRate) / InNewBeatsPerMinute);
 		RecalculateDurationsBasedOnFramesPerTick();
 	}
 
@@ -177,7 +183,7 @@ namespace Audio
 	{
 		check(InNewSampleRate >= 0);
 
-		FramesPerTick = (InNewSampleRate / SampleRate) * static_cast<float>(FramesPerTick);
+		FramesPerTick = FMath::Max(1.0f, (InNewSampleRate / SampleRate) * static_cast<float>(FramesPerTick));
 		SampleRate = InNewSampleRate;
 
 		RecalculateDurationsBasedOnFramesPerTick();
