Add scoped draw event for VT upload

diff --git a/Engine/Source/Runtime/Engine/Private/VT/VirtualTextureUploadCache.cpp b/Engine/Source/Runtime/Engine/Private/VT/VirtualTextureUploadCache.cpp
--- a/Engine/Source/Runtime/Engine/Private/VT/VirtualTextureUploadCache.cpp
+++ b/Engine/Source/Runtime/Engine/Private/VT/VirtualTextureUploadCache.cpp
@@ -91,6 +91,7 @@ void FVirtualTextureUploadCache::Finalize(FRDGBuilder& GraphBuilder)
 	check(IsInRenderingThread());
 
 	FRHICommandListImmediate& RHICmdList = GraphBuilder.RHICmdList;
+	SCOPED_DRAW_EVENT(RHICmdList, FVirtualTextureUploadCache_Finalize);
 
 	// Multi-GPU support : May be ineffecient for AFR.
 	SCOPED_GPU_MASK(RHICmdList, FRHIGPUMask::All());
