Add missing Linear plasticity function to physics constraints blueprint library
Add ability to get constraints from a specific body to SkeletalMeshComponent

diff --git a/Engine/Source/Runtime/Engine/Classes/Components/SkeletalMeshComponent.h b/Engine/Source/Runtime/Engine/Classes/Components/SkeletalMeshComponent.h
--- a/Engine/Source/Runtime/Engine/Classes/Components/SkeletalMeshComponent.h
+++ b/Engine/Source/Runtime/Engine/Classes/Components/SkeletalMeshComponent.h
@@ -2099,6 +2099,17 @@ public:
 	UFUNCTION(BlueprintCallable, Category = "Physics", meta = (Keywords = "Components|SkeletalMesh"))
 	void GetConstraints(bool bIncludesTerminated, TArray<FConstraintInstanceAccessor>& OutConstraints);
 
+	/** Gets all the constraints attached to a body 
+	* @param BodyName name of the body to get the attached constraints from 
+	* @param bParentConstraints return constraints where BodyName is the child of the constraint
+	* @param bChildConstraints return constraints where BodyName is the parent of the constraint
+	* @param bParentConstraints return constraints attached to the parent of the body 
+	* @param IncludesTerminated whether or not to return terminated constraints
+	* @param OutConstraints returned list of constraints matching the parameters
+	* */
+	UFUNCTION(BlueprintCallable, Category = "Physics", meta = (Keywords = "Components|SkeletalMesh"))
+	void GetConstraintsFromBody(FName BodyName, bool bParentConstraints, bool bChildConstraints, bool bIncludesTerminated, TArray<FConstraintInstanceAccessor>& OutConstraints);
+
 	/** Sets the Angular Motion Ranges for a named constraint
 	*  @param InBoneName  Name of bone to adjust constraint ranges for
 	*  @param Swing1LimitAngle	 Size of limit in degrees, 0 means locked, 180 means free
diff --git a/Engine/Source/Runtime/Engine/Classes/PhysicsEngine/ConstraintInstance.h b/Engine/Source/Runtime/Engine/Classes/PhysicsEngine/ConstraintInstance.h
--- a/Engine/Source/Runtime/Engine/Classes/PhysicsEngine/ConstraintInstance.h
+++ b/Engine/Source/Runtime/Engine/Classes/PhysicsEngine/ConstraintInstance.h
@@ -540,6 +540,18 @@ public:
 		UpdatePlasticity();
 	}
 
+	/** Gets Whether it is possible to reset the target position */
+	bool HasLinearPlasticity() const
+	{
+		return ProfileInstance.bLinearPlasticity;
+	}
+
+	/** Gets Delta from target needed to reset the target joint */
+	float GetLinearPlasticityThreshold() const
+	{
+		return ProfileInstance.LinearPlasticityThreshold;
+	}
+
 	/** Sets the Angular Breakable properties
 	*	@param bInAngularBreakable		Whether it is possible to break the joint with angular force
 	*	@param InAngularBreakThreshold	Torque needed to break the joint
diff --git a/Engine/Source/Runtime/Engine/Classes/PhysicsEngine/ConstraintInstanceBlueprintLibrary.h b/Engine/Source/Runtime/Engine/Classes/PhysicsEngine/ConstraintInstanceBlueprintLibrary.h
--- a/Engine/Source/Runtime/Engine/Classes/PhysicsEngine/ConstraintInstanceBlueprintLibrary.h
+++ b/Engine/Source/Runtime/Engine/Classes/PhysicsEngine/ConstraintInstanceBlueprintLibrary.h
@@ -178,6 +178,31 @@ public:
 		float& LinearBreakThreshold
 	);
 
+
+	/** Sets Constraint Linear Plasticity properties
+	*	@param Accessor						Constraint accessor to change
+	*	@param bLinearPlasticity			Whether it is possible to reset the target position from the linear displacement
+	*	@param LinearPlasticityThreshold	Delta from target needed to reset the target joint
+	*/
+	UFUNCTION(BlueprintCallable, Category = "Physics|Constraints")
+		static void SetLinearPlasticity(
+			UPARAM(ref) FConstraintInstanceAccessor& Accessor,
+			bool bLinearPlasticity,
+			float LinearPlasticityThreshold
+		);
+
+	/** Sets Constraint Linear Plasticity properties
+	*	@param Accessor						Constraint accessor to query
+	*	@param bAngularPlasticity			Whether it is possible to reset the target position from the linear displacement
+	*	@param AngularPlasticityThreshold	Delta from target needed to reset the target joint
+	*/
+	UFUNCTION(BlueprintCallable, Category = "Physics|Constraints")
+		static void GetLinearPlasticity(
+			UPARAM(ref) FConstraintInstanceAccessor& Accessor,
+			bool& bLinearPlasticity,
+			float& LinearPlasticityThreshold
+		);
+
 	//---------------------------------------------------------------------------------------------------
 	//
 	// ANGULAR LIMITS 
diff --git a/Engine/Source/Runtime/Engine/Private/PhysicsEngine/ConstraintInstanceBlueprintLibrary.cpp b/Engine/Source/Runtime/Engine/Private/PhysicsEngine/ConstraintInstanceBlueprintLibrary.cpp
--- a/Engine/Source/Runtime/Engine/Private/PhysicsEngine/ConstraintInstanceBlueprintLibrary.cpp
+++ b/Engine/Source/Runtime/Engine/Private/PhysicsEngine/ConstraintInstanceBlueprintLibrary.cpp
@@ -205,6 +205,36 @@ void UConstraintInstanceBlueprintLibrary::GetLinearBreakable(
 	}
 }
 
+void UConstraintInstanceBlueprintLibrary::SetLinearPlasticity(
+	FConstraintInstanceAccessor& Accessor,
+	bool bLinearPlasticity,
+	float LinearPlasticityThreshold
+	)
+{
+	if (FConstraintInstance* ConstraintInstance = Accessor.Get())
+	{
+		ConstraintInstance->SetLinearPlasticity(bLinearPlasticity, LinearPlasticityThreshold);
+	}
+}
+
+void UConstraintInstanceBlueprintLibrary::GetLinearPlasticity(
+	FConstraintInstanceAccessor& Accessor,
+	bool& bLinearPlasticity,
+	float& LinearPlasticityThreshold
+	)
+{
+	if (FConstraintInstance* ConstraintInstance = Accessor.Get())
+	{
+		bLinearPlasticity = ConstraintInstance->HasLinearPlasticity();
+		LinearPlasticityThreshold = ConstraintInstance->GetLinearPlasticityThreshold();
+	}
+	else
+	{
+		bLinearPlasticity = false;
+		LinearPlasticityThreshold = 0.0f;
+	}
+}
+
 //---------------------------------------------------------------------------------------------------
 //
 // ANGULAR LIMITS 
diff --git a/Engine/Source/Runtime/Engine/Private/SkeletalMeshComponentPhysics.cpp b/Engine/Source/Runtime/Engine/Private/SkeletalMeshComponentPhysics.cpp
--- a/Engine/Source/Runtime/Engine/Private/SkeletalMeshComponentPhysics.cpp
+++ b/Engine/Source/Runtime/Engine/Private/SkeletalMeshComponentPhysics.cpp
@@ -1906,6 +1906,30 @@ void USkeletalMeshComponent::GetConstraints(bool bIncludesTerminated, TArray<FCo
 	}
 }
 
+void  USkeletalMeshComponent::GetConstraintsFromBody(FName BodyName, bool bParentConstraints, bool bChildConstraints, bool bIncludesTerminated, TArray<FConstraintInstanceAccessor>& OutConstraints)
+{
+	if (UPhysicsAsset* const PhysicsAsset = GetPhysicsAsset())
+	{
+		for (int32 i = 0; i < Constraints.Num(); i++)
+		{
+			if (FConstraintInstance* ConstraintInstance = Constraints[i])
+			{
+				if (bIncludesTerminated || !ConstraintInstance->IsTerminated())
+				{
+					if (bParentConstraints && ConstraintInstance->GetChildBoneName() == BodyName)
+					{
+						OutConstraints.Add(FConstraintInstanceAccessor(this, i));
+					}
+					if (bChildConstraints && ConstraintInstance->GetParentBoneName() == BodyName)
+					{
+						OutConstraints.Add(FConstraintInstanceAccessor(this, i));
+					}
+				}
+			}
+		}
+	}
+}
+
 void USkeletalMeshComponent::BreakConstraint(FVector Impulse, FVector HitLocation, FName InBoneName)
 {
 	// you can enable/disable the instanced weights by calling
