GitHub #6794 : Expose GetPlayerViewPoint to Blueprints.
https://github.com/EpicGames/UnrealEngine/pull/6794

diff --git a/Engine/Source/Runtime/Engine/Classes/GameFramework/Controller.h b/Engine/Source/Runtime/Engine/Classes/GameFramework/Controller.h
--- a/Engine/Source/Runtime/Engine/Classes/GameFramework/Controller.h
+++ b/Engine/Source/Runtime/Engine/Classes/GameFramework/Controller.h
@@ -324,6 +324,7 @@ public:
 	 * @output	out_Location, view location of player
 	 * @output	out_rotation, view rotation of player
 	 */
+	UFUNCTION(BlueprintCallable, Category = Pawn)
 	virtual void GetPlayerViewPoint( FVector& Location, FRotator& Rotation ) const;
 
 	/** GameMode failed to spawn pawn for me. */
