Fix compilation error by deprecating BeginGraphDragAction without mouse event input in EdGraphSchema

diff --git a/Engine/Source/Editor/Kismet/Private/SMyBlueprint.cpp b/Engine/Source/Editor/Kismet/Private/SMyBlueprint.cpp
--- a/Engine/Source/Editor/Kismet/Private/SMyBlueprint.cpp
+++ b/Engine/Source/Editor/Kismet/Private/SMyBlueprint.cpp
@@ -1739,7 +1739,7 @@ FReply SMyBlueprint::OnActionDragged( const TArray< TSharedPtr<FEdGraphSchemaAct
 			{
 				if (FuncAction->EdGraph->GetSchema()->CanGraphBeDropped(InAction))
 				{
-					return FuncAction->EdGraph->GetSchema()->BeginGraphDragAction(InAction);
+					return FuncAction->EdGraph->GetSchema()->BeginGraphDragAction(InAction, FPointerEvent());
 				}
 			}
 
diff --git a/Engine/Source/Runtime/Engine/Classes/EdGraph/EdGraphSchema.h b/Engine/Source/Runtime/Engine/Classes/EdGraph/EdGraphSchema.h
--- a/Engine/Source/Runtime/Engine/Classes/EdGraph/EdGraphSchema.h
+++ b/Engine/Source/Runtime/Engine/Classes/EdGraph/EdGraphSchema.h
@@ -975,7 +975,13 @@ class ENGINE_API UEdGraphSchema : public UObject
 	/*
 	 * Begins a drag and drop action to drag a graph action into another graph
 	 */
-	virtual FReply BeginGraphDragAction(TSharedPtr<FEdGraphSchemaAction> InAction, const FPointerEvent& MouseEvent = FPointerEvent() ) const { return FReply::Unhandled(); }
+	UE_DEPRECATED(5.0, "Use version that takes FPointerEvent instead.")
+	virtual FReply BeginGraphDragAction(TSharedPtr<FEdGraphSchemaAction> InAction) const { return FReply::Unhandled(); }
+
+	/*
+	* Begins a drag and drop action to drag a graph action into another graph
+	*/
+	virtual FReply BeginGraphDragAction(TSharedPtr<FEdGraphSchemaAction> InAction, const FPointerEvent& MouseEvent) const { return FReply::Unhandled(); }
 
 	/**
 	 * Gets display information for a graph
