LevelInstance: Change default verbosity + change RerunConstructionScript verbosity

diff --git a/Engine/Source/Runtime/Engine/Private/LevelInstance/LevelInstancePrivate.h b/Engine/Source/Runtime/Engine/Private/LevelInstance/LevelInstancePrivate.h
--- a/Engine/Source/Runtime/Engine/Private/LevelInstance/LevelInstancePrivate.h
+++ b/Engine/Source/Runtime/Engine/Private/LevelInstance/LevelInstancePrivate.h
@@ -4,5 +4,5 @@
 
 #include "CoreMinimal.h"
 
-DECLARE_LOG_CATEGORY_EXTERN(LogLevelInstance, Verbose, All);
+DECLARE_LOG_CATEGORY_EXTERN(LogLevelInstance, Log, All);
 
diff --git a/Engine/Source/Runtime/Engine/Private/LevelInstance/Packed/PackedLevelInstanceActor.cpp b/Engine/Source/Runtime/Engine/Private/LevelInstance/Packed/PackedLevelInstanceActor.cpp
--- a/Engine/Source/Runtime/Engine/Private/LevelInstance/Packed/PackedLevelInstanceActor.cpp
+++ b/Engine/Source/Runtime/Engine/Private/LevelInstance/Packed/PackedLevelInstanceActor.cpp
@@ -71,7 +71,7 @@ void APackedLevelInstance::RerunConstructionScripts()
 		if (PackedVersion != GetClass()->GetDefaultObject<APackedLevelInstance>()->PackedVersion)
 		{
 			bShouldRerunConstructionScript = true;
-			UE_LOG(LogLevelInstance, Warning, TEXT("RerunConstructionScript was executed on %s (%s) because its version (%s) doesn't match latest version (%s). Resaving this actor will fix this"),
+			UE_LOG(LogLevelInstance, Verbose, TEXT("RerunConstructionScript was executed on %s (%s) because its version (%s) doesn't match latest version (%s). Resaving this actor will fix this"),
 				*GetPathName(),
 				*GetPackage()->GetPathName(),
 				*PackedVersion.ToString(),
