Change whether to use GPUSkinCache when ray tracing is on to be a per mesh component setting, so individual mesh can opt in/out based on its own SupportRayTracing setting

diff --git a/Engine/Source/Runtime/Engine/Classes/Engine/RendererSettings.h b/Engine/Source/Runtime/Engine/Classes/Engine/RendererSettings.h
--- a/Engine/Source/Runtime/Engine/Classes/Engine/RendererSettings.h
+++ b/Engine/Source/Runtime/Engine/Classes/Engine/RendererSettings.h
@@ -785,7 +785,7 @@ class ENGINE_API URendererSettings : public UDeveloperSettings
 
 	UPROPERTY(config, EditAnywhere, Category = Optimizations, meta = (
 		ConsoleVariable = "r.SkinCache.DefaultBehavior", DisplayName = "Default Skin Cache Behavior",
-		ToolTip = "Default behavior if all skeletal meshes are included/excluded from the skin cache. If Ray Tracing is enabled, will force inclusive behavior."))
+		ToolTip = "Default behavior if all skeletal meshes are included/excluded from the skin cache. If Support Ray Tracing is enabled on a mesh, will force inclusive behavior on that mesh."))
 	ESkinCacheDefaultBehavior DefaultSkinCacheBehavior;
 
 	UPROPERTY(config, EditAnywhere, Category = Optimizations, meta = (
diff --git a/Engine/Source/Runtime/Engine/Classes/Engine/SkeletalMesh.h b/Engine/Source/Runtime/Engine/Classes/Engine/SkeletalMesh.h
--- a/Engine/Source/Runtime/Engine/Classes/Engine/SkeletalMesh.h
+++ b/Engine/Source/Runtime/Engine/Classes/Engine/SkeletalMesh.h
@@ -223,10 +223,10 @@ namespace nvidia
 UENUM()
 enum class ESkinCacheUsage : uint8
 {
-	// Auto will defer to child or global behavior based on context. If Ray Tracing is enabled, will imply Enabled
+	// Auto will defer to child or global behavior based on context. If Support Ray Tracing is enabled on the mesh, will imply Enabled
 	Auto		= 0,
 
-	// Mesh will not use the skin cache. If Ray Tracing is enabled, will imply Enabled
+	// Mesh will not use the skin cache. If Support Ray Tracing is enabled on the mesh, will imply Enabled
 	Disabled	= uint8(-1),
 
 	// Mesh will use the skin cache
@@ -236,7 +236,7 @@ enum class ESkinCacheUsage : uint8
 UENUM()
 enum class ESkinCacheDefaultBehavior : uint8
 {
-	// All skeletal meshes are excluded from the skin cache. Each must opt in individually. If Ray Tracing is enabled, will imply Inclusive
+	// All skeletal meshes are excluded from the skin cache. Each must opt in individually. If Support Ray Tracing is enabled on a mesh, will force inclusive behavior on that mesh
 	Exclusive = 0,
 
 	// All skeletal meshes are included into the skin cache. Each must opt out individually
@@ -398,7 +398,7 @@ struct FSkeletalMeshLODInfo
 	FString SourceImportFilename;
 
 	/**
-	 * How this LOD uses the skin cache feature. Auto will defer to the default project global option. If Ray Tracing is enabled, will imply Enabled
+	 * How this LOD uses the skin cache feature. Auto will defer to the default project global option. If Support Ray Tracing is enabled on the mesh, will imply Enabled
 	 */
 	UPROPERTY(EditAnywhere, Category = SkeletalMeshLODInfo)
 	ESkinCacheUsage SkinCacheUsage = ESkinCacheUsage::Auto;
diff --git a/Engine/Source/Runtime/Engine/Private/Components/SkinnedMeshComponent.cpp b/Engine/Source/Runtime/Engine/Private/Components/SkinnedMeshComponent.cpp
--- a/Engine/Source/Runtime/Engine/Private/Components/SkinnedMeshComponent.cpp
+++ b/Engine/Source/Runtime/Engine/Private/Components/SkinnedMeshComponent.cpp
@@ -1471,7 +1471,11 @@ bool USkinnedMeshComponent::IsSkinCacheAllowed(int32 LodIdx) const
 {
 	static const IConsoleVariable* CVarDefaultGPUSkinCacheBehavior = IConsoleManager::Get().FindConsoleVariable(TEXT("r.SkinCache.DefaultBehavior"));
 
-	bool bIsRayTracing = IsRayTracingEnabled();
+	bool bIsRayTracing = FGPUSkinCache::IsGPUSkinCacheRayTracingSupported();
+	if (bIsRayTracing && SkeletalMesh)
+	{
+		bIsRayTracing = SkeletalMesh->bSupportRayTracing;
+	}
 
 	bool bGlobalDefault = CVarDefaultGPUSkinCacheBehavior && ESkinCacheDefaultBehavior(CVarDefaultGPUSkinCacheBehavior->GetInt()) == ESkinCacheDefaultBehavior::Inclusive;
 
diff --git a/Engine/Source/Runtime/Engine/Private/GPUSkinCache.cpp b/Engine/Source/Runtime/Engine/Private/GPUSkinCache.cpp
--- a/Engine/Source/Runtime/Engine/Private/GPUSkinCache.cpp
+++ b/Engine/Source/Runtime/Engine/Private/GPUSkinCache.cpp
@@ -59,7 +59,7 @@ static TAutoConsoleVariable<int32> CVarEnableGPUSkinCache(
 static TAutoConsoleVariable<int32> CVarDefaultGPUSkinCacheBehavior(
 	TEXT("r.SkinCache.DefaultBehavior"),
 	(int32)ESkinCacheDefaultBehavior::Inclusive,
-	TEXT("Default behavior if all skeletal meshes are included/excluded from the skin cache. If Ray Tracing is enabled, will imply Inclusive.\n")
+	TEXT("Default behavior if all skeletal meshes are included/excluded from the skin cache. If Support Ray Tracing is enabled on a mesh, will force inclusive behavior on that mesh.\n")
 	TEXT(" Exclusive ( 0): All skeletal meshes are excluded from the skin cache. Each must opt in individually.\n")
 	TEXT(" Inclusive ( 1): All skeletal meshes are included into the skin cache. Each must opt out individually. (default)")
 	);
@@ -1688,7 +1688,7 @@ void FGPUSkinCache::ProcessEntry(
 	}
 }
 
-static bool IsGPUSkinCacheRayTracingSupported()
+bool FGPUSkinCache::IsGPUSkinCacheRayTracingSupported()
 {
 #if RHI_RAYTRACING
 	static const auto CVar = IConsoleManager::Get().FindConsoleVariable(TEXT("r.RayTracing.Geometry.SupportSkeletalMeshes"));
diff --git a/Engine/Source/Runtime/Engine/Public/GPUSkinCache.h b/Engine/Source/Runtime/Engine/Public/GPUSkinCache.h
--- a/Engine/Source/Runtime/Engine/Public/GPUSkinCache.h
+++ b/Engine/Source/Runtime/Engine/Public/GPUSkinCache.h
@@ -205,6 +205,8 @@ public:
 
 	ENGINE_API uint64 GetExtraRequiredMemoryAndReset();
 
+	static bool IsGPUSkinCacheRayTracingSupported();
+
 	enum
 	{
 		NUM_BUFFERS = 2,
