Properly delete the old WorldSettings when setting a new one on a level to prevent accumulating stray world settings
Added a warning to mention stray world settings were detected and resaving the level or external actors is needed for proper cleanup

diff --git a/Engine/Source/Runtime/Engine/Private/Level.cpp b/Engine/Source/Runtime/Engine/Private/Level.cpp
--- a/Engine/Source/Runtime/Engine/Private/Level.cpp
+++ b/Engine/Source/Runtime/Engine/Private/Level.cpp
@@ -2109,13 +2109,16 @@ void ULevel::SetWorldSettings(AWorldSettings* NewWorldSettings)
 			}
 		}
 
-		if (WorldSettings)
+		// Assign the new world settings before destroying the old one
+		// since level will prevent destruction of the world settings if it matches the cached value;
+		AWorldSettings* OldWorldSettings = WorldSettings;
+		WorldSettings = NewWorldSettings;
+
+		if (OldWorldSettings)
 		{
 			// Makes no sense to have two WorldSettings so destroy existing one
-			WorldSettings->Destroy();
+			OldWorldSettings->Destroy();
 		}
-
-		WorldSettings = NewWorldSettings;
 	}
 }
 
diff --git a/Engine/Source/Runtime/Engine/Private/World.cpp b/Engine/Source/Runtime/Engine/Private/World.cpp
--- a/Engine/Source/Runtime/Engine/Private/World.cpp
+++ b/Engine/Source/Runtime/Engine/Private/World.cpp
@@ -1552,7 +1552,22 @@ void UWorld::RepairWorldSettings()
 		{
 			PersistentLevel->OwningWorld = nullptr;
 		}
+
+		// Now that we have set the proper world settings, clean up any other stay that may have accumulated due to legacy behaviors
+		if (PersistentLevel->Actors.Num() > 1)
+		{
+			for (int32 Index = 1, ActorNum = PersistentLevel->Actors.Num(); Index < ActorNum; ++Index)
+			{
+				AActor* Actor = PersistentLevel->Actors[Index];
+				if (Actor->IsA<AWorldSettings>())
+				{
+					UE_LOG(LogWorld, Warning, TEXT("Extra World Settings '%s' actor found. Resave level or actors to clean up"), *Actor->GetName());
+					Actor->Destroy();
+				}
+			}
+		}
 	}
+
 	check(GetWorldSettings());
 }
 
