Temp fix for broken HLOD in case the World Settings class is changed

diff --git a/Engine/Source/Runtime/Engine/Private/WorldPartition/HLOD/HLODSubsystem.cpp b/Engine/Source/Runtime/Engine/Private/WorldPartition/HLOD/HLODSubsystem.cpp
--- a/Engine/Source/Runtime/Engine/Private/WorldPartition/HLOD/HLODSubsystem.cpp
+++ b/Engine/Source/Runtime/Engine/Private/WorldPartition/HLOD/HLODSubsystem.cpp
@@ -99,7 +99,7 @@ void UHLODSubsystem::OnWorldPartitionRegistered(UWorldPartition* InWorldPartitio
 	// Build cell to HLOD mapping
 	for (const UWorldPartitionRuntimeCell* Cell : StreamingCells)
 	{
-		CellsHLODMapping.Emplace(Cell);
+		CellsHLODMapping.Emplace(Cell->GetFName());
 	}
 }
 
@@ -114,7 +114,15 @@ void UHLODSubsystem::RegisterHLODActor(AWorldPartitionHLOD* InWorldPartitionHLOD
 	TRACE_CPUPROFILER_EVENT_SCOPE(UHLODSubsystem::RegisterHLODActor);
 
 	const TSoftObjectPtr<UWorldPartitionRuntimeCell>& RuntimeCell = InWorldPartitionHLOD->GetSourceCell();
-	FCellHLODMapping* CellHLODs = CellsHLODMapping.Find(RuntimeCell);
+	const FString CellPath = FPackageName::GetShortName(RuntimeCell.ToSoftObjectPath().GetSubPathString());
+	FString CellName;
+	FString CellContext;
+	if (!CellPath.Split(TEXT("."), &CellContext, &CellName, ESearchCase::CaseSensitive, ESearchDir::FromEnd))
+	{
+		CellName = CellPath;
+	}
+
+	FCellHLODMapping* CellHLODs = CellsHLODMapping.Find(FName(CellName));
 
 #if WITH_EDITOR
 	UE_LOG(LogHLODSubsystem, Verbose, TEXT("Registering HLOD %s (%s) for cell %s"), *InWorldPartitionHLOD->GetActorLabel(), *InWorldPartitionHLOD->GetActorGuid().ToString(), *RuntimeCell.ToString());
@@ -137,7 +145,15 @@ void UHLODSubsystem::UnregisterHLODActor(AWorldPartitionHLOD* InWorldPartitionHL
 	TRACE_CPUPROFILER_EVENT_SCOPE(UHLODSubsystem::UnregisterHLODActor);
 
 	const TSoftObjectPtr<UWorldPartitionRuntimeCell>& RuntimeCell = InWorldPartitionHLOD->GetSourceCell();
-	FCellHLODMapping* CellHLODs = CellsHLODMapping.Find(RuntimeCell);
+	const FString CellPath = FPackageName::GetShortName(RuntimeCell.ToSoftObjectPath().GetSubPathString());
+	FString CellName;
+	FString CellContext;
+	if (!CellPath.Split(TEXT("."), &CellContext, &CellName, ESearchCase::CaseSensitive, ESearchDir::FromEnd))
+	{
+		CellName = CellPath;
+	}
+
+	FCellHLODMapping* CellHLODs = CellsHLODMapping.Find(FName(CellName));
 
 #if WITH_EDITOR
 	UE_LOG(LogHLODSubsystem, Verbose, TEXT("Unregistering HLOD %s (%s) for cell %s"), *InWorldPartitionHLOD->GetActorLabel(), *InWorldPartitionHLOD->GetActorGuid().ToString(), *RuntimeCell.ToString());
@@ -152,7 +168,7 @@ void UHLODSubsystem::UnregisterHLODActor(AWorldPartitionHLOD* InWorldPartitionHL
 
 void UHLODSubsystem::OnCellShown(const UWorldPartitionRuntimeCell* InCell)
 {
-	FCellHLODMapping& CellHLODs = CellsHLODMapping.FindChecked(InCell);
+	FCellHLODMapping& CellHLODs = CellsHLODMapping.FindChecked(InCell->GetFName());
 	CellHLODs.bIsCellVisible = true;
 
 #if WITH_EDITOR
@@ -170,7 +186,7 @@ void UHLODSubsystem::OnCellShown(const UWorldPartitionRuntimeCell* InCell)
 
 void UHLODSubsystem::OnCellHidden(const UWorldPartitionRuntimeCell* InCell)
 {
-	FCellHLODMapping& CellHLODs = CellsHLODMapping.FindChecked(InCell);
+	FCellHLODMapping& CellHLODs = CellsHLODMapping.FindChecked(InCell->GetFName());
 	CellHLODs.bIsCellVisible = false;
 
 #if WITH_EDITOR
diff --git a/Engine/Source/Runtime/Engine/Public/WorldPartition/HLOD/HLODSubsystem.h b/Engine/Source/Runtime/Engine/Public/WorldPartition/HLOD/HLODSubsystem.h
--- a/Engine/Source/Runtime/Engine/Public/WorldPartition/HLOD/HLODSubsystem.h
+++ b/Engine/Source/Runtime/Engine/Public/WorldPartition/HLOD/HLODSubsystem.h
@@ -59,5 +59,5 @@ private:
 		}
 	};
 
-	TMap<TSoftObjectPtr<UWorldPartitionRuntimeCell>, FCellHLODMapping> CellsHLODMapping;
+	TMap<FName, FCellHLODMapping> CellsHLODMapping;
 };
