Fixed uninitialized LocalToPrimitive matrix for skeletal meshes.

diff --git a/Engine/Source/Runtime/Engine/Private/SkeletalMesh.cpp b/Engine/Source/Runtime/Engine/Private/SkeletalMesh.cpp
--- a/Engine/Source/Runtime/Engine/Private/SkeletalMesh.cpp
+++ b/Engine/Source/Runtime/Engine/Private/SkeletalMesh.cpp
@@ -5683,6 +5683,8 @@ FSkeletalMeshSceneProxy::FSkeletalMeshSceneProxy(const USkinnedMeshComponent* Co
 		// 100% generic default-instance
 		Instances.SetNum(1);
 		FPrimitiveInstance& Instance = Instances[0];
+		Instance.LocalToPrimitive.SetIdentity();
+		Instance.PrevLocalToPrimitive.SetIdentity();
 		Instance.LocalBounds = GetLocalBounds();
 	}
 
