Added material cache flags for vertex interpolator, per-instance custom data, and per-instance random usage in a material graph.

diff --git a/Engine/Source/Runtime/Engine/Private/Materials/HLSLMaterialTranslator.cpp b/Engine/Source/Runtime/Engine/Private/Materials/HLSLMaterialTranslator.cpp
--- a/Engine/Source/Runtime/Engine/Private/Materials/HLSLMaterialTranslator.cpp
+++ b/Engine/Source/Runtime/Engine/Private/Materials/HLSLMaterialTranslator.cpp
@@ -151,7 +151,6 @@ FHLSLMaterialTranslator::FHLSLMaterialTranslator(FMaterial* InMaterial,
 ,	bUsesDistanceCullFade(false)
 ,	bIsFullyRough(0)
 ,	bAllowCodeChunkGeneration(true)
-,	bUsesPerInstanceCustomData(false)
 ,	bUsesAnisotropy(false)
 ,	bMaterialIsStrata(false)
 ,	AllocatedUserTexCoords()
@@ -1427,8 +1426,10 @@ void FHLSLMaterialTranslator::GetMaterialEnvironment(EShaderPlatform InPlatform,
 		OutEnvironment.SetDefine(TEXT("VIRTUAL_TEXTURE_OUTPUT"), 1);
 	}
 
-	OutEnvironment.SetDefine(TEXT("USES_PER_INSTANCE_CUSTOM_DATA"), bUsesPerInstanceCustomData && Material->IsUsedWithInstancedStaticMeshes());
-		
+	OutEnvironment.SetDefine(TEXT("USES_PER_INSTANCE_CUSTOM_DATA"), MaterialCompilationOutput.bUsesPerInstanceCustomData && Material->IsUsedWithInstancedStaticMeshes());
+	OutEnvironment.SetDefine(TEXT("USES_PER_INSTANCE_RANDOM"), MaterialCompilationOutput.bUsesPerInstanceRandom && Material->IsUsedWithInstancedStaticMeshes());
+	OutEnvironment.SetDefine(TEXT("USES_VERTEX_INTERPOLATOR"), MaterialCompilationOutput.bUsesVertexInterpolator);
+
 	OutEnvironment.SetDefine(TEXT("MATERIAL_SKY_ATMOSPHERE"), bUsesSkyAtmosphere);
 	OutEnvironment.SetDefine(TEXT("INTERPOLATE_VERTEX_COLOR"), bUsesVertexColor);
 	OutEnvironment.SetDefine(TEXT("NEEDS_PARTICLE_COLOR"), bUsesParticleColor); 
@@ -6690,6 +6691,8 @@ int32 FHLSLMaterialTranslator::VertexInterpolator(uint32 InterpolatorIndex)
 		return Errorf(TEXT("Invalid custom interpolator index."));
 	}
 
+	MaterialCompilationOutput.bUsesVertexInterpolator = true;
+
 	UMaterialExpressionVertexInterpolator* Interpolator = *InterpolatorPtr;
 	check(Interpolator->InterpolatorIndex == InterpolatorIndex);
 	check(Interpolator->InterpolatedType & MCT_Float);
@@ -9643,6 +9646,7 @@ int32 FHLSLMaterialTranslator::PerInstanceRandom()
 	}
 	else
 	{
+		MaterialCompilationOutput.bUsesPerInstanceRandom = true;
 		return AddInlinedCodeChunkZeroDeriv(MCT_Float, TEXT("GetPerInstanceRandom(Parameters)"));
 	}
 }
@@ -9678,7 +9682,7 @@ int32 FHLSLMaterialTranslator::PerInstanceCustomData(int32 DataIndex, int32 Defa
 	}
 	else
 	{
-		bUsesPerInstanceCustomData = true;
+		MaterialCompilationOutput.bUsesPerInstanceCustomData = true;
 		return AddInlinedCodeChunkZeroDeriv(MCT_Float, TEXT("GetPerInstanceCustomData(Parameters, %d, %s)"), DataIndex, *GetParameterCode(DefaultValueIndex));
 	}
 }
diff --git a/Engine/Source/Runtime/Engine/Private/Materials/HLSLMaterialTranslator.h b/Engine/Source/Runtime/Engine/Private/Materials/HLSLMaterialTranslator.h
--- a/Engine/Source/Runtime/Engine/Private/Materials/HLSLMaterialTranslator.h
+++ b/Engine/Source/Runtime/Engine/Private/Materials/HLSLMaterialTranslator.h
@@ -368,9 +368,6 @@ protected:
 	uint32 bIsFullyRough : 1;
 	/** true if allowed to generate code chunks. Translator operates in two phases; generate all code chunks & query meta data based on generated code chunks. */
 	uint32 bAllowCodeChunkGeneration : 1;
-	
-	/** True if this material reads any per-instance custom data */
-	uint32 bUsesPerInstanceCustomData : 1;
 
 	/** True if this material write anisotropy material property */
 	uint32 bUsesAnisotropy : 1;
diff --git a/Engine/Source/Runtime/Engine/Private/Materials/Material.cpp b/Engine/Source/Runtime/Engine/Private/Materials/Material.cpp
--- a/Engine/Source/Runtime/Engine/Private/Materials/Material.cpp
+++ b/Engine/Source/Runtime/Engine/Private/Materials/Material.cpp
@@ -5364,9 +5364,23 @@ bool UMaterial::CanBeClusterRoot() const
 
 void UMaterial::GetAssetRegistryTags(TArray<FAssetRegistryTag>& OutTags) const
 {
-	// CachedExpressionData may be NULL here for default objects
-	const bool bHasSceneColor = CachedExpressionData ? CachedExpressionData->bHasSceneColor : false;
-	OutTags.Add(FAssetRegistryTag("HasSceneColor", bHasSceneColor ? TEXT("True") : TEXT("False"), FAssetRegistryTag::TT_Alphabetical));
+	// CachedExpressionData may be nullptr here for default objects
+	{
+		const bool bHasSceneColor = CachedExpressionData ? CachedExpressionData->bHasSceneColor : false;
+		OutTags.Add(FAssetRegistryTag("HasSceneColor", bHasSceneColor ? TEXT("True") : TEXT("False"), FAssetRegistryTag::TT_Alphabetical));
+	}
+	{
+		const bool bHasPerInstanceRandom = CachedExpressionData ? CachedExpressionData->bHasPerInstanceRandom : false;
+		OutTags.Add(FAssetRegistryTag("HasPerInstanceRandom", bHasPerInstanceRandom ? TEXT("True") : TEXT("False"), FAssetRegistryTag::TT_Alphabetical));
+	}
+	{
+		const bool bHasPerInstanceCustomData = CachedExpressionData ? CachedExpressionData->bHasPerInstanceCustomData : false;
+		OutTags.Add(FAssetRegistryTag("HasPerInstanceCustomData", bHasPerInstanceCustomData ? TEXT("True") : TEXT("False"), FAssetRegistryTag::TT_Alphabetical));
+	}
+	{
+		const bool bHasVertexInterpolator = CachedExpressionData ? CachedExpressionData->bHasVertexInterpolator : false;
+		OutTags.Add(FAssetRegistryTag("HasVertexInterpolator", bHasVertexInterpolator ? TEXT("True") : TEXT("False"), FAssetRegistryTag::TT_Alphabetical));
+	}
 	Super::GetAssetRegistryTags(OutTags);
 }
 
diff --git a/Engine/Source/Runtime/Engine/Private/Materials/MaterialCachedData.cpp b/Engine/Source/Runtime/Engine/Private/Materials/MaterialCachedData.cpp
--- a/Engine/Source/Runtime/Engine/Private/Materials/MaterialCachedData.cpp
+++ b/Engine/Source/Runtime/Engine/Private/Materials/MaterialCachedData.cpp
@@ -15,6 +15,8 @@
 #include "Materials/MaterialExpressionShaderStageSwitch.h"
 #include "Materials/MaterialExpressionMakeMaterialAttributes.h"
 #include "Materials/MaterialExpressionParameter.h"
+#include "Materials/MaterialExpressionPerInstanceCustomData.h"
+#include "Materials/MaterialExpressionPerInstanceRandom.h"
 #include "Materials/MaterialExpressionRuntimeVirtualTextureSampleParameter.h"
 #include "Materials/MaterialExpressionScalarParameter.h"
 #include "Materials/MaterialExpressionStaticBoolParameter.h"
@@ -56,6 +58,9 @@ void FMaterialCachedExpressionData::Reset()
 	QualityLevelsUsed.AddDefaulted(EMaterialQualityLevel::Num);
 	bHasRuntimeVirtualTextureOutput = false;
 	bHasSceneColor = false;
+	bHasPerInstanceCustomData = false;
+	bHasPerInstanceRandom = false;
+	bHasVertexInterpolator = false;
 	MaterialAttributesPropertyConnectedBitmask = 0;
 
 	static_assert((uint32)(EMaterialProperty::MP_MAX)-1 <= (8 * sizeof(MaterialAttributesPropertyConnectedBitmask)), "MaterialAttributesPropertyConnectedBitmask cannot contain entire EMaterialProperty enumeration.");
@@ -600,7 +605,6 @@ bool FMaterialCachedExpressionData::UpdateForExpressions(const FMaterialCachedEx
 					}
 				}
 			}
-
 		}
 		else if (Expression->IsA(UMaterialExpressionRuntimeVirtualTextureOutput::StaticClass()))
 		{
@@ -610,6 +614,18 @@ bool FMaterialCachedExpressionData::UpdateForExpressions(const FMaterialCachedEx
 		{
 			bHasSceneColor = true;
 		}
+		else if (Expression->IsA(UMaterialExpressionPerInstanceRandom::StaticClass()))
+		{
+			bHasPerInstanceRandom = true;
+		}
+		else if (Expression->IsA(UMaterialExpressionPerInstanceCustomData::StaticClass()))
+		{
+			bHasPerInstanceCustomData = true;
+		}
+		else if (Expression->IsA(UMaterialExpressionVertexInterpolator::StaticClass()))
+		{
+			bHasVertexInterpolator = true;
+		}
 		else if (UMaterialExpressionMaterialAttributeLayers* LayersExpression = Cast<UMaterialExpressionMaterialAttributeLayers>(Expression))
 		{
 			checkf(Association == GlobalParameter, TEXT("UMaterialExpressionMaterialAttributeLayers can't be nested"));
diff --git a/Engine/Source/Runtime/Engine/Private/Materials/MaterialHLSLEmitter.cpp b/Engine/Source/Runtime/Engine/Private/Materials/MaterialHLSLEmitter.cpp
--- a/Engine/Source/Runtime/Engine/Private/Materials/MaterialHLSLEmitter.cpp
+++ b/Engine/Source/Runtime/Engine/Private/Materials/MaterialHLSLEmitter.cpp
@@ -423,7 +423,9 @@ static void GetMaterialEnvironment(EShaderPlatform InPlatform,
 		OutEnvironment.SetDefine(TEXT("VIRTUAL_TEXTURE_OUTPUT"), 1);
 	}
 
-	OutEnvironment.SetDefine(TEXT("USES_PER_INSTANCE_CUSTOM_DATA"), false);// bUsesPerInstanceCustomData&& Material->IsUsedWithInstancedStaticMeshes());
+	OutEnvironment.SetDefine(TEXT("USES_PER_INSTANCE_CUSTOM_DATA"), MaterialCompilationOutput.bUsesPerInstanceCustomData && InMaterial.IsUsedWithInstancedStaticMeshes());
+	OutEnvironment.SetDefine(TEXT("USES_PER_INSTANCE_RANDOM"), MaterialCompilationOutput.bUsesPerInstanceRandom && InMaterial.IsUsedWithInstancedStaticMeshes());
+	OutEnvironment.SetDefine(TEXT("USES_VERTEX_INTERPOLATOR"), MaterialCompilationOutput.bUsesVertexInterpolator);
 
 	// @todo MetalMRT: Remove this hack and implement proper atmospheric-fog solution for Metal MRT...
 	OutEnvironment.SetDefine(TEXT("MATERIAL_ATMOSPHERIC_FOG"), false);// !IsMetalMRTPlatform(InPlatform) ? bUsesAtmosphericFog : 0);
diff --git a/Engine/Source/Runtime/Engine/Private/Materials/MaterialShared.cpp b/Engine/Source/Runtime/Engine/Private/Materials/MaterialShared.cpp
--- a/Engine/Source/Runtime/Engine/Private/Materials/MaterialShared.cpp
+++ b/Engine/Source/Runtime/Engine/Private/Materials/MaterialShared.cpp
@@ -1687,6 +1687,21 @@ uint32 FMaterialResource::GetStencilCompare() const
 	return GetMaterialDomain() == MD_PostProcess ? uint32(Material->StencilCompare.GetValue()) : 0;
 }
 
+bool FMaterialResource::HasPerInstanceCustomData() const
+{
+	return Material->GetCachedExpressionData().bHasPerInstanceCustomData;
+}
+
+bool FMaterialResource::HasPerInstanceRandom() const
+{
+	return Material->GetCachedExpressionData().bHasPerInstanceRandom;
+}
+
+bool FMaterialResource::HasVertexInterpolator() const
+{
+	return Material->GetCachedExpressionData().bHasVertexInterpolator;
+}
+
 bool FMaterialResource::HasRuntimeVirtualTextureOutput() const
 {
 	return Material->GetCachedExpressionData().bHasRuntimeVirtualTextureOutput;
diff --git a/Engine/Source/Runtime/Engine/Public/MaterialCachedData.h b/Engine/Source/Runtime/Engine/Public/MaterialCachedData.h
--- a/Engine/Source/Runtime/Engine/Public/MaterialCachedData.h
+++ b/Engine/Source/Runtime/Engine/Public/MaterialCachedData.h
@@ -236,6 +236,9 @@ struct FMaterialCachedExpressionData
 	FMaterialCachedExpressionData()
 		: bHasRuntimeVirtualTextureOutput(false)
 		, bHasSceneColor(false)
+		, bHasPerInstanceCustomData(false)
+		, bHasPerInstanceRandom(false)
+		, bHasVertexInterpolator(false)
 	{}
 
 #if WITH_EDITOR
@@ -295,6 +298,15 @@ struct FMaterialCachedExpressionData
 	UPROPERTY()
 	uint32 bHasSceneColor : 1;
 
+	UPROPERTY()
+	uint32 bHasPerInstanceCustomData : 1;
+
+	UPROPERTY()
+	uint32 bHasPerInstanceRandom : 1;
+
+	UPROPERTY()
+	uint32 bHasVertexInterpolator : 1;
+
 	/** Each bit corresponds to EMaterialProperty connection status. */
 	UPROPERTY()
 	uint32 MaterialAttributesPropertyConnectedBitmask = 0;
diff --git a/Engine/Source/Runtime/Engine/Public/MaterialShared.h b/Engine/Source/Runtime/Engine/Public/MaterialShared.h
--- a/Engine/Source/Runtime/Engine/Public/MaterialShared.h
+++ b/Engine/Source/Runtime/Engine/Public/MaterialShared.h
@@ -563,6 +563,9 @@ public:
 		bUsesGlobalDistanceField(false),
 		bUsesPixelDepthOffset(false),
 		bUsesDistanceCullFade(false),
+		bUsesPerInstanceCustomData(false),
+		bUsesPerInstanceRandom(false),
+		bUsesVertexInterpolator(false),
 		bHasRuntimeVirtualTextureOutputNode(false),
 		bUsesAnisotropy(false)
 	{}
@@ -653,6 +656,15 @@ public:
 	/** true if the material uses distance cull fade */
 	LAYOUT_BITFIELD(uint8, bUsesDistanceCullFade, 1);
 
+	/** true if the material uses per-instance custom data */
+	LAYOUT_BITFIELD(uint8, bUsesPerInstanceCustomData, 1);
+
+	/** true if the material uses per-instance random */
+	LAYOUT_BITFIELD(uint8, bUsesPerInstanceRandom, 1);
+
+	/** true if the material uses vertex interpolator */
+	LAYOUT_BITFIELD(uint8, bUsesVertexInterpolator, 1);
+
 	/** true if the material writes to a runtime virtual texture custom output node. */
 	LAYOUT_BITFIELD(uint8, bHasRuntimeVirtualTextureOutputNode, 1);
 
@@ -1719,6 +1731,9 @@ public:
 	virtual bool IsStencilTestEnabled() const { return false; }
 	virtual uint32 GetStencilRefValue() const { return 0; }
 	virtual uint32 GetStencilCompare() const { return 0; }
+	virtual bool HasPerInstanceCustomData() const { return false; }
+	virtual bool HasPerInstanceRandom() const { return false; }
+	virtual bool HasVertexInterpolator() const { return false; }
 	virtual bool HasRuntimeVirtualTextureOutput() const { return false; }
 	virtual bool HasMaterialLayers() const { return false; }
 	virtual bool CastsRayTracedShadows() const { return true; }
@@ -2597,6 +2612,9 @@ public:
 	ENGINE_API virtual bool ShouldApplyCloudFogging() const override;
 	ENGINE_API virtual bool IsSky() const override;
 	ENGINE_API virtual bool ComputeFogPerPixel() const override;
+	ENGINE_API virtual bool HasPerInstanceCustomData() const override;
+	ENGINE_API virtual bool HasPerInstanceRandom() const override;
+	ENGINE_API virtual bool HasVertexInterpolator() const override;
 	ENGINE_API virtual bool HasRuntimeVirtualTextureOutput() const override;
 	ENGINE_API virtual bool HasMaterialLayers() const override;
 	ENGINE_API virtual bool CastsRayTracedShadows() const override;
@@ -3169,6 +3187,9 @@ struct FMaterialShaderParameters
 			uint64 bIsTranslucencyWritingCustomDepth : 1;
 			uint64 bIsDitheredLODTransition : 1;
 			uint64 bIsUsedWithInstancedStaticMeshes : 1;
+			uint64 bHasPerInstanceCustomData : 1;
+			uint64 bHasPerInstanceRandom : 1;
+			uint64 bHasVertexInterpolator : 1;
 			uint64 bHasRuntimeVirtualTextureOutput : 1;
 			uint64 bIsUsedWithLidarPointCloud : 1;
 			uint64 bIsUsedWithVirtualHeightfieldMesh : 1;
@@ -3228,6 +3249,9 @@ struct FMaterialShaderParameters
 		bIsTranslucencyWritingCustomDepth = InMaterial->IsTranslucencyWritingCustomDepth();
 		bIsDitheredLODTransition = InMaterial->IsDitheredLODTransition();
 		bIsUsedWithInstancedStaticMeshes = InMaterial->IsUsedWithInstancedStaticMeshes();
+		bHasPerInstanceCustomData = InMaterial->HasPerInstanceCustomData();
+		bHasPerInstanceRandom = InMaterial->HasPerInstanceRandom();
+		bHasVertexInterpolator = InMaterial->HasVertexInterpolator();
 		bHasRuntimeVirtualTextureOutput = InMaterial->HasRuntimeVirtualTextureOutput();
 		bIsUsedWithLidarPointCloud = InMaterial->IsUsedWithLidarPointCloud();
 		bIsUsedWithVirtualHeightfieldMesh = InMaterial->IsUsedWithVirtualHeightfieldMesh();
