Replaced assert during LevelScriptActor regeneration with an error message and more information. This was occurring on a map with a level blueprint having an invalid/missing base class.

#rnx

diff --git a/Engine/Source/Runtime/Engine/Private/Level.cpp b/Engine/Source/Runtime/Engine/Private/Level.cpp
--- a/Engine/Source/Runtime/Engine/Private/Level.cpp
+++ b/Engine/Source/Runtime/Engine/Private/Level.cpp
@@ -3039,44 +3039,51 @@ void ULevel::RepairLevelScript()
 
 void ULevel::RegenerateLevelScriptActor()
 {
+	check(LevelScriptBlueprint);
+
 	bool bResetDebugObject = false;
 
-	check(LevelScriptBlueprint);
 	UClass* SpawnClass = (LevelScriptBlueprint->GeneratedClass) ? LevelScriptBlueprint->GeneratedClass : LevelScriptBlueprint->SkeletonGeneratedClass;
-	check(SpawnClass);
 
-	// Get rid of the old LevelScriptActor
-	if (LevelScriptActor)
+	if (SpawnClass)
 	{
-		// Clear the current debug object and indicate that it needs to be reset (below).
-		if (LevelScriptBlueprint->GetObjectBeingDebugged() == LevelScriptActor)
+		// Get rid of the old LevelScriptActor
+		if (LevelScriptActor)
 		{
-			bResetDebugObject = true;
-			LevelScriptBlueprint->SetObjectBeingDebugged(nullptr);
-		}
+			// Clear the current debug object and indicate that it needs to be reset (below).
+			if (LevelScriptBlueprint->GetObjectBeingDebugged() == LevelScriptActor)
+			{
+				bResetDebugObject = true;
+				LevelScriptBlueprint->SetObjectBeingDebugged(nullptr);
+			}
 
-		LevelScriptActor->MarkPendingKill();
-		LevelScriptActor = nullptr;
-	}
+			LevelScriptActor->MarkPendingKill();
+			LevelScriptActor = nullptr;
+		}
 
-	check(OwningWorld);
-	// Create the new one
-	FActorSpawnParameters SpawnInfo;
-	SpawnInfo.OverrideLevel = this;
-	LevelScriptActor = OwningWorld->SpawnActor<ALevelScriptActor>(SpawnClass, SpawnInfo);
+		check(OwningWorld);
+		// Create the new one
+		FActorSpawnParameters SpawnInfo;
+		SpawnInfo.OverrideLevel = this;
+		LevelScriptActor = OwningWorld->SpawnActor<ALevelScriptActor>(SpawnClass, SpawnInfo);
 
-	if (LevelScriptActor)
-	{
-		// Reset the current debug object to the new instance if it was previously set to the old instance.
-		if (bResetDebugObject)
+		if (LevelScriptActor)
 		{
-			LevelScriptBlueprint->SetObjectBeingDebugged(LevelScriptActor);
-		}
+			// Reset the current debug object to the new instance if it was previously set to the old instance.
+			if (bResetDebugObject)
+			{
+				LevelScriptBlueprint->SetObjectBeingDebugged(LevelScriptActor);
+			}
 
-		LevelScriptActor->ClearFlags(RF_Transactional);
-		check(LevelScriptActor->GetLevel() == this);
-		// Finally, fixup all the bound events to point to their new LSA
-		FBlueprintEditorUtils::FixLevelScriptActorBindings(LevelScriptActor, LevelScriptBlueprint);
+			LevelScriptActor->ClearFlags(RF_Transactional);
+			check(LevelScriptActor->GetLevel() == this);
+			// Finally, fixup all the bound events to point to their new LSA
+			FBlueprintEditorUtils::FixLevelScriptActorBindings(LevelScriptActor, LevelScriptBlueprint);
+		}
+	}
+	else
+	{
+		UE_LOG(LogLevel, Error, TEXT("Skipped regeneration of LevelScriptActor due to blueprint '%s' having no spawnable class. A probable cause is that the blueprint is deriving from an invalid class, and may not function."), *LevelScriptBlueprint->GetName());
 	}
 }
 #endif // WITH_EDITOR
