fixed crash building svt for rvt

diff --git a/Engine/Source/Runtime/Engine/Private/VT/VirtualTextureTranscodeCache.cpp b/Engine/Source/Runtime/Engine/Private/VT/VirtualTextureTranscodeCache.cpp
--- a/Engine/Source/Runtime/Engine/Private/VT/VirtualTextureTranscodeCache.cpp
+++ b/Engine/Source/Runtime/Engine/Private/VT/VirtualTextureTranscodeCache.cpp
@@ -432,12 +432,7 @@ FVTTranscodeTileHandle FVirtualTextureTranscodeCache::SubmitTask(FVirtualTexture
 			TaskEntry.StageTileHandle[LayerIndex] = InUploadCache.PrepareTileForUpload(StagingBufferForLayer, LayerFormat, TilePixelSize);
 
 			// If there aren't any prerequisites, it's not worth launching a task for a raw memcpy, just do the work now
-			if (!Prerequisites && Chunk.CodecType[LayerIndex] == EVirtualTextureCodec::RawGPU)
-			{
-				FTranscodeTask LocalTask(StagingBuffer, InParams);
-				LocalTask.DoTask(false);
-			}
-			else
+			if (Prerequisites || Chunk.CodecType[LayerIndex] != EVirtualTextureCodec::RawGPU)
 			{
 				// TODO - if this logic gets more complex, should mask off layers that are processed inline when kicking off jobs
 				// for now in practice, should either process all layers inline, or all layers in job
@@ -450,7 +445,11 @@ FVTTranscodeTileHandle FVirtualTextureTranscodeCache::SubmitTask(FVirtualTexture
 	{
 		TaskEntry.GraphEvent = TGraphTask<FTranscodeTask>::CreateTask(Prerequisites).ConstructAndDispatchWhenReady(StagingBuffer, InParams);
 	}
-
+	else
+	{
+		FTranscodeTask LocalTask(StagingBuffer, InParams);
+		LocalTask.DoTask(false);
+	}
 	return FVTTranscodeTileHandle(TaskIndex, TaskEntry.Magic);
 }
 
