Using cached poses in AnimBP will cause assertion when running in Debug mode
#fix make the arrays containing the wrapped animation attributes heap-based, when saving a cache pose for the 2nd time the previousy stack-based attributes array will try to free its memory (which has already been freed due to memstack being reset)

diff --git a/Engine/Source/Runtime/Engine/Classes/Animation/AttributesContainer.h b/Engine/Source/Runtime/Engine/Classes/Animation/AttributesContainer.h
--- a/Engine/Source/Runtime/Engine/Classes/Animation/AttributesContainer.h
+++ b/Engine/Source/Runtime/Engine/Classes/Animation/AttributesContainer.h
@@ -74,8 +74,8 @@ namespace UE
 				const int32 NumTypes = UniqueTypes.Num();
 				for (int32 TypeIndex = 0; TypeIndex < NumTypes; ++TypeIndex)
 				{
-					const TArray<TWrappedAttribute<OtherAllocator>, OtherAllocator> Array = Other.Values[TypeIndex];
-					TArray<TWrappedAttribute<InAllocator>, InAllocator>& ValuesArray = Values.AddDefaulted_GetRef();
+					const TArray<TWrappedAttribute<OtherAllocator>>& Array = Other.Values[TypeIndex];
+					TArray<TWrappedAttribute<InAllocator>>& ValuesArray = Values.AddDefaulted_GetRef();
 
 					const TWeakObjectPtr<const UScriptStruct>& WeakScriptStruct = UniqueTypes[TypeIndex];
 					for (const TWrappedAttribute<OtherAllocator>& StructOnScope : Array)
@@ -100,8 +100,8 @@ namespace UE
 					const int32 NumTypes = UniqueTypes.Num();
 					for (int32 TypeIndex = 0; TypeIndex < NumTypes; ++TypeIndex)
 					{
-						const TArray<TWrappedAttribute<InAllocator>, InAllocator> Array = Other.Values[TypeIndex];
-						TArray<TWrappedAttribute<InAllocator>, InAllocator>& ValuesArray = Values.AddDefaulted_GetRef();
+						const TArray<TWrappedAttribute<InAllocator>>& Array = Other.Values[TypeIndex];
+						TArray<TWrappedAttribute<InAllocator>>& ValuesArray = Values.AddDefaulted_GetRef();
 
 						const TWeakObjectPtr<const UScriptStruct>& WeakScriptStruct = UniqueTypes[TypeIndex];
 						for (const TWrappedAttribute<InAllocator>& StructOnScope : Array)
@@ -242,7 +242,7 @@ namespace UE
 
 					UniqueTypedBoneIndices[TypeIndex].AddUnique(InAttributeId.GetIndex());
 
-					TArray<TWrappedAttribute<InAllocator>, InAllocator>& TypedArray = Values[TypeIndex];
+					TArray<TWrappedAttribute<InAllocator>>& TypedArray = Values[TypeIndex];
 					NewIndex = TypedArray.Add(InScriptStruct);
 					TWrappedAttribute<InAllocator>& StructData = TypedArray[NewIndex];
 					InScriptStruct->InitializeStruct(StructData.template GetPtr<void>());
@@ -285,7 +285,7 @@ namespace UE
 				int32 AttributeIndex = AttributeIds.IndexOfByKey(InAttributeId);
 
 				ensure(Values.IsValidIndex(TypeIndex));
-				TArray<TWrappedAttribute<InAllocator>, InAllocator>& TypedArray = Values[TypeIndex];
+				TArray<TWrappedAttribute<InAllocator>>& TypedArray = Values[TypeIndex];
 
 				// Should only add an attribute once
 				if (AttributeIndex == INDEX_NONE)
@@ -340,7 +340,7 @@ namespace UE
 
 					const int32 AttributeIndex = AttributeIds.IndexOfByKey(InAttributeId);
 					ensure(Values.IsValidIndex(TypeIndex));
-					TArray<TWrappedAttribute<InAllocator>, InAllocator>& TypedArray = Values[TypeIndex];
+					TArray<TWrappedAttribute<InAllocator>>& TypedArray = Values[TypeIndex];
 
 					if (AttributeIndex != INDEX_NONE)
 					{
@@ -383,7 +383,7 @@ namespace UE
 					const int32 AttributeIndex = AttributeIds.IndexOfByKey(InAttributeId);
 
 					ensure(Values.IsValidIndex(TypeIndex));
-					const TArray<TWrappedAttribute<InAllocator>, InAllocator>& TypedArray = Values[TypeIndex];
+					const TArray<TWrappedAttribute<InAllocator>>& TypedArray = Values[TypeIndex];
 
 					if (AttributeIndex != INDEX_NONE)
 					{
@@ -414,7 +414,7 @@ namespace UE
 					const int32 AttributeIndex = AttributeIds.IndexOfByKey(InAttributeId);
 
 					ensure(Values.IsValidIndex(TypeIndex));
-					const TArray<TWrappedAttribute<InAllocator>, InAllocator>& TypedArray = Values[TypeIndex];
+					const TArray<TWrappedAttribute<InAllocator>>& TypedArray = Values[TypeIndex];
 
 					if (AttributeIndex != INDEX_NONE)
 					{
@@ -450,7 +450,7 @@ namespace UE
 					const int32 AttributeIndex = AttributeIds.IndexOfByKey(InAttributeId);
 
 					ensure(Values.IsValidIndex(TypeIndex));
-					const TArray<TWrappedAttribute<InAllocator>, InAllocator>& TypedArray = Values[TypeIndex];
+					const TArray<TWrappedAttribute<InAllocator>>& TypedArray = Values[TypeIndex];
 
 					if (AttributeIndex != INDEX_NONE)
 					{
@@ -524,7 +524,7 @@ namespace UE
 					// Can only remove if it exists
 					if (AttributeIndex != INDEX_NONE)
 					{
-						TArray<TWrappedAttribute<InAllocator>, InAllocator>& TypedArray = Values[TypeIndex];
+						TArray<TWrappedAttribute<InAllocator>>& TypedArray = Values[TypeIndex];
 
 						// If removing the last one, remove all data for the attribute type
 						if (TypedArray.Num() == 1)
@@ -607,7 +607,7 @@ namespace UE
 			/*
 			* @return Array of all the contained values for the provided TypeIndex
 			*/
-			const TArray<TWrappedAttribute<InAllocator>, InAllocator>& GetValues(int32 TypeIndex) const { return Values[TypeIndex]; }
+			const TArray<TWrappedAttribute<InAllocator>>& GetValues(int32 TypeIndex) const { return Values[TypeIndex]; }
 
 			/*
 			* @return Array of all the contained attribute types (UScriptStruct)
@@ -702,7 +702,7 @@ namespace UE
 			/*
 			* @return Array of all the contained values for the provided TypeIndex
 			*/
-			TArray<TWrappedAttribute<InAllocator>, InAllocator>& GetValuesInternal(int32 TypeIndex)  { return Values[TypeIndex]; }
+			TArray<TWrappedAttribute<InAllocator>>& GetValuesInternal(int32 TypeIndex)  { return Values[TypeIndex]; }
 
 			/*
 			* @return Array of all the contained keys for the provided TypeIndex
@@ -731,7 +731,7 @@ namespace UE
 			/** Unique bone indices for all contained entries of a specific attribute type */
 			TArray<TArray<int32, InAllocator>, InAllocator> UniqueTypedBoneIndices;
 			TArray<TArray<FAttributeId, InAllocator>, InAllocator> AttributeIdentifiers;
-			TArray<TArray<TWrappedAttribute<InAllocator>, InAllocator>, InAllocator> Values;
+			TArray<TArray<TWrappedAttribute<InAllocator>>> Values;
 			TArray<TWeakObjectPtr<const UScriptStruct>, InAllocator> UniqueTypes;
 
 			template <typename OtherBoneIndexType, typename OtherAllocator>
diff --git a/Engine/Source/Runtime/Engine/Private/Animation/AttributeBlendData.cpp b/Engine/Source/Runtime/Engine/Private/Animation/AttributeBlendData.cpp
--- a/Engine/Source/Runtime/Engine/Private/Animation/AttributeBlendData.cpp
+++ b/Engine/Source/Runtime/Engine/Private/Animation/AttributeBlendData.cpp
@@ -230,7 +230,7 @@ namespace UE { namespace Anim {
 		const int32 TypeIndex = AttributeContainers.FindTypeIndex(AttributeScriptStruct);
 		if (TypeIndex != INDEX_NONE)
 		{
-			const TArray<TWrappedAttribute<FAnimStackAllocator>, FAnimStackAllocator>& ValuesArray = AttributeContainers.GetValues(TypeIndex);
+			const TArray<TWrappedAttribute<FAnimStackAllocator>>& ValuesArray = AttributeContainers.GetValues(TypeIndex);
 			const TArray<FAttributeId, FAnimStackAllocator>& AttributeIdentifiers = AttributeContainers.GetKeys(TypeIndex);
 			
 			for (int32 AttributeIndex = 0; AttributeIndex < AttributeIdentifiers.Num(); ++AttributeIndex)
diff --git a/Engine/Source/Runtime/Engine/Private/Animation/AttributesRuntime.cpp b/Engine/Source/Runtime/Engine/Private/Animation/AttributesRuntime.cpp
--- a/Engine/Source/Runtime/Engine/Private/Animation/AttributesRuntime.cpp
+++ b/Engine/Source/Runtime/Engine/Private/Animation/AttributesRuntime.cpp
@@ -498,7 +498,7 @@ void Attributes::MirrorAttributes(FStackAttributeContainer& CustomAttributes, co
 				}
 			}
 
-			TArray<TWrappedAttribute<FAnimStackAllocator>, FAnimStackAllocator>& ValuesArray = FStackAttributeContainerAccessor::GetValues(CustomAttributes, TypeIndex);
+			TArray<TWrappedAttribute<FAnimStackAllocator>>& ValuesArray = FStackAttributeContainerAccessor::GetValues(CustomAttributes, TypeIndex);
 			const TArray<FAttributeId, FAnimStackAllocator>& AttributeIdentifiers = CustomAttributes.GetKeys(TypeIndex);
 
 			// gather attributes that are on mirrored bones 
diff --git a/Engine/Source/Runtime/Engine/Public/Animation/AttributesRuntime.h b/Engine/Source/Runtime/Engine/Public/Animation/AttributesRuntime.h
--- a/Engine/Source/Runtime/Engine/Public/Animation/AttributesRuntime.h
+++ b/Engine/Source/Runtime/Engine/Public/Animation/AttributesRuntime.h
@@ -44,7 +44,7 @@ namespace UE
 		template<class BoneIndexType, typename InAllocator>
         struct TAttributeContainerAccessor
 		{
-			static TArray<TWrappedAttribute<InAllocator>, InAllocator>& GetValues(TAttributeContainer<BoneIndexType, InAllocator>& Attributes,  int32 TypeIndex) {return Attributes.GetValuesInternal(TypeIndex);}
+			static TArray<TWrappedAttribute<InAllocator>>& GetValues(TAttributeContainer<BoneIndexType, InAllocator>& Attributes,  int32 TypeIndex) {return Attributes.GetValuesInternal(TypeIndex);}
 			static TArray<FAttributeId, InAllocator>& GetKeys(TAttributeContainer<BoneIndexType, InAllocator>& Attributes,  int32 TypeIndex) {return Attributes.GetKeysInternal(TypeIndex);}
 		};
 
