Texture streaming optimisation : Automatic detection at cook of PrimitiveComponents that never reference any streamable textures.
When it's the case, GetStreamingRenderAssetInfo skips the part where it processes used materials & streamed textures.
In a big project, saves ~35% of inclusive average time of FRenderAssetStreamingManager::IncrementalUpdate.

diff --git a/Engine/Source/Runtime/Engine/Classes/Components/PrimitiveComponent.h b/Engine/Source/Runtime/Engine/Classes/Components/PrimitiveComponent.h
--- a/Engine/Source/Runtime/Engine/Classes/Components/PrimitiveComponent.h
+++ b/Engine/Source/Runtime/Engine/Classes/Components/PrimitiveComponent.h
@@ -593,6 +593,9 @@ protected:
 	/** Result of last call to AreAllCollideableDescendantsRelative(). */
 	uint8 bCachedAllCollideableDescendantsRelative : 1;
 
+	UPROPERTY()
+	uint8 bHasNoStreamableTextures : 1;
+
 public:
 	/** If true then DoCustomNavigableGeometryExport will be called to collect navigable geometry of this component. */
 	UPROPERTY()
@@ -600,8 +603,8 @@ public:
 
 public:
 #if WITH_EDITORONLY_DATA
-		UPROPERTY()
-			TEnumAsByte<enum EHitProxyPriority> HitProxyPriority;
+	UPROPERTY()
+	TEnumAsByte<enum EHitProxyPriority> HitProxyPriority;
 #endif
 
 private:
@@ -1694,6 +1697,9 @@ public:
 	/** Return True if a primitive's parameters as well as its position is static during gameplay, and can thus use static lighting. */
 	bool HasStaticLighting() const;
 
+	/** Return true if primitive can skip getting texture streaming render asset info. */
+	bool CanSkipGetTextureStreamingRenderAssetInfo() const;
+
 	/** Returns true if the component is static and has the right static mesh setup to support lightmaps. */
 	virtual bool HasValidSettingsForStaticLighting(bool bOverlookInvalidComponents) const 
 	{
@@ -2035,6 +2041,11 @@ public:
 	virtual bool CanEditChange(const FProperty* InProperty) const override;
 	virtual void UpdateCollisionProfile();
 	virtual void PostEditImport() override;
+	PRAGMA_DISABLE_DEPRECATION_WARNINGS // Suppress compiler warning on override of deprecated function
+	UE_DEPRECATED(5.0, "Use version that takes FObjectPreSaveContext instead.")
+	virtual void PreSave(const class ITargetPlatform* TargetPlatform) override;
+	PRAGMA_ENABLE_DEPRECATION_WARNINGS
+	virtual void PreSave(FObjectPreSaveContext ObjectSaveContext) override;
 #endif // WITH_EDITOR
 	//~ End UObject Interface.
 
diff --git a/Engine/Source/Runtime/Engine/Classes/Engine/Texture.h b/Engine/Source/Runtime/Engine/Classes/Engine/Texture.h
--- a/Engine/Source/Runtime/Engine/Classes/Engine/Texture.h
+++ b/Engine/Source/Runtime/Engine/Classes/Engine/Texture.h
@@ -998,6 +998,12 @@ protected:
 	TArray<TObjectPtr<UAssetUserData>> AssetUserData;
 
 private:
+
+#if WITH_EDITOR
+	/** Used to mark texture streamable state when cooking. */
+	TOptional<bool> bCookedIsStreamable;
+#endif
+
 	/** The texture's resource, can be NULL */
 	class FTextureResource*	PrivateResource;
 	/** Value updated and returned by the render-thread to allow
@@ -1341,6 +1347,11 @@ public:
 	/** Sets the minimum number of mips that must be resident in memory (cannot be streamed). */
 	static void SetMinTextureResidentMipCount(int32 InMinTextureResidentMipCount);
 
+#if WITH_EDITOR
+	/** Called by ULevel::MarkNoStreamableTexturesPrimitiveComponents when cooking level. */
+	bool IsCandidateForTextureStreaming(const ITargetPlatform* InTargetPlatform) const;
+#endif
+
 protected:
 
 	/** The minimum number of mips that must be resident in memory (cannot be streamed). */
@@ -1363,6 +1374,7 @@ protected:
 	/** Notify any loaded material instances that the texture has changed. */
 	ENGINE_API void NotifyMaterials(const ENotifyMaterialsEffectOnShaders EffectOnShaders = ENotifyMaterialsEffectOnShaders::Default);
 
+	virtual bool GetStreamableRenderResourceState(FTexturePlatformData* InPlatformData, FStreamableRenderResourceState& OutState) const { return false; }
 #endif //WITH_EDITOR
 
 	void BeginFinalReleaseResource();
@@ -1375,7 +1387,7 @@ protected:
 	 * @param	MaxMipCount - optional limitation on the max mip count.
 	 * @return  The state to be passed to FStreamableTextureResource.
 	 */
-	FStreamableRenderResourceState GetResourcePostInitState(FTexturePlatformData* PlatformData, bool bAllowStreaming, int32 MinRequestMipCount = 0, int32 MaxMipCount = 0) const;
+	FStreamableRenderResourceState GetResourcePostInitState(FTexturePlatformData* PlatformData, bool bAllowStreaming, int32 MinRequestMipCount = 0, int32 MaxMipCount = 0, bool bSkipCanBeLoaded = false) const;
 };
 
 /** 
diff --git a/Engine/Source/Runtime/Engine/Classes/Engine/Texture2D.h b/Engine/Source/Runtime/Engine/Classes/Engine/Texture2D.h
--- a/Engine/Source/Runtime/Engine/Classes/Engine/Texture2D.h
+++ b/Engine/Source/Runtime/Engine/Classes/Engine/Texture2D.h
@@ -307,4 +307,10 @@ public:
 	 * This can reduce page table overhead but potentially increase the number of physical pools allocated.
 	 */
 	virtual bool IsVirtualTexturedWithSinglePhysicalSpace() const { return false;  }
+
+protected:
+
+#if WITH_EDITOR
+	virtual bool GetStreamableRenderResourceState(FTexturePlatformData* InPlatformData, FStreamableRenderResourceState& OutState) const override;
+#endif
 };
diff --git a/Engine/Source/Runtime/Engine/Classes/Engine/Texture2DArray.cpp b/Engine/Source/Runtime/Engine/Classes/Engine/Texture2DArray.cpp
--- a/Engine/Source/Runtime/Engine/Classes/Engine/Texture2DArray.cpp
+++ b/Engine/Source/Runtime/Engine/Classes/Engine/Texture2DArray.cpp
@@ -35,6 +35,20 @@ UTexture2DArray::UTexture2DArray(const FObjectInitializer& ObjectInitializer) :
 #endif
 }
 
+#if WITH_EDITOR
+bool UTexture2DArray::GetStreamableRenderResourceState(FTexturePlatformData* InPlatformData, FStreamableRenderResourceState& OutState) const
+{
+	TGuardValue<FTexturePlatformData*> Guard(const_cast<UTexture2DArray*>(this)->PlatformData, InPlatformData);
+	const FPixelFormatInfo& FormatInfo = GPixelFormats[GetPixelFormat()];
+	if (GetNumMips() > 0 && FormatInfo.Supported)
+	{
+		OutState = GetResourcePostInitState(PlatformData, GSupportsTexture2DArrayStreaming, 0, 0, /*bSkipCanBeLoaded*/ true);
+		return true;
+	}
+	return false;
+}
+#endif
+
 FTextureResource* UTexture2DArray::CreateResource()
 {
 	const FPixelFormatInfo& FormatInfo = GPixelFormats[GetPixelFormat()];
diff --git a/Engine/Source/Runtime/Engine/Classes/Engine/Texture2DArray.h b/Engine/Source/Runtime/Engine/Classes/Engine/Texture2DArray.h
--- a/Engine/Source/Runtime/Engine/Classes/Engine/Texture2DArray.h
+++ b/Engine/Source/Runtime/Engine/Classes/Engine/Texture2DArray.h
@@ -122,5 +122,6 @@ protected:
 
 #if WITH_EDITOR
 	void UpdateMipGenSettings();
+	virtual bool GetStreamableRenderResourceState(FTexturePlatformData* InPlatformData, FStreamableRenderResourceState& OutState) const override;
 #endif
 };
\ No newline at end of file
diff --git a/Engine/Source/Runtime/Engine/Classes/Engine/VolumeTexture.h b/Engine/Source/Runtime/Engine/Classes/Engine/VolumeTexture.h
--- a/Engine/Source/Runtime/Engine/Classes/Engine/VolumeTexture.h
+++ b/Engine/Source/Runtime/Engine/Classes/Engine/VolumeTexture.h
@@ -131,6 +131,7 @@ protected:
 
 #if WITH_EDITOR
 	void UpdateMipGenSettings();
+	virtual bool GetStreamableRenderResourceState(FTexturePlatformData* InPlatformData, FStreamableRenderResourceState& OutState) const override;
 #endif
 };
 
diff --git a/Engine/Source/Runtime/Engine/Private/Components/MeshComponent.cpp b/Engine/Source/Runtime/Engine/Private/Components/MeshComponent.cpp
--- a/Engine/Source/Runtime/Engine/Private/Components/MeshComponent.cpp
+++ b/Engine/Source/Runtime/Engine/Private/Components/MeshComponent.cpp
@@ -440,6 +440,11 @@ void UMeshComponent::CacheMaterialParameterNameIndices()
 
 void UMeshComponent::GetStreamingTextureInfoInner(FStreamingTextureLevelContext& LevelContext, const TArray<FStreamingTextureBuildInfo>* PreBuiltData, float ComponentScaling, TArray<FStreamingRenderAssetPrimitiveInfo>& OutStreamingTextures) const
 {
+	if (CanSkipGetTextureStreamingRenderAssetInfo())
+	{
+		return;
+	}
+
 	LevelContext.BindBuildData(PreBuiltData);
 
 	const int32 NumMaterials = GetNumMaterials();
diff --git a/Engine/Source/Runtime/Engine/Private/Components/PrimitiveComponent.cpp b/Engine/Source/Runtime/Engine/Private/Components/PrimitiveComponent.cpp
--- a/Engine/Source/Runtime/Engine/Private/Components/PrimitiveComponent.cpp
+++ b/Engine/Source/Runtime/Engine/Private/Components/PrimitiveComponent.cpp
@@ -34,6 +34,7 @@
 #include "UObject/RenderingObjectVersion.h"
 #include "UObject/FortniteMainBranchObjectVersion.h"
 #include "EngineModule.h"
+#include "UObject/ObjectSaveContext.h"
 
 #if WITH_EDITOR
 #include "Engine/LODActor.h"
@@ -409,6 +410,11 @@ bool UPrimitiveComponent::HasStaticLighting() const
 
 void UPrimitiveComponent::GetStreamingRenderAssetInfo(FStreamingTextureLevelContext& LevelContext, TArray<FStreamingRenderAssetPrimitiveInfo>& OutStreamingRenderAssets) const
 {
+	if (CanSkipGetTextureStreamingRenderAssetInfo())
+	{
+		return;
+	}
+
 	if (CVarStreamingUseNewMetrics.GetValueOnGameThread() != 0)
 	{
 		LevelContext.BindBuildData(nullptr);
@@ -1240,6 +1246,23 @@ void UPrimitiveComponent::PostDuplicate(bool bDuplicateForPIE)
 	Super::PostDuplicate(bDuplicateForPIE);
 }
 
+static int32 GEnableAutoDetectNoStreamableTextures = 1;
+static FAutoConsoleVariableRef CVarEnableAutoDetectNoStreamableTextures(
+	TEXT("r.Streaming.EnableAutoDetectNoStreamableTextures"),
+	GEnableAutoDetectNoStreamableTextures,
+	TEXT("Enables auto-detection at cook time of primitive components with no streamable textures. Can also be turned-off at runtime to skip optimisation."),
+	ECVF_Default
+);
+
+bool UPrimitiveComponent::CanSkipGetTextureStreamingRenderAssetInfo() const
+{
+#if WITH_EDITOR
+	return false;
+#else
+	return GEnableAutoDetectNoStreamableTextures && bHasNoStreamableTextures;
+#endif
+}
+
 #if WITH_EDITOR
 /**
  * Called after importing property values for this object (paste, duplicate or .t3d import)
@@ -1262,6 +1285,57 @@ void UPrimitiveComponent::PostEditImport()
 	// Setup the transient internal primitive data array here after import (to support duplicate/paste)
 	ResetCustomPrimitiveData();
 }
+
+void UPrimitiveComponent::PreSave(const class ITargetPlatform* TargetPlatform)
+{
+	PRAGMA_DISABLE_DEPRECATION_WARNINGS;
+	Super::PreSave(TargetPlatform);
+	PRAGMA_ENABLE_DEPRECATION_WARNINGS;
+}
+
+void UPrimitiveComponent::PreSave(FObjectPreSaveContext ObjectSaveContext)
+{
+	Super::PreSave(ObjectSaveContext);
+
+	if (!IsTemplate() && ObjectSaveContext.IsCooking())
+	{
+		// Reset flag
+		bHasNoStreamableTextures = false;
+
+		if (!GEnableAutoDetectNoStreamableTextures || (Mobility != EComponentMobility::Static))
+		{
+			return;
+		}
+
+		TArray<UMaterialInterface*> Materials;
+		GetUsedMaterials(Materials);
+		if (Materials.IsEmpty())
+		{
+			return;
+		}
+
+		TSet<const UTexture*> Textures;
+		for (UMaterialInterface* Material : Materials)
+		{
+			if (Material)
+			{
+				Material->GetReferencedTexturesAndOverrides(Textures);
+			}
+		}
+
+		bool bHasStreamableTextures = false;
+		for (const UTexture* Texture : Textures)
+		{
+			if (Texture && Texture->IsCandidateForTextureStreaming(ObjectSaveContext.GetTargetPlatform()))
+			{
+				bHasStreamableTextures = true;
+				break;
+			}
+		}
+
+		bHasNoStreamableTextures = !bHasStreamableTextures;
+	}
+}
 #endif
 
 void UPrimitiveComponent::BeginDestroy()
diff --git a/Engine/Source/Runtime/Engine/Private/Components/StaticMeshComponent.cpp b/Engine/Source/Runtime/Engine/Private/Components/StaticMeshComponent.cpp
--- a/Engine/Source/Runtime/Engine/Private/Components/StaticMeshComponent.cpp
+++ b/Engine/Source/Runtime/Engine/Private/Components/StaticMeshComponent.cpp
@@ -913,8 +913,21 @@ void UStaticMeshComponent::GetStreamingRenderAssetInfo(FStreamingTextureLevelCon
 		return;
 	}
 
-	const float TransformScale = GetTextureStreamingTransformScale();
-	GetStreamingTextureInfoInner(LevelContext, Mobility == EComponentMobility::Static ? &StreamingTextureData : nullptr, TransformScale * StreamingDistanceMultiplier, OutStreamingRenderAssets);
+	// Since GetTextureStreamingTransformScale can be slow for certain component types, only call it if necessary
+	TOptional<float> LazyTransformScale;
+	auto GetTransformScale = [this, &LazyTransformScale]()
+	{
+		if (!LazyTransformScale.IsSet())
+		{
+			LazyTransformScale = GetTextureStreamingTransformScale();
+		}
+		return *LazyTransformScale;
+	};
+	
+	if (!CanSkipGetTextureStreamingRenderAssetInfo())
+	{
+		GetStreamingTextureInfoInner(LevelContext, Mobility == EComponentMobility::Static ? &StreamingTextureData : nullptr, GetTransformScale() * StreamingDistanceMultiplier, OutStreamingRenderAssets);
+	}
 
 	// Process the lightmaps and shadowmaps entries.
 	for (int32 LODIndex = 0; LODIndex < LODData.Num(); ++LODIndex)
@@ -928,7 +941,7 @@ void UStaticMeshComponent::GetStreamingRenderAssetInfo(FStreamingTextureLevelCon
 			const FVector2D& Scale = Lightmap->GetCoordinateScale();
 			if (Scale.X > SMALL_NUMBER && Scale.Y > SMALL_NUMBER)
 			{
-				const float TexelFactor = GetStaticMesh()->GetLightmapUVDensity() * TransformScale / FMath::Min(Scale.X, Scale.Y);
+				const float TexelFactor = GetStaticMesh()->GetLightmapUVDensity() * GetTransformScale() / FMath::Min(Scale.X, Scale.Y);
 				new (OutStreamingRenderAssets) FStreamingRenderAssetPrimitiveInfo(Lightmap->GetTexture(LightmapIndex), Bounds, TexelFactor, PackedRelativeBox_Identity);
 				new (OutStreamingRenderAssets) FStreamingRenderAssetPrimitiveInfo(Lightmap->GetAOMaterialMaskTexture(), Bounds, TexelFactor, PackedRelativeBox_Identity);
 				new (OutStreamingRenderAssets) FStreamingRenderAssetPrimitiveInfo(Lightmap->GetSkyOcclusionTexture(), Bounds, TexelFactor, PackedRelativeBox_Identity);
@@ -941,7 +954,7 @@ void UStaticMeshComponent::GetStreamingRenderAssetInfo(FStreamingTextureLevelCon
 			const FVector2D& Scale = Shadowmap->GetCoordinateScale();
 			if (Scale.X > SMALL_NUMBER && Scale.Y > SMALL_NUMBER)
 			{
-				const float TexelFactor = GetStaticMesh()->GetLightmapUVDensity() * TransformScale / FMath::Min(Scale.X, Scale.Y);
+				const float TexelFactor = GetStaticMesh()->GetLightmapUVDensity() * GetTransformScale() / FMath::Min(Scale.X, Scale.Y);
 				new (OutStreamingRenderAssets) FStreamingRenderAssetPrimitiveInfo(Shadowmap->GetTexture(), Bounds, TexelFactor, PackedRelativeBox_Identity);
 			}
 		}
diff --git a/Engine/Source/Runtime/Engine/Private/Texture.cpp b/Engine/Source/Runtime/Engine/Private/Texture.cpp
--- a/Engine/Source/Runtime/Engine/Private/Texture.cpp
+++ b/Engine/Source/Runtime/Engine/Private/Texture.cpp
@@ -913,7 +913,32 @@ const TArray<UAssetUserData*>* UTexture::GetAssetUserDataArray() const
 	return &ToRawPtrTArrayUnsafe(AssetUserData);
 }
 
-FStreamableRenderResourceState UTexture::GetResourcePostInitState(FTexturePlatformData* PlatformData, bool bAllowStreaming, int32 MinRequestMipCount, int32 MaxMipCount) const
+#if WITH_EDITOR
+// Based on target platform, returns wether texture is candidate to be streamed.
+// This method is used to decide if PrimitiveComponent's bHasNoStreamableTextures flag can be set to true.
+// See ULevel::MarkNoStreamableTexturesPrimitiveComponents for details.
+bool UTexture::IsCandidateForTextureStreaming(const ITargetPlatform* InTargetPlatform) const
+{
+	const bool bIsVirtualTextureStreaming = InTargetPlatform->SupportsFeature(ETargetPlatformFeatures::VirtualTextureStreaming) ? VirtualTextureStreaming : false;
+	const bool bIsCandidateForTextureStreaming = InTargetPlatform->SupportsFeature(ETargetPlatformFeatures::TextureStreaming) && !bIsVirtualTextureStreaming;
+
+	if (bIsCandidateForTextureStreaming &&
+		!NeverStream &&
+		LODGroup != TEXTUREGROUP_UI &&
+		MipGenSettings != TMGS_NoMipmaps)
+	{
+		// If bCookedIsStreamable flag was previously computed, use it.
+		if (bCookedIsStreamable.IsSet())
+		{
+			return *bCookedIsStreamable;
+		}
+		return true;
+	}
+	return false;
+}
+#endif
+
+FStreamableRenderResourceState UTexture::GetResourcePostInitState(FTexturePlatformData* PlatformData, bool bAllowStreaming, int32 MinRequestMipCount, int32 MaxMipCount, bool bSkipCanBeLoaded) const
 {
 	// Create the resource with a mip count limit taking in consideration the asset LODBias.
 	// This ensures that the mip count stays constant when toggling asset streaming at runtime.
@@ -943,7 +968,7 @@ FStreamableRenderResourceState UTexture::GetResourcePostInitState(FTexturePlatfo
 		NumOfNonStreamingMips < NumMips && 
 		LODGroup != TEXTUREGROUP_UI && 
 		bAllowStreaming &&
-		PlatformData->CanBeLoaded())
+		(bSkipCanBeLoaded || PlatformData->CanBeLoaded()))
 	{
 		bMakeStreamble  = true;
 	}
diff --git a/Engine/Source/Runtime/Engine/Private/Texture2D.cpp b/Engine/Source/Runtime/Engine/Private/Texture2D.cpp
--- a/Engine/Source/Runtime/Engine/Private/Texture2D.cpp
+++ b/Engine/Source/Runtime/Engine/Private/Texture2D.cpp
@@ -930,6 +930,30 @@ bool UTexture2D::HasAlphaChannel() const
 	return false;
 }
 
+#if WITH_EDITOR
+bool UTexture2D::GetStreamableRenderResourceState(FTexturePlatformData* InPlatformData, FStreamableRenderResourceState& OutState) const
+{
+	TGuardValue<FTexturePlatformData*> Guard(const_cast<UTexture2D*>(this)->PrivatePlatformData, InPlatformData);
+	if (PlatformData)
+	{
+		if (IsCurrentlyVirtualTextured())
+		{
+			return false;
+		}
+
+		const EPixelFormat PixelFormat = GetPixelFormat();
+		const int32 NumMips = FMath::Min3<int32>(PlatformData->Mips.Num(), GMaxTextureMipCount, FStreamableRenderResourceState::MAX_LOD_COUNT);
+		if (NumMips && GPixelFormats[PixelFormat].Supported &&
+			(NumMips > 1 || FMath::Max(GetSizeX(), GetSizeY()) <= (int32)GetMax2DTextureDimension()))
+		{
+			OutState = GetResourcePostInitState(PlatformData, true, 0, NumMips, /*bSkipCanBeLoaded*/ true);
+			return true;
+		}
+	}
+	return false;
+}
+#endif
+
 FTextureResource* UTexture2D::CreateResource()
 {
 	TRACE_CPUPROFILER_EVENT_SCOPE(UTexture2D::CreateResource)
diff --git a/Engine/Source/Runtime/Engine/Private/TextureDerivedData.cpp b/Engine/Source/Runtime/Engine/Private/TextureDerivedData.cpp
--- a/Engine/Source/Runtime/Engine/Private/TextureDerivedData.cpp
+++ b/Engine/Source/Runtime/Engine/Private/TextureDerivedData.cpp
@@ -2105,6 +2105,7 @@ void UTexture::SerializeCookedPlatformData(FArchive& Ar)
 #if WITH_EDITOR
 	if (Ar.IsCooking() && Ar.IsPersistent())
 	{
+		bCookedIsStreamable.Reset();
 		if (!Ar.CookingTarget()->IsServerOnly())
 		{
 			FTextureBuildSettings BuildSettings;
@@ -2160,6 +2161,14 @@ void UTexture::SerializeCookedPlatformData(FArchive& Ar)
 			{
 				FTexturePlatformData* PlatformDataToSave = PlatformDataToSerialize[i];
 				PlatformDataToSave->FinishCache();
+
+				// Update bCookedIsStreamable for later use in IsCandidateForTextureStreaming
+				FStreamableRenderResourceState State;
+				if (GetStreamableRenderResourceState(PlatformDataToSave, State))
+				{
+					bCookedIsStreamable = !bCookedIsStreamable.IsSet() ? State.bSupportsStreaming : (*bCookedIsStreamable || State.bSupportsStreaming);
+				}
+
 				FName PixelFormatName = PixelFormatEnum->GetNameByValue(PlatformDataToSave->PixelFormat);
 				Ar << PixelFormatName;
 				int64 SkipOffsetLoc = Ar.Tell();
diff --git a/Engine/Source/Runtime/Engine/Private/VolumeTexture.cpp b/Engine/Source/Runtime/Engine/Private/VolumeTexture.cpp
--- a/Engine/Source/Runtime/Engine/Private/VolumeTexture.cpp
+++ b/Engine/Source/Runtime/Engine/Private/VolumeTexture.cpp
@@ -279,6 +279,21 @@ uint32 UVolumeTexture::CalcTextureMemorySizeEnum( ETextureMipCount Enum ) const
 	}
 }
 
+#if WITH_EDITOR
+bool UVolumeTexture::GetStreamableRenderResourceState(FTexturePlatformData* InPlatformData, FStreamableRenderResourceState& OutState) const
+{
+	TGuardValue<FTexturePlatformData*> Guard(const_cast<UVolumeTexture*>(this)->PlatformData, InPlatformData);
+	const FPixelFormatInfo& FormatInfo = GPixelFormats[GetPixelFormat()];
+	const bool bFormatIsSupported = FormatInfo.Supported;
+	if (GetNumMips() > 0 && GSupportsTexture3D && bFormatIsSupported)
+	{
+		OutState = GetResourcePostInitState(PlatformData, GSupportsVolumeTextureStreaming, 0, 0, /*bSkipCanBeLoaded*/ true);
+		return true;
+	}
+	return false;
+}
+#endif
+
 FTextureResource* UVolumeTexture::CreateResource()
 {
 	const FPixelFormatInfo& FormatInfo = GPixelFormats[GetPixelFormat()];
