Fix skin cache over reporting out of memory figure.

The cause is each skeletal mesh section calls ProcessEntry on skin cache, if skin cache doesn't have enough memory to allocate for the mesh, that error will be reported N times where N equals number of sections. This change lets the subsequent sections skip calling ProcessEntry if the first section fails.

diff --git a/Engine/Source/Runtime/Engine/Private/GPUSkinCache.cpp b/Engine/Source/Runtime/Engine/Private/GPUSkinCache.cpp
--- a/Engine/Source/Runtime/Engine/Private/GPUSkinCache.cpp
+++ b/Engine/Source/Runtime/Engine/Private/GPUSkinCache.cpp
@@ -1474,7 +1474,7 @@ void FGPUSkinCache::DoDispatch(FRHICommandListImmediate& RHICmdList, FGPUSkinCac
 	SkinCacheEntry->UpdateVertexFactoryDeclaration(Section);
 }
 
-void FGPUSkinCache::ProcessEntry(
+bool FGPUSkinCache::ProcessEntry(
 	FRHICommandListImmediate& RHICmdList, 
 	FGPUBaseSkinVertexFactory* VertexFactory,
 	FGPUSkinPassthroughVertexFactory* TargetVertexFactory, 
@@ -1578,7 +1578,7 @@ void FGPUSkinCache::ProcessEntry(
 			}
 
 			// Couldn't fit; caller will notify OOM
-			return;
+			return false;
 		}
 
 		InOutEntry = new FGPUSkinCacheEntry(this, Skin, NewPositionAllocation);
@@ -1686,6 +1686,8 @@ void FGPUSkinCache::ProcessEntry(
 	{
 		DoDispatch(RHICmdList, InOutEntry, Section, RevisionNumber);
 	}
+
+	return true;
 }
 
 bool FGPUSkinCache::IsGPUSkinCacheRayTracingSupported()
diff --git a/Engine/Source/Runtime/Engine/Private/SkeletalRenderGPUSkin.cpp b/Engine/Source/Runtime/Engine/Private/SkeletalRenderGPUSkin.cpp
--- a/Engine/Source/Runtime/Engine/Private/SkeletalRenderGPUSkin.cpp
+++ b/Engine/Source/Runtime/Engine/Private/SkeletalRenderGPUSkin.cpp
@@ -569,6 +569,7 @@ void FSkeletalMeshObjectGPUSkin::ProcessUpdatedDynamicData(FGPUSkinCache* GPUSki
 
 	if (bDataPresent)
 	{
+		bool bSkinCacheResult = true;
 		for (int32 SectionIdx = 0; SectionIdx < Sections.Num(); SectionIdx++)
 		{
 			const FSkelMeshRenderSection& Section = Sections[SectionIdx];
@@ -654,22 +655,26 @@ void FSkeletalMeshObjectGPUSkin::ProcessUpdatedDynamicData(FGPUSkinCache* GPUSki
 				// Matrices are transposed in UE meaning matrix multiples need to happen in reverse ((AB)x = b becomes xTBTAT = b).
 				FMatrix44f LocalToCloth = DynamicData->ClothObjectLocalToWorld * ClothLocalToWorld.Inverse();
 
-				GPUSkinCache->ProcessEntry(
-					RHICmdList, 
-					VertexFactory,
-					VertexFactoryData.PassthroughVertexFactories[SectionIdx].Get(), 
-					Section, 
-					this, 
-					&LOD.VertexOffsetVertexBuffers,
-					bMorph ? &LOD.MorphVertexBuffer : 0, 
-					bClothFactory ? &LODData.ClothVertexBuffer : 0,
-					bClothFactory ? DynamicData->ClothingSimData.Find(Section.CorrespondClothAssetIndex) : 0, 
-					LocalToCloth, 
-					DynamicData->ClothBlendWeight, 
-					RevisionNumber, 
-					SectionIdx, 
-					SkinCacheEntry
-					);
+				// ProcessEntry returns false if not enough memory is left in skin cache to allocate for the mesh, if that happens don't try to process subsequent sections because they will also fail.
+				if (bSkinCacheResult)
+				{
+					bSkinCacheResult = GPUSkinCache->ProcessEntry(
+						RHICmdList,
+						VertexFactory,
+						VertexFactoryData.PassthroughVertexFactories[SectionIdx].Get(),
+						Section,
+						this,
+						&LOD.VertexOffsetVertexBuffers,
+						bMorph ? &LOD.MorphVertexBuffer : 0,
+						bClothFactory ? &LODData.ClothVertexBuffer : 0,
+						bClothFactory ? DynamicData->ClothingSimData.Find(Section.CorrespondClothAssetIndex) : 0,
+						LocalToCloth,
+						DynamicData->ClothBlendWeight,
+						RevisionNumber,
+						SectionIdx,
+						SkinCacheEntry
+						);
+				}
 			}
 
 			if (bNeedFence)
diff --git a/Engine/Source/Runtime/Engine/Public/GPUSkinCache.h b/Engine/Source/Runtime/Engine/Public/GPUSkinCache.h
--- a/Engine/Source/Runtime/Engine/Public/GPUSkinCache.h
+++ b/Engine/Source/Runtime/Engine/Public/GPUSkinCache.h
@@ -154,7 +154,7 @@ public:
 	FCachedGeometry::Section GetCachedGeometry(FGPUSkinCacheEntry* InOutEntry, uint32 SectionId);
 	void UpdateSkinWeightBuffer(FGPUSkinCacheEntry* Entry);
 
-	void ProcessEntry(
+	bool ProcessEntry(
 		FRHICommandListImmediate& RHICmdList, 
 		FGPUBaseSkinVertexFactory* VertexFactory,
 		FGPUSkinPassthroughVertexFactory* TargetVertexFactory, 
