make non-VT textures with layers a warning rather than error and fix logs

diff --git a/Engine/Source/Runtime/Engine/Private/TextureDerivedDataTask.cpp b/Engine/Source/Runtime/Engine/Private/TextureDerivedDataTask.cpp
--- a/Engine/Source/Runtime/Engine/Private/TextureDerivedDataTask.cpp
+++ b/Engine/Source/Runtime/Engine/Private/TextureDerivedDataTask.cpp
@@ -409,13 +409,16 @@ void FTextureCacheDerivedDataWorker::BuildTexture(bool bReplaceExistingDDC)
 		if (TextureData.Blocks.Num() > 1)
 		{
 			// This warning can happen if user attempts to import a UDIM without VT enabled
-			UE_LOG(LogTexture, Warning, TEXT("Texture %s was imported as UDIM with %d blocks but VirtualTexturing is not enabled, only the 1001 block will be availiable"),
-				*Texture.GetName(), TextureData.Blocks.Num());
+			UE_LOG(LogTexture, Warning, TEXT("Texture %s was imported as UDIM with %d blocks but VirtualTexturing is not enabled, only the first block will be available"),
+				*Texture.GetPathName(), TextureData.Blocks.Num());
 		}
 
 		// No user-facing way to generated multi-layered textures currently, so this should not occur
-		ensureMsgf(TextureData.Layers.Num() == 1, TEXT("Texture %s has %d layers bu VirtualTexturing is not enabled, only layer0 will be availiable"),
-			*Texture.GetName(), TextureData.Blocks.Num());
+		if (TextureData.Layers.Num() > 1)
+		{
+			UE_LOG(LogTexture, Warning, TEXT("Texture %s has %d layers but VirtualTexturing is not enabled, only the first layer will be available"),
+				*Texture.GetPathName(), TextureData.Layers.Num());
+		}
 
 		check(DerivedData->Mips.Num() == 0);
 		DerivedData->SizeX = 0;
