Adding additional check information around AttributeTypes

diff --git a/Engine/Source/Runtime/Engine/Public/Animation/AttributeTypes.h b/Engine/Source/Runtime/Engine/Public/Animation/AttributeTypes.h
--- a/Engine/Source/Runtime/Engine/Public/Animation/AttributeTypes.h
+++ b/Engine/Source/Runtime/Engine/Public/Animation/AttributeTypes.h
@@ -139,7 +139,8 @@ namespace UE
 			{
 				AttributeTypes::Initialize();
 				const int32 Index = AttributeTypes::RegisteredTypes.IndexOfByKey(WeakStruct);
-				checkf(Index != INDEX_NONE, TEXT("Missing operator for custom attribute, type was not registered previously"));
+				ensure(WeakStruct.IsValid());
+				checkf(Index != INDEX_NONE, TEXT("Missing operator for attribute, type %s was not registered previously"), *WeakStruct->GetName());
 				return AttributeTypes::Operators[Index].Get();
 			}
 
