PR #8129: Skip string wrapping if UCanvas::WarpString() is called with a null font. (Contributed by [REDACTED_PERSON])

diff --git a/Engine/Source/Runtime/Engine/Private/UserInterface/Canvas.cpp b/Engine/Source/Runtime/Engine/Private/UserInterface/Canvas.cpp
--- a/Engine/Source/Runtime/Engine/Private/UserInterface/Canvas.cpp
+++ b/Engine/Source/Runtime/Engine/Private/UserInterface/Canvas.cpp
@@ -1262,7 +1262,14 @@ void UCanvas::CanvasStringSize( FTextSizingParameters& Parameters, const TCHAR*
 
 void UCanvas::WrapString( FCanvasWordWrapper& Wrapper, FTextSizingParameters& Parameters, const float InCurX, const TCHAR* const pText, TArray<FWrappedStringElement>& out_Lines, FCanvasWordWrapper::FWrappedLineData* const OutWrappedLineData)
 {
-	Wrapper.Execute(pText, Parameters, out_Lines, OutWrappedLineData);
+	if (Parameters.DrawFont)
+	{
+		Wrapper.Execute(pText, Parameters, out_Lines, OutWrappedLineData);
+	}
+	else
+	{
+		UE_LOG(LogCanvas, Warning, TEXT("UCanvas::WrapString() called with a NULL Font!"));
+	}
 }
 
 void UCanvas::WrapString( FTextSizingParameters& Parameters, const float InCurX, const TCHAR* const pText, TArray<FWrappedStringElement>& out_Lines, FCanvasWordWrapper::FWrappedLineData* const OutWrappedLineData)
