Fix world partition class redirectors to properly support classes from plugins and modules.

diff --git a/Engine/Source/Runtime/Engine/Private/Actor.cpp b/Engine/Source/Runtime/Engine/Private/Actor.cpp
--- a/Engine/Source/Runtime/Engine/Private/Actor.cpp
+++ b/Engine/Source/Runtime/Engine/Private/Actor.cpp
@@ -468,7 +468,7 @@ void AActor::GetExternalActorExtendedAssetRegistryTags(TArray<FAssetRegistryTag>
 	{
 		TUniquePtr<FWorldPartitionActorDesc> ActorDesc(CreateActorDesc());
 
-		const FString ActorMetaDataClass = GetParentNativeClass(GetClass())->GetName();
+		const FString ActorMetaDataClass = GetParentNativeClass(GetClass())->GetPathName();
 		static FName NAME_ActorMetaDataClass(TEXT("ActorMetaDataClass"));
 		OutTags.Add(UObject::FAssetRegistryTag(NAME_ActorMetaDataClass, ActorMetaDataClass, UObject::FAssetRegistryTag::TT_Hidden));
 
diff --git a/Engine/Source/Runtime/Engine/Private/WorldPartition/ActorDescContainer.cpp b/Engine/Source/Runtime/Engine/Private/WorldPartition/ActorDescContainer.cpp
--- a/Engine/Source/Runtime/Engine/Private/WorldPartition/ActorDescContainer.cpp
+++ b/Engine/Source/Runtime/Engine/Private/WorldPartition/ActorDescContainer.cpp
@@ -8,7 +8,7 @@
 #include "WorldPartition/WorldPartitionHandle.h"
 #include "Misc/Base64.h"
 #include "UObject/ObjectSaveContext.h"
-#include "UObject/LinkerLoad.h"
+#include "UObject/CoreRedirects.h"
 #endif
 
 UActorDescContainer::UActorDescContainer(const FObjectInitializer& ObjectInitializer)
@@ -47,23 +47,27 @@ void UActorDescContainer::Initialize(UWorld* InWorld, FName InPackageName)
 
 	auto GetActorDescriptor = [this](const FAssetData& InAssetData) -> TUniquePtr<FWorldPartitionActorDesc>
 	{
-		FName ActorClassName;
+		FString ActorMetaDataClass;
 		static FName NAME_ActorMetaDataClass(TEXT("ActorMetaDataClass"));
-		if (InAssetData.GetTagValue(NAME_ActorMetaDataClass, ActorClassName))
+		if (InAssetData.GetTagValue(NAME_ActorMetaDataClass, ActorMetaDataClass))
 		{
-			// Look for a class redirectors
-			FName NewActorClassName = FLinkerLoad::FindNewNameForClass(ActorClassName, false);
-			if (NewActorClassName != NAME_None)
-			{
-				ActorClassName = NewActorClassName;
-			}
-
 			FString ActorMetaDataStr;
 			static FName NAME_ActorMetaData(TEXT("ActorMetaData"));
 			if (InAssetData.GetTagValue(NAME_ActorMetaData, ActorMetaDataStr))
 			{
+				FString ActorClassName;
+				FString ActorPackageName;
+				if (!ActorMetaDataClass.Split(TEXT("."), &ActorPackageName, &ActorClassName))
+				{
+					ActorClassName = *ActorMetaDataClass;
+				}
+
+				// Look for a class redirectors
+				FCoreRedirectObjectName OldClassName = FCoreRedirectObjectName(*ActorClassName, NAME_None, *ActorPackageName);
+				FCoreRedirectObjectName NewClassName = FCoreRedirects::GetRedirectedName(ECoreRedirectFlags::Type_Class, OldClassName);
+
 				bool bIsValidClass = true;
-				UClass* ActorClass = FindObject<UClass>(ANY_PACKAGE, *ActorClassName.ToString(), true);
+				UClass* ActorClass = FindObject<UClass>(ANY_PACKAGE, *NewClassName.ToString(), true);
 
 				if (!ActorClass)
 				{
