World Partition HLODs - Added automatic texture size selection based on hints
Supports either:
* From texel density
* From expected mesh screen size
* From expected mesh draw distance

Compute best texture size, taking UV map waste into account.
Currently only supported by the Approximate Actors merge method, to be integrated in the ProxyLOD later.

diff --git a/Engine/Plugins/Experimental/MeshModelingToolset/Source/GeometryProcessingAdapters/Private/ApproximateActorsImpl.cpp b/Engine/Plugins/Experimental/MeshModelingToolset/Source/GeometryProcessingAdapters/Private/ApproximateActorsImpl.cpp
--- a/Engine/Plugins/Experimental/MeshModelingToolset/Source/GeometryProcessingAdapters/Private/ApproximateActorsImpl.cpp
+++ b/Engine/Plugins/Experimental/MeshModelingToolset/Source/GeometryProcessingAdapters/Private/ApproximateActorsImpl.cpp
@@ -647,6 +647,46 @@ static TSharedPtr<FApproximationMeshData> GenerateApproximationMesh(
 }
 
 
+static int32 GetMeshTextureSizeFromTargetTexelDensity(const FDynamicMesh3& Mesh, float TargetTexelDensity)
+{
+	const FDynamicMeshUVOverlay* UVOverlay = Mesh.Attributes()->PrimaryUV();
+	double Mesh3DArea = 0;
+	double MeshUVArea = 0;
+	for (int TriangleID : Mesh.TriangleIndicesItr())
+	{
+		// World space area
+		Mesh3DArea += Mesh.GetTriArea(TriangleID);
+
+		FIndex3i UVVertices = UVOverlay->GetTriangle(TriangleID);
+		FTriangle2d TriangleUV = FTriangle2d(
+			(FVector2d)UVOverlay->GetElement(UVVertices.A),
+			(FVector2d)UVOverlay->GetElement(UVVertices.B),
+			(FVector2d)UVOverlay->GetElement(UVVertices.C));
+
+		// UV space area
+		MeshUVArea += TriangleUV.Area();
+	}
+	double TexelRatio = FMath::Sqrt(MeshUVArea / Mesh3DArea) * 100;
+
+	// Compute the perfect texture size that would get us to our texture density
+	// Also compute the nearest power of two sizes (below and above our target)
+	const int32 SizePerfect = FMath::CeilToInt(TargetTexelDensity / TexelRatio);
+	const int32 SizeHi = FMath::RoundUpToPowerOfTwo(SizePerfect);
+	const int32 SizeLo = SizeHi >> 1;
+
+	// Compute the texel density we achieve with these two texture sizes
+	const double TexelDensityLo = SizeLo * TexelRatio;
+	const double TexelDensityHi = SizeHi * TexelRatio;
+
+	// Select best match between low & high res textures.
+	const double TexelDensityLoDiff = TargetTexelDensity - TexelDensityLo;
+	const double TexelDensityHiDiff = TexelDensityHi - TargetTexelDensity;
+	const int32 BestTextureSize = TexelDensityLoDiff < TexelDensityHiDiff ? SizeLo : SizeHi;
+
+	return BestTextureSize;
+}
+
+
 IGeometryProcessing_ApproximateActors::FOptions FApproximateActorsImpl::ConstructOptions(const FMeshApproximationSettings& UseSettings)
 {
 	//
@@ -839,9 +879,28 @@ void FApproximateActorsImpl::GenerateApproximationForActorSet(const TArray<AActo
 	FDynamicMesh3 FinalMesh = MoveTemp(ApproximationMeshData->Mesh);
 	FMeshTangentsd FinalMeshTangents = MoveTemp(ApproximationMeshData->Tangents);
 
+	FOptions OverridenOptions = Options;
+
+	// evaluate required texture size if needed
+	if (Options.TextureSizePolicy == ETextureSizePolicy::TexelDensity)
+	{
+		const int32 MaxTextureSize = 8192;
+		const int32 BestTextureSize = GetMeshTextureSizeFromTargetTexelDensity(FinalMesh, Options.MeshTexelDensity);
+
+		if (BestTextureSize > MaxTextureSize)
+		{
+			UE_LOG(LogApproximateActors, Warning, TEXT("Mesh would require %dx%d textures, clamping down to maximum (%dx%d)"), BestTextureSize, BestTextureSize, MaxTextureSize, MaxTextureSize);
+			OverridenOptions.TextureImageSize = MaxTextureSize;
+		}
+		else
+		{
+			OverridenOptions.TextureImageSize = BestTextureSize;
+		}
+	}
+	
 	// bake textures for Actor
 	FGeneratedResultTextures GeneratedTextures;
-	BakeTexturesFromPhotoCapture(SceneCapture, Options,
+	BakeTexturesFromPhotoCapture(SceneCapture, OverridenOptions,
 		GeneratedTextures,
 		&FinalMesh, &FinalMeshTangents);
 
diff --git a/Engine/Source/Developer/GeometryProcessingInterfaces/Public/GeometryProcessingInterfaces/ApproximateActors.h b/Engine/Source/Developer/GeometryProcessingInterfaces/Public/GeometryProcessingInterfaces/ApproximateActors.h
--- a/Engine/Source/Developer/GeometryProcessingInterfaces/Public/GeometryProcessingInterfaces/ApproximateActors.h
+++ b/Engine/Source/Developer/GeometryProcessingInterfaces/Public/GeometryProcessingInterfaces/ApproximateActors.h
@@ -58,6 +58,12 @@ public:
 		GeometricTolerance = 2
 	};
 
+	enum class ETextureSizePolicy : uint8
+	{
+		TextureSize = 0,
+		TexelDensity = 1
+	};
+
 
 	/**
 	 * Input options to Actor Approximation process
@@ -134,7 +140,9 @@ public:
 		FName PackedMRSTexParamName = FName("PackedMRS");
 
 		// output texture options
+		ETextureSizePolicy TextureSizePolicy = ETextureSizePolicy::TextureSize;
 		int32 TextureImageSize = 1024;
+		float MeshTexelDensity = 0;
 
 		// supersampling parameter
 		int32 AntiAliasMultiSampling = 0;
diff --git a/Engine/Source/Developer/MaterialUtilities/Private/MaterialUtilities.cpp b/Engine/Source/Developer/MaterialUtilities/Private/MaterialUtilities.cpp
--- a/Engine/Source/Developer/MaterialUtilities/Private/MaterialUtilities.cpp
+++ b/Engine/Source/Developer/MaterialUtilities/Private/MaterialUtilities.cpp
@@ -3073,3 +3073,37 @@ void FMaterialUtilities::GeneratedBinnedTextureSquares(const FVector2D Destinati
 		Box = Texture.Area;
 	}
 }
+
+float FMaterialUtilities::ComputeRequiredTexelDensityFromScreenSize(const float InScreenSize, float InWorldSpaceRadius)
+{
+	static const float ScreenX = 1920;
+
+	float WorldSizeCM = InWorldSpaceRadius * 2;
+	float WorldSizeMeter = WorldSizeCM / 100;
+
+	float ScreenSizePercent = InScreenSize;
+	float ScreenSizePixel = ScreenSizePercent * ScreenX;
+
+	float TexelDensityPerMeter = ScreenSizePixel / WorldSizeMeter;
+
+	return TexelDensityPerMeter;
+}
+
+float FMaterialUtilities::ComputeRequiredTexelDensityFromDrawDistance(const float InDrawDistance, float InWorldSpaceRadius)
+{
+	// Generate a projection matrix.
+	static const float ScreenX = 1920;
+	static const float ScreenY = 1080;
+	static const float HalfFOVRad = FMath::DegreesToRadians(45.0f);
+	static const FMatrix ProjectionMatrix = FPerspectiveMatrix(HalfFOVRad, ScreenX, ScreenY, 0.01f);
+
+	float WorldSizeCM = InWorldSpaceRadius * 2;
+	float WorldSizeMeter = WorldSizeCM / 100;
+
+	float ScreenSizePercent = ComputeBoundsScreenSize(FVector::ZeroVector, InWorldSpaceRadius, FVector(0.0f, 0.0f, InDrawDistance), ProjectionMatrix);
+	float ScreenSizePixel = ScreenSizePercent * ScreenX;
+
+	float TexelDensityPerMeter = ScreenSizePixel / WorldSizeMeter;
+
+	return TexelDensityPerMeter;
+}
diff --git a/Engine/Source/Developer/MaterialUtilities/Public/MaterialUtilities.h b/Engine/Source/Developer/MaterialUtilities/Public/MaterialUtilities.h
--- a/Engine/Source/Developer/MaterialUtilities/Public/MaterialUtilities.h
+++ b/Engine/Source/Developer/MaterialUtilities/Public/MaterialUtilities.h
@@ -521,6 +521,23 @@ public:
 
 	/** Creates a proxy material and the required texture assets */
 	static UMaterialInterface* CreateProxyMaterialAndTextures(const FString& PackageName, const FString& AssetName, const FBakeOutput& BakeOutput, const FMeshData& MeshData, const FMaterialData& MaterialData, UMaterialOptions* Options);
+
+	/** Compute the required texel density needed to properly represent an object/objects covering the provided world space radius
+	* 
+	* @param InScreenSize		Screen size at which the objects are expected to be be shown
+	* @param InWorldSpaceRadius	World space radius of the objects
+	* @return					The required texel density per METER
+	*/
+	static float ComputeRequiredTexelDensityFromScreenSize(const float InScreenSize, float InWorldSpaceRadius);
+
+	/** Compute the required texel density needed to properly represent an object/objects covering the provided world space radius
+	*
+	* @param InDrawDistance		Draw distance at which the objects are expected to be shown
+	* @param InWorldSpaceRadius	World space radius of the objects
+	* @return					The required texel density per METER
+	*/
+	static float ComputeRequiredTexelDensityFromDrawDistance(const float InDrawDistance, float InWorldSpaceRadius);
+
 private:
 	
 	/**
diff --git a/Engine/Source/Developer/WorldPartitionHLODUtilities/Private/HLODBuilderMeshApproximate.cpp b/Engine/Source/Developer/WorldPartitionHLODUtilities/Private/HLODBuilderMeshApproximate.cpp
--- a/Engine/Source/Developer/WorldPartitionHLODUtilities/Private/HLODBuilderMeshApproximate.cpp
+++ b/Engine/Source/Developer/WorldPartitionHLODUtilities/Private/HLODBuilderMeshApproximate.cpp
@@ -10,6 +10,7 @@
 #include "Engine/StaticMesh.h"
 #include "Materials/Material.h"
 #include "Materials/MaterialInstance.h"
+#include "MaterialUtilities.h"
 
 #include "Modules/ModuleManager.h"
 #include "IGeometryProcessingInterfacesModule.h"
@@ -48,14 +49,45 @@ TArray<UPrimitiveComponent*> FHLODBuilder_MeshApproximate::CreateComponents(AWor
 	Options.bUsePackedMRS = true;
 	Options.PackedMRSTexParamName = FName("PackedTexture");
 
-	// Scale capture size based on grid promotion (higher z grid index -> higher area covered)
-	uint64 GridIndexX = 0;
-	uint64 GridIndexY = 0;
-	uint64 GridIndexZ = 0;
-	InHLODActor->GetGridIndices(GridIndexX, GridIndexY, GridIndexZ);
-	uint64 TexSizeMultiplier = (uint64)FMath::Pow(2.0f, GridIndexZ);
-	Options.TextureImageSize *= TexSizeMultiplier;
-	Options.RenderCaptureImageSize *= TexSizeMultiplier;
+	// Gather bounds of the input components
+	auto GetActorsBounds = [&]() -> FBoxSphereBounds
+	{
+		FBoxSphereBounds Bounds;
+		bool bFirst = true;
+
+		for (UPrimitiveComponent* Component : InSubComponents)
+		{
+			FBoxSphereBounds ComponentBounds = Component->Bounds;
+			Bounds = bFirst ? ComponentBounds : Bounds + ComponentBounds;
+			bFirst = false;
+		}
+
+		return Bounds;
+	};
+	
+	// Compute texel density if needed, depending on the TextureSizingType setting
+	ETextureSizingType TextureSizingType = UseSettings.MaterialSettings.TextureSizingType;
+	float TexelDensityPerMeter = 0.0f;
+
+	IGeometryProcessing_ApproximateActors::ETextureSizePolicy TextureSizePolicy = IGeometryProcessing_ApproximateActors::ETextureSizePolicy::TextureSize;
+	if (TextureSizingType == ETextureSizingType::TextureSizingType_AutomaticFromTexelDensity)
+	{
+		TexelDensityPerMeter = UseSettings.MaterialSettings.TargetTexelDensityPerMeter;
+		TextureSizePolicy = IGeometryProcessing_ApproximateActors::ETextureSizePolicy::TexelDensity;
+	}
+	else if (TextureSizingType == ETextureSizingType::TextureSizingType_AutomaticFromMeshScreenSize)
+	{
+		TexelDensityPerMeter = FMaterialUtilities::ComputeRequiredTexelDensityFromScreenSize(UseSettings.MaterialSettings.MeshMaxScreenSizePercent, GetActorsBounds().SphereRadius);
+		TextureSizePolicy = IGeometryProcessing_ApproximateActors::ETextureSizePolicy::TexelDensity;
+	}
+	else if (TextureSizingType == ETextureSizingType::TextureSizingType_AutomaticFromMeshDrawDistance)
+	{
+		TexelDensityPerMeter = FMaterialUtilities::ComputeRequiredTexelDensityFromDrawDistance(UseSettings.MaterialSettings.MeshMinDrawDistance, GetActorsBounds().SphereRadius);
+		TextureSizePolicy = IGeometryProcessing_ApproximateActors::ETextureSizePolicy::TexelDensity;
+	}
+
+	Options.MeshTexelDensity = TexelDensityPerMeter;
+	Options.TextureSizePolicy = TextureSizePolicy;
 
 	// run actor approximation computation
 	IGeometryProcessing_ApproximateActors::FResults Results;
diff --git a/Engine/Source/Developer/WorldPartitionHLODUtilities/WorldPartitionHLODUtilities.Build.cs b/Engine/Source/Developer/WorldPartitionHLODUtilities/WorldPartitionHLODUtilities.Build.cs
--- a/Engine/Source/Developer/WorldPartitionHLODUtilities/WorldPartitionHLODUtilities.Build.cs
+++ b/Engine/Source/Developer/WorldPartitionHLODUtilities/WorldPartitionHLODUtilities.Build.cs
@@ -20,6 +20,7 @@ public class WorldPartitionHLODUtilities : ModuleRules
 			{
 				"EditorFramework",
 				"Engine",
+				"MaterialUtilities",
 				"MeshDescription",
 				"StaticMeshDescription",
 				"UnrealEd",
diff --git a/Engine/Source/Editor/MergeActors/MergeActors.Build.cs b/Engine/Source/Editor/MergeActors/MergeActors.Build.cs
--- a/Engine/Source/Editor/MergeActors/MergeActors.Build.cs
+++ b/Engine/Source/Editor/MergeActors/MergeActors.Build.cs
@@ -35,6 +35,7 @@ public class MergeActors : ModuleRules
 				"CoreUObject",
 				"Engine",
                 "InputCore",
+				"MaterialUtilities",
 				"Slate",
 				"SlateCore",
                 "EditorStyle",
diff --git a/Engine/Source/Editor/MergeActors/Private/MeshApproximationTool/MeshApproximationTool.cpp b/Engine/Source/Editor/MergeActors/Private/MeshApproximationTool/MeshApproximationTool.cpp
--- a/Engine/Source/Editor/MergeActors/Private/MeshApproximationTool/MeshApproximationTool.cpp
+++ b/Engine/Source/Editor/MergeActors/Private/MeshApproximationTool/MeshApproximationTool.cpp
@@ -18,6 +18,7 @@
 #include "Misc/MessageDialog.h"
 #include "MeshApproximationTool/SMeshApproximationDialog.h"
 #include "MeshUtilities.h"
+#include "MaterialUtilities.h"
 #include "IContentBrowserSingleton.h"
 #include "ContentBrowserModule.h"
 #include "AssetRegistryModule.h"
@@ -157,6 +158,53 @@ bool FMeshApproximationTool::RunMerge(const FString& PackageName, const TArray<T
 		return false;
 	}
 
+    // Gather bounds of the input components
+	auto GetActorsBounds = [&]() -> FBoxSphereBounds
+	{
+		FBoxSphereBounds Bounds;
+		bool bFirst = true;
+
+		for (auto& Component : SelectedComponents)
+		{
+			if (!Component.IsValid() || !Component.Get()->bShouldIncorporate)
+			{
+				continue;
+			}
+
+			if (UPrimitiveComponent* PrimitiveComponent = Cast<UPrimitiveComponent>(Component.Get()->PrimComponent))
+			{
+				// Append bounds of visible components only
+				FBoxSphereBounds ComponentBounds = PrimitiveComponent->Bounds;
+				Bounds = bFirst ? ComponentBounds : Bounds + ComponentBounds;
+				bFirst = false;
+			}
+		}
+
+		return Bounds;
+	};
+	
+
+	// Compute texel density if needed, depending on the TextureSizingType setting
+	ETextureSizingType TextureSizingType = SettingsObject->Settings.MaterialSettings.TextureSizingType;
+	float TexelDensityPerMeter = 0.0f;
+
+	IGeometryProcessing_ApproximateActors::ETextureSizePolicy TextureSizePolicy = IGeometryProcessing_ApproximateActors::ETextureSizePolicy::TextureSize;
+	if (TextureSizingType == ETextureSizingType::TextureSizingType_AutomaticFromTexelDensity)
+	{
+		TexelDensityPerMeter = SettingsObject->Settings.MaterialSettings.TargetTexelDensityPerMeter;
+		TextureSizePolicy = IGeometryProcessing_ApproximateActors::ETextureSizePolicy::TexelDensity;
+	}
+	else if (TextureSizingType == ETextureSizingType::TextureSizingType_AutomaticFromMeshScreenSize)
+	{
+		TexelDensityPerMeter = FMaterialUtilities::ComputeRequiredTexelDensityFromScreenSize(SettingsObject->Settings.MaterialSettings.MeshMaxScreenSizePercent, GetActorsBounds().SphereRadius);
+		TextureSizePolicy = IGeometryProcessing_ApproximateActors::ETextureSizePolicy::TexelDensity;
+	}
+	else if (TextureSizingType == ETextureSizingType::TextureSizingType_AutomaticFromMeshDrawDistance)
+	{
+		TexelDensityPerMeter = FMaterialUtilities::ComputeRequiredTexelDensityFromScreenSize(SettingsObject->Settings.MaterialSettings.MeshMinDrawDistance, GetActorsBounds().SphereRadius);
+		TextureSizePolicy = IGeometryProcessing_ApproximateActors::ETextureSizePolicy::TexelDensity;
+	}
+
 	const FMeshApproximationSettings& UseSettings = SettingsObject->Settings;
 
 	IGeometryProcessingInterfacesModule& GeomProcInterfaces = FModuleManager::Get().LoadModuleChecked<IGeometryProcessingInterfacesModule>("GeometryProcessingInterfaces");
@@ -168,6 +216,9 @@ bool FMeshApproximationTool::RunMerge(const FString& PackageName, const TArray<T
 
 	IGeometryProcessing_ApproximateActors::FOptions Options = ApproxActorsAPI->ConstructOptions(UseSettings);
 	Options.BasePackagePath = PackageName;
+
+	Options.TextureSizePolicy = TextureSizePolicy;
+	Options.MeshTexelDensity = TexelDensityPerMeter;
 	
 	// run actor approximation computation
 	IGeometryProcessing_ApproximateActors::FResults Results;
diff --git a/Engine/Source/Runtime/Engine/Classes/Engine/MaterialMerging.h b/Engine/Source/Runtime/Engine/Classes/Engine/MaterialMerging.h
--- a/Engine/Source/Runtime/Engine/Classes/Engine/MaterialMerging.h
+++ b/Engine/Source/Runtime/Engine/Classes/Engine/MaterialMerging.h
@@ -14,6 +14,9 @@ enum ETextureSizingType
 	TextureSizingType_UseAutomaticBiasedSizes UMETA(DisplayName = "Use automatically biased texture sizes based on TextureSize"),
 	TextureSizingType_UseManualOverrideTextureSize UMETA(DisplayName = "Use per property manually overriden texture sizes"),
 	TextureSizingType_UseSimplygonAutomaticSizing UMETA(DisplayName = "Use Simplygon's automatic texture sizing"),
+	TextureSizingType_AutomaticFromTexelDensity UMETA(DisplayName = "Automatic - From Texel Density"),
+	TextureSizingType_AutomaticFromMeshScreenSize UMETA(DisplayName = "Automatic - From Mesh Screen Size"),
+	TextureSizingType_AutomaticFromMeshDrawDistance UMETA(DisplayName = "Automatic - From Mesh Draw Distance"),
 	TextureSizingType_MAX,
 };
 
@@ -29,12 +32,28 @@ struct FMaterialProxySettings
 {
 	GENERATED_USTRUCT_BODY()
 
+	// Method that should be used to generate the sizes of the output textures
+	UPROPERTY(Category = Material, BlueprintReadWrite, EditAnywhere)
+	TEnumAsByte<ETextureSizingType> TextureSizingType;
+
 	// Size of generated BaseColor map
-	UPROPERTY(Category = Material, BlueprintReadWrite, EditAnywhere, meta =(ClampMin = "1", UIMin = "1"))
+	UPROPERTY(Category = Material, BlueprintReadWrite, EditAnywhere, meta =(ClampMin = "1", UIMin = "1", EditConditionHides, EditCondition = "TextureSizingType == ETextureSizingType::TextureSizingType_UseSingleTextureSize || TextureSizingType == ETextureSizingType::TextureSizingType_UseAutomaticBiasedSizes"))
 	FIntPoint TextureSize;
+
+	// Target texel density
+	UPROPERTY(Category = Material, BlueprintReadWrite, EditAnywhere, meta = (ClampMin = "0.1", ClampMax = "1024", EditConditionHides, EditCondition = "TextureSizingType == ETextureSizingType::TextureSizingType_AutomaticFromTexelDensity"))
+	float TargetTexelDensityPerMeter;
+
+	// Expected maximum screen size for the mesh
+	UPROPERTY(Category = Material, BlueprintReadWrite, EditAnywhere, meta = (ClampMin = "0.01", ClampMax = "1.0", EditConditionHides, EditCondition = "TextureSizingType == ETextureSizingType::TextureSizingType_AutomaticFromMeshScreenSize"))
+	float MeshMaxScreenSizePercent;
+
+	// Expected minimum distance at which the mesh will be rendered
+	UPROPERTY(Category = Material, BlueprintReadWrite, EditAnywhere, meta = (ClampMin = "0", EditConditionHides, EditCondition = "TextureSizingType == ETextureSizingType::TextureSizingType_AutomaticFromMeshDrawDistance"))
+	float MeshMinDrawDistance;
 	
 	// Gutter space to take into account 
-	UPROPERTY(Category = Material, BlueprintReadWrite, AdvancedDisplay, EditAnywhere, meta=(DisplayAfter="TextureSizingType"))
+	UPROPERTY(Category = Material, BlueprintReadWrite, AdvancedDisplay, EditAnywhere)
 	float GutterSpace;
 
 	// Constant value to use for the Metallic property
@@ -65,10 +84,6 @@ struct FMaterialProxySettings
 	UPROPERTY(Category = Material, BlueprintReadWrite, EditAnywhere, meta = (DisplayAfter="bAmbientOcclusionMap", ClampMin = "0", ClampMax = "1", UIMin = "0", UIMax = "1", editcondition = "!bAmbientOcclusionMap"))
 	float AmbientOcclusionConstant;
 
-	// Method that should be used to generate the sizes of the output textures
-	UPROPERTY(Category = Material, BlueprintReadWrite, EditAnywhere)
-	TEnumAsByte<ETextureSizingType> TextureSizingType;
-
 	UPROPERTY()
 	TEnumAsByte<EMaterialMergeType> MaterialMergeType;
 
@@ -165,7 +180,11 @@ struct FMaterialProxySettings
 	FIntPoint AmbientOcclusionTextureSize;
 
 	FMaterialProxySettings()
-		: TextureSize(1024, 1024)
+		: TextureSizingType(TextureSizingType_UseSingleTextureSize)
+		, TextureSize(1024, 1024)
+		, TargetTexelDensityPerMeter(5.0f)
+		, MeshMaxScreenSizePercent(0.5f)
+		, MeshMinDrawDistance(10000.0f)
 		, GutterSpace(4.0f)
 		, MetallicConstant(0.0f)
 		, RoughnessConstant(0.5f)
@@ -174,7 +193,6 @@ struct FMaterialProxySettings
 		, OpacityConstant(1.0f)
 		, OpacityMaskConstant(1.0f)
 		, AmbientOcclusionConstant(1.0f)
-		, TextureSizingType(TextureSizingType_UseSingleTextureSize)
 		, MaterialMergeType(EMaterialMergeType::MaterialMergeType_Default)
 		, BlendMode(BLEND_Opaque)
 		, bAllowTwoSidedMaterial(true)
@@ -206,6 +224,9 @@ struct FMaterialProxySettings
 	{
 		return TextureSize == Other.TextureSize
 			&& TextureSizingType == Other.TextureSizingType
+			&& TargetTexelDensityPerMeter == Other.TargetTexelDensityPerMeter
+			&& MeshMaxScreenSizePercent == Other.MeshMaxScreenSizePercent
+			&& MeshMinDrawDistance == Other.MeshMinDrawDistance
 			&& GutterSpace == Other.GutterSpace
 			&& bNormalMap == Other.bNormalMap
 			&& bTangentMap == Other.bTangentMap
diff --git a/Engine/Source/Runtime/GeometryCore/Public/TriangleTypes.h b/Engine/Source/Runtime/GeometryCore/Public/TriangleTypes.h
--- a/Engine/Source/Runtime/GeometryCore/Public/TriangleTypes.h
+++ b/Engine/Source/Runtime/GeometryCore/Public/TriangleTypes.h
@@ -82,6 +82,12 @@ struct TTriangle2
 	{
 		return SignedArea(V[0], V[1], V[2]);
 	}
+	
+	/** @return unsigned area of the triangle */
+	RealType Area() const
+	{
+		return TMathUtil<RealType>::Abs(SignedArea());
+	}
 
 
 	/**
