Fixed broken velocity vectors

diff --git a/Engine/Plugins/Experimental/GPULightmass/Source/GPULightmass/Private/LightmapRenderer.cpp b/Engine/Plugins/Experimental/GPULightmass/Source/GPULightmass/Private/LightmapRenderer.cpp
--- a/Engine/Plugins/Experimental/GPULightmass/Source/GPULightmass/Private/LightmapRenderer.cpp
+++ b/Engine/Plugins/Experimental/GPULightmass/Source/GPULightmass/Private/LightmapRenderer.cpp
@@ -336,7 +336,8 @@ void FCachedRayTracingSceneData::SetupViewUniformBufferFromSceneRenderState(FSce
 					InstanceIndex, /* Primitive Id */
 					PrimitiveUniformShaderParameters.LocalToWorld,
 					PrimitiveUniformShaderParameters.PreviousLocalToWorld,
-					INVALID_LAST_UPDATE_FRAME
+					INVALID_LAST_UPDATE_FRAME,
+					false /* Has Previous Transform */
 				);
 
 				for (int32 LODIndex = 0; LODIndex < Instance.LODLightmapRenderStates.Num(); LODIndex++)
@@ -415,7 +416,8 @@ void FCachedRayTracingSceneData::SetupViewUniformBufferFromSceneRenderState(FSce
 						PrimitiveId,
 						InstanceGroup.LocalToWorld,
 						InstanceGroup.LocalToWorld,
-						INVALID_LAST_UPDATE_FRAME
+						INVALID_LAST_UPDATE_FRAME,
+						false /* Has Previous Transform */
 					)
 				);
 			}
@@ -479,7 +481,8 @@ void FCachedRayTracingSceneData::SetupViewUniformBufferFromSceneRenderState(FSce
 				PrimitiveId,
 				PrimitiveUniformShaderParameters.LocalToWorld,
 				PrimitiveUniformShaderParameters.PreviousLocalToWorld,
-				INVALID_LAST_UPDATE_FRAME
+				INVALID_LAST_UPDATE_FRAME,
+				false /* Has Previous Transform */
 			));
 
 			PrimitiveSceneData.Add(FPrimitiveSceneShaderData(PrimitiveUniformShaderParameters));
diff --git a/Engine/Source/Runtime/Engine/Private/InstanceUniformShaderParameters.cpp b/Engine/Source/Runtime/Engine/Private/InstanceUniformShaderParameters.cpp
--- a/Engine/Source/Runtime/Engine/Private/InstanceUniformShaderParameters.cpp
+++ b/Engine/Source/Runtime/Engine/Private/InstanceUniformShaderParameters.cpp
@@ -9,11 +9,12 @@ FInstanceSceneShaderData::FInstanceSceneShaderData(
 	uint32 PrimitiveId,
 	const FRenderTransform& PrimitiveLocalToWorld,
 	const FRenderTransform& PrimitivePrevLocalToWorld,
-	uint32 LastUpdateFrame
+	uint32 LastUpdateFrame,
+	bool bHasPreviousTransform
 )
 : Data(InPlace, NoInit)
 {
-	Setup(Instance, PrimitiveId, PrimitiveLocalToWorld, PrimitivePrevLocalToWorld, LastUpdateFrame);
+	Setup(Instance, PrimitiveId, PrimitiveLocalToWorld, PrimitivePrevLocalToWorld, LastUpdateFrame, bHasPreviousTransform);
 }
 
 void FInstanceSceneShaderData::Setup(
@@ -21,7 +22,8 @@ void FInstanceSceneShaderData::Setup(
 	uint32 PrimitiveId,
 	const FRenderTransform& PrimitiveToWorld,
 	const FRenderTransform& PrevPrimitiveToWorld,
-	uint32 LastUpdateFrame
+	uint32 LastUpdateFrame,
+	bool bHasPreviousTransform
 )
 {
 	// Note: layout must match GetInstanceData in SceneData.ush
@@ -30,7 +32,11 @@ void FInstanceSceneShaderData::Setup(
 	// This is a read-only setting that will cause all shaders to recompile if changed.
 
 	FRenderTransform LocalToWorld		= Instance.LocalToPrimitive * PrimitiveToWorld;
-	FRenderTransform PrevLocalToWorld	= Instance.PrevLocalToPrimitive * PrevPrimitiveToWorld;
+
+	// TODO: Remove
+	const FRenderTransform& PrevLocalToPrimitive = bHasPreviousTransform ? Instance.PrevLocalToPrimitive : Instance.LocalToPrimitive;
+	FRenderTransform PrevLocalToWorld = PrevLocalToPrimitive * PrevPrimitiveToWorld;
+	//FRenderTransform PrevLocalToWorld	= Instance.PrevLocalToPrimitive * PrevPrimitiveToWorld;
 
 	// Remove shear
 	LocalToWorld.Orthonormalize();
@@ -91,7 +97,8 @@ ENGINE_API const FInstanceSceneShaderData& GetDummyInstanceSceneShaderData()
 		0xFFFFFFFFu, /* Primitive Id */
 		FRenderTransform::Identity, /* Primitive LocalToWorld */
 		FRenderTransform::Identity,  /* Primitive PrevLocalToWorld */
-		INVALID_LAST_UPDATE_FRAME
+		INVALID_LAST_UPDATE_FRAME,
+		false /* Has Previous Transform */
 	);
 	return DummyShaderData;
 }
diff --git a/Engine/Source/Runtime/Engine/Public/InstanceUniformShaderParameters.h b/Engine/Source/Runtime/Engine/Public/InstanceUniformShaderParameters.h
--- a/Engine/Source/Runtime/Engine/Public/InstanceUniformShaderParameters.h
+++ b/Engine/Source/Runtime/Engine/Public/InstanceUniformShaderParameters.h
@@ -79,6 +79,8 @@ FORCEINLINE FPrimitiveInstance ConstructPrimitiveInstance(
 )
 {
 	FPrimitiveInstance Result;
+	Result.LocalToPrimitive.SetIdentity();
+	Result.PrevLocalToPrimitive.SetIdentity();
 	Result.LightMapAndShadowMapUVBias			= LightMapAndShadowMapUVBias;
 	Result.LocalBounds							= LocalBounds;
 	Result.NaniteHierarchyOffset				= NaniteHierarchyOffset;
@@ -110,7 +112,8 @@ struct FInstanceSceneShaderData
 			0, /* Primitive Id */
 			FRenderTransform::Identity,  /* LocalToWorld */
 			FRenderTransform::Identity,  /* PrevLocalToWorld */
-			INVALID_LAST_UPDATE_FRAME
+			INVALID_LAST_UPDATE_FRAME,
+			false /* Has Previous Transform */
 		);
 	}
 
@@ -119,7 +122,8 @@ struct FInstanceSceneShaderData
 		uint32 PrimitiveId,
 		const FRenderTransform& PrimitiveLocalToWorld,
 		const FRenderTransform& PrimitivePrevLocalToWorld,
-		uint32 LastUpdateFrame
+		uint32 LastUpdateFrame,
+		bool bHasPreviousTransform
 	);
 
 	ENGINE_API void Setup(
@@ -127,7 +131,8 @@ struct FInstanceSceneShaderData
 		uint32 PrimitiveId,
 		const FRenderTransform& PrimitiveLocalToWorld,
 		const FRenderTransform& PrimitivePrevLocalToWorld,
-		uint32 LastUpdateFrame
+		uint32 LastUpdateFrame,
+		bool bHasPreviousTransform
 	);
 };
 
diff --git a/Engine/Source/Runtime/Renderer/Private/GPUScene.cpp b/Engine/Source/Runtime/Renderer/Private/GPUScene.cpp
--- a/Engine/Source/Runtime/Renderer/Private/GPUScene.cpp
+++ b/Engine/Source/Runtime/Renderer/Private/GPUScene.cpp
@@ -171,6 +171,8 @@ struct FInstanceUploadInfo
 	FRenderTransform PrevPrimitiveToWorld;
 	int32 PrimitiveID = INDEX_NONE;
 	uint32 LastUpdateSceneFrameNumber = ~uint32(0);
+
+	bool bHasPreviousTransforms = false;
 };
 
 /**
@@ -256,6 +258,8 @@ struct FUploadDataSourceAdapterScenePrimitives
 			InstanceUploadInfo.LastUpdateSceneFrameNumber = SceneFrameNumber;
 			InstanceUploadInfo.PrimitiveID = PrimitiveID;
 			InstanceUploadInfo.PrimitiveToWorld = PrimitiveSceneProxy->GetLocalToWorld();
+			InstanceUploadInfo.bHasPreviousTransforms = PrimitiveSceneProxy->HasPrevInstanceTransforms();
+
 			{
 				bool bHasPrecomputedVolumetricLightmap{};
 				bool bOutputVelocity{};
@@ -754,7 +758,8 @@ void FGPUScene::UpdateInternal(FRDGBuilder& GraphBuilder, FScene& Scene)
 								InstanceDataGPU.PrimitiveId,
 								PrimitiveToWorld,
 								PrevPrimitiveToWorld,
-								SceneFrameNumber
+								SceneFrameNumber,
+								PrimitiveSceneProxy->HasPrevInstanceTransforms()
 							);
 
 							// unpack again, just as on GPU...
@@ -1221,7 +1226,8 @@ void FGPUScene::UploadGeneral(FRHICommandListImmediate& RHICmdList, FScene *Scen
 											UploadInfo.PrimitiveID,
 											UploadInfo.PrimitiveToWorld,
 											UploadInfo.PrevPrimitiveToWorld,
-											UploadInfo.LastUpdateSceneFrameNumber
+											UploadInfo.LastUpdateSceneFrameNumber,
+											UploadInfo.bHasPreviousTransforms
 										);
 
 										void* DstRefs[FInstanceSceneShaderData::InstanceDataStrideInFloat4s];
diff --git a/Engine/Source/Runtime/Renderer/Private/Renderer.cpp b/Engine/Source/Runtime/Renderer/Private/Renderer.cpp
--- a/Engine/Source/Runtime/Renderer/Private/Renderer.cpp
+++ b/Engine/Source/Runtime/Renderer/Private/Renderer.cpp
@@ -184,7 +184,8 @@ void FRendererModule::DrawTileMesh(FCanvasRenderContext& RenderContext, FMeshPas
 					0, /* Primitive Id */
 					PrimitiveParams.LocalToWorld,
 					PrimitiveParams.PreviousLocalToWorld,
-					INVALID_LAST_UPDATE_FRAME
+					INVALID_LAST_UPDATE_FRAME,
+					false /* Has Previous Transform */
 				);
 
 				// Set up the parameters for the LightmapSceneData from the given LCI data 
