WorldPartition: detect slow/critical performance (when performance is critical block streaming until it completes)

- In order to block StreamingGrid needs to be tagged so that we don't block on all grids (ex: HLOD grid)
- Block can be disabled with cvar (wp.Runtime.ToggleBlockOnSlowStreaming)

diff --git a/Engine/Source/Editor/UnrealEd/Private/EditorEngine.cpp b/Engine/Source/Editor/UnrealEd/Private/EditorEngine.cpp
--- a/Engine/Source/Editor/UnrealEd/Private/EditorEngine.cpp
+++ b/Engine/Source/Editor/UnrealEd/Private/EditorEngine.cpp
@@ -1771,6 +1771,13 @@ void UEditorEngine::Tick( float DeltaSeconds, bool bIdleMode )
 					bAWorldTicked = true;
 					TickType = LEVELTICK_All;
 
+					// Block on async loading if requested.
+					if (PlayWorld->bRequestedBlockOnAsyncLoading)
+					{
+						BlockTillLevelStreamingCompleted(PlayWorld);
+						PlayWorld->bRequestedBlockOnAsyncLoading = false;
+					}
+
 					if (!bFirstTick)
 					{
 						// Update sky light first because sky diffuse will be visible in reflection capture indirect specular
diff --git a/Engine/Source/Runtime/Engine/Classes/Engine/World.h b/Engine/Source/Runtime/Engine/Classes/Engine/World.h
--- a/Engine/Source/Runtime/Engine/Classes/Engine/World.h
+++ b/Engine/Source/Runtime/Engine/Classes/Engine/World.h
@@ -1055,6 +1055,9 @@ public:
 	/** Examine all streaming levels and determine which ones should be considered. */
 	void PopulateStreamingLevelsToConsider();
 
+	/** Whether the world is currently in a BlockTillLevelStreamingCompleted() call */
+	bool GetIsInBlockTillLevelStreamingCompleted() const { return bIsInBlockTillLevelStreamingCompleted; }
+
 	/** Prefix we used to rename streaming levels, non empty in PIE and standalone preview */
 	UPROPERTY()
 	FString										StreamingLevelsPrefix;
@@ -1175,7 +1178,7 @@ public:
 	
 	/** Whether it was requested that the engine bring up a loading screen and block on async loading. */   
 	uint8 bRequestedBlockOnAsyncLoading:1;
-
+		
 	/** Whether actors have been initialized for play */
 	uint8 bActorsInitialized:1;
 
@@ -1222,6 +1225,8 @@ public:
 	uint8 bAreConstraintsDirty:1;
 
 private:
+	/** Whether the world is currently in a BlockTillLevelStreamingCompleted() call */
+	uint8 bIsInBlockTillLevelStreamingCompleted : 1;
 
 	/** Whether the render scene for this World should be created with HitProxies or not */
 	uint8 bRequiresHitProxies:1;
diff --git a/Engine/Source/Runtime/Engine/Private/World.cpp b/Engine/Source/Runtime/Engine/Private/World.cpp
--- a/Engine/Source/Runtime/Engine/Private/World.cpp
+++ b/Engine/Source/Runtime/Engine/Private/World.cpp
@@ -12,6 +12,7 @@
 #include "Logging/LogScopedVerbosityOverride.h"
 #include "Stats/StatsMisc.h"
 #include "Misc/ScopedSlowTask.h"
+#include "Misc/ScopeExit.h"
 #include "Misc/CoreDelegates.h"
 #include "Modules/ModuleManager.h"
 #include "UObject/UObjectIterator.h"
@@ -3596,8 +3597,19 @@ void FStreamingLevelsToConsider::AddReferencedObjects(UObject* InThis, FReferenc
 
 void UWorld::BlockTillLevelStreamingCompleted()
 {
-	const double StartTime = FPlatformTime::Seconds();
+	if (bIsInBlockTillLevelStreamingCompleted)
+	{
+		return;
+	}
 
+	bIsInBlockTillLevelStreamingCompleted = true;
+	ON_SCOPE_EXIT
+	{
+		bIsInBlockTillLevelStreamingCompleted = false;
+	};
+
+	const double StartTime = FPlatformTime::Seconds();
+	
 	bool bWorkToDo = false;
 	do
 	{
diff --git a/Engine/Source/Runtime/Engine/Private/WorldPartition/WorldPartition.cpp b/Engine/Source/Runtime/Engine/Private/WorldPartition/WorldPartition.cpp
--- a/Engine/Source/Runtime/Engine/Private/WorldPartition/WorldPartition.cpp
+++ b/Engine/Source/Runtime/Engine/Private/WorldPartition/WorldPartition.cpp
@@ -629,6 +629,10 @@ void UWorldPartition::RegisterDelegates()
 	if (World->IsGameWorld())
 	{
 		World->OnWorldBeginPlay.AddUObject(this, &UWorldPartition::OnWorldBeginPlay);
+
+#if !UE_BUILD_SHIPPING
+		FCoreDelegates::OnGetOnScreenMessages.AddUObject(this, &UWorldPartition::GetOnScreenMessages);
+#endif
 	}
 }
 
@@ -654,8 +658,22 @@ void UWorldPartition::UnregisterDelegates()
 	if (World->IsGameWorld())
 	{
 		World->OnWorldBeginPlay.RemoveAll(this);
+
+#if !UE_BUILD_SHIPPING
+		FCoreDelegates::OnGetOnScreenMessages.RemoveAll(this);
+#endif
+	}
+}
+
+#if !UE_BUILD_SHIPPING
+void UWorldPartition::GetOnScreenMessages(FCoreDelegates::FSeverityMessageMap& OutMessages)
+{
+	if (StreamingPolicy)
+	{
+		StreamingPolicy->GetOnScreenMessages(OutMessages);
 	}
 }
+#endif
 
 void UWorldPartition::OnWorldBeginPlay()
 {
diff --git a/Engine/Source/Runtime/Engine/Private/WorldPartition/WorldPartitionRuntimeHash.cpp b/Engine/Source/Runtime/Engine/Private/WorldPartition/WorldPartitionRuntimeHash.cpp
--- a/Engine/Source/Runtime/Engine/Private/WorldPartition/WorldPartitionRuntimeHash.cpp
+++ b/Engine/Source/Runtime/Engine/Private/WorldPartition/WorldPartitionRuntimeHash.cpp
@@ -279,6 +279,36 @@ void UWorldPartitionRuntimeHash::SortStreamingCellsByImportance(const TSet<const
 	});
 }
 
+EWorldPartitionStreamingPerformance UWorldPartitionRuntimeHash::GetStreamingPerformance(const TSet<const UWorldPartitionRuntimeCell*>& CellsToActivate) const
+{
+	EWorldPartitionStreamingPerformance StreamingPerformance = EWorldPartitionStreamingPerformance::Good;
+
+	if (!CellsToActivate.IsEmpty() && GetWorld()->bMatchStarted)
+	{
+		UWorld* World = GetWorld();
+
+		for (const UWorldPartitionRuntimeCell* Cell : CellsToActivate)
+		{
+			if (Cell->GetBlockOnSlowLoading() && !Cell->IsAlwaysLoaded() && Cell->GetStreamingStatus() != LEVEL_Visible)
+			{
+				EWorldPartitionStreamingPerformance CellPerformance = GetStreamingPerformanceForCell(Cell);
+				// Cell Performance is worst than previous cell performance
+				if (CellPerformance > StreamingPerformance)
+				{
+					StreamingPerformance = CellPerformance;
+					// Early out performance is critical
+					if (StreamingPerformance == EWorldPartitionStreamingPerformance::Critical)
+					{
+						return StreamingPerformance;
+					}
+				}
+			}
+		}
+	}
+
+	return StreamingPerformance;
+}
+
 void UWorldPartitionRuntimeHash::FStreamingSourceCells::AddCell(const UWorldPartitionRuntimeCell* InCell, const FWorldPartitionStreamingSource& InSource)
 {
 	InCell->CacheStreamingSourceInfo(InSource);
diff --git a/Engine/Source/Runtime/Engine/Private/WorldPartition/WorldPartitionRuntimeSpatialHash.cpp b/Engine/Source/Runtime/Engine/Private/WorldPartition/WorldPartitionRuntimeSpatialHash.cpp
--- a/Engine/Source/Runtime/Engine/Private/WorldPartition/WorldPartitionRuntimeSpatialHash.cpp
+++ b/Engine/Source/Runtime/Engine/Private/WorldPartition/WorldPartitionRuntimeSpatialHash.cpp
@@ -66,11 +66,24 @@ static FAutoConsoleVariableRef CVarShowRuntimeSpatialHashGridLevelCount(
 	GShowRuntimeSpatialHashGridLevelCount,
 	TEXT("Used to choose how many grid levels to display when showing world partition runtime hash."));
 
+static float GBlockOnSlowStreamingRatio = 0.25f;
+static FAutoConsoleVariableRef CVarBlockOnSlowStreamingRatio(
+	TEXT("wp.Runtime.BlockOnSlowStreamingRatio"),
+	GBlockOnSlowStreamingRatio,
+	TEXT("Ratio of DistanceToCell / LoadingRange to use to determine if World Partition streaming needs to block"));
+
+static float GBlockOnSlowStreamingWarningFactor = 2.f;
+static FAutoConsoleVariableRef CVarBlockOnSlowStreamingWarningFactor(
+	TEXT("wp.Runtime.BlockOnSlowStreamingWarningFactor"),
+	GBlockOnSlowStreamingWarningFactor,
+	TEXT("Factor of wp.Runtime.BlockOnSlowStreamingRatio we want to start notifying the user"));
+
 // ------------------------------------------------------------------------------------------------
 FSpatialHashStreamingGrid::FSpatialHashStreamingGrid()
 	: Origin(ForceInitToZero)
 	, CellSize(0)
 	, LoadingRange(0.0f)
+	, bBlockOnSlowStreaming(false)
 	, DebugColor(ForceInitToZero)
 	, WorldBounds(ForceInitToZero)
 	, bClientOnlyVisible(false)
@@ -102,6 +115,11 @@ const FSquare2DGridHelper& FSpatialHashStreamingGrid::GetGridHelper() const
 	return *GridHelper;
 }
 
+int32 FSpatialHashStreamingGrid::GetCellSize(int32 Level) const
+{
+	return GetGridHelper().Levels[Level].CellSize;
+}
+
 void FSpatialHashStreamingGrid::GetCells(const FWorldPartitionStreamingQuerySource& QuerySource, TSet<const UWorldPartitionRuntimeCell*>& OutCells) const
 {
 	TRACE_CPUPROFILER_EVENT_SCOPE(FSpatialHashStreamingGrid::GetCells_QuerySource);
@@ -767,6 +785,7 @@ bool UWorldPartitionRuntimeSpatialHash::CreateStreamingGrid(const FSpatialHashRu
 	CurrentStreamingGrid.CellSize = PartionedActors.CellSize;
 	CurrentStreamingGrid.WorldBounds = PartionedActors.WorldBounds;
 	CurrentStreamingGrid.LoadingRange = RuntimeGrid.LoadingRange;
+	CurrentStreamingGrid.bBlockOnSlowStreaming = RuntimeGrid.bBlockOnSlowStreaming;
 	CurrentStreamingGrid.DebugColor = RuntimeGrid.DebugColor;
 	CurrentStreamingGrid.bClientOnlyVisible = RuntimeGrid.bClientOnlyVisible;
 
@@ -867,6 +886,7 @@ bool UWorldPartitionRuntimeSpatialHash::CreateStreamingGrid(const FSpatialHashRu
 				StreamingCell->Position = FVector(Bounds.GetCenter(), 0.f);
 				StreamingCell->SetDebugInfo(CellGlobalCoords, CurrentStreamingGrid.GridName);
 				StreamingCell->SetClientOnlyVisible(CurrentStreamingGrid.bClientOnlyVisible);
+				StreamingCell->SetBlockOnSlowLoading(CurrentStreamingGrid.bBlockOnSlowStreaming);
 
 				UE_LOG(LogWorldPartition, Verbose, TEXT("Cell%s %s Actors = %d Bounds (%s)"), bIsCellAlwaysLoaded ? TEXT(" (AlwaysLoaded)") : TEXT(""), *StreamingCell->GetName(), FilteredActors.Num(), *Bounds.ToString());
 
@@ -1110,6 +1130,53 @@ bool UWorldPartitionRuntimeSpatialHash::GetStreamingCells(const TArray<FWorldPar
 	return !!(OutActivateCells.Num() + OutLoadCells.Num());
 }
 
+const TMap<FName, const FSpatialHashStreamingGrid*>& UWorldPartitionRuntimeSpatialHash::GetNameToGridMapping() const
+{
+	if (NameToGridMapping.IsEmpty())
+	{
+		for (const FSpatialHashStreamingGrid& StreamingGrid : StreamingGrids)
+		{
+			NameToGridMapping.Add(StreamingGrid.GridName, &StreamingGrid);
+		}
+	}
+
+	return NameToGridMapping;
+}
+
+EWorldPartitionStreamingPerformance UWorldPartitionRuntimeSpatialHash::GetStreamingPerformanceForCell(const UWorldPartitionRuntimeCell* Cell) const
+{
+	// If base class already returning critical. Early out.
+	if (Super::GetStreamingPerformanceForCell(Cell) == EWorldPartitionStreamingPerformance::Critical)
+	{
+		return EWorldPartitionStreamingPerformance::Critical;
+	}
+
+	check(Cell->GetBlockOnSlowLoading());
+	const float BlockOnSlowStreamingWarningRatio = GBlockOnSlowStreamingRatio * GBlockOnSlowStreamingWarningFactor;
+	
+	const UWorldPartitionRuntimeSpatialHashCell* StreamingCell = Cast<const UWorldPartitionRuntimeSpatialHashCell>(Cell);
+	check(StreamingCell);
+
+	const FSpatialHashStreamingGrid* const* StreamingGrid = GetNameToGridMapping().Find(StreamingCell->GetGridName());
+	check(StreamingGrid && *StreamingGrid);
+	
+	const float LoadingRange = (*StreamingGrid)->LoadingRange;
+	const float Distance = FMath::Sqrt(StreamingCell->CachedSourceMinDistanceSquare) - ((float)(*StreamingGrid)->GetCellSize(StreamingCell->Level) / 2.f);
+
+	const float Ratio = Distance / LoadingRange;
+
+	if (Ratio < GBlockOnSlowStreamingRatio)
+	{
+		return EWorldPartitionStreamingPerformance::Critical;
+	}
+	else if (Ratio < BlockOnSlowStreamingWarningRatio)
+	{
+		return EWorldPartitionStreamingPerformance::Slow;
+	}
+
+	return EWorldPartitionStreamingPerformance::Good;
+}
+
 FVector2D UWorldPartitionRuntimeSpatialHash::GetDraw2DDesiredFootprint(const FVector2D& CanvasSize) const
 {
 	return FVector2D(CanvasSize.X * GetFilteredStreamingGrids().Num(), CanvasSize.Y);
diff --git a/Engine/Source/Runtime/Engine/Private/WorldPartition/WorldPartitionRuntimeSpatialHashCell.cpp b/Engine/Source/Runtime/Engine/Private/WorldPartition/WorldPartitionRuntimeSpatialHashCell.cpp
--- a/Engine/Source/Runtime/Engine/Private/WorldPartition/WorldPartitionRuntimeSpatialHashCell.cpp
+++ b/Engine/Source/Runtime/Engine/Private/WorldPartition/WorldPartitionRuntimeSpatialHashCell.cpp
@@ -17,7 +17,8 @@ static FAutoConsoleVariableRef CVarRuntimeSpatialHashCellToSourceAngleContributi
 UWorldPartitionRuntimeSpatialHashCell::UWorldPartitionRuntimeSpatialHashCell(const FObjectInitializer& ObjectInitializer)
 : Super(ObjectInitializer)
 , Level(0)
-, CachedSourceMinDistance(0.f)
+, CachedSourceMinDistanceSquare(0.f)
+, CachedSourceMinSortDistance(0.f)
 {}
 
 #if WITH_EDITOR
@@ -55,7 +56,8 @@ bool UWorldPartitionRuntimeSpatialHashCell::CacheStreamingSourceInfo(const FWorl
 	float ModulatedDistance = SqrDistance * AngleFactor;
 
 	// If cache was dirtied, use value, else use minimum with existing cached value
-	CachedSourceMinDistance = bWasCacheDirtied ? ModulatedDistance : FMath::Min(ModulatedDistance, CachedSourceMinDistance);
+	CachedSourceMinSortDistance = bWasCacheDirtied ? ModulatedDistance : FMath::Min(ModulatedDistance, CachedSourceMinSortDistance);
+	CachedSourceMinDistanceSquare = bWasCacheDirtied ? SqrDistance : FMath::Min(SqrDistance, CachedSourceMinDistanceSquare);
 
 	return bWasCacheDirtied;
 }
@@ -73,7 +75,7 @@ int32 UWorldPartitionRuntimeSpatialHashCell::SortCompare(const UWorldPartitionRu
 		if (Result == 0)
 		{
 			// Closest distance (lower value is higher prio)
-			const float Diff = CachedSourceMinDistance - Other->CachedSourceMinDistance;
+			const float Diff = CachedSourceMinSortDistance - Other->CachedSourceMinSortDistance;
 			Result = Diff < 0.f ? -1 : (Diff > 0.f ? 1 : 0);
 		}
 	}
diff --git a/Engine/Source/Runtime/Engine/Private/WorldPartition/WorldPartitionStreamingPolicy.cpp b/Engine/Source/Runtime/Engine/Private/WorldPartition/WorldPartitionStreamingPolicy.cpp
--- a/Engine/Source/Runtime/Engine/Private/WorldPartition/WorldPartitionStreamingPolicy.cpp
+++ b/Engine/Source/Runtime/Engine/Private/WorldPartition/WorldPartitionStreamingPolicy.cpp
@@ -27,17 +27,31 @@
 #include "LevelEditorViewport.h"
 #endif
 
+#define LOCTEXT_NAMESPACE "WorldPartitionStreamingPolicy"
+
 static int32 GMaxLoadingStreamingCells = 4;
 static FAutoConsoleVariableRef CMaxLoadingStreamingCells(
 	TEXT("wp.Runtime.MaxLoadingStreamingCells"),
 	GMaxLoadingStreamingCells,
 	TEXT("Used to limit the number of concurrent loading world partition streaming cells."));
 
+static int32 GBlockOnSlowStreaming = 1;
+static FAutoConsoleCommand CVarBlockOnSlowStreaming(
+	TEXT("wp.Runtime.ToggleBlockOnSlowStreaming"),
+	TEXT("Toggles if streaming needs to block when to slow to catchup."),
+	FConsoleCommandDelegate::CreateLambda([] { GBlockOnSlowStreaming = !GBlockOnSlowStreaming; }));
+
+
 UWorldPartitionStreamingPolicy::UWorldPartitionStreamingPolicy(const FObjectInitializer& ObjectInitializer)
 : Super(ObjectInitializer) 
 , UpdateStreamingStateEpoch(0)
 , SortedAddToWorldCellsEpoch(INT_MIN)
 , DataLayersStatesServerEpoch(INT_MIN)
+, StreamingPerformance(EWorldPartitionStreamingPerformance::Good)
+#if !UE_BUILD_SHIPPING
+, OnScreenMessageStartTime(0.0)
+, OnScreenMessageStreamingPerformance(EWorldPartitionStreamingPerformance::Good)
+#endif
 {
 	if (!IsTemplate())
 	{
@@ -240,11 +254,89 @@ void UWorldPartitionStreamingPolicy::UpdateStreamingState()
 		SetTargetStateForCells(EWorldPartitionRuntimeCellState::Loaded, ToLoadCells);
 	}
 
+	// Evaluate streaming performance based on cells that should be activated
+	UpdateStreamingPerformance(ActivateStreamingCells);
+	
 	// Reset frame StreamingSourceCells (optimization to avoid reallocation at every call to UpdateStreamingState)
 	FrameActivateCells.Reset();
 	FrameLoadCells.Reset();
 }
 
+void UWorldPartitionStreamingPolicy::UpdateStreamingPerformance(const TSet<const UWorldPartitionRuntimeCell*>& CellsToActivate)
+{		
+	TRACE_CPUPROFILER_EVENT_SCOPE(UWorldPartitionStreamingPolicy::UpdateStreamingPerformance);
+	UWorld* World = GetWorld();
+	// If we are currently in a blocked loading just reset the on screen message time and return
+	if (StreamingPerformance == EWorldPartitionStreamingPerformance::Critical && World->GetIsInBlockTillLevelStreamingCompleted())
+	{
+#if !UE_BUILD_SHIPPING
+		OnScreenMessageStartTime = FPlatformTime::Seconds();
+#endif
+		return;
+	}
+
+	EWorldPartitionStreamingPerformance NewStreamingPerformance = WorldPartition->RuntimeHash->GetStreamingPerformance(CellsToActivate);
+		
+	if (StreamingPerformance != NewStreamingPerformance)
+	{
+		UE_LOG(LogWorldPartition, Log, TEXT("Streaming performance changed: %s -> %s"),
+			*StaticEnum<EWorldPartitionStreamingPerformance>()->GetDisplayNameTextByValue((int64)StreamingPerformance).ToString(),
+			*StaticEnum<EWorldPartitionStreamingPerformance>()->GetDisplayNameTextByValue((int64)NewStreamingPerformance).ToString());
+		
+		StreamingPerformance = NewStreamingPerformance;
+	}
+
+#if !UE_BUILD_SHIPPING
+	if (StreamingPerformance != EWorldPartitionStreamingPerformance::Good)
+	{
+		// performance still bad keep message alive
+		OnScreenMessageStartTime = FPlatformTime::Seconds();
+		OnScreenMessageStreamingPerformance = StreamingPerformance;
+	}
+#endif
+	
+	if (StreamingPerformance == EWorldPartitionStreamingPerformance::Critical)
+	{
+		// This is a first very simple standalone implementation of handling of critical streaming conditions.
+		if (GBlockOnSlowStreaming && !World->GetIsInBlockTillLevelStreamingCompleted() && World->GetNetMode() == NM_Standalone)
+		{
+			World->bRequestedBlockOnAsyncLoading = true;
+		}
+	}
+}
+
+#if !UE_BUILD_SHIPPING
+void UWorldPartitionStreamingPolicy::GetOnScreenMessages(FCoreDelegates::FSeverityMessageMap& OutMessages)
+{
+	// Keep displaying for 2 seconds (or more if health stays bad)
+	double DisplayTime = FPlatformTime::Seconds() - OnScreenMessageStartTime;
+	if (DisplayTime < 2.0)
+	{
+		switch (OnScreenMessageStreamingPerformance)
+		{
+		case EWorldPartitionStreamingPerformance::Critical:
+			OutMessages.Add(FCoreDelegates::EOnScreenMessageSeverity::Error, LOCTEXT("WPStreamingCritical", "[Critical] WorldPartition Streaming Performance"));
+			break;
+		case EWorldPartitionStreamingPerformance::Slow:
+			OutMessages.Add(FCoreDelegates::EOnScreenMessageSeverity::Warning, LOCTEXT("WPStreamingWarning", "[Slow] WorldPartition Streaming Performance"));
+			break;
+		default:
+			break;
+		}
+	}
+	else
+	{
+		OnScreenMessageStreamingPerformance = EWorldPartitionStreamingPerformance::Good;
+	}
+}
+#endif
+
+bool UWorldPartitionStreamingPolicy::ShouldSkipCellForPerformance(const UWorldPartitionRuntimeCell* Cell) const
+{
+	// When performance is degrading start skipping non blocking cells
+	return StreamingPerformance != EWorldPartitionStreamingPerformance::Good && !Cell->GetBlockOnSlowLoading();
+}
+
 int32 UWorldPartitionStreamingPolicy::GetMaxCellsToLoad() const
 {
 	// This policy limits the number of concurrent loading streaming cells, except if match hasn't started
@@ -290,6 +382,11 @@ void UWorldPartitionStreamingPolicy::SetCellsStateToLoaded(const TSet<const UWor
 	// Trigger cell loading. Depending on actual state of cell limit loading.
 	for (const UWorldPartitionRuntimeCell* Cell : SortedCells)
 	{
+		if (ShouldSkipCellForPerformance(Cell))
+		{
+			continue;
+		}
+
 		UE_LOG(LogWorldPartition, Verbose, TEXT("UWorldPartitionStreamingPolicy::LoadCells %s"), *Cell->GetName());
 		if (ActivatedCells.Contains(Cell))
 		{
@@ -322,6 +419,11 @@ void UWorldPartitionStreamingPolicy::SetCellsStateToActivated(const TSet<const U
 	// Trigger cell activation. Depending on actual state of cell limit loading.
 	for (const UWorldPartitionRuntimeCell* Cell : SortedCells)
 	{
+		if (ShouldSkipCellForPerformance(Cell))
+		{
+			continue;
+		}
+
 		UE_LOG(LogWorldPartition, Verbose, TEXT("UWorldPartitionStreamingPolicy::ActivateCells %s"), *Cell->GetName());
 		if (LoadedCells.Contains(Cell))
 		{
@@ -460,3 +562,5 @@ void UWorldPartitionStreamingPolicy::DrawRuntimeHash3D()
 		WorldPartition->RuntimeHash->Draw3D(StreamingSources);
 	}
 }
+
+#undef LOCTEXT_NAMESPACE
\ No newline at end of file
diff --git a/Engine/Source/Runtime/Engine/Public/WorldPartition/WorldPartition.h b/Engine/Source/Runtime/Engine/Public/WorldPartition/WorldPartition.h
--- a/Engine/Source/Runtime/Engine/Public/WorldPartition/WorldPartition.h
+++ b/Engine/Source/Runtime/Engine/Public/WorldPartition/WorldPartition.h
@@ -4,6 +4,7 @@
 #include "CoreMinimal.h"
 #include "UObject/ObjectMacros.h"
 #include "UObject/Object.h"
+#include "Misc/CoreDelegates.h"
 #include "GameFramework/Actor.h"
 #include "Templates/SubclassOf.h"
 #include "WorldPartition/WorldPartitionActorDesc.h"
@@ -244,5 +245,9 @@ private:
 	void ApplyActorTransform(AActor* Actor, const FTransform& InTransform);
 #endif
 
+#if !UE_BUILD_SHIPPING
+	void GetOnScreenMessages(FCoreDelegates::FSeverityMessageMap& OutMessages);
+#endif
+
 	friend class UWorldPartitionStreamingPolicy;
 };
\ No newline at end of file
diff --git a/Engine/Source/Runtime/Engine/Public/WorldPartition/WorldPartitionRuntimeCell.h b/Engine/Source/Runtime/Engine/Public/WorldPartition/WorldPartitionRuntimeCell.h
--- a/Engine/Source/Runtime/Engine/Public/WorldPartition/WorldPartitionRuntimeCell.h
+++ b/Engine/Source/Runtime/Engine/Public/WorldPartition/WorldPartitionRuntimeCell.h
@@ -129,7 +129,7 @@ class UWorldPartitionRuntimeCell : public UObject
 	virtual const FString& GetDebugName() const { return DebugName; }
 	virtual bool IsDebugShown() const;
 	virtual int32 SortCompare(const UWorldPartitionRuntimeCell* Other) const;
-
+	virtual FName GetGridName() const { return GridName; }
 	/** Caches information on streaming source that will be used later on to sort cell. Returns true if cache was reset, else returns false. */
 	virtual bool CacheStreamingSourceInfo(const struct FWorldPartitionStreamingSource& Source) const;
 
@@ -142,7 +142,10 @@ class UWorldPartitionRuntimeCell : public UObject
 		return Algo::AnyOf(DataLayers, [&InDataLayers](const FName& DataLayer) { return InDataLayers.Contains(DataLayer); });
 	}
 
+	bool GetBlockOnSlowLoading() const { return bBlockOnSlowLoading; }
 #if WITH_EDITOR
+	void SetBlockOnSlowLoading(bool bInBlockOnSlowLoading) { bBlockOnSlowLoading = bInBlockOnSlowLoading; }
+
 	void SetClientOnlyVisible(bool bInClientOnlyVisible) { bClientOnlyVisible = bInClientOnlyVisible; }
 	bool GetClientOnlyVisible() const { return bClientOnlyVisible; }
 
@@ -195,6 +198,9 @@ private:
 	UPROPERTY()
 	bool bClientOnlyVisible;
 
+	UPROPERTY()
+	bool bBlockOnSlowLoading;
+
 	// Source Priority
 	mutable int32 CachedSourcePriority;
 
diff --git a/Engine/Source/Runtime/Engine/Public/WorldPartition/WorldPartitionRuntimeHash.h b/Engine/Source/Runtime/Engine/Public/WorldPartition/WorldPartitionRuntimeHash.h
--- a/Engine/Source/Runtime/Engine/Public/WorldPartition/WorldPartitionRuntimeHash.h
+++ b/Engine/Source/Runtime/Engine/Public/WorldPartition/WorldPartitionRuntimeHash.h
@@ -16,6 +16,14 @@
 
 class UWorldPartitionRuntimeCell;
 
+UENUM()
+enum class EWorldPartitionStreamingPerformance
+{
+	Good,
+	Slow,
+	Critical
+};
+
 UCLASS(Abstract, Config=Engine, AutoExpandCategories=(WorldPartition), Within = WorldPartition)
 class ENGINE_API UWorldPartitionRuntimeHash : public UObject
 {
@@ -60,6 +68,7 @@ class ENGINE_API UWorldPartitionRuntimeHash : public UObject
 	virtual bool GetStreamingCells(const FWorldPartitionStreamingQuerySource& QuerySource, TSet<const UWorldPartitionRuntimeCell*>& OutCells) const { return false; }
 	virtual bool GetStreamingCells(const TArray<FWorldPartitionStreamingSource>& Sources, UWorldPartitionRuntimeHash::FStreamingSourceCells& OutActivateCells, UWorldPartitionRuntimeHash::FStreamingSourceCells& OutLoadCells) const { return false; };
 	virtual void SortStreamingCellsByImportance(const TSet<const UWorldPartitionRuntimeCell*>& InCells, const TArray<FWorldPartitionStreamingSource>& InSources, TArray<const UWorldPartitionRuntimeCell*, TInlineAllocator<256>>& OutSortedCells) const;
+	EWorldPartitionStreamingPerformance GetStreamingPerformance(const TSet<const UWorldPartitionRuntimeCell*>& CellToActivate) const;
 
 	/* Returns desired footprint that Draw2D should take relative to given Canvas size (the value can exceed the given size).
 	 * UWorldPartitionSubSystem will re-adapt the size relative to all others UWorldPartitionRuntimeHash and provide the correct size to Draw2D.
@@ -73,6 +82,8 @@ class ENGINE_API UWorldPartitionRuntimeHash : public UObject
 	virtual bool ContainsRuntimeHash(const FString& Name) const { return false; }
 
 protected:
+	virtual EWorldPartitionStreamingPerformance GetStreamingPerformanceForCell(const UWorldPartitionRuntimeCell* Cell) const { return EWorldPartitionStreamingPerformance::Good; }
+
 #if WITH_EDITOR
 	virtual void CheckForErrorsInternal(const TMap<FGuid, FWorldPartitionActorViewProxy>& ActorDescList) const;
 
diff --git a/Engine/Source/Runtime/Engine/Public/WorldPartition/WorldPartitionRuntimeSpatialHash.h b/Engine/Source/Runtime/Engine/Public/WorldPartition/WorldPartitionRuntimeSpatialHash.h
--- a/Engine/Source/Runtime/Engine/Public/WorldPartition/WorldPartitionRuntimeSpatialHash.h
+++ b/Engine/Source/Runtime/Engine/Public/WorldPartition/WorldPartitionRuntimeSpatialHash.h
@@ -103,6 +103,9 @@ struct FSpatialHashStreamingGrid
 	UPROPERTY()
 	float LoadingRange;
 
+	UPROPERTY()
+	bool bBlockOnSlowStreaming;
+		
 	UPROPERTY()
 	FLinearColor DebugColor;
 
@@ -127,6 +130,7 @@ struct FSpatialHashStreamingGrid
 	}
 
 	// Used by PIE/Game
+	int32 GetCellSize(int32 Level) const;
 	void GetCells(const FWorldPartitionStreamingQuerySource& QuerySource, TSet<const UWorldPartitionRuntimeCell*>& OutCells) const;
 	void GetCells(const TArray<FWorldPartitionStreamingSource>& Sources, const class UDataLayerSubsystem* DataLayerSubsystem, UWorldPartitionRuntimeHash::FStreamingSourceCells& OutActivateCells, UWorldPartitionRuntimeHash::FStreamingSourceCells& OutLoadCells) const;
 	void GetAlwaysLoadedCells(const UDataLayerSubsystem* DataLayerSubsystem, TSet<const UWorldPartitionRuntimeCell*>& OutActivateCells, TSet<const UWorldPartitionRuntimeCell*>& OutLoadCells) const;
@@ -154,6 +158,7 @@ struct FSpatialHashRuntimeGrid
 #if WITH_EDITORONLY_DATA
 		: CellSize(0)
 		, LoadingRange(1000)
+		, bBlockOnSlowStreaming(false)
 		, Priority(0)
 		, DebugColor(FLinearColor::MakeRandomColor())
 		, bClientOnlyVisible(false)
@@ -170,6 +175,10 @@ struct FSpatialHashRuntimeGrid
 	UPROPERTY(EditAnywhere, Category=Settings)
 	float LoadingRange;
 
+	/** Should streaming block in situations where cells aren't getting loaded fast enough. */
+	UPROPERTY(EditAnywhere, Category = Settings)
+	bool bBlockOnSlowStreaming;
+		
 	UPROPERTY(EditAnywhere, Category=Settings)
 	int32 Priority;
 
@@ -232,6 +241,7 @@ public:
 
 protected:
 	bool ShouldConsiderClientOnlyVisibleCells() const;
+	virtual EWorldPartitionStreamingPerformance GetStreamingPerformanceForCell(const UWorldPartitionRuntimeCell* Cell) const override;
 
 #if WITH_EDITOR
 	virtual bool GenerateStreaming(EWorldPartitionStreamingMode Mode, class UWorldPartitionStreamingPolicy* StreamingPolicy, TArray<FString>* OutPackagesToGenerate = nullptr) override;
@@ -263,6 +273,8 @@ private:
 	UPROPERTY(NonPIEDuplicateTransient)
 	TArray<FSpatialHashStreamingGrid> StreamingGrids;
 
+	mutable TMap<FName, const FSpatialHashStreamingGrid*> NameToGridMapping;
+	
 	virtual FVector2D GetDraw2DDesiredFootprint(const FVector2D& CanvasSize) const override;
 	virtual void Draw2D(class UCanvas* Canvas, const TArray<FWorldPartitionStreamingSource>& Sources, const FVector2D& PartitionCanvasSize, FVector2D& Offset) const override;
 	virtual void Draw3D(const TArray<FWorldPartitionStreamingSource>& Sources) const override;
@@ -271,6 +283,7 @@ private:
 private:
 	void GetAlwaysLoadedStreamingCells(const FSpatialHashStreamingGrid& StreamingGrid, TSet<const UWorldPartitionRuntimeCell*>& Cells) const;
 	void GetStreamingCells(const FVector& Position, const FSpatialHashStreamingGrid& StreamingGrid, TSet<const UWorldPartitionRuntimeCell*>& Cells) const;
+	const TMap<FName, const FSpatialHashStreamingGrid*>& GetNameToGridMapping() const;
 #if WITH_EDITOR
 	bool CreateStreamingGrid(const FSpatialHashRuntimeGrid& RuntimeGrid, const FSquare2DGridHelper& PartionedActors, EWorldPartitionStreamingMode Mode, UWorldPartitionStreamingPolicy* StreamingPolicy, TArray<FString>* OutPackagesToGenerate = nullptr);
 #endif
diff --git a/Engine/Source/Runtime/Engine/Public/WorldPartition/WorldPartitionRuntimeSpatialHashCell.h b/Engine/Source/Runtime/Engine/Public/WorldPartition/WorldPartitionRuntimeSpatialHashCell.h
--- a/Engine/Source/Runtime/Engine/Public/WorldPartition/WorldPartitionRuntimeSpatialHashCell.h
+++ b/Engine/Source/Runtime/Engine/Public/WorldPartition/WorldPartitionRuntimeSpatialHashCell.h
@@ -23,9 +23,12 @@ class UWorldPartitionRuntimeSpatialHashCell : public UWorldPartitionRuntimeCell
 
 	UPROPERTY()
 	int32 Level;
+		
+	// Used to determine if the streaming of cells is lagging
+	mutable float CachedSourceMinDistanceSquare;
 
 	// Computed and cached value used by SortCompare to sort Cells
-	mutable float CachedSourceMinDistance;
+	mutable float CachedSourceMinSortDistance;
 
 #if WITH_EDITORONLY_DATA
 	UPROPERTY()
diff --git a/Engine/Source/Runtime/Engine/Public/WorldPartition/WorldPartitionStreamingPolicy.h b/Engine/Source/Runtime/Engine/Public/WorldPartition/WorldPartitionStreamingPolicy.h
--- a/Engine/Source/Runtime/Engine/Public/WorldPartition/WorldPartitionStreamingPolicy.h
+++ b/Engine/Source/Runtime/Engine/Public/WorldPartition/WorldPartitionStreamingPolicy.h
@@ -10,6 +10,7 @@
 #pragma once
 
 #include "Containers/Set.h"
+#include "Misc/CoreDelegates.h"
 #include "WorldPartition/WorldPartitionRuntimeCell.h"
 #include "WorldPartition/WorldPartitionRuntimeHash.h"
 #include "WorldPartition/WorldPartitionStreamingSource.h"
@@ -43,6 +44,10 @@ public:
 	virtual void RemapSoftObjectPath(FSoftObjectPath& ObjectPath) {}
 #endif
 
+#if !UE_BUILD_SHIPPING
+	virtual void GetOnScreenMessages(FCoreDelegates::FSeverityMessageMap& OutMessages);
+#endif
+
 	virtual UObject* GetSubObject(const TCHAR* SubObjectPath) { return nullptr; }
 
 	const TArray<FWorldPartitionStreamingSource>& GetStreamingSources() const { return StreamingSources; }
@@ -54,6 +59,8 @@ protected:
 	virtual int32 GetCellLoadingCount() const { return 0; }
 	virtual int32 GetMaxCellsToLoad() const;
 	virtual void UpdateStreamingSources();
+	void UpdateStreamingPerformance(const TSet<const UWorldPartitionRuntimeCell*>& CellsToActivate);
+	bool ShouldSkipCellForPerformance(const UWorldPartitionRuntimeCell* Cell) const;
 
 	const UWorldPartition* WorldPartition;
 	TSet<const UWorldPartitionRuntimeCell*> LoadedCells;
@@ -68,4 +75,10 @@ protected:
 	mutable TArray<const UWorldPartitionRuntimeCell*, TInlineAllocator<256>> SortedAddToWorldCells;
 
 	int32 DataLayersStatesServerEpoch;
+
+	EWorldPartitionStreamingPerformance StreamingPerformance;
+#if !UE_BUILD_SHIPPING
+	double OnScreenMessageStartTime;
+	EWorldPartitionStreamingPerformance  OnScreenMessageStreamingPerformance;
+#endif
 };
\ No newline at end of file
