Change moment of initialization of Animation Attribute system - previously had a race-condition triggered through parallel eval

diff --git a/Engine/Source/Runtime/AnimGraphRuntime/Private/AnimGraphRuntimeModule.cpp b/Engine/Source/Runtime/AnimGraphRuntime/Private/AnimGraphRuntimeModule.cpp
--- a/Engine/Source/Runtime/AnimGraphRuntime/Private/AnimGraphRuntimeModule.cpp
+++ b/Engine/Source/Runtime/AnimGraphRuntime/Private/AnimGraphRuntimeModule.cpp
@@ -5,6 +5,7 @@
 #include "BoneControllers/AnimNode_AnimDynamics.h"
 #include "UObject/UObjectIterator.h"
 #include "Animation/AnimInstance.h"
+#include "Animation/AttributeTypes.h"
 
 //////////////////////////////////////////////////////////////////////////
 // FAnimGraphRuntimeModule
@@ -14,6 +15,10 @@ class FAnimGraphRuntimeModule : public IModuleInterface
 public:
 	virtual void StartupModule() override
 	{
+		FCoreDelegates::OnPostEngineInit.AddLambda([]()
+		{			
+			UE::Anim::AttributeTypes::Initialize();
+		});
 	}
 
 	virtual void ShutdownModule() override
diff --git a/Engine/Source/Runtime/Engine/Private/Animation/AttributeTypes.cpp b/Engine/Source/Runtime/Engine/Private/Animation/AttributeTypes.cpp
--- a/Engine/Source/Runtime/Engine/Private/Animation/AttributeTypes.cpp
+++ b/Engine/Source/Runtime/Engine/Private/Animation/AttributeTypes.cpp
@@ -14,15 +14,13 @@ namespace UE
 		void AttributeTypes::Initialize()
 		{
 			static bool bInitialized = false;
-			if (!bInitialized)
-			{
-				bInitialized = true;
-
-				RegisterType<FFloatAnimationAttribute>();
-				RegisterType<FIntegerAnimationAttribute>();
-				RegisterType<FStringAnimationAttribute>();
-				RegisterType<FTransformAnimationAttribute>();
-			}
+			checkf(bInitialized == false, TEXT("Trying to initialize attribute type system multiple times"));
+			
+			bInitialized = true;
+			RegisterType<FFloatAnimationAttribute>();
+			RegisterType<FIntegerAnimationAttribute>();
+			RegisterType<FStringAnimationAttribute>();
+			RegisterType<FTransformAnimationAttribute>();
 		}		
 	}
 }
diff --git a/Engine/Source/Runtime/Engine/Public/Animation/AttributeTypes.h b/Engine/Source/Runtime/Engine/Public/Animation/AttributeTypes.h
--- a/Engine/Source/Runtime/Engine/Public/Animation/AttributeTypes.h
+++ b/Engine/Source/Runtime/Engine/Public/Animation/AttributeTypes.h
@@ -34,19 +34,17 @@ namespace UE
 		struct ENGINE_API AttributeTypes
 		{
 		protected:
-			static void Initialize();
-
 			static TArray<TWeakObjectPtr<const UScriptStruct>> RegisteredTypes;
 			static TArray<TUniquePtr<IAttributeBlendOperator>> Operators;
 			static TArray<TWeakObjectPtr<const UScriptStruct>> InterpolatableTypes;
 		public:
+			
+			static void Initialize();
 
 			/** Used for registering an attribute type for which TAttributeTypeTraits::WithCustomBlendOperator is set to true, use RegisterType() otherwise */
 			template<typename AttributeType, typename OperatorType, typename... OperatorArgs>
 			static void RegisterTypeWithOperator(OperatorArgs&&... args)
 			{
-				AttributeTypes::Initialize();
-
 				static_assert(TAttributeTypeTraits<AttributeType>::WithCustomBlendOperator, "Attribute type does not require a custom blend operation");
 				UScriptStruct* ScriptStruct = AttributeType::StaticStruct();
 
@@ -70,7 +68,6 @@ namespace UE
 			{
 				static_assert(!TAttributeTypeTraits<AttributeType>::WithCustomBlendOperator, "Attribute type requires a custom blend operation");
 
-				AttributeTypes::Initialize();
 				UScriptStruct* ScriptStruct = AttributeType::StaticStruct();
 				AttributeTypes::RegisteredTypes.Add(ScriptStruct);
 				
@@ -93,7 +90,6 @@ namespace UE
 				static_assert(!TAttributeTypeTraits<AttributeType>::WithCustomBlendOperator, "Attribute type requires a custom blend operation");
 				static_assert(TModels<CBlendableAttribute, AttributeType>::Value, "Missing arithmetic operators required for Attribute blending");
 
-				AttributeTypes::Initialize();
 				UScriptStruct* ScriptStruct = AttributeType::StaticStruct();
 				AttributeTypes::RegisteredTypes.Add(ScriptStruct);
 
@@ -111,7 +107,6 @@ namespace UE
 			{
 				static_assert(!TAttributeTypeTraits<AttributeType>::WithCustomBlendOperator, "Attribute type requires a custom blend operation");
 
-				AttributeTypes::Initialize();
 				UScriptStruct* ScriptStruct = AttributeType::StaticStruct();
 				AttributeTypes::RegisteredTypes.Add(ScriptStruct);
 
@@ -137,7 +132,6 @@ namespace UE
 			/** Returns the blend operator for the provided type, asserts when the type is not registered */
 			static const IAttributeBlendOperator* GetTypeOperator(TWeakObjectPtr<const UScriptStruct> WeakStruct)
 			{
-				AttributeTypes::Initialize();
 				const int32 Index = AttributeTypes::RegisteredTypes.IndexOfByKey(WeakStruct);
 				ensure(WeakStruct.IsValid());
 				checkf(Index != INDEX_NONE, TEXT("Missing operator for attribute, type %s was not registered previously"), *WeakStruct->GetName());
@@ -147,21 +141,18 @@ namespace UE
 			/** Returns whether or not the provided type can be interpolated, defaults to false when the type is not registered */
 			static bool CanInterpolateType(TWeakObjectPtr<const UScriptStruct> WeakStruct)
 			{
-				AttributeTypes::Initialize();
 				return AttributeTypes::InterpolatableTypes.Contains(WeakStruct);
 			}
 
 			/** Returns whether or not the type is registered */
 			static bool IsTypeRegistered(const UScriptStruct* ScriptStruct)
 			{
-				AttributeTypes::Initialize();
 				return AttributeTypes::RegisteredTypes.Contains(ScriptStruct);
 			}
 
 			/** Returns all registered types */
 			static TArray<TWeakObjectPtr<const UScriptStruct>>& GetRegisteredTypes()
 			{
-				AttributeTypes::Initialize();
 				return AttributeTypes::RegisteredTypes;
 			}
 		};
