Use the function to access material DDC master key instead of a define directly.

- Requested on UDN, makes it easier for the licensees to append their own ID to the master key.

diff --git a/Engine/Source/Runtime/Engine/Private/Materials/MaterialShader.cpp b/Engine/Source/Runtime/Engine/Private/Materials/MaterialShader.cpp
--- a/Engine/Source/Runtime/Engine/Private/Materials/MaterialShader.cpp
+++ b/Engine/Source/Runtime/Engine/Private/Materials/MaterialShader.cpp
@@ -183,7 +183,7 @@ static FString GetMaterialShaderMapKeyString(const FMaterialShaderMapId& ShaderM
 	FMaterialAttributeDefinitionMap::AppendDDCKeyString(ShaderMapKeyString);
 	FShaderCompileUtilities::AppendGBufferDDCKeyString(Platform, ShaderMapKeyString);
 
-	return FDerivedDataCacheInterface::BuildCacheKey(TEXT("MATSM"), MATERIALSHADERMAP_DERIVEDDATA_VER, *ShaderMapKeyString);
+	return FDerivedDataCacheInterface::BuildCacheKey(TEXT("MATSM"), *GetMaterialShaderMapDDCKey(), *ShaderMapKeyString);
 }
 
 static FString FSHA1_HashString(const FString& Key)
