UStaticMesh - Added a few utility methods
* GetNumTriangles()
* GetNumTexCoords()
* GetNumNaniteVertices()
* GetNumNaniteTriangles()

diff --git a/Engine/Source/Runtime/Engine/Classes/Engine/StaticMesh.h b/Engine/Source/Runtime/Engine/Classes/Engine/StaticMesh.h
--- a/Engine/Source/Runtime/Engine/Classes/Engine/StaticMesh.h
+++ b/Engine/Source/Runtime/Engine/Classes/Engine/StaticMesh.h
@@ -1582,6 +1582,26 @@ public:
 	 */
 	ENGINE_API int32 GetNumVertices(int32 LODIndex) const;
 
+	/**
+	 * Returns the number of triangles for the specified LOD.
+	 */
+	ENGINE_API int32 GetNumTriangles(int32 LODIndex) const;
+
+	/**
+	 * Returns the number of tex coords for the specified LOD.
+	 */
+	ENGINE_API int32 GetNumTexCoords(int32 LODIndex) const;
+
+	/**
+	 * Returns the number of vertices of the Nanite representation of this mesh.
+	 */
+	ENGINE_API int32 GetNumNaniteVertices() const;
+
+	/**
+	 * Returns the number of triangles of the Nanite representation of this mesh.
+	 */
+	ENGINE_API int32 GetNumNaniteTriangles() const;
+
 	/**
 	 * Returns the number of LODs used by the mesh.
 	 */
diff --git a/Engine/Source/Runtime/Engine/Private/StaticMesh.cpp b/Engine/Source/Runtime/Engine/Private/StaticMesh.cpp
--- a/Engine/Source/Runtime/Engine/Private/StaticMesh.cpp
+++ b/Engine/Source/Runtime/Engine/Private/StaticMesh.cpp
@@ -3119,11 +3119,59 @@ int32 UStaticMesh::GetNumVertices(int32 LODIndex) const
 	int32 NumVertices = 0;
 	if (GetRenderData() && GetRenderData()->LODResources.IsValidIndex(LODIndex))
 	{
-		NumVertices = GetRenderData()->LODResources[LODIndex].VertexBuffers.StaticMeshVertexBuffer.GetNumVertices();
+		NumVertices = GetRenderData()->LODResources[LODIndex].GetNumVertices();
 	}
 	return NumVertices;
 }
 
+int32 UStaticMesh::GetNumTriangles(int32 LODIndex) const
+{
+	int32 NumTriangles = 0;
+	if (GetRenderData() && GetRenderData()->LODResources.IsValidIndex(LODIndex))
+	{
+		NumTriangles = GetRenderData()->LODResources[LODIndex].GetNumTriangles();
+	}
+	return NumTriangles;
+}
+
+int32 UStaticMesh::GetNumTexCoords(int32 LODIndex) const
+{
+	int32 NumTexCoords = 0;
+	if (GetRenderData() && GetRenderData()->LODResources.IsValidIndex(LODIndex))
+	{
+		NumTexCoords = GetRenderData()->LODResources[LODIndex].GetNumTexCoords();
+	}
+	return NumTexCoords;
+}
+
+int32 UStaticMesh::GetNumNaniteVertices() const
+{
+	int32 NumVertices = 0;
+	if (HasValidNaniteData())
+	{
+		const Nanite::FResources& Resources = GetRenderData()->NaniteResources;
+		if (Resources.RootClusterPage.Num() > 0)
+		{
+			NumVertices = Resources.NumInputVertices;
+		}
+	}
+	return NumVertices;
+}
+
+int32 UStaticMesh::GetNumNaniteTriangles() const
+{
+	int32 NumTriangles = 0;
+	if (HasValidNaniteData())
+	{
+		const Nanite::FResources& Resources = GetRenderData()->NaniteResources;
+		if (Resources.RootClusterPage.Num() > 0)
+		{
+			NumTriangles = Resources.NumInputTriangles;
+		}
+	}
+	return NumTriangles;
+}
+
 int32 UStaticMesh::GetNumLODs() const
 {
 	int32 NumLODs = 0;
@@ -3986,17 +4034,8 @@ void UStaticMesh::GetAssetRegistryTags(TArray<FAssetRegistryTag>& OutTags) const
 		ComplexityString = LexToString((ECollisionTraceFlag)GetBodySetup()->GetCollisionTraceFlag());
 	}
 
-	int32 NumNaniteTriangles = 0;
-	int32 NumNaniteVertices = 0;
-	if (GetRenderData())
-	{
-		const Nanite::FResources& Resources = GetRenderData()->NaniteResources;
-		if (Resources.RootClusterPage.Num() > 0)
-		{
-			NumNaniteTriangles = Resources.NumInputTriangles;
-			NumNaniteVertices = Resources.NumInputVertices;
-		}
-	}
+	int32 NumNaniteTriangles = GetNumNaniteTriangles();
+	int32 NumNaniteVertices = GetNumNaniteVertices();
 
 	OutTags.Add(FAssetRegistryTag("NaniteTriangles", FString::FromInt(NumNaniteTriangles), FAssetRegistryTag::TT_Numerical));
 	OutTags.Add(FAssetRegistryTag("NaniteVertices", FString::FromInt(NumNaniteVertices), FAssetRegistryTag::TT_Numerical));
