Partial backout of 16546734: use default max resolution of 512 pixels for cubemaps generated from longlat textures.

diff --git a/Engine/Source/Runtime/Engine/Private/TextureDerivedData.cpp b/Engine/Source/Runtime/Engine/Private/TextureDerivedData.cpp
--- a/Engine/Source/Runtime/Engine/Private/TextureDerivedData.cpp
+++ b/Engine/Source/Runtime/Engine/Private/TextureDerivedData.cpp
@@ -396,6 +396,11 @@ static void GetTextureBuildSettings(
 		OutBuildSettings.DiffuseConvolveMipLevel = GDiffuseConvolveMipLevel;
 		const UTextureCube* Cube = CastChecked<UTextureCube>(&Texture);
 		OutBuildSettings.bLongLatSource = (Cube->Source.GetNumSlices() == 1);
+		if (OutBuildSettings.bLongLatSource && Texture.MaxTextureSize <= 0)
+		{
+			// long/lat source use 512 as default
+			OutBuildSettings.MaxTextureResolution = 512;
+		}
 	}
 	else if (Texture.IsA(UTexture2DArray::StaticClass()))
 	{
