New default heightfog setup:
- by default, sky atmosphere affect height fog (r.SupportSkyAtmosphereAffectsHeightFog cvar is true by default now, cannot version renderer setting)
- by default, new heightfog components will have color contributions set to black.
- SupportSkyAtmosphereAffectsHeightFog is set to off in enginetest project to not disrupt tests.

diff --git a/Engine/Source/Runtime/Core/Public/UObject/UE5MainStreamObjectVersion.h b/Engine/Source/Runtime/Core/Public/UObject/UE5MainStreamObjectVersion.h
--- a/Engine/Source/Runtime/Core/Public/UObject/UE5MainStreamObjectVersion.h
+++ b/Engine/Source/Runtime/Core/Public/UObject/UE5MainStreamObjectVersion.h
@@ -151,6 +151,9 @@ struct CORE_API FUE5MainStreamObjectVersion
 		// Removed the AtmosphericFog component with conversion to SkyAtmosphere component
 		RemovedAtmosphericFog,
 
+		// The SkyAtmosphere now light up the heightfog by default, and by default the height fog has a black color.
+		SkyAtmosphereAffectsHeightFogWithBetterDefault,
+
 		// -----<new versions can be added above this line>-------------------------------------------------
 		VersionPlusOne,
 		LatestVersion = VersionPlusOne - 1
diff --git a/Engine/Source/Runtime/Engine/Classes/Components/ExponentialHeightFogComponent.h b/Engine/Source/Runtime/Engine/Classes/Components/ExponentialHeightFogComponent.h
--- a/Engine/Source/Runtime/Engine/Classes/Components/ExponentialHeightFogComponent.h
+++ b/Engine/Source/Runtime/Engine/Classes/Components/ExponentialHeightFogComponent.h
@@ -63,8 +63,11 @@ class ENGINE_API UExponentialHeightFogComponent : public USceneComponent
 	UPROPERTY(EditAnywhere, BlueprintReadOnly, Category=ExponentialHeightFogComponent)
 	FExponentialHeightFogData SecondFogData;
 
-	UPROPERTY(BlueprintReadOnly, interp, Category=ExponentialHeightFogComponent)
-	FLinearColor FogInscatteringColor;
+	UPROPERTY()
+	FLinearColor FogInscatteringColor_DEPRECATED;
+
+	UPROPERTY(BlueprintReadOnly, interp, Category=ExponentialHeightFogComponent, meta = (DisplayName = "Fog Inscattering Color"))
+	FLinearColor FogInscatteringLuminance;
 
 	/** 
 	 * Cubemap that can be specified for fog color, which is useful to make distant, heavily fogged scene elements match the sky.
@@ -103,12 +106,15 @@ class ENGINE_API UExponentialHeightFogComponent : public USceneComponent
 	UPROPERTY(BlueprintReadOnly, interp, Category=DirectionalInscattering)
 	float DirectionalInscatteringStartDistance;
 
+	UPROPERTY()
+	FLinearColor DirectionalInscatteringColor_DEPRECATED;
+
 	/** 
 	 * Controls the color of the directional inscattering, which is used to approximate inscattering from a directional light. 
 	 * Note: there must be a directional light with bUsedAsAtmosphereSunLight enabled for DirectionalInscattering to be used.
 	 */
-	UPROPERTY(BlueprintReadOnly, interp, Category=DirectionalInscattering)
-	FLinearColor DirectionalInscatteringColor;
+	UPROPERTY(BlueprintReadOnly, interp, Category=DirectionalInscattering, meta = (DisplayName = "Directional Inscattering Color"))
+	FLinearColor DirectionalInscatteringLuminance;
 
 	/** 
 	 * Maximum opacity of the fog.  
@@ -255,6 +261,7 @@ public:
 	virtual void PostEditChangeProperty(FPropertyChangedEvent& PropertyChangedEvent) override;
 #endif // WITH_EDITOR
 	virtual void PostInterpChange(FProperty* PropertyThatChanged) override;
+	virtual void Serialize(FArchive& Ar) override;
 	//~ End UObject Interface
 };
 
diff --git a/Engine/Source/Runtime/Engine/Private/Components/HeightFogComponent.cpp b/Engine/Source/Runtime/Engine/Private/Components/HeightFogComponent.cpp
--- a/Engine/Source/Runtime/Engine/Private/Components/HeightFogComponent.cpp
+++ b/Engine/Source/Runtime/Engine/Private/Components/HeightFogComponent.cpp
@@ -15,15 +15,18 @@
 #include "Engine/ExponentialHeightFog.h"
 #include "Net/UnrealNetwork.h"
 #include "Components/BillboardComponent.h"
+#include "UObject/UE5MainStreamObjectVersion.h"
 
 UExponentialHeightFogComponent::UExponentialHeightFogComponent(const FObjectInitializer& ObjectInitializer)
 	: Super(ObjectInitializer)
 {
-	FogInscatteringColor = FLinearColor(0.447f, 0.638f, 1.0f);
+	FogInscatteringColor_DEPRECATED = FLinearColor(0.447f, 0.638f, 1.0f);
+	FogInscatteringLuminance = FLinearColor::Black;
 
 	DirectionalInscatteringExponent = 4.0f;
 	DirectionalInscatteringStartDistance = 10000.0f;
-	DirectionalInscatteringColor = FLinearColor(0.25f, 0.25f, 0.125f);
+	DirectionalInscatteringColor_DEPRECATED = FLinearColor(0.25f, 0.25f, 0.125f);
+	DirectionalInscatteringLuminance = FLinearColor::Black;
 
 	InscatteringTextureTint = FLinearColor::White;
 	FullyDirectionalInscatteringColorDistance = 100000.0f;
@@ -85,8 +88,8 @@ bool UExponentialHeightFogComponent::CanEditChange(const FProperty* InProperty)
 
 		if (PropertyName == GET_MEMBER_NAME_STRING_CHECKED(UExponentialHeightFogComponent, DirectionalInscatteringExponent) ||
 			PropertyName == GET_MEMBER_NAME_STRING_CHECKED(UExponentialHeightFogComponent, DirectionalInscatteringStartDistance) ||
-			PropertyName == GET_MEMBER_NAME_STRING_CHECKED(UExponentialHeightFogComponent, DirectionalInscatteringColor) ||
-			PropertyName == GET_MEMBER_NAME_STRING_CHECKED(UExponentialHeightFogComponent, FogInscatteringColor))
+			PropertyName == GET_MEMBER_NAME_STRING_CHECKED(UExponentialHeightFogComponent, DirectionalInscatteringLuminance) ||
+			PropertyName == GET_MEMBER_NAME_STRING_CHECKED(UExponentialHeightFogComponent, FogInscatteringLuminance))
 		{
 			return !InscatteringColorCubemap;
 		}
@@ -137,9 +140,9 @@ void UExponentialHeightFogComponent::SetFogDensity(float Value)
 
 void UExponentialHeightFogComponent::SetFogInscatteringColor(FLinearColor Value)
 {
-	if(FogInscatteringColor != Value)
+	if(FogInscatteringLuminance != Value)
 	{
-		FogInscatteringColor = Value;
+		FogInscatteringLuminance = Value;
 		MarkRenderStateDirty();
 	}
 }
@@ -209,9 +212,9 @@ void UExponentialHeightFogComponent::SetDirectionalInscatteringStartDistance(flo
 
 void UExponentialHeightFogComponent::SetDirectionalInscatteringColor(FLinearColor Value)
 {
-	if(DirectionalInscatteringColor != Value)
+	if(DirectionalInscatteringLuminance != Value)
 	{
-		DirectionalInscatteringColor = Value;
+		DirectionalInscatteringLuminance = Value;
 		MarkRenderStateDirty();
 	}
 }
@@ -306,6 +309,19 @@ void UExponentialHeightFogComponent::SetVolumetricFogDistance(float NewValue)
 	}
 }
 
+void UExponentialHeightFogComponent::Serialize(FArchive& Ar)
+{
+	Super::Serialize(Ar);
+
+	Ar.UsingCustomVersion(FUE5MainStreamObjectVersion::GUID);
+
+	if (Ar.IsLoading() && (Ar.CustomVer(FUE5MainStreamObjectVersion::GUID) < FUE5MainStreamObjectVersion::SkyAtmosphereAffectsHeightFogWithBetterDefault))
+	{
+		FogInscatteringLuminance = FogInscatteringColor_DEPRECATED;
+		DirectionalInscatteringLuminance = DirectionalInscatteringColor_DEPRECATED;
+	}
+}
+
 //////////////////////////////////////////////////////////////////////////
 // AExponentialHeightFog
 
diff --git a/Engine/Source/Runtime/Renderer/Private/SceneCore.cpp b/Engine/Source/Runtime/Renderer/Private/SceneCore.cpp
--- a/Engine/Source/Runtime/Renderer/Private/SceneCore.cpp
+++ b/Engine/Source/Runtime/Renderer/Private/SceneCore.cpp
@@ -361,7 +361,7 @@ FExponentialHeightFogSceneInfo::FExponentialHeightFogSceneInfo(const UExponentia
 	FogCutoffDistance(InComponent->FogCutoffDistance),
 	DirectionalInscatteringExponent(InComponent->DirectionalInscatteringExponent),
 	DirectionalInscatteringStartDistance(InComponent->DirectionalInscatteringStartDistance),
-	DirectionalInscatteringColor(InComponent->DirectionalInscatteringColor)
+	DirectionalInscatteringColor(InComponent->DirectionalInscatteringLuminance)
 {
 	FogData[0].Height = InComponent->GetComponentLocation().Z;
 	FogData[1].Height = InComponent->GetComponentLocation().Z + InComponent->SecondFogData.FogHeightOffset;
@@ -373,7 +373,7 @@ FExponentialHeightFogSceneInfo::FExponentialHeightFogSceneInfo(const UExponentia
 	FogData[1].Density = InComponent->SecondFogData.FogDensity / 1000.0f;
 	FogData[1].HeightFalloff = InComponent->SecondFogData.FogHeightFalloff / 1000.0f;
 
-	FogColor = InComponent->InscatteringColorCubemap ? InComponent->InscatteringTextureTint : InComponent->FogInscatteringColor;
+	FogColor = InComponent->InscatteringColorCubemap ? InComponent->InscatteringTextureTint : InComponent->FogInscatteringLuminance;
 	InscatteringColorCubemap = InComponent->InscatteringColorCubemap;
 	InscatteringColorCubemapAngle = InComponent->InscatteringColorCubemapAngle * (PI / 180.f);
 	FullyDirectionalInscatteringColorDistance = InComponent->FullyDirectionalInscatteringColorDistance;
diff --git a/Engine/Source/Runtime/Renderer/Private/SkyAtmosphereRendering.cpp b/Engine/Source/Runtime/Renderer/Private/SkyAtmosphereRendering.cpp
--- a/Engine/Source/Runtime/Renderer/Private/SkyAtmosphereRendering.cpp
+++ b/Engine/Source/Runtime/Renderer/Private/SkyAtmosphereRendering.cpp
@@ -39,7 +39,7 @@ static TAutoConsoleVariable<int32> CVarSupportSkyAtmosphere(
 // The project setting for the sky atmosphere component to affect the height fog (disable runtime and shader code)
 static TAutoConsoleVariable<int32> CVarSupportSkyAtmosphereAffectsHeightFog(
 	TEXT("r.SupportSkyAtmosphereAffectsHeightFog"),
-	0,
+	1,
 	TEXT("Enables SkyAtmosphere affecting height fog. It requires r.SupportSkyAtmosphere to be true."),
 	ECVF_ReadOnly | ECVF_RenderThreadSafe);
 
