DataLayer: Add logging when trying to set datalayer state on unknown Data Layer

diff --git a/Engine/Source/Runtime/Engine/Private/WorldPartition/DataLayer/DataLayerSubsystem.cpp b/Engine/Source/Runtime/Engine/Private/WorldPartition/DataLayer/DataLayerSubsystem.cpp
--- a/Engine/Source/Runtime/Engine/Private/WorldPartition/DataLayer/DataLayerSubsystem.cpp
+++ b/Engine/Source/Runtime/Engine/Private/WorldPartition/DataLayer/DataLayerSubsystem.cpp
@@ -4,6 +4,7 @@
 #include "WorldPartition/DataLayer/WorldDataLayers.h"
 #include "WorldPartition/DataLayer/DataLayer.h"
 #include "WorldPartition/WorldPartitionDebugHelper.h"
+#include "WorldPartition/WorldPartition.h"
 #include "Engine/Engine.h"
 #include "Engine/World.h"
 #include "Engine/Canvas.h"
@@ -76,25 +77,53 @@ UDataLayer* UDataLayerSubsystem::GetDataLayerFromName(FName InDataLayerName) con
 
 void UDataLayerSubsystem::SetDataLayerState(const UDataLayer* InDataLayer, EDataLayerState InState)
 {
-	if (AWorldDataLayers* WorldDataLayers = GetWorld()->GetWorldDataLayers())
+	if (InDataLayer)
 	{
-		WorldDataLayers->SetDataLayerState(FActorDataLayer(InDataLayer->GetFName()), InState);
+		if (AWorldDataLayers* WorldDataLayers = GetWorld()->GetWorldDataLayers())
+		{
+			WorldDataLayers->SetDataLayerState(FActorDataLayer(InDataLayer->GetFName()), InState);
+		}
+	}
+	else
+	{
+		UE_LOG(LogWorldPartition, Warning, TEXT("UDataLayerSubsystem::SetDataLayerState called with null Data Layer"));
 	}
 }
 
 void UDataLayerSubsystem::SetDataLayerStateByName(const FName& InDataLayerName, EDataLayerState InState)
 {
-	SetDataLayerState(GetDataLayerFromName(InDataLayerName), InState);
+	if (UDataLayer* DataLayer = GetDataLayerFromName(InDataLayerName))
+	{
+		SetDataLayerState(DataLayer, InState);
+	}
+	else
+	{
+		UE_LOG(LogWorldPartition, Warning, TEXT("UDataLayerSubsystem::SetDataLayerStateByName unknown Data Layer: '%s'"), *InDataLayerName.ToString());
+	}
 }
 
 void UDataLayerSubsystem::SetDataLayerState(const FActorDataLayer& InDataLayer, EDataLayerState InState)
 {
-	SetDataLayerState(GetDataLayerFromName(InDataLayer.Name), InState);
+	if (UDataLayer* DataLayer = GetDataLayerFromName(InDataLayer.Name))
+	{
+		SetDataLayerState(DataLayer, InState);
+	}
+	else
+	{
+		UE_LOG(LogWorldPartition, Warning, TEXT("UDataLayerSubsystem::SetDataLayerState unknown Data Layer: '%s'"), *InDataLayer.Name.ToString());
+	}
 }
 
 void UDataLayerSubsystem::SetDataLayerStateByLabel(const FName& InDataLayerLabel, EDataLayerState InState)
 {
-	SetDataLayerState(GetDataLayerFromLabel(InDataLayerLabel), InState);
+	if (UDataLayer* DataLayer = GetDataLayerFromLabel(InDataLayerLabel))
+	{
+		SetDataLayerState(DataLayer, InState);
+	}
+	else
+	{
+		UE_LOG(LogWorldPartition, Warning, TEXT("UDataLayerSubsystem::SetDataLayerStateByLabel unknown Data Layer: '%s'"), *InDataLayerLabel.ToString());
+	}
 }
 
 EDataLayerState UDataLayerSubsystem::GetDataLayerState(const UDataLayer* InDataLayer) const
