Added wp.Runtime.UpdateStreamingSources console command to stop updating streaming sources (mainly for debugging purposes)

diff --git a/Engine/Source/Runtime/Engine/Private/WorldPartition/WorldPartitionStreamingPolicy.cpp b/Engine/Source/Runtime/Engine/Private/WorldPartition/WorldPartitionStreamingPolicy.cpp
--- a/Engine/Source/Runtime/Engine/Private/WorldPartition/WorldPartitionStreamingPolicy.cpp
+++ b/Engine/Source/Runtime/Engine/Private/WorldPartition/WorldPartitionStreamingPolicy.cpp
@@ -29,6 +29,12 @@
 
 #define LOCTEXT_NAMESPACE "WorldPartitionStreamingPolicy"
 
+static int32 GUpdateStreamingSources = 1;
+static FAutoConsoleVariableRef CVarUpdateStreamingSources(
+	TEXT("wp.Runtime.UpdateStreamingSources"),
+	GUpdateStreamingSources,
+	TEXT("Set to 0 to stop updating (freeze) world partition streaming sources."));
+
 static int32 GMaxLoadingStreamingCells = 4;
 static FAutoConsoleVariableRef CMaxLoadingStreamingCells(
 	TEXT("wp.Runtime.MaxLoadingStreamingCells"),
@@ -62,6 +68,11 @@ UWorldPartitionStreamingPolicy::UWorldPartitionStreamingPolicy(const FObjectInit
 
 void UWorldPartitionStreamingPolicy::UpdateStreamingSources()
 {
+	if (!GUpdateStreamingSources)
+	{
+		return;
+	}
+
 	TRACE_CPUPROFILER_EVENT_SCOPE(UWorldPartitionStreamingPolicy::UpdateStreamingSources);
 
 	StreamingSources.Reset();
