Matinee/Framework: Removed Matinee from AActor, AGameModeBase, UWorld, etc.

diff --git a/Engine/Source/Editor/UnrealEd/Classes/Editor/EditorEngine.h b/Engine/Source/Editor/UnrealEd/Classes/Editor/EditorEngine.h
--- a/Engine/Source/Editor/UnrealEd/Classes/Editor/EditorEngine.h
+++ b/Engine/Source/Editor/UnrealEd/Classes/Editor/EditorEngine.h
@@ -36,6 +36,7 @@
 
 class APlayerStart;
 class Error;
+class AMatineeActor;
 class FEditorViewportClient;
 class FEditorWorldManager;
 class FMessageLog;
diff --git a/Engine/Source/Runtime/Engine/Classes/Camera/PlayerCameraManager.h b/Engine/Source/Runtime/Engine/Classes/Camera/PlayerCameraManager.h
--- a/Engine/Source/Runtime/Engine/Classes/Camera/PlayerCameraManager.h
+++ b/Engine/Source/Runtime/Engine/Classes/Camera/PlayerCameraManager.h
@@ -421,7 +421,7 @@ public:
 	/** True to smoothly interpolate color scale values when they change. */
 	uint32 bEnableColorScaleInterp : 1;
 
-	/** True if clients are handling setting their own viewtarget and the server should not replicate it (e.g. during certain Matinee sequences) */
+	/** True if clients are handling setting their own viewtarget and the server should not replicate it. */
 	UPROPERTY(EditAnywhere, BlueprintReadWrite, Category=PlayerCameraManager)
 	uint32 bClientSimulatingViewTarget : 1;
 
diff --git a/Engine/Source/Runtime/Engine/Classes/Engine/Engine.h b/Engine/Source/Runtime/Engine/Classes/Engine/Engine.h
--- a/Engine/Source/Runtime/Engine/Classes/Engine/Engine.h
+++ b/Engine/Source/Runtime/Engine/Classes/Engine/Engine.h
@@ -23,7 +23,6 @@
 
 #define WITH_DYNAMIC_RESOLUTION (!UE_SERVER)
 
-class AMatineeActor;
 class APlayerController;
 class Error;
 class FCanvas;
@@ -1559,9 +1558,6 @@ public:
 	UPROPERTY(transient)
 	float SelectionHighlightIntensityBillboards;
 
-	/** The "outermost" active matinee, if any. */
-	TWeakObjectPtr<AMatineeActor> ActiveMatinee;
-
 	/** Delegate handling when streaming pause begins. Set initially in FStreamingPauseRenderingModule::StartupModule() but can then be overridden by games. */
 	void RegisterBeginStreamingPauseRenderingDelegate( FBeginStreamingPauseDelegate* InDelegate );
 	FBeginStreamingPauseDelegate* BeginStreamingPauseDelegate;
diff --git a/Engine/Source/Runtime/Engine/Classes/Engine/LocalPlayer.h b/Engine/Source/Runtime/Engine/Classes/Engine/LocalPlayer.h
--- a/Engine/Source/Runtime/Engine/Classes/Engine/LocalPlayer.h
+++ b/Engine/Source/Runtime/Engine/Classes/Engine/LocalPlayer.h
@@ -241,7 +241,6 @@ public:
 	bool HandleExecCommand( const TCHAR* Cmd, FOutputDevice& Ar );
 	bool HandleToggleDrawEventsCommand( const TCHAR* Cmd, FOutputDevice& Ar );
 	bool HandleToggleStreamingVolumesCommand( const TCHAR* Cmd, FOutputDevice& Ar );
-	bool HandleCancelMatineeCommand( const TCHAR* Cmd, FOutputDevice& Ar );
 	
 protected:
 	/**
diff --git a/Engine/Source/Runtime/Engine/Classes/Engine/World.h b/Engine/Source/Runtime/Engine/Classes/Engine/World.h
--- a/Engine/Source/Runtime/Engine/Classes/Engine/World.h
+++ b/Engine/Source/Runtime/Engine/Classes/Engine/World.h
@@ -36,7 +36,6 @@ class ACameraActor;
 class AController;
 class AGameModeBase;
 class AGameStateBase;
-class AMatineeActor;
 class APhysicsVolume;
 class APlayerController;
 class AWorldSettings;
@@ -3828,9 +3827,6 @@ public:
 	/** Sets world origin at specified position and stream-in all relevant levels */
 	void NavigateTo(FIntVector InLocation);
 
-	/** Gets all matinee actors for the current level */
-	void GetMatineeActors( TArray<AMatineeActor*>& OutMatineeActors );
-
 	/** Updates all physics constraint actor joint locations.  */
 	void UpdateConstraintActors();
 
diff --git a/Engine/Source/Runtime/Engine/Classes/GameFramework/Actor.h b/Engine/Source/Runtime/Engine/Classes/GameFramework/Actor.h
--- a/Engine/Source/Runtime/Engine/Classes/GameFramework/Actor.h
+++ b/Engine/Source/Runtime/Engine/Classes/GameFramework/Actor.h
@@ -28,7 +28,6 @@
 
 class AActor;
 class AController;
-class AMatineeActor;
 class APawn;
 class APlayerController;
 class UActorChannel;
@@ -734,10 +733,6 @@ protected:
 	FVector PivotOffset;
 #endif
 
-	/** The matinee actors that control this actor. */
-	UPROPERTY(Transient)
-	TArray<TObjectPtr<class AMatineeActor>> ControllingMatineeActors;
-
 	/** Handle for efficient management of LifeSpanExpired timer */
 	FTimerHandle TimerHandle_LifeSpanExpired;
 
@@ -2470,9 +2465,6 @@ public:
 	UFUNCTION(BlueprintCallable, Category="Collision", meta=(UnsafeDuringActorConstruction="true"))
 	bool IsOverlappingActor(const AActor* Other) const;
 
-	/** Returns whether a MatineeActor is currently controlling this Actor */
-	bool IsMatineeControlled() const;
-
 	/** See if the root component has ModifyFrequency of MF_Static */
 	bool IsRootComponentStatic() const;
 
@@ -2564,17 +2556,6 @@ public:
 	/** Allow actors to initialize themselves on the C++ side after all of their components have been initialized, only called during gameplay */
 	virtual void PostInitializeComponents();
 
-	/**
-	 * Adds a controlling matinee actor for use during matinee playback
-	 * @param InMatineeActor	The matinee actor which controls this actor
-	 */
-	void AddControllingMatineeActor( AMatineeActor& InMatineeActor );
-
-	/**
-	 * Removes a controlling matinee actor
-	 * @param InMatineeActor	The matinee actor which currently controls this actor
-	 */
-	void RemoveControllingMatineeActor( AMatineeActor& InMatineeActor );
 
 	/** Dispatches ReceiveHit virtual and OnComponentHit delegate */
 	virtual void DispatchPhysicsCollisionHit(const struct FRigidBodyCollisionInfo& MyInfo, const struct FRigidBodyCollisionInfo& OtherInfo, const FCollisionImpactData& RigidCollisionData);
diff --git a/Engine/Source/Runtime/Engine/Classes/GameFramework/GameMode.h b/Engine/Source/Runtime/Engine/Classes/GameFramework/GameMode.h
--- a/Engine/Source/Runtime/Engine/Classes/GameFramework/GameMode.h
+++ b/Engine/Source/Runtime/Engine/Classes/GameFramework/GameMode.h
@@ -192,9 +192,6 @@ public:
 	/** SetViewTarget of player control on server change */
 	virtual void SetSeamlessTravelViewTarget(APlayerController* PC);
 
-	/** Called when this PC is in cinematic mode, and its matinee is canceled by the user. */
-	virtual void MatineeCancelled();
-
 	/**
 	 * Called from CommitMapChange before unloading previous level. Used for asynchronous level streaming
 	 * @param PreviousMapName - Name of the previous persistent level
diff --git a/Engine/Source/Runtime/Engine/Classes/GameFramework/GameModeBase.h b/Engine/Source/Runtime/Engine/Classes/GameFramework/GameModeBase.h
--- a/Engine/Source/Runtime/Engine/Classes/GameFramework/GameModeBase.h
+++ b/Engine/Source/Runtime/Engine/Classes/GameFramework/GameModeBase.h
@@ -475,7 +475,7 @@ public:
 protected:
 
 	/**
-	 *	Check to see if we should start in cinematic mode (e.g. matinee movie capture).
+	 *	Check to see if we should start in cinematic mode 
 	 * 	@param	OutHidePlayer		Whether to hide the player
 	 *	@param	OutHideHud			Whether to hide the HUD		
 	 *	@param	OutDisableMovement	Whether to disable movement
diff --git a/Engine/Source/Runtime/Engine/Classes/GameFramework/PlayerController.h b/Engine/Source/Runtime/Engine/Classes/GameFramework/PlayerController.h
--- a/Engine/Source/Runtime/Engine/Classes/GameFramework/PlayerController.h
+++ b/Engine/Source/Runtime/Engine/Classes/GameFramework/PlayerController.h
@@ -806,8 +806,6 @@ public:
 	 */
 	virtual bool IsPlayerMuted(const class FUniqueNetId& PlayerId);
 
-	/** Notification when a matinee director track starts or stops controlling the ViewTarget of this PlayerController */
-	virtual void NotifyDirectorControl(bool bNowControlling, class AMatineeActor* CurrentMatinee);
 
 	/** Console control commands, useful when remote debugging so you can't touch the console the normal way */
 	UFUNCTION(exec)
@@ -1653,16 +1651,6 @@ public:
 	 */
 	void BuildHiddenComponentList(const FVector& ViewLocation, TSet<FPrimitiveComponentId>& HiddenComponentsOut);
 
-	/**
-	 * Sets the Matinee director track instance that's currently possessing this player controller
-	 * @param   NewControllingDirector    The director track instance that's now controlling this player controller (or nullptr for none)
-	 * @param	bClientSimulatingViewTarget	True to allow clients to simulate their own camera cuts (ignored if NewControllingDirector is nullptr).
-	 */
-	void SetControllingDirector(UInterpTrackInstDirector* NewControllingDirector, bool bClientSimulatingViewTarget);
-
-	/** Returns the Matinee director track that's currently possessing this player controller, or nullptr for none */
-	UInterpTrackInstDirector* GetControllingDirector();
-
 	/** spawn cameras for servers and owning players */
 	virtual void SpawnPlayerCameraManager();
 
diff --git a/Engine/Source/Runtime/Engine/Private/Actor.cpp b/Engine/Source/Runtime/Engine/Private/Actor.cpp
--- a/Engine/Source/Runtime/Engine/Private/Actor.cpp
+++ b/Engine/Source/Runtime/Engine/Private/Actor.cpp
@@ -28,9 +28,6 @@
 #include "Logging/MessageLog.h"
 #include "Net/UnrealNetwork.h"
 #include "Net/RepLayout.h"
-#include "Matinee/MatineeActor.h"
-#include "Matinee/InterpGroup.h"
-#include "Matinee/InterpGroupInst.h"
 #include "Engine/Canvas.h"
 #include "DisplayDebugHelpers.h"
 #include "Animation/AnimInstance.h"
@@ -682,28 +679,6 @@ void AActor::SetTickableWhenPaused(bool bTickableWhenPaused)
 	PrimaryActorTick.bTickEvenWhenPaused = bTickableWhenPaused;
 }
 
-void AActor::AddControllingMatineeActor( AMatineeActor& InMatineeActor )
-{
-	if (RootComponent)
-	{
-		RootComponent->PrimaryComponentTick.AddPrerequisite(&InMatineeActor, InMatineeActor.PrimaryActorTick);
-	}
-
-	PrimaryActorTick.AddPrerequisite(&InMatineeActor, InMatineeActor.PrimaryActorTick);
-	ControllingMatineeActors.AddUnique(&InMatineeActor);
-}
-
-void AActor::RemoveControllingMatineeActor( AMatineeActor& InMatineeActor )
-{
-	if (RootComponent)
-	{
-		RootComponent->PrimaryComponentTick.RemovePrerequisite(&InMatineeActor, InMatineeActor.PrimaryActorTick);
-	}
-
-	PrimaryActorTick.RemovePrerequisite(&InMatineeActor, InMatineeActor.PrimaryActorTick);
-	ControllingMatineeActors.RemoveSwap(&InMatineeActor);
-}
-
 void AActor::BeginDestroy()
 {
 	UnregisterAllComponents();
@@ -2213,33 +2188,6 @@ bool AActor::IsInPersistentLevel(bool bIncludeLevelStreamingPersistent) const
 														World->GetStreamingLevels()[0]->GetLoadedLevel() == MyLevel ) );
 }
 
-
-bool AActor::IsMatineeControlled() const 
-{
-	bool bMovedByMatinee = false;
-	for(auto It(ControllingMatineeActors.CreateConstIterator()); It; It++)
-	{
-		AMatineeActor* ControllingMatineeActor = *It;
-		if(ControllingMatineeActor != nullptr)
-		{
-			UInterpGroupInst* GroupInst = ControllingMatineeActor->FindGroupInst(this);
-			if(GroupInst != nullptr)
-			{
-				if(GroupInst->Group && GroupInst->Group->HasMoveTrack())
-				{
-					bMovedByMatinee = true;
-					break;
-				}
-			}
-			else
-			{
-				UE_LOG(LogActor, Log, TEXT("IsMatineeControlled: ControllingMatineeActor is set but no GroupInstance (%s)"), *GetPathName());
-			}
-		}
-	}
-	return bMovedByMatinee;
-}
-
 bool AActor::IsRootComponentStatic() const
 {
 	return(RootComponent != nullptr && RootComponent->Mobility == EComponentMobility::Static);
diff --git a/Engine/Source/Runtime/Engine/Private/CharacterMovementComponentAsync.cpp b/Engine/Source/Runtime/Engine/Private/CharacterMovementComponentAsync.cpp
--- a/Engine/Source/Runtime/Engine/Private/CharacterMovementComponentAsync.cpp
+++ b/Engine/Source/Runtime/Engine/Private/CharacterMovementComponentAsync.cpp
@@ -367,7 +367,7 @@ void FCharacterMovementComponentAsyncInput::PerformMovement(float DeltaSeconds,
 		// Update character state based on change from movement
 		UpdateCharacterStateAfterMovement(DeltaSeconds, Output);
 
-		if ((bAllowPhysicsRotationDuringAnimRootMotion || !Output.RootMotionParams.bHasRootMotion)/* && !CharacterOwner->IsMatineeControlled()*/)
+		if ((bAllowPhysicsRotationDuringAnimRootMotion || !Output.RootMotionParams.bHasRootMotion))
 		{
 			PhysicsRotation(DeltaSeconds, Output);
 		}
@@ -602,94 +602,82 @@ void FCharacterMovementComponentAsyncInput::UpdateBasedMovement(float DeltaSecon
 		const FQuatRotationTranslationMatrix OldLocalToWorld(OldBaseQuat, OldBaseLocation);
 		const FQuatRotationTranslationMatrix NewLocalToWorld(NewBaseQuat, NewBaseLocation);
 
+		FQuat FinalQuat = UpdatedComponentInput->GetRotation();
 
-		// TODO Matinee
-		if (false)//CharacterOwner->IsMatineeControlled())
+		if (bRotationChanged && !bIgnoreBaseRotation)
 		{
-			/*FRotationTranslationMatrix HardRelMatrix(CharacterOwner->GetBasedMovement().Rotation, CharacterOwner->GetBasedMovement().Location);
-			const FMatrix NewWorldTM = HardRelMatrix * NewLocalToWorld;
-			const FQuat NewWorldRot = bIgnoreBaseRotation ? UpdatedComponent->GetComponentQuat() : NewWorldTM.ToQuat();
-			MoveUpdatedComponent(NewWorldTM.GetOrigin() - UpdatedComponent->GetComponentLocation(), NewWorldRot, true);*/
-		}
-		else
-		{
-			FQuat FinalQuat = UpdatedComponentInput->GetRotation();
+			// Apply change in rotation and pipe through FaceRotation to maintain axis restrictions
+			const FQuat PawnOldQuat = UpdatedComponentInput->GetRotation();
+			const FQuat TargetQuat = Output.DeltaQuat * FinalQuat;
+			FRotator TargetRotator(TargetQuat);
 
-			if (bRotationChanged && !bIgnoreBaseRotation)
-			{
-				// Apply change in rotation and pipe through FaceRotation to maintain axis restrictions
-				const FQuat PawnOldQuat = UpdatedComponentInput->GetRotation();
-				const FQuat TargetQuat = Output.DeltaQuat * FinalQuat;
-				FRotator TargetRotator(TargetQuat);
+			// Do we need this value after all?
 
-				// Do we need this value after all?
+			CharacterInput->FaceRotation(TargetRotator, 0.0f, *this, Output);
+			FinalQuat =  Output.CharacterOutput->Rotation.Quaternion();//UpdatedComponent->GetComponentQuat(); supposed to be modified by MockFaceRotation, si this ok?
 
-				CharacterInput->FaceRotation(TargetRotator, 0.0f, *this, Output);
-				FinalQuat =  Output.CharacterOutput->Rotation.Quaternion();//UpdatedComponent->GetComponentQuat(); supposed to be modified by MockFaceRotation, si this ok?
 
-
-				if (PawnOldQuat.Equals(FinalQuat, 1e-6f))
+			if (PawnOldQuat.Equals(FinalQuat, 1e-6f))
+			{
+				// Nothing changed. This means we probably are using another rotation mechanism (bOrientToMovement etc). We should still follow the base object.
+				// @todo: This assumes only Yaw is used, currently a valid assumption. This is the only reason FaceRotation() is used above really, aside from being a virtual hook.
+				if (bOrientRotationToMovement || (bUseControllerDesiredRotation /*&& CharacterOwner->Controller*/))
 				{
-					// Nothing changed. This means we probably are using another rotation mechanism (bOrientToMovement etc). We should still follow the base object.
-					// @todo: This assumes only Yaw is used, currently a valid assumption. This is the only reason FaceRotation() is used above really, aside from being a virtual hook.
-					if (bOrientRotationToMovement || (bUseControllerDesiredRotation /*&& CharacterOwner->Controller*/))
-					{
-						TargetRotator.Pitch = 0.f;
-						TargetRotator.Roll = 0.f;
+					TargetRotator.Pitch = 0.f;
+					TargetRotator.Roll = 0.f;
 
-						MoveUpdatedComponent(FVector::ZeroVector, FQuat(TargetRotator), /*bSweep=*/false, Output);
-						//FinalQuat = UpdatedComponent->GetComponentQuat();
-						FinalQuat = UpdatedComponentInput->GetRotation();
-					}
+					MoveUpdatedComponent(FVector::ZeroVector, FQuat(TargetRotator), /*bSweep=*/false, Output);
+					//FinalQuat = UpdatedComponent->GetComponentQuat();
+					FinalQuat = UpdatedComponentInput->GetRotation();
 				}
-
-				// TODO Camera
-				// Pipe through ControlRotation, to affect camera.
-				/*if (CharacterOwner->Controller)
-				{
-					const FQuat PawnDeltaRotation = FinalQuat * PawnOldQuat.Inverse();
-					FRotator FinalRotation = FinalQuat.Rotator();
-					UpdateBasedRotation(FinalRotation, PawnDeltaRotation.Rotator());
-					FinalQuat = UpdatedComponent->GetComponentQuat();
-				}*/
 			}
 
-			// We need to offset the base of the character here, not its origin, so offset by half height
-			float HalfHeight = Output.ScaledCapsuleHalfHeight;
-			float Radius = Output.ScaledCapsuleRadius;
+			// TODO Camera
+			// Pipe through ControlRotation, to affect camera.
+			/*if (CharacterOwner->Controller)
+			{
+				const FQuat PawnDeltaRotation = FinalQuat * PawnOldQuat.Inverse();
+				FRotator FinalRotation = FinalQuat.Rotator();
+				UpdateBasedRotation(FinalRotation, PawnDeltaRotation.Rotator());
+				FinalQuat = UpdatedComponent->GetComponentQuat();
+			}*/
+		}
 
-			FVector const BaseOffset(0.0f, 0.0f, HalfHeight);
-			FVector const LocalBasePos = OldLocalToWorld.InverseTransformPosition(UpdatedComponentInput->GetPosition() - BaseOffset);
-			FVector const NewWorldPos = ConstrainLocationToPlane(NewLocalToWorld.TransformPosition(LocalBasePos) + BaseOffset);
-			Output.DeltaPosition = ConstrainDirectionToPlane(NewWorldPos - UpdatedComponentInput->GetPosition());
+		// We need to offset the base of the character here, not its origin, so offset by half height
+		float HalfHeight = Output.ScaledCapsuleHalfHeight;
+		float Radius = Output.ScaledCapsuleRadius;
 
-			// move attached actor
-			if (false)//bFastAttachedMove) // TODO bFastAttachedMove
-			{
-				// we're trusting no other obstacle can prevent the move here
-				//	UpdatedComponent->SetWorldLocationAndRotation(NewWorldPos, FinalQuat, false);
+		FVector const BaseOffset(0.0f, 0.0f, HalfHeight);
+		FVector const LocalBasePos = OldLocalToWorld.InverseTransformPosition(UpdatedComponentInput->GetPosition() - BaseOffset);
+		FVector const NewWorldPos = ConstrainLocationToPlane(NewLocalToWorld.TransformPosition(LocalBasePos) + BaseOffset);
+		Output.DeltaPosition = ConstrainDirectionToPlane(NewWorldPos - UpdatedComponentInput->GetPosition());
 
-				UpdatedComponentInput->SetPosition(NewWorldPos);
-				UpdatedComponentInput->SetRotation(FinalQuat);
-			}
-			else
+		// move attached actor
+		if (false)//bFastAttachedMove) // TODO bFastAttachedMove
+		{
+			// we're trusting no other obstacle can prevent the move here
+			//	UpdatedComponent->SetWorldLocationAndRotation(NewWorldPos, FinalQuat, false);
+
+			UpdatedComponentInput->SetPosition(NewWorldPos);
+			UpdatedComponentInput->SetRotation(FinalQuat);
+		}
+		else
+		{
+			// hack - transforms between local and world space introducing slight error FIXMESTEVE - discuss with engine team: just skip the transforms if no rotation?
+			FVector BaseMoveDelta = NewBaseLocation - OldBaseLocation;
+			if (!bRotationChanged && (BaseMoveDelta.X == 0.f) && (BaseMoveDelta.Y == 0.f))
 			{
-				// hack - transforms between local and world space introducing slight error FIXMESTEVE - discuss with engine team: just skip the transforms if no rotation?
-				FVector BaseMoveDelta = NewBaseLocation - OldBaseLocation;
-				if (!bRotationChanged && (BaseMoveDelta.X == 0.f) && (BaseMoveDelta.Y == 0.f))
-				{
-					Output.DeltaPosition.X = 0.f;
-					Output.DeltaPosition.Y = 0.f;
-				}
+				Output.DeltaPosition.X = 0.f;
+				Output.DeltaPosition.Y = 0.f;
+			}
 
-				FHitResult MoveOnBaseHit(1.f);
-				const FVector OldLocation = UpdatedComponentInput->GetPosition();
-				MoveUpdatedComponent(Output.DeltaPosition, FinalQuat, true, Output, &MoveOnBaseHit);
-				if ((UpdatedComponentInput->GetPosition() - (OldLocation + Output.DeltaPosition)).IsNearlyZero() == false)
-				{
-					// TODO OnUnableToFollowBaseMove
-					//OnUnableToFollowBaseMove(DeltaPosition, OldLocation, MoveOnBaseHit);
-				}
+			FHitResult MoveOnBaseHit(1.f);
+			const FVector OldLocation = UpdatedComponentInput->GetPosition();
+			MoveUpdatedComponent(Output.DeltaPosition, FinalQuat, true, Output, &MoveOnBaseHit);
+			if ((UpdatedComponentInput->GetPosition() - (OldLocation + Output.DeltaPosition)).IsNearlyZero() == false)
+			{
+				// TODO OnUnableToFollowBaseMove
+				//OnUnableToFollowBaseMove(DeltaPosition, OldLocation, MoveOnBaseHit);
 			}
 		}
 
diff --git a/Engine/Source/Runtime/Engine/Private/Components/CharacterMovementComponent.cpp b/Engine/Source/Runtime/Engine/Private/Components/CharacterMovementComponent.cpp
--- a/Engine/Source/Runtime/Engine/Private/Components/CharacterMovementComponent.cpp
+++ b/Engine/Source/Runtime/Engine/Private/Components/CharacterMovementComponent.cpp
@@ -1690,7 +1690,7 @@ void UCharacterMovementComponent::SimulatedTick(float DeltaSeconds)
 			// Avoid moving the mesh during movement if SmoothClientPosition will take care of it.
 			{
 				const FScopedPreventAttachedComponentMove PreventMeshMovement(bPreventMeshMovement ? Mesh : nullptr);
-				if (CharacterOwner->IsMatineeControlled() || CharacterOwner->IsPlayingRootMotion())
+				if (CharacterOwner->IsPlayingRootMotion())
 				{
 					PerformMovement(DeltaSeconds);
 				}
@@ -2108,81 +2108,71 @@ void UCharacterMovementComponent::UpdateBasedMovement(float DeltaSeconds)
 		const FQuatRotationTranslationMatrix OldLocalToWorld(OldBaseQuat, OldBaseLocation);
 		const FQuatRotationTranslationMatrix NewLocalToWorld(NewBaseQuat, NewBaseLocation);
 
-		if( CharacterOwner->IsMatineeControlled() )
-		{
-			FRotationTranslationMatrix HardRelMatrix(CharacterOwner->GetBasedMovement().Rotation, CharacterOwner->GetBasedMovement().Location);
-			const FMatrix NewWorldTM = HardRelMatrix * NewLocalToWorld;
-			const FQuat NewWorldRot = bIgnoreBaseRotation ? UpdatedComponent->GetComponentQuat() : NewWorldTM.ToQuat();
-			MoveUpdatedComponent( NewWorldTM.GetOrigin() - UpdatedComponent->GetComponentLocation(), NewWorldRot, true );
-		}
-		else
-		{
-			FQuat FinalQuat = UpdatedComponent->GetComponentQuat();
+		FQuat FinalQuat = UpdatedComponent->GetComponentQuat();
 			
-			if (bRotationChanged && !bIgnoreBaseRotation)
-			{
-				// Apply change in rotation and pipe through FaceRotation to maintain axis restrictions
-				const FQuat PawnOldQuat = UpdatedComponent->GetComponentQuat();
-				const FQuat TargetQuat = DeltaQuat * FinalQuat;
-				FRotator TargetRotator(TargetQuat);
-				CharacterOwner->FaceRotation(TargetRotator, 0.f);
-				FinalQuat = UpdatedComponent->GetComponentQuat();
-
-				if (PawnOldQuat.Equals(FinalQuat, 1e-6f))
-				{
-					// Nothing changed. This means we probably are using another rotation mechanism (bOrientToMovement etc). We should still follow the base object.
-					// @todo: This assumes only Yaw is used, currently a valid assumption. This is the only reason FaceRotation() is used above really, aside from being a virtual hook.
-					if (bOrientRotationToMovement || (bUseControllerDesiredRotation && CharacterOwner->Controller))
-					{
-						TargetRotator.Pitch = 0.f;
-						TargetRotator.Roll = 0.f;
-						MoveUpdatedComponent(FVector::ZeroVector, TargetRotator, false);
-						FinalQuat = UpdatedComponent->GetComponentQuat();
-					}
-				}
+		if (bRotationChanged && !bIgnoreBaseRotation)
+		{
+			// Apply change in rotation and pipe through FaceRotation to maintain axis restrictions
+			const FQuat PawnOldQuat = UpdatedComponent->GetComponentQuat();
+			const FQuat TargetQuat = DeltaQuat * FinalQuat;
+			FRotator TargetRotator(TargetQuat);
+			CharacterOwner->FaceRotation(TargetRotator, 0.f);
+			FinalQuat = UpdatedComponent->GetComponentQuat();
 
-				// Pipe through ControlRotation, to affect camera.
-				if (CharacterOwner->Controller)
+			if (PawnOldQuat.Equals(FinalQuat, 1e-6f))
+			{
+				// Nothing changed. This means we probably are using another rotation mechanism (bOrientToMovement etc). We should still follow the base object.
+				// @todo: This assumes only Yaw is used, currently a valid assumption. This is the only reason FaceRotation() is used above really, aside from being a virtual hook.
+				if (bOrientRotationToMovement || (bUseControllerDesiredRotation && CharacterOwner->Controller))
 				{
-					const FQuat PawnDeltaRotation = FinalQuat * PawnOldQuat.Inverse();
-					FRotator FinalRotation = FinalQuat.Rotator();
-					UpdateBasedRotation(FinalRotation, PawnDeltaRotation.Rotator());
+					TargetRotator.Pitch = 0.f;
+					TargetRotator.Roll = 0.f;
+					MoveUpdatedComponent(FVector::ZeroVector, TargetRotator, false);
 					FinalQuat = UpdatedComponent->GetComponentQuat();
 				}
 			}
 
-			// We need to offset the base of the character here, not its origin, so offset by half height
-			float HalfHeight, Radius;
-			CharacterOwner->GetCapsuleComponent()->GetScaledCapsuleSize(Radius, HalfHeight);
+			// Pipe through ControlRotation, to affect camera.
+			if (CharacterOwner->Controller)
+			{
+				const FQuat PawnDeltaRotation = FinalQuat * PawnOldQuat.Inverse();
+				FRotator FinalRotation = FinalQuat.Rotator();
+				UpdateBasedRotation(FinalRotation, PawnDeltaRotation.Rotator());
+				FinalQuat = UpdatedComponent->GetComponentQuat();
+			}
+		}
+
+		// We need to offset the base of the character here, not its origin, so offset by half height
+		float HalfHeight, Radius;
+		CharacterOwner->GetCapsuleComponent()->GetScaledCapsuleSize(Radius, HalfHeight);
 
-			FVector const BaseOffset(0.0f, 0.0f, HalfHeight);
-			FVector const LocalBasePos = OldLocalToWorld.InverseTransformPosition(UpdatedComponent->GetComponentLocation() - BaseOffset);
-			FVector const NewWorldPos = ConstrainLocationToPlane(NewLocalToWorld.TransformPosition(LocalBasePos) + BaseOffset);
-			DeltaPosition = ConstrainDirectionToPlane(NewWorldPos - UpdatedComponent->GetComponentLocation());
+		FVector const BaseOffset(0.0f, 0.0f, HalfHeight);
+		FVector const LocalBasePos = OldLocalToWorld.InverseTransformPosition(UpdatedComponent->GetComponentLocation() - BaseOffset);
+		FVector const NewWorldPos = ConstrainLocationToPlane(NewLocalToWorld.TransformPosition(LocalBasePos) + BaseOffset);
+		DeltaPosition = ConstrainDirectionToPlane(NewWorldPos - UpdatedComponent->GetComponentLocation());
 
-			// move attached actor
-			if (bFastAttachedMove)
+		// move attached actor
+		if (bFastAttachedMove)
+		{
+			// we're trusting no other obstacle can prevent the move here
+			UpdatedComponent->SetWorldLocationAndRotation(NewWorldPos, FinalQuat, false);
+		}
+		else
+		{
+			// hack - transforms between local and world space introducing slight error FIXMESTEVE - discuss with engine team: just skip the transforms if no rotation?
+			FVector BaseMoveDelta = NewBaseLocation - OldBaseLocation;
+			if (!bRotationChanged && (BaseMoveDelta.X == 0.f) && (BaseMoveDelta.Y == 0.f))
 			{
-				// we're trusting no other obstacle can prevent the move here
-				UpdatedComponent->SetWorldLocationAndRotation(NewWorldPos, FinalQuat, false);
+				DeltaPosition.X = 0.f;
+				DeltaPosition.Y = 0.f;
 			}
-			else
-			{
-				// hack - transforms between local and world space introducing slight error FIXMESTEVE - discuss with engine team: just skip the transforms if no rotation?
-				FVector BaseMoveDelta = NewBaseLocation - OldBaseLocation;
-				if (!bRotationChanged && (BaseMoveDelta.X == 0.f) && (BaseMoveDelta.Y == 0.f))
-				{
-					DeltaPosition.X = 0.f;
-					DeltaPosition.Y = 0.f;
-				}
 
-				FHitResult MoveOnBaseHit(1.f);
-				const FVector OldLocation = UpdatedComponent->GetComponentLocation();
-				MoveUpdatedComponent(DeltaPosition, FinalQuat, true, &MoveOnBaseHit);
-				if ((UpdatedComponent->GetComponentLocation() - (OldLocation + DeltaPosition)).IsNearlyZero() == false)
-				{
-					OnUnableToFollowBaseMove(DeltaPosition, OldLocation, MoveOnBaseHit);
-				}
+			FHitResult MoveOnBaseHit(1.f);
+			const FVector OldLocation = UpdatedComponent->GetComponentLocation();
+			MoveUpdatedComponent(DeltaPosition, FinalQuat, true, &MoveOnBaseHit);
+			if ((UpdatedComponent->GetComponentLocation() - (OldLocation + DeltaPosition)).IsNearlyZero() == false)
+			{
+				OnUnableToFollowBaseMove(DeltaPosition, OldLocation, MoveOnBaseHit);
 			}
 		}
 
@@ -2490,7 +2480,7 @@ void UCharacterMovementComponent::PerformMovement(float DeltaSeconds)
 		// Update character state based on change from movement
 		UpdateCharacterStateAfterMovement(DeltaSeconds);
 
-		if ((bAllowPhysicsRotationDuringAnimRootMotion || !HasAnimRootMotion()) && !CharacterOwner->IsMatineeControlled())
+		if (bAllowPhysicsRotationDuringAnimRootMotion || !HasAnimRootMotion())
 		{
 			PhysicsRotation(DeltaSeconds);
 		}
@@ -2652,7 +2642,7 @@ void UCharacterMovementComponent::SaveBaseLocation()
 		// Read transforms into OldBaseLocation, OldBaseQuat. Do this regardless of whether the object is movable, since mobility can change.
 		MovementBaseUtility::GetMovementBaseTransform(MovementBase, CharacterOwner->GetBasedMovement().BoneName, OldBaseLocation, OldBaseQuat);
 
-		if (MovementBaseUtility::UseRelativeLocation(MovementBase) && !CharacterOwner->IsMatineeControlled())
+		if (MovementBaseUtility::UseRelativeLocation(MovementBase))
 		{
 			// Relative Location
 			const FVector RelativeLocation = UpdatedComponent->GetComponentLocation() - OldBaseLocation;
@@ -3486,10 +3476,6 @@ bool UCharacterMovementComponent::CanStartPathFollowing() const
 		{
 			return false;
 		}
-		else if (CharacterOwner->IsMatineeControlled())
-		{
-			return false;
-		}
 	}
 
 	return Super::CanStartPathFollowing();
@@ -6960,10 +6946,6 @@ FString UCharacterMovementComponent::GetMovementName() const
 		{
 			return TEXT("Rigid Body");
 		}
-		else if ( CharacterOwner->IsMatineeControlled() )
-		{
-			return TEXT("Matinee");
-		}
 	}
 
 	// Using character movement
@@ -12073,7 +12055,6 @@ void UCharacterMovementComponent::FillAsyncInput(const FVector& InputVector, FCh
 		return;
 	}
 
-	ensure(CharacterOwner->IsMatineeControlled() == false); // See UpdateBasedMovement
 	ensure(UpdatedComponent); // assumed to exist in  MockMoveUpdatedComponent
 	ensure(HasRootMotionSources() == false); // Skipped some impl in PerformMovement, is this required for first test?
 	ensure(CharacterOwner->GetCapsuleComponent()); // check in FindFloor
diff --git a/Engine/Source/Runtime/Engine/Private/GameMode.cpp b/Engine/Source/Runtime/Engine/Private/GameMode.cpp
--- a/Engine/Source/Runtime/Engine/Private/GameMode.cpp
+++ b/Engine/Source/Runtime/Engine/Private/GameMode.cpp
@@ -66,7 +66,6 @@ void AGameMode::InitGame(const FString& MapName, const FString& Options, FString
 	}
 
 	// Bind to delegates
-	FGameDelegates::Get().GetMatineeCancelledDelegate().AddUObject(this, &AGameMode::MatineeCancelled);
 	FGameDelegates::Get().GetPendingConnectionLostDelegate().AddUObject(this, &AGameMode::NotifyPendingConnectionLost);
 	FGameDelegates::Get().GetPreCommitMapChangeDelegate().AddUObject(this, &AGameMode::PreCommitMapChange);
 	FGameDelegates::Get().GetPostCommitMapChangeDelegate().AddUObject(this, &AGameMode::PostCommitMapChange);
@@ -797,8 +796,6 @@ bool AGameMode::IsHandlingReplays()
 	return bHandleDedicatedServerReplays && GetNetMode() == ENetMode::NM_DedicatedServer;
 }
 
-void AGameMode::MatineeCancelled() {}
-
 void AGameMode::PreCommitMapChange(const FString& PreviousMapName, const FString& NextMapName) {}
 
 void AGameMode::PostCommitMapChange() {}
diff --git a/Engine/Source/Runtime/Engine/Private/GameModeBase.cpp b/Engine/Source/Runtime/Engine/Private/GameModeBase.cpp
--- a/Engine/Source/Runtime/Engine/Private/GameModeBase.cpp
+++ b/Engine/Source/Runtime/Engine/Private/GameModeBase.cpp
@@ -2,7 +2,6 @@
 
 #include "GameFramework/GameModeBase.h"
 #include "GameFramework/GameNetworkManager.h"
-#include "Matinee/MatineeActor.h"
 #include "Engine/LevelScriptActor.h"
 #include "Engine/World.h"
 #include "Misc/CommandLine.h"
@@ -948,19 +947,11 @@ void AGameModeBase::GenericPlayerInitialization(AController* C)
 
 		bool HidePlayer = false, HideHUD = false, DisableMovement = false, DisableTurning = false;
 
-		// Check to see if we should start in cinematic mode (matinee movie capture)
+		// Check to see if we should start in cinematic mode
 		if (ShouldStartInCinematicMode(PC, HidePlayer, HideHUD, DisableMovement, DisableTurning))
 		{
 			PC->SetCinematicMode(true, HidePlayer, HideHUD, DisableMovement, DisableTurning);
 		}
-
-		// Add the player to any matinees running so that it gets in on any cinematics already running, etc
-		TArray<AMatineeActor*> AllMatineeActors;
-		GetWorld()->GetMatineeActors(AllMatineeActors);
-		for (int32 i = 0; i < AllMatineeActors.Num(); i++)
-		{
-			AllMatineeActors[i]->AddPlayerToDirectorTracks(PC);
-		}
 	}
 }
 
diff --git a/Engine/Source/Runtime/Engine/Private/Interpolation.cpp b/Engine/Source/Runtime/Engine/Private/Interpolation.cpp
--- a/Engine/Source/Runtime/Engine/Private/Interpolation.cpp
+++ b/Engine/Source/Runtime/Engine/Private/Interpolation.cpp
@@ -876,11 +876,7 @@ void AMatineeActor::InitInterp()
 
 	if( MatineeData )
 	{
-		// Register myself as the active matinee if one is not active.
-		if (GEngine->ActiveMatinee.IsValid() == false)
-		{
-			GEngine->ActiveMatinee = this;
-		}
+
 
 		TMap<FName, FInterpGroupActorInfo*> InterpGroupToActorInfoMap;
 	
@@ -976,14 +972,6 @@ PRAGMA_ENABLE_DEPRECATION_WARNINGS
 			}
 		}
 
-		// set matinee actor when initialize it, otherwise, we'll have random tick order
-		for(int32 i=0; i<GroupInst.Num(); i++)
-		{
-			if(GroupInst[i]->GroupActor)
-			{
-				GroupInst[i]->GroupActor->AddControllingMatineeActor( *this );
-			}
-		}
 
 		EnableCinematicMode(true);
 	}
@@ -1000,22 +988,12 @@ void AMatineeActor::TermInterp()
 	// Destroy each group instance.
 	for(int32 i=0; i<GroupInst.Num(); i++)
 	{
-		if(GroupInst[i]->GroupActor)
-		{
-			GroupInst[i]->GroupActor->RemoveControllingMatineeActor( *this );
-		}
 
 		GroupInst[i]->TermGroupInst(true);
 
 	}
 	GroupInst.Empty();
 
-	// Unregister myself as the active matinee if one is not active.
-	if (GEngine->ActiveMatinee.Get() == this)
-	{
-		GEngine->ActiveMatinee.Reset();
-	}
-
 	// disable cinematic mode
 	EnableCinematicMode(false);
 }
@@ -1490,11 +1468,6 @@ void AMatineeActor::DeleteActorGroupInfo(class UInterpGroup * Group, AActor* Act
 			{
 				if (GroupInfo.Actors[ActorID] == ActorToDelete)
 				{
-					if (ActorToDelete)
-					{
-						ActorToDelete->RemoveControllingMatineeActor(*this);
-					}
-
 					GroupInfo.Actors.RemoveAt(ActorID);
 					return;
 				}
@@ -1512,15 +1485,6 @@ void AMatineeActor::DeleteGroupinfo(class UInterpGroup * GroupToDelete)
 
 		if(GroupToDelete->GetFName() == Info.ObjectName)
 		{
-			for (int32 ActorIndex = 0; ActorIndex<Info.Actors.Num(); ++ActorIndex)
-			{
-				if ( Info.Actors[ActorIndex] )
-				{
-					// clear Matinee Actor
-					Info.Actors[ActorIndex]->RemoveControllingMatineeActor(*this);
-				}
-			}
-
 			GroupActorInfos.RemoveAt( InfoIndex );
 			break;
 		}
@@ -1543,17 +1507,7 @@ void AMatineeActor::ReplaceActorGroupInfo(UInterpGroup * Group, AActor* OldActor
 				{
 					if (GroupInfo.Actors[ActorID] == OldActor)
 					{
-						if (OldActor)
-						{
-							OldActor->RemoveControllingMatineeActor(*this);
-						}
-
 						GroupInfo.Actors[ActorID] = NewActor;
-
-						if (NewActor)
-						{
-							NewActor->AddControllingMatineeActor(*this);
-						}
 						return;
 					}
 				}
@@ -1743,11 +1697,6 @@ void AMatineeActor::InitGroupActorForGroup(class UInterpGroup* InGroup, class AA
 		{
 			bFoundGroup = true;
 			Info.Actors.AddUnique( InGroupActor );
-
-			if (InGroupActor)
-			{
-				InGroupActor->AddControllingMatineeActor(*this);
-			}
 		}
 	}
 
@@ -1758,11 +1707,6 @@ void AMatineeActor::InitGroupActorForGroup(class UInterpGroup* InGroup, class AA
 		NewInfo.Actors.Add( InGroupActor );
 
 		GroupActorInfos.Add( NewInfo );
-
-		if (InGroupActor)
-		{
-			InGroupActor->AddControllingMatineeActor(*this);
-		}
 	}
 
 	PostEditChange();
@@ -6735,22 +6679,9 @@ void UInterpTrackDirector::UpdateTrack(float NewPosition, UInterpTrackInst* TrIn
 						// If the actor's current view target is a director track camera, then we want to store 
 						// the director track's 'old view target' in case the current Matinee sequence finishes
 						// before our's does.
-						UInterpTrackInstDirector* PreviousDirInst = PC->GetControllingDirector();
-						if( PreviousDirInst != NULL && PreviousDirInst->OldViewTarget != NULL )
-						{
-							// Store the underlying director track's old view target so we can restore this later
-							DirInst->OldViewTarget = PreviousDirInst->OldViewTarget;
-						}
-						else
-						{
-							DirInst->OldViewTarget = ViewTarget;
-						}
+						DirInst->OldViewTarget = ViewTarget;
 					}
 
-					PC->SetControllingDirector(DirInst, bSimulateCameraCutsOnClients);
-	
-					PC->NotifyDirectorControl(true, MatineeActor); 
-
 					//UE_LOG(LogMatinee, Log, TEXT("UInterpTrackDirector::UpdateTrack SetViewTarget ViewGroupInst->GroupActor Time:%f Name: %s"), PC->GetWorld()->GetTimeSeconds(), *ViewGroupInst->GroupActor->GetFName());
 					// Change view to desired view target.
 					FViewTargetTransitionParams TransitionParams;
@@ -6779,9 +6710,6 @@ void UInterpTrackDirector::UpdateTrack(float NewPosition, UInterpTrackInst* TrIn
 					PC->SetViewTarget( DirInst->OldViewTarget, TransitionParams );
 				}
 
-				PC->NotifyDirectorControl(false, MatineeActor); 
-				PC->SetControllingDirector(NULL, false);
-				
 				DirInst->OldViewTarget = NULL;
 			}
 		}
@@ -6941,9 +6869,6 @@ void UInterpTrackInstDirector::TermTrackInst(UInterpTrack* Track)
 				PC->SetViewTarget(OldViewTarget);
 			}
 		}
-		// this may be a duplicate call if it was already called in UpdateTrack(), but that's better than not at all and leaving code thinking we stayed in matinee forever
-		PC->NotifyDirectorControl(false, MatineeActor);
-		PC->SetControllingDirector(NULL, false);
 	}
 	
 	OldViewTarget = NULL;
diff --git a/Engine/Source/Runtime/Engine/Private/LocalPlayer.cpp b/Engine/Source/Runtime/Engine/Private/LocalPlayer.cpp
--- a/Engine/Source/Runtime/Engine/Private/LocalPlayer.cpp
+++ b/Engine/Source/Runtime/Engine/Private/LocalPlayer.cpp
@@ -21,9 +21,6 @@
 #include "UnrealEngine.h"
 #include "EngineUtils.h"
 
-#include "Matinee/MatineeActor.h"
-#include "Matinee/InterpData.h"
-#include "Matinee/InterpGroupInst.h"
 #include "Net/OnlineEngineInterface.h"
 #include "SceneManagement.h"
 #include "Physics/PhysicsInterfaceCore.h"
@@ -50,7 +47,7 @@ DEFINE_LOG_CATEGORY(LogPlayerManagement);
 static TAutoConsoleVariable<int32> CVarViewportTest(
 	TEXT("r.Test.ConstrainedView"),
 	0,
-	TEXT("Allows to test different viewport rectangle configuations (in game only) as they can happen when using Matinee/Editor.\n")
+	TEXT("Allows to test different viewport rectangle configuations (in game only) as they can happen when using cinematics/Editor.\n")
 	TEXT("0: off(default)\n")
 	TEXT("1..7: Various Configuations"),
 	ECVF_RenderThreadSafe);
@@ -870,7 +867,6 @@ FSceneView* ULocalPlayer::CalcSceneView( class FSceneViewFamily* ViewFamily,
 		}
 
 		//	CAMERA OVERRIDE
-		//	NOTE: Matinee works through this channel
 		View->OverridePostProcessSettings(ViewInfo.PostProcessSettings, ViewInfo.PostProcessBlendWeight);
 
 		if (PlayerController->PlayerCameraManager)
@@ -1363,52 +1359,6 @@ bool ULocalPlayer::HandleToggleStreamingVolumesCommand( const TCHAR* Cmd, FOutpu
 	return true;
 }
 
-bool ULocalPlayer::HandleCancelMatineeCommand( const TCHAR* Cmd, FOutputDevice& Ar )
-{
-	// allow optional parameter for initial time in the matinee that this won't work (ie,
-	// 'cancelmatinee 5' won't do anything in the first 5 seconds of the matinee)
-	float InitialNoSkipTime = FCString::Atof(Cmd);
-
-	// is the player in cinematic mode?
-	if (PlayerController->bCinematicMode)
-	{
-		bool bFoundMatinee = false;
-		// if so, look for all active matinees that has this Player in a director group
-		for (TActorIterator<AMatineeActor> It(GetWorld()); It; ++It)
-		{
-			AMatineeActor* MatineeActor = *It;
-
-			// is it currently playing (and skippable)?
-			if (MatineeActor->bIsPlaying && MatineeActor->bIsSkippable && (MatineeActor->bClientSideOnly || MatineeActor->GetWorld()->IsServer()))
-			{
-				for (int32 GroupIndex = 0; GroupIndex < MatineeActor->GroupInst.Num(); GroupIndex++)
-				{
-					// is the PC the group actor?
-					if (MatineeActor->GroupInst[GroupIndex]->GetGroupActor() == PlayerController)
-					{
-						const float RightBeforeEndTime = 0.1f;
-						// make sure we aren';t already at the end (or before the allowed skip time)
-						if ((MatineeActor->InterpPosition < MatineeActor->MatineeData->InterpLength - RightBeforeEndTime) &&
-							(MatineeActor->InterpPosition >= InitialNoSkipTime))
-						{
-							// skip to end
-							MatineeActor->SetPosition(MatineeActor->MatineeData->InterpLength - RightBeforeEndTime, true);
-							bFoundMatinee = true;
-						}
-					}
-				}
-			}
-		}
-
-		if (bFoundMatinee)
-		{
-			FGameDelegates::Get().GetMatineeCancelledDelegate().Broadcast();
-		}
-	}
-	return true;
-}
-
-
 bool ULocalPlayer::Exec(UWorld* InWorld, const TCHAR* Cmd,FOutputDevice& Ar)
 {
 #if WITH_EDITOR
@@ -1521,12 +1471,6 @@ bool ULocalPlayer::Exec(UWorld* InWorld, const TCHAR* Cmd,FOutputDevice& Ar)
 	{
 		return HandleToggleStreamingVolumesCommand( Cmd, Ar );
 	}
-	// @hack: This is a test matinee skipping function, quick and dirty to see if it's good enough for
-	// gameplay. Will fix up better when we have some testing done!
-	else if (FParse::Command(&Cmd, TEXT("CANCELMATINEE")))
-	{
-		return HandleCancelMatineeCommand( Cmd, Ar );
-	}
 	else if(ViewportClient && ViewportClient->Exec( InWorld, Cmd,Ar))
 	{
 		return true;
diff --git a/Engine/Source/Runtime/Engine/Private/PlayerController.cpp b/Engine/Source/Runtime/Engine/Private/PlayerController.cpp
--- a/Engine/Source/Runtime/Engine/Private/PlayerController.cpp
+++ b/Engine/Source/Runtime/Engine/Private/PlayerController.cpp
@@ -44,8 +44,6 @@
 #include "IXRInput.h"
 #include "GameFramework/TouchInterface.h"
 #include "DisplayDebugHelpers.h"
-#include "Matinee/InterpTrackInstDirector.h"
-#include "Matinee/MatineeActor.h"
 #include "Engine/ActorChannel.h"
 #include "GameFramework/PawnMovementComponent.h"
 #include "GameFramework/SpectatorPawn.h"
@@ -433,14 +431,6 @@ AActor* APlayerController::GetViewTarget() const
 
 void APlayerController::SetViewTarget(class AActor* NewViewTarget, struct FViewTargetTransitionParams TransitionParams)
 {
-	// if we're being controlled by a director track, update it with the new viewtarget 
-	// so it returns to the proper viewtarget when it finishes.
-	UInterpTrackInstDirector* const Director = GetControllingDirector();
-	if (Director)
-	{
-		Director->OldViewTarget = NewViewTarget;
-	}
-
 	if (PlayerCameraManager)
 	{
 		PlayerCameraManager->SetViewTarget(NewViewTarget, TransitionParams);
@@ -534,24 +524,6 @@ ACameraActor* APlayerController::GetAutoActivateCameraForPlayer() const
 	return NULL;
 }
 
-
-
-void APlayerController::SetControllingDirector(UInterpTrackInstDirector* NewControllingDirector, bool bClientSimulatingViewTarget)
-{
-	ControllingDirTrackInst = NewControllingDirector;
-
-	if (PlayerCameraManager != NULL)
-	{
-		PlayerCameraManager->bClientSimulatingViewTarget = (NewControllingDirector != NULL) ? bClientSimulatingViewTarget : false;
-	}
-}
-
-
-UInterpTrackInstDirector* APlayerController::GetControllingDirector()
-{
-	return ControllingDirTrackInst;
-}
-
 /// @cond DOXYGEN_WARNINGS
 
 bool APlayerController::ServerNotifyLoadedWorld_Validate(FName WorldPackageName)
@@ -676,20 +648,6 @@ void APlayerController::InitInputSystem()
 	check(World);
 	World->PersistentLevel->PushPendingAutoReceiveInput(this);
 
-	// add the player to any matinees running so that it gets in on any cinematics already running, etc
-	// (already done on server in PostLogin())
-	if (GetLocalRole() < ROLE_Authority)
-	{
-		TArray<AMatineeActor*> AllMatineeActors;
-		World->GetMatineeActors(AllMatineeActors);
-
-		// tell them all to add this PC to any running Director tracks
-		for (int32 i = 0; i < AllMatineeActors.Num(); i++)
-		{
-			AllMatineeActors[i]->AddPlayerToDirectorTracks(this);
-		}
-	}
-
 	// setup optional touchscreen interface
 	CreateTouchInterface();
 }
@@ -3517,16 +3475,6 @@ bool APlayerController::IsPlayerMuted(const FUniqueNetId& PlayerId)
 	return MuteList.IsPlayerMuted(PlayerId);
 }
 
-void APlayerController::NotifyDirectorControl(bool bNowControlling, AMatineeActor* CurrentMatinee)
-{
-	// matinee is done, make sure client syncs up viewtargets, since we were ignoring
-	// ClientSetViewTarget during the matinee.
-	if (!bNowControlling && (GetNetMode() == NM_Client) && PlayerCameraManager && PlayerCameraManager->bClientSimulatingViewTarget)
-	{
-		ServerVerifyViewTarget();
-	}
-}
-
 /// @cond DOXYGEN_WARNINGS
 
 void APlayerController::ClientWasKicked_Implementation(const FText& KickReason)
diff --git a/Engine/Source/Runtime/Engine/Private/Replication.cpp b/Engine/Source/Runtime/Engine/Private/Replication.cpp
--- a/Engine/Source/Runtime/Engine/Private/Replication.cpp
+++ b/Engine/Source/Runtime/Engine/Private/Replication.cpp
@@ -166,7 +166,6 @@ void AMatineeActor::PostNetReceive()
 				UInterpGroupInst * GrInst = FindGroupInst(Actor);
 				if (Actor != NULL && !Actor->IsPendingKill() && GrInst != NULL) 
 				{
-					Actor->AddControllingMatineeActor(*this);
 					// fire an event if we're really playing (and not just starting it up to do a position update)
 					if (bIsPlaying)
 					{
@@ -222,7 +221,6 @@ void AMatineeActor::PostNetReceive()
 				AActor* Actor = ControlledActors[ ActorIndex ];
 				if (Actor != NULL)
 				{
-					Actor->RemoveControllingMatineeActor(*this);
 
 					// fire an event if we were really playing (and not just starting it up to do a position update)
 					if (SavedbIsPlaying)
diff --git a/Engine/Source/Runtime/Engine/Private/UnrealClient.cpp b/Engine/Source/Runtime/Engine/Private/UnrealClient.cpp
--- a/Engine/Source/Runtime/Engine/Private/UnrealClient.cpp
+++ b/Engine/Source/Runtime/Engine/Private/UnrealClient.cpp
@@ -1115,19 +1115,8 @@ int32 FStatHitchesData::DrawStat(FViewport* InViewport, FCanvas* InCanvas, int32
 			Hitches[OverwriteIndex] = DeltaSeconds;
 			When[OverwriteIndex] = CurrentTime;
 			OverwriteIndex = (OverwriteIndex + 1) % NumHitches;
-			if (GEngine->ActiveMatinee.IsValid())
-			{
-				float MatineeTime = GEngine->ActiveMatinee.Get()->InterpPosition;
-				float MatineeMM = FPlatformMath::TruncToFloat(MatineeTime / 60.0f);
-				float MatineeSS = FPlatformMath::TruncToFloat(MatineeTime - MatineeMM * 60.0f);
-				float MatineeMS = FPlatformMath::TruncToFloat((MatineeTime - MatineeMM * 60.0f - MatineeSS) * 1000.0f);
-				UE_LOG(LogClient, Warning, TEXT("HITCH @ %02dm:%02d.%03ds,%d,%d,%d"),
-					(int32)MatineeMM, (int32)MatineeSS, (int32)MatineeMS, int32(MatineeTime * 1000), int32(DeltaSeconds * 1000), Count++);
-			}
-			else
-			{
-				UE_LOG(LogClient, Warning, TEXT("HITCH %d              running cnt = %5d"), int32(DeltaSeconds * 1000), Count++);
-			}
+
+			UE_LOG(LogClient, Warning, TEXT("HITCH %d              running cnt = %5d"), int32(DeltaSeconds * 1000), Count++);
 		}
 
 		const int32 MaxY = InCanvas->GetRenderTarget()->GetSizeXY().Y;
diff --git a/Engine/Source/Runtime/Engine/Private/World.cpp b/Engine/Source/Runtime/Engine/Private/World.cpp
--- a/Engine/Source/Runtime/Engine/Private/World.cpp
+++ b/Engine/Source/Runtime/Engine/Private/World.cpp
@@ -62,7 +62,6 @@
 #include "Engine/WorldComposition.h"
 #include "WorldPartition/WorldPartition.h"
 #include "WorldPartition/WorldPartitionSubsystem.h"
-#include "Matinee/MatineeActor.h"
 #include "Misc/NetworkVersion.h"
 #include "GameFramework/GameNetworkManager.h"
 #include "GameFramework/CharacterMovementComponent.h"
@@ -6182,22 +6181,6 @@ void UWorld::NavigateTo(FIntVector InLocation)
 	FlushLevelStreaming();
 }
 
-void UWorld::GetMatineeActors(TArray<class AMatineeActor*>& OutMatineeActors)
-{
-	check( IsGameWorld() && GetCurrentLevel());
-
-	ULevel* CurLevel = GetCurrentLevel();
-	for( int32 ActorIndex = 0; ActorIndex < CurLevel->Actors.Num(); ++ActorIndex )
-	{
-		AActor* Actor = CurLevel->Actors[ ActorIndex ];
-		AMatineeActor* MatineeActor = Cast<AMatineeActor>( Actor );
-		if( MatineeActor )
-		{
-			OutMatineeActors.Add( MatineeActor );
-		}
-	}
-}
-
 /*-----------------------------------------------------------------------------
 	Seamless world traveling
 -----------------------------------------------------------------------------*/
