Matinee/Animation: Removed Matinee related functions from Skeletal Mesh components.

diff --git a/Engine/Source/Runtime/Engine/Classes/Animation/AnimMontage.h b/Engine/Source/Runtime/Engine/Classes/Animation/AnimMontage.h
--- a/Engine/Source/Runtime/Engine/Classes/Animation/AnimMontage.h
+++ b/Engine/Source/Runtime/Engine/Classes/Animation/AnimMontage.h
@@ -585,20 +585,11 @@ private:
 	float GetRemainingPlayTimeToSectionEnd(const FMontageSubStepper& MontageSubStepper) const;
 
 public:
-	/** static functions that are used by matinee functionality */
-	static UAnimMontage* SetMatineeAnimPositionInner(FName SlotName, USkeletalMeshComponent* SkeletalMeshComponent, UAnimSequenceBase* InAnimSequence, float InPosition, bool bLooping);
-	static UAnimMontage* PreviewMatineeSetAnimPositionInner(FName SlotName, USkeletalMeshComponent* SkeletalMeshComponent, UAnimSequenceBase* InAnimSequence, float InPosition, bool bLooping, bool bFireNotifies, float DeltaTime);
-	
-	static UAnimMontage* SetMatineeAnimPositionInner(FName SlotName, UAnimInstance* AnimInstance, UAnimSequenceBase* InAnimSequence, float InPosition, bool bLooping);
-	static UAnimMontage* PreviewMatineeSetAnimPositionInner(FName SlotName, USkeletalMeshComponent* SkeletalMeshComponent, UAnimInstance* AnimInstance, UAnimSequenceBase* InAnimSequence, float InPosition, bool bLooping, bool bFireNotifies, float DeltaTime);
-	
 	/** static functions that are used by sequencer montage support*/
 	static UAnimMontage* SetSequencerMontagePosition(FName SlotName, USkeletalMeshComponent* SkeletalMeshComponent, int32& InOutInstanceId, UAnimSequenceBase* InAnimSequence, float InFromPosition, float InToPosition, float Weight, bool bLooping, bool bPlaying);
 	static UAnimMontage* PreviewSequencerMontagePosition(FName SlotName, USkeletalMeshComponent* SkeletalMeshComponent, int32& InOutInstanceId, UAnimSequenceBase* InAnimSequence, float InFromPosition, float InToPosition, float Weight, bool bLooping, bool bFireNotifies, bool bPlaying);
 	static UAnimMontage* SetSequencerMontagePosition(FName SlotName, UAnimInstance* AnimInstance, int32& InOutInstanceId, UAnimSequenceBase* InAnimSequence, float InFromPosition, float InToPosition, float Weight, bool bLooping, bool bPlaying);
 	static UAnimMontage* PreviewSequencerMontagePosition(FName SlotName, USkeletalMeshComponent* SkeletalMeshComponent, UAnimInstance* AnimInstance, int32& InOutInstanceId, UAnimSequenceBase* InAnimSequence, float InFromPosition, float InToPosition, float Weight, bool bLooping, bool bFireNotifies, bool bPlaying);
-private:
-	static UAnimMontage* InitializeMatineeControl(FName SlotName, UAnimInstance* AnimInstance, UAnimSequenceBase* InAnimSequence, bool bLooping);
 };
 
 /* 
diff --git a/Engine/Source/Runtime/Engine/Classes/Animation/SkeletalMeshActor.h b/Engine/Source/Runtime/Engine/Classes/Animation/SkeletalMeshActor.h
--- a/Engine/Source/Runtime/Engine/Classes/Animation/SkeletalMeshActor.h
+++ b/Engine/Source/Runtime/Engine/Classes/Animation/SkeletalMeshActor.h
@@ -5,7 +5,6 @@
 #include "UObject/ObjectMacros.h"
 #include "Engine/EngineTypes.h"
 #include "GameFramework/Actor.h"
-#include "Matinee/MatineeAnimInterface.h"
 #include "SkeletalMeshActor.generated.h"
 
 class UAnimMontage;
@@ -20,7 +19,7 @@ class UAnimSequence;
  * @see USkeletalMesh
  */
 UCLASS(ClassGroup=ISkeletalMeshes, Blueprintable, ComponentWrapperClass, ConversionRoot, meta=(ChildCanTick))
-class ENGINE_API ASkeletalMeshActor : public AActor, public IMatineeAnimInterface
+class ENGINE_API ASkeletalMeshActor : public AActor
 {
 	GENERATED_UCLASS_BODY()
 
@@ -83,18 +82,6 @@ public:
 	virtual void PostInitializeComponents() override;
 	//~ End AActor Interface
 
-	//~ Begin IMatineeAnimInterface Interface
-	virtual void PreviewBeginAnimControl(class UInterpGroup* InInterpGroup) override;
-	virtual void PreviewSetAnimPosition(FName SlotName, int32 ChannelIndex, UAnimSequence* InAnimSequence, float InPosition, bool bLooping, bool bFireNotifies, float AdvanceTime) override;
-	virtual void PreviewSetAnimWeights(TArray<FAnimSlotInfo>& SlotInfos) override;
-	virtual void PreviewFinishAnimControl(class UInterpGroup* InInterpGroup) override;
-	virtual void GetAnimControlSlotDesc(TArray<struct FAnimSlotDesc>& OutSlotDescs) override {};
-	virtual void SetAnimWeights( const TArray<struct FAnimSlotInfo>& SlotInfos ) override;
-	virtual void BeginAnimControl(class UInterpGroup* InInterpGroup) override;
-	virtual void SetAnimPosition(FName SlotName, int32 ChannelIndex, class UAnimSequence* InAnimSequence, float InPosition, bool bFireNotifies, bool bLooping) override;
-	virtual void FinishAnimControl(class UInterpGroup* InInterpGroup) override;
-	//~ End IMatineeAnimInterface Interface
-
 private:
 	// utility function to see if it can play animation or not
 	bool CanPlayAnimation(class UAnimSequenceBase* AnimAssetBase=NULL) const;
diff --git a/Engine/Source/Runtime/Engine/Classes/Components/SkeletalMeshComponent.h b/Engine/Source/Runtime/Engine/Classes/Components/SkeletalMeshComponent.h
--- a/Engine/Source/Runtime/Engine/Classes/Components/SkeletalMeshComponent.h
+++ b/Engine/Source/Runtime/Engine/Classes/Components/SkeletalMeshComponent.h
@@ -321,14 +321,6 @@ class ENGINE_API USkeletalMeshComponent : public USkinnedMeshComponent, public I
 	friend struct FAnimNode_LinkedAnimLayer;
 	friend struct FLinkedInstancesAdapter;
 
-	/**
-	 * Animation 
-	 */
-	
-	/** @Todo anim: Matinee related data start - this needs to be replaced to new system. */
-	
-	/** @Todo anim: Matinee related data end - this needs to be replaced to new system. */
-
 public:
 #if WITH_EDITORONLY_DATA
 	/** The blueprint for creating an AnimationScript. */
diff --git a/Engine/Source/Runtime/Engine/Private/Animation/AnimMontage.cpp b/Engine/Source/Runtime/Engine/Private/Animation/AnimMontage.cpp
--- a/Engine/Source/Runtime/Engine/Private/Animation/AnimMontage.cpp
+++ b/Engine/Source/Runtime/Engine/Private/Animation/AnimMontage.cpp
@@ -2711,131 +2711,6 @@ void FAnimMontageInstance::BranchingPointEventHandler(const FBranchingPointMarke
 	}
 }
 
-UAnimMontage* FAnimMontageInstance::InitializeMatineeControl(FName SlotName, UAnimInstance* AnimInstance, UAnimSequenceBase* InAnimSequence, bool bLooping)
-{
-	UAnimMontage* MontageToPlay = Cast<UAnimMontage>(InAnimSequence);
-
-	if (UAnimSingleNodeInstance* SingleNodeInst = Cast<UAnimSingleNodeInstance>(AnimInstance))
-	{
-		// Single node anim instance
-		if (SingleNodeInst->GetCurrentAsset() != InAnimSequence)
-		{
-			SingleNodeInst->SetAnimationAsset(InAnimSequence, bLooping);
-			SingleNodeInst->SetPosition(0.0f);
-		}
-
-		if (SingleNodeInst->IsLooping() != bLooping)
-		{
-			SingleNodeInst->SetLooping(bLooping);
-		}
-
-		return MontageToPlay;
-	}
-	else if (UAnimInstance* AnimInst = AnimInstance)
-	{
-		UAnimMontage* PreviousMontage = nullptr;
-
-		if (MontageToPlay)
-		{
-			if (!AnimInst->Montage_IsPlaying(MontageToPlay))
-			{
-				// Will reuse an existing montage instance for this montage, if one already exists
-				AnimInst->Montage_Play(MontageToPlay, 0.f);
-			}
-
-			return MontageToPlay;
-		}
-
-		// We need to attempt find an existing slot animation
-		for (FAnimMontageInstance* MontageInstance : AnimInst->MontageInstances)
-		{
-			if (!MontageInstance || !MontageInstance->IsActive())
-			{
-				continue;
-			}
-
-			// Try and find an anim track for the desired slot
-			UAnimMontage* ThisMontage = MontageInstance->Montage;
-			const FAnimTrack* AnimTrack = ( ThisMontage && ThisMontage->GetOuter() == GetTransientPackage() ) ? ThisMontage->GetAnimationData(SlotName) : nullptr;
-			if (!AnimTrack)
-			{
-				continue;
-			}
-
-			// Try and find our asset on this track
-			if (AnimTrack->AnimSegments.Num() == 1 && AnimTrack->AnimSegments[0].AnimReference == InAnimSequence)
-			{
-				// We've already found an active animation in this slot that's playing our animation. Use that.
-				return ThisMontage;
-			}
-
-			// Something animating this slot that's not us - set weight to be 0 on this slot
-			MontageInstance->BlendStartAlpha = 1.0f;
-			MontageInstance->Blend.SetDesiredValue(0.f);
-			MontageInstance->Blend.SetAlpha(1.f);
-		}
-
-		return AnimInst->PlaySlotAnimationAsDynamicMontage(InAnimSequence, SlotName, 0.0f, 0.0f, 0.f, 1);
-	}
-
-	return nullptr;
-}
-
-UAnimMontage* FAnimMontageInstance::SetMatineeAnimPositionInner(FName SlotName, USkeletalMeshComponent* SkeletalMeshComponent, UAnimSequenceBase* InAnimSequence, float InPosition, bool bLooping)
-{
-	if (SkeletalMeshComponent)
-	{
-		return SetMatineeAnimPositionInner(SlotName, SkeletalMeshComponent->GetAnimInstance(), InAnimSequence, InPosition, bLooping);
-	}
-
-	return nullptr;
-}
-
-UAnimMontage* FAnimMontageInstance::SetMatineeAnimPositionInner(FName SlotName, UAnimInstance* AnimInstance, UAnimSequenceBase* InAnimSequence, float InPosition, bool bLooping)
-{
-	UAnimMontage* PlayingMontage = InitializeMatineeControl(SlotName, AnimInstance, InAnimSequence, bLooping);
-	UAnimInstance* AnimInst = AnimInstance;
-	if (UAnimSingleNodeInstance* SingleNodeInst = Cast<UAnimSingleNodeInstance>(AnimInst))
-	{
-		if (SingleNodeInst->GetCurrentTime() != InPosition)
-		{
-			SingleNodeInst->SetPosition(InPosition);
-		}
-	}
-	else if (PlayingMontage && AnimInst)
-	{
-		FAnimMontageInstance* AnimMontageInst = AnimInst->GetActiveInstanceForMontage(PlayingMontage);
-		if (!AnimMontageInst)
-		{
-			UE_LOG(LogSkeletalMesh, Warning, TEXT("Unable to set animation position for montage on slot name: %s"), *SlotName.ToString());
-			return nullptr;
-		}
-
-		// ensure full weighting to this instance
-		AnimMontageInst->Blend.SetDesiredValue(1.f);
-		AnimMontageInst->Blend.SetAlpha(1.f);
-		AnimMontageInst->BlendStartAlpha = 1.0f;
-
-		AnimMontageInst->SetNextPositionWithEvents(InPosition);
-	}
-	else
-	{
-		UE_LOG(LogSkeletalMesh, Warning, TEXT("Invalid animation configuration when attempting to set animation possition with : %s"), *InAnimSequence->GetName());
-	}
-
-	return PlayingMontage;
-}
-
-UAnimMontage* FAnimMontageInstance::PreviewMatineeSetAnimPositionInner(FName SlotName, USkeletalMeshComponent* SkeletalMeshComponent, UAnimSequenceBase* InAnimSequence, float InPosition, bool bLooping, bool bFireNotifies, float DeltaTime)
-{
-	if (SkeletalMeshComponent)
-	{
-		return PreviewMatineeSetAnimPositionInner(SlotName, SkeletalMeshComponent, SkeletalMeshComponent->GetAnimInstance(), InAnimSequence, InPosition, bLooping, bFireNotifies, DeltaTime);
-	}
-
-	return nullptr;
-}
-
 UAnimMontage* FAnimMontageInstance::PreviewSequencerMontagePosition(FName SlotName, USkeletalMeshComponent* SkeletalMeshComponent, int32& InOutInstanceId, UAnimSequenceBase* InAnimSequence, float InFromPosition, float InToPosition, float Weight, bool bLooping, bool bFireNotifies, bool bPlaying)
 {
 	if (SkeletalMeshComponent)
@@ -2846,82 +2721,6 @@ UAnimMontage* FAnimMontageInstance::PreviewSequencerMontagePosition(FName SlotNa
 	return nullptr;
 }
 
-UAnimMontage* FAnimMontageInstance::PreviewMatineeSetAnimPositionInner(FName SlotName, USkeletalMeshComponent* SkeletalMeshComponent, UAnimInstance* AnimInstance, UAnimSequenceBase* InAnimSequence, float InPosition, bool bLooping, bool bFireNotifies, float DeltaTime)
-{
-	// Codepath for updating an animation when the skeletal mesh component is not going to be ticked (ie in editor)
-	UAnimMontage* PlayingMontage = InitializeMatineeControl(SlotName, AnimInstance, InAnimSequence, bLooping);
-
-	UAnimInstance* AnimInst = AnimInstance;
-
-	FAnimMontageInstance* MontageInstanceToUpdate = AnimInst && PlayingMontage ? AnimInst->GetActiveInstanceForMontage(PlayingMontage) : nullptr;
-	float PreviousPosition = InPosition;
-
-	if (UAnimSingleNodeInstance* SingleNodeInst = Cast<UAnimSingleNodeInstance>(AnimInst))
-	{
-		PreviousPosition = SingleNodeInst->GetCurrentTime();
-
-		// If we're playing a montage, we fire notifies explicitly below (rather than allowing the single node instance to do it)
-		const bool bFireNotifiesHere = bFireNotifies && PlayingMontage == nullptr;
-
-		if (DeltaTime == 0.f)
-		{
-			const float PreviousTime = InPosition;
-			SingleNodeInst->SetPositionWithPreviousTime(InPosition, PreviousTime, bFireNotifiesHere);
-		}
-		else
-		{
-			SingleNodeInst->SetPosition(InPosition, bFireNotifiesHere);
-		}
-	}
-	else if (MontageInstanceToUpdate)
-	{
-		// ensure full weighting to this instance
-		MontageInstanceToUpdate->Blend.SetDesiredValue(1.f);
-		MontageInstanceToUpdate->Blend.SetAlpha(1.f);
-		MontageInstanceToUpdate->BlendStartAlpha = 1.0f;
-
-		PreviousPosition = AnimInst->Montage_GetPosition(PlayingMontage);
-		AnimInst->Montage_SetPosition(PlayingMontage, InPosition);
-	}
-	else
-	{
-		UE_LOG(LogSkeletalMesh, Warning, TEXT("Invalid animation configuration when attempting to set animation possition with : %s"), *InAnimSequence->GetName());
-	}
-
-	// Now force the animation system to update, if we have a montage instance
-	if (MontageInstanceToUpdate)
-	{
-		AnimInst->UpdateAnimation(DeltaTime, false);
-
-		// since we don't advance montage in the tick, we manually have to handle notifies
-		MontageInstanceToUpdate->HandleEvents(PreviousPosition, InPosition, NULL);
-
-		if (!bFireNotifies)
-		{
-			AnimInst->NotifyQueue.Reset(SkeletalMeshComponent);
-		}
-
-		// Allow the proxy to update (this also filters unfiltered notifies)
-		if (AnimInst->NeedsUpdate())
-		{
-			AnimInst->ParallelUpdateAnimation();
-		}
-
-		// Explicitly call post update (also triggers notifies)
-		AnimInst->PostUpdateAnimation();
-	}
-
-	// Update space bases so new animation position has an effect.
-	SkeletalMeshComponent->RefreshBoneTransforms();
-	SkeletalMeshComponent->RefreshSlaveComponents();
-	SkeletalMeshComponent->UpdateComponentToWorld();
-	SkeletalMeshComponent->FinalizeBoneTransform();
-	SkeletalMeshComponent->MarkRenderTransformDirty();
-	SkeletalMeshComponent->MarkRenderDynamicDataDirty();
-
-	return PlayingMontage;
-}
-
 UAnimMontage* FAnimMontageInstance::SetSequencerMontagePosition(FName SlotName, USkeletalMeshComponent* SkeletalMeshComponent, int32& InOutInstanceId, UAnimSequenceBase* InAnimSequence, float InFromPosition, float InToPosition, float Weight, bool bLooping, bool bPlaying)
 {
 	if (SkeletalMeshComponent)
diff --git a/Engine/Source/Runtime/Engine/Private/Animation/SkeletalMeshActor.cpp b/Engine/Source/Runtime/Engine/Private/Animation/SkeletalMeshActor.cpp
--- a/Engine/Source/Runtime/Engine/Private/Animation/SkeletalMeshActor.cpp
+++ b/Engine/Source/Runtime/Engine/Private/Animation/SkeletalMeshActor.cpp
@@ -37,61 +37,6 @@ void ASkeletalMeshActor::GetLifetimeReplicatedProps(TArray< FLifetimeProperty >
 	DOREPLIFETIME(ASkeletalMeshActor, ReplicatedMaterial1);
 }
 
-void ASkeletalMeshActor::PreviewBeginAnimControl(UInterpGroup* InInterpGroup)
-{
-	if (CanPlayAnimation())
-	{
-		UAnimInstance* AnimInst = SkeletalMeshComponent->GetAnimInstance();
-		if (!AnimInst)
-		{
-			SkeletalMeshComponent->SetAnimationMode(EAnimationMode::Type::AnimationSingleNode);
-		}
-	}
-}
-
-
-void ASkeletalMeshActor::PreviewFinishAnimControl(UInterpGroup* InInterpGroup)
-{
-	if (CanPlayAnimation())
-	{
-		// if in editor, reset the Animations, makes easier for artist to see them visually and align them
-		// in game, we keep the last pose that matinee kept. If you'd like it to have animation, you'll need to have AnimTree or AnimGraph to handle correctly
-		if (SkeletalMeshComponent->GetAnimationMode() == EAnimationMode::Type::AnimationBlueprint)
-		{
-			UAnimInstance* AnimInst = SkeletalMeshComponent->GetAnimInstance();
-			if (AnimInst)
-			{
-				AnimInst->Montage_Stop(0.f);
-				AnimInst->UpdateAnimation(0.f, false);
-			}
-		}
-		// Update space bases to reset it back to ref pose
-		SkeletalMeshComponent->RefreshBoneTransforms();
-		SkeletalMeshComponent->RefreshSlaveComponents();
-		SkeletalMeshComponent->UpdateComponentToWorld();
-	}
-}
-
-
-void ASkeletalMeshActor::PreviewSetAnimPosition(FName SlotName, int32 ChannelIndex, UAnimSequence* InAnimSequence, float InPosition, bool bLooping, bool bFireNotifies, float DeltaTime)
-{
-	if (CanPlayAnimation(InAnimSequence))
-	{
-		TWeakObjectPtr<class UAnimMontage>& CurrentlyPlayingMontage = CurrentlyPlayingMontages.FindOrAdd(SlotName);
-		CurrentlyPlayingMontage = FAnimMontageInstance::PreviewMatineeSetAnimPositionInner(SlotName, SkeletalMeshComponent, SkeletalMeshComponent->GetAnimInstance(),  InAnimSequence, InPosition, bLooping, bFireNotifies, DeltaTime);
-	}
-}
-
-void ASkeletalMeshActor::PreviewSetAnimWeights(TArray<FAnimSlotInfo>& SlotInfos)
-{
-	//no support yet
-}
-
-void ASkeletalMeshActor::SetAnimWeights(const TArray<struct FAnimSlotInfo>& SlotInfos)
-{
-	//no support yet
-}
-
 /** Check SkeletalMeshActor for errors. */
 #if WITH_EDITOR
 void ASkeletalMeshActor::CheckForErrors()
@@ -201,52 +146,12 @@ void ASkeletalMeshActor::OnRep_ReplicatedMaterial1()
 	SkeletalMeshComponent->SetMaterial(1, ReplicatedMaterial1);
 }
 
-void ASkeletalMeshActor::BeginAnimControl(UInterpGroup* InInterpGroup)
-{
-	if (CanPlayAnimation())
-	{
-		UAnimInstance* AnimInst = SkeletalMeshComponent->GetAnimInstance();
-		if (!AnimInst)
-		{
-			SkeletalMeshComponent->SetAnimationMode(EAnimationMode::Type::AnimationSingleNode);
-		}
-	}
-}
-
 bool ASkeletalMeshActor::CanPlayAnimation(class UAnimSequenceBase* AnimAssetBase/*=NULL*/) const
 {
 	return (SkeletalMeshComponent->SkeletalMesh && SkeletalMeshComponent->SkeletalMesh->GetSkeleton() &&
 		(!AnimAssetBase || SkeletalMeshComponent->SkeletalMesh->GetSkeleton()->IsCompatible(AnimAssetBase->GetSkeleton())));
 }
 
-void ASkeletalMeshActor::SetAnimPosition(FName SlotName, int32 ChannelIndex, UAnimSequence* InAnimSequence, float InPosition, bool bFireNotifies, bool bLooping)
-{
-	if (CanPlayAnimation(InAnimSequence))
-	{
-		TWeakObjectPtr<class UAnimMontage>& CurrentlyPlayingMontage = CurrentlyPlayingMontages.FindOrAdd(SlotName);
-		CurrentlyPlayingMontage = FAnimMontageInstance::SetMatineeAnimPositionInner(SlotName, SkeletalMeshComponent->GetAnimInstance(), InAnimSequence, InPosition, bLooping);
-	}
-}
-
-void ASkeletalMeshActor::FinishAnimControl(UInterpGroup* InInterpGroup)
-{
-	if (SkeletalMeshComponent->GetAnimationMode() == EAnimationMode::Type::AnimationBlueprint)
-	{
-		UAnimInstance* AnimInst = SkeletalMeshComponent->GetAnimInstance();
-		if (AnimInst)
-		{
-			AnimInst->Montage_Stop(0.f);
-			AnimInst->UpdateAnimation(0.f, false);
-		}
-
-		// Update space bases to reset it back to ref pose
-		SkeletalMeshComponent->RefreshBoneTransforms();
-		SkeletalMeshComponent->RefreshSlaveComponents();
-		SkeletalMeshComponent->UpdateComponentToWorld();
-	}
-}
-
-
 #if WITH_EDITOR
 
 bool ASkeletalMeshActor::GetReferencedContentObjects(TArray<UObject*>& Objects) const
