Improve implementation of switch nodes in HLSL material generation to detect cases where both expressions are constants

This pattern was already applied in ShadowReplace, so this change abstracts the code such that the same logic can apply to all other similar cases.

diff --git a/Engine/Source/Runtime/Engine/Private/Materials/HLSLMaterialTranslator.cpp b/Engine/Source/Runtime/Engine/Private/Materials/HLSLMaterialTranslator.cpp
--- a/Engine/Source/Runtime/Engine/Private/Materials/HLSLMaterialTranslator.cpp
+++ b/Engine/Source/Runtime/Engine/Private/Materials/HLSLMaterialTranslator.cpp
@@ -8002,81 +8002,67 @@ int32 FHLSLMaterialTranslator::PrecomputedAOMask()
 	return ResultIdx;
 }
 
-int32 FHLSLMaterialTranslator::GIReplace(int32 Direct, int32 StaticIndirect, int32 DynamicIndirect)
-{ 
-	if(Direct == INDEX_NONE || DynamicIndirect == INDEX_NONE)
+int32 FHLSLMaterialTranslator::GenericSwitch(const TCHAR* SwitchExpressionText, int32 IfTrue, int32 IfFalse)
+{
+	if (IfTrue == INDEX_NONE || IfFalse == INDEX_NONE)
 	{
 		return INDEX_NONE;
 	}
 
-	EMaterialValueType ResultType = GetArithmeticResultType(Direct, DynamicIndirect);
-
-	return AddCodeChunk(ResultType,TEXT("(GetGIReplaceState() ? (%s) : (%s))"), *GetParameterCode(DynamicIndirect), *GetParameterCode(Direct));
-}
-
-int32 FHLSLMaterialTranslator::ShadowReplace(int32 Default, int32 Shadow)
-{
-	if (Default == INDEX_NONE || Shadow == INDEX_NONE)
+	// exactly the same inputs on both sides - no need to generate anything extra
+	if (IfTrue == IfFalse)
 	{
-		return INDEX_NONE;
+		return IfTrue;
 	}
 
-	FMaterialUniformExpression* DefaultExpression = GetParameterUniformExpression(Default);
-	FMaterialUniformExpression* ShadowExpression = GetParameterUniformExpression(Shadow);
-	if (DefaultExpression &&
-		ShadowExpression &&
-		DefaultExpression->IsConstant() &&
-		ShadowExpression->IsConstant())
+	FMaterialUniformExpression* IfTrueExpression = GetParameterUniformExpression(IfTrue);
+	FMaterialUniformExpression* IfFalseExpression = GetParameterUniformExpression(IfFalse);
+	if (IfTrueExpression &&
+		IfFalseExpression &&
+		IfTrueExpression->IsConstant() &&
+		IfFalseExpression->IsConstant())
 	{
 		FMaterialRenderContext DummyContext(nullptr, *Material, nullptr);
-		FLinearColor DefaultValue;
-		FLinearColor ShadowValue;
-		DefaultExpression->GetNumberValue(DummyContext, DefaultValue);
-		ShadowExpression->GetNumberValue(DummyContext, ShadowValue);
-		if (DefaultValue == ShadowValue)
+		FLinearColor IfTrueValue;
+		FLinearColor IfFalseValue;
+		IfTrueExpression->GetNumberValue(DummyContext, IfTrueValue);
+		IfFalseExpression->GetNumberValue(DummyContext, IfFalseValue);
+		if (IfTrueValue == IfFalseValue)
 		{
 			// If both inputs are wired to == constant values, avoid adding the runtime switch
 			// This will avoid breaking various offline checks for constant values
-			return Default;
+			return IfTrue;
 		}
 	}
 
-	EMaterialValueType ResultType = GetArithmeticResultType(Default, Shadow);
-	return AddCodeChunk(ResultType, TEXT("(GetShadowReplaceState() ? (%s) : (%s))"), *GetParameterCode(Shadow), *GetParameterCode(Default));
+	EMaterialValueType ResultType = GetArithmeticResultType(IfTrue, IfFalse);
+	return AddCodeChunk(ResultType, TEXT("(%s ? (%s) : (%s))"), SwitchExpressionText, *GetParameterCode(IfTrue), *GetParameterCode(IfFalse));
 }
 
+int32 FHLSLMaterialTranslator::GIReplace(int32 Direct, int32 StaticIndirect, int32 DynamicIndirect)
+{ 
+	return GenericSwitch(TEXT("GetGIReplaceState()"), DynamicIndirect, Direct);
+}
 
-int32 FHLSLMaterialTranslator::ReflectionCapturePassSwitch(int32 Default, int32 Reflection)
+int32 FHLSLMaterialTranslator::ShadowReplace(int32 Default, int32 Shadow)
 {
-	if (Default == INDEX_NONE || Reflection == INDEX_NONE)
-	{
-		return INDEX_NONE;
-	}
+	return GenericSwitch(TEXT("GetShadowReplaceState()"), Shadow, Default);
+}
+
 
-	EMaterialValueType ResultType = GetArithmeticResultType(Default, Reflection);
-	return AddCodeChunk(ResultType, TEXT("(GetReflectionCapturePassSwitchState() ? (%s) : (%s))"), *GetParameterCode(Reflection), *GetParameterCode(Default));
+int32 FHLSLMaterialTranslator::ReflectionCapturePassSwitch(int32 Default, int32 Reflection)
+{
+	return GenericSwitch(TEXT("GetReflectionCapturePassSwitchState()"), Reflection, Default);
 }
 
 int32 FHLSLMaterialTranslator::RayTracingQualitySwitchReplace(int32 Normal, int32 RayTraced)
 {
-	if (Normal == INDEX_NONE || RayTraced == INDEX_NONE)
-	{
-		return INDEX_NONE;
-	}
-
-	EMaterialValueType ResultType = GetArithmeticResultType(Normal, RayTraced);
-	return AddCodeChunk(ResultType, TEXT("(GetRayTracingQualitySwitch() ? (%s) : (%s))"), *GetParameterCode(RayTraced), *GetParameterCode(Normal));
+	return GenericSwitch(TEXT("GetRayTracingQualitySwitch()"), RayTraced, Normal);
 }
 
 int32 FHLSLMaterialTranslator::VirtualTextureOutputReplace(int32 Default, int32 VirtualTexture)
 {
-	if (Default == INDEX_NONE || VirtualTexture == INDEX_NONE)
-	{
-		return INDEX_NONE;
-	}
-
-	EMaterialValueType ResultType = GetArithmeticResultType(Default, VirtualTexture);
-	return AddCodeChunk(ResultType, TEXT("(GetRuntimeVirtualTextureOutputSwitch() ? (%s) : (%s))"), *GetParameterCode(VirtualTexture), *GetParameterCode(Default));
+	return GenericSwitch(TEXT("GetRuntimeVirtualTextureOutputSwitch()"), VirtualTexture, Default);
 }
 
 int32 FHLSLMaterialTranslator::ObjectOrientation()
diff --git a/Engine/Source/Runtime/Engine/Private/Materials/HLSLMaterialTranslator.h b/Engine/Source/Runtime/Engine/Private/Materials/HLSLMaterialTranslator.h
--- a/Engine/Source/Runtime/Engine/Private/Materials/HLSLMaterialTranslator.h
+++ b/Engine/Source/Runtime/Engine/Private/Materials/HLSLMaterialTranslator.h
@@ -612,6 +612,8 @@ protected:
 	// GetArithmeticResultType
 	EMaterialValueType GetArithmeticResultType(EMaterialValueType TypeA, EMaterialValueType TypeB);
 
+	int32 GenericSwitch(const TCHAR* Function, int32 IfTrue, int32 IfFalse);
+
 	// FMaterialCompiler interface.
 
 	/** 
