Fixed strata nodes typos.
Thanks Colton!

diff --git a/Engine/Source/Runtime/Engine/Classes/Materials/MaterialExpressionStrata.h b/Engine/Source/Runtime/Engine/Classes/Materials/MaterialExpressionStrata.h
--- a/Engine/Source/Runtime/Engine/Classes/Materials/MaterialExpressionStrata.h
+++ b/Engine/Source/Runtime/Engine/Classes/Materials/MaterialExpressionStrata.h
@@ -87,7 +87,7 @@ class UMaterialExpressionStrataSlabBSDF : public UMaterialExpressionStrataBSDF
 	FExpressionInput Roughness;
 		
 	/**
-	 * Controls the anisotropy factor of the roughness. Positive value elongates the specular lobe along the Tangent vector, Negative value elongates the specular lobe along the perpendiculat of the Tangent. (type = float, unit = unitless).
+	 * Controls the anisotropy factor of the roughness. Positive value elongates the specular lobe along the Tangent vector, Negative value elongates the specular lobe along the perpendicular of the Tangent. (type = float, unit = unitless).
 	 */
 	UPROPERTY()
 	FExpressionInput Anisotropy;
@@ -147,7 +147,7 @@ class UMaterialExpressionStrataSlabBSDF : public UMaterialExpressionStrataBSDF
 	FExpressionInput FuzzAmount;
 
 	/**
-	 * The base colod of the fuzz.
+	 * The base color of the fuzz.
 	 */
 	UPROPERTY()
 	FExpressionInput FuzzColor;
@@ -238,7 +238,7 @@ class UMaterialExpressionStrataUnlitBSDF : public UMaterialExpressionStrataBSDF
 	FExpressionInput EmissiveColor;
 
 	/**
-	 * The amount of transmitted light from the back side of the surface to the front side of the surface ()
+	 * The amount of transmitted light from the back side of the surface to the front side of the surface (type = float3, unit = unitless, defaults to 1)
 	 */
 	UPROPERTY()
 	FExpressionInput TransmittanceColor;
@@ -326,7 +326,7 @@ class UMaterialExpressionStrataSingleLayerWaterBSDF : public UMaterialExpression
 	FExpressionInput BaseColor;
 
 	/**
-	 * whether the surface represents a dielectric (such as plastic) or a conductor (such as metal). (type = float, unit = unitless, defaults to 0 = dielectric)
+	 * Whether the surface represents a dielectric (such as plastic) or a conductor (such as metal). (type = float, unit = unitless, defaults to 0 = dielectric)
 	 */
 	UPROPERTY()
 	FExpressionInput Metallic;
