Fix race condition in FPerInstanceRenderData::UpdateBoundsTransforms_Concurrent

diff --git a/Engine/Source/Runtime/Engine/Private/InstancedStaticMesh.cpp b/Engine/Source/Runtime/Engine/Private/InstancedStaticMesh.cpp
--- a/Engine/Source/Runtime/Engine/Private/InstancedStaticMesh.cpp
+++ b/Engine/Source/Runtime/Engine/Private/InstancedStaticMesh.cpp
@@ -991,6 +991,9 @@ void FPerInstanceRenderData::UpdateFromPreallocatedData(FStaticMeshInstanceData&
 	ENQUEUE_RENDER_COMMAND(FInstanceBuffer_UpdateFromPreallocatedData)(
 		[InInstanceBufferDataPtr, InInstanceBuffer, this](FRHICommandListImmediate& RHICmdList)
 		{
+			// The assignment to InstanceData shared pointer kills the old data
+			// If UpdateBoundsTask is in-flight it will crash
+			EnsureInstanceDataUpdated();
 			InInstanceBuffer->InstanceData = InInstanceBufferDataPtr;
 			if (!InInstanceBuffer->CondSetFlushToGPUPending())
 			{
@@ -1113,6 +1116,13 @@ const TArray<FRenderTransform>& FPerInstanceRenderData::GetPerInstanceTransforms
 
 void FPerInstanceRenderData::UpdateFromCommandBuffer(FInstanceUpdateCmdBuffer& CmdBuffer)
 {
+	// UpdateFromCommandBuffer reallocates InstanceData in InstanceBuffer
+	// If UpdateBoundsTask is in-flight it will crash
+	ENQUEUE_RENDER_COMMAND(EnsureInstanceDataUpdatedCmd)(
+		[this](FRHICommandList&) {
+		EnsureInstanceDataUpdated();
+	});
+
 	InstanceBuffer.UpdateFromCommandBuffer_Concurrent(CmdBuffer);
 	UpdateBoundsTransforms_Concurrent();
 }
