ApproxActors: expose hard-normals options. Change UProperty metatags to allow smaller min-dimensions to support smaller objects. Change default nanite proxy percent in CreateStaticMeshUtil to 0 (ie minimal 2000 triangles proxy).
#rnx

diff --git a/Engine/Plugins/Experimental/MeshModelingToolset/Source/GeometryProcessingAdapters/Private/ApproximateActorsImpl.cpp b/Engine/Plugins/Experimental/MeshModelingToolset/Source/GeometryProcessingAdapters/Private/ApproximateActorsImpl.cpp
--- a/Engine/Plugins/Experimental/MeshModelingToolset/Source/GeometryProcessingAdapters/Private/ApproximateActorsImpl.cpp
+++ b/Engine/Plugins/Experimental/MeshModelingToolset/Source/GeometryProcessingAdapters/Private/ApproximateActorsImpl.cpp
@@ -457,51 +457,6 @@
 		return Result;
 	}
 
-	Progress.EnterProgressFrame(1.f, LOCTEXT("RemoveHidden", "Removing Hidden Geometry..."));
-
-	if (Options.OcclusionPolicy == IGeometryProcessing_ApproximateActors::EOcclusionPolicy::VisibilityBased)
-	{
-		TRACE_CPUPROFILER_EVENT_SCOPE(ApproximateActorsImpl_Generate_Occlusion);
-		TRemoveOccludedTriangles<FDynamicMesh3> Remover(CurResultMesh);
-		Remover.InsideMode = EOcclusionCalculationMode::SimpleOcclusionTest;
-		Remover.TriangleSamplingMethod = EOcclusionTriangleSampling::Centroids;
-		Remover.AddTriangleSamples = 5;
-		FDynamicMeshAABBTree3 CurResultMeshSpatial(CurResultMesh, false);
-		{
-			TRACE_CPUPROFILER_EVENT_SCOPE(ApproximateActorsImpl_Generate_Occlusion_Spatial);
-			CurResultMeshSpatial.Build();
-		}
-
-		TArray<FTransform3d> NoTransforms;
-		NoTransforms.Add(FTransform3d::Identity());
-		TArray<FDynamicMeshAABBTree3*> Spatials;
-		Spatials.Add(&CurResultMeshSpatial);
-		{
-			TRACE_CPUPROFILER_EVENT_SCOPE(ApproximateActorsImpl_Generate_Occlusion_Compute);
-			Remover.Select(NoTransforms, Spatials, {}, NoTransforms);
-		}
-		if (Remover.RemovedT.Num() > 0)
-		{
-			FMeshFaceSelection Selection(CurResultMesh);
-			{
-				TRACE_CPUPROFILER_EVENT_SCOPE(ApproximateActorsImpl_Generate_Occlusion_Clean);
-				Selection.Select(Remover.RemovedT);
-				Selection.ExpandToOneRingNeighbours(1);
-				Selection.ContractBorderByOneRingNeighbours(4);
-			}
-			FDynamicMeshEditor Editor(CurResultMesh);
-			{
-				TRACE_CPUPROFILER_EVENT_SCOPE(ApproximateActorsImpl_Generate_Occlusion_Delete);
-				Editor.RemoveTriangles(Selection.AsArray(), true);
-			}
-		}
-
-		if (Options.bVerbose)
-		{
-			UE_LOG(LogApproximateActors, Warning, TEXT("Occlusion-Filtered mesh has %d triangles"), CurResultMesh->TriangleCount());
-		}
-	}
-
 	Progress.EnterProgressFrame(1.f, LOCTEXT("SimplifyingMesh", "Simplifying Mesh..."));
 
 	FVolPresMeshSimplification Simplifier(CurResultMesh);
@@ -562,6 +517,55 @@
 		}
 	}
 
+
+	Progress.EnterProgressFrame(1.f, LOCTEXT("RemoveHidden", "Removing Hidden Geometry..."));
+
+	if (Options.OcclusionPolicy == IGeometryProcessing_ApproximateActors::EOcclusionPolicy::VisibilityBased)
+	{
+		TRACE_CPUPROFILER_EVENT_SCOPE(ApproximateActorsImpl_Generate_Occlusion);
+		TRemoveOccludedTriangles<FDynamicMesh3> Remover(CurResultMesh);
+		Remover.InsideMode = EOcclusionCalculationMode::SimpleOcclusionTest;
+		Remover.TriangleSamplingMethod = EOcclusionTriangleSampling::VerticesAndCentroids;
+		Remover.AddTriangleSamples = 50;
+		Remover.AddRandomRays = 50;
+		FDynamicMeshAABBTree3 CurResultMeshSpatial(CurResultMesh, false);
+		{
+			TRACE_CPUPROFILER_EVENT_SCOPE(ApproximateActorsImpl_Generate_Occlusion_Spatial);
+			CurResultMeshSpatial.Build();
+		}
+
+		TArray<FTransform3d> NoTransforms;
+		NoTransforms.Add(FTransform3d::Identity());
+		TArray<FDynamicMeshAABBTree3*> Spatials;
+		Spatials.Add(&CurResultMeshSpatial);
+		{
+			TRACE_CPUPROFILER_EVENT_SCOPE(ApproximateActorsImpl_Generate_Occlusion_Compute);
+			Remover.Select(NoTransforms, Spatials, {}, NoTransforms);
+		}
+		if (Remover.RemovedT.Num() > 0)
+		{
+			FMeshFaceSelection Selection(CurResultMesh);
+			{
+				TRACE_CPUPROFILER_EVENT_SCOPE(ApproximateActorsImpl_Generate_Occlusion_Clean);
+				Selection.Select(Remover.RemovedT);
+				Selection.ExpandToOneRingNeighbours(1);
+				Selection.ContractBorderByOneRingNeighbours(2);
+			}
+			FDynamicMeshEditor Editor(CurResultMesh);
+			{
+				TRACE_CPUPROFILER_EVENT_SCOPE(ApproximateActorsImpl_Generate_Occlusion_Delete);
+				Editor.RemoveTriangles(Selection.AsArray(), true);
+			}
+		}
+
+		if (Options.bVerbose)
+		{
+			UE_LOG(LogApproximateActors, Warning, TEXT("Occlusion-Filtered mesh has %d triangles"), CurResultMesh->TriangleCount());
+		}
+	}
+
+
+
 	// re-enable attributes
 	CurResultMesh->EnableAttributes();
 
@@ -735,6 +739,9 @@
 		Options.MeshSimplificationPolicy = IGeometryProcessing_ApproximateActors::ESimplificationPolicy::FixedTriangleCount;
 	}
 
+	Options.bCalculateHardNormals = UseSettings.bEstimateHardNormals;
+	Options.HardNormalsAngleDeg = FMath::Clamp(UseSettings.HardNormalAngle, 0.001f, 89.99f);
+
 	Options.TextureImageSize = UseSettings.MaterialSettings.TextureSize.X;
 	Options.AntiAliasMultiSampling = FMath::Max(1, UseSettings.MultiSamplingAA);
 
diff --git a/Engine/Plugins/Experimental/MeshModelingToolset/Source/ModelingComponentsEditorOnly/Public/AssetUtils/CreateStaticMeshUtil.h b/Engine/Plugins/Experimental/MeshModelingToolset/Source/ModelingComponentsEditorOnly/Public/AssetUtils/CreateStaticMeshUtil.h
--- a/Engine/Plugins/Experimental/MeshModelingToolset/Source/ModelingComponentsEditorOnly/Public/AssetUtils/CreateStaticMeshUtil.h
+++ b/Engine/Plugins/Experimental/MeshModelingToolset/Source/ModelingComponentsEditorOnly/Public/AssetUtils/CreateStaticMeshUtil.h
@@ -65,7 +65,7 @@ namespace AssetUtils
 		// Whether to generate a nanite-enabled mesh
 		bool bGenerateNaniteEnabledMesh = false;
 		// Percentage of triangles to reduce down to for generating a coarse proxy mesh from the Nanite mesh
-		float NaniteProxyTrianglePercent = 100;
+		float NaniteProxyTrianglePercent = 0;
 
 		// Whether ray tracing will be supported on this mesh. Disable this to save memory if the generated mesh will only be rendered in the distance
 		bool bSupportRayTracing = true;
diff --git a/Engine/Source/Runtime/Engine/Classes/Engine/MeshMerging.h b/Engine/Source/Runtime/Engine/Classes/Engine/MeshMerging.h
--- a/Engine/Source/Runtime/Engine/Classes/Engine/MeshMerging.h
+++ b/Engine/Source/Runtime/Engine/Classes/Engine/MeshMerging.h
@@ -717,7 +717,7 @@ struct FMeshApproximationSettings
 	EMeshApproximationType OutputType = EMeshApproximationType::MeshAndMaterials;
 
 	/** Approximation Accuracy in Meters, will determine (eg) voxel resolution */
-	UPROPERTY(BlueprintReadWrite, EditAnywhere, Category = ShapeSettings, meta = (ClampMin = "0.1"))
+	UPROPERTY(BlueprintReadWrite, EditAnywhere, Category = ShapeSettings, meta = (DisplayName = "Approximation Accuracy (meters)", ClampMin = "0.001"))
 	float ApproximationAccuracy = 1.0;
 
 	/** Maximum allowable voxel count along main directions. This is a limit on ApproximationAccuracy */
@@ -729,7 +729,7 @@ struct FMeshApproximationSettings
 	bool bAttemptAutoThickening = true;
 
 	/** Multiplier on Approximation Accuracy used for auto-thickening */
-	UPROPERTY(BlueprintReadWrite, EditAnywhere, Category = ShapeSettings, meta = (ClampMin = "0.1", EditCondition = "bAttemptAutoThickening"))
+	UPROPERTY(BlueprintReadWrite, EditAnywhere, Category = ShapeSettings, meta = (ClampMin = "0.001", EditCondition = "bAttemptAutoThickening"))
 	float TargetMinThicknessMultiplier = 1.5;
 
 	/** Optional methods to attempt to close off the bottom of open meshes */
@@ -737,15 +737,17 @@ struct FMeshApproximationSettings
 	EMeshApproximationBaseCappingType BaseCapping = EMeshApproximationBaseCappingType::NoBaseCapping;
 
 
-	/** */
+	/** Winding Threshold controls hole filling at open mesh borders. Smaller value means "more/rounder" filling */
 	UPROPERTY(BlueprintReadWrite, EditAnywhere, Category = ShapeSettings, meta = (ClampMin = "0.01", ClampMax = "0.99"))
 	float WindingThreshold = 0.5;
 
+
+	/** If true, topological expand/contract is used to try to fill small gaps between objects. */
 	UPROPERTY(BlueprintReadWrite, EditAnywhere, Category = ShapeSettings)
 	bool bFillGaps = true;
 
-	/** Distance in Meters to expand approximation to fill gaps */
-	UPROPERTY(BlueprintReadWrite, EditAnywhere, Category = ShapeSettings, meta = (ClampMin = "0.1", EditCondition = "bFillGaps"))
+	/** Distance in Meters to expand/contract to fill gaps */
+	UPROPERTY(BlueprintReadWrite, EditAnywhere, Category = ShapeSettings, meta = (DisplayName = "Gap Filling Distance (meters)", ClampMin = "0.001", EditCondition = "bFillGaps"))
 	float GapDistance = 0.1;
 
 
@@ -765,16 +767,26 @@ struct FMeshApproximationSettings
 	float TrianglesPerM = 2.0f;
 
 	/** Allowable Geometric Deviation in Meters when SimplifyMethod incorporates a Geometric Tolerance */
-	UPROPERTY(BlueprintReadWrite, EditAnywhere, Category = SimplifySettings, meta = (ClampMin = "0.0001", EditCondition = "SimplifyMethod == EMeshApproximationSimplificationPolicy::GeometricTolerance"))
+	UPROPERTY(BlueprintReadWrite, EditAnywhere, Category = SimplifySettings, meta = (DisplayName = "Geometric Deviation (meters)", ClampMin = "0.0001", EditCondition = "SimplifyMethod == EMeshApproximationSimplificationPolicy::GeometricTolerance"))
 	float GeometricDeviation = 0.1f;
 
+
+	/** If true, normal angle will be used to estimate hard normals */
+	UPROPERTY(BlueprintReadWrite, EditAnywhere, Category = NormalsSettings)
+	bool bEstimateHardNormals = true;
+
+	UPROPERTY(BlueprintReadWrite, EditAnywhere, Category = NormalsSettings, meta = (ClampMin = "0.0", ClampMax = "90.0", EditCondition = "bEstimateHardNormals"))
+	float HardNormalAngle = 60.0f;
+
+
+
 	/** Whether to generate a nanite-enabled mesh */
 	UPROPERTY(BlueprintReadWrite, EditAnywhere, Category = MeshSettings)
 	bool bGenerateNaniteEnabledMesh = false;
 
 	/** Percentage of triangles to reduce down to for generating a coarse proxy mesh from the Nanite mesh */
 	UPROPERTY(BlueprintReadWrite, EditAnywhere, AdvancedDisplay, Category = MeshSettings, meta = (EditConditionHides, EditCondition = "bGenerateNaniteEnabledMesh", ClampMin = 0, ClampMax = 100))
-	float NaniteProxyTrianglePercent = 100;
+	float NaniteProxyTrianglePercent = 0;
 
 	/** Whether ray tracing will be supported on this mesh. Disable this to save memory if the generated mesh will only be rendered in the distance. */
 	UPROPERTY(BlueprintReadWrite, EditAnywhere, Category = MeshSettings)
@@ -784,6 +796,7 @@ struct FMeshApproximationSettings
 	UPROPERTY(BlueprintReadWrite, EditAnywhere, Category = MeshSettings)
 	bool bAllowDistanceField = true;
 
+
 	/** If Value is > 1, Multisample output baked textures by this amount in each direction (eg 4 == 16x supersampling) */
 	UPROPERTY(BlueprintReadWrite, EditAnywhere, Category = MaterialSettings, meta = (ClampMin = "0", ClampMax = "8", UIMin = "0", UIMax = "4"))
 	int32 MultiSamplingAA = 0;
