Using "Remove Selected" for LOD Bone Reduction doesn't inactivate bones in the Skeleton Tree
#fix use (clamped) GetNumLODs() -1, rather than going through MeshObject directly as it might be null when re-registering/creating of the component

diff --git a/Engine/Source/Runtime/Engine/Private/Components/SkinnedMeshComponent.cpp b/Engine/Source/Runtime/Engine/Private/Components/SkinnedMeshComponent.cpp
--- a/Engine/Source/Runtime/Engine/Private/Components/SkinnedMeshComponent.cpp
+++ b/Engine/Source/Runtime/Engine/Private/Components/SkinnedMeshComponent.cpp
@@ -3072,11 +3072,11 @@ bool USkinnedMeshComponent::UpdateLODStatus_Internal(int32 InMasterPoseComponent
 		const int32 LODBias = GSkeletalMeshLODBias;
 #endif
 
-		int32 MinLodIndex = ComputeMinLOD();
-		int32 MaxLODIndex = MinLodIndex;
+		const int32 MinLodIndex = ComputeMinLOD();
+		const int32 MaxLODIndex = FMath::Max(GetNumLODs() - 1, MinLodIndex);
+
 		if (MeshObject)
 		{
-			MaxLODIndex = MeshObject->GetSkeletalMeshRenderData().LODRenderData.Num() - 1;
 			MaxDistanceFactor = MeshObject->MaxDistanceFactor;
 		}
 
