User is not prompted to confirm copying anim notifies when some notifies fall outside the range of the destination anim sequence
User is not prompted to delete preexisting anim notifies from destination anim sequence when copying anim notifies from another animation
#fix changed some behaviour around CopyAnimNotifies UFunction
    - Added bool flag allowing for deleting notifies when doing the copy
    - Added log warnings to CopyNotifies, and LogAnimation category to filter when applying Anim Modifiers

diff --git a/Engine/Source/Editor/AnimationBlueprintLibrary/Private/AnimationBlueprintLibrary.cpp b/Engine/Source/Editor/AnimationBlueprintLibrary/Private/AnimationBlueprintLibrary.cpp
--- a/Engine/Source/Editor/AnimationBlueprintLibrary/Private/AnimationBlueprintLibrary.cpp
+++ b/Engine/Source/Editor/AnimationBlueprintLibrary/Private/AnimationBlueprintLibrary.cpp
@@ -916,7 +916,7 @@ void UAnimationBlueprintLibrary::ReplaceAnimNotifies(UAnimSequenceBase* Animatio
 	ReplaceAnimNotifies_Helper(AnimationSequenceBase, OldNotifyClass.Get(), NewNotifyClass.Get(), OnNotifyReplaced, FOnNotifyStateReplaced());
 }
 
-void UAnimationBlueprintLibrary::CopyAnimNotifiesFromSequence(UAnimSequenceBase* SourceAnimationSequenceBase, UAnimSequenceBase* DestinationAnimSequenceBase)
+void UAnimationBlueprintLibrary::CopyAnimNotifiesFromSequence(UAnimSequenceBase* SourceAnimationSequenceBase, UAnimSequenceBase* DestinationAnimSequenceBase, bool bDeleteExistingNotifies)
 {
 	if (SourceAnimationSequenceBase == nullptr)
 	{
@@ -933,7 +933,7 @@ void UAnimationBlueprintLibrary::CopyAnimNotifiesFromSequence(UAnimSequenceBase*
 	else
 	{
 		const bool bShowDialogs = false;
-		UE::Anim::CopyNotifies(SourceAnimationSequenceBase, DestinationAnimSequenceBase, bShowDialogs);
+		UE::Anim::CopyNotifies(SourceAnimationSequenceBase, DestinationAnimSequenceBase, bShowDialogs, bDeleteExistingNotifies);
 	}
 }
 
diff --git a/Engine/Source/Editor/AnimationBlueprintLibrary/Public/AnimationBlueprintLibrary.h b/Engine/Source/Editor/AnimationBlueprintLibrary/Public/AnimationBlueprintLibrary.h
--- a/Engine/Source/Editor/AnimationBlueprintLibrary/Public/AnimationBlueprintLibrary.h
+++ b/Engine/Source/Editor/AnimationBlueprintLibrary/Public/AnimationBlueprintLibrary.h
@@ -219,7 +219,7 @@ public:
 
 	/** Copies animation notifies from Src Animation Sequence to Dest. Creates anim notify tracks as necessary. Returns true on success. */
 	UFUNCTION(BlueprintCallable, Category = "AnimationBlueprintLibrary|NotifyEvents")
-	static void CopyAnimNotifiesFromSequence(UAnimSequenceBase* SourceAnimationSequenceBase, UAnimSequenceBase* DestinationAnimationSequenceBase);
+	static void CopyAnimNotifiesFromSequence(UAnimSequenceBase* SourceAnimationSequenceBase, UAnimSequenceBase* DestinationAnimationSequenceBase, bool bDeleteExistingNotifies = false);
 
 	// Notify Tracks
 
diff --git a/Engine/Source/Runtime/Engine/Classes/Animation/AnimSequenceHelpers.h b/Engine/Source/Runtime/Engine/Classes/Animation/AnimSequenceHelpers.h
--- a/Engine/Source/Runtime/Engine/Classes/Animation/AnimSequenceHelpers.h
+++ b/Engine/Source/Runtime/Engine/Classes/Animation/AnimSequenceHelpers.h
@@ -90,13 +90,14 @@ namespace Anim {
 	/**
 	* Copies over any individual curve from FRawCurveTracks to a UAnimDataModel instance targeted by the provided controller.
 	*
-	* @param	SourceAnimSeq		Animation Sequence to copy the notifies from
-	* @param	Skeleton			Animation Sequence to copy the notifies to
-	* @param	bShowDialogs		Whether or not to show any user-facing dialogs for confirmation of the copy
+	* @param	SourceAnimSeq				Animation Sequence to copy the notifies from
+	* @param	DestAnimSeq					Animation Sequence to copy the notifies to
+	* @param	bShowDialogs				Whether or not to show any user-facing dialogs for confirmation of the copy
+	* @param	bDeleteExistingNotifies		Whether or not to delete all notifies found on the destination sequence
 	*
 	* @return Whether or not the copy was succesful
 	*/
-	ENGINE_API bool CopyNotifies(const UAnimSequenceBase* SourceAnimSeq, UAnimSequenceBase* DestAnimSeq, bool bShowDialogs = true);
+	ENGINE_API bool CopyNotifies(const UAnimSequenceBase* SourceAnimSeq, UAnimSequenceBase* DestAnimSeq, bool bShowDialogs = true, bool bDeleteExistingNotifies = false);
 
 	namespace AnimationData
 	{
diff --git a/Engine/Source/Runtime/Engine/Private/Animation/AnimSequenceHelpers.cpp b/Engine/Source/Runtime/Engine/Private/Animation/AnimSequenceHelpers.cpp
--- a/Engine/Source/Runtime/Engine/Private/Animation/AnimSequenceHelpers.cpp
+++ b/Engine/Source/Runtime/Engine/Private/Animation/AnimSequenceHelpers.cpp
@@ -396,7 +396,7 @@ void Retargeting::RetargetPose(FCompactPose& InOutPose, const FName& RetargetSou
 
 
 #if WITH_EDITOR
-bool CopyNotifies(const UAnimSequenceBase* SourceAnimSeq, UAnimSequenceBase* DestAnimSeq, bool bShowDialogs)
+bool CopyNotifies(const UAnimSequenceBase* SourceAnimSeq, UAnimSequenceBase* DestAnimSeq, bool bShowDialogs, bool bDeleteNotifies)
 {
 	// Abort if source == destination.
 	if (SourceAnimSeq == DestAnimSeq)
@@ -407,7 +407,7 @@ bool CopyNotifies(const UAnimSequenceBase* SourceAnimSeq, UAnimSequenceBase* Des
 	// If the destination sequence is shorter than the source sequence, we'll be dropping notifies that
 	// occur at later times than the dest sequence is long.  Give the user a chance to abort if we
 	// find any notifies that won't be copied over.
-	if (bShowDialogs && DestAnimSeq->GetPlayLength() < SourceAnimSeq->GetPlayLength())
+	if (DestAnimSeq->GetPlayLength() < SourceAnimSeq->GetPlayLength())
 	{
 		for (int32 NotifyIndex = 0; NotifyIndex < SourceAnimSeq->Notifies.Num(); ++NotifyIndex)
 		{
@@ -415,7 +415,9 @@ bool CopyNotifies(const UAnimSequenceBase* SourceAnimSeq, UAnimSequenceBase* Des
 			const FAnimNotifyEvent& SrcNotifyEvent = SourceAnimSeq->Notifies[NotifyIndex];
 			if (SrcNotifyEvent.GetTriggerTime() > DestAnimSeq->GetPlayLength())
 			{
-				const bool bProceed = EAppReturnType::Yes == FMessageDialog::Open(EAppMsgType::YesNo, NSLOCTEXT("UnrealEd", "SomeNotifiesWillNotBeCopiedQ", "Some notifies will not be copied because the destination sequence is not long enough.  Proceed?"));
+				UE_LOG(LogAnimation, Warning, TEXT("Animation Notify trigger time %f falls outside of the destination animation sequence its length %f, notify will not be copied."), SrcNotifyEvent.GetTriggerTime(), DestAnimSeq->GetPlayLength());
+				
+				const bool bProceed = !bShowDialogs || EAppReturnType::Yes == FMessageDialog::Open(EAppMsgType::YesNo, NSLOCTEXT("UnrealEd", "SomeNotifiesWillNotBeCopiedQ", "Some notifies will not be copied because the destination sequence is not long enough.  Proceed?"));
 				if (!bProceed)
 				{
 					return false;
@@ -430,14 +432,16 @@ bool CopyNotifies(const UAnimSequenceBase* SourceAnimSeq, UAnimSequenceBase* Des
 
 	// If the destination sequence contains any notifies, ask the user if they'd like
 	// to delete the existing notifies before copying over from the source sequence.
-	if (bShowDialogs && DestAnimSeq->Notifies.Num() > 0)
+	if (DestAnimSeq->Notifies.Num() > 0)
 	{
-		const bool bDeleteExistingNotifies = EAppReturnType::Yes == FMessageDialog::Open(EAppMsgType::YesNo, FText::Format(
-			NSLOCTEXT("UnrealEd", "DestSeqAlreadyContainsNotifiesMergeQ", "The destination sequence already contains {0} notifies.  Delete these before copying?"), FText::AsNumber(DestAnimSeq->Notifies.Num())));
+		const bool bDeleteExistingNotifies = bDeleteNotifies || (bShowDialogs && EAppReturnType::Yes == FMessageDialog::Open(EAppMsgType::YesNo, FText::Format(
+			NSLOCTEXT("UnrealEd", "DestSeqAlreadyContainsNotifiesMergeQ", "The destination sequence already contains {0} notifies.  Delete these before copying?"), FText::AsNumber(DestAnimSeq->Notifies.Num()))));
+
 		if (bDeleteExistingNotifies)
 		{
 			DestAnimSeq->Notifies.Empty();
 			DestAnimSeq->MarkPackageDirty();
+			DestAnimSeq->RefreshCacheData();
 		}
 	}
 
