Anim node pose by name is expensive
#fix in addition to previous optimizations also added an optimized evaluation path for a single pose curve (used by pose by name node), roughly 3x speed-up for 24 instances of AnimBP with 4 pose by name nodes using pose asset with 50 stored poses

diff --git a/Engine/Source/Runtime/Engine/Private/Animation/PoseAsset.cpp b/Engine/Source/Runtime/Engine/Private/Animation/PoseAsset.cpp
--- a/Engine/Source/Runtime/Engine/Private/Animation/PoseAsset.cpp
+++ b/Engine/Source/Runtime/Engine/Private/Animation/PoseAsset.cpp
@@ -591,27 +591,7 @@ bool UPoseAsset::GetAnimationPose(struct FAnimationPoseData& OutAnimationPoseDat
 
 		const FBoneContainer& RequiredBones = OutPose.GetBoneContainer();
 		USkeleton* MySkeleton = GetSkeleton();
-
-		// TLS storage for working data
-		FPoseAssetEvalData& EvalData = FPoseAssetEvalData::Get();
-
-		// this contains compact bone pose list that this pose cares
-		TArray<FBoneIndices>& BoneIndices = EvalData.BoneIndices;
-		const int32 TrackNum = PoseContainer.Tracks.Num();
-		BoneIndices.SetNumUninitialized(TrackNum, false);
-
-		for(int32 TrackIndex = 0; TrackIndex < TrackNum; ++TrackIndex)
-		{
-			const int32 SkeletonBoneIndex = PoseContainer.TrackBoneIndices[TrackIndex];
-			const FCompactPoseBoneIndex PoseBoneIndex = RequiredBones.GetCompactPoseIndexFromSkeletonIndex(SkeletonBoneIndex);
-			// we add even if it's invalid because we want it to match with track index
-			BoneIndices[TrackIndex].SkeletonBoneIndex = SkeletonBoneIndex;
-			BoneIndices[TrackIndex].CompactBoneIndex = PoseBoneIndex;
-		}
-
-		// you could only have morphtargets
-		// so can't return here yet when bone indices is empty
-
+		
 		check(PoseContainer.IsValid());
 
 		if (bAdditivePose)
@@ -622,158 +602,251 @@ bool UPoseAsset::GetAnimationPose(struct FAnimationPoseData& OutAnimationPoseDat
 		{
 			OutPose.ResetToRefPose();
 		}
-		
-		const bool bNormalizeWeight = bAdditivePose == false;
-		const int32 NumPoses = PoseContainer.Poses.Num();
-		TArray<float>& PoseWeights = EvalData.PoseWeights;
-		PoseWeights.Reset();
-		PoseWeights.SetNumZeroed(NumPoses, false);
 
-		TArray<int32>& WeightedPoseIndices = EvalData.PoseWeightedIndices;
-		WeightedPoseIndices.Reset();
-
-		TArray<bool>& WeightedPoses = EvalData.WeightedPoses;
-		WeightedPoses.Reset();
-		WeightedPoses.SetNumZeroed(NumPoses, false);
-
-		float TotalWeight = 0.f;
-		// we iterate through to see if we have that corresponding pose
+		const int32 TrackNum = PoseContainer.Tracks.Num();
 
-		const int32 NumPoseCurves = ExtractionContext.PoseCurves.Num();
-		for (int32 CurveIndex = 0; CurveIndex < NumPoseCurves; ++CurveIndex)
+		// Single pose optimized evaluation path - explicitly used by PoseByName Animation Node
+		if (ExtractionContext.PoseCurves.Num() == 1)
 		{
-			const FPoseCurve& Curve = ExtractionContext.PoseCurves[CurveIndex];
+			const FPoseCurve& Curve = ExtractionContext.PoseCurves[0];
 			const int32 PoseIndex = Curve.PoseIndex; 
 			if (ensure(PoseIndex != INDEX_NONE))
 			{
-				const FPoseData& PoseData = PoseContainer.Poses[PoseIndex];
-				const float Value = Curve.Value;
+				const FPoseData& Pose = PoseContainer.Poses[PoseIndex];
+				// Clamp weight for non-additive pose assets rather than normalizing the weight
+				const float Weight = bAdditivePose ? Curve.Value : FMath::Clamp(Curve.Value, 0.f, 1.f);
 
-				// we only add to the list if it's not additive Or if it's additive, we don't want to add base pose index
-				// and has weight
-				if ((!bAdditivePose || PoseIndex != BasePoseIndex) && FAnimationRuntime::HasWeight(Value))
+				// Only generate pose if the single weight is actually relevant
+				if(FAnimWeight::IsRelevant(Weight))
 				{
-					TotalWeight += Value;
+					// Blend curve
+					PoseContainer.BlendPoseCurve(&Pose, OutCurve, Weight);
 
-					// Set pose weight and bit, and add weighted pose index
-					PoseWeights[PoseIndex] = Value;
-					WeightedPoseIndices.Add(PoseIndex);
-					WeightedPoses[PoseIndex] = true;
+					// Per-track (bone) transform
+					for (int32 TrackIndex = 0; TrackIndex < TrackNum; ++TrackIndex)
+					{
+						const int32 SkeletonBoneIndex = PoseContainer.TrackBoneIndices[TrackIndex];
+						const FCompactPoseBoneIndex CompactIndex = RequiredBones.GetCompactPoseIndexFromSkeletonIndex(SkeletonBoneIndex);
+						
+						// If bone index is invalid, or not required for the pose - skip
+						if (CompactIndex == INDEX_NONE || !ExtractionContext.IsBoneRequired(CompactIndex.GetInt()))
+						{
+							continue;
+						}
+					
+						// Check if this track is part of the pose
+						const TArray<FPoseAssetInfluence>& PoseInfluences = PoseContainer.TrackPoseInfluenceIndices[TrackIndex].Influences;						
+						const int32 InfluenceIndex = PoseInfluences.IndexOfByPredicate([PoseIndex](const FPoseAssetInfluence& Influence) -> bool
+						{
+							return Influence.PoseIndex == PoseIndex;
+						});
+
+						if (InfluenceIndex != INDEX_NONE)
+						{
+							FTransform& OutBoneTransform =  OutPose[CompactIndex];
+							const FPoseAssetInfluence& Influence = PoseInfluences[InfluenceIndex];
+							const int32& BonePoseIndex = Influence.BoneTransformIndex;
+
+							const FPoseData& PoseData = PoseContainer.Poses[PoseIndex];
+							const FTransform& BonePose = PoseData.LocalSpacePose[BonePoseIndex];
+
+							// Apply additive, overriede or blend using pose weight
+							if (bAdditivePose)
+							{
+								BlendFromIdentityAndAccumulateAdditively_Custom(OutBoneTransform, BonePose, Weight);
+							}
+							else if (FAnimWeight::IsFullWeight(Weight))
+							{
+								OutBoneTransform = BonePose;
+							}
+							else
+							{
+								OutBoneTransform = OutBoneTransform * ScalarRegister( 1 - Weight);
+								OutBoneTransform.AccumulateWithShortestRotation(BonePose, ScalarRegister(Weight));
+							}
+
+							// Retarget the bone transform
+							FAnimationRuntime::RetargetBoneTransform(MySkeleton, GetRetargetTransformsSourceName(), GetRetargetTransforms(), OutBoneTransform, SkeletonBoneIndex, CompactIndex, RequiredBones, bAdditivePose);
+							OutBoneTransform.NormalizeRotation();
+						}
+					}					
+					
+					return true;
 				}
 			}
 		}
-
-		const int32 TotalNumberOfValidPoses = WeightedPoseIndices.Num();
-		if (TotalNumberOfValidPoses > 0)
+		else
 		{
-			//if full pose, we'll have to normalize by weight
-			if (bNormalizeWeight && TotalWeight > 1.f)
+			// TLS storage for working data
+			FPoseAssetEvalData& EvalData = FPoseAssetEvalData::Get();
+
+			// this contains compact bone pose list that this pose cares
+			TArray<FBoneIndices>& BoneIndices = EvalData.BoneIndices;
+			BoneIndices.SetNumUninitialized(TrackNum, false);
+
+			for(int32 TrackIndex = 0; TrackIndex < TrackNum; ++TrackIndex)
 			{
-				for (const int32& WeightedPoseIndex : WeightedPoseIndices)
+				const int32 SkeletonBoneIndex = PoseContainer.TrackBoneIndices[TrackIndex];
+				const FCompactPoseBoneIndex PoseBoneIndex = RequiredBones.GetCompactPoseIndexFromSkeletonIndex(SkeletonBoneIndex);
+				// we add even if it's invalid because we want it to match with track index
+				BoneIndices[TrackIndex].SkeletonBoneIndex = SkeletonBoneIndex;
+				BoneIndices[TrackIndex].CompactBoneIndex = PoseBoneIndex;
+			}
+
+			// you could only have morphtargets
+			// so can't return here yet when bone indices is empty
+			const bool bNormalizeWeight = bAdditivePose == false;
+			const int32 NumPoses = PoseContainer.Poses.Num();
+			TArray<float>& PoseWeights = EvalData.PoseWeights;
+			PoseWeights.Reset();
+			PoseWeights.SetNumZeroed(NumPoses, false);
+
+			TArray<int32>& WeightedPoseIndices = EvalData.PoseWeightedIndices;
+			WeightedPoseIndices.Reset();
+
+			TArray<bool>& WeightedPoses = EvalData.WeightedPoses;
+			WeightedPoses.Reset();
+			WeightedPoses.SetNumZeroed(NumPoses, false);
+
+			float TotalWeight = 0.f;
+			// we iterate through to see if we have that corresponding pose
+
+			const int32 NumPoseCurves = ExtractionContext.PoseCurves.Num();
+			for (int32 CurveIndex = 0; CurveIndex < NumPoseCurves; ++CurveIndex)
+			{
+				const FPoseCurve& Curve = ExtractionContext.PoseCurves[CurveIndex];
+				const int32 PoseIndex = Curve.PoseIndex; 
+				if (ensure(PoseIndex != INDEX_NONE))
 				{
-					float& PoseWeight = PoseWeights[WeightedPoseIndex];
-					PoseWeight /= TotalWeight;
+					const FPoseData& PoseData = PoseContainer.Poses[PoseIndex];
+					const float Value = Curve.Value;
 
-					// Do curve blending inline as we are looping over weights anyway
-					const FPoseData& Pose = PoseContainer.Poses[WeightedPoseIndex];
-					PoseContainer.BlendPoseCurve(&Pose, OutCurve, PoseWeight);
+					// we only add to the list if it's not additive Or if it's additive, we don't want to add base pose index
+					// and has weight
+					if ((!bAdditivePose || PoseIndex != BasePoseIndex) && FAnimationRuntime::HasWeight(Value))
+					{
+						TotalWeight += Value;
+
+						// Set pose weight and bit, and add weighted pose index
+						PoseWeights[PoseIndex] = Value;
+						WeightedPoseIndices.Add(PoseIndex);
+						WeightedPoses[PoseIndex] = true;
+					}
 				}
 			}
-            else
-            {
- 				// Take the matching curve weights from the selected poses, and blend them using the
-				// weighting that we need from each pose. This is much faster than grabbing each
-				// blend curve and blending them in their entirety, especially when there are very
-				// few active curves for each pose and many curves for the entire pose asset.
-				for (const int32& WeightedPoseIndex : WeightedPoseIndices)
-				{
-					const FPoseData& Pose = PoseContainer.Poses[WeightedPoseIndex];
-					const float& Weight = PoseWeights[WeightedPoseIndex];
-					PoseContainer.BlendPoseCurve(&Pose, OutCurve, Weight);
-				}
-            }
 
-			// Final per-track (bone) transform
-			FTransform OutBoneTransform;
-			for (int32 TrackIndex = 0; TrackIndex < TrackNum; ++TrackIndex)
+			const int32 TotalNumberOfValidPoses = WeightedPoseIndices.Num();
+			if (TotalNumberOfValidPoses > 0)
 			{
-				const FCompactPoseBoneIndex& CompactIndex = BoneIndices[TrackIndex].CompactBoneIndex;
-
-				// If bone index is invalid, or not required for the pose - skip
-				if (CompactIndex == INDEX_NONE || !ExtractionContext.IsBoneRequired(CompactIndex.GetInt()))
+				//if full pose, we'll have to normalize by weight
+				if (bNormalizeWeight && TotalWeight > 1.f)
 				{
-					continue;
-				}
-				
-				const TArray<FPoseAssetInfluence>& PoseInfluences = PoseContainer.TrackPoseInfluenceIndices[TrackIndex].Influences;
+					for (const int32& WeightedPoseIndex : WeightedPoseIndices)
+					{
+						float& PoseWeight = PoseWeights[WeightedPoseIndex];
+						PoseWeight /= TotalWeight;
 
-				// When additive, or for any bone that has no pose influences. Set transform to input.
-				if (bAdditivePose || PoseInfluences.Num() == 0)
+						// Do curve blending inline as we are looping over weights anyway
+						const FPoseData& Pose = PoseContainer.Poses[WeightedPoseIndex];
+						PoseContainer.BlendPoseCurve(&Pose, OutCurve, PoseWeight);
+					}
+				}
+				else
 				{
-					OutBoneTransform = OutPose[CompactIndex];
+					// Take the matching curve weights from the selected poses, and blend them using the
+					// weighting that we need from each pose. This is much faster than grabbing each
+					// blend curve and blending them in their entirety, especially when there are very
+					// few active curves for each pose and many curves for the entire pose asset.
+					for (const int32& WeightedPoseIndex : WeightedPoseIndices)
+					{
+						const FPoseData& Pose = PoseContainer.Poses[WeightedPoseIndex];
+						const float& Weight = PoseWeights[WeightedPoseIndex];
+						PoseContainer.BlendPoseCurve(&Pose, OutCurve, Weight);
+					}
 				}
 
-				const int32 NumInfluences = PoseInfluences.Num();
-				if (NumInfluences)
+				// Final per-track (bone) transform
+				FTransform OutBoneTransform;
+				for (int32 TrackIndex = 0; TrackIndex < TrackNum; ++TrackIndex)
 				{
-					float TotalLocalWeight = 0.f;
-					bool bSet = false;
-					// Only loop over poses known to influence this track its final transform
-					for (int32 Index = 0; Index < NumInfluences; ++Index)
+					const FCompactPoseBoneIndex& CompactIndex = BoneIndices[TrackIndex].CompactBoneIndex;
+
+					// If bone index is invalid, or not required for the pose - skip
+					if (CompactIndex == INDEX_NONE || !ExtractionContext.IsBoneRequired(CompactIndex.GetInt()))
 					{
-						const FPoseAssetInfluence& Influence = PoseInfluences[Index];
-						const int32& PoseIndex = Influence.PoseIndex;
-						const int32& BonePoseIndex = Influence.BoneTransformIndex;
-						
-						// Only processs pose if its weighted
-						if(WeightedPoses[PoseIndex])
-						{
-							const float& Weight = PoseWeights[PoseIndex];
-							TotalLocalWeight += Weight;
-						
-							const FPoseData& PoseData = PoseContainer.Poses[PoseIndex];
-							const FTransform& BonePose = PoseData.LocalSpacePose[BonePoseIndex];
+						continue;
+					}
+					
+					const TArray<FPoseAssetInfluence>& PoseInfluences = PoseContainer.TrackPoseInfluenceIndices[TrackIndex].Influences;
 
-							// Set weighted value For first pose, if applicable
-							if(!bSet && !bAdditivePose)
-							{								
-								OutBoneTransform = BonePose * ScalarRegister(Weight);
-								bSet = true;
-							}
-							else
+					// When additive, or for any bone that has no pose influences. Set transform to input.
+					if (bAdditivePose || PoseInfluences.Num() == 0)
+					{
+						OutBoneTransform = OutPose[CompactIndex];
+					}
+
+					const int32 NumInfluences = PoseInfluences.Num();
+					if (NumInfluences)
+					{
+						float TotalLocalWeight = 0.f;
+						bool bSet = false;
+						// Only loop over poses known to influence this track its final transform
+						for (int32 Index = 0; Index < NumInfluences; ++Index)
+						{
+							const FPoseAssetInfluence& Influence = PoseInfluences[Index];
+							const int32& PoseIndex = Influence.PoseIndex;
+							const int32& BonePoseIndex = Influence.BoneTransformIndex;
+							
+							// Only processs pose if its weighted
+							if(WeightedPoses[PoseIndex])
 							{
-								if (bAdditivePose)
+								const float& Weight = PoseWeights[PoseIndex];
+								TotalLocalWeight += Weight;
+							
+								const FPoseData& PoseData = PoseContainer.Poses[PoseIndex];
+								const FTransform& BonePose = PoseData.LocalSpacePose[BonePoseIndex];
+
+								// Set weighted value For first pose, if applicable
+								if(!bSet && !bAdditivePose)
 								{
-									BlendFromIdentityAndAccumulateAdditively_Custom(OutBoneTransform, BonePose, Weight);
+									OutBoneTransform = BonePose * ScalarRegister(Weight);
+									bSet = true;
 								}
 								else
 								{
-									OutBoneTransform.AccumulateWithShortestRotation(BonePose, ScalarRegister(Weight));
+									if (bAdditivePose)
+									{
+										BlendFromIdentityAndAccumulateAdditively_Custom(OutBoneTransform, BonePose, Weight);
+									}
+									else
+									{
+										OutBoneTransform.AccumulateWithShortestRotation(BonePose, ScalarRegister(Weight));
+									}
 								}
 							}
 						}
-					}
 
-					// In case no influencing poses had any weight, set transform to input
-					if(!FAnimWeight::IsRelevant(TotalLocalWeight))
-					{
-						OutBoneTransform = OutPose[CompactIndex];
-					}
-					else if (!FAnimWeight::IsFullWeight(TotalLocalWeight) && !bAdditivePose)
-					{
-						OutBoneTransform.AccumulateWithShortestRotation(OutPose[CompactIndex], ScalarRegister(1.f - TotalLocalWeight));
+						// In case no influencing poses had any weight, set transform to input
+						if(!FAnimWeight::IsRelevant(TotalLocalWeight))
+						{
+							OutBoneTransform = OutPose[CompactIndex];
+						}
+						else if (!FAnimWeight::IsFullWeight(TotalLocalWeight) && !bAdditivePose)
+						{
+							OutBoneTransform.AccumulateWithShortestRotation(OutPose[CompactIndex], ScalarRegister(1.f - TotalLocalWeight));
+						}
 					}
-				}
 
-				// Retarget the blended transform, and copy to output pose
-				FAnimationRuntime::RetargetBoneTransform(MySkeleton, GetRetargetTransformsSourceName(), GetRetargetTransforms(), OutBoneTransform,  BoneIndices[TrackIndex].SkeletonBoneIndex, CompactIndex, RequiredBones, bAdditivePose);
+					// Retarget the blended transform, and copy to output pose
+					FAnimationRuntime::RetargetBoneTransform(MySkeleton, GetRetargetTransformsSourceName(), GetRetargetTransforms(), OutBoneTransform,  BoneIndices[TrackIndex].SkeletonBoneIndex, CompactIndex, RequiredBones, bAdditivePose);
 
-				OutPose[CompactIndex] = OutBoneTransform;
-				OutPose[CompactIndex].NormalizeRotation();
-			}
+					OutPose[CompactIndex] = OutBoneTransform;
+					OutPose[CompactIndex].NormalizeRotation();
+				}
 
-			return true;
-		}
+				return true;
+			}
+		}		
 	}
 
 	return false;
