Fix Actor with HISM Component copy-paste failed check: add missing NotifyIfStaticMeshChanged in PostEditImport

diff --git a/Engine/Source/Runtime/Engine/Classes/Components/StaticMeshComponent.h b/Engine/Source/Runtime/Engine/Classes/Components/StaticMeshComponent.h
--- a/Engine/Source/Runtime/Engine/Classes/Components/StaticMeshComponent.h
+++ b/Engine/Source/Runtime/Engine/Classes/Components/StaticMeshComponent.h
@@ -442,6 +442,7 @@ public:
 	virtual void BeginCacheForCookedPlatformData(const ITargetPlatform* TargetPlatform) override;
 	virtual bool IsCachedCookedPlatformDataLoaded(const ITargetPlatform* TargetPlatform) override;
 	virtual void PostDuplicate(bool bDuplicateForPIE) override;
+	virtual void PostEditImport() override;
 	virtual void InitializeComponent() override;
 	virtual void UpdateBounds() override;
 #endif // WITH_EDITOR
diff --git a/Engine/Source/Runtime/Engine/Private/Components/StaticMeshComponent.cpp b/Engine/Source/Runtime/Engine/Private/Components/StaticMeshComponent.cpp
--- a/Engine/Source/Runtime/Engine/Private/Components/StaticMeshComponent.cpp
+++ b/Engine/Source/Runtime/Engine/Private/Components/StaticMeshComponent.cpp
@@ -635,6 +635,13 @@ void UStaticMeshComponent::PostDuplicate(bool bDuplicateForPIE)
 	Super::PostDuplicate(bDuplicateForPIE);
 }
 
+void UStaticMeshComponent::PostEditImport()
+{
+	NotifyIfStaticMeshChanged();
+
+	Super::PostEditImport();
+}
+
 #endif // #if WITH_EDITOR
 
 void UStaticMeshComponent::OnRegister()
