Potential crash due to unused variable in GetRefVertexLocationTyped
#fix removing dead code - caused a crash for licenseed

diff --git a/Engine/Source/Runtime/Engine/Private/SkeletalMesh.cpp b/Engine/Source/Runtime/Engine/Private/SkeletalMesh.cpp
--- a/Engine/Source/Runtime/Engine/Private/SkeletalMesh.cpp
+++ b/Engine/Source/Runtime/Engine/Private/SkeletalMesh.cpp
@@ -6730,12 +6730,9 @@ FVector GetRefVertexLocationTyped(
 	for (int32 InfluenceIndex = 0; InfluenceIndex < MaxBoneInfluences; InfluenceIndex++)
 #endif
 	{
-		const int32 MeshBoneIndex = Section.BoneMap[SkinWeightVertexBuffer.GetBoneIndex(BufferVertIndex, InfluenceIndex)];
 		const float	Weight = (float)SkinWeightVertexBuffer.GetBoneWeight(BufferVertIndex, InfluenceIndex) / 255.0f;
 		{
-			const FMatrix BoneTransformMatrix = FMatrix::Identity;//Mesh->GetComposedRefPoseMatrix(MeshBoneIndex);
-			const FMatrix RefToLocal = Mesh->GetRefBasesInvMatrix()[MeshBoneIndex] * BoneTransformMatrix;
-
+			const FMatrix BoneTransformMatrix = FMatrix::Identity;
 			SkinnedPos += BoneTransformMatrix.TransformPosition(PositionBuffer.VertexPosition(BufferVertIndex)) * Weight;
 		}
 	}
@@ -6774,10 +6771,8 @@ void GetRefTangentBasisTyped(const USkeletalMesh* Mesh, const FSkelMeshRenderSec
 	for (int32 InfluenceIndex = 0; InfluenceIndex < MaxBoneInfluences; InfluenceIndex++)
 #endif
 	{
-		const int32 MeshBoneIndex = Section.BoneMap[SkinWeightVertexBuffer.GetBoneIndex(BufferVertIndex, InfluenceIndex)];
 		const float	Weight = (float)SkinWeightVertexBuffer.GetBoneWeight(BufferVertIndex, InfluenceIndex) / 255.0f;
-		const FMatrix BoneTransformMatrix = FMatrix::Identity;//Mesh->GetComposedRefPoseMatrix(MeshBoneIndex);
-		//const FMatrix RefToLocal = Mesh->RefBasesInvMatrix[MeshBoneIndex] * BoneTransformMatrix;
+		const FMatrix BoneTransformMatrix = FMatrix::Identity;
 		OutTangentX += BoneTransformMatrix.TransformVector(VertexTangentX) * Weight;
 		OutTangentY += BoneTransformMatrix.TransformVector(VertexTangentY) * Weight;
 		OutTangentZ += BoneTransformMatrix.TransformVector(VertexTangentZ) * Weight;
