Free up DuplicatedVertexBuffer CPU resource when not in editor to save memory

diff --git a/Engine/Source/Runtime/Engine/Private/GPUSkinCache.cpp b/Engine/Source/Runtime/Engine/Private/GPUSkinCache.cpp
--- a/Engine/Source/Runtime/Engine/Private/GPUSkinCache.cpp
+++ b/Engine/Source/Runtime/Engine/Private/GPUSkinCache.cpp
@@ -1212,7 +1212,9 @@ void FGPUSkinCache::DispatchUpdateSkinTangents(FRHICommandListImmediate& RHICmdL
 
 			if (!GAllowDupedVertsForRecomputeTangents)
 			{
+#if WITH_EDITOR
 				check(LodData.RenderSections[SectionIndex].DuplicatedVerticesBuffer.DupVertData.Num() && LodData.RenderSections[SectionIndex].DuplicatedVerticesBuffer.DupVertIndexData.Num());
+#endif
 				DispatchData.DuplicatedIndices = LodData.RenderSections[SectionIndex].DuplicatedVerticesBuffer.DuplicatedVerticesIndexBuffer.VertexBufferSRV;
 				DispatchData.DuplicatedIndicesIndices = LodData.RenderSections[SectionIndex].DuplicatedVerticesBuffer.LengthAndIndexDuplicatedVerticesIndexBuffer.VertexBufferSRV;
 			}
diff --git a/Engine/Source/Runtime/Engine/Private/SkeletalMeshLODRenderData.cpp b/Engine/Source/Runtime/Engine/Private/SkeletalMeshLODRenderData.cpp
--- a/Engine/Source/Runtime/Engine/Private/SkeletalMeshLODRenderData.cpp
+++ b/Engine/Source/Runtime/Engine/Private/SkeletalMeshLODRenderData.cpp
@@ -143,14 +143,24 @@ void FSkeletalMeshLODRenderData::InitResources(bool bNeedsVertexColors, int32 LO
 	// DuplicatedVerticesBuffer is used only for SkinCache and Editor features which is SM5 only
     if (IsFeatureLevelSupported(GMaxRHIShaderPlatform, ERHIFeatureLevel::SM5))
 	{
-		if (GPUSkinCacheNeedsDuplicatedVertices())
+		const bool bSkinCacheNeedsDuplicatedVertices = GPUSkinCacheNeedsDuplicatedVertices();
+		for (auto& RenderSection : RenderSections)
 		{
-			for (auto& RenderSection : RenderSections)
+			if (bSkinCacheNeedsDuplicatedVertices)
 			{
+				// No need to discard CPU data in cooked builds as bNeedsCPUAccess is false (see FDuplicatedVerticesBuffer constructor), 
+				// so it'd be auto-discarded after the RHI has copied the resource data. Keep CPU data when in the editor for geometry operations.
 				check(RenderSection.DuplicatedVerticesBuffer.DupVertData.Num());
 				BeginInitResource(&RenderSection.DuplicatedVerticesBuffer);
 			}
-		}
+			else
+			{
+#if !WITH_EDITOR
+				// Discard CPU data in cooked builds. Keep CPU data when in the editor for geometry operations.
+				RenderSection.DuplicatedVerticesBuffer.ReleaseCPUResources();
+#endif
+			}
+}
     }
 
 	// UseGPUMorphTargets() can be toggled only on SM5 atm
@@ -272,7 +282,9 @@ void FSkeletalMeshLODRenderData::ReleaseResources()
 		{
 			for (auto& RenderSection : RenderSections)
 			{
+#if WITH_EDITOR
 				check(RenderSection.DuplicatedVerticesBuffer.DupVertData.Num());
+#endif
 				BeginReleaseResource(&RenderSection.DuplicatedVerticesBuffer);
 			}
 		}
diff --git a/Engine/Source/Runtime/Engine/Public/Rendering/SkeletalMeshDuplicatedVerticesBuffer.h b/Engine/Source/Runtime/Engine/Public/Rendering/SkeletalMeshDuplicatedVerticesBuffer.h
--- a/Engine/Source/Runtime/Engine/Public/Rendering/SkeletalMeshDuplicatedVerticesBuffer.h
+++ b/Engine/Source/Runtime/Engine/Public/Rendering/SkeletalMeshDuplicatedVerticesBuffer.h
@@ -38,7 +38,14 @@ public:
 
     bool bHasOverlappingVertices;
 
-    FDuplicatedVerticesBuffer() : DupVertData(true), DupVertIndexData(true), bHasOverlappingVertices(false) {}
+    FDuplicatedVerticesBuffer() 
+#if WITH_EDITOR
+	// Keep CPU copy in editor for geometry operations
+	: DupVertData(true), DupVertIndexData(true)
+#else
+	: DupVertData(false), DupVertIndexData(false)
+#endif
+	, bHasOverlappingVertices(false) {}
 	
     /** Destructor. */
 	virtual ~FDuplicatedVerticesBuffer() {}
@@ -92,6 +99,12 @@ public:
         }
     }
 
+	void ReleaseCPUResources()
+	{
+		DupVertData.Discard();
+		DupVertIndexData.Discard();
+	}
+
     virtual void InitRHI() override
     {
         {
