Engine: Added AController::OnPossessedPawnChanged which is broadcast when the posessed pawn changes under all circumstances (client or authority), and gives both the old and new controlled pawn to simply consumer logic without additional tracking

diff --git a/Engine/Source/Runtime/Engine/Classes/GameFramework/Controller.h b/Engine/Source/Runtime/Engine/Classes/GameFramework/Controller.h
--- a/Engine/Source/Runtime/Engine/Classes/GameFramework/Controller.h
+++ b/Engine/Source/Runtime/Engine/Classes/GameFramework/Controller.h
@@ -18,6 +18,7 @@ class UDamageType;
 UDELEGATE(BlueprintAuthorityOnly)
 DECLARE_DYNAMIC_MULTICAST_DELEGATE_FourParams( FInstigatedAnyDamageSignature, float, Damage, const UDamageType*, DamageType, AActor*, DamagedActor, AActor*, DamageCauser );
 DECLARE_MULTICAST_DELEGATE_OneParam(FPawnChangedSignature, APawn* /*NewPawn*/);
+DECLARE_DYNAMIC_MULTICAST_DELEGATE_TwoParams(FOnPossessedPawnChanged, APawn*, OldPawn, APawn*, NewPawn);
 
 /**
  * Controllers are non-physical actors that can possess a Pawn to control
@@ -56,6 +57,10 @@ public:
 	UPROPERTY(BlueprintAssignable)
 	FInstigatedAnyDamageSignature OnInstigatedAnyDamage;
 
+	/** Called on both authorities and clients when the possessed pawn changes (either OldPawn or NewPawn might be nullptr) */
+	UPROPERTY(BlueprintAssignable, Category=Pawn)
+	FOnPossessedPawnChanged OnPossessedPawnChanged;
+
 	/** Current gameplay state this controller is in */
 	UPROPERTY()
 	FName StateName;
@@ -80,7 +85,7 @@ private:
 	TObjectPtr<USceneComponent> TransformComponent;
 
 protected:
-	/** Delegate broadcasted when possessing a new pawn or unpossessing one */
+	/** Delegate broadcast on authorities when possessing a new pawn or unpossessing one */
 	FPawnChangedSignature OnNewPawn;
 
 	/** The control rotation of the Controller. See GetControlRotation. */
diff --git a/Engine/Source/Runtime/Engine/Private/Controller.cpp b/Engine/Source/Runtime/Engine/Private/Controller.cpp
--- a/Engine/Source/Runtime/Engine/Private/Controller.cpp
+++ b/Engine/Source/Runtime/Engine/Private/Controller.cpp
@@ -316,10 +316,10 @@ void AController::Possess(APawn* InPawn)
 
 	REDIRECT_OBJECT_TO_VLOG(InPawn, this);
 
-	const APawn* CurrentPawn = GetPawn();
+	APawn* CurrentPawn = GetPawn();
 
 	// A notification is required when the current assigned pawn is not possessed (i.e. pawn assigned before calling Possess)
-	const bool bNotificationRequired = (CurrentPawn != nullptr && CurrentPawn->GetController() == nullptr);
+	const bool bNotificationRequired = (CurrentPawn != nullptr) && (CurrentPawn->GetController() == nullptr);
 
 	// To preserve backward compatibility we keep notifying derived classed for null pawn in case some
 	// overrides decided to react differently when asked to possess a null pawn.
@@ -328,10 +328,11 @@ void AController::Possess(APawn* InPawn)
 
 	// Notify when pawn to possess (different than the assigned one) has been accepted by the native class or notification is explicitly required
 	APawn* NewPawn = GetPawn();
-	if (NewPawn != CurrentPawn || bNotificationRequired)
+	if ((NewPawn != CurrentPawn) || bNotificationRequired)
 	{
 		ReceivePossess(NewPawn);
 		OnNewPawn.Broadcast(NewPawn);
+		OnPossessedPawnChanged.Broadcast(bNotificationRequired ? nullptr : CurrentPawn, NewPawn);
 	}
 }
 
@@ -383,6 +384,7 @@ void AController::UnPossess()
 	{
 		ReceiveUnPossess(CurrentPawn);
 		OnNewPawn.Broadcast(NewPawn);
+		OnPossessedPawnChanged.Broadcast(CurrentPawn, NewPawn);
 	}
 }
 
@@ -519,16 +521,22 @@ void AController::SetPawnFromRep(APawn* InPawn)
 
 void AController::OnRep_Pawn()
 {
+	APawn* StrongOldPawn = OldPawn.Get();
 	// Detect when pawn changes, so we can NULL out the controller on the old pawn
-	if ( OldPawn != NULL && Pawn != OldPawn.Get() && OldPawn->Controller == this )
+	if ((StrongOldPawn != nullptr) && (Pawn != StrongOldPawn) && (StrongOldPawn->Controller == this))
 	{
 		// Set the old controller to NULL, since we are no longer the owner, and can't rely on it replicating to us anymore
-		OldPawn->Controller = NULL;
+		StrongOldPawn->Controller = nullptr;
 	}
 
 	OldPawn = Pawn;
 
 	SetPawn(Pawn);
+
+	if (StrongOldPawn != Pawn)
+	{
+		OnPossessedPawnChanged.Broadcast(StrongOldPawn, Pawn);
+	}
 }
 
 void AController::OnRep_PlayerState()
