Fixed deprecated method from DebugRenderSceneProxy

diff --git a/Engine/Source/Runtime/Engine/Private/DebugRenderSceneProxy.cpp b/Engine/Source/Runtime/Engine/Private/DebugRenderSceneProxy.cpp
--- a/Engine/Source/Runtime/Engine/Private/DebugRenderSceneProxy.cpp
+++ b/Engine/Source/Runtime/Engine/Private/DebugRenderSceneProxy.cpp
@@ -49,7 +49,7 @@ void FDebugDrawDelegateHelper::UnregisterDebugDrawDelegate()
 
 void  FDebugDrawDelegateHelper::ReregisterDebugDrawDelegate()
 {
-	ensureMsgf(State != UndefinedState, TEXT("ReregisterDebugDrawDelgeate is in an invalid State: %i !"), State);
+	ensureMsgf(State != UndefinedState, TEXT("ReregisterDebugDrawDelegate is in an invalid State: %i !"), State);
 	if (State == RegisteredState)
 	{
 		UnregisterDebugDrawDelegate();
