Explicitly remove scaling from camera ViewRotationMatrix, as this can accumulte due to floating point precision limitations

diff --git a/Engine/Source/Runtime/Engine/Private/SceneView.cpp b/Engine/Source/Runtime/Engine/Private/SceneView.cpp
--- a/Engine/Source/Runtime/Engine/Private/SceneView.cpp
+++ b/Engine/Source/Runtime/Engine/Private/SceneView.cpp
@@ -503,11 +503,11 @@ FVector4 CreateInvDeviceZToWorldZTransform(const FMatrix& ProjMatrix)
 
 void FViewMatrices::Init(const FMinimalInitializer& Initializer)
 {
-	//check(Initializer.ViewRotationMatrix.GetOrigin().IsNearlyZero());
-	check(FVector::Distance(Initializer.ViewRotationMatrix.GetScaleVector(), FVector::OneVector) < KINDA_SMALL_NUMBER);
+	FMatrix ViewRotationMatrix = Initializer.ViewRotationMatrix;
+	const FVector ViewRotationScaling = ViewRotationMatrix.ExtractScaling();
+	ensureMsgf(FVector::Distance(ViewRotationScaling, FVector::OneVector) < KINDA_SMALL_NUMBER, TEXT("ViewRotation matrix accumulated scaling (%f, %f, %f)"), ViewRotationScaling.X, ViewRotationScaling.Y, ViewRotationScaling.Z);
 
 	FVector LocalViewOrigin = Initializer.ViewOrigin;
-	FMatrix ViewRotationMatrix = Initializer.ViewRotationMatrix;
 	if (!ViewRotationMatrix.GetOrigin().IsNearlyZero(0.0f))
 	{
 		LocalViewOrigin += ViewRotationMatrix.InverseTransformPosition(FVector::ZeroVector);
