Bink Audio :: 4/4 :: Integration :: #pf 60aea2116905a6000136901e #pf 60c9308ef7cffb0001674132

Integration touches several places in the engine:

1) SoundWave -- A check box enables Bink Audio as the codec of choice for that sound wave.

2) Decoder - Each supported platform's AudioMixer now returns BINK if the soundwave requests it. Additionally, the TargetPlatform returns BINK as an available codec, and returns it to the cooking code if the sound wave requests it.

3) Encode - TargetPlatform.Build.cs adds the encoder to the editor dependencies, and it gets picked up in the TPMM formats search.

diff --git a/Engine/Source/Developer/Android/AndroidTargetPlatform/Private/AndroidTargetPlatform.cpp b/Engine/Source/Developer/Android/AndroidTargetPlatform/Private/AndroidTargetPlatform.cpp
--- a/Engine/Source/Developer/Android/AndroidTargetPlatform/Private/AndroidTargetPlatform.cpp
+++ b/Engine/Source/Developer/Android/AndroidTargetPlatform/Private/AndroidTargetPlatform.cpp
@@ -616,6 +616,7 @@ FName FAndroidTargetPlatform::GetWaveFormat( const class USoundWave* Wave ) cons
 {
 	static const FName NAME_ADPCM(TEXT("ADPCM"));
 	static const FName NAME_OGG(TEXT("OGG"));
+	static const FName NAME_BINKA(TEXT("BINKA"));
 
 	static bool bFormatRead = false;
 	static FName NAME_FORMAT;
@@ -650,6 +651,11 @@ FName FAndroidTargetPlatform::GetWaveFormat( const class USoundWave* Wave ) cons
 		}
 	}
 
+	if (Wave->bUseBinkAudio)
+	{
+		return NAME_BINKA;
+	}
+
 	if (Wave->IsSeekableStreaming())
 	{
 		return NAME_ADPCM;
@@ -661,9 +667,11 @@ FName FAndroidTargetPlatform::GetWaveFormat( const class USoundWave* Wave ) cons
 
 void FAndroidTargetPlatform::GetAllWaveFormats(TArray<FName>& OutFormats) const
 {
-	static FName NAME_OGG(TEXT("OGG"));
-	static FName NAME_ADPCM(TEXT("ADPCM"));
+	static const FName NAME_OGG(TEXT("OGG"));
+	static const FName NAME_ADPCM(TEXT("ADPCM"));
+	static const FName NAME_BINKA(TEXT("BINKA"));
 
+	OutFormats.Add(NAME_BINKA);
 	OutFormats.Add(NAME_OGG);
 	OutFormats.Add(NAME_ADPCM);
 }
diff --git a/Engine/Source/Developer/IOS/IOSTargetPlatform/Private/IOSTargetPlatform.cpp b/Engine/Source/Developer/IOS/IOSTargetPlatform/Private/IOSTargetPlatform.cpp
--- a/Engine/Source/Developer/IOS/IOSTargetPlatform/Private/IOSTargetPlatform.cpp
+++ b/Engine/Source/Developer/IOS/IOSTargetPlatform/Private/IOSTargetPlatform.cpp
@@ -674,6 +674,11 @@ const UTextureLODSettings& FIOSTargetPlatform::GetTextureLODSettings() const
 
 FName FIOSTargetPlatform::GetWaveFormat( const class USoundWave* Wave ) const
 {
+	if (Wave->bUseBinkAudio)
+	{
+		static FName NAME_BINKA(TEXT("BINKA"));
+		return NAME_BINKA;
+	}
 	static FName NAME_ADPCM(TEXT("ADPCM"));
 	return NAME_ADPCM;
 }
@@ -681,8 +686,10 @@ FName FIOSTargetPlatform::GetWaveFormat( const class USoundWave* Wave ) const
 
 void FIOSTargetPlatform::GetAllWaveFormats(TArray<FName>& OutFormat) const
 {
+	static FName NAME_BINKA(TEXT("BINKA"));
 	static FName NAME_ADPCM(TEXT("ADPCM"));
 	OutFormat.Add(NAME_ADPCM);
+	OutFormat.Add(NAME_BINKA);
 }
 
 #endif // WITH_ENGINE
diff --git a/Engine/Source/Developer/Linux/LinuxTargetPlatform/Private/LinuxTargetPlatform.h b/Engine/Source/Developer/Linux/LinuxTargetPlatform/Private/LinuxTargetPlatform.h
--- a/Engine/Source/Developer/Linux/LinuxTargetPlatform/Private/LinuxTargetPlatform.h
+++ b/Engine/Source/Developer/Linux/LinuxTargetPlatform/Private/LinuxTargetPlatform.h
@@ -322,6 +322,12 @@ public:
 		static const FName NAME_ADPCM(TEXT("ADPCM"));
 		static const FName NAME_OGG(TEXT("OGG"));
 		static const FName NAME_OPUS(TEXT("OPUS"));
+		static const FName NAME_BINKA(TEXT("BINKA"));
+
+		if (Wave->bUseBinkAudio)
+		{
+			return NAME_BINKA;
+		}
 
 		if (Wave->IsSeekableStreaming())
 		{
@@ -341,7 +347,9 @@ public:
 		static const FName NAME_ADPCM(TEXT("ADPCM"));
 		static const FName NAME_OGG(TEXT("OGG"));
 		static const FName NAME_OPUS(TEXT("OPUS"));
+		static const FName NAME_BINKA(TEXT("BINKA"));
 
+		OutFormats.Add(NAME_BINKA);
 		OutFormats.Add(NAME_ADPCM);
 		OutFormats.Add(NAME_OGG);
 		OutFormats.Add(NAME_OPUS);
diff --git a/Engine/Source/Developer/Mac/MacTargetPlatform/Private/GenericMacTargetPlatform.h b/Engine/Source/Developer/Mac/MacTargetPlatform/Private/GenericMacTargetPlatform.h
--- a/Engine/Source/Developer/Mac/MacTargetPlatform/Private/GenericMacTargetPlatform.h
+++ b/Engine/Source/Developer/Mac/MacTargetPlatform/Private/GenericMacTargetPlatform.h
@@ -326,6 +326,12 @@ return TSuper::SupportsFeature(Feature);
 		static const FName NAME_ADPCM(TEXT("ADPCM"));
 		static const FName NAME_OGG(TEXT("OGG"));
 		static const FName NAME_OPUS(TEXT("OPUS"));
+		static const FName NAME_BINKA(TEXT("BINKA"));
+
+		if (Wave->bUseBinkAudio)
+		{
+			return NAME_BINKA;
+		}
 
 		if (Wave->IsSeekableStreaming())
 		{
@@ -345,10 +351,12 @@ return TSuper::SupportsFeature(Feature);
 		static const FName NAME_ADPCM(TEXT("ADPCM"));
 		static const FName NAME_OGG(TEXT("OGG"));
 		static const FName NAME_OPUS(TEXT("OPUS"));
+		static const FName NAME_BINKA(TEXT("BINKA"));
 
 		OutFormats.Add(NAME_ADPCM);
 		OutFormats.Add(NAME_OGG);
 		OutFormats.Add(NAME_OPUS);
+		OutFormats.Add(NAME_BINKA);
 	}
 
 #endif //WITH_ENGINE
diff --git a/Engine/Source/Developer/TargetPlatform/TargetPlatform.Build.cs b/Engine/Source/Developer/TargetPlatform/TargetPlatform.Build.cs
--- a/Engine/Source/Developer/TargetPlatform/TargetPlatform.Build.cs
+++ b/Engine/Source/Developer/TargetPlatform/TargetPlatform.Build.cs
@@ -58,6 +58,7 @@ public class TargetPlatform : ModuleRules
 					DynamicallyLoadedModuleNames.Add("AudioFormatADPCM");
 					DynamicallyLoadedModuleNames.Add("AudioFormatOgg");
 					DynamicallyLoadedModuleNames.Add("AudioFormatOpus");
+					DynamicallyLoadedModuleNames.Add("AudioFormatBink");
 				}
 
 				if (Target.Type == TargetType.Editor || Target.Type == TargetType.Program)
@@ -78,6 +79,7 @@ public class TargetPlatform : ModuleRules
 					DynamicallyLoadedModuleNames.Add("AudioFormatADPCM");
 					DynamicallyLoadedModuleNames.Add("AudioFormatOgg");
 					DynamicallyLoadedModuleNames.Add("AudioFormatOpus");
+					DynamicallyLoadedModuleNames.Add("AudioFormatBink");
 				}
 
 			}
@@ -92,6 +94,7 @@ public class TargetPlatform : ModuleRules
 					DynamicallyLoadedModuleNames.Add("AudioFormatADPCM");
 					DynamicallyLoadedModuleNames.Add("AudioFormatOgg");
 					DynamicallyLoadedModuleNames.Add("AudioFormatOpus");
+					DynamicallyLoadedModuleNames.Add("AudioFormatBink");
 				}
 			}
 		}
diff --git a/Engine/Source/Developer/Windows/WindowsTargetPlatform/Private/GenericWindowsTargetPlatform.h b/Engine/Source/Developer/Windows/WindowsTargetPlatform/Private/GenericWindowsTargetPlatform.h
--- a/Engine/Source/Developer/Windows/WindowsTargetPlatform/Private/GenericWindowsTargetPlatform.h
+++ b/Engine/Source/Developer/Windows/WindowsTargetPlatform/Private/GenericWindowsTargetPlatform.h
@@ -416,6 +416,12 @@ public:
 		static const FName NAME_ADPCM(TEXT("ADPCM"));
 		static const FName NAME_OGG(TEXT("OGG"));
 		static const FName NAME_OPUS(TEXT("OPUS"));
+		static const FName NAME_BINKA(TEXT("BINKA"));
+
+		if (Wave->bUseBinkAudio)
+		{
+			return NAME_BINKA;
+		}
 
 		// Seekable streams need to pick a codec which allows fixed-sized frames so we can compute stream chunk index to load
 		if (Wave->IsSeekableStreaming())
@@ -436,9 +442,11 @@ public:
 	virtual void GetAllWaveFormats(TArray<FName>& OutFormats) const override
 	{
 		static const FName NAME_ADPCM(TEXT("ADPCM"));
-		static FName NAME_OGG(TEXT("OGG"));
-		static FName NAME_OPUS(TEXT("OPUS"));
+		static const FName NAME_OGG(TEXT("OGG"));
+		static const FName NAME_OPUS(TEXT("OPUS"));
+		static const FName NAME_BINKA(TEXT("BINKA"));
 
+		OutFormats.Add(NAME_BINKA);
 		OutFormats.Add(NAME_ADPCM);
 		OutFormats.Add(NAME_OGG);
 		OutFormats.Add(NAME_OPUS);
diff --git a/Engine/Source/Runtime/Android/AudioMixerAndroid/AudioMixerAndroid.Build.cs b/Engine/Source/Runtime/Android/AudioMixerAndroid/AudioMixerAndroid.Build.cs
--- a/Engine/Source/Runtime/Android/AudioMixerAndroid/AudioMixerAndroid.Build.cs
+++ b/Engine/Source/Runtime/Android/AudioMixerAndroid/AudioMixerAndroid.Build.cs
@@ -13,7 +13,8 @@ public class AudioMixerAndroid : ModuleRules
 		PrivateDependencyModuleNames.AddRange(
 			new string[] {
 				"Core",
-				"AudioMixerCore"
+				"AudioMixerCore",
+				"BinkAudioDecoder"
 			}
 			);
 			
diff --git a/Engine/Source/Runtime/Android/AudioMixerAndroid/Private/AudioMixerPlatformAndroid.cpp b/Engine/Source/Runtime/Android/AudioMixerAndroid/Private/AudioMixerPlatformAndroid.cpp
--- a/Engine/Source/Runtime/Android/AudioMixerAndroid/Private/AudioMixerPlatformAndroid.cpp
+++ b/Engine/Source/Runtime/Android/AudioMixerAndroid/Private/AudioMixerPlatformAndroid.cpp
@@ -10,6 +10,7 @@
 #if WITH_ENGINE
 #include "VorbisAudioInfo.h"
 #include "ADPCMAudioInfo.h"
+#include "BinkAudioInfo.h"
 #include "AudioPluginUtilities.h"
 #endif
 
@@ -427,6 +428,12 @@ namespace Audio
 	{
 #if WITH_ENGINE
 		static FName NAME_ADPCM(TEXT("ADPCM"));
+		static const FName NAME_BINKA(TEXT("BINKA"));
+
+		if (InSoundWave->bUseBinkAudio)
+		{
+			return NAME_BINKA;
+		}
 
 		if (InSoundWave->IsSeekableStreaming())
 		{
@@ -454,8 +461,10 @@ namespace Audio
 	ICompressedAudioInfo* FMixerPlatformAndroid::CreateCompressedAudioInfo(USoundWave* InSoundWave)
 	{
 #if WITH_ENGINE
-		static FName NAME_OGG(TEXT("OGG"));
-		static FName NAME_ADPCM(TEXT("ADPCM"));
+		if (InSoundWave->bUseBinkAudio)
+		{
+			return new FBinkAudioInfo();
+		}
 
 		if (InSoundWave->IsSeekableStreaming())
 		{
@@ -476,8 +485,10 @@ namespace Audio
 		}
 
 #if WITH_ENGINE
-		static FName NAME_OGG(TEXT("OGG"));
-		static FName NAME_ADPCM(TEXT("ADPCM"));
+		if (InSoundWave->UseBinkAudio())
+		{
+			return new FBinkAudioInfo();
+		}
 
 		if (InSoundWave->IsSeekableStreaming())
 		{
diff --git a/Engine/Source/Runtime/Apple/AudioMixerAudioUnit/AudioMixerAudioUnit.Build.cs b/Engine/Source/Runtime/Apple/AudioMixerAudioUnit/AudioMixerAudioUnit.Build.cs
--- a/Engine/Source/Runtime/Apple/AudioMixerAudioUnit/AudioMixerAudioUnit.Build.cs
+++ b/Engine/Source/Runtime/Apple/AudioMixerAudioUnit/AudioMixerAudioUnit.Build.cs
@@ -22,10 +22,16 @@ public class AudioMixerAudioUnit : ModuleRules
 
         if (Target.bCompileAgainstEngine)
         {
-            // Engine module is required for CompressedAudioInfo implementations.
-            PrivateDependencyModuleNames.Add("Engine");
+			// Engine module is required for CompressedAudioInfo implementations.
+			PrivateDependencyModuleNames.AddRange(
+				new string[]
+				{
+					"Engine",
+					"BinkAudioDecoder"
+				}
+			);
 
-            AddEngineThirdPartyPrivateStaticDependencies(Target,
+			AddEngineThirdPartyPrivateStaticDependencies(Target,
                 "UEOgg",
                 "Vorbis",
                 "VorbisFile"
diff --git a/Engine/Source/Runtime/Apple/AudioMixerAudioUnit/Private/AudioMixerPlatformAudioUnit.cpp b/Engine/Source/Runtime/Apple/AudioMixerAudioUnit/Private/AudioMixerPlatformAudioUnit.cpp
--- a/Engine/Source/Runtime/Apple/AudioMixerAudioUnit/Private/AudioMixerPlatformAudioUnit.cpp
+++ b/Engine/Source/Runtime/Apple/AudioMixerAudioUnit/Private/AudioMixerPlatformAudioUnit.cpp
@@ -10,7 +10,7 @@
 #include "Misc/ConfigCacheIni.h"
 #include "Misc/CoreDelegates.h"
 #include "ADPCMAudioInfo.h"
-
+#include "BinkAudioInfo.h"
 
 /*
  This implementation only depends on the audio units API which allows it to run on MacOS, iOS and tvOS.
@@ -390,6 +390,12 @@ namespace Audio
 	
 	FName FMixerPlatformAudioUnit::GetRuntimeFormat(USoundWave* InSoundWave)
 	{
+		if (InSoundWave->bUseBinkAudio)
+		{
+			static const FName NAME_BINKA(TEXT("BINKA"));
+			return NAME_BINKA;
+		}
+
 		static FName NAME_ADPCM(TEXT("ADPCM"));
 		return NAME_ADPCM;
 	}
@@ -401,11 +407,20 @@ namespace Audio
 	
 	ICompressedAudioInfo* FMixerPlatformAudioUnit::CreateCompressedAudioInfo(USoundWave* InSoundWave)
 	{
+		if (InSoundWave->bUseBinkAudio)
+		{
+			return new FBinkAudioInfo();
+		}
 		return new FADPCMAudioInfo();
 	}
 
 	ICompressedAudioInfo* FMixerPlatformAudioUnit::CreateCompressedAudioInfo(const FSoundWaveProxyPtr& InSoundWave)
 	{
+		if (InSoundWave->UseBinkAudio())
+		{
+			return new FBinkAudioInfo();
+		}
+
 		return new FADPCMAudioInfo();
 	}
 	
diff --git a/Engine/Source/Runtime/Apple/AudioMixerCoreAudio/AudioMixerCoreAudio.Build.cs b/Engine/Source/Runtime/Apple/AudioMixerCoreAudio/AudioMixerCoreAudio.Build.cs
--- a/Engine/Source/Runtime/Apple/AudioMixerCoreAudio/AudioMixerCoreAudio.Build.cs
+++ b/Engine/Source/Runtime/Apple/AudioMixerCoreAudio/AudioMixerCoreAudio.Build.cs
@@ -20,6 +20,7 @@ public class AudioMixerCoreAudio : ModuleRules
 				"Core",
 				"CoreUObject",
 				"AudioMixerCore",
+				"BinkAudioDecoder"
 			}
 			);
 
diff --git a/Engine/Source/Runtime/Apple/AudioMixerCoreAudio/Private/AudioMixerPlatformCoreAudio.cpp b/Engine/Source/Runtime/Apple/AudioMixerCoreAudio/Private/AudioMixerPlatformCoreAudio.cpp
--- a/Engine/Source/Runtime/Apple/AudioMixerCoreAudio/Private/AudioMixerPlatformCoreAudio.cpp
+++ b/Engine/Source/Runtime/Apple/AudioMixerCoreAudio/Private/AudioMixerPlatformCoreAudio.cpp
@@ -11,6 +11,7 @@
 #include "VorbisAudioInfo.h"
 #include "OpusAudioInfo.h"
 #include "ADPCMAudioInfo.h"
+#include "BinkAudioInfo.h"
 #endif
 /**
 * CoreAudio System Headers
@@ -365,6 +366,12 @@ namespace Audio
 		static const FName NAME_ADPCM(TEXT("ADPCM"));
 		static const FName NAME_OGG(TEXT("OGG"));
 		static const FName NAME_OPUS(TEXT("OPUS"));
+		static const FName NAME_BINKA(TEXT("BINKA"));
+
+		if (InSoundWave->bUseBinkAudio)
+		{
+			return NAME_BINKA;
+		}
 
 		if (InSoundWave->IsSeekableStreaming())
 		{
@@ -392,6 +399,11 @@ namespace Audio
 #if WITH_ENGINE
 		check(InSoundWave);
 
+		if (InSoundWave->bUseBinkAudio)
+		{
+			return new FBinkAudioInfo();
+		}
+
 		if (InSoundWave->IsSeekableStreaming())
 		{
 			return new FADPCMAudioInfo();
@@ -426,6 +438,11 @@ namespace Audio
 			return nullptr;
 		}
 
+		if (InSoundWave->UseBinkAudio())
+		{
+			return new FBinkAudioInfo();
+		}
+
 		if (InSoundWave->IsSeekableStreaming())
 		{
 			return new FADPCMAudioInfo();
diff --git a/Engine/Source/Runtime/Engine/Classes/Sound/SoundWave.h b/Engine/Source/Runtime/Engine/Classes/Sound/SoundWave.h
--- a/Engine/Source/Runtime/Engine/Classes/Sound/SoundWave.h
+++ b/Engine/Source/Runtime/Engine/Classes/Sound/SoundWave.h
@@ -343,6 +343,7 @@ public:
 
 	void ReleaseCompressedAudio();
 
+	bool UseBinkAudio() const { return *bUseBinkAudioPtr; }
 	bool IsStreaming() const { return *bIsStreamingPtr; }
 	bool IsSeekableStreaming() const { return *bIsStreamingPtr && *bSeekableStreamingPtr; }
 	bool IsRetainingAudio() const { return FirstChunk.IsValid(); }
@@ -386,10 +387,14 @@ private:
 	FAudioChunkHandle FirstChunk;
 	ESoundWaveLoadingBehavior LoadingBehavior;
 
+	// These hold a reference to shared memory containing flags that need to get
+	// routed from the SoundWave. Whenever the relevant bool changes on the SoundWave,
+	// the ProxyPtr on the sound wave should get updated so we can see them.
 	TSharedPtr<FThreadSafeBool, ESPMode::ThreadSafe> bLoadingBehaviorOverriddenPtr;
 	TSharedPtr<FThreadSafeBool, ESPMode::ThreadSafe> bIsStreamingPtr;
 	TSharedPtr<FThreadSafeBool, ESPMode::ThreadSafe> bSeekableStreamingPtr;
 	TSharedPtr<FThreadSafeBool, ESPMode::ThreadSafe> bShouldUseStreamCachingPtr;
+	TSharedPtr<FThreadSafeBool, ESPMode::ThreadSafe> bUseBinkAudioPtr;
 
 #if WITH_EDITOR
 	TSharedPtr<FThreadSafeCounter> CurrentChunkRevision{ nullptr };
@@ -440,6 +445,14 @@ private:
 	TSharedPtr<FThreadSafeBool, ESPMode::ThreadSafe> bIsSeekableStreamingProxyFlag{ MakeShared<FThreadSafeBool, ESPMode::ThreadSafe>() };
 public:
 
+	/** If true, the sound will compress to the Bink Audio format whenever available. */
+	UPROPERTY(EditAnywhere, Category = "Format")
+	uint8 bUseBinkAudio : 1;
+private:
+	// updated to reflect bUseBinkAudio in PostEditChangeProperty()
+	TSharedPtr<FThreadSafeBool, ESPMode::ThreadSafe> bUseBinkAudioProxyFlag{ MakeShared<FThreadSafeBool, ESPMode::ThreadSafe>() };
+public:
+
 	// Loading behavior members are lazily initialized in const getters
 	/** Specifies how and when compressed audio data is loaded for asset if stream caching is enabled. */
 	UPROPERTY(EditAnywhere, Category = "Loading", meta = (DisplayName = "Loading Behavior Override"))
diff --git a/Engine/Source/Runtime/Engine/Private/SoundWave.cpp b/Engine/Source/Runtime/Engine/Private/SoundWave.cpp
--- a/Engine/Source/Runtime/Engine/Private/SoundWave.cpp
+++ b/Engine/Source/Runtime/Engine/Private/SoundWave.cpp
@@ -1811,8 +1811,10 @@ void USoundWave::PostEditChangeProperty(FPropertyChangedEvent& PropertyChangedEv
 	static const FName SampleRateFName = GET_MEMBER_NAME_CHECKED(USoundWave,SampleRateQuality);
 	static const FName StreamingFName = GET_MEMBER_NAME_CHECKED(USoundWave, bStreaming);
 	static const FName SeekableStreamingFName = GET_MEMBER_NAME_CHECKED(USoundWave, bSeekableStreaming);
+	static const FName UseBinkAudioFName = GET_MEMBER_NAME_CHECKED(USoundWave, bUseBinkAudio);
 
 	// force proxy flags to be up to date
+	*bUseBinkAudioProxyFlag = bUseBinkAudio;
 	*bIsSeekableStreamingProxyFlag = bSeekableStreaming;
 	*bIsStreamingProxyFlag = IsStreaming(nullptr);
 	*bShouldUseStreamCachingProxyFlag = ShouldUseStreamCaching();
@@ -1824,7 +1826,7 @@ void USoundWave::PostEditChangeProperty(FPropertyChangedEvent& PropertyChangedEv
 		if (FProperty* PropertyThatChanged = PropertyChangedEvent.Property)
 		{
 			const FName& Name = PropertyThatChanged->GetFName();
-			if (Name == CompressionQualityFName || Name == SampleRateFName || Name == StreamingFName || Name == SeekableStreamingFName)
+			if (Name == CompressionQualityFName || Name == SampleRateFName || Name == StreamingFName || Name == SeekableStreamingFName || Name == UseBinkAudioFName)
 			{
 				InvalidateCompressedData();
 				FreeResources();
@@ -2492,6 +2494,7 @@ FSoundWaveProxyPtr USoundWave::CreateSoundWaveProxy()
 	EnsureZerothChunkIsLoaded();
 #endif // #if WITH_EDITORONLY_DATA
 
+	*bUseBinkAudioProxyFlag = bUseBinkAudio;
 	*bIsSeekableStreamingProxyFlag = bSeekableStreaming;
 	*bIsStreamingProxyFlag = IsStreaming(nullptr);
 	*bShouldUseStreamCachingProxyFlag = ShouldUseStreamCaching();
@@ -2505,6 +2508,7 @@ TUniquePtr<Audio::IProxyData> USoundWave::CreateNewProxyData(const Audio::FProxy
 	EnsureZerothChunkIsLoaded();
 #endif // #if WITH_EDITORONLY_DATA
 
+	*bUseBinkAudioProxyFlag = bUseBinkAudio;
 	*bIsSeekableStreamingProxyFlag = bSeekableStreaming;
 	*bIsStreamingProxyFlag = IsStreaming(nullptr);
 	*bShouldUseStreamCachingProxyFlag = ShouldUseStreamCaching();
@@ -3145,6 +3149,7 @@ FSoundWaveProxy::FSoundWaveProxy(USoundWave* InWave)
 	, bIsStreamingPtr(InWave->bIsStreamingProxyFlag)
 	, bSeekableStreamingPtr(InWave->bIsSeekableStreamingProxyFlag)
 	, bShouldUseStreamCachingPtr(InWave->bShouldUseStreamCachingProxyFlag)
+	, bUseBinkAudioPtr(InWave->bUseBinkAudioProxyFlag)
 #if WITH_EDITOR
 	, CurrentChunkRevision(InWave->CurrentChunkRevision)
 #endif // #if WITH_EDITOR
diff --git a/Engine/Source/Runtime/Linux/AudioMixerSDL/AudioMixerSDL.Build.cs b/Engine/Source/Runtime/Linux/AudioMixerSDL/AudioMixerSDL.Build.cs
--- a/Engine/Source/Runtime/Linux/AudioMixerSDL/AudioMixerSDL.Build.cs
+++ b/Engine/Source/Runtime/Linux/AudioMixerSDL/AudioMixerSDL.Build.cs
@@ -21,6 +21,13 @@ public class AudioMixerSDL : ModuleRules
 			PrivateIncludePaths.Add("Runtime/Linux/AudioMixerSDL/Private/" + PlatformName);
 		}
 
+        {        	
+			// Bink Audio isn't yet built for other SDL platforms (e.g. QNX, Stadia),
+			// however the module won't provide a lib and will define WITH_BINK_AUDIO
+			// to 0 for those platforms.
+			PrivateDependencyModuleNames.Add("BinkAudioDecoder");
+		}
+
 		PrivateDependencyModuleNames.AddRange(
 			new string[] {
 				"Core",
diff --git a/Engine/Source/Runtime/Linux/AudioMixerSDL/Private/AudioMixerPlatformSDL.cpp b/Engine/Source/Runtime/Linux/AudioMixerSDL/Private/AudioMixerPlatformSDL.cpp
--- a/Engine/Source/Runtime/Linux/AudioMixerSDL/Private/AudioMixerPlatformSDL.cpp
+++ b/Engine/Source/Runtime/Linux/AudioMixerSDL/Private/AudioMixerPlatformSDL.cpp
@@ -12,6 +12,10 @@
 #include "OpusAudioInfo.h"
 #include "VorbisAudioInfo.h"
 #include "ADPCMAudioInfo.h"
+#if WITH_BINK_AUDIO
+#include "BinkAudioInfo.h"
+#endif // WITH_BINK_AUDIO
+
 #endif // WITH_ENGINE
 
 DECLARE_LOG_CATEGORY_EXTERN(LogAudioMixerSDL, Log, All);
@@ -332,7 +336,14 @@ namespace Audio
 		static FName NAME_OGG(TEXT("OGG"));
 		static FName NAME_OPUS(TEXT("OPUS"));
 		static FName NAME_ADPCM(TEXT("ADPCM"));
-
+		
+#if WITH_BINK_AUDIO
+		static FName NAME_BINKA(TEXT("BINKA"));
+		if (InSoundWave->bUseBinkAudio)
+		{
+			return NAME_BINKA;
+		}
+#endif // WITH_BINK_AUDIO
 		if (InSoundWave->IsStreaming(nullptr))
 		{
 			if (InSoundWave->IsSeekableStreaming())
@@ -359,6 +370,13 @@ namespace Audio
 #if WITH_ENGINE
 		check(InSoundWave);
 
+#if WITH_BINK_AUDIO
+		if (InSoundWave->bUseBinkAudio)
+		{
+			return new FBinkAudioInfo();
+		}
+#endif // WITH_BINK_AUDIO
+
 		if (InSoundWave->IsStreaming())
 		{
 			if (InSoundWave->IsSeekableStreaming())
@@ -385,6 +403,14 @@ namespace Audio
 	ICompressedAudioInfo* FMixerPlatformSDL::CreateCompressedAudioInfo(const FSoundWaveProxyPtr& InSoundWave)
 	{
 #if WITH_ENGINE
+
+#if WITH_BINK_AUDIO
+		if (InSoundWave->UseBinkAudio())
+		{
+			return new FBinkAudioInfo();
+		}
+#endif // WITH_BINK_AUDIO
+
 		if (InSoundWave->IsStreaming())
 		{
 			if (InSoundWave->IsSeekableStreaming())
diff --git a/Engine/Source/Runtime/Windows/AudioMixerXAudio2/AudioMixerXAudio2.Build.cs b/Engine/Source/Runtime/Windows/AudioMixerXAudio2/AudioMixerXAudio2.Build.cs
--- a/Engine/Source/Runtime/Windows/AudioMixerXAudio2/AudioMixerXAudio2.Build.cs
+++ b/Engine/Source/Runtime/Windows/AudioMixerXAudio2/AudioMixerXAudio2.Build.cs
@@ -13,7 +13,13 @@ public class AudioMixerXAudio2 : ModuleRules
 		if (Target.bCompileAgainstEngine)
         {
 			// Engine module is required for CompressedAudioInfo implementations.
-            PrivateDependencyModuleNames.Add("Engine");
+            PrivateDependencyModuleNames.AddRange(
+				new string[]
+				{
+					"Engine",
+					"BinkAudioDecoder"
+				}
+			);
 
 			AddEngineThirdPartyPrivateStaticDependencies(Target,
 			"UEOgg",
diff --git a/Engine/Source/Runtime/Windows/AudioMixerXAudio2/Private/AudioMixerPlatformXAudio2.cpp b/Engine/Source/Runtime/Windows/AudioMixerXAudio2/Private/AudioMixerPlatformXAudio2.cpp
--- a/Engine/Source/Runtime/Windows/AudioMixerXAudio2/Private/AudioMixerPlatformXAudio2.cpp
+++ b/Engine/Source/Runtime/Windows/AudioMixerXAudio2/Private/AudioMixerPlatformXAudio2.cpp
@@ -25,6 +25,9 @@
 #include "VorbisAudioInfo.h"
 #include "ADPCMAudioInfo.h"
 #include "AudioPluginUtilities.h"
+#if WITH_BINK_AUDIO
+#include "BinkAudioInfo.h"
+#endif // WITH_BINK_AUDIO
 #endif //WITH_ENGINE
 
 #include "CoreGlobals.h"
@@ -1468,8 +1471,17 @@ namespace Audio
 		static FName NAME_OPUS(TEXT("OPUS"));
 		static FName NAME_XMA(TEXT("XMA"));
 		static FName NAME_ADPCM(TEXT("ADPCM"));
+		static FName NAME_BINKA(TEXT("BINKA"));
 
 #if WITH_ENGINE
+
+#if WITH_BINK_AUDIO
+		if (InSoundWave->bUseBinkAudio)
+		{
+			return NAME_BINKA;
+		}
+#endif	// WITH_BINK_AUDIO	
+
 		if (InSoundWave->IsStreaming(nullptr))
 		{
 			if (InSoundWave->IsSeekableStreaming())
@@ -1509,6 +1521,13 @@ namespace Audio
 #if WITH_ENGINE
 		check(InSoundWave);
 
+#if WITH_BINK_AUDIO
+		if (InSoundWave->bUseBinkAudio)
+		{
+			return new FBinkAudioInfo();
+		}
+#endif // WITH_BINK_AUDIO
+
 		if (InSoundWave->IsStreaming(nullptr))
 		{
 			if (InSoundWave->IsSeekableStreaming())
@@ -1558,6 +1577,14 @@ namespace Audio
 		}
 
 #if WITH_ENGINE
+
+#if WITH_BINK_AUDIO
+		if (InSoundWave->UseBinkAudio())
+		{
+			return new FBinkAudioInfo();
+		}
+#endif // WITH_BINK_AUDIO		
+
 		// Decoding with FSoundWaveProxy is only supported ot streaming audio at the moment
 		if (false == InSoundWave->IsStreaming())
 		{
diff --git a/Engine/Source/Runtime/Windows/XAudio2/Private/XAudio2Device.cpp b/Engine/Source/Runtime/Windows/XAudio2/Private/XAudio2Device.cpp
--- a/Engine/Source/Runtime/Windows/XAudio2/Private/XAudio2Device.cpp
+++ b/Engine/Source/Runtime/Windows/XAudio2/Private/XAudio2Device.cpp
@@ -39,6 +39,7 @@ THIRD_PARTY_INCLUDES_END
 #endif
 
 #include "ADPCMAudioInfo.h"
+#include "BinkAudioInfo.h"
 
 DEFINE_LOG_CATEGORY(LogXAudio2);
 
@@ -383,6 +384,11 @@ class ICompressedAudioInfo* FXAudio2Device::CreateCompressedAudioInfo(USoundWave
 {
 	check(SoundWave);
 
+	if (SoundWave->bUseBinkAudio)
+	{
+		return new FBinkAudioInfo();
+	}
+
 	if (SoundWave->IsStreaming(nullptr))
 	{
 		if (SoundWave->IsSeekableStreaming())
diff --git a/Engine/Source/Runtime/Windows/XAudio2/Public/XAudio2Device.h b/Engine/Source/Runtime/Windows/XAudio2/Public/XAudio2Device.h
--- a/Engine/Source/Runtime/Windows/XAudio2/Public/XAudio2Device.h
+++ b/Engine/Source/Runtime/Windows/XAudio2/Public/XAudio2Device.h
@@ -123,6 +123,12 @@ private:
 		static const FName NAME_OGG(TEXT("OGG"));
 		static const FName NAME_OPUS(TEXT("OPUS"));
 		static const FName NAME_XMA(TEXT("XMA"));
+		static const FName NAME_BINKA(TEXT("BINKA"));
+
+		if (SoundWave->bUseBinkAudio)
+		{
+			return NAME_BINKA;
+		}
 
 		if (SoundWave->IsStreaming(nullptr))
 		{
diff --git a/Engine/Source/Runtime/Windows/XAudio2/XAudio2.Build.cs b/Engine/Source/Runtime/Windows/XAudio2/XAudio2.Build.cs
--- a/Engine/Source/Runtime/Windows/XAudio2/XAudio2.Build.cs
+++ b/Engine/Source/Runtime/Windows/XAudio2/XAudio2.Build.cs
@@ -13,6 +13,7 @@ public class XAudio2 : ModuleRules
 				"Core",
 				"CoreUObject",
 				"Engine",
+				"BinkAudioDecoder",
                 "AudioMixerCore" // for AudioPlatformSettings::GetPlatformSettings() in AudioMixerTypes.h.
 			}
 			);
