AddToWorld/RemoveFromWorld now respect the frame time limits set by s.LevelStreamingActorsUpdateTimeLimit and s.UnregisterComponentsTimeLimit.

diff --git a/Engine/Source/Runtime/Engine/Private/World.cpp b/Engine/Source/Runtime/Engine/Private/World.cpp
--- a/Engine/Source/Runtime/Engine/Private/World.cpp
+++ b/Engine/Source/Runtime/Engine/Private/World.cpp
@@ -2547,8 +2547,12 @@ private:
 
 #endif // PERF_TRACK_DETAILED_ASYNC_STATS
 
+// Cumulated time passed in UWorld::AddToWorld since last call to UWorld::UpdateLevelStreaming.
+static double GAddToWorldTimeCumul = 0.0;
+
 void UWorld::AddToWorld( ULevel* Level, const FTransform& LevelTransform, bool bConsiderTimeLimit )
 {
+	TRACE_CPUPROFILER_EVENT_SCOPE(AddToWorld);
 	SCOPE_CYCLE_COUNTER(STAT_AddToWorldTime);
 	CSV_SCOPED_TIMING_STAT_EXCLUSIVE(AddToWorld);
 
@@ -2606,10 +2610,12 @@ void UWorld::AddToWorld( ULevel* Level, const FTransform& LevelTransform, bool b
 
 	// Don't consider the time limit if the match hasn't started as we need to ensure that the levels are fully loaded
 	bConsiderTimeLimit &= bMatchStarted && bIsGameWorld;
-	double TimeLimit = GLevelStreamingActorsUpdateTimeLimit;
+	double TimeLimit = 0.0;
 
 	if (bConsiderTimeLimit)
 	{
+		TimeLimit = GLevelStreamingActorsUpdateTimeLimit;
+
 		// Give the actor initialization code more time if we're performing a high priority load or are in seamless travel
 		if (AWorldSettings* WorldSettings = GetWorldSettings(false, false))
 		{
@@ -2618,6 +2624,9 @@ void UWorld::AddToWorld( ULevel* Level, const FTransform& LevelTransform, bool b
 				TimeLimit += GPriorityLevelStreamingActorsUpdateExtraTime;
 			}
 		}
+
+		// Remove cumulated time since UpdateLevelStreaming
+		TimeLimit = FMath::Max<double>(0.0, TimeLimit - GAddToWorldTimeCumul);
 	}
 
 	if( bExecuteNextStep && !Level->bAlreadyMovedActors )
@@ -2860,6 +2869,10 @@ void UWorld::AddToWorld( ULevel* Level, const FTransform& LevelTransform, bool b
 		UE_LOG(LogStreaming, Display, TEXT("Perform Last Step       : %6.2f ms"), PerformLastStepTime * 1000 );
 	}
 #endif // PERF_TRACK_DETAILED_ASYNC_STATS
+
+	// Delta time in ms.
+	double DeltaTime = (FPlatformTime::Seconds() - StartTime) * 1000;
+	GAddToWorldTimeCumul += DeltaTime;
 }
 
 void UWorld::BeginTearingDown()
@@ -2873,8 +2886,12 @@ void UWorld::BeginTearingDown()
 	FWorldDelegates::OnWorldBeginTearDown.Broadcast(this);
 }
 
+// Cumulated time doing IncrementalUnregisterComponents in UWorld::RemoveFromWorld since last call to UWorld::UpdateLevelStreaming.
+static double GRemoveFromWorldUnregisterComponentTimeCumul = 0.0;
+
 void UWorld::RemoveFromWorld( ULevel* Level, bool bAllowIncrementalRemoval )
 {
+	TRACE_CPUPROFILER_EVENT_SCOPE(RemoveFromWorld);
 	SCOPE_CYCLE_COUNTER(STAT_RemoveFromWorldTime);
 	CSV_SCOPED_TIMING_STAT_EXCLUSIVE(RemoveFromWorld);
 
@@ -2889,7 +2906,7 @@ void UWorld::RemoveFromWorld( ULevel* Level, bool bAllowIncrementalRemoval )
 	{
 		Level->bIsBeingRemoved = true;
 
-		if (Level->bRequireFullVisibilityToRender)
+		if (Level->bRequireFullVisibilityToRender && Level->bIsVisible)
 		{
 			Level->bIsVisible = false;
 			ULevelStreaming::BroadcastLevelVisibleStatus(this, Level->GetOutermost()->GetFName(), false);
@@ -2916,22 +2933,29 @@ void UWorld::RemoveFromWorld( ULevel* Level, bool bAllowIncrementalRemoval )
 				BeginRemoval();
 			}
 
-			if ( CurrentLevelPendingInvisibility == Level )
+			if (CurrentLevelPendingInvisibility == Level)
 			{
-				// Incrementally unregister actor components. 
-				// This avoids spikes on the renderthread and gamethread when we subsequently call ClearLevelComponents() further down
-				check(IsGameWorld());
-				int32 NumComponentsToUnregister = GLevelStreamingComponentsUnregistrationGranularity;
-				do
+				double TimeLimit = FMath::Max<double>(0.0, GLevelStreamingUnregisterComponentsTimeLimit - GRemoveFromWorldUnregisterComponentTimeCumul);
+				if (TimeLimit > 0.0)
 				{
-					if (Level->IncrementalUnregisterComponents(NumComponentsToUnregister))
+					// Incrementally unregister actor components. 
+					// This avoids spikes on the renderthread and gamethread when we subsequently call ClearLevelComponents() further down
+					check(IsGameWorld());
+					int32 NumComponentsToUnregister = GLevelStreamingComponentsUnregistrationGranularity;
+					do
 					{
-						// We're done, so the level can be removed
-						CurrentLevelPendingInvisibility = nullptr;
-						bFinishRemovingLevel = true;
-						break;
-					}
-				} while (!IsTimeLimitExceeded(TEXT("unregistering components"), StartTime, Level, GLevelStreamingUnregisterComponentsTimeLimit));
+						if (Level->IncrementalUnregisterComponents(NumComponentsToUnregister))
+						{
+							// We're done, so the level can be removed
+							CurrentLevelPendingInvisibility = nullptr;
+							bFinishRemovingLevel = true;
+							break;
+						}
+					} while (!IsTimeLimitExceeded(TEXT("unregistering components"), StartTime, Level, TimeLimit));
+
+					double DeltaTime = (FPlatformTime::Seconds() - StartTime) * 1000;
+					GRemoveFromWorldUnregisterComponentTimeCumul += DeltaTime;
+				}
 			}
 		}
 		else
@@ -3671,10 +3695,15 @@ ULevel* UWorld::CalculateNextPreferredLevelPendingVisibility() const
 
 void UWorld::UpdateLevelStreaming()
 {
+	TRACE_CPUPROFILER_EVENT_SCOPE(UWorld::UpdateLevelStreaming);
 	SCOPE_CYCLE_COUNTER(STAT_UpdateLevelStreamingTime);
 	CSV_SCOPED_TIMING_STAT_EXCLUSIVE(UpdateLevelStreaming);
 	LLM_SCOPE(ELLMTag::LoadMapMisc);
 
+	// Reset counters
+	GAddToWorldTimeCumul = 0.0;
+	GRemoveFromWorldUnregisterComponentTimeCumul = 0.0;
+
 	// do nothing if level streaming is frozen
 	if (bIsLevelStreamingFrozen)
 	{
