Fix discrepancy between FPackageHarvester::ProcessExport serialization (not building outdated tree because of ArIsObjectReferenceCollector flag) and the following WriteExports serialization modifying object state.

diff --git a/Engine/Source/Runtime/Engine/Private/HierarchicalInstancedStaticMesh.cpp b/Engine/Source/Runtime/Engine/Private/HierarchicalInstancedStaticMesh.cpp
--- a/Engine/Source/Runtime/Engine/Private/HierarchicalInstancedStaticMesh.cpp
+++ b/Engine/Source/Runtime/Engine/Private/HierarchicalInstancedStaticMesh.cpp
@@ -2090,9 +2090,8 @@ void UHierarchicalInstancedStaticMeshComponent::Serialize(FArchive& Ar)
 
 	Ar.UsingCustomVersion(FReleaseObjectVersion::GUID);
 
-	// If we are saving make sure the tree is up to date. For Undo/Redo the PostEditUndo will rebuild the tree.
-	// Properly building the tree here will avoid the need to call BuildTreeIfOutdated in different duplication/load use cases.
-	if (Ar.IsSaving() && !Ar.IsTransacting() && !Ar.IsObjectReferenceCollector())
+	// Make sure to build tree before Save/Duplicate
+	if (Ar.IsSaving() && Ar.IsPersistent())
 	{
 		BuildTreeIfOutdated(/*Async*/false, /*ForceUpdate*/false);
 	}
