Fixed a bug where undo in the mirror data table was not restoring the location of the row

diff --git a/Engine/Source/Runtime/Engine/Classes/Animation/MirrorDataTable.h b/Engine/Source/Runtime/Engine/Classes/Animation/MirrorDataTable.h
--- a/Engine/Source/Runtime/Engine/Classes/Animation/MirrorDataTable.h
+++ b/Engine/Source/Runtime/Engine/Classes/Animation/MirrorDataTable.h
@@ -104,25 +104,14 @@ class UMirrorDataTable : public UDataTable
 public:
 	UMirrorDataTable(const FObjectInitializer& ObjectInitializer);
 
-	ENGINE_API virtual void GetPreloadDependencies(TArray<UObject*>& OutDeps) override;
-
-	ENGINE_API virtual void Serialize(FArchive& Ar) override;
-
 	ENGINE_API virtual void PostLoad() override;
 
-	ENGINE_API virtual void EmptyTable() override;
-
 #if WITH_EDITOR
 
-	ENGINE_API virtual void CleanBeforeStructChange() override;
-
-	ENGINE_API virtual void RestoreAfterStructChange() override;
-
 	ENGINE_API virtual void PreEditChange(FProperty* PropertyThatWillChange) override;
 
 	ENGINE_API virtual void PostEditChangeProperty(FPropertyChangedEvent& PropertyChangedEvent) override;
 
-	ENGINE_API virtual void PostEditUndo() override;
 
 #endif // WITH_EDITOR
 
@@ -209,7 +198,5 @@ protected:
 
 	// Fill BoneToMirrorBoneIndex, CurveMirrorSourceUIDArray, CurveMirrorTargetUIDArray and NotifyToMirrorNotifyIndex based on the Skeleton and Table Contents
 	ENGINE_API void FillMirrorArrays();
-
-	void HandleDataTableChanged();
 };
 
diff --git a/Engine/Source/Runtime/Engine/Private/Animation/MirrorDataTable.cpp b/Engine/Source/Runtime/Engine/Private/Animation/MirrorDataTable.cpp
--- a/Engine/Source/Runtime/Engine/Private/Animation/MirrorDataTable.cpp
+++ b/Engine/Source/Runtime/Engine/Private/Animation/MirrorDataTable.cpp
@@ -57,16 +57,7 @@ UMirrorDataTable::UMirrorDataTable(const FObjectInitializer& ObjectInitializer)
 #if WITH_EDITORONLY_DATA
 	OnDataTableChanged().AddUObject(this, &UMirrorDataTable::FillMirrorArrays);
 #endif 
-}
 
-void UMirrorDataTable::GetPreloadDependencies(TArray<UObject*>& OutDeps)
-{
-	Super::GetPreloadDependencies(OutDeps);
-}
-
-void UMirrorDataTable::Serialize(FArchive& Ar)
-{
-	Super::Serialize(Ar); 
 }
 
 void UMirrorDataTable::PostLoad()
@@ -80,20 +71,7 @@ void UMirrorDataTable::PostLoad()
 	}
 }
 
-void UMirrorDataTable::EmptyTable()
-{
-}
-
 #if WITH_EDITOR
-void UMirrorDataTable::CleanBeforeStructChange()
-{
-	Super::CleanBeforeStructChange();
-}
-
-void UMirrorDataTable::RestoreAfterStructChange()
-{
-	Super::RestoreAfterStructChange();
-}
 
 void UMirrorDataTable::PreEditChange(FProperty* PropertyThatWillChange)
 {
@@ -121,14 +99,8 @@ void UMirrorDataTable::PostEditChangeProperty(FPropertyChangedEvent& PropertyCha
 	}
 }
 
-void UMirrorDataTable::PostEditUndo()
-{
-	Super::PostEditUndo();
-}
-
 #endif // WITH_EDITOR
 
-
 FName UMirrorDataTable::GetSettingsMirrorName(FName InName)
 {
 	UAnimationSettings* AnimationSettings = UAnimationSettings::Get();
@@ -462,6 +434,4 @@ void UMirrorDataTable::FillMirrorArrays()
 	}
 }
 
-
-
 #undef LOCTEXT_NAMESPACE
\ No newline at end of file
