VirtualTexture builds: silently fail to build rather than crashing on textures with no pixels.
#rnx

diff --git a/Engine/Source/Runtime/Engine/Private/TextureDerivedDataTask.cpp b/Engine/Source/Runtime/Engine/Private/TextureDerivedDataTask.cpp
--- a/Engine/Source/Runtime/Engine/Private/TextureDerivedDataTask.cpp
+++ b/Engine/Source/Runtime/Engine/Private/TextureDerivedDataTask.cpp
@@ -359,6 +359,11 @@ void FTextureCacheDerivedDataWorker::BuildTexture(bool bReplaceExistingDDC)
 		return;
 	}
 
+	if (!bHasValidMip0)
+	{
+		return;
+	}
+
 	if (bForVirtualTextureStreamingBuild)
 	{
 		if (DerivedData->VTData == nullptr)
@@ -403,7 +408,7 @@ void FTextureCacheDerivedDataWorker::BuildTexture(bool bReplaceExistingDDC)
 			}
 		}
 	}
-	else if (bHasValidMip0)
+	else
 	{
 		// Only support single Block/Layer here (Blocks and Layers are intended for VT support)
 		if (TextureData.Blocks.Num() > 1)
@@ -921,7 +926,7 @@ void FTextureCacheDerivedDataWorker::Finalize()
 		BuildTexture();
 	}
 		
-	if (BuildSettingsPerLayer[0].bVirtualStreamable) // Texture.VirtualTextureStreaming is more a hint that might be overruled by the buildsettings
+	if (bSucceeded && BuildSettingsPerLayer[0].bVirtualStreamable) // Texture.VirtualTextureStreaming is more a hint that might be overruled by the buildsettings
 	{
 		check((DerivedData->VTData != nullptr) == Texture.VirtualTextureStreaming); 
 	}
