Add UnloadIntersectingCells

diff --git a/Engine/Source/Runtime/Engine/Private/WorldPartition/WorldPartitionVolume.cpp b/Engine/Source/Runtime/Engine/Private/WorldPartition/WorldPartitionVolume.cpp
--- a/Engine/Source/Runtime/Engine/Private/WorldPartition/WorldPartitionVolume.cpp
+++ b/Engine/Source/Runtime/Engine/Private/WorldPartition/WorldPartitionVolume.cpp
@@ -18,13 +18,25 @@ void AWorldPartitionVolume::LoadIntersectingCells()
 {
 	if (UWorldPartition* WorldPartition = GetWorld()->GetWorldPartition())
 	{
-		FVector Origin;
-		FVector Extent;
-		GetActorBounds(false, Origin, Extent);
+		WorldPartition->LoadEditorCells(GetIntersectingBounds());
+	}
+}
 
-		WorldPartition->LoadEditorCells(FBox(Origin - Extent, Origin + Extent));
+void AWorldPartitionVolume::UnloadIntersectingCells()
+{
+	if (UWorldPartition* WorldPartition = GetWorld()->GetWorldPartition())
+	{
+		WorldPartition->UnloadEditorCells(GetIntersectingBounds());
 	}
 }
+
+FBox AWorldPartitionVolume::GetIntersectingBounds() const
+{
+	FVector Origin;
+	FVector Extent;
+	GetActorBounds(false, Origin, Extent);
+	return FBox(Origin - Extent, Origin + Extent);
+}
 #endif
 
 #undef LOCTEXT_NAMESPACE
\ No newline at end of file
diff --git a/Engine/Source/Runtime/Engine/Public/WorldPartition/WorldPartitionVolume.h b/Engine/Source/Runtime/Engine/Public/WorldPartition/WorldPartitionVolume.h
--- a/Engine/Source/Runtime/Engine/Public/WorldPartition/WorldPartitionVolume.h
+++ b/Engine/Source/Runtime/Engine/Public/WorldPartition/WorldPartitionVolume.h
@@ -22,6 +22,10 @@ public:
 
 #if WITH_EDITOR
 	void LoadIntersectingCells();
+	void UnloadIntersectingCells();
 	virtual EActorGridPlacement GetDefaultGridPlacement() const override { return EActorGridPlacement::AlwaysLoaded; }
+
+private:
+	FBox GetIntersectingBounds() const;
 #endif
 };
\ No newline at end of file
