Engine: Improved UWorld::GetDebugDisplayName() to work properly outside of PIE ticks

diff --git a/Engine/Source/Runtime/Engine/Private/UnrealEngine.cpp b/Engine/Source/Runtime/Engine/Private/UnrealEngine.cpp
--- a/Engine/Source/Runtime/Engine/Private/UnrealEngine.cpp
+++ b/Engine/Source/Runtime/Engine/Private/UnrealEngine.cpp
@@ -419,58 +419,66 @@ bool GIsPrepareMapChangeBroken = false;
 ENGINE_API FString GPlayInEditorContextString(TEXT("invalid"));
 #endif
 
-ENGINE_API void UpdatePlayInEditorWorldDebugString(const FWorldContext* WorldContext)
+ENGINE_API FString GetDebugStringForWorld(const UWorld* World)
 {
-#if WITH_EDITOR
-	if (!GIsEditor)
-	{
-		return;
-	}
+	FString WorldName;
 
-	if (WorldContext == nullptr)
+	if (World == nullptr)
 	{
-		ensure(GPlayInEditorID == INDEX_NONE);
-		GPlayInEditorContextString = NSLOCTEXT("Engine", "PlayWorldIsNotActive", "Not in a play world").ToString();
+		WorldName = NSLOCTEXT("Engine", "PlayWorldBeingCreated", "(World Being Created)").ToString();
 	}
 	else
 	{
-		FString WorldName;
-		if (UWorld* World = WorldContext->World())
+		FWorldContext* WorldContext = GEngine->GetWorldContextFromWorld(World);
+
+		switch (World->GetNetMode())
 		{
-			switch (World->GetNetMode())
-			{
-			case NM_Standalone:
-				WorldName = NSLOCTEXT("Engine", "PlayWorldIsStandalone", "Standalone").ToString();
-				break;
+		case NM_Standalone:
+			WorldName = NSLOCTEXT("Engine", "PlayWorldIsStandalone", "Standalone").ToString();
+			break;
 
-			case NM_ListenServer:
-				WorldName = NSLOCTEXT("Engine", "PlayWorldIsListenServer", "Listen Server").ToString();
-				break;
+		case NM_ListenServer:
+			WorldName = NSLOCTEXT("Engine", "PlayWorldIsListenServer", "Listen Server").ToString();
+			break;
 
-			case NM_DedicatedServer:
-				WorldName = NSLOCTEXT("Engine", "PlayWorldIsDedicatedServer", "Dedicated Server").ToString();
-				break;
+		case NM_DedicatedServer:
+			WorldName = NSLOCTEXT("Engine", "PlayWorldIsDedicatedServer", "Dedicated Server").ToString();
+			break;
 
-			case NM_Client:
-				// 0 is always the server, use PIEInstance so it matches the in-editor UI
-				WorldName = FText::Format(NSLOCTEXT("Engine", "PlayWorldIsClient", "Client {0}"), FText::AsNumber(WorldContext->PIEInstance)).ToString();
-				break;
+		case NM_Client:
+			// 0 is always the server, use PIEInstance so it matches the in-editor UI
+			WorldName = FText::Format(NSLOCTEXT("Engine", "PlayWorldIsClient", "Client {0}"), FText::AsNumber(WorldContext ? WorldContext->PIEInstance : 0)).ToString();
+			break;
 
-			default:
-				unimplemented();
-			};
-		}
-		else
-		{
-			WorldName = NSLOCTEXT("Engine", "PlayWorldBeingCreated", "(World Being Created)").ToString();
-		}
+		default:
+			unimplemented();
+		};
 
-		if (!WorldContext->CustomDescription.IsEmpty())
+		if ((WorldContext != nullptr) && !WorldContext->CustomDescription.IsEmpty())
 		{
 			WorldName += TEXT(" ") + WorldContext->CustomDescription;
 		}
+	}
+
+	return WorldName;
+}
+
+ENGINE_API void UpdatePlayInEditorWorldDebugString(const FWorldContext* WorldContext)
+{
+#if WITH_EDITOR
+	if (!GIsEditor)
+	{
+		return;
+	}
 
-		GPlayInEditorContextString = WorldName;
+	if (WorldContext == nullptr)
+	{
+		ensure(GPlayInEditorID == INDEX_NONE);
+		GPlayInEditorContextString = NSLOCTEXT("Engine", "PlayWorldIsNotActive", "Not in a play world").ToString();
+	}
+	else
+	{
+		GPlayInEditorContextString = GetDebugStringForWorld(WorldContext->World());
 	}
 #endif
 }
diff --git a/Engine/Source/Runtime/Engine/Private/World.cpp b/Engine/Source/Runtime/Engine/Private/World.cpp
--- a/Engine/Source/Runtime/Engine/Private/World.cpp
+++ b/Engine/Source/Runtime/Engine/Private/World.cpp
@@ -5243,15 +5243,8 @@ void UWorld::SetWorldDataLayers(AWorldDataLayers* NewWorldDataLayers)
 
 FString UWorld::GetDebugDisplayName() const
 {
-#if WITH_EDITOR
-	if (GIsEditor)
-	{
-		extern ENGINE_API FString GPlayInEditorContextString;
-
-		return FString::Printf(TEXT("%s (%s)"), *GPlayInEditorContextString, *GetPathNameSafe(this));
-	}
-#endif
-	return GetPathNameSafe(this);
+	extern ENGINE_API FString GetDebugStringForWorld(const UWorld* World);
+	return FString::Printf(TEXT("%s (%s)"), *GetDebugStringForWorld(this), *GetPathNameSafe(this));
 }
 
 UWorldPartition* UWorld::GetWorldPartition() const
