Fixes for per-bone smoothing with blendspace graphs. Also reduces the work done at runtime (in some ways).

diff --git a/Engine/Source/Runtime/Engine/Classes/Animation/BlendProfile.h b/Engine/Source/Runtime/Engine/Classes/Animation/BlendProfile.h
--- a/Engine/Source/Runtime/Engine/Classes/Animation/BlendProfile.h
+++ b/Engine/Source/Runtime/Engine/Classes/Animation/BlendProfile.h
@@ -136,7 +136,7 @@ public:
 	void FillBoneScalesArray(TArray<float>& OutBoneBlendProfileFactors, const FBoneContainer& BoneContainer) const;
 
 	// IInterpolationIndexProvider
-	virtual int32 GetPerBoneInterpolationIndex(int32 BoneIndex, const FBoneContainer& RequiredBones) const override;
+	virtual int32 GetPerBoneInterpolationIndex(int32 BoneIndex, const FBoneContainer& RequiredBones, const IInterpolationIndexProvider::FPerBoneInterpolationData* Data) const override;
 	// End IInterpolationIndexProvider
 
 	// UObject
diff --git a/Engine/Source/Runtime/Engine/Classes/Animation/BlendSpace.h b/Engine/Source/Runtime/Engine/Classes/Animation/BlendSpace.h
--- a/Engine/Source/Runtime/Engine/Classes/Animation/BlendSpace.h
+++ b/Engine/Source/Runtime/Engine/Classes/Animation/BlendSpace.h
@@ -481,11 +481,20 @@ public:
 	//~ End UAnimationAsset Interface
 	
 	// Begin IInterpolationIndexProvider Overrides
+
+	/**
+	 * Sorts the PerBoneBlend data into a form that can be repeatedly used in GetPerBoneInterpolationIndex
+	 */
+	virtual TSharedPtr<IInterpolationIndexProvider::FPerBoneInterpolationData> GetPerBoneInterpolationData(const USkeleton* Skeleton) const override;
+
 	/**
 	* Get PerBoneInterpolationIndex for the input BoneIndex
 	* If nothing found, return INDEX_NONE
 	*/
-	virtual int32 GetPerBoneInterpolationIndex(int32 BoneIndex, const FBoneContainer& RequiredBones) const override;	
+	virtual int32 GetPerBoneInterpolationIndex(
+		int32                                                         BoneIndex, 
+		const FBoneContainer&                                         RequiredBones, 
+		const IInterpolationIndexProvider::FPerBoneInterpolationData* Data) const override;	
 	// End IInterpolationIndexProvider Overrides
 
 	/** Returns whether or not the given additive animation type is compatible with the blendspace type */
@@ -768,7 +777,6 @@ public:
 
 protected:
 
-#if WITH_EDITORONLY_DATA
 	/**
 	* Per bone sample smoothing settings, which affect the specified bone and all its descendants in the skeleton.
 	* These act as overrides to the global sample smoothing speed, which means the global sample smoothing speed does
@@ -777,13 +785,6 @@ protected:
 	*/
 	UPROPERTY(EditAnywhere, Category = SampleSmoothing, meta = (DisplayName="Per Bone Overrides"))
 	TArray<FPerBoneInterpolation> PerBoneBlend;
-#endif
-
-	/**
-	 * PerBoneBlend sorted so that we find the bones in the correct order, working back from leaf bones.
-	 */
-	UPROPERTY()
-	TArray<FPerBoneInterpolation> SortedPerBoneBlend;
 
 	/** Track index to get marker data from. Samples are tested for the suitability of marker based sync
 	    during load and if we can use marker based sync we cache an index to a representative sample here */
diff --git a/Engine/Source/Runtime/Engine/Private/Animation/AnimNode_StateMachine.cpp b/Engine/Source/Runtime/Engine/Private/Animation/AnimNode_StateMachine.cpp
--- a/Engine/Source/Runtime/Engine/Private/Animation/AnimNode_StateMachine.cpp
+++ b/Engine/Source/Runtime/Engine/Private/Animation/AnimNode_StateMachine.cpp
@@ -876,9 +876,10 @@ void FAnimNode_StateMachine::EvaluateTransitionStandardBlendInternal(FPoseContex
 	// If we have a blend profile we need to blend per bone
 	if(Transition.BlendProfile)
 	{
+		TSharedPtr<IInterpolationIndexProvider::FPerBoneInterpolationData> Data = Transition.BlendProfile->GetPerBoneInterpolationData(Output.AnimInstanceProxy->GetSkeleton());
 		for(FCompactPoseBoneIndex BoneIndex : Output.Pose.ForEachBoneIndex())
 		{
-			const int32 PerBoneIndex = Transition.BlendProfile->GetPerBoneInterpolationIndex(BoneIndex.GetInt(), Output.AnimInstanceProxy->GetRequiredBones());
+			const int32 PerBoneIndex = Transition.BlendProfile->GetPerBoneInterpolationIndex(BoneIndex.GetInt(), Output.AnimInstanceProxy->GetRequiredBones(), Data.Get());
 
 			// Use defined per-bone scale if the bone has a scale specified in the blend profile
 			ScalarRegister FirstWeight = PerBoneIndex != INDEX_NONE ? ScalarRegister(Transition.StateBlendData[1].PerBoneBlendData[PerBoneIndex]) : ScalarRegister(VPreviousWeight);
diff --git a/Engine/Source/Runtime/Engine/Private/Animation/AnimationRuntime.cpp b/Engine/Source/Runtime/Engine/Private/Animation/AnimationRuntime.cpp
--- a/Engine/Source/Runtime/Engine/Private/Animation/AnimationRuntime.cpp
+++ b/Engine/Source/Runtime/Engine/Private/Animation/AnimationRuntime.cpp
@@ -565,9 +565,10 @@ void FAnimationRuntime::BlendPosesTogetherPerBone(TArrayView<const FCompactPose>
 
 	TArray<int32> PerBoneIndices;
 	PerBoneIndices.AddUninitialized(OutPose.GetNumBones());
+	TSharedPtr<IInterpolationIndexProvider::FPerBoneInterpolationData> Data = InterpolationIndexProvider->GetPerBoneInterpolationData(OutPose.GetBoneContainer().GetSkeletonAsset());
 	for (int32 BoneIndex = 0; BoneIndex < PerBoneIndices.Num(); ++BoneIndex)
 	{
-		PerBoneIndices[BoneIndex] = InterpolationIndexProvider->GetPerBoneInterpolationIndex(RequiredBoneIndices[BoneIndex], OutPose.GetBoneContainer());
+		PerBoneIndices[BoneIndex] = InterpolationIndexProvider->GetPerBoneInterpolationIndex(RequiredBoneIndices[BoneIndex], OutPose.GetBoneContainer(), Data.Get());
 	}
 
 	BlendPosePerBone<ETransformBlendMode::Overwrite>(PerBoneIndices, BlendSampleDataCache[0], OutPose, SourcePoses[0]);
@@ -625,9 +626,10 @@ void FAnimationRuntime::BlendPosesTogetherPerBone(TArrayView<const FCompactPose>
 
 	TArray<int32> PerBoneIndices;
 	PerBoneIndices.AddUninitialized(OutPose.GetNumBones());
+	TSharedPtr<IInterpolationIndexProvider::FPerBoneInterpolationData> Data = InterpolationIndexProvider->GetPerBoneInterpolationData(OutPose.GetBoneContainer().GetSkeletonAsset());
 	for (int32 BoneIndex = 0; BoneIndex < PerBoneIndices.Num(); ++BoneIndex)
 	{
-		PerBoneIndices[BoneIndex] = InterpolationIndexProvider->GetPerBoneInterpolationIndex(RequiredBoneIndices[BoneIndex], OutPose.GetBoneContainer());
+		PerBoneIndices[BoneIndex] = InterpolationIndexProvider->GetPerBoneInterpolationIndex(RequiredBoneIndices[BoneIndex], OutPose.GetBoneContainer(), Data.Get());
 	}
 
 	BlendPosePerBone<ETransformBlendMode::Overwrite>(PerBoneIndices, BlendSampleDataCache[BlendSampleDataCacheIndices[0]], OutPose, SourcePoses[0]);
@@ -704,7 +706,7 @@ void FAnimationRuntime::BlendPosesTogetherPerBoneInMeshSpace(TArrayView<FCompact
 	BlendPosesTogetherPerBone(SourcePoses, SourceCurves, SourceAttributes, BlendSpace, BlendSampleDataCache, OutAnimationPoseData);
 
 	// now result atoms has the output with mesh space rotation. Convert back to local space, start from back
-	for (const FCompactPoseBoneIndex BoneIndex : OutPose.ForEachBoneIndex())
+	for (const FCompactPoseBoneIndex BoneIndex : OutPose.ForEachBoneIndexReverse())
 	{
 		const FCompactPoseBoneIndex ParentIndex = OutPose.GetParentBoneIndex(BoneIndex);
 		if (ParentIndex != INDEX_NONE)
diff --git a/Engine/Source/Runtime/Engine/Private/Animation/BlendProfile.cpp b/Engine/Source/Runtime/Engine/Private/Animation/BlendProfile.cpp
--- a/Engine/Source/Runtime/Engine/Private/Animation/BlendProfile.cpp
+++ b/Engine/Source/Runtime/Engine/Private/Animation/BlendProfile.cpp
@@ -134,7 +134,7 @@ float UBlendProfile::GetEntryBlendScale(const int32 InEntryIdx) const
 	return GetDefaultBlendScale();
 }
 
-int32 UBlendProfile::GetPerBoneInterpolationIndex(int32 BoneIndex, const FBoneContainer& RequiredBones) const
+int32 UBlendProfile::GetPerBoneInterpolationIndex(int32 BoneIndex, const FBoneContainer& RequiredBones, const IInterpolationIndexProvider::FPerBoneInterpolationData* Data) const
 {
 	// Our internal entries are skeleton bone indices, but we have pose indices here - convert to skeleton
 	int32 ActualBoneIndex = RequiredBones.GetPoseToSkeletonBoneIndexArray()[BoneIndex];
diff --git a/Engine/Source/Runtime/Engine/Private/Animation/BlendSpace.cpp b/Engine/Source/Runtime/Engine/Private/Animation/BlendSpace.cpp
--- a/Engine/Source/Runtime/Engine/Private/Animation/BlendSpace.cpp
+++ b/Engine/Source/Runtime/Engine/Private/Animation/BlendSpace.cpp
@@ -152,9 +152,8 @@ void UBlendSpace::PostEditChangeProperty(struct FPropertyChangedEvent& PropertyC
 	const FName MemberPropertyName = PropertyChangedEvent.MemberProperty ? PropertyChangedEvent.MemberProperty->GetFName() : NAME_None;
 	const FName PropertyName = PropertyChangedEvent.Property ? PropertyChangedEvent.Property->GetFName() : NAME_None;
 
-	if ((MemberPropertyName == GET_MEMBER_NAME_CHECKED(UBlendSpace, PerBoneBlend)
-		 && PropertyName == GET_MEMBER_NAME_CHECKED(FBoneReference, BoneName)) 
-		|| PropertyName == GET_MEMBER_NAME_CHECKED(UBlendSpace, PerBoneBlend))
+	if (MemberPropertyName == GET_MEMBER_NAME_CHECKED(UBlendSpace, PerBoneBlend) ||
+		PropertyName == GET_MEMBER_NAME_CHECKED(FPerBoneInterpolation, InterpolationSpeedPerSec))
 	{
 		InitializePerBoneBlend();
 	}
@@ -245,7 +244,7 @@ bool UBlendSpace::UpdateBlendSamples_Internal(
 	if (GetSamplesFromBlendInput(InBlendSpacePosition, NewSampleDataList, InOutCachedTriangulationIndex, true))
 	{
 		// if target weight interpolation is set
-		if (TargetWeightInterpolationSpeedPerSec > 0.f || SortedPerBoneBlend.Num() > 0)
+		if (TargetWeightInterpolationSpeedPerSec > 0.f || PerBoneBlend.Num() > 0)
 		{
 			// target weight interpolation
 			if (InterpolateWeightOfSampleData(InDeltaTime, InOutOldSampleDataList, NewSampleDataList, InOutSampleDataCache))
@@ -669,24 +668,68 @@ void UBlendSpace::RuntimeValidateMarkerData()
 
 #endif // WITH_EDITOR
 
-// Note that this gets called once per bone in the skeleton - so any excessive looping in here is
-// expensive. See UE-110576
-int32 UBlendSpace::GetPerBoneInterpolationIndex(int32 BoneIndex, const FBoneContainer& RequiredBones) const
+/** When per-bone blend data are required to be sorted, this stores the sorted copy and the index of the original */
+struct FSortedPerBoneInterpolation
 {
-	for (int32 Iter = 0; Iter < SortedPerBoneBlend.Num(); ++Iter)
+	FSortedPerBoneInterpolation(const FPerBoneInterpolation& Original, int32 Index) :
+		PerBoneBlend(Original), OriginalIndex(Index) {}
+	FPerBoneInterpolation PerBoneBlend;
+	/** Index into the original array */
+	int32 OriginalIndex;
+};
+struct FSortedPerBoneInterpolationData : public IInterpolationIndexProvider::FPerBoneInterpolationData
+{
+	TArray<FSortedPerBoneInterpolation> Data;
+};
+
+TSharedPtr<IInterpolationIndexProvider::FPerBoneInterpolationData> UBlendSpace::GetPerBoneInterpolationData(const USkeleton* RuntimeSkeleton) const 
+{
+	FSortedPerBoneInterpolationData* Data = new FSortedPerBoneInterpolationData();
+	Data->Data.SetNumUninitialized(PerBoneBlend.Num());
+	for (int32 Iter = 0 ; Iter != PerBoneBlend.Num() ; ++Iter)
+	{
+		Data->Data[Iter] = FSortedPerBoneInterpolation(PerBoneBlend[Iter], Iter);
+		if (!IsAsset())
+		{
+			FBoneReference& Bone = Data->Data[Iter].PerBoneBlend.BoneReference;
+			// Note that for blendspace graphs, the bone index won't have been set since the skeleton wasn't known at
+			// creation time. In that case, look it up (but note that it could be slow).
+			Bone.Initialize(RuntimeSkeleton);
+		}
+	}
+
+	Data->Data.Sort([](const FSortedPerBoneInterpolation& A, const FSortedPerBoneInterpolation& B)
+					{ return A.PerBoneBlend.BoneReference.BoneIndex > B.PerBoneBlend.BoneReference.BoneIndex; });
+	return MakeShareable<FSortedPerBoneInterpolationData>(Data);
+}
+
+int32 UBlendSpace::GetPerBoneInterpolationIndex(
+	int32                                                         BoneIndex, 
+	const FBoneContainer&                                         RequiredBones,
+	const IInterpolationIndexProvider::FPerBoneInterpolationData* Data) const
+{
+	if (!ensure(Data))
+	{
+		return INDEX_NONE;
+	}
+	const TArray<FSortedPerBoneInterpolation>& SortedData = static_cast<const FSortedPerBoneInterpolationData*>(Data)->Data;
+
+	for (int32 Iter = 0; Iter < SortedData.Num(); ++Iter)
 	{
-		const FBoneReference& SmoothedBone = SortedPerBoneBlend[Iter].BoneReference;
+		const FPerBoneInterpolation& PerBoneInterpolation = SortedData[Iter].PerBoneBlend;
+		int32 OriginalIndex = SortedData[Iter].OriginalIndex;
+		const FBoneReference& SmoothedBone = PerBoneInterpolation.BoneReference;
 		if (SmoothedBone.IsValidToEvaluate(RequiredBones))
 		{
 			int SmoothedBonePoseIndex = RequiredBones.GetCompactPoseIndexFromSkeletonIndex(SmoothedBone.BoneIndex).GetInt();
 			if (SmoothedBonePoseIndex == BoneIndex)
 			{
-				return Iter;
+				return OriginalIndex;
 			}
 			// Returns true if BoneIndex is a child of SmoothedBonePoseIndex. Searches up from BoneIndex
 			if (RequiredBones.BoneIsChildOf(BoneIndex, SmoothedBonePoseIndex))
 			{
-				return Iter;
+				return OriginalIndex;
 			}
 		}
 	}
@@ -828,7 +871,7 @@ void UBlendSpace::GetAnimationPose_Internal(TArray<FBlendSampleData>& BlendSampl
 
 	TArrayView<FCompactPose> ChildrenPosesView(ChildrenPoses);
 
-	if (SortedPerBoneBlend.Num() > 0)
+	if (PerBoneBlend.Num() > 0)
 	{
 		if (IsValidAdditive())
 		{
@@ -1156,8 +1199,15 @@ void UBlendSpace::ValidateSampleData()
 		}
 	}
 
-	// set rotation blend in mesh space
-	bRotationBlendInMeshSpace = ContainsMatchingSamples(AAT_RotationOffsetMeshSpace);
+	// Set the space for blending
+	if (IsAsset())
+	{
+		bRotationBlendInMeshSpace = ContainsMatchingSamples(AAT_RotationOffsetMeshSpace);
+	}
+	else
+	{
+		bRotationBlendInMeshSpace = IsValidAdditive();
+	}
 
 	SampleIndexWithMarkers = bAllMarkerPatternsMatch ? SampleWithMarkers : INDEX_NONE;
 
@@ -1435,24 +1485,10 @@ bool UBlendSpace::IsTooCloseToExistingSamplePoint(const FVector& SampleValue, in
 void UBlendSpace::InitializePerBoneBlend()
 {
 	const USkeleton* MySkeleton = GetSkeleton();
-#if WITH_EDITOR
 	for (FPerBoneInterpolation& BoneInterpolationData : PerBoneBlend)
 	{
 		BoneInterpolationData.Initialize(MySkeleton);
 	}
-	SortedPerBoneBlend = PerBoneBlend;
-#endif
-	for (FPerBoneInterpolation& BoneInterpolationData : SortedPerBoneBlend)
-	{
-		BoneInterpolationData.Initialize(MySkeleton);
-	}
-#if WITH_EDITOR
-	// Sort this by bigger to smaller. This is because the per-bone blends are intended to affect the
-	// children of as well as the specified bones, and we need to find them (in
-	// GetPerBoneInterpolationIndex) in order working back from the leaf bones.
-	SortedPerBoneBlend.Sort([](const FPerBoneInterpolation& A, const FPerBoneInterpolation& B) 
-		{ return A.BoneReference.BoneIndex > B.BoneReference.BoneIndex; });
-#endif
 }
 
 void UBlendSpace::TickFollowerSamples(TArray<FBlendSampleData>& SampleDataList, const int32 HighestWeightIndex, FAnimAssetTickContext& Context, bool bResetMarkerDataOnFollowers) const
@@ -1610,10 +1646,10 @@ bool UBlendSpace::InterpolateWeightOfSampleData(float DeltaTime, const TArray<FB
 		FBlendSampleData OldSample = *OldIt;
 		bool bTargetSampleExists = false;
 
-		if (OldSample.PerBoneBlendData.Num() != SortedPerBoneBlend.Num())
+		if (OldSample.PerBoneBlendData.Num() != PerBoneBlend.Num())
 		{
-			OldSample.PerBoneBlendData.Init(OldSample.TotalWeight, SortedPerBoneBlend.Num());
-			OldSample.PerBoneWeightRate.Init(OldSample.WeightRate, SortedPerBoneBlend.Num());
+			OldSample.PerBoneBlendData.Init(OldSample.TotalWeight, PerBoneBlend.Num());
+			OldSample.PerBoneWeightRate.Init(OldSample.WeightRate, PerBoneBlend.Num());
 		}
 
 		for (auto NewIt = NewSampleDataList.CreateConstIterator(); NewIt; ++NewIt)
@@ -1634,7 +1670,7 @@ bool UBlendSpace::InterpolateWeightOfSampleData(float DeltaTime, const TArray<FB
 					SmoothWeight(
 						InterpData.PerBoneBlendData[Iter], InterpData.PerBoneWeightRate[Iter],
 						OldSample.PerBoneBlendData[Iter], OldSample.PerBoneWeightRate[Iter], NewSample.TotalWeight,
-						DeltaTime, SortedPerBoneBlend[Iter].InterpolationSpeedPerSec, bTargetWeightInterpolationEaseInOut);
+						DeltaTime, PerBoneBlend[Iter].InterpolationSpeedPerSec, bTargetWeightInterpolationEaseInOut);
 					TotalPerBoneWeight += InterpData.PerBoneBlendData[Iter];
 				}
 
@@ -1661,7 +1697,7 @@ bool UBlendSpace::InterpolateWeightOfSampleData(float DeltaTime, const TArray<FB
 				SmoothWeight(
 					InterpData.PerBoneBlendData[Iter], InterpData.PerBoneWeightRate[Iter],
 					OldSample.PerBoneBlendData[Iter], OldSample.PerBoneWeightRate[Iter], 0.0f,
-					DeltaTime, SortedPerBoneBlend[Iter].InterpolationSpeedPerSec, bTargetWeightInterpolationEaseInOut);
+					DeltaTime, PerBoneBlend[Iter].InterpolationSpeedPerSec, bTargetWeightInterpolationEaseInOut);
 				TotalPerBoneWeight += InterpData.PerBoneBlendData[Iter];
 			}
 
@@ -1682,10 +1718,10 @@ bool UBlendSpace::InterpolateWeightOfSampleData(float DeltaTime, const TArray<FB
 		FBlendSampleData OldSample = *OldIt;
 		bool bOldSampleExists = false;
 
-		if (OldSample.PerBoneBlendData.Num() != SortedPerBoneBlend.Num())
+		if (OldSample.PerBoneBlendData.Num() != PerBoneBlend.Num())
 		{
-			OldSample.PerBoneBlendData.Init(OldSample.TotalWeight, SortedPerBoneBlend.Num());
-			OldSample.PerBoneWeightRate.Init(OldSample.WeightRate, SortedPerBoneBlend.Num());
+			OldSample.PerBoneBlendData.Init(OldSample.TotalWeight, PerBoneBlend.Num());
+			OldSample.PerBoneWeightRate.Init(OldSample.WeightRate, PerBoneBlend.Num());
 		}
 
 		for (auto NewIt = FinalSampleDataList.CreateConstIterator(); NewIt; ++NewIt)
@@ -1716,7 +1752,7 @@ bool UBlendSpace::InterpolateWeightOfSampleData(float DeltaTime, const TArray<FB
 				SmoothWeight(
 					InterpData.PerBoneBlendData[Iter], InterpData.PerBoneWeightRate[Iter],
 					OldSample.PerBoneBlendData[Iter], OldSample.PerBoneWeightRate[Iter], Target,
-					DeltaTime, SortedPerBoneBlend[Iter].InterpolationSpeedPerSec, bTargetWeightInterpolationEaseInOut);
+					DeltaTime, PerBoneBlend[Iter].InterpolationSpeedPerSec, bTargetWeightInterpolationEaseInOut);
 				TotalPerBoneWeight += InterpData.PerBoneBlendData[Iter];
 			}
 			if (InterpData.TotalWeight > ZERO_ANIMWEIGHT_THRESH || TotalPerBoneWeight > ZERO_ANIMWEIGHT_THRESH)
diff --git a/Engine/Source/Runtime/Engine/Public/AnimationRuntime.h b/Engine/Source/Runtime/Engine/Public/AnimationRuntime.h
--- a/Engine/Source/Runtime/Engine/Public/AnimationRuntime.h
+++ b/Engine/Source/Runtime/Engine/Public/AnimationRuntime.h
@@ -64,13 +64,26 @@ void BlendCurves(const TArrayView<const FBlendedCurve* const> SourceCurves, cons
 
 
 /////////////////////////////////////////////////////////
-/** Interface used to provide interpolation indices for per bone blends
-  *
+/** 
+  * Interface used to provide interpolation indices for per bone blends
   */
 class ENGINE_API IInterpolationIndexProvider
 {
 public:
-	virtual int32 GetPerBoneInterpolationIndex(int32 BoneIndex, const FBoneContainer& RequiredBones) const = 0;
+	struct FPerBoneInterpolationData
+	{
+		virtual ~FPerBoneInterpolationData() {}
+	};
+
+	// There may be times when the implementation can pre-calculate data needed for GetPerBoneInterpolationIndex, as the
+	// latter is often called multiple times whilst iterating over a skeleton.
+	virtual TSharedPtr<FPerBoneInterpolationData> GetPerBoneInterpolationData(const USkeleton* Skeleton) const { return nullptr; }
+
+	// Implementation should return the index into the PerBoneBlendData array that would be required when looking
+	// up/blending BoneIndex. This call will be passed the results of GetPerBoneInterpolationData, so the two functions
+	// should be matched.
+	virtual int32 GetPerBoneInterpolationIndex(
+		int32 BoneIndex, const FBoneContainer& RequiredBones, const FPerBoneInterpolationData* Data) const = 0;
 };
 
 /** In AnimationRunTime Library, we extract animation data based on Skeleton hierarchy, not ref pose hierarchy. 
