Fix CIS for non-editor builds

diff --git a/Engine/Source/Runtime/Engine/Classes/Animation/AnimNodeBase.h b/Engine/Source/Runtime/Engine/Classes/Animation/AnimNodeBase.h
--- a/Engine/Source/Runtime/Engine/Classes/Animation/AnimNodeBase.h
+++ b/Engine/Source/Runtime/Engine/Classes/Animation/AnimNodeBase.h
@@ -913,10 +913,10 @@ protected:
 	// Get anim node constant/folded data of the specified type given the identifier. Do not use directly - use GET_MUTABLE_ANIM_NODE_DATA
 	// Note: will assert if data is not held on the instance/dynamic. Use GetInstanceDataPtr/GET_INSTANCE_ANIM_NODE_DATA_PTR if the value
 	// might not be mutable, which will return null.
+#if WITH_EDITORONLY_DATA
 	template<typename DataType>
 	DataType& GetMutableData(UE::Anim::FNodeDataId InId, UObject* InObject = nullptr)
 	{
-#if WITH_EDITORONLY_DATA
 		if(NodeData)
 		{
 			return *static_cast<DataType*>(NodeData->GetMutableData(InId, this, InObject));
@@ -925,11 +925,8 @@ protected:
 		{
 			return *InId.GetProperty()->ContainerPtrToValuePtr<DataType>(this);
 		}
-#else
-		check(NodeData);
-		return *static_cast<DataType*>(NodeData->GetMutableData(InId, this, InObject));
-#endif
 	}
+#endif
 
 	// Get anim node mutable data of the specified type given the identifier. Do not use directly - use GET_INSTANCE_ANIM_NODE_DATA_PTR
 	// @return nullptr if the data is not mutable/dynamic
