Texture: Simplified the compact binary copy of the build settings

#rnx

diff --git a/Engine/Source/Developer/TextureBuild/Private/TextureBuildFunction.cpp b/Engine/Source/Developer/TextureBuild/Private/TextureBuildFunction.cpp
--- a/Engine/Source/Developer/TextureBuild/Private/TextureBuildFunction.cpp
+++ b/Engine/Source/Developer/TextureBuild/Private/TextureBuildFunction.cpp
@@ -18,6 +18,7 @@
 #include "Serialization/FileRegions.h"
 #include "Serialization/MemoryWriter.h"
 #include "TextureCompressorModule.h"
+#include "TextureFormatManager.h"
 
 DEFINE_LOG_CATEGORY_STATIC(LogTextureBuildFunction, Log, All);
 
@@ -35,104 +36,119 @@ static const TCHAR* GetPixelFormatString(EPixelFormat InPixelFormat)
 	}
 }
 
-static void ReadCbField(FCbFieldView FieldView, FColor& OutColor)
+static void ReadCbField(FCbFieldView Field, bool& OutValue) { OutValue = Field.AsBool(OutValue); }
+static void ReadCbField(FCbFieldView Field, int32& OutValue) { OutValue = Field.AsInt32(OutValue); }
+static void ReadCbField(FCbFieldView Field, uint8& OutValue) { OutValue = Field.AsUInt8(OutValue); }
+static void ReadCbField(FCbFieldView Field, uint32& OutValue) { OutValue = Field.AsUInt32(OutValue); }
+static void ReadCbField(FCbFieldView Field, float& OutValue) { OutValue = Field.AsFloat(OutValue); }
+
+static void ReadCbField(FCbFieldView Field, FName& OutValue)
+{
+	if (Field.IsString())
+	{
+		OutValue = FName(FUTF8ToTCHAR(Field.AsString()));
+	}
+}
+
+static void ReadCbField(FCbFieldView Field, FColor& OutValue)
 {
-	// Choosing the big endian ordering
-	FCbFieldViewIterator ColorCbArrayIt = FieldView.AsArrayView().CreateViewIterator();
-	OutColor.A = ColorCbArrayIt->AsUInt8(); ++ColorCbArrayIt;
-	OutColor.R = ColorCbArrayIt->AsUInt8(); ++ColorCbArrayIt;
-	OutColor.G = ColorCbArrayIt->AsUInt8(); ++ColorCbArrayIt;
-	OutColor.B = ColorCbArrayIt->AsUInt8(); ++ColorCbArrayIt;
+	FCbFieldViewIterator It = Field.AsArrayView().CreateViewIterator();
+	OutValue.A = It++->AsUInt8(OutValue.A);
+	OutValue.R = It++->AsUInt8(OutValue.R);
+	OutValue.G = It++->AsUInt8(OutValue.G);
+	OutValue.B = It++->AsUInt8(OutValue.B);
 }
 
-static void ReadCbField(FCbFieldView FieldView, FVector4& OutVec4)
+static void ReadCbField(FCbFieldView Field, FVector4& OutValue)
 {
-	FCbFieldViewIterator Vec4CbArrayIt = FieldView.AsArrayView().CreateViewIterator();
-	OutVec4.X = Vec4CbArrayIt->AsFloat(); ++Vec4CbArrayIt;
-	OutVec4.Y = Vec4CbArrayIt->AsFloat(); ++Vec4CbArrayIt;
-	OutVec4.Z = Vec4CbArrayIt->AsFloat(); ++Vec4CbArrayIt;
-	OutVec4.W = Vec4CbArrayIt->AsFloat(); ++Vec4CbArrayIt;
+	FCbFieldViewIterator It = Field.AsArrayView().CreateViewIterator();
+	OutValue.X = It++->AsFloat(OutValue.X);
+	OutValue.Y = It++->AsFloat(OutValue.Y);
+	OutValue.Z = It++->AsFloat(OutValue.Z);
+	OutValue.W = It++->AsFloat(OutValue.W);
 }
 
-static void ReadCbField(FCbFieldView FieldView, FIntPoint& OutIntPoint)
+static void ReadCbField(FCbFieldView Field, FIntPoint& OutValue)
 {
-	FCbFieldViewIterator IntPointCbArrayIt = FieldView.AsArrayView().CreateViewIterator();
-	OutIntPoint.X = IntPointCbArrayIt->AsInt32(); ++IntPointCbArrayIt;
-	OutIntPoint.Y = IntPointCbArrayIt->AsInt32(); ++IntPointCbArrayIt;
+	FCbFieldViewIterator It = Field.AsArrayView().CreateViewIterator();
+	OutValue.X = It++->AsInt32(OutValue.X);
+	OutValue.Y = It++->AsInt32(OutValue.Y);
 }
 
-static void ReadBuildSettingsFromCompactBinary(const FCbObject& Object, FTextureBuildSettings& OutBuildSettings)
+static FTextureBuildSettings ReadBuildSettingsFromCompactBinary(const FCbObjectView& Object)
 {
-	OutBuildSettings.FormatConfigOverride = Object.FindView("FormatConfigOverride").AsObjectView();
-	FCbObjectView ColorAdjustmentCbObj = Object.FindView("ColorAdjustment").AsObjectView();
-	FColorAdjustmentParameters& ColorAdjustment = OutBuildSettings.ColorAdjustment;
-	ColorAdjustment.AdjustBrightness = ColorAdjustmentCbObj.FindView("AdjustBrightness").AsFloat();
-	ColorAdjustment.AdjustBrightnessCurve = ColorAdjustmentCbObj.FindView("AdjustBrightnessCurve").AsFloat();
-	ColorAdjustment.AdjustSaturation = ColorAdjustmentCbObj.FindView("AdjustSaturation").AsFloat();
-	ColorAdjustment.AdjustVibrance = ColorAdjustmentCbObj.FindView("AdjustVibrance").AsFloat();
-	ColorAdjustment.AdjustRGBCurve = ColorAdjustmentCbObj.FindView("AdjustRGBCurve").AsFloat();
-	ColorAdjustment.AdjustHue = ColorAdjustmentCbObj.FindView("AdjustHue").AsFloat();
-	ColorAdjustment.AdjustMinAlpha = ColorAdjustmentCbObj.FindView("AdjustMinAlpha").AsFloat();
-	ColorAdjustment.AdjustMaxAlpha = ColorAdjustmentCbObj.FindView("AdjustMaxAlpha").AsFloat();
-	ReadCbField(Object.FindView("AlphaCoverageThresholds"), OutBuildSettings.AlphaCoverageThresholds);
-	OutBuildSettings.MipSharpening = Object.FindView("MipSharpening").AsFloat();
-	OutBuildSettings.DiffuseConvolveMipLevel = Object.FindView("DiffuseConvolveMipLevel").AsUInt32();
-	OutBuildSettings.SharpenMipKernelSize = Object.FindView("SharpenMipKernelSize").AsUInt32();
-	OutBuildSettings.MaxTextureResolution = Object.FindView("MaxTextureResolution").AsUInt32();
-	OutBuildSettings.TextureFormatName = FName(Object.FindView("TextureFormatName").AsString());
-	OutBuildSettings.bHDRSource = Object.FindView("bHDRSource").AsBool();
-	OutBuildSettings.MipGenSettings = Object.FindView("MipGenSettings").AsUInt8();
-	OutBuildSettings.bCubemap = Object.FindView("bCubemap").AsBool();
-	OutBuildSettings.bTextureArray = Object.FindView("bTextureArray").AsBool();
-	OutBuildSettings.bVolume = Object.FindView("bVolume").AsBool();
-	OutBuildSettings.bLongLatSource = Object.FindView("bLongLatSource").AsBool();
-	OutBuildSettings.bSRGB = Object.FindView("bSRGB").AsBool();
-	OutBuildSettings.bUseLegacyGamma = Object.FindView("bUseLegacyGamma").AsBool();
-	OutBuildSettings.bPreserveBorder = Object.FindView("bPreserveBorder").AsBool();
-	OutBuildSettings.bForceNoAlphaChannel = Object.FindView("bForceNoAlphaChannel").AsBool();
-	OutBuildSettings.bForceAlphaChannel = Object.FindView("bForceAlphaChannel").AsBool();
-	OutBuildSettings.bDitherMipMapAlpha = Object.FindView("bDitherMipMapAlpha").AsBool();
-	OutBuildSettings.bComputeBokehAlpha = Object.FindView("bComputeBokehAlpha").AsBool();
-	OutBuildSettings.bReplicateRed = Object.FindView("bReplicateRed").AsBool();
-	OutBuildSettings.bReplicateAlpha = Object.FindView("bReplicateAlpha").AsBool();
-	OutBuildSettings.bDownsampleWithAverage = Object.FindView("bDownsampleWithAverage").AsBool();
-	OutBuildSettings.bSharpenWithoutColorShift = Object.FindView("bSharpenWithoutColorShift").AsBool();
-	OutBuildSettings.bBorderColorBlack = Object.FindView("bBorderColorBlack").AsBool();
-	OutBuildSettings.bFlipGreenChannel = Object.FindView("bFlipGreenChannel").AsBool();
-	OutBuildSettings.bApplyYCoCgBlockScale = Object.FindView("bApplyYCoCgBlockScale").AsBool();
-	OutBuildSettings.bApplyKernelToTopMip = Object.FindView("bApplyKernelToTopMip").AsBool();
-	OutBuildSettings.bRenormalizeTopMip = Object.FindView("bRenormalizeTopMip").AsBool();
-	OutBuildSettings.CompositeTextureMode = Object.FindView("CompositeTextureMode").AsUInt8();
-	OutBuildSettings.CompositePower = Object.FindView("CompositePower").AsFloat();
-	OutBuildSettings.LODBias = Object.FindView("LODBias").AsUInt32();
-	OutBuildSettings.LODBiasWithCinematicMips = Object.FindView("LODBiasWithCinematicMips").AsUInt32();
-	ReadCbField(Object.FindView("TopMipSize"), OutBuildSettings.TopMipSize);
-	OutBuildSettings.VolumeSizeZ = Object.FindView("VolumeSizeZ").AsInt32();
-	OutBuildSettings.ArraySlices = Object.FindView("ArraySlices").AsInt32();
-	OutBuildSettings.bStreamable = Object.FindView("bStreamable").AsBool();
-	OutBuildSettings.bVirtualStreamable = Object.FindView("bVirtualStreamable").AsBool();
-	OutBuildSettings.bChromaKeyTexture = Object.FindView("bChromaKeyTexture").AsBool();
-	OutBuildSettings.PowerOfTwoMode = Object.FindView("PowerOfTwoMode").AsUInt8();
-	ReadCbField(Object.FindView("PaddingColor"), OutBuildSettings.PaddingColor);
-	ReadCbField(Object.FindView("ChromaKeyColor"), OutBuildSettings.ChromaKeyColor);
-	OutBuildSettings.ChromaKeyThreshold = Object.FindView("ChromaKeyThreshold").AsFloat();
-	OutBuildSettings.CompressionQuality = Object.FindView("CompressionQuality").AsInt32();
-	OutBuildSettings.LossyCompressionAmount = Object.FindView("LossyCompressionAmount").AsInt32();
-	OutBuildSettings.Downscale = Object.FindView("Downscale").AsFloat();
-	OutBuildSettings.DownscaleOptions = Object.FindView("DownscaleOptions").AsUInt8();
-	OutBuildSettings.VirtualAddressingModeX = Object.FindView("VirtualAddressingModeX").AsInt32();
-	OutBuildSettings.VirtualAddressingModeY = Object.FindView("VirtualAddressingModeY").AsInt32();
-	OutBuildSettings.VirtualTextureTileSize = Object.FindView("VirtualTextureTileSize").AsInt32();
-	OutBuildSettings.VirtualTextureBorderSize = Object.FindView("VirtualTextureBorderSize").AsInt32();
-	OutBuildSettings.bVirtualTextureEnableCompressZlib = Object.FindView("bVirtualTextureEnableCompressZlib").AsBool();
-	OutBuildSettings.bVirtualTextureEnableCompressCrunch = Object.FindView("bVirtualTextureEnableCompressCrunch").AsBool();
-	OutBuildSettings.bHasEditorOnlyData = Object.FindView("bHasEditorOnlyData").AsBool();
+	FTextureBuildSettings BuildSettings;
+	BuildSettings.FormatConfigOverride = Object["FormatConfigOverride"].AsObjectView();
+	FCbObjectView ColorAdjustmentCbObj = Object["ColorAdjustment"].AsObjectView();
+	FColorAdjustmentParameters& ColorAdjustment = BuildSettings.ColorAdjustment;
+	ReadCbField(ColorAdjustmentCbObj["AdjustBrightness"], ColorAdjustment.AdjustBrightness);
+	ReadCbField(ColorAdjustmentCbObj["AdjustBrightnessCurve"], ColorAdjustment.AdjustBrightnessCurve);
+	ReadCbField(ColorAdjustmentCbObj["AdjustSaturation"], ColorAdjustment.AdjustSaturation);
+	ReadCbField(ColorAdjustmentCbObj["AdjustVibrance"], ColorAdjustment.AdjustVibrance);
+	ReadCbField(ColorAdjustmentCbObj["AdjustRGBCurve"], ColorAdjustment.AdjustRGBCurve);
+	ReadCbField(ColorAdjustmentCbObj["AdjustHue"], ColorAdjustment.AdjustHue);
+	ReadCbField(ColorAdjustmentCbObj["AdjustMinAlpha"], ColorAdjustment.AdjustMinAlpha);
+	ReadCbField(ColorAdjustmentCbObj["AdjustMaxAlpha"], ColorAdjustment.AdjustMaxAlpha);
+	ReadCbField(Object["AlphaCoverageThresholds"], BuildSettings.AlphaCoverageThresholds);
+	ReadCbField(Object["MipSharpening"], BuildSettings.MipSharpening);
+	ReadCbField(Object["DiffuseConvolveMipLevel"], BuildSettings.DiffuseConvolveMipLevel);
+	ReadCbField(Object["SharpenMipKernelSize"], BuildSettings.SharpenMipKernelSize);
+	ReadCbField(Object["MaxTextureResolution"], BuildSettings.MaxTextureResolution);
+	ReadCbField(Object["TextureFormatName"], BuildSettings.TextureFormatName);
+	ReadCbField(Object["bHDRSource"], BuildSettings.bHDRSource);
+	ReadCbField(Object["MipGenSettings"], BuildSettings.MipGenSettings);
+	BuildSettings.bCubemap = Object["bCubemap"].AsBool(BuildSettings.bCubemap);
+	BuildSettings.bTextureArray = Object["bTextureArray"].AsBool(BuildSettings.bTextureArray);
+	BuildSettings.bVolume = Object["bVolume"].AsBool(BuildSettings.bVolume);
+	BuildSettings.bLongLatSource = Object["bLongLatSource"].AsBool(BuildSettings.bLongLatSource);
+	BuildSettings.bSRGB = Object["bSRGB"].AsBool(BuildSettings.bSRGB);
+	BuildSettings.bUseLegacyGamma = Object["bUseLegacyGamma"].AsBool(BuildSettings.bUseLegacyGamma);
+	BuildSettings.bPreserveBorder = Object["bPreserveBorder"].AsBool(BuildSettings.bPreserveBorder);
+	BuildSettings.bForceNoAlphaChannel = Object["bForceNoAlphaChannel"].AsBool(BuildSettings.bForceNoAlphaChannel);
+	BuildSettings.bForceAlphaChannel = Object["bForceAlphaChannel"].AsBool(BuildSettings.bForceAlphaChannel);
+	BuildSettings.bDitherMipMapAlpha = Object["bDitherMipMapAlpha"].AsBool(BuildSettings.bDitherMipMapAlpha);
+	BuildSettings.bComputeBokehAlpha = Object["bComputeBokehAlpha"].AsBool(BuildSettings.bComputeBokehAlpha);
+	BuildSettings.bReplicateRed = Object["bReplicateRed"].AsBool(BuildSettings.bReplicateRed);
+	BuildSettings.bReplicateAlpha = Object["bReplicateAlpha"].AsBool(BuildSettings.bReplicateAlpha);
+	BuildSettings.bDownsampleWithAverage = Object["bDownsampleWithAverage"].AsBool(BuildSettings.bDownsampleWithAverage);
+	BuildSettings.bSharpenWithoutColorShift = Object["bSharpenWithoutColorShift"].AsBool(BuildSettings.bSharpenWithoutColorShift);
+	BuildSettings.bBorderColorBlack = Object["bBorderColorBlack"].AsBool(BuildSettings.bBorderColorBlack);
+	BuildSettings.bFlipGreenChannel = Object["bFlipGreenChannel"].AsBool(BuildSettings.bFlipGreenChannel);
+	BuildSettings.bApplyYCoCgBlockScale = Object["bApplyYCoCgBlockScale"].AsBool(BuildSettings.bApplyYCoCgBlockScale);
+	BuildSettings.bApplyKernelToTopMip = Object["bApplyKernelToTopMip"].AsBool(BuildSettings.bApplyKernelToTopMip);
+	BuildSettings.bRenormalizeTopMip = Object["bRenormalizeTopMip"].AsBool(BuildSettings.bRenormalizeTopMip);
+	ReadCbField(Object["CompositeTextureMode"], BuildSettings.CompositeTextureMode);
+	ReadCbField(Object["CompositePower"], BuildSettings.CompositePower);
+	ReadCbField(Object["LODBias"], BuildSettings.LODBias);
+	ReadCbField(Object["LODBiasWithCinematicMips"], BuildSettings.LODBiasWithCinematicMips);
+	ReadCbField(Object["TopMipSize"], BuildSettings.TopMipSize);
+	ReadCbField(Object["VolumeSizeZ"], BuildSettings.VolumeSizeZ);
+	ReadCbField(Object["ArraySlices"], BuildSettings.ArraySlices);
+	BuildSettings.bStreamable = Object["bStreamable"].AsBool(BuildSettings.bStreamable);
+	BuildSettings.bVirtualStreamable = Object["bVirtualStreamable"].AsBool(BuildSettings.bVirtualStreamable);
+	BuildSettings.bChromaKeyTexture = Object["bChromaKeyTexture"].AsBool(BuildSettings.bChromaKeyTexture);
+	ReadCbField(Object["PowerOfTwoMode"], BuildSettings.PowerOfTwoMode);
+	ReadCbField(Object["PaddingColor"], BuildSettings.PaddingColor);
+	ReadCbField(Object["ChromaKeyColor"], BuildSettings.ChromaKeyColor);
+	ReadCbField(Object["ChromaKeyThreshold"], BuildSettings.ChromaKeyThreshold);
+	ReadCbField(Object["CompressionQuality"], BuildSettings.CompressionQuality);
+	ReadCbField(Object["LossyCompressionAmount"], BuildSettings.LossyCompressionAmount);
+	ReadCbField(Object["Downscale"], BuildSettings.Downscale);
+	ReadCbField(Object["DownscaleOptions"], BuildSettings.DownscaleOptions);
+	ReadCbField(Object["VirtualAddressingModeX"], BuildSettings.VirtualAddressingModeX);
+	ReadCbField(Object["VirtualAddressingModeY"], BuildSettings.VirtualAddressingModeY);
+	ReadCbField(Object["VirtualTextureTileSize"], BuildSettings.VirtualTextureTileSize);
+	ReadCbField(Object["VirtualTextureBorderSize"], BuildSettings.VirtualTextureBorderSize);
+	BuildSettings.bVirtualTextureEnableCompressZlib = Object["bVirtualTextureEnableCompressZlib"].AsBool(BuildSettings.bVirtualTextureEnableCompressZlib);
+	BuildSettings.bVirtualTextureEnableCompressCrunch = Object["bVirtualTextureEnableCompressCrunch"].AsBool(BuildSettings.bVirtualTextureEnableCompressCrunch);
+	BuildSettings.bHasEditorOnlyData = Object["bHasEditorOnlyData"].AsBool(BuildSettings.bHasEditorOnlyData);
+	return BuildSettings;
 }
 
-static void ReadOutputSettingsFromCompactBinary(const FCbObject& Object, int32& NumInlineMips, FString& MipKeyPrefix)
+static void ReadOutputSettingsFromCompactBinary(const FCbObjectView& Object, int32& NumInlineMips, FString& MipKeyPrefix)
 {
-	NumInlineMips = Object.FindView("NumInlineMips").AsInt32();
-	MipKeyPrefix = FString(Object.FindView("MipKeyPrefix").AsString());
+	NumInlineMips = Object["NumInlineMips"].AsInt32();
+	MipKeyPrefix = FUTF8ToTCHAR(Object["MipKeyPrefix"].AsString());
 }
 
 static ERawImageFormat::Type ComputeRawImageFormat(ETextureSourceFormat SourceFormat)
@@ -151,20 +167,24 @@ static ERawImageFormat::Type ComputeRawImageFormat(ETextureSourceFormat SourceFo
 	}
 }
 
-static void ReadTextureSourceFromCompactBinary(const FCbObject& Object, UE::DerivedData::FBuildContext& Context, TArray<FImage>& OutMips)
+static bool TryReadTextureSourceFromCompactBinary(const FCbObjectView& Object, UE::DerivedData::FBuildContext& Context, TArray<FImage>& OutMips)
 {
-	FAnsiStringView InputKey = Object.FindView("Input").AsString();
-	FSharedBuffer InputBuffer = Context.FindInput(FUTF8ToTCHAR(InputKey));
+	FSharedBuffer InputBuffer = Context.FindInput(Object["Input"].AsUuid().ToString());
+	if (!InputBuffer)
+	{
+		UE_LOG(LogTextureBuildFunction, Error, TEXT("Missing input %s."), *Object["Input"].AsUuid().ToString());
+		return false;
+	}
 
-	ETextureSourceCompressionFormat CompressionFormat = (ETextureSourceCompressionFormat)Object.FindView("CompressionFormat").AsUInt8();
-	ETextureSourceFormat SourceFormat = (ETextureSourceFormat)Object.FindView("SourceFormat").AsUInt8();
+	ETextureSourceCompressionFormat CompressionFormat = (ETextureSourceCompressionFormat)Object["CompressionFormat"].AsUInt8();
+	ETextureSourceFormat SourceFormat = (ETextureSourceFormat)Object["SourceFormat"].AsUInt8();
 
 	ERawImageFormat::Type RawImageFormat = ComputeRawImageFormat(SourceFormat);
 
-	EGammaSpace GammaSpace = (EGammaSpace)Object.FindView("GammaSpace").AsUInt8();
-	int32 NumSlices = Object.FindView("NumSlices").AsInt32();
-	int32 SizeX = Object.FindView("SizeX").AsInt32();
-	int32 SizeY = Object.FindView("SizeY").AsInt32();
+	EGammaSpace GammaSpace = (EGammaSpace)Object["GammaSpace"].AsUInt8();
+	int32 NumSlices = Object["NumSlices"].AsInt32();
+	int32 SizeX = Object["SizeX"].AsInt32();
+	int32 SizeY = Object["SizeY"].AsInt32();
 	int32 MipSizeX = SizeX;
 	int32 MipSizeY = SizeY;
 
@@ -190,20 +210,20 @@ static void ReadTextureSourceFromCompactBinary(const FCbObject& Object, UE::Deri
 		}
 		break;
 		default:
-			checkf(false, TEXT("Unexpected source compression format encountered while attempting to build a texture."));
-			break;
+			UE_LOG(LogTextureBuildFunction, Error, TEXT("Unexpected source compression format encountered while attempting to build a texture."));
+			return false;
 		}
 		DecompressedSourceData = IntermediateDecompressedData.GetData();
 		InputBuffer.Reset();
 	}
 
-	FCbArrayView MipsCbArrayView = Object.FindView("Mips").AsArrayView();
+	FCbArrayView MipsCbArrayView = Object["Mips"].AsArrayView();
 	OutMips.Reserve(MipsCbArrayView.Num());
-	for (FCbFieldViewIterator MipsCbArrayIt = MipsCbArrayView.CreateViewIterator(); MipsCbArrayIt; ++MipsCbArrayIt)
+	for (FCbFieldView MipsCbArrayIt : MipsCbArrayView)
 	{
-		FCbObjectView MipCbObjectView = MipsCbArrayIt->AsObjectView();
-		int64 MipOffset = MipCbObjectView.FindView("Offset").AsInt64();
-		int64 MipSize = MipCbObjectView.FindView("Size").AsInt64();
+		FCbObjectView MipCbObjectView = MipsCbArrayIt.AsObjectView();
+		int64 MipOffset = MipCbObjectView["Offset"].AsInt64();
+		int64 MipSize = MipCbObjectView["Size"].AsInt64();
 
 		FImage* SourceMip = new(OutMips) FImage(
 			MipSizeX, MipSizeY,
@@ -231,32 +251,53 @@ static void ReadTextureSourceFromCompactBinary(const FCbObject& Object, UE::Deri
 		MipSizeY = FMath::Max(MipSizeY / 2, 1);
 	}
 
+	return true;
 }
 
 void FTextureBuildFunction::Configure(UE::DerivedData::FBuildConfigContext& Context) const
 {
-	Context.SetCacheBucket(Context.CreateCacheBucket(TEXT("Texture")));
+	Context.SetCacheBucket(Context.CreateCacheBucket(TEXT("Texture"_SV)));
 }
 
 void FTextureBuildFunction::Build(UE::DerivedData::FBuildContext& Context) const
 {
-	ITextureCompressorModule& Compressor = FModuleManager::LoadModuleChecked<ITextureCompressorModule>(TEXTURE_COMPRESSOR_MODULENAME);
+	const FCbObject Settings = Context.FindConstant(TEXT("Settings"_SV));
+	if (!Settings)
+	{
+		UE_LOG(LogTextureBuildFunction, Error, TEXT("Settings are not available."));
+		return;
+	}
+
+	const FTextureBuildSettings BuildSettings = ReadBuildSettingsFromCompactBinary(Settings["Build"].AsObjectView());
+
+	const uint16 RequiredTextureFormatVersion = Settings["FormatVersion"].AsUInt16();
+	if (ITextureFormatManagerModule* TFM = GetTextureFormatManager())
+	{
+		const ITextureFormat* TextureFormat = TFM->FindTextureFormat(BuildSettings.TextureFormatName);
+		const uint16 CurrentTextureFormatVersion = TextureFormat ? TextureFormat->GetVersion(BuildSettings.TextureFormatName, &BuildSettings) : 0;
+		if (CurrentTextureFormatVersion != RequiredTextureFormatVersion)
+		{
+			UE_LOG(LogTextureBuildFunction, Error, TEXT("%s has version %hu when version %hu is required."),
+				*BuildSettings.TextureFormatName.ToString(), CurrentTextureFormatVersion, RequiredTextureFormatVersion);;
+			return;
+		}
+	}
 
-	const FCbObject BuildSettingsCbObj = Context.FindConstant(TEXT("TextureBuildSettings"));
-	FTextureBuildSettings BuildSettings;
-	ReadBuildSettingsFromCompactBinary(Context.FindConstant(TEXT("TextureBuildSettings")), BuildSettings);
 	int32 NumInlineMips;
 	FString MipKeyPrefix;
-	ReadOutputSettingsFromCompactBinary(Context.FindConstant(TEXT("TextureOutputSettings")), NumInlineMips, MipKeyPrefix);
+	ReadOutputSettingsFromCompactBinary(Settings["Output"].AsObjectView(), NumInlineMips, MipKeyPrefix);
 
 	TArray<FImage> SourceMips;
-	ReadTextureSourceFromCompactBinary(Context.FindConstant(TEXT("TextureSource")), Context, SourceMips);
+	if (!TryReadTextureSourceFromCompactBinary(Settings["Source"].AsObjectView(), Context, SourceMips))
+	{
+		return;
+	}
 
 	TArray<FImage> AssociatedNormalSourceMips;
-	const FCbObject AssociatedNormalSourceMipsCbObj = Context.FindConstant(TEXT("CompositeTextureSource"));
-	if (AssociatedNormalSourceMipsCbObj)
+	if (const FCbObjectView CompositeSource = Settings["CompositeSource"].AsObjectView();
+		CompositeSource && !TryReadTextureSourceFromCompactBinary(CompositeSource, Context, AssociatedNormalSourceMips))
 	{
-		ReadTextureSourceFromCompactBinary(AssociatedNormalSourceMipsCbObj, Context, AssociatedNormalSourceMips);
+		return;
 	}
 
 	UE_LOG(LogTextureBuildFunction, Display, TEXT("Compressing %d source mip(s) (%dx%d) to %s..."), SourceMips.Num(), SourceMips[0].SizeX, SourceMips[0].SizeY, *BuildSettings.TextureFormatName.ToString());
@@ -265,7 +306,8 @@ void FTextureBuildFunction::Build(UE::DerivedData::FBuildContext& Context) const
 	uint32 NumMipsInTail;
 	uint32 ExtData;
 
-	Compressor.BuildTexture(SourceMips,
+	FModuleManager::LoadModuleChecked<ITextureCompressorModule>(TEXTURE_COMPRESSOR_MODULENAME).BuildTexture(
+		SourceMips,
 		AssociatedNormalSourceMips,
 		BuildSettings,
 		CompressedMips,
@@ -373,7 +415,7 @@ void FTextureBuildFunction::Build(UE::DerivedData::FBuildContext& Context) const
 			EFileRegionType FileRegion = FFileRegion::SelectType(EPixelFormat(CompressedMip.PixelFormat));
 			MipFooterWriter << FileRegion;
 
-			TStringBuilder<128> DerivedDataKey;
+			TStringBuilder<512> DerivedDataKey;
 			if (!bIsInlineMip)
 			{
 				// This has to match the behavior in GetTextureDerivedMipKey in TextureDerivedData.cpp
@@ -393,6 +435,6 @@ void FTextureBuildFunction::Build(UE::DerivedData::FBuildContext& Context) const
 
 
 		FCompositeBuffer CompositeResult(OrderedBuffers);
-		Context.AddPayload(UE::DerivedData::FPayloadId::FromName(TEXT("Texture")), FCompressedBuffer::Compress(NAME_Default, CompositeResult));
+		Context.AddPayload(UE::DerivedData::FPayloadId::FromName(TEXT("Texture"_SV)), FCompressedBuffer::Compress(NAME_Default, CompositeResult));
 	}
 }
diff --git a/Engine/Source/Runtime/Engine/Private/TextureDerivedData.cpp b/Engine/Source/Runtime/Engine/Private/TextureDerivedData.cpp
--- a/Engine/Source/Runtime/Engine/Private/TextureDerivedData.cpp
+++ b/Engine/Source/Runtime/Engine/Private/TextureDerivedData.cpp
@@ -268,6 +268,15 @@ static void SerializeForKey(FArchive& Ar, const FTextureBuildSettings& Settings)
 	}
 }
 
+/**
+ * Returns the texture derived data version.
+ */
+const FGuid& GetTextureDerivedDataVersion()
+{
+	static FGuid Version(TEXTURE_DERIVEDDATA_VER);
+	return Version;
+}
+
 /**
  * Constructs a derived data key from the key suffix.
  * @param KeySuffix - The key suffix.
diff --git a/Engine/Source/Runtime/Engine/Private/TextureDerivedDataBuildExporter.cpp b/Engine/Source/Runtime/Engine/Private/TextureDerivedDataBuildExporter.cpp
--- a/Engine/Source/Runtime/Engine/Private/TextureDerivedDataBuildExporter.cpp
+++ b/Engine/Source/Runtime/Engine/Private/TextureDerivedDataBuildExporter.cpp
@@ -166,11 +166,8 @@ void FTextureDerivedDataBuildExporter::ExportTextureBuild(const UTexture& Textur
 
 	UE::DerivedData::FBuildActionBuilder ActionBuilder = DerivedDataBuild->CreateAction(TexturePath, BuildFunctionName);
 
-	ComposeTextureBuildFunctionConstants(KeySuffix, Texture, BuildSettings, LayerIndex, NumInlineMips, 
-		[&ActionBuilder] (FStringView Key, const FCbObject& Value)
-		{
-			ActionBuilder.AddConstant(Key, Value);
-		});
+	ActionBuilder.AddConstant(TEXT("Settings"_SV),
+		SaveTextureBuildSettings(KeySuffix, Texture, BuildSettings, LayerIndex, NumInlineMips));
 
 	ActionBuilder.AddInput(Texture.Source.GetId().ToString(), ExportedTextureBulkDataHash, ExportedTextureBulkDataSize);
 	if ((bool)Texture.CompositeTexture && !ExportedCompositeTextureBulkDataHash.IsZero())
diff --git a/Engine/Source/Runtime/Engine/Private/TextureDerivedDataBuildUtils.cpp b/Engine/Source/Runtime/Engine/Private/TextureDerivedDataBuildUtils.cpp
--- a/Engine/Source/Runtime/Engine/Private/TextureDerivedDataBuildUtils.cpp
+++ b/Engine/Source/Runtime/Engine/Private/TextureDerivedDataBuildUtils.cpp
@@ -5,50 +5,65 @@
 #if WITH_EDITOR
 #include "Engine/Texture.h"
 #include "Interfaces/ITextureFormat.h"
-#include "Interfaces/ITextureFormat.h"
 #include "Interfaces/ITextureFormatManagerModule.h"
 #include "Interfaces/ITextureFormatModule.h"
 #include "Serialization/CompactBinary.h"
 #include "Serialization/CompactBinaryWriter.h"
+#include "String/Find.h"
 #include "TextureCompressorModule.h"
 #include "TextureFormatManager.h"
 #include "TextureResource.h"
 
-void GetTextureDerivedDataKeyFromSuffix(const FString& KeySuffix, FString& OutKey);
+const FGuid& GetTextureDerivedDataVersion();
 void GetTextureDerivedMipKey(int32 MipIndex, const FTexture2DMipMap& Mip, const FString& KeySuffix, FString& OutKey);
 
-static void WriteCbField(FCbWriter& Writer, const FAnsiStringView& Name, const FColor& Color)
+template <typename ValueType>
+static void WriteCbField(FCbWriter& Writer, FAnsiStringView Name, const ValueType& Value)
+{
+	Writer << Name << Value;
+}
+
+static void WriteCbField(FCbWriter& Writer, FAnsiStringView Name, const FColor& Value)
 {
-	// Choosing the big endian ordering
 	Writer.BeginArray(Name);
-	Writer.AddInteger(Color.A);
-	Writer.AddInteger(Color.R);
-	Writer.AddInteger(Color.G);
-	Writer.AddInteger(Color.B);
+	Writer.AddInteger(Value.A);
+	Writer.AddInteger(Value.R);
+	Writer.AddInteger(Value.G);
+	Writer.AddInteger(Value.B);
 	Writer.EndArray();
 }
 
-static void WriteCbField(FCbWriter& Writer, const FAnsiStringView& Name, const FVector4& Vec4)
+static void WriteCbField(FCbWriter& Writer, FAnsiStringView Name, const FVector4& Value)
 {
 	Writer.BeginArray(Name);
-	Writer.AddFloat(Vec4.X);
-	Writer.AddFloat(Vec4.Y);
-	Writer.AddFloat(Vec4.Z);
-	Writer.AddFloat(Vec4.W);
+	Writer.AddFloat(Value.X);
+	Writer.AddFloat(Value.Y);
+	Writer.AddFloat(Value.Z);
+	Writer.AddFloat(Value.W);
 	Writer.EndArray();
 }
 
-static void WriteCbField(FCbWriter& Writer, const FAnsiStringView& Name, const FIntPoint& IntPoint)
+static void WriteCbField(FCbWriter& Writer, FAnsiStringView Name, const FIntPoint& Value)
 {
 	Writer.BeginArray(Name);
-	Writer.AddInteger(IntPoint.X);
-	Writer.AddInteger(IntPoint.Y);
+	Writer.AddInteger(Value.X);
+	Writer.AddInteger(Value.Y);
 	Writer.EndArray();
 }
 
-static FCbObject WriteBuildSettingsToCompactBinary(const FTextureBuildSettings& BuildSettings, const ITextureFormat* TextureFormat)
+template <typename ValueType>
+static void WriteCbFieldWithDefault(FCbWriter& Writer, FAnsiStringView Name, ValueType Value, ValueType Default)
 {
-	FCbWriter Writer;
+	if (Value != Default)
+	{
+		WriteCbField(Writer, Name, Forward<ValueType>(Value));
+	}
+}
+
+static void WriteBuildSettings(FCbWriter& Writer, const FTextureBuildSettings& BuildSettings, const ITextureFormat* TextureFormat)
+{
+	FTextureBuildSettings DefaultSettings;
+
 	Writer.BeginObject();
 
 	if (BuildSettings.FormatConfigOverride)
@@ -59,81 +74,87 @@ static FCbObject WriteBuildSettingsToCompactBinary(const FTextureBuildSettings&
 	{
 		Writer.AddObject("FormatConfigOverride", TextureFormatConfig);
 	}
-	Writer.BeginObject("ColorAdjustment");
-	Writer.AddFloat("AdjustBrightness", BuildSettings.ColorAdjustment.AdjustBrightness);
-	Writer.AddFloat("AdjustBrightnessCurve", BuildSettings.ColorAdjustment.AdjustBrightnessCurve);
-	Writer.AddFloat("AdjustSaturation", BuildSettings.ColorAdjustment.AdjustSaturation);
-	Writer.AddFloat("AdjustVibrance", BuildSettings.ColorAdjustment.AdjustVibrance);
-	Writer.AddFloat("AdjustRGBCurve", BuildSettings.ColorAdjustment.AdjustRGBCurve);
-	Writer.AddFloat("AdjustHue", BuildSettings.ColorAdjustment.AdjustHue);
-	Writer.AddFloat("AdjustMinAlpha", BuildSettings.ColorAdjustment.AdjustMinAlpha);
-	Writer.AddFloat("AdjustMaxAlpha", BuildSettings.ColorAdjustment.AdjustMaxAlpha);
-	Writer.EndObject();
 
-	WriteCbField(Writer, "AlphaCoverageThresholds", BuildSettings.AlphaCoverageThresholds);
-
-	Writer.AddFloat("MipSharpening", BuildSettings.MipSharpening);
-	Writer.AddInteger("DiffuseConvolveMipLevel", BuildSettings.DiffuseConvolveMipLevel);
-	Writer.AddInteger("SharpenMipKernelSize", BuildSettings.SharpenMipKernelSize);
-	Writer.AddInteger("MaxTextureResolution", BuildSettings.MaxTextureResolution);
-	Writer.AddString("TextureFormatName", WriteToString<64>(BuildSettings.TextureFormatName));
-	Writer.AddBool("bHDRSource", BuildSettings.bHDRSource);
-	Writer.AddInteger("MipGenSettings", BuildSettings.MipGenSettings);
-	Writer.AddBool("bCubemap", BuildSettings.bCubemap);
-	Writer.AddBool("bTextureArray", BuildSettings.bTextureArray);
-	Writer.AddBool("bVolume", BuildSettings.bVolume);
-	Writer.AddBool("bLongLatSource", BuildSettings.bLongLatSource);
-	Writer.AddBool("bSRGB", BuildSettings.bSRGB);
-	Writer.AddBool("bUseLegacyGamma", BuildSettings.bUseLegacyGamma);
-	Writer.AddBool("bPreserveBorder", BuildSettings.bPreserveBorder);
-	Writer.AddBool("bForceNoAlphaChannel", BuildSettings.bForceNoAlphaChannel);
-	Writer.AddBool("bForceAlphaChannel", BuildSettings.bForceAlphaChannel);
-	Writer.AddBool("bDitherMipMapAlpha", BuildSettings.bDitherMipMapAlpha);
-	Writer.AddBool("bComputeBokehAlpha", BuildSettings.bComputeBokehAlpha);
-	Writer.AddBool("bReplicateRed", BuildSettings.bReplicateRed);
-	Writer.AddBool("bReplicateAlpha", BuildSettings.bReplicateAlpha);
-	Writer.AddBool("bDownsampleWithAverage", BuildSettings.bDownsampleWithAverage);
-	Writer.AddBool("bSharpenWithoutColorShift", BuildSettings.bSharpenWithoutColorShift);
-	Writer.AddBool("bBorderColorBlack", BuildSettings.bBorderColorBlack);
-	Writer.AddBool("bFlipGreenChannel", BuildSettings.bFlipGreenChannel);
-	Writer.AddBool("bApplyYCoCgBlockScale", BuildSettings.bApplyYCoCgBlockScale);
-	Writer.AddBool("bApplyKernelToTopMip", BuildSettings.bApplyKernelToTopMip);
-	Writer.AddBool("bRenormalizeTopMip", BuildSettings.bRenormalizeTopMip);
-	Writer.AddInteger("CompositeTextureMode", BuildSettings.CompositeTextureMode);
-	Writer.AddFloat("CompositePower", BuildSettings.CompositePower);
-	Writer.AddInteger("LODBias", BuildSettings.LODBias);
-	Writer.AddInteger("LODBiasWithCinematicMips", BuildSettings.LODBiasWithCinematicMips);
-
-	WriteCbField(Writer, "TopMipSize", BuildSettings.TopMipSize);
-
-	Writer.AddInteger("VolumeSizeZ", BuildSettings.VolumeSizeZ);
-	Writer.AddInteger("ArraySlices", BuildSettings.ArraySlices);
-	Writer.AddBool("bStreamable", BuildSettings.bStreamable);
-	Writer.AddBool("bVirtualStreamable", BuildSettings.bVirtualStreamable);
-	Writer.AddBool("bChromaKeyTexture", BuildSettings.bChromaKeyTexture);
-	Writer.AddInteger("PowerOfTwoMode", BuildSettings.PowerOfTwoMode);
-	WriteCbField(Writer, "PaddingColor", BuildSettings.PaddingColor);
-	WriteCbField(Writer, "ChromaKeyColor", BuildSettings.ChromaKeyColor);
-	Writer.AddFloat("ChromaKeyThreshold", BuildSettings.ChromaKeyThreshold);
-	Writer.AddInteger("CompressionQuality", BuildSettings.CompressionQuality);
-	Writer.AddInteger("LossyCompressionAmount", BuildSettings.LossyCompressionAmount);
-	Writer.AddFloat("Downscale", BuildSettings.Downscale);
-	Writer.AddInteger("DownscaleOptions", BuildSettings.DownscaleOptions);
-	Writer.AddInteger("VirtualAddressingModeX", BuildSettings.VirtualAddressingModeX);
-	Writer.AddInteger("VirtualAddressingModeY", BuildSettings.VirtualAddressingModeY);
-	Writer.AddInteger("VirtualTextureTileSize", BuildSettings.VirtualTextureTileSize);
-	Writer.AddInteger("VirtualTextureBorderSize", BuildSettings.VirtualTextureBorderSize);
-	Writer.AddBool("bVirtualTextureEnableCompressZlib", BuildSettings.bVirtualTextureEnableCompressZlib);
-	Writer.AddBool("bVirtualTextureEnableCompressCrunch", BuildSettings.bVirtualTextureEnableCompressCrunch);
-	Writer.AddBool("bHasEditorOnlyData", BuildSettings.bHasEditorOnlyData);
+	if (BuildSettings.ColorAdjustment.AdjustBrightness != DefaultSettings.ColorAdjustment.AdjustBrightness ||
+		BuildSettings.ColorAdjustment.AdjustBrightnessCurve != DefaultSettings.ColorAdjustment.AdjustBrightnessCurve ||
+		BuildSettings.ColorAdjustment.AdjustSaturation != DefaultSettings.ColorAdjustment.AdjustSaturation ||
+		BuildSettings.ColorAdjustment.AdjustVibrance != DefaultSettings.ColorAdjustment.AdjustVibrance ||
+		BuildSettings.ColorAdjustment.AdjustRGBCurve != DefaultSettings.ColorAdjustment.AdjustRGBCurve ||
+		BuildSettings.ColorAdjustment.AdjustHue != DefaultSettings.ColorAdjustment.AdjustHue ||
+		BuildSettings.ColorAdjustment.AdjustMinAlpha != DefaultSettings.ColorAdjustment.AdjustMinAlpha ||
+		BuildSettings.ColorAdjustment.AdjustMaxAlpha != DefaultSettings.ColorAdjustment.AdjustMaxAlpha)
+	{
+		Writer.BeginObject("ColorAdjustment");
+		WriteCbFieldWithDefault(Writer, "AdjustBrightness", BuildSettings.ColorAdjustment.AdjustBrightness, DefaultSettings.ColorAdjustment.AdjustBrightness);
+		WriteCbFieldWithDefault(Writer, "AdjustBrightnessCurve", BuildSettings.ColorAdjustment.AdjustBrightnessCurve, DefaultSettings.ColorAdjustment.AdjustBrightnessCurve);
+		WriteCbFieldWithDefault(Writer, "AdjustSaturation", BuildSettings.ColorAdjustment.AdjustSaturation, DefaultSettings.ColorAdjustment.AdjustSaturation);
+		WriteCbFieldWithDefault(Writer, "AdjustVibrance", BuildSettings.ColorAdjustment.AdjustVibrance, DefaultSettings.ColorAdjustment.AdjustVibrance);
+		WriteCbFieldWithDefault(Writer, "AdjustRGBCurve", BuildSettings.ColorAdjustment.AdjustRGBCurve, DefaultSettings.ColorAdjustment.AdjustRGBCurve);
+		WriteCbFieldWithDefault(Writer, "AdjustHue", BuildSettings.ColorAdjustment.AdjustHue, DefaultSettings.ColorAdjustment.AdjustHue);
+		WriteCbFieldWithDefault(Writer, "AdjustMinAlpha", BuildSettings.ColorAdjustment.AdjustMinAlpha, DefaultSettings.ColorAdjustment.AdjustMinAlpha);
+		WriteCbFieldWithDefault(Writer, "AdjustMaxAlpha", BuildSettings.ColorAdjustment.AdjustMaxAlpha, DefaultSettings.ColorAdjustment.AdjustMaxAlpha);
+		Writer.EndObject();
+	}
+
+	WriteCbFieldWithDefault(Writer, "AlphaCoverageThresholds", BuildSettings.AlphaCoverageThresholds, DefaultSettings.AlphaCoverageThresholds);
+	WriteCbFieldWithDefault(Writer, "MipSharpening", BuildSettings.MipSharpening, DefaultSettings.MipSharpening);
+	WriteCbFieldWithDefault(Writer, "DiffuseConvolveMipLevel", BuildSettings.DiffuseConvolveMipLevel, DefaultSettings.DiffuseConvolveMipLevel);
+	WriteCbFieldWithDefault(Writer, "SharpenMipKernelSize", BuildSettings.SharpenMipKernelSize, DefaultSettings.SharpenMipKernelSize);
+	WriteCbFieldWithDefault(Writer, "MaxTextureResolution", BuildSettings.MaxTextureResolution, DefaultSettings.MaxTextureResolution);
+	WriteCbFieldWithDefault(Writer, "TextureFormatName", WriteToString<64>(BuildSettings.TextureFormatName).ToView(), TEXT(""_SV));
+	WriteCbFieldWithDefault(Writer, "bHDRSource", BuildSettings.bHDRSource, DefaultSettings.bHDRSource);
+	WriteCbFieldWithDefault(Writer, "MipGenSettings", BuildSettings.MipGenSettings, DefaultSettings.MipGenSettings);
+	WriteCbFieldWithDefault(Writer, "bCubemap", BuildSettings.bCubemap, DefaultSettings.bCubemap);
+	WriteCbFieldWithDefault(Writer, "bTextureArray", BuildSettings.bTextureArray, DefaultSettings.bTextureArray);
+	WriteCbFieldWithDefault(Writer, "bVolume", BuildSettings.bVolume, DefaultSettings.bVolume);
+	WriteCbFieldWithDefault(Writer, "bLongLatSource", BuildSettings.bLongLatSource, DefaultSettings.bLongLatSource);
+	WriteCbFieldWithDefault(Writer, "bSRGB", BuildSettings.bSRGB, DefaultSettings.bSRGB);
+	WriteCbFieldWithDefault(Writer, "bUseLegacyGamma", BuildSettings.bUseLegacyGamma, DefaultSettings.bUseLegacyGamma);
+	WriteCbFieldWithDefault(Writer, "bPreserveBorder", BuildSettings.bPreserveBorder, DefaultSettings.bPreserveBorder);
+	WriteCbFieldWithDefault(Writer, "bForceNoAlphaChannel", BuildSettings.bForceNoAlphaChannel, DefaultSettings.bForceNoAlphaChannel);
+	WriteCbFieldWithDefault(Writer, "bForceAlphaChannel", BuildSettings.bForceAlphaChannel, DefaultSettings.bForceAlphaChannel);
+	WriteCbFieldWithDefault(Writer, "bDitherMipMapAlpha", BuildSettings.bDitherMipMapAlpha, DefaultSettings.bDitherMipMapAlpha);
+	WriteCbFieldWithDefault(Writer, "bComputeBokehAlpha", BuildSettings.bComputeBokehAlpha, DefaultSettings.bComputeBokehAlpha);
+	WriteCbFieldWithDefault(Writer, "bReplicateRed", BuildSettings.bReplicateRed, DefaultSettings.bReplicateRed);
+	WriteCbFieldWithDefault(Writer, "bReplicateAlpha", BuildSettings.bReplicateAlpha, DefaultSettings.bReplicateAlpha);
+	WriteCbFieldWithDefault(Writer, "bDownsampleWithAverage", BuildSettings.bDownsampleWithAverage, DefaultSettings.bDownsampleWithAverage);
+	WriteCbFieldWithDefault(Writer, "bSharpenWithoutColorShift", BuildSettings.bSharpenWithoutColorShift, DefaultSettings.bSharpenWithoutColorShift);
+	WriteCbFieldWithDefault(Writer, "bBorderColorBlack", BuildSettings.bBorderColorBlack, DefaultSettings.bBorderColorBlack);
+	WriteCbFieldWithDefault(Writer, "bFlipGreenChannel", BuildSettings.bFlipGreenChannel, DefaultSettings.bFlipGreenChannel);
+	WriteCbFieldWithDefault(Writer, "bApplyYCoCgBlockScale", BuildSettings.bApplyYCoCgBlockScale, DefaultSettings.bApplyYCoCgBlockScale);
+	WriteCbFieldWithDefault(Writer, "bApplyKernelToTopMip", BuildSettings.bApplyKernelToTopMip, DefaultSettings.bApplyKernelToTopMip);
+	WriteCbFieldWithDefault(Writer, "bRenormalizeTopMip", BuildSettings.bRenormalizeTopMip, DefaultSettings.bRenormalizeTopMip);
+	WriteCbFieldWithDefault(Writer, "CompositeTextureMode", BuildSettings.CompositeTextureMode, DefaultSettings.CompositeTextureMode);
+	WriteCbFieldWithDefault(Writer, "CompositePower", BuildSettings.CompositePower, DefaultSettings.CompositePower);
+	WriteCbFieldWithDefault(Writer, "LODBias", BuildSettings.LODBias, DefaultSettings.LODBias);
+	WriteCbFieldWithDefault(Writer, "LODBiasWithCinematicMips", BuildSettings.LODBiasWithCinematicMips, DefaultSettings.LODBiasWithCinematicMips);
+	WriteCbFieldWithDefault(Writer, "TopMipSize", BuildSettings.TopMipSize, DefaultSettings.TopMipSize);
+	WriteCbFieldWithDefault(Writer, "VolumeSizeZ", BuildSettings.VolumeSizeZ, DefaultSettings.VolumeSizeZ);
+	WriteCbFieldWithDefault(Writer, "ArraySlices", BuildSettings.ArraySlices, DefaultSettings.ArraySlices);
+	WriteCbFieldWithDefault(Writer, "bStreamable", BuildSettings.bStreamable, DefaultSettings.bStreamable);
+	WriteCbFieldWithDefault(Writer, "bVirtualStreamable", BuildSettings.bVirtualStreamable, DefaultSettings.bVirtualStreamable);
+	WriteCbFieldWithDefault(Writer, "bChromaKeyTexture", BuildSettings.bChromaKeyTexture, DefaultSettings.bChromaKeyTexture);
+	WriteCbFieldWithDefault(Writer, "PowerOfTwoMode", BuildSettings.PowerOfTwoMode, DefaultSettings.PowerOfTwoMode);
+	WriteCbFieldWithDefault(Writer, "PaddingColor", BuildSettings.PaddingColor, DefaultSettings.PaddingColor);
+	WriteCbFieldWithDefault(Writer, "ChromaKeyColor", BuildSettings.ChromaKeyColor, DefaultSettings.ChromaKeyColor);
+	WriteCbFieldWithDefault(Writer, "ChromaKeyThreshold", BuildSettings.ChromaKeyThreshold, DefaultSettings.ChromaKeyThreshold);
+	WriteCbFieldWithDefault(Writer, "CompressionQuality", BuildSettings.CompressionQuality, DefaultSettings.CompressionQuality);
+	WriteCbFieldWithDefault(Writer, "LossyCompressionAmount", BuildSettings.LossyCompressionAmount, DefaultSettings.LossyCompressionAmount);
+	WriteCbFieldWithDefault(Writer, "Downscale", BuildSettings.Downscale, DefaultSettings.Downscale);
+	WriteCbFieldWithDefault(Writer, "DownscaleOptions", BuildSettings.DownscaleOptions, DefaultSettings.DownscaleOptions);
+	WriteCbFieldWithDefault(Writer, "VirtualAddressingModeX", BuildSettings.VirtualAddressingModeX, DefaultSettings.VirtualAddressingModeX);
+	WriteCbFieldWithDefault(Writer, "VirtualAddressingModeY", BuildSettings.VirtualAddressingModeY, DefaultSettings.VirtualAddressingModeY);
+	WriteCbFieldWithDefault(Writer, "VirtualTextureTileSize", BuildSettings.VirtualTextureTileSize, DefaultSettings.VirtualTextureTileSize);
+	WriteCbFieldWithDefault(Writer, "VirtualTextureBorderSize", BuildSettings.VirtualTextureBorderSize, DefaultSettings.VirtualTextureBorderSize);
+	WriteCbFieldWithDefault(Writer, "bVirtualTextureEnableCompressZlib", BuildSettings.bVirtualTextureEnableCompressZlib, DefaultSettings.bVirtualTextureEnableCompressZlib);
+	WriteCbFieldWithDefault(Writer, "bVirtualTextureEnableCompressCrunch", BuildSettings.bVirtualTextureEnableCompressCrunch, DefaultSettings.bVirtualTextureEnableCompressCrunch);
+	WriteCbFieldWithDefault(Writer, "bHasEditorOnlyData", BuildSettings.bHasEditorOnlyData, DefaultSettings.bHasEditorOnlyData);
 
 	Writer.EndObject();
-	return Writer.Save().AsObject();
 }
 
-static FCbObject WriteOutputSettingsToCompactBinary(int32 NumInlineMips, const FString& KeySuffix)
+static void WriteOutputSettings(FCbWriter& Writer, int32 NumInlineMips, const FString& KeySuffix)
 {
-	FCbWriter Writer;
 	Writer.BeginObject();
 
 	Writer.AddInteger("NumInlineMips", NumInlineMips);
@@ -150,29 +171,32 @@ static FCbObject WriteOutputSettingsToCompactBinary(int32 NumInlineMips, const F
 	Writer.AddString("MipKeyPrefix",*MipDerivedDataKey);
 
 	Writer.EndObject();
-	return Writer.Save().AsObject();
 }
 
-static FCbObject WriteTextureSourceToCompactBinary(const FTextureSource& TextureSource, EGammaSpace GammaSpace)
+static void WriteSource(FCbWriter& Writer, const UTexture& Texture, int32 LayerIndex)
 {
-	FCbWriter Writer;
+	const FTextureSource& Source = Texture.Source;
+
+	FTextureFormatSettings TextureFormatSettings;
+	Texture.GetLayerFormatSettings(LayerIndex, TextureFormatSettings);
+	EGammaSpace GammaSpace = TextureFormatSettings.SRGB ? (Texture.bUseLegacyGamma ? EGammaSpace::Pow22 : EGammaSpace::sRGB) : EGammaSpace::Linear;
+
 	Writer.BeginObject();
 
-	Writer.AddString("Input", TextureSource.GetId().ToString());
-	Writer.AddInteger("CompressionFormat", TextureSource.GetSourceCompression());
-	Writer.AddInteger("SourceFormat", TextureSource.GetFormat());
-	Writer.AddInteger("GammaSpace", (int32)GammaSpace);
-	Writer.AddInteger("NumSlices", TextureSource.GetNumSlices());
-	Writer.AddInteger("SizeX", TextureSource.GetSizeX());
-	Writer.AddInteger("SizeY", TextureSource.GetSizeY());
+	Writer.AddUuid("Input", Source.GetId());
+	Writer.AddInteger("CompressionFormat", Source.GetSourceCompression());
+	Writer.AddInteger("SourceFormat", Source.GetFormat());
+	Writer.AddInteger("GammaSpace", static_cast<uint8>(GammaSpace));
+	Writer.AddInteger("NumSlices", Source.GetNumSlices());
+	Writer.AddInteger("SizeX", Source.GetSizeX());
+	Writer.AddInteger("SizeY", Source.GetSizeY());
 	Writer.BeginArray("Mips");
-	const int32 NumMips = TextureSource.GetNumMips();
 	int64 Offset = 0;
-	for (int32 MipIndex = 0; MipIndex < NumMips; ++MipIndex)
+	for (int32 MipIndex = 0, MipCount = Source.GetNumMips(); MipIndex < MipCount; ++MipIndex)
 	{
 		Writer.BeginObject();
 		Writer.AddInteger("Offset", Offset);
-		const int64 MipSize = TextureSource.CalcMipSize(MipIndex);
+		const int64 MipSize = Source.CalcMipSize(MipIndex);
 		Writer.AddInteger("Size", MipSize);
 		Offset += MipSize;
 		Writer.EndObject();
@@ -180,67 +204,73 @@ static FCbObject WriteTextureSourceToCompactBinary(const FTextureSource& Texture
 	Writer.EndArray();
 
 	Writer.EndObject();
-	return Writer.Save().AsObject();
 }
 
 FString GetTextureBuildFunctionName(const FTextureBuildSettings& BuildSettings)
 {
-	const ITextureFormat* TextureFormat = nullptr;
 	FName TextureFormatModuleName;
-	ITextureFormatModule* TextureFormatModule = nullptr;
 
-	ITextureFormatManagerModule* TFM = GetTextureFormatManager();
-	if (TFM)
+	if (ITextureFormatManagerModule* TFM = GetTextureFormatManager())
 	{
-		TextureFormat = TFM->FindTextureFormatAndModule(BuildSettings.TextureFormatName, TextureFormatModuleName, TextureFormatModule);
+		ITextureFormatModule* TextureFormatModule = nullptr;
+		if (!TFM->FindTextureFormatAndModule(BuildSettings.TextureFormatName, TextureFormatModuleName, TextureFormatModule))
+		{
+			return FString();
+		}
 	}
 
-	if (TextureFormat == nullptr)
+	// Texture format modules are inconsistent in their naming, e.g., TextureFormatUncompressed, <Platform>TextureFormat.
+	// Attempt to unify the naming of build functions as <Format>Texture.
+	TStringBuilder<64> FunctionName;
+	FunctionName << TextureFormatModuleName << TEXT("Texture"_SV);
+	if (int32 Index = UE::String::FindFirst(FunctionName, TEXT("TextureFormat"_SV)); Index != INDEX_NONE)
 	{
-		return FString();
+		FunctionName.RemoveAt(Index, TEXT("TextureFormat"_SV).Len());
 	}
-
-	// Texture format modules are inconsistent in their naming.  eg: TextureFormatUncompressed, PS5TextureFormat
-	// We attempt to  unify the naming here when specifying build function names.
-	TStringBuilder<64> BuildFunctionNameBuilder;
-	BuildFunctionNameBuilder << TextureFormatModuleName.ToString().Replace(TEXT("TextureFormat"), TEXT(""));
-	BuildFunctionNameBuilder << TEXT("Texture");
-
-	return BuildFunctionNameBuilder.ToString();
+	return FString(FunctionName);
 }
 
-void ComposeTextureBuildFunctionConstants(const FString& KeySuffix, const UTexture& Texture, const FTextureBuildSettings& BuildSettings, int32 LayerIndex, int32 NumInlineMips, FTextureConstantOperator Operator)
+FCbObject SaveTextureBuildSettings(const FString& KeySuffix, const UTexture& Texture, const FTextureBuildSettings& BuildSettings, int32 LayerIndex, int32 NumInlineMips)
 {
 	const ITextureFormat* TextureFormat = nullptr;
-	FName TextureFormatModuleName;
-	ITextureFormatModule* TextureFormatModule = nullptr;
-
-	ITextureFormatManagerModule* TFM = GetTextureFormatManager();
-	if (TFM)
+	if (ITextureFormatManagerModule* TFM = GetTextureFormatManager())
 	{
+		FName TextureFormatModuleName;
+		ITextureFormatModule* TextureFormatModule = nullptr;
 		TextureFormat = TFM->FindTextureFormatAndModule(BuildSettings.TextureFormatName, TextureFormatModuleName, TextureFormatModule);
 	}
-
 	if (TextureFormat == nullptr)
 	{
-		return;
+		return FCbObject();
 	}
 
-	Operator(TEXT("TextureBuildSettings"), WriteBuildSettingsToCompactBinary(BuildSettings, TextureFormat));
-	Operator(TEXT("TextureOutputSettings"), WriteOutputSettingsToCompactBinary(NumInlineMips, KeySuffix));
+	FCbWriter Writer;
+	Writer.BeginObject();
 
-	FTextureFormatSettings TextureFormatSettings;
-	Texture.GetLayerFormatSettings(LayerIndex, TextureFormatSettings);
-	EGammaSpace TextureGammaSpace = TextureFormatSettings.SRGB ? (Texture.bUseLegacyGamma ? EGammaSpace::Pow22 : EGammaSpace::sRGB) : EGammaSpace::Linear;
-	Operator(TEXT("TextureSource"), WriteTextureSourceToCompactBinary(Texture.Source, TextureGammaSpace));
-	if ((bool)Texture.CompositeTexture)
+	Writer.AddUuid("BuildVersion", GetTextureDerivedDataVersion());
+
+	if (uint16 TextureFormatVersion = TextureFormat->GetVersion(BuildSettings.TextureFormatName, &BuildSettings))
 	{
-		FTextureFormatSettings CompositeTextureFormatSettings;
-		Texture.CompositeTexture->GetLayerFormatSettings(LayerIndex, CompositeTextureFormatSettings);
-		EGammaSpace CompositeTextureGammaSpace = CompositeTextureFormatSettings.SRGB ? (Texture.CompositeTexture->bUseLegacyGamma ? EGammaSpace::Pow22 : EGammaSpace::sRGB) : EGammaSpace::Linear;
+		Writer.AddInteger("FormatVersion", TextureFormatVersion);
+	}
+
+	Writer.SetName("Build");
+	WriteBuildSettings(Writer, BuildSettings, TextureFormat);
+
+	Writer.SetName("Output");
+	WriteOutputSettings(Writer, NumInlineMips, KeySuffix);
 
-		Operator(TEXT("CompositeTextureSource"), WriteTextureSourceToCompactBinary(Texture.CompositeTexture->Source, CompositeTextureGammaSpace));
+	Writer.SetName("Source");
+	WriteSource(Writer, Texture, LayerIndex);
+
+	if (Texture.CompositeTexture)
+	{
+		Writer.SetName("CompositeSource");
+		WriteSource(Writer, *Texture.CompositeTexture, LayerIndex);
 	}
+
+	Writer.EndObject();
+	return Writer.Save().AsObject();
 }
 
 #endif // WITH_EDITOR
diff --git a/Engine/Source/Runtime/Engine/Private/TextureDerivedDataBuildUtils.h b/Engine/Source/Runtime/Engine/Private/TextureDerivedDataBuildUtils.h
--- a/Engine/Source/Runtime/Engine/Private/TextureDerivedDataBuildUtils.h
+++ b/Engine/Source/Runtime/Engine/Private/TextureDerivedDataBuildUtils.h
@@ -4,20 +4,14 @@
 
 #if WITH_EDITOR
 
-#include "CoreMinimal.h"
-#include "Containers/StringView.h"
-#include "Templates/Function.h"
-
-namespace UE::DerivedData { class FBuildActionBuilder; }
+#include "CoreTypes.h"
 
 class FCbObject;
+class FString;
 class UTexture;
 struct FTextureBuildSettings;
 
-using FTextureConstantOperator = TFunctionRef<void (FStringView Key, const FCbObject& Value)>;
-
-// Suitable for use with an action builder or in the future a definition builder
 FString GetTextureBuildFunctionName(const FTextureBuildSettings& BuildSettings);
-void ComposeTextureBuildFunctionConstants(const FString& KeySuffix, const UTexture& Texture, const FTextureBuildSettings& BuildSettings, int32 LayerIndex, int32 NumInlineMips, FTextureConstantOperator Operator);
+FCbObject SaveTextureBuildSettings(const FString& KeySuffix, const UTexture& Texture, const FTextureBuildSettings& BuildSettings, int32 LayerIndex, int32 NumInlineMips);
 
 #endif // WITH_EDITOR
diff --git a/Engine/Source/Runtime/Engine/Private/TextureDerivedDataTask.cpp b/Engine/Source/Runtime/Engine/Private/TextureDerivedDataTask.cpp
--- a/Engine/Source/Runtime/Engine/Private/TextureDerivedDataTask.cpp
+++ b/Engine/Source/Runtime/Engine/Private/TextureDerivedDataTask.cpp
@@ -445,11 +445,8 @@ void FTextureCacheDerivedDataWorker::BuildTexture(bool bReplaceExistingDDC)
 			FString TexturePath = Texture.GetPathName();
 			FBuildDefinitionBuilder DefinitionBuilder = Build.CreateDefinition(TexturePath, BuildFunctionName);
 
-			ComposeTextureBuildFunctionConstants(KeySuffix, Texture, BuildSettingsPerLayer[0], 0, NUM_INLINE_DERIVED_MIPS, 
-				[&DefinitionBuilder] (FStringView Key, const FCbObject& Value)
-				{
-					DefinitionBuilder.AddConstant(Key, Value);
-				});
+			DefinitionBuilder.AddConstant(TEXT("Settings"_SV),
+				SaveTextureBuildSettings(KeySuffix, Texture, BuildSettingsPerLayer[0], 0, NUM_INLINE_DERIVED_MIPS));
 
 			for (const FBuildInputRecord& InputRecord : BuildInputRecords)
 			{
