PR #6996: Add USkinnedMeshComponent::GetRefPoseTransform (Contributed by [REDACTED_PERSON])

diff --git a/Engine/Source/Runtime/Engine/Classes/Components/SkinnedMeshComponent.h b/Engine/Source/Runtime/Engine/Classes/Components/SkinnedMeshComponent.h
--- a/Engine/Source/Runtime/Engine/Classes/Components/SkinnedMeshComponent.h
+++ b/Engine/Source/Runtime/Engine/Classes/Components/SkinnedMeshComponent.h
@@ -1411,8 +1411,18 @@ public:
 	 *
 	 * @return Local space reference position 
 	 */
-	UFUNCTION(BlueprintCallable, Category = "Components|SkinnedMesh")
-	FVector GetRefPosePosition(int32 BoneIndex);
+	UFUNCTION(BlueprintPure, Category = "Components|SkinnedMesh")
+	FVector GetRefPosePosition(int32 BoneIndex) const;
+
+	/** 
+	 * Gets the local-space transform of a bone in the reference pose. 
+	 *
+	 * @param BoneIndex Index of the bone
+	 *
+	 * @return Local space reference transform 
+	 */
+	UFUNCTION(BlueprintPure, Category = "Components|SkinnedMesh")
+	FTransform GetRefPoseTransform(int32 BoneIndex) const;
 
 	/** finds a vector pointing along the given axis of the given bone
 	 *
diff --git a/Engine/Source/Runtime/Engine/Private/Components/SkinnedMeshComponent.cpp b/Engine/Source/Runtime/Engine/Private/Components/SkinnedMeshComponent.cpp
--- a/Engine/Source/Runtime/Engine/Private/Components/SkinnedMeshComponent.cpp
+++ b/Engine/Source/Runtime/Engine/Private/Components/SkinnedMeshComponent.cpp
@@ -1553,7 +1553,7 @@ bool USkinnedMeshComponent::BoneIsChildOf(FName BoneName, FName ParentBoneName)
 }
 
 
-FVector USkinnedMeshComponent::GetRefPosePosition(int32 BoneIndex)
+FVector USkinnedMeshComponent::GetRefPosePosition(int32 BoneIndex) const
 {
 	if(SkeletalMesh && (BoneIndex >= 0) && (BoneIndex < SkeletalMesh->GetRefSkeleton().GetNum()))
 	{
@@ -1565,6 +1565,17 @@ FVector USkinnedMeshComponent::GetRefPosePosition(int32 BoneIndex)
 	}
 }
 
+FTransform USkinnedMeshComponent::GetRefPoseTransform(int32 BoneIndex) const
+{
+	if(SkeletalMesh && (BoneIndex >= 0) && (BoneIndex < SkeletalMesh->GetRefSkeleton().GetNum()))
+	{
+		return SkeletalMesh->GetRefSkeleton().GetRefBonePose()[BoneIndex];
+	}
+	else
+	{
+		return FTransform::Identity;
+	}
+}
 
 void USkinnedMeshComponent::SetSkeletalMesh(USkeletalMesh* InSkelMesh, bool bReinitPose)
 {
