Material Instance changes should cause path tracer to restart accumulation

FMaterialRenderProxy doesn't appear to know which FScene it belongs to, so I had to go through the global render module and flag all scenes instead

diff --git a/Engine/Source/Runtime/Engine/Private/Materials/MaterialInstance.cpp b/Engine/Source/Runtime/Engine/Private/Materials/MaterialInstance.cpp
--- a/Engine/Source/Runtime/Engine/Private/Materials/MaterialInstance.cpp
+++ b/Engine/Source/Runtime/Engine/Private/Materials/MaterialInstance.cpp
@@ -415,6 +415,12 @@ void GameThread_UpdateMIParameter(const UMaterialInstance* Instance, const Param
 	}
 }
 
+void FMaterialInstanceResource::RenderThread_SignalParameterChange() const
+{
+	// a material parameter changed - path tracer accumulation should be reset
+	GetRendererModule().InvalidatePathTracedOutput();
+}
+
 #if WITH_EDITOR
 template<typename ParameterType>
 static void RemapLayerParameterIndicesArray(TArray<ParameterType>& Parameters, const TArray<int32>& RemapLayerIndices)
diff --git a/Engine/Source/Runtime/Engine/Private/Materials/MaterialInstanceSupport.h b/Engine/Source/Runtime/Engine/Private/Materials/MaterialInstanceSupport.h
--- a/Engine/Source/Runtime/Engine/Private/Materials/MaterialInstanceSupport.h
+++ b/Engine/Source/Runtime/Engine/Private/Materials/MaterialInstanceSupport.h
@@ -126,6 +126,8 @@ public:
 		LLM_SCOPE(ELLMTag::MaterialInstance);
 
 		InvalidateUniformExpressionCache(false);
+		RenderThread_SignalParameterChange();
+
 		TArray<TNamedParameter<ValueType> >& ValueArray = GetValueArray<ValueType>();
 		const int32 ParameterCount = ValueArray.Num();
 		for (int32 ParameterIndex = 0; ParameterIndex < ParameterCount; ++ParameterIndex)
@@ -169,6 +171,8 @@ private:
 	template <typename ValueType> TArray<TNamedParameter<ValueType> >& GetValueArray() { return ScalarParameterArray; }
 	template <typename ValueType> const TArray<TNamedParameter<ValueType> >& GetValueArray() const { return ScalarParameterArray; }
 
+	void RenderThread_SignalParameterChange() const;
+
 	/** The parent of the material instance. */
 	UMaterialInterface* Parent;
 
diff --git a/Engine/Source/Runtime/Engine/Public/SceneInterface.h b/Engine/Source/Runtime/Engine/Public/SceneInterface.h
--- a/Engine/Source/Runtime/Engine/Public/SceneInterface.h
+++ b/Engine/Source/Runtime/Engine/Public/SceneInterface.h
@@ -202,6 +202,9 @@ public:
 	/** Invalidates pages in a runtime virtual texture object. */
 	virtual void InvalidateRuntimeVirtualTexture(class URuntimeVirtualTextureComponent* Component, FBoxSphereBounds const& WorldBounds) {}
 
+	/** Mark scene as needing to restart path tracer accumulation. */
+	virtual void InvalidatePathTracedOutput() {}
+
 	/** 
 	 * Retrieves primitive uniform shader parameters that are internal to the renderer.
 	 */
diff --git a/Engine/Source/Runtime/RenderCore/Public/RendererInterface.h b/Engine/Source/Runtime/RenderCore/Public/RendererInterface.h
--- a/Engine/Source/Runtime/RenderCore/Public/RendererInterface.h
+++ b/Engine/Source/Runtime/RenderCore/Public/RendererInterface.h
@@ -834,5 +834,8 @@ public:
 	 * The intended use is for stand-alone rendering that involves Scene proxies (that then may need the machinery to render GPU-Scene aware primitives.
 	 */
 	virtual IScenePrimitiveRenderingContext* BeginScenePrimitiveRendering(FRDGBuilder& GraphBuilder, FSceneViewFamily* ViewFamily) = 0;
+
+	/** Mark all the current scenes as needing to restart path tracer accumulation */
+	virtual void InvalidatePathTracedOutput() = 0;
 };
 
diff --git a/Engine/Source/Runtime/Renderer/Private/RendererModule.h b/Engine/Source/Runtime/Renderer/Private/RendererModule.h
--- a/Engine/Source/Runtime/Renderer/Private/RendererModule.h
+++ b/Engine/Source/Runtime/Renderer/Private/RendererModule.h
@@ -120,6 +120,8 @@ public:
 
 	virtual IScenePrimitiveRenderingContext* BeginScenePrimitiveRendering(FRDGBuilder& GraphBuilder, FSceneViewFamily* ViewFamily) override;
 
+	virtual void InvalidatePathTracedOutput() override;
+
 private:
 	TSet<FSceneInterface*> AllocatedScenes;
 	FOnPostOpaqueRender PostOpaqueRenderDelegate;
diff --git a/Engine/Source/Runtime/Renderer/Private/RendererScene.cpp b/Engine/Source/Runtime/Renderer/Private/RendererScene.cpp
--- a/Engine/Source/Runtime/Renderer/Private/RendererScene.cpp
+++ b/Engine/Source/Runtime/Renderer/Private/RendererScene.cpp
@@ -2698,6 +2698,12 @@ void FScene::InvalidateRuntimeVirtualTexture(class URuntimeVirtualTextureCompone
 	}
 }
 
+void FScene::InvalidatePathTracedOutput()
+{
+	checkSlow(IsInRenderingThread());
+	bPathTracingNeedsInvalidation = true;
+}
+
 void FScene::FlushDirtyRuntimeVirtualTextures()
 {
 	checkSlow(IsInRenderingThread());
@@ -4888,6 +4894,14 @@ FSceneViewStateInterface* FRendererModule::AllocateViewState()
 	return new FSceneViewState();
 }
 
+void FRendererModule::InvalidatePathTracedOutput()
+{
+	for (FSceneInterface* Scene : AllocatedScenes)
+	{
+		Scene->InvalidatePathTracedOutput();
+	}
+}
+
 //////////////////////////////////////////////////////////////////////////
 
 FLatentGPUTimer::FLatentGPUTimer(FRenderQueryPoolRHIRef InTimerQueryPool, int32 InAvgSamples)
diff --git a/Engine/Source/Runtime/Renderer/Private/ScenePrivate.h b/Engine/Source/Runtime/Renderer/Private/ScenePrivate.h
--- a/Engine/Source/Runtime/Renderer/Private/ScenePrivate.h
+++ b/Engine/Source/Runtime/Renderer/Private/ScenePrivate.h
@@ -2979,6 +2979,7 @@ public:
 	virtual void RemoveRuntimeVirtualTexture(class URuntimeVirtualTextureComponent* Component) override;
 	virtual void GetRuntimeVirtualTextureHidePrimitiveMask(uint8& bHideMaskEditor, uint8& bHideMaskGame) const override;
 	virtual void InvalidateRuntimeVirtualTexture(class URuntimeVirtualTextureComponent* Component, FBoxSphereBounds const& WorldBounds) override;
+	virtual void InvalidatePathTracedOutput() override;
 	virtual void GetPrimitiveUniformShaderParameters_RenderThread(const FPrimitiveSceneInfo* PrimitiveSceneInfo, bool& bHasPrecomputedVolumetricLightmap, FMatrix& PreviousLocalToWorld, int32& SingleCaptureIndex, bool& bOutputVelocity) const override;
 	virtual void UpdateLightTransform(ULightComponent* Light) override;
 	virtual void UpdateLightColorAndBrightness(ULightComponent* Light) override;
