- Fix a crash caused by not setting a blend mode when a material is both Thin Translucent and After Motion Blur.
- Add material errors to the HLSL material translator to disallow attempting to use Thin Translucents or Modulate blend mode in the After Motion Blur translucency pass.

diff --git a/Engine/Source/Runtime/Engine/Private/Materials/HLSLMaterialTranslator.cpp b/Engine/Source/Runtime/Engine/Private/Materials/HLSLMaterialTranslator.cpp
--- a/Engine/Source/Runtime/Engine/Private/Materials/HLSLMaterialTranslator.cpp
+++ b/Engine/Source/Runtime/Engine/Private/Materials/HLSLMaterialTranslator.cpp
@@ -736,46 +736,46 @@ bool FHLSLMaterialTranslator::Translate()
 		
 		if (!bEnableExecutionFlow)
 		{
-		if (Domain == MD_Volume || (Domain == MD_Surface && IsSubsurfaceShadingModel(MaterialShadingModels)))
-		{
-			// Note we don't test for the blend mode as you can have a translucent material using the subsurface shading model
+			if (Domain == MD_Volume || (Domain == MD_Surface && IsSubsurfaceShadingModel(MaterialShadingModels)))
+			{
+				// Note we don't test for the blend mode as you can have a translucent material using the subsurface shading model
 
-			// another ForceCast as CompilePropertyAndSetMaterialProperty() can return MCT_Float which we don't want here
-			int32 SubsurfaceColor = Material->CompilePropertyAndSetMaterialProperty(MP_SubsurfaceColor, this);
-			SubsurfaceColor = ForceCast(SubsurfaceColor, FMaterialAttributeDefinitionMap::GetValueType(MP_SubsurfaceColor), MFCF_ExactMatch | MFCF_ReplicateValue);
+				// another ForceCast as CompilePropertyAndSetMaterialProperty() can return MCT_Float which we don't want here
+				int32 SubsurfaceColor = Material->CompilePropertyAndSetMaterialProperty(MP_SubsurfaceColor, this);
+				SubsurfaceColor = ForceCast(SubsurfaceColor, FMaterialAttributeDefinitionMap::GetValueType(MP_SubsurfaceColor), MFCF_ExactMatch | MFCF_ReplicateValue);
 
-			static FName NameSubsurfaceProfile(TEXT("__SubsurfaceProfile"));
+				static FName NameSubsurfaceProfile(TEXT("__SubsurfaceProfile"));
 
-			// 1.0f is is a not used profile - later this gets replaced with the actual profile
-			int32 CodeSubsurfaceProfile = ForceCast(ScalarParameter(NameSubsurfaceProfile, 1.0f), MCT_Float1);
+				// 1.0f is is a not used profile - later this gets replaced with the actual profile
+				int32 CodeSubsurfaceProfile = ForceCast(ScalarParameter(NameSubsurfaceProfile, 1.0f), MCT_Float1);
 
 				Chunk[MP_SubsurfaceColor] = AppendVector(SubsurfaceColor, CodeSubsurfaceProfile);
-		}
+			}
 
 			Chunk[MP_CustomData0] = Material->CompilePropertyAndSetMaterialProperty(MP_CustomData0, this);
 			Chunk[MP_CustomData1] = Material->CompilePropertyAndSetMaterialProperty(MP_CustomData1, this);
 			Chunk[MP_AmbientOcclusion] = Material->CompilePropertyAndSetMaterialProperty(MP_AmbientOcclusion, this);
 
-		if (IsTranslucentBlendMode(BlendMode) || MaterialShadingModels.HasShadingModel(MSM_SingleLayerWater))
-		{
-			int32 UserRefraction = ForceCast(Material->CompilePropertyAndSetMaterialProperty(MP_Refraction, this), MCT_Float1);
-			int32 RefractionDepthBias = ForceCast(ScalarParameter(FName(TEXT("RefractionDepthBias")), Material->GetRefractionDepthBiasValue()), MCT_Float1);
+			if (IsTranslucentBlendMode(BlendMode) || MaterialShadingModels.HasShadingModel(MSM_SingleLayerWater))
+			{
+				int32 UserRefraction = ForceCast(Material->CompilePropertyAndSetMaterialProperty(MP_Refraction, this), MCT_Float1);
+				int32 RefractionDepthBias = ForceCast(ScalarParameter(FName(TEXT("RefractionDepthBias")), Material->GetRefractionDepthBiasValue()), MCT_Float1);
 
-			Chunk[MP_Refraction] = AppendVector(UserRefraction, RefractionDepthBias);
-		}
+				Chunk[MP_Refraction] = AppendVector(UserRefraction, RefractionDepthBias);
+			}
 
-		if (bCompileForComputeShader)
-		{
+			if (bCompileForComputeShader)
+			{
 				Chunk[CompiledMP_EmissiveColorCS] = Material->CompilePropertyAndSetMaterialProperty(MP_EmissiveColor, this, SF_Compute);
-		}
+			}
 
-		if (Chunk[MP_WorldPositionOffset] != INDEX_NONE)
-		{
-			// Only calculate previous WPO if there is a current WPO
-			Chunk[CompiledMP_PrevWorldPositionOffset] = Material->CompilePropertyAndSetMaterialProperty(MP_WorldPositionOffset, this, SF_Vertex, true);
-		}
+			if (Chunk[MP_WorldPositionOffset] != INDEX_NONE)
+			{
+				// Only calculate previous WPO if there is a current WPO
+				Chunk[CompiledMP_PrevWorldPositionOffset] = Material->CompilePropertyAndSetMaterialProperty(MP_WorldPositionOffset, this, SF_Vertex, true);
+			}
 
-		Chunk[MP_PixelDepthOffset] = Material->CompilePropertyAndSetMaterialProperty(MP_PixelDepthOffset, this);
+			Chunk[MP_PixelDepthOffset] = Material->CompilePropertyAndSetMaterialProperty(MP_PixelDepthOffset, this);
 		}
 
 		ResourcesString = TEXT("");
@@ -854,17 +854,17 @@ bool FHLSLMaterialTranslator::Translate()
 
 		if (!bEnableExecutionFlow)
 		{
-		for (uint32 CustomUVIndex = MP_CustomizedUVs0; CustomUVIndex <= MP_CustomizedUVs7; CustomUVIndex++)
-		{
-			// Only compile custom UV inputs for UV channels requested by the pixel shader inputs
-			// Any unconnected inputs will have a texcoord generated for them in Material->CompileProperty, which will pass through the vertex (uncustomized) texture coordinates
-			// Note: this is using NumUserTexCoords, which is set by translating all the pixel properties above
-			if (CustomUVIndex - MP_CustomizedUVs0 < SavedNumUserTexCoords)
+			for (uint32 CustomUVIndex = MP_CustomizedUVs0; CustomUVIndex <= MP_CustomizedUVs7; CustomUVIndex++)
 			{
-				Chunk[CustomUVIndex] = Material->CompilePropertyAndSetMaterialProperty((EMaterialProperty)CustomUVIndex, this);
+				// Only compile custom UV inputs for UV channels requested by the pixel shader inputs
+				// Any unconnected inputs will have a texcoord generated for them in Material->CompileProperty, which will pass through the vertex (uncustomized) texture coordinates
+				// Note: this is using NumUserTexCoords, which is set by translating all the pixel properties above
+				if (CustomUVIndex - MP_CustomizedUVs0 < SavedNumUserTexCoords)
+				{
+					Chunk[CustomUVIndex] = Material->CompilePropertyAndSetMaterialProperty((EMaterialProperty)CustomUVIndex, this);
+				}
 			}
 		}
-		}
 
 		// Output the implementation for any custom expressions we will call below.
 		for (int32 ExpressionIndex = 0; ExpressionIndex < CustomExpressions.Num(); ExpressionIndex++)
@@ -922,7 +922,6 @@ bool FHLSLMaterialTranslator::Translate()
 				Errorf(TEXT("Translucency after DOF with BLEND_Modulate is only allowed on platforms that support dual-blending. Consider using BLEND_Translucent with black emissive"));
 			}
 		}
-		
 
 		// Don't allow opaque and masked materials to scene depth as the results are undefined
 		if (bUsesSceneDepth && Domain != MD_PostProcess && !IsTranslucentBlendMode(BlendMode))
@@ -1030,6 +1029,11 @@ bool FHLSLMaterialTranslator::Translate()
 			{
 				Errorf(TEXT("ThinTranslucent materials requires the use of ThinTranslucentMaterial output node."));
 			}
+			if (Material->IsTranslucencyAfterMotionBlurEnabled())
+			{
+				// We don't currently have a separate translucency modulation pass for After Motion Blur
+				Errorf(TEXT("ThinTranslucent materials are not currently supported in the \"After Motion Blur\" translucency pass."));
+			}
 		}
 
 		if (Material->IsStrataMaterial())
@@ -1064,6 +1068,12 @@ bool FHLSLMaterialTranslator::Translate()
 			}
 		}
 
+		if (BlendMode == BLEND_Modulate && Material->IsTranslucencyAfterMotionBlurEnabled())
+		{
+			// We don't currently have a separate translucency modulation pass for After Motion Blur
+			Errorf(TEXT("Blend Mode \"Modulate\" materials are not currently supported in the \"After Motion Blur\" translucency pass."));
+		}
+
 		// Catch any modifications to NumUserTexCoords that will not seen by customized UVs
 		check(SavedNumUserTexCoords == GetNumUserTexCoords());
 
diff --git a/Engine/Source/Runtime/Renderer/Private/BasePassRendering.cpp b/Engine/Source/Runtime/Renderer/Private/BasePassRendering.cpp
--- a/Engine/Source/Runtime/Renderer/Private/BasePassRendering.cpp
+++ b/Engine/Source/Runtime/Renderer/Private/BasePassRendering.cpp
@@ -173,6 +173,11 @@ void SetTranslucentRenderState(FMeshPassProcessorRenderState& DrawRenderState, c
 				// Alpha we will leave unchanged.
 				DrawRenderState.SetBlendState(TStaticBlendState<CW_RGBA, BO_Add, BF_Zero, BF_Source1Color, BO_Add, BF_Zero, BF_One>::GetRHI());
 			}
+			else if (InTranslucencyPassType == ETranslucencyPass::TPT_TranslucencyAfterMotionBlur)
+			{
+				// We don't actually currently support color modulation in the post-motion blur pass at the moment, so just do the same as post-DOF for now
+				DrawRenderState.SetBlendState(TStaticBlendState<CW_RGBA, BO_Add, BF_One, BF_Source1Color, BO_Add, BF_Zero, BF_One>::GetRHI());
+			}
 		}
 		else
 		{
