UMaterialExpressionPerInstanceCustomData3Vector material expression

diff --git a/Engine/Shaders/Private/MaterialTemplate.ush b/Engine/Shaders/Private/MaterialTemplate.ush
--- a/Engine/Shaders/Private/MaterialTemplate.ush
+++ b/Engine/Shaders/Private/MaterialTemplate.ush
@@ -959,6 +959,34 @@ float GetPerInstanceCustomData(FMaterialVertexParameters Parameters, int Index,
 #endif
 }
 
+// Per Instance Custom Data Getter (Vertex Shader Only)
+/** Get the per-instance custom data when instancing */
+MaterialFloat3 GetPerInstanceCustomData3Vector(FMaterialVertexParameters Parameters, int Index, MaterialFloat3 DefaultValue)
+{
+#if (USE_INSTANCING || USE_INSTANCE_CULLING) && USES_PER_INSTANCE_CUSTOM_DATA
+	if (InstanceVF.NumCustomDataFloats > 0 && Index + 2 < InstanceVF.NumCustomDataFloats)
+	{
+	#if USE_INSTANCE_CULLING
+		int InstanceDataIndex = Parameters.PrimitiveLocalInstanceIndex;
+	#else
+		int InstanceDataIndex = Parameters.InstanceId + Parameters.InstanceOffset;
+	#endif
+
+		int BufferStartIndex = InstanceDataIndex * InstanceVF.NumCustomDataFloats;
+		int FloatIndex = clamp(Index, 0, InstanceVF.NumCustomDataFloats - 1);
+		return float3(InstanceVF.InstanceCustomDataBuffer[BufferStartIndex + FloatIndex], 
+		              InstanceVF.InstanceCustomDataBuffer[BufferStartIndex + FloatIndex + 1],
+					  InstanceVF.InstanceCustomDataBuffer[BufferStartIndex + FloatIndex + 2]);
+	}
+	else
+	{
+		return DefaultValue;
+	}
+#else
+	return DefaultValue;
+#endif
+}
+
 /** Transforms a vector from tangent space to world space, prescaling by an amount calculated previously */
 MaterialFloat3 TransformTangentVectorToWorld_PreScaled(FMaterialTessellationParameters Parameters, MaterialFloat3 InTangentVector)
 {
diff --git a/Engine/Source/Runtime/Engine/Classes/Materials/MaterialExpressionPerInstanceCustomData.h b/Engine/Source/Runtime/Engine/Classes/Materials/MaterialExpressionPerInstanceCustomData.h
--- a/Engine/Source/Runtime/Engine/Classes/Materials/MaterialExpressionPerInstanceCustomData.h
+++ b/Engine/Source/Runtime/Engine/Classes/Materials/MaterialExpressionPerInstanceCustomData.h
@@ -21,6 +21,28 @@ class UMaterialExpressionPerInstanceCustomData : public UMaterialExpression
 	UPROPERTY(EditAnywhere, Category = MaterialExpressionPerInstanceCustomData)
 	uint32 DataIndex;
 
+	//~ Begin UMaterialExpression Interface
+#if WITH_EDITOR
+	virtual int32 Compile(class FMaterialCompiler* Compiler, int32 OutputIndex) override;
+	virtual void GetCaption(TArray<FString>& OutCaptions) const override;
+#endif
+	//~ End UMaterialExpression Interface
+};
+
+UCLASS(collapsecategories, hidecategories=Object)
+class UMaterialExpressionPerInstanceCustomData3Vector : public UMaterialExpression
+{
+	GENERATED_UCLASS_BODY()
+
+	UPROPERTY(meta = (RequiredInput = "false", ToolTip = "Defaults to 'ConstDefaultValue' if not specified; Default value is used when no instances provided."))
+	FExpressionInput DefaultValue;
+
+	UPROPERTY(EditAnywhere, Category = MaterialExpressionPerInstanceCustomData, meta = (OverridingInputProperty = "DefaultValue"))
+	FLinearColor ConstDefaultValue;
+
+	UPROPERTY(EditAnywhere, Category = MaterialExpressionPerInstanceCustomData)
+	uint32 DataIndex;
+
 	//~ Begin UMaterialExpression Interface
 #if WITH_EDITOR
 	virtual int32 Compile(class FMaterialCompiler* Compiler, int32 OutputIndex) override;
diff --git a/Engine/Source/Runtime/Engine/Private/Materials/HLSLMaterialTranslator.cpp b/Engine/Source/Runtime/Engine/Private/Materials/HLSLMaterialTranslator.cpp
--- a/Engine/Source/Runtime/Engine/Private/Materials/HLSLMaterialTranslator.cpp
+++ b/Engine/Source/Runtime/Engine/Private/Materials/HLSLMaterialTranslator.cpp
@@ -9698,6 +9698,25 @@ int32 FHLSLMaterialTranslator::PerInstanceCustomData(int32 DataIndex, int32 Defa
 	}
 }
 
+/**
+ *	Returns a custom data on a per-instance basis when instancing
+ *	@DataIndex - index in array that represents custom data
+ *
+ *	@return	Code index
+ */
+int32 FHLSLMaterialTranslator::PerInstanceCustomData3Vector(int32 DataIndex, int32 DefaultValueIndex)
+{
+	if (ShaderFrequency != SF_Vertex)
+	{
+		return NonVertexShaderExpressionError();
+	}
+	else
+	{
+		MaterialCompilationOutput.bUsesPerInstanceCustomData = true;
+		return AddInlinedCodeChunkZeroDeriv(MCT_Float3, TEXT("GetPerInstanceCustomData3Vector(Parameters, %d, %s)"), DataIndex, *GetParameterCode(DefaultValueIndex));
+	}
+}
+
 /**
 	* Returns a float2 texture coordinate after 2x2 transform and offset applied
 	*
diff --git a/Engine/Source/Runtime/Engine/Private/Materials/HLSLMaterialTranslator.h b/Engine/Source/Runtime/Engine/Private/Materials/HLSLMaterialTranslator.h
--- a/Engine/Source/Runtime/Engine/Private/Materials/HLSLMaterialTranslator.h
+++ b/Engine/Source/Runtime/Engine/Private/Materials/HLSLMaterialTranslator.h
@@ -1042,6 +1042,14 @@ protected:
 	 */
 	virtual int32 PerInstanceCustomData(int32 DataIndex, int32 DefaultValueIndex) override;
 
+	/**
+	 * Returns a custom data on a per-instance basis when instancing
+	 * @DataIndex - index in array that represents custom data
+	 *
+	 * @return	Code index
+	 */
+	virtual int32 PerInstanceCustomData3Vector(int32 DataIndex, int32 DefaultValueIndex) override;
+
 	/**
 	 * Returns a float2 texture coordinate after 2x2 transform and offset applied
 	 *
diff --git a/Engine/Source/Runtime/Engine/Private/Materials/MaterialExpressions.cpp b/Engine/Source/Runtime/Engine/Private/Materials/MaterialExpressions.cpp
--- a/Engine/Source/Runtime/Engine/Private/Materials/MaterialExpressions.cpp
+++ b/Engine/Source/Runtime/Engine/Private/Materials/MaterialExpressions.cpp
@@ -17372,6 +17372,45 @@ void UMaterialExpressionPerInstanceCustomData::GetCaption(TArray<FString>& OutCa
 }
 #endif // WITH_EDITOR
 
+///////////////////////////////////////////////////////////////////////////////
+// UMaterialExpressionPerInstanceCustomData3Vector
+///////////////////////////////////////////////////////////////////////////////
+UMaterialExpressionPerInstanceCustomData3Vector::UMaterialExpressionPerInstanceCustomData3Vector(const FObjectInitializer& ObjectInitializer)
+	: Super(ObjectInitializer)
+{
+#if WITH_EDITORONLY_DATA
+	// Structure to hold one-time initialization
+	struct FConstructorStatics
+	{
+		FText NAME_Custom;
+		FConstructorStatics()
+			: NAME_Custom(LOCTEXT("Custom", "Custom"))
+		{
+		}
+	};
+	static FConstructorStatics ConstructorStatics;
+
+	MenuCategories.Add(ConstructorStatics.NAME_Custom);
+
+	bShaderInputData = true;
+#endif
+}
+
+#if WITH_EDITOR
+
+int32 UMaterialExpressionPerInstanceCustomData3Vector::Compile(class FMaterialCompiler* Compiler, int32 OutputIndex)
+{
+	int32 DefaultArgument = DefaultValue.GetTracedInput().Expression ? DefaultValue.Compile(Compiler) : Compiler->Constant3(ConstDefaultValue.R, ConstDefaultValue.G, ConstDefaultValue.B);
+	return Compiler->PerInstanceCustomData3Vector(DataIndex, DefaultArgument);
+}
+
+void UMaterialExpressionPerInstanceCustomData3Vector::GetCaption(TArray<FString>& OutCaptions) const
+{
+	OutCaptions.Add(FString::Printf(TEXT("PerInstanceCustomData[%d, %d, %d]"), DataIndex, DataIndex + 1, DataIndex + 2));
+}
+
+#endif // WITH_EDITOR
+
 ///////////////////////////////////////////////////////////////////////////////
 // UMaterialExpressionAntialiasedTextureMask
 ///////////////////////////////////////////////////////////////////////////////
diff --git a/Engine/Source/Runtime/Engine/Public/MaterialCompiler.h b/Engine/Source/Runtime/Engine/Public/MaterialCompiler.h
--- a/Engine/Source/Runtime/Engine/Public/MaterialCompiler.h
+++ b/Engine/Source/Runtime/Engine/Public/MaterialCompiler.h
@@ -356,6 +356,7 @@ public:
 	virtual int32 PerInstanceRandom() = 0;
 	virtual int32 PerInstanceFadeAmount() = 0;
 	virtual int32 PerInstanceCustomData(int32 DataIndex, int32 DefaultValueIndex) = 0;
+	virtual int32 PerInstanceCustomData3Vector(int32 DataIndex, int32 DefaultValueIndex) = 0;
 	virtual int32 AntialiasedTextureMask(int32 Tex, int32 UV, float Threshold, uint8 Channel) = 0;
 	virtual int32 Sobol(int32 Cell, int32 Index, int32 Seed) = 0;
 	virtual int32 TemporalSobol(int32 Index, int32 Seed) = 0;
@@ -724,6 +725,7 @@ public:
 	virtual int32 PerInstanceRandom() override { return Compiler->PerInstanceRandom(); }
 	virtual int32 PerInstanceFadeAmount() override { return Compiler->PerInstanceFadeAmount(); }
 	virtual int32 PerInstanceCustomData(int32 DataIndex, int32 DefaultValueIndex) override { return Compiler->PerInstanceCustomData(DataIndex, DefaultValueIndex); }
+	virtual int32 PerInstanceCustomData3Vector(int32 DataIndex, int32 DefaultValueIndex) override { return Compiler->PerInstanceCustomData3Vector(DataIndex, DefaultValueIndex); }
 	virtual int32 DepthOfFieldFunction(int32 Depth, int32 FunctionValueIndex) override
 	{
 		return Compiler->DepthOfFieldFunction(Depth, FunctionValueIndex);
