Remove useless sleeps when nothing to process

#rnx

diff --git a/Engine/Source/Runtime/Engine/Private/DistanceFieldAtlas.cpp b/Engine/Source/Runtime/Engine/Private/DistanceFieldAtlas.cpp
--- a/Engine/Source/Runtime/Engine/Private/DistanceFieldAtlas.cpp
+++ b/Engine/Source/Runtime/Engine/Private/DistanceFieldAtlas.cpp
@@ -642,7 +642,7 @@ void FDistanceFieldAsyncQueue::BlockUntilBuildComplete(UStaticMesh* StaticMesh,
 void FDistanceFieldAsyncQueue::BlockUntilAllBuildsComplete()
 {
 	TRACE_CPUPROFILER_EVENT_SCOPE(FDistanceFieldAsyncQueue::BlockUntilAllBuildsComplete)
-	do 
+	while (true)
 	{
 #if WITH_EDITOR
 		FStaticMeshCompilingManager::Get().FinishAllCompilation();
@@ -657,9 +657,14 @@ void FDistanceFieldAsyncQueue::BlockUntilAllBuildsComplete()
 		}
 
 		ProcessAsyncTasks();
+
+		if (GetNumOutstandingTasks() <= 0)
+		{
+			break;
+		}
+
 		FPlatformProcess::Sleep(.01f);
 	} 
-	while (GetNumOutstandingTasks() > 0);
 }
 
 void FDistanceFieldAsyncQueue::Build(FAsyncDistanceFieldTask* Task, FQueuedThreadPool& BuildThreadPool)
diff --git a/Engine/Source/Runtime/Engine/Private/MeshCardRepresentation.cpp b/Engine/Source/Runtime/Engine/Private/MeshCardRepresentation.cpp
--- a/Engine/Source/Runtime/Engine/Private/MeshCardRepresentation.cpp
+++ b/Engine/Source/Runtime/Engine/Private/MeshCardRepresentation.cpp
@@ -499,7 +499,7 @@ void FCardRepresentationAsyncQueue::BlockUntilBuildComplete(UStaticMesh* StaticM
 void FCardRepresentationAsyncQueue::BlockUntilAllBuildsComplete()
 {
 	TRACE_CPUPROFILER_EVENT_SCOPE(FCardRepresentationAsyncQueue::BlockUntilAllBuildsComplete)
-	do 
+	while (true)
 	{
 #if WITH_EDITOR
 		FStaticMeshCompilingManager::Get().FinishAllCompilation();
@@ -514,9 +514,14 @@ void FCardRepresentationAsyncQueue::BlockUntilAllBuildsComplete()
 		}
 
 		ProcessAsyncTasks();
+
+		if (GetNumOutstandingTasks() <= 0)
+		{
+			break;
+		}
+
 		FPlatformProcess::Sleep(.01f);
 	} 
-	while (GetNumOutstandingTasks() > 0);
 }
 
 void FCardRepresentationAsyncQueue::Build(FAsyncCardRepresentationTask* Task, FQueuedThreadPool& BuildThreadPool)
