Added world partition streaming source automatic velocity computation. Currently only used in debug display but could be used by runtime streaming prioritization.

diff --git a/Engine/Source/Runtime/Engine/Private/WorldPartition/WorldPartitionStreamingPolicy.cpp b/Engine/Source/Runtime/Engine/Private/WorldPartition/WorldPartitionStreamingPolicy.cpp
--- a/Engine/Source/Runtime/Engine/Private/WorldPartition/WorldPartitionStreamingPolicy.cpp
+++ b/Engine/Source/Runtime/Engine/Private/WorldPartition/WorldPartitionStreamingPolicy.cpp
@@ -139,6 +139,28 @@ void UWorldPartitionStreamingPolicy::UpdateStreamingSources()
 			StreamingSources.Add(StreamingSource);
 		}
 	}
+
+	// Update streaming sources velocity
+	const float CurrentTime = WorldPartition->GetWorld()->GetTimeSeconds();
+	TSet<FName, DefaultKeyFuncs<FName>, TInlineSetAllocator<8>> ValidStreamingSources;
+	for (FWorldPartitionStreamingSource& StreamingSource : StreamingSources)
+	{
+		if (!StreamingSource.Name.IsNone())
+		{
+			ValidStreamingSources.Add(StreamingSource.Name);
+			FStreamingSourceVelocity& SourceVelocity = StreamingSourcesVelocity.FindOrAdd(StreamingSource.Name);
+			StreamingSource.Velocity = SourceVelocity.GetAverageVelocity(StreamingSource.Location, CurrentTime);
+		}
+	}
+
+	// Cleanup StreamingSourcesVelocity
+	for (auto It(StreamingSourcesVelocity.CreateIterator()); It; ++It)
+	{
+		if (!ValidStreamingSources.Contains(It.Key()))
+		{
+			It.RemoveCurrent();
+		}
+	}
 }
 
 #define WORLDPARTITION_LOG_UPDATESTREAMINGSTATE(Verbosity)\
@@ -574,4 +596,51 @@ void UWorldPartitionStreamingPolicy::DrawRuntimeHash3D()
 	}
 }
 
+/*
+ * FStreamingSourceVelocity Implementation
+ */
+
+FStreamingSourceVelocity::FStreamingSourceVelocity()
+: LastIndex(INDEX_NONE)
+, LastUpdateTime(-1.0)
+, VelocitiesHistorySum(0.f)
+{
+	VelocitiesHistory.SetNumZeroed(VELOCITY_HISTORY_SAMPLE_COUNT);
+}
+
+float FStreamingSourceVelocity::GetAverageVelocity(const FVector& NewPosition, const float CurrentTime)
+{
+	const float TeleportDistance = 100.f;
+	const float MaxDeltaSeconds = 5.f;
+	const bool bIsFirstCall = (LastIndex == INDEX_NONE);
+	const float DeltaSeconds = bIsFirstCall ? 0.f : (CurrentTime - LastUpdateTime);
+	const float Distance = bIsFirstCall ? 0.f : ((NewPosition - LastPosition) * 0.01f).Size();
+	if (bIsFirstCall)
+	{
+		LastIndex = 0;
+	}
+
+	ON_SCOPE_EXIT
+	{
+		LastUpdateTime = CurrentTime;
+		LastPosition = NewPosition;
+	};
+
+	// Handle invalid cases
+	if (bIsFirstCall || (DeltaSeconds <= 0.f) || (DeltaSeconds > MaxDeltaSeconds) || (Distance > TeleportDistance))
+	{
+		return 0.f;
+	}
+
+	// Compute velocity (m/s)
+	const float Velocity = Distance / DeltaSeconds;
+	// Update velocities history buffer and sum
+	LastIndex = (LastIndex + 1) % VELOCITY_HISTORY_SAMPLE_COUNT;
+	VelocitiesHistorySum = FMath::Max<float>(0.f, (VelocitiesHistorySum + Velocity - VelocitiesHistory[LastIndex]));
+	VelocitiesHistory[LastIndex] = Velocity;
+
+	// return average
+	return (VelocitiesHistorySum / VELOCITY_HISTORY_SAMPLE_COUNT);
+}
+
 #undef LOCTEXT_NAMESPACE
\ No newline at end of file
diff --git a/Engine/Source/Runtime/Engine/Private/WorldPartition/WorldPartitionSubsystem.cpp b/Engine/Source/Runtime/Engine/Private/WorldPartition/WorldPartitionSubsystem.cpp
--- a/Engine/Source/Runtime/Engine/Private/WorldPartition/WorldPartitionSubsystem.cpp
+++ b/Engine/Source/Runtime/Engine/Private/WorldPartition/WorldPartitionSubsystem.cpp
@@ -284,7 +284,7 @@ void UWorldPartitionSubsystem::Draw(UCanvas* Canvas, class APlayerController* PC
 			Pos = CurrentOffset + FVector2D(MaxTextWidth + 10, 0.f);
 			for (const FWorldPartitionStreamingSource& StreamingSource : *StreamingSources)
 			{
-				const FString Text = FString::Printf(TEXT("Priority: %d | Pos: %s | Rot: %s"), StreamingSource.Priority, *StreamingSource.Location.ToString(), *StreamingSource.Rotation.ToString());
+				const FString Text = FString::Printf(TEXT("Priority: %d | Pos: %s | Rot: %s | Vel: %3.2f m/s (%d mph)"), StreamingSource.Priority, *StreamingSource.Location.ToString(), *StreamingSource.Rotation.ToString(), StreamingSource.Velocity, (int32)(StreamingSource.Velocity*2.23694f));
 				FWorldPartitionDebugHelper::DrawText(Canvas, Text, GEngine->GetSmallFont(), FColor::White, Pos);
 			}
 			CurrentOffset.Y = Pos.Y;
diff --git a/Engine/Source/Runtime/Engine/Public/WorldPartition/WorldPartitionStreamingPolicy.h b/Engine/Source/Runtime/Engine/Public/WorldPartition/WorldPartitionStreamingPolicy.h
--- a/Engine/Source/Runtime/Engine/Public/WorldPartition/WorldPartitionStreamingPolicy.h
+++ b/Engine/Source/Runtime/Engine/Public/WorldPartition/WorldPartitionStreamingPolicy.h
@@ -18,6 +18,23 @@
 
 class UWorldPartition;
 
+/**
+ * Helper to compute streaming source velocity based on position history.
+ */
+struct FStreamingSourceVelocity
+{
+	FStreamingSourceVelocity();
+	float GetAverageVelocity(const FVector& NewPosition, const float CurrentTime);
+
+private:
+	enum { VELOCITY_HISTORY_SAMPLE_COUNT = 16 };
+	int32 LastIndex;
+	float LastUpdateTime;
+	FVector LastPosition;
+	float VelocitiesHistorySum;
+	TArray<float, TInlineAllocator<VELOCITY_HISTORY_SAMPLE_COUNT>> VelocitiesHistory;
+};
+
 UCLASS(Abstract, Within = WorldPartition)
 class UWorldPartitionStreamingPolicy : public UObject
 {
@@ -65,7 +82,10 @@ protected:
 	const UWorldPartition* WorldPartition;
 	TSet<const UWorldPartitionRuntimeCell*> LoadedCells;
 	TSet<const UWorldPartitionRuntimeCell*> ActivatedCells;
+
+	// Streaming Sources
 	TArray<FWorldPartitionStreamingSource> StreamingSources;
+	TMap<FName, FStreamingSourceVelocity> StreamingSourcesVelocity;
 
 	UWorldPartitionRuntimeHash::FStreamingSourceCells FrameActivateCells;
 	UWorldPartitionRuntimeHash::FStreamingSourceCells FrameLoadCells;
diff --git a/Engine/Source/Runtime/Engine/Public/WorldPartition/WorldPartitionStreamingSource.h b/Engine/Source/Runtime/Engine/Public/WorldPartition/WorldPartitionStreamingSource.h
--- a/Engine/Source/Runtime/Engine/Public/WorldPartition/WorldPartitionStreamingSource.h
+++ b/Engine/Source/Runtime/Engine/Public/WorldPartition/WorldPartitionStreamingSource.h
@@ -85,14 +85,16 @@ struct ENGINE_API FWorldPartitionStreamingSource
 {
 	FWorldPartitionStreamingSource()
 		: Priority(EStreamingSourcePriority::Default)
+		, Velocity(0.f)
 	{}
 
-	FWorldPartitionStreamingSource(FName InName, const FVector& InLocation, const FRotator& InRotation, EStreamingSourceTargetState InTargetState, EStreamingSourcePriority InPriority = EStreamingSourcePriority::Default)
+	FWorldPartitionStreamingSource(FName InName, const FVector& InLocation, const FRotator& InRotation, EStreamingSourceTargetState InTargetState, EStreamingSourcePriority InPriority = EStreamingSourcePriority::Default, float InVelocity = 0.f)
 		: Name(InName)
 		, Location(InLocation)
 		, Rotation(InRotation)
 		, TargetState(InTargetState)
 		, Priority(InPriority)
+		, Velocity(InVelocity)
 	{
 	}
 
@@ -103,6 +105,7 @@ struct ENGINE_API FWorldPartitionStreamingSource
 	FRotator Rotation;
 	EStreamingSourceTargetState TargetState;
 	EStreamingSourcePriority Priority;
+	float Velocity; // Computed automatically
 };
 
 /**
