Provided an interface to query element parent<->child hierarchy information and used it to unify selection normalization

Previously we relied on element implementations to know what their parent or child elements might be, and to deal with them correctly in the following cases:
 1) When "normalising" a selection to be safe for operations like a move or delete, eg) removing elements that are children of other selected elements.
 2) When deleting an element, also ensuring that any implicitly destroyed child elements were deselected, eg) deleting an ISM component that has selected static mesh instances.

This approach hurts the modularity of the elements themselves, as it requires that the element implementations have intrinsic knowledge of any other child (or parent) element types that may exist, and to deal with them accordingly when needed. As a concrete example, an ISM component may have child elements in the form of static mesh instance elements, however the generic component element type does not (and should not) directly know that static mesh instance elements exist.

To address this issue, we've added a new interface, UTypedElementHierarchyInterface, which can be used to provide information about the logical parent<->child hierarchy information of elements. This is implemented as follows for our current element types:
 - Actor:
   - GetParentElement returns an invalid element handle.
   - GetChildElements returns the element handles of any components on the actor.
 - Component:
   - GetParentElement returns its owner actor element handle.
   - GetChildElements returns nothing by default, but UActorComponent::GetComponentChildElements may be overridden to control this behavior.
     - eg) UInstancedStaticMeshComponent overrides it to return its static mesh instance element handles.
 - SMInstance:
   - GetParentElement returns its owner ISM component element handle.
   - GetChildElements returns nothing.

Now the problems listed above can be solved by using this interface instead of relying on the element implementations:
 1) Selection normalization is now handled by UTypedElementSelectionSet, replacing the previous duplicate implementations of UTypedElementCommonActions and UTypedElementViewportInteraction.
 2) UTypedElementSelectionSet will now walk and also update the selection state of child elements, when asked to via the FTypedElementSelectionOptions.

Breaking Changes:
 - UTypedElementCommonActions::GetElementsForAction and UTypedElementViewportInteraction::GetSelectedElementsToMove have been removed, in favor of using the selection normalization functions of UTypedElementSelectionSet.
   - Note: You may still favor using FLevelEditorViewportClient::GetElementsToManipulate, as it will have removed normalized elements that also cannot be moved by a gizmo.
 - UTypedElementCommonActions::DeleteElements and UTypedElementCommonActions::DuplicateElements have been removed, as these functions operated on an unnormalized selection and could be unsafe. UTypedElementCommonActions::DeleteNormalizedElements and UTypedElementCommonActions::DuplicateNormalizedElements should be used instead.
 - AGroupActor::ForEachActorInGroup and AGroupActor::ForEachMovableActorInGroup now take a second AGroupActor* argument in their callback.

diff --git a/Engine/Source/Editor/AssetPlacementEdMode/Private/Tools/PlacementEraseTool.cpp b/Engine/Source/Editor/AssetPlacementEdMode/Private/Tools/PlacementEraseTool.cpp
--- a/Engine/Source/Editor/AssetPlacementEdMode/Private/Tools/PlacementEraseTool.cpp
+++ b/Engine/Source/Editor/AssetPlacementEdMode/Private/Tools/PlacementEraseTool.cpp
@@ -60,6 +60,12 @@ void UPlacementModeEraseTool::OnTick(float DeltaTime)
 		return;
 	}
 
+	UTypedElementSelectionSet* SelectionSet = AssetEditorContext->GetMutableSelectionSet();
+	if (!SelectionSet)
+	{
+		return;
+	}
+
 	TGCObjectScopeGuard<UTypedElementList> ElementsToDelete(UTypedElementRegistry::GetInstance()->CreateElementList());
 
 	TArray<FTypedElementHandle> HitElements = GetElementsInBrushRadius();
@@ -92,10 +98,9 @@ void UPlacementModeEraseTool::OnTick(float DeltaTime)
 
 	if (ElementsToDelete.Get()->HasElements())
 	{
-		UTypedElementSelectionSet* SelectionSet = AssetEditorContext->GetMutableSelectionSet();
-		if (SelectionSet)
-		{
-			ElementCommonActions->DeleteElementsInList(ElementsToDelete.Get(), AssetEditorContext->GetEditingWorld(), SelectionSet, FTypedElementDeletionOptions());
-		}
+		TGCObjectScopeGuard<UTypedElementList> NormalizedElementsToDelete(SelectionSet->GetNormalizedElementList(ElementsToDelete.Get(), FTypedElementSelectionNormalizationOptions()));
+		ElementCommonActions->DeleteNormalizedElements(NormalizedElementsToDelete.Get(), AssetEditorContext->GetEditingWorld(), SelectionSet, FTypedElementDeletionOptions());
+		NormalizedElementsToDelete.Get()->Empty();
+		ElementsToDelete.Get()->Empty();
 	}
 }
diff --git a/Engine/Source/Editor/EditorFramework/Private/Elements/Framework/TypedElementViewportInteraction.cpp b/Engine/Source/Editor/EditorFramework/Private/Elements/Framework/TypedElementViewportInteraction.cpp
--- a/Engine/Source/Editor/EditorFramework/Private/Elements/Framework/TypedElementViewportInteraction.cpp
+++ b/Engine/Source/Editor/EditorFramework/Private/Elements/Framework/TypedElementViewportInteraction.cpp
@@ -5,11 +5,6 @@
 #include "Elements/Framework/TypedElementRegistry.h"
 #include "Elements/Framework/TypedElementUtil.h"
 
-void FTypedElementViewportInteractionCustomization::GetElementsToMove(const TTypedElement<UTypedElementWorldInterface>& InElementWorldHandle, const ETypedElementViewportInteractionWorldType InWorldType, const UTypedElementSelectionSet* InSelectionSet, UTypedElementList* OutElementsToMove)
-{
-	OutElementsToMove->Add(InElementWorldHandle);
-}
-
 bool FTypedElementViewportInteractionCustomization::GetGizmoPivotLocation(const TTypedElement<UTypedElementWorldInterface>& InElementWorldHandle, const UE::Widget::EWidgetMode InWidgetMode, FVector& OutPivotLocation)
 {
 	FTransform ElementWorldTransform;
@@ -136,22 +131,6 @@ void FTypedElementViewportInteractionCustomization::MirrorElement(const TTypedEl
 }
 
 
-void UTypedElementViewportInteraction::GetSelectedElementsToMove(const UTypedElementSelectionSet* InSelectionSet, const ETypedElementViewportInteractionWorldType InWorldType, UTypedElementList* OutElementsToMove) const
-{
-	OutElementsToMove->Reset();
-
-	InSelectionSet->ForEachSelectedElement<UTypedElementWorldInterface>([this, InSelectionSet, InWorldType, OutElementsToMove](const TTypedElement<UTypedElementWorldInterface>& InElementWorldHandle)
-	{
-		if (InElementWorldHandle.CanEditElement())
-		{
-			FTypedElementViewportInteractionElement ViewportInteractionElement(InElementWorldHandle, GetInterfaceCustomizationByTypeId(InElementWorldHandle.GetId().GetTypeId()));
-			check(ViewportInteractionElement.IsSet());
-			ViewportInteractionElement.GetElementsToMove(InWorldType, InSelectionSet, OutElementsToMove);
-		}
-		return true;
-	});
-}
-
 void UTypedElementViewportInteraction::BeginGizmoManipulation(const UTypedElementList* InElementsToMove, const UE::Widget::EWidgetMode InWidgetMode)
 {
 	{
diff --git a/Engine/Source/Editor/EditorFramework/Public/Elements/Framework/TypedElementViewportInteraction.h b/Engine/Source/Editor/EditorFramework/Public/Elements/Framework/TypedElementViewportInteraction.h
--- a/Engine/Source/Editor/EditorFramework/Public/Elements/Framework/TypedElementViewportInteraction.h
+++ b/Engine/Source/Editor/EditorFramework/Public/Elements/Framework/TypedElementViewportInteraction.h
@@ -11,12 +11,6 @@
 class UTypedElementList;
 class UTypedElementSelectionSet;
 
-enum class ETypedElementViewportInteractionWorldType : uint8
-{
-	Editor,
-	PlayInEditor,
-};
-
 enum class ETypedElementViewportInteractionGizmoManipulationType : uint8
 {
 	/** The gizmo manipulation resulted in a drag operation */
@@ -34,7 +28,6 @@ public:
 	virtual ~FTypedElementViewportInteractionCustomization() = default;
 
 	//~ See UTypedElementViewportInteraction for API docs
-	virtual void GetElementsToMove(const TTypedElement<UTypedElementWorldInterface>& InElementWorldHandle, const ETypedElementViewportInteractionWorldType InWorldType, const UTypedElementSelectionSet* InSelectionSet, UTypedElementList* OutElementsToMove);
 	virtual bool GetGizmoPivotLocation(const TTypedElement<UTypedElementWorldInterface>& InElementWorldHandle, const UE::Widget::EWidgetMode InWidgetMode, FVector& OutPivotLocation);
 	virtual void PreGizmoManipulationStarted(TArrayView<const FTypedElementHandle> InElementHandles, const UE::Widget::EWidgetMode InWidgetMode);
 	virtual void GizmoManipulationStarted(const TTypedElement<UTypedElementWorldInterface>& InElementWorldHandle, const UE::Widget::EWidgetMode InWidgetMode);
@@ -76,7 +69,6 @@ public:
 	}
 
 	//~ See UTypedElementViewportInteraction for API docs
-	void GetElementsToMove(const ETypedElementViewportInteractionWorldType InWorldType, const UTypedElementSelectionSet* InSelectionSet, UTypedElementList* OutElementsToMove) { ViewportInteractionCustomization->GetElementsToMove(ElementWorldHandle, InWorldType, InSelectionSet, OutElementsToMove); }
 	bool GetGizmoPivotLocation(const UE::Widget::EWidgetMode InWidgetMode, FVector& OutPivotLocation) const { return ViewportInteractionCustomization->GetGizmoPivotLocation(ElementWorldHandle, InWidgetMode, OutPivotLocation); }
 	void GizmoManipulationStarted(const UE::Widget::EWidgetMode InWidgetMode) const { ViewportInteractionCustomization->GizmoManipulationStarted(ElementWorldHandle, InWidgetMode); }
 	void GizmoManipulationDeltaUpdate(const UE::Widget::EWidgetMode InWidgetMode, const EAxisList::Type InDragAxis, const FInputDeviceState& InInputState, const FTransform& InDeltaTransform, const FVector& InPivotLocation) const { ViewportInteractionCustomization->GizmoManipulationDeltaUpdate(ElementWorldHandle, InWidgetMode, InDragAxis, InInputState, InDeltaTransform, InPivotLocation); }
@@ -98,11 +90,6 @@ class EDITORFRAMEWORK_API UTypedElementViewportInteraction : public UObject, pub
 	GENERATED_BODY()
 
 public:
-	/**
-	 * Get the elements from the given selection set that can be moved (eg, by a gizmo).
-	 */
-	void GetSelectedElementsToMove(const UTypedElementSelectionSet* InSelectionSet, const ETypedElementViewportInteractionWorldType InWorldType, UTypedElementList* OutElementsToMove) const;
-
 	/**
 	 * Notify that the gizmo is potentially about to start manipulating the transform of the given set of elements (calculated from calling GetSelectedElementsToMove).
 	 */
diff --git a/Engine/Source/Editor/LevelEditor/Private/Elements/Actor/ActorElementLevelEditorCommonActionsCustomization.cpp b/Engine/Source/Editor/LevelEditor/Private/Elements/Actor/ActorElementLevelEditorCommonActionsCustomization.cpp
--- a/Engine/Source/Editor/LevelEditor/Private/Elements/Actor/ActorElementLevelEditorCommonActionsCustomization.cpp
+++ b/Engine/Source/Editor/LevelEditor/Private/Elements/Actor/ActorElementLevelEditorCommonActionsCustomization.cpp
@@ -9,19 +9,6 @@
 #include "EditorModeManager.h"
 #include "Toolkits/IToolkitHost.h"
 
-void FActorElementLevelEditorCommonActionsCustomization::GetElementsForAction(const TTypedElement<UTypedElementWorldInterface>& InElementWorldHandle, const UTypedElementList* InElementList, UTypedElementList* OutElementsForAction)
-{
-	AActor* Actor = ActorElementDataUtil::GetActorFromHandleChecked(InElementWorldHandle);
-
-	if (InElementList->HasElementsOfType(NAME_Components))
-	{
-		// If we have components selected then we will perform the action on those rather than the actors
-		return;
-	}
-
-	FTypedElementCommonActionsCustomization::GetElementsForAction(InElementWorldHandle, InElementList, OutElementsForAction);
-}
-
 bool FActorElementLevelEditorCommonActionsCustomization::DeleteElements(UTypedElementWorldInterface* InWorldInterface, TArrayView<const FTypedElementHandle> InElementHandles, UWorld* InWorld, UTypedElementSelectionSet* InSelectionSet, const FTypedElementDeletionOptions& InDeletionOptions)
 {
 	if (const IToolkitHost* ToolkitHostPtr = GetToolkitHost())
diff --git a/Engine/Source/Editor/LevelEditor/Private/Elements/Actor/ActorElementLevelEditorCommonActionsCustomization.h b/Engine/Source/Editor/LevelEditor/Private/Elements/Actor/ActorElementLevelEditorCommonActionsCustomization.h
--- a/Engine/Source/Editor/LevelEditor/Private/Elements/Actor/ActorElementLevelEditorCommonActionsCustomization.h
+++ b/Engine/Source/Editor/LevelEditor/Private/Elements/Actor/ActorElementLevelEditorCommonActionsCustomization.h
@@ -8,7 +8,6 @@
 class FActorElementLevelEditorCommonActionsCustomization : public FTypedElementCommonActionsCustomization, public FTypedElementAssetEditorToolkitHostMixin
 {
 public:
-	virtual void GetElementsForAction(const TTypedElement<UTypedElementWorldInterface>& InElementWorldHandle, const UTypedElementList* InElementList, UTypedElementList* OutElementsForAction) override;
 	virtual bool DeleteElements(UTypedElementWorldInterface* InWorldInterface, TArrayView<const FTypedElementHandle> InElementHandles, UWorld* InWorld, UTypedElementSelectionSet* InSelectionSet, const FTypedElementDeletionOptions& InDeletionOptions) override;
 	virtual void DuplicateElements(UTypedElementWorldInterface* InWorldInterface, TArrayView<const FTypedElementHandle> InElementHandles, UWorld* InWorld, const FVector& InLocationOffset, TArray<FTypedElementHandle>& OutNewElements) override;
 };
diff --git a/Engine/Source/Editor/LevelEditor/Private/Elements/Actor/ActorElementLevelEditorSelectionCustomization.cpp b/Engine/Source/Editor/LevelEditor/Private/Elements/Actor/ActorElementLevelEditorSelectionCustomization.cpp
--- a/Engine/Source/Editor/LevelEditor/Private/Elements/Actor/ActorElementLevelEditorSelectionCustomization.cpp
+++ b/Engine/Source/Editor/LevelEditor/Private/Elements/Actor/ActorElementLevelEditorSelectionCustomization.cpp
@@ -59,6 +59,20 @@ FTypedElementHandle FActorElementLevelEditorSelectionCustomization::GetSelection
 	return InElementSelectionHandle;
 }
 
+void FActorElementLevelEditorSelectionCustomization::GetNormalizedElements(const TTypedElement<UTypedElementSelectionInterface>& InElementSelectionHandle, const UTypedElementList* InSelectionSet, const FTypedElementSelectionNormalizationOptions& InNormalizationOptions, UTypedElementList* OutNormalizedElements)
+{
+	AActor* Actor = ActorElementDataUtil::GetActorFromHandleChecked(InElementSelectionHandle);
+
+	if (InSelectionSet->HasElementsOfType(NAME_Components))
+	{
+		// If we have components selected then we will use those rather than the actors
+		// The component may still choose to use its owner actor rather than itself
+		return;
+	}
+
+	FActorElementLevelEditorSelectionCustomization::AppendNormalizedActors(Actor, InSelectionSet, InNormalizationOptions, OutNormalizedElements);
+}
+
 bool FActorElementLevelEditorSelectionCustomization::CanSelectActorElement(const TTypedElement<UTypedElementSelectionInterface>& InActorSelectionHandle, const FTypedElementSelectionOptions& InSelectionOptions) const
 {
 	AActor* Actor = ActorElementDataUtil::GetActorFromHandleChecked(InActorSelectionHandle);
@@ -333,3 +347,52 @@ bool FActorElementLevelEditorSelectionCustomization::DeselectActorGroup(AGroupAc
 
 	return bSelectionChanged;
 }
+
+void FActorElementLevelEditorSelectionCustomization::AppendNormalizedActors(AActor* InActor, const UTypedElementList* InSelectionSet, const FTypedElementSelectionNormalizationOptions& InNormalizationOptions, UTypedElementList* OutNormalizedElements)
+{
+	if (InNormalizationOptions.FollowAttachment())
+	{
+		// Ensure that only parent-most actors are included if we were asked to follow attachments
+		for (AActor* Parent = InActor->GetAttachParentActor(); Parent; Parent = Parent->GetAttachParentActor())
+		{
+			FTypedElementHandle ParentHandle = UEngineElementsLibrary::AcquireEditorActorElementHandle(Parent, /*bAllowCreate*/false);
+			if (ParentHandle && InSelectionSet->Contains(ParentHandle))
+			{
+				return;
+			}
+		}
+	}
+
+	if (InNormalizationOptions.ExpandGroups())
+	{
+		AGroupActor* ParentGroup = AGroupActor::GetRootForActor(InActor, true, true);
+		if (ParentGroup && UActorGroupingUtils::IsGroupingActive())
+		{
+			// Skip if the group is already in the normalized list, since this logic will have already run
+			FTypedElementHandle ParentGroupElementHandle = UEngineElementsLibrary::AcquireEditorActorElementHandle(ParentGroup);
+			if (ParentGroupElementHandle && !OutNormalizedElements->Contains(ParentGroupElementHandle))
+			{
+				ParentGroup->ForEachActorInGroup([InSelectionSet, &InNormalizationOptions, OutNormalizedElements](AActor* InGroupedActor, AGroupActor* InGroupActor)
+				{
+					// Check that we've not got a parent attachment within the group/selection
+					if (InNormalizationOptions.FollowAttachment() && (GroupActorHelpers::ActorHasParentInGroup(InGroupedActor, InGroupActor) || GroupActorHelpers::ActorHasParentInSelection(InGroupedActor, InSelectionSet)))
+					{
+						return;
+					}
+
+					if (FTypedElementHandle GroupedActorElementHandle = UEngineElementsLibrary::AcquireEditorActorElementHandle(InGroupedActor))
+					{
+						OutNormalizedElements->Add(MoveTemp(GroupedActorElementHandle));
+					}
+				});
+
+				check(OutNormalizedElements->Contains(ParentGroupElementHandle));
+			}
+		}
+	}
+
+	if (FTypedElementHandle ActorElementHandle = UEngineElementsLibrary::AcquireEditorActorElementHandle(InActor))
+	{
+		OutNormalizedElements->Add(MoveTemp(ActorElementHandle));
+	}
+}
diff --git a/Engine/Source/Editor/LevelEditor/Private/Elements/Actor/ActorElementLevelEditorSelectionCustomization.h b/Engine/Source/Editor/LevelEditor/Private/Elements/Actor/ActorElementLevelEditorSelectionCustomization.h
--- a/Engine/Source/Editor/LevelEditor/Private/Elements/Actor/ActorElementLevelEditorSelectionCustomization.h
+++ b/Engine/Source/Editor/LevelEditor/Private/Elements/Actor/ActorElementLevelEditorSelectionCustomization.h
@@ -16,6 +16,7 @@ public:
 	virtual bool DeselectElement(const TTypedElement<UTypedElementSelectionInterface>& InElementSelectionHandle, UTypedElementList* InSelectionSet, const FTypedElementSelectionOptions& InSelectionOptions) override;
 	virtual bool AllowSelectionModifiers(const TTypedElement<UTypedElementSelectionInterface>& InElementSelectionHandle, const UTypedElementList* InSelectionSet) override;
 	virtual FTypedElementHandle GetSelectionElement(const TTypedElement<UTypedElementSelectionInterface>& InElementSelectionHandle, const UTypedElementList* InCurrentSelection, const ETypedElementSelectionMethod InSelectionMethod) override;
+	virtual void GetNormalizedElements(const TTypedElement<UTypedElementSelectionInterface>& InElementSelectionHandle, const UTypedElementList* InSelectionSet, const FTypedElementSelectionNormalizationOptions& InNormalizationOptions, UTypedElementList* OutNormalizedElements) override;
 
 	bool CanSelectActorElement(const TTypedElement<UTypedElementSelectionInterface>& InActorSelectionHandle, const FTypedElementSelectionOptions& InSelectionOptions) const;
 	bool CanDeselectActorElement(const TTypedElement<UTypedElementSelectionInterface>& InActorSelectionHandle, const FTypedElementSelectionOptions& InSelectionOptions) const;
@@ -25,4 +26,6 @@ public:
 
 	bool SelectActorGroup(AGroupActor* InGroupActor, UTypedElementList* InSelectionSet, const FTypedElementSelectionOptions& InSelectionOptions, const bool bForce);
 	bool DeselectActorGroup(AGroupActor* InGroupActor, UTypedElementList* InSelectionSet, const FTypedElementSelectionOptions& InSelectionOptions, const bool bForce);
+
+	static void AppendNormalizedActors(AActor* InActor, const UTypedElementList* InSelectionSet, const FTypedElementSelectionNormalizationOptions& InNormalizationOptions, UTypedElementList* OutNormalizedElements);
 };
diff --git a/Engine/Source/Editor/LevelEditor/Private/Elements/Actor/ActorElementLevelEditorViewportInteractionCustomization.cpp b/Engine/Source/Editor/LevelEditor/Private/Elements/Actor/ActorElementLevelEditorViewportInteractionCustomization.cpp
--- a/Engine/Source/Editor/LevelEditor/Private/Elements/Actor/ActorElementLevelEditorViewportInteractionCustomization.cpp
+++ b/Engine/Source/Editor/LevelEditor/Private/Elements/Actor/ActorElementLevelEditorViewportInteractionCustomization.cpp
@@ -14,42 +14,6 @@
 #include "ActorGroupingUtils.h"
 #include "LevelEditorViewport.h"
 
-void FActorElementLevelEditorViewportInteractionCustomization::GetElementsToMove(const TTypedElement<UTypedElementWorldInterface>& InElementWorldHandle, const ETypedElementViewportInteractionWorldType InWorldType, const UTypedElementSelectionSet* InSelectionSet, UTypedElementList* OutElementsToMove)
-{
-	AActor* Actor = ActorElementDataUtil::GetActorFromHandleChecked(InElementWorldHandle);
-
-	if (InSelectionSet->GetElementList()->HasElementsOfType(NAME_Components))
-	{
-		// If we have components selected then we will move those rather than the actors
-		// The component may still choose to move its owner actor rather than itself
-		return;
-	}
-
-	// Ensure that only parent-most actor(s) are moved
-	bool bParentAlsoSelected = false;
-	AActor* Parent = Actor->GetAttachParentActor();
-	while (Parent != nullptr)
-	{
-		if (Parent->IsSelected())
-		{
-			bParentAlsoSelected = true;
-			break;
-		}
-
-		Parent = Parent->GetAttachParentActor();
-	}
-	if (bParentAlsoSelected)
-	{
-		return;
-	}
-
-
-	if (CanMoveActorInViewport(Actor, InWorldType))
-	{
-		AppendActorsToMove(Actor, InSelectionSet, OutElementsToMove);
-	}
-}
-
 void FActorElementLevelEditorViewportInteractionCustomization::GizmoManipulationStarted(const TTypedElement<UTypedElementWorldInterface>& InElementWorldHandle, const UE::Widget::EWidgetMode InWidgetMode)
 {
 	AActor* Actor = ActorElementDataUtil::GetActorFromHandleChecked(InElementWorldHandle);
@@ -142,64 +106,3 @@ void FActorElementLevelEditorViewportInteractionCustomization::ModifyScale(AActo
 		}
 	}
 }
-
-bool FActorElementLevelEditorViewportInteractionCustomization::CanMoveActorInViewport(const AActor* InActor, const ETypedElementViewportInteractionWorldType InWorldType)
-{
-	if (!GEditor || !InActor)
-	{
-		return false;
-	}
-
-	// The actor cannot be location locked
-	if (InActor->IsLockLocation())
-	{
-		return false;
-	}
-
-	// If the actor has a root component, but it cannot be moved, then the actor cannot move.
-	const bool bIsPlayInEditorWorld = (InWorldType == ETypedElementViewportInteractionWorldType::PlayInEditor);
-	if (bIsPlayInEditorWorld && InActor->GetRootComponent() && !InActor->IsRootComponentMovable())
-	{
-		return false;
-	}
-
-	// The actor needs to be in the current viewport world
-	if (GEditor->PlayWorld)
-	{
-		const UWorld* CurrentWorld = InActor->GetWorld();
-		const UWorld* RequiredWorld = bIsPlayInEditorWorld ? GEditor->PlayWorld : GEditor->EditorWorld;
-		if (CurrentWorld != RequiredWorld)
-		{
-			return false;
-		}
-	}
-
-	return true;
-}
-
-void FActorElementLevelEditorViewportInteractionCustomization::AppendActorsToMove(AActor* InActor, const UTypedElementSelectionSet* InSelectionSet, UTypedElementList* OutElementsToMove)
-{
-	auto AddActorElement = [OutElementsToMove](AActor* InActorToAdd)
-	{
-		if(FTypedElementHandle ActorElementHandle = UEngineElementsLibrary::AcquireEditorActorElementHandle(InActorToAdd))
-		{
-			OutElementsToMove->Add(MoveTemp(ActorElementHandle));
-		}
-	};
-
-	AGroupActor* ParentGroup = AGroupActor::GetRootForActor(InActor, true, true);
-	if (ParentGroup && UActorGroupingUtils::IsGroupingActive())
-	{
-		// Skip if the group is already in the list of things to move, since this logic will have already run
-		FTypedElementHandle ParentGroupElementHandle = UEngineElementsLibrary::AcquireEditorActorElementHandle(ParentGroup);
-		if (ParentGroupElementHandle && !OutElementsToMove->Contains(ParentGroupElementHandle))
-		{
-			ParentGroup->ForEachMovableActorInGroup(InSelectionSet, AddActorElement);
-			check(OutElementsToMove->Contains(ParentGroupElementHandle));
-		}
-	}
-	else
-	{
-		AddActorElement(InActor);
-	}
-}
diff --git a/Engine/Source/Editor/LevelEditor/Private/Elements/Component/ComponentElementLevelEditorSelectionCustomization.cpp b/Engine/Source/Editor/LevelEditor/Private/Elements/Component/ComponentElementLevelEditorSelectionCustomization.cpp
--- a/Engine/Source/Editor/LevelEditor/Private/Elements/Component/ComponentElementLevelEditorSelectionCustomization.cpp
+++ b/Engine/Source/Editor/LevelEditor/Private/Elements/Component/ComponentElementLevelEditorSelectionCustomization.cpp
@@ -2,9 +2,11 @@
 
 #include "Elements/Component/ComponentElementLevelEditorSelectionCustomization.h"
 #include "Elements/Component/ComponentElementData.h"
-#include "GameFramework/Actor.h"
 #include "Components/ActorComponent.h"
 
+#include "Elements/Actor/ActorElementLevelEditorSelectionCustomization.h"
+#include "GameFramework/Actor.h"
+
 #include "Elements/Framework/TypedElementList.h"
 #include "Elements/Framework/EngineElementsLibrary.h"
 
@@ -81,6 +83,49 @@ FTypedElementHandle FComponentElementLevelEditorSelectionCustomization::GetSelec
 	return InElementSelectionHandle;
 }
 
+void FComponentElementLevelEditorSelectionCustomization::GetNormalizedElements(const TTypedElement<UTypedElementSelectionInterface>& InElementSelectionHandle, const UTypedElementList* InSelectionSet, const FTypedElementSelectionNormalizationOptions& InNormalizationOptions, UTypedElementList* OutNormalizedElements)
+{
+	UActorComponent* Component = ComponentElementDataUtil::GetComponentFromHandleChecked(InElementSelectionHandle);
+
+	if (USceneComponent* SceneComponent = Cast<USceneComponent>(Component))
+	{
+		if (AActor* ComponentOwner = SceneComponent->GetOwner())
+		{
+#if DO_CHECK
+			{
+				FTypedElementHandle ComponentOwnerElementHandle = UEngineElementsLibrary::AcquireEditorActorElementHandle(ComponentOwner, /*bAllowCreate*/false);
+				const bool bIsOwnerSelected = ComponentOwnerElementHandle && InSelectionSet->Contains(ComponentOwnerElementHandle);
+				ensureMsgf(bIsOwnerSelected, TEXT("Owner(%s) of %s is not selected"), *ComponentOwner->GetFullName(), *Component->GetFullName());
+			}
+#endif	// DO_CHECK
+
+			if (InNormalizationOptions.FollowAttachment())
+			{
+				if (ComponentOwner->GetRootComponent() == SceneComponent)
+				{
+					// If it is a root component, use the parent actor instead
+					FActorElementLevelEditorSelectionCustomization::AppendNormalizedActors(ComponentOwner, InSelectionSet, InNormalizationOptions, OutNormalizedElements);
+					return;
+				}
+				else
+				{
+					// Check to see if any parent is selected
+					for (USceneComponent* Parent = SceneComponent->GetAttachParent(); Parent; Parent = Parent->GetAttachParent())
+					{
+						FTypedElementHandle ParentElementHandle = UEngineElementsLibrary::AcquireEditorComponentElementHandle(Parent, /*bAllowCreate*/false);
+						if (ParentElementHandle && InSelectionSet->Contains(ParentElementHandle))
+						{
+							return;
+						}
+					}
+				}
+			}
+		}
+	}
+
+	OutNormalizedElements->Add(InElementSelectionHandle);
+}
+
 bool FComponentElementLevelEditorSelectionCustomization::CanSelectComponentElement(const TTypedElement<UTypedElementSelectionInterface>& InComponentSelectionHandle, const FTypedElementSelectionOptions& InSelectionOptions) const
 {
 	const UActorComponent* Component = ComponentElementDataUtil::GetComponentFromHandleChecked(InComponentSelectionHandle);
diff --git a/Engine/Source/Editor/LevelEditor/Private/Elements/Component/ComponentElementLevelEditorSelectionCustomization.h b/Engine/Source/Editor/LevelEditor/Private/Elements/Component/ComponentElementLevelEditorSelectionCustomization.h
--- a/Engine/Source/Editor/LevelEditor/Private/Elements/Component/ComponentElementLevelEditorSelectionCustomization.h
+++ b/Engine/Source/Editor/LevelEditor/Private/Elements/Component/ComponentElementLevelEditorSelectionCustomization.h
@@ -13,6 +13,7 @@ public:
 	virtual bool SelectElement(const TTypedElement<UTypedElementSelectionInterface>& InElementSelectionHandle, UTypedElementList* InSelectionSet, const FTypedElementSelectionOptions& InSelectionOptions) override;
 	virtual bool DeselectElement(const TTypedElement<UTypedElementSelectionInterface>& InElementSelectionHandle, UTypedElementList* InSelectionSet, const FTypedElementSelectionOptions& InSelectionOptions) override;
 	virtual FTypedElementHandle GetSelectionElement(const TTypedElement<UTypedElementSelectionInterface>& InElementSelectionHandle, const UTypedElementList* InCurrentSelection, const ETypedElementSelectionMethod InSelectionMethod) override;
+	virtual void GetNormalizedElements(const TTypedElement<UTypedElementSelectionInterface>& InElementSelectionHandle, const UTypedElementList* InSelectionSet, const FTypedElementSelectionNormalizationOptions& InNormalizationOptions, UTypedElementList* OutNormalizedElements) override;
 
 	bool CanSelectComponentElement(const TTypedElement<UTypedElementSelectionInterface>& InComponentSelectionHandle, const FTypedElementSelectionOptions& InSelectionOptions) const;
 	bool CanDeselectComponentElement(const TTypedElement<UTypedElementSelectionInterface>& InComponentSelectionHandle, const FTypedElementSelectionOptions& InSelectionOptions) const;
diff --git a/Engine/Source/Editor/LevelEditor/Private/Elements/Component/ComponentElementLevelEditorViewportInteractionCustomization.cpp b/Engine/Source/Editor/LevelEditor/Private/Elements/Component/ComponentElementLevelEditorViewportInteractionCustomization.cpp
--- a/Engine/Source/Editor/LevelEditor/Private/Elements/Component/ComponentElementLevelEditorViewportInteractionCustomization.cpp
+++ b/Engine/Source/Editor/LevelEditor/Private/Elements/Component/ComponentElementLevelEditorViewportInteractionCustomization.cpp
@@ -12,51 +12,6 @@
 #include "LevelEditorViewport.h"
 #include "Kismet2/ComponentEditorUtils.h"
 
-void FComponentElementLevelEditorViewportInteractionCustomization::GetElementsToMove(const TTypedElement<UTypedElementWorldInterface>& InElementWorldHandle, const ETypedElementViewportInteractionWorldType InWorldType, const UTypedElementSelectionSet* InSelectionSet, UTypedElementList* OutElementsToMove)
-{
-	UActorComponent* Component = ComponentElementDataUtil::GetComponentFromHandleChecked(InElementWorldHandle);
-
-	if (USceneComponent* SceneComponent = Cast<USceneComponent>(Component))
-	{
-		if (AActor* ComponentOwner = SceneComponent->GetOwner())
-		{
-#if DO_CHECK
-			{
-				FTypedElementHandle ComponentOwnerElementHandle = UEngineElementsLibrary::AcquireEditorActorElementHandle(ComponentOwner, /*bAllowCreate*/false);
-				const bool bIsOwnerSelected = ComponentOwnerElementHandle && InSelectionSet->IsElementSelected(ComponentOwnerElementHandle, FTypedElementIsSelectedOptions());
-				ensureMsgf(bIsOwnerSelected, TEXT("Owner(%s) of %s is not selected"), *ComponentOwner->GetFullName(), *Component->GetFullName());
-			}
-#endif	// DO_CHECK
-
-			if (ComponentOwner->GetRootComponent() == SceneComponent)
-			{
-				if (FActorElementLevelEditorViewportInteractionCustomization::CanMoveActorInViewport(ComponentOwner, InWorldType))
-				{
-					// If it is a root component, use the parent actor instead
-					FActorElementLevelEditorViewportInteractionCustomization::AppendActorsToMove(ComponentOwner, InSelectionSet, OutElementsToMove);
-				}
-			}
-			else
-			{
-				// Check to see if any parent is selected
-				bool bHasSelectedParent = false;
-				for (USceneComponent* Parent = SceneComponent->GetAttachParent(); Parent && !bHasSelectedParent; Parent = Parent->GetAttachParent())
-				{
-					FTypedElementHandle ParentElementHandle = UEngineElementsLibrary::AcquireEditorComponentElementHandle(Parent, /*bAllowCreate*/false);
-					bHasSelectedParent = ParentElementHandle && InSelectionSet->IsElementSelected(ParentElementHandle, FTypedElementIsSelectedOptions());
-				}
-
-				const bool bCanMoveInEditorWorld = (InWorldType == ETypedElementViewportInteractionWorldType::Editor) || (SceneComponent->Mobility == EComponentMobility::Type::Movable);
-				if (!bHasSelectedParent && bCanMoveInEditorWorld)
-				{
-					// If no parent of this component is also in the selection set, move it
-					OutElementsToMove->Add(InElementWorldHandle);
-				}
-			}
-		}
-	}
-}
-
 void FComponentElementLevelEditorViewportInteractionCustomization::GizmoManipulationStarted(const TTypedElement<UTypedElementWorldInterface>& InElementWorldHandle, const UE::Widget::EWidgetMode InWidgetMode)
 {
 	UActorComponent* Component = ComponentElementDataUtil::GetComponentFromHandleChecked(InElementWorldHandle);
diff --git a/Engine/Source/Editor/LevelEditor/Public/Elements/Actor/ActorElementLevelEditorViewportInteractionCustomization.h b/Engine/Source/Editor/LevelEditor/Public/Elements/Actor/ActorElementLevelEditorViewportInteractionCustomization.h
--- a/Engine/Source/Editor/LevelEditor/Public/Elements/Actor/ActorElementLevelEditorViewportInteractionCustomization.h
+++ b/Engine/Source/Editor/LevelEditor/Public/Elements/Actor/ActorElementLevelEditorViewportInteractionCustomization.h
@@ -8,15 +8,11 @@
 class LEVELEDITOR_API FActorElementLevelEditorViewportInteractionCustomization : public FActorElementEditorViewportInteractionCustomization, public FTypedElementAssetEditorLevelEditorViewportClientMixin
 {
 public:
-	virtual void GetElementsToMove(const TTypedElement<UTypedElementWorldInterface>& InElementWorldHandle, const ETypedElementViewportInteractionWorldType InWorldType, const UTypedElementSelectionSet* InSelectionSet, UTypedElementList* OutElementsToMove) override;
 	virtual void GizmoManipulationStarted(const TTypedElement<UTypedElementWorldInterface>& InElementWorldHandle, const UE::Widget::EWidgetMode InWidgetMode) override;
 	virtual void GizmoManipulationDeltaUpdate(const TTypedElement<UTypedElementWorldInterface>& InElementWorldHandle, const UE::Widget::EWidgetMode InWidgetMode, const EAxisList::Type InDragAxis, const FInputDeviceState& InInputState, const FTransform& InDeltaTransform, const FVector& InPivotLocation) override;
 	virtual void GizmoManipulationStopped(const TTypedElement<UTypedElementWorldInterface>& InElementWorldHandle, const UE::Widget::EWidgetMode InWidgetMode, const ETypedElementViewportInteractionGizmoManipulationType InManipulationType) override;
 	virtual void PostGizmoManipulationStopped(TArrayView<const FTypedElementHandle> InElementHandles, const UE::Widget::EWidgetMode InWidgetMode) override;
 
-	static bool CanMoveActorInViewport(const AActor* InActor, const ETypedElementViewportInteractionWorldType InWorldType);
-	static void AppendActorsToMove(AActor* InActor, const UTypedElementSelectionSet* InSelectionSet, UTypedElementList* OutElementsToMove);
-
 private:
 	void ModifyScale(AActor* InActor, const EAxisList::Type InDragAxis, FVector& ScaleDelta, bool bCheckSmallExtent);
 };
diff --git a/Engine/Source/Editor/LevelEditor/Public/Elements/Component/ComponentElementLevelEditorViewportInteractionCustomization.h b/Engine/Source/Editor/LevelEditor/Public/Elements/Component/ComponentElementLevelEditorViewportInteractionCustomization.h
--- a/Engine/Source/Editor/LevelEditor/Public/Elements/Component/ComponentElementLevelEditorViewportInteractionCustomization.h
+++ b/Engine/Source/Editor/LevelEditor/Public/Elements/Component/ComponentElementLevelEditorViewportInteractionCustomization.h
@@ -10,7 +10,6 @@ class FLevelEditorViewportClient;
 class LEVELEDITOR_API FComponentElementLevelEditorViewportInteractionCustomization : public FComponentElementEditorViewportInteractionCustomization, public FTypedElementAssetEditorLevelEditorViewportClientMixin
 {
 public:
-	virtual void GetElementsToMove(const TTypedElement<UTypedElementWorldInterface>& InElementWorldHandle, const ETypedElementViewportInteractionWorldType InWorldType, const UTypedElementSelectionSet* InSelectionSet, UTypedElementList* OutElementsToMove) override;
 	virtual void GizmoManipulationStarted(const TTypedElement<UTypedElementWorldInterface>& InElementWorldHandle, const UE::Widget::EWidgetMode InWidgetMode) override;
 	virtual void GizmoManipulationDeltaUpdate(const TTypedElement<UTypedElementWorldInterface>& InElementWorldHandle, const UE::Widget::EWidgetMode InWidgetMode, const EAxisList::Type InDragAxis, const FInputDeviceState& InInputState, const FTransform& InDeltaTransform, const FVector& InPivotLocation) override;
 	virtual void GizmoManipulationStopped(const TTypedElement<UTypedElementWorldInterface>& InElementWorldHandle, const UE::Widget::EWidgetMode InWidgetMode, const ETypedElementViewportInteractionGizmoManipulationType InManipulationType) override;
diff --git a/Engine/Source/Editor/UnrealEd/Classes/Editor/GroupActor.h b/Engine/Source/Editor/UnrealEd/Classes/Editor/GroupActor.h
--- a/Engine/Source/Editor/UnrealEd/Classes/Editor/GroupActor.h
+++ b/Engine/Source/Editor/UnrealEd/Classes/Editor/GroupActor.h
@@ -13,6 +13,7 @@
 
 class FLevelEditorViewportClient;
 class FViewport;
+class UTypedElementList;
 class UTypedElementSelectionSet;
 
 UCLASS(hidedropdown,MinimalAPI, notplaceable, NotBlueprintable)
@@ -57,14 +58,14 @@ class AGroupActor : public AActor
 	 * Enumerate all actors and subgroups for this group (includes this).
 	 * @param InCallback		The function to call for each actor
 	 */
-	UNREALED_API void ForEachActorInGroup(TFunctionRef<void(AActor*)> InCallback);
+	UNREALED_API void ForEachActorInGroup(TFunctionRef<void(AActor*, AGroupActor*)> InCallback);
 
 	/**
 	 * Enumerate all actors and subgroups for this group that should be moved alongside this actor (includes this).
 	 * @param InSelectionSet	The selection set to use, or null to use the global selection set
 	 * @param InCallback		The function to call for each movable actor
 	 */
-	UNREALED_API void ForEachMovableActorInGroup(const UTypedElementSelectionSet* InSelectionSet, TFunctionRef<void(AActor*)> InCallback);
+	UNREALED_API void ForEachMovableActorInGroup(const UTypedElementSelectionSet* InSelectionSet, TFunctionRef<void(AActor*, AGroupActor*)> InCallback);
 	
 	/**
 	 * Apply given deltas to all actors and subgroups for this group.
@@ -223,5 +224,10 @@ private:
 	void PostRemove();
 };
 
+namespace GroupActorHelpers
+{
 
+UNREALED_API bool ActorHasParentInGroup(const AActor* Actor, const AGroupActor* GroupActor);
+UNREALED_API bool ActorHasParentInSelection(const AActor* Actor, const UTypedElementList* SelectionSet);
 
+} // namespace GroupActorHelpers
diff --git a/Engine/Source/Editor/UnrealEd/Private/EditorActor.cpp b/Engine/Source/Editor/UnrealEd/Private/EditorActor.cpp
--- a/Engine/Source/Editor/UnrealEd/Private/EditorActor.cpp
+++ b/Engine/Source/Editor/UnrealEd/Private/EditorActor.cpp
@@ -811,7 +811,8 @@ bool UUnrealEdEngine::DeleteComponents(const TArray<UActorComponent*>& InCompone
 	const FTypedElementSelectionOptions SelectionOptions = FTypedElementSelectionOptions()
 		.SetAllowHidden(true)
 		.SetAllowGroups(false)
-		.SetWarnIfLocked(false);
+		.SetWarnIfLocked(false)
+		.SetChildElementInclusionMethod(ETypedElementChildInclusionMethod::Recursive);
 
 	// Clear the selection of any components we may delete
 	for (UActorComponent* ComponentToDelete : EditableComponentsToDelete)
@@ -944,7 +945,8 @@ bool UUnrealEdEngine::DeleteActors(const TArray<AActor*>& InActorsToDelete, UWor
 	const FTypedElementSelectionOptions SelectionOptions = FTypedElementSelectionOptions()
 		.SetAllowHidden(true)
 		.SetAllowGroups(false)
-		.SetWarnIfLocked(false);
+		.SetWarnIfLocked(false)
+		.SetChildElementInclusionMethod(ETypedElementChildInclusionMethod::Recursive);
 
 	ULayersSubsystem* LayersSubsystem = GEditor->GetEditorSubsystem<ULayersSubsystem>();
 	for (AActor* Actor : ActorsToDelete)
diff --git a/Engine/Source/Editor/UnrealEd/Private/EditorSelectUtils.cpp b/Engine/Source/Editor/UnrealEd/Private/EditorSelectUtils.cpp
--- a/Engine/Source/Editor/UnrealEd/Private/EditorSelectUtils.cpp
+++ b/Engine/Source/Editor/UnrealEd/Private/EditorSelectUtils.cpp
@@ -332,8 +332,8 @@ void UUnrealEdEngine::UpdatePivotLocationForSelection( bool bOnChange )
 	ElementsToManipulate->ForEachElement<UTypedElementWorldInterface>([&SingleWorldElement](const TTypedElement<UTypedElementWorldInterface>& InWorldElement)
 	{
 #if DO_CHECK
+		if (ULevel* OwnerLevel = InWorldElement.GetOwnerLevel())
 		{
-			ULevel* OwnerLevel = InWorldElement.GetOwnerLevel();
 			const bool bIsTemplate = InWorldElement.IsTemplateElement();
 			const bool bLevelLocked = FLevelUtils::IsLevelLocked(OwnerLevel);
 			check(bIsTemplate || !bLevelLocked);
diff --git a/Engine/Source/Editor/UnrealEd/Private/GroupActor.cpp b/Engine/Source/Editor/UnrealEd/Private/GroupActor.cpp
--- a/Engine/Source/Editor/UnrealEd/Private/GroupActor.cpp
+++ b/Engine/Source/Editor/UnrealEd/Private/GroupActor.cpp
@@ -103,13 +103,13 @@ bool AGroupActor::IsSelected() const
 	return (IsLocked() && HasSelectedActors()) || Super::IsSelected();
 }
 
-void AGroupActor::ForEachActorInGroup(TFunctionRef<void(AActor*)> InCallback)
+void AGroupActor::ForEachActorInGroup(TFunctionRef<void(AActor*, AGroupActor*)> InCallback)
 {
 	for (AActor* Actor : GroupActors)
 	{
 		if (Actor)
 		{
-			InCallback(Actor);
+			InCallback(Actor, this);
 		}
 	}
 	for (AGroupActor* SubGroup : SubGroups)
@@ -119,19 +119,19 @@ void AGroupActor::ForEachActorInGroup(TFunctionRef<void(AActor*)> InCallback)
 			SubGroup->ForEachActorInGroup(InCallback);
 		}
 	}
-	InCallback(this);
+	InCallback(this, this);
 }
 
 namespace GroupActorHelpers
 {
 
-bool ActorHasParentInGroup(const TArray<AActor*>& GroupActors, const AActor* Actor)
+bool ActorHasParentInGroup(const AActor* Actor, const AGroupActor* GroupActor)
 {
-	check(Actor);
+	check(Actor && GroupActor);
 	// Check that we've not got a parent attachment within the group.
 	if (USceneComponent* RootComponent = Actor->GetRootComponent())
 	{
-		for (const AActor* OtherActor : GroupActors)
+		for (const AActor* OtherActor : GroupActor->GroupActors)
 		{
 			if (OtherActor && OtherActor != Actor)
 			{
@@ -147,7 +147,7 @@ bool ActorHasParentInGroup(const TArray<AActor*>& GroupActors, const AActor* Act
 	return false;
 }
 
-bool ActorHasParentInSelection(const AActor* Actor, const UTypedElementSelectionSet* SelectionSet)
+bool ActorHasParentInSelection(const AActor* Actor, const UTypedElementList* SelectionSet)
 {
 	check(Actor);
 	if (SelectionSet)
@@ -155,7 +155,7 @@ bool ActorHasParentInSelection(const AActor* Actor, const UTypedElementSelection
 		for (const AActor* ParentActor = Actor->GetAttachParentActor(); ParentActor; ParentActor = ParentActor->GetAttachParentActor())
 		{
 			FTypedElementHandle ParentActorElementHandle = UEngineElementsLibrary::AcquireEditorActorElementHandle(ParentActor, /*bAllowCreate*/false);
-			if (ParentActorElementHandle && SelectionSet->IsElementSelected(ParentActorElementHandle, FTypedElementIsSelectedOptions()))
+			if (ParentActorElementHandle && SelectionSet->Contains(ParentActorElementHandle))
 			{
 				return true;
 			}
@@ -166,7 +166,7 @@ bool ActorHasParentInSelection(const AActor* Actor, const UTypedElementSelection
 
 } // namespace GroupActorHelpers
 
-void AGroupActor::ForEachMovableActorInGroup(const UTypedElementSelectionSet* InSelectionSet, TFunctionRef<void(AActor*)> InCallback)
+void AGroupActor::ForEachMovableActorInGroup(const UTypedElementSelectionSet* InSelectionSet, TFunctionRef<void(AActor*, AGroupActor*)> InCallback)
 {
 	const UTypedElementSelectionSet* SelectionSet = InSelectionSet ? InSelectionSet : GEditor->GetSelectedActors()->GetElementSelectionSet();
 	for (AActor* Actor : GroupActors)
@@ -174,10 +174,10 @@ void AGroupActor::ForEachMovableActorInGroup(const UTypedElementSelectionSet* In
 		if (Actor)
 		{
 			// Check that we've not got a parent attachment within the group/selection
-			const bool bCanApplyDelta = !GroupActorHelpers::ActorHasParentInGroup(GroupActors, Actor) && !GroupActorHelpers::ActorHasParentInSelection(Actor, SelectionSet);
+			const bool bCanApplyDelta = !GroupActorHelpers::ActorHasParentInGroup(Actor, this) && !GroupActorHelpers::ActorHasParentInSelection(Actor, SelectionSet->GetElementList());
 			if (bCanApplyDelta)
 			{
-				InCallback(Actor);
+				InCallback(Actor, this);
 			}
 		}
 	}
@@ -188,14 +188,14 @@ void AGroupActor::ForEachMovableActorInGroup(const UTypedElementSelectionSet* In
 			SubGroup->ForEachMovableActorInGroup(SelectionSet, InCallback);
 		}
 	}
-	InCallback(this);
+	InCallback(this, this);
 }
 
 void AGroupActor::GroupApplyDelta(const FVector& InDrag, const FRotator& InRot, const FVector& InScale )
 {
-	ForEachMovableActorInGroup(nullptr, [&InDrag, &InRot, &InScale](AActor* InActor)
+	ForEachMovableActorInGroup(nullptr, [&InDrag, &InRot, &InScale](AActor* InGroupedActor, AGroupActor* InGroupActor)
 	{
-		GEditor->ApplyDeltaToActor(InActor, true, &InDrag, &InRot, &InScale);
+		GEditor->ApplyDeltaToActor(InGroupActor, true, &InDrag, &InRot, &InScale);
 	});
 }
 
diff --git a/Engine/Source/Editor/UnrealEd/Private/LevelEditorViewport.cpp b/Engine/Source/Editor/UnrealEd/Private/LevelEditorViewport.cpp
--- a/Engine/Source/Editor/UnrealEd/Private/LevelEditorViewport.cpp
+++ b/Engine/Source/Editor/UnrealEd/Private/LevelEditorViewport.cpp
@@ -28,6 +28,7 @@
 #include "Components/ModelComponent.h"
 #include "Kismet2/ComponentEditorUtils.h"
 #include "Engine/Selection.h"
+#include "UObject/GCObjectScopeGuard.h"
 #include "UObject/UObjectIterator.h"
 #include "EngineUtils.h"
 #include "Editor.h"
@@ -1603,9 +1604,9 @@ void FTrackingTransaction::Begin(const FText& Description, AActor* AdditionalAct
 		// Recursively call modify on any groups, as viewport manipulation of an actor within a group tends to affect the entire group
 		for (AGroupActor* GroupPendingModify : GroupsPendingModify)
 		{
-			GroupPendingModify->ForEachActorInGroup([](AActor* InActor)
+			GroupPendingModify->ForEachActorInGroup([](AActor* InGroupedActor, AGroupActor* InGroupActor)
 			{
-				InActor->Modify();
+				InGroupedActor->Modify();
 			});
 		}
 
@@ -2604,7 +2605,12 @@ bool FLevelEditorViewportClient::InputWidgetDelta(FViewport* InViewport, EAxisLi
 							if (CommonActions)
 							{
 								const UTypedElementList* ElementsToManipulate = GetElementsToManipulate();
-								TArray<FTypedElementHandle> DuplicatedElements = CommonActions->DuplicateElementsInList(ElementsToManipulate, GetWorld(), FVector::ZeroVector);
+								TArray<FTypedElementHandle> DuplicatedElements;
+								{
+									TGCObjectScopeGuard<UTypedElementList> ElementsToDuplicate(ElementsToManipulate->Clone()); // Clone this, as the duplicate may triger construction scripts to modify the selection as we duplicate
+									DuplicatedElements = CommonActions->DuplicateNormalizedElements(ElementsToDuplicate.Get(), GetWorld(), FVector::ZeroVector);
+									ElementsToDuplicate.Get()->Empty();
+								}
 
 								// Exclusively select the new elements, so that future gizmo interaction manipulates those items instead
 								if (DuplicatedElements.Num() > 0)
@@ -3531,8 +3537,35 @@ void FLevelEditorViewportClient::CacheElementsToManipulate(const bool bForceRefr
 
 	if (!bHasCachedElementsToManipulate)
 	{
+		const FTypedElementSelectionNormalizationOptions NormalizationOptions = FTypedElementSelectionNormalizationOptions()
+			.SetExpandGroups(true)
+			.SetFollowAttachment(true);
+
 		const UTypedElementSelectionSet* SelectionSet = GetSelectionSet();
-		ViewportInteraction->GetSelectedElementsToMove(SelectionSet, bIsSimulateInEditorViewport ? ETypedElementViewportInteractionWorldType::PlayInEditor : ETypedElementViewportInteractionWorldType::Editor, CachedElementsToManipulate);
+		SelectionSet->GetNormalizedSelection(NormalizationOptions, CachedElementsToManipulate);
+
+		// Remove any elements that cannot be moved
+		CachedElementsToManipulate->RemoveAll<UTypedElementWorldInterface>([this](const TTypedElement<UTypedElementWorldInterface>& InWorldElement)
+		{
+			if (!InWorldElement.CanMoveElement(bIsSimulateInEditorViewport ? ETypedElementWorldType::Game : ETypedElementWorldType::Editor))
+			{
+				return true;
+			}
+
+			// This element must belong to the current viewport world
+			if (GEditor->PlayWorld)
+			{
+				const UWorld* CurrentWorld = InWorldElement.GetOwnerWorld();
+				const UWorld* RequiredWorld = bIsSimulateInEditorViewport ? GEditor->PlayWorld : GEditor->EditorWorld;
+				if (CurrentWorld != RequiredWorld)
+				{
+					return true;
+				}
+			}
+
+			return false;
+		});
+
 		bHasCachedElementsToManipulate = true;
 	}
 }
diff --git a/Engine/Source/Editor/UnrealEd/Private/Subsystems/EditorActorSubsystem.cpp b/Engine/Source/Editor/UnrealEd/Private/Subsystems/EditorActorSubsystem.cpp
--- a/Engine/Source/Editor/UnrealEd/Private/Subsystems/EditorActorSubsystem.cpp
+++ b/Engine/Source/Editor/UnrealEd/Private/Subsystems/EditorActorSubsystem.cpp
@@ -530,17 +530,28 @@ bool UEditorActorSubsystem::DestroyActors(const TArray<AActor*>& ActorsToDestroy
 
 	FEditorDelegates::OnDeleteActorsBegin.Broadcast();
 	
-	TArray<FTypedElementHandle> ActorHandles;
-	ActorHandles.Reserve(ActorsToDestroy.Num());
-	for (AActor* ActorToDestroy : ActorsToDestroy)
+	// Make sure these actors are no longer selected
+	USelection* ActorSelection = GEditor->GetSelectedActors();
+	if (UTypedElementSelectionSet* SelectionSet = ActorSelection->GetElementSelectionSet())
 	{
-		ActorHandles.Add(UEngineElementsLibrary::AcquireEditorActorElementHandle(ActorToDestroy));
-	}
+		TArray<FTypedElementHandle> ActorHandles;
+		ActorHandles.Reserve(ActorsToDestroy.Num());
+		for (AActor* ActorToDestroy : ActorsToDestroy)
+		{
+			if (FTypedElementHandle ActorHandle = UEngineElementsLibrary::AcquireEditorActorElementHandle(ActorToDestroy, /*bAllowCreate*/false))
+			{
+				ActorHandles.Add(ActorHandle);
+			}
+		}
+	
+		const FTypedElementSelectionOptions SelectionOptions = FTypedElementSelectionOptions()
+			.SetAllowHidden(true)
+			.SetAllowGroups(false)
+			.SetWarnIfLocked(false)
+			.SetChildElementInclusionMethod(ETypedElementChildInclusionMethod::Recursive);
 
-	//To avoid dangling gizmo after actors have been destroyed
-	USelection* ActorSelection = GEditor->GetSelectedActors();
-	UTypedElementSelectionSet* SelectionSet = ActorSelection->GetElementSelectionSet();
-	SelectionSet->DeselectElements(ActorHandles, FTypedElementSelectionOptions());
+		SelectionSet->DeselectElements(ActorHandles, SelectionOptions);
+	}
 
 	ULayersSubsystem* Layers = GEditor->GetEditorSubsystem<ULayersSubsystem>();
 	Layers->DisassociateActorsFromLayers(ActorsToDestroy);
diff --git a/Engine/Source/Runtime/Engine/Classes/Components/ActorComponent.h b/Engine/Source/Runtime/Engine/Classes/Components/ActorComponent.h
--- a/Engine/Source/Runtime/Engine/Classes/Components/ActorComponent.h
+++ b/Engine/Source/Runtime/Engine/Classes/Components/ActorComponent.h
@@ -17,6 +17,8 @@
 #include "ComponentInstanceDataCache.h"
 #include "ActorComponent.generated.h"
 
+struct FTypedElementHandle;
+
 class AActor;
 class UActorComponent;
 class UAssetUserData;
@@ -859,6 +861,12 @@ public:
 	/** Called after ApplyToComponent has run. */
 	virtual void PostApplyToComponent();
 
+	/**
+	 * Get the logical child elements of this component, if any.
+	 * @see UTypedElementHierarchyInterface.
+	 */
+	virtual void GetComponentChildElements(TArray<FTypedElementHandle>& OutElementHandles, const bool bAllowCreate = true) {}
+
 	//~ Begin UObject Interface.
 	virtual void BeginDestroy() override;
 	virtual bool NeedsLoadForClient() const override;
diff --git a/Engine/Source/Runtime/Engine/Classes/Components/InstancedStaticMeshComponent.h b/Engine/Source/Runtime/Engine/Classes/Components/InstancedStaticMeshComponent.h
--- a/Engine/Source/Runtime/Engine/Classes/Components/InstancedStaticMeshComponent.h
+++ b/Engine/Source/Runtime/Engine/Classes/Components/InstancedStaticMeshComponent.h
@@ -333,6 +333,7 @@ public:
 
 	//~ Begin UActorComponent Interface
 	virtual TStructOnScope<FActorComponentInstanceData> GetComponentInstanceData() const override;
+	virtual void GetComponentChildElements(TArray<FTypedElementHandle>& OutElementHandles, const bool bAllowCreate = true) override;
 	virtual bool IsHLODRelevant() const override;
 	//~ End UActorComponent Interface
 
@@ -457,6 +458,7 @@ protected:
 
 	//~ ISMInstanceManager interface
 	virtual bool CanEditSMInstance(const FSMInstanceId& InstanceId) const override;
+	virtual bool CanMoveSMInstance(const FSMInstanceId& InstanceId, const ETypedElementWorldType InWorldType) const override;
 	virtual bool GetSMInstanceTransform(const FSMInstanceId& InstanceId, FTransform& OutInstanceTransform, bool bWorldSpace = false) const override;
 	virtual bool SetSMInstanceTransform(const FSMInstanceId& InstanceId, const FTransform& InstanceTransform, bool bWorldSpace = false, bool bMarkRenderStateDirty = false, bool bTeleport = false) override;
 	virtual void NotifySMInstanceMovementStarted(const FSMInstanceId& InstanceId) override;
diff --git a/Engine/Source/Runtime/Engine/Classes/GameFramework/LightWeightInstanceStaticMeshManager.h b/Engine/Source/Runtime/Engine/Classes/GameFramework/LightWeightInstanceStaticMeshManager.h
--- a/Engine/Source/Runtime/Engine/Classes/GameFramework/LightWeightInstanceStaticMeshManager.h
+++ b/Engine/Source/Runtime/Engine/Classes/GameFramework/LightWeightInstanceStaticMeshManager.h
@@ -51,6 +51,7 @@ protected:
 
 	//~ ISMInstanceManager interface
 	virtual bool CanEditSMInstance(const FSMInstanceId& InstanceId) const override;
+	virtual bool CanMoveSMInstance(const FSMInstanceId& InstanceId, const ETypedElementWorldType InWorldType) const override;
 	virtual bool GetSMInstanceTransform(const FSMInstanceId& InstanceId, FTransform& OutInstanceTransform, bool bWorldSpace = false) const override;
 	virtual bool SetSMInstanceTransform(const FSMInstanceId& InstanceId, const FTransform& InstanceTransform, bool bWorldSpace = false, bool bMarkRenderStateDirty = false, bool bTeleport = false) override;
 	virtual void NotifySMInstanceMovementStarted(const FSMInstanceId& InstanceId) override;
diff --git a/Engine/Source/Runtime/Engine/Private/Elements/Actor/ActorElementHierarchyInterface.cpp b/Engine/Source/Runtime/Engine/Private/Elements/Actor/ActorElementHierarchyInterface.cpp
new file mode 100644
--- /dev/null
+++ b/Engine/Source/Runtime/Engine/Private/Elements/Actor/ActorElementHierarchyInterface.cpp
@@ -0,0 +1,24 @@
+// Copyright Epic Games, Inc. All Rights Reserved.
+
+#include "Elements/Actor/ActorElementHierarchyInterface.h"
+#include "Elements/Actor/ActorElementData.h"
+#include "GameFramework/Actor.h"
+
+#include "Elements/Framework/EngineElementsLibrary.h"
+
+void UActorElementHierarchyInterface::GetChildElements(const FTypedElementHandle& InElementHandle, TArray<FTypedElementHandle>& OutElementHandles, const bool bAllowCreate)
+{
+	if (AActor* Actor = ActorElementDataUtil::GetActorFromHandle(InElementHandle))
+	{
+#if WITH_EDITOR
+		for (UActorComponent* Component : Actor->GetComponents())
+		{
+			FTypedElementHandle ComponentHandle = Component ? UEngineElementsLibrary::AcquireEditorComponentElementHandle(Component, bAllowCreate) : FTypedElementHandle();
+			if (ComponentHandle)
+			{
+				OutElementHandles.Add(MoveTemp(ComponentHandle));
+			}
+		}
+#endif	// WITH_EDITOR
+	}
+}
diff --git a/Engine/Source/Runtime/Engine/Private/Elements/Actor/ActorElementWorldInterface.cpp b/Engine/Source/Runtime/Engine/Private/Elements/Actor/ActorElementWorldInterface.cpp
--- a/Engine/Source/Runtime/Engine/Private/Elements/Actor/ActorElementWorldInterface.cpp
+++ b/Engine/Source/Runtime/Engine/Private/Elements/Actor/ActorElementWorldInterface.cpp
@@ -37,6 +37,30 @@ bool UActorElementWorldInterface::GetBounds(const FTypedElementHandle& InElement
 	return false;
 }
 
+bool UActorElementWorldInterface::CanMoveElement(const FTypedElementHandle& InElementHandle, const ETypedElementWorldType InWorldType)
+{
+	if (const AActor* Actor = ActorElementDataUtil::GetActorFromHandle(InElementHandle))
+	{
+#if WITH_EDITOR
+		// The actor cannot be location locked
+		if (InWorldType == ETypedElementWorldType::Editor && Actor->IsLockLocation())
+		{
+			return false;
+		}
+#endif	// WITH_EDITOR
+
+		// If the actor has a root component, but it cannot be moved, then the actor cannot move.
+		if (InWorldType == ETypedElementWorldType::Game && Actor->GetRootComponent() && !Actor->IsRootComponentMovable())
+		{
+			return false;
+		}
+
+		return true;
+	}
+
+	return false;
+}
+
 bool UActorElementWorldInterface::GetWorldTransform(const FTypedElementHandle& InElementHandle, FTransform& OutTransform)
 {
 	if (const AActor* Actor = ActorElementDataUtil::GetActorFromHandle(InElementHandle))
diff --git a/Engine/Source/Runtime/Engine/Private/Elements/Component/ComponentElementHierarchyInterface.cpp b/Engine/Source/Runtime/Engine/Private/Elements/Component/ComponentElementHierarchyInterface.cpp
new file mode 100644
--- /dev/null
+++ b/Engine/Source/Runtime/Engine/Private/Elements/Component/ComponentElementHierarchyInterface.cpp
@@ -0,0 +1,29 @@
+// Copyright Epic Games, Inc. All Rights Reserved.
+
+#include "Elements/Component/ComponentElementHierarchyInterface.h"
+#include "Elements/Component/ComponentElementData.h"
+#include "Components/ActorComponent.h"
+
+#include "Elements/Framework/EngineElementsLibrary.h"
+
+FTypedElementHandle UComponentElementHierarchyInterface::GetParentElement(const FTypedElementHandle& InElementHandle, const bool bAllowCreate)
+{
+	if (UActorComponent* Component = ComponentElementDataUtil::GetComponentFromHandle(InElementHandle))
+	{
+#if WITH_EDITOR
+		if (AActor* OwnerActor = Component->GetOwner())
+		{
+			return UEngineElementsLibrary::AcquireEditorActorElementHandle(OwnerActor, bAllowCreate);
+		}
+#endif	// WITH_EDITOR
+	}
+	return FTypedElementHandle();
+}
+
+void UComponentElementHierarchyInterface::GetChildElements(const FTypedElementHandle& InElementHandle, TArray<FTypedElementHandle>& OutElementHandles, const bool bAllowCreate)
+{
+	if (UActorComponent* Component = ComponentElementDataUtil::GetComponentFromHandle(InElementHandle))
+	{
+		Component->GetComponentChildElements(OutElementHandles, bAllowCreate);
+	}
+}
diff --git a/Engine/Source/Runtime/Engine/Private/Elements/Component/ComponentElementWorldInterface.cpp b/Engine/Source/Runtime/Engine/Private/Elements/Component/ComponentElementWorldInterface.cpp
--- a/Engine/Source/Runtime/Engine/Private/Elements/Component/ComponentElementWorldInterface.cpp
+++ b/Engine/Source/Runtime/Engine/Private/Elements/Component/ComponentElementWorldInterface.cpp
@@ -51,6 +51,19 @@ bool UComponentElementWorldInterface::GetBounds(const FTypedElementHandle& InEle
 	return false;
 }
 
+bool UComponentElementWorldInterface::CanMoveElement(const FTypedElementHandle& InElementHandle, const ETypedElementWorldType InWorldType)
+{
+	if (const UActorComponent* Component = ComponentElementDataUtil::GetComponentFromHandle(InElementHandle))
+	{
+		if (const USceneComponent* SceneComponent = Cast<USceneComponent>(Component))
+		{
+			return InWorldType == ETypedElementWorldType::Editor || SceneComponent->Mobility == EComponentMobility::Movable;
+		}
+	}
+
+	return false;
+}
+
 bool UComponentElementWorldInterface::GetWorldTransform(const FTypedElementHandle& InElementHandle, FTransform& OutTransform)
 {
 	if (const UActorComponent* Component = ComponentElementDataUtil::GetComponentFromHandle(InElementHandle))
diff --git a/Engine/Source/Runtime/Engine/Private/Elements/Framework/EngineElements.cpp b/Engine/Source/Runtime/Engine/Private/Elements/Framework/EngineElements.cpp
--- a/Engine/Source/Runtime/Engine/Private/Elements/Framework/EngineElements.cpp
+++ b/Engine/Source/Runtime/Engine/Private/Elements/Framework/EngineElements.cpp
@@ -11,18 +11,21 @@
 
 #include "Elements/Actor/ActorElementData.h"
 #include "Elements/Actor/ActorElementAssetDataInterface.h"
+#include "Elements/Actor/ActorElementHierarchyInterface.h"
 #include "Elements/Actor/ActorElementObjectInterface.h"
 #include "Elements/Actor/ActorElementCounterInterface.h"
 #include "Elements/Actor/ActorElementWorldInterface.h"
 #include "Elements/Actor/ActorElementSelectionInterface.h"
 
 #include "Elements/Component/ComponentElementData.h"
+#include "Elements/Component/ComponentElementHierarchyInterface.h"
 #include "Elements/Component/ComponentElementObjectInterface.h"
 #include "Elements/Component/ComponentElementCounterInterface.h"
 #include "Elements/Component/ComponentElementWorldInterface.h"
 #include "Elements/Component/ComponentElementSelectionInterface.h"
 
 #include "Elements/SMInstance/SMInstanceElementData.h"
+#include "Elements/SMInstance/SMInstanceElementHierarchyInterface.h"
 #include "Elements/SMInstance/SMInstanceElementWorldInterface.h"
 #include "Elements/SMInstance/SMInstanceElementSelectionInterface.h"
 #include "Elements/SMInstance/SMInstanceElementAssetDataInterface.h"
@@ -48,6 +51,7 @@ void RegisterEngineActorElements()
 
 	Registry->RegisterElementType<FActorElementData>(NAME_Actor);
 	Registry->RegisterElementInterface<UTypedElementAssetDataInterface>(NAME_Actor, NewObject<UActorElementAssetDataInterface>());
+	Registry->RegisterElementInterface<UTypedElementHierarchyInterface>(NAME_Actor, NewObject<UActorElementHierarchyInterface>());
 	Registry->RegisterElementInterface<UTypedElementObjectInterface>(NAME_Actor, NewObject<UActorElementObjectInterface>());
 	Registry->RegisterElementInterface<UTypedElementCounterInterface>(NAME_Actor, NewObject<UActorElementCounterInterface>());
 	Registry->RegisterElementInterface<UTypedElementWorldInterface>(NAME_Actor, NewObject<UActorElementWorldInterface>());
@@ -59,6 +63,7 @@ void RegisterEngineComponentElements()
 	UTypedElementRegistry* Registry = UTypedElementRegistry::GetInstance();
 
 	Registry->RegisterElementType<FComponentElementData>(NAME_Components);
+	Registry->RegisterElementInterface<UTypedElementHierarchyInterface>(NAME_Components, NewObject<UComponentElementHierarchyInterface>());
 	Registry->RegisterElementInterface<UTypedElementObjectInterface>(NAME_Components, NewObject<UComponentElementObjectInterface>());
 	Registry->RegisterElementInterface<UTypedElementCounterInterface>(NAME_Components, NewObject<UComponentElementCounterInterface>());
 	Registry->RegisterElementInterface<UTypedElementWorldInterface>(NAME_Components, NewObject<UComponentElementWorldInterface>());
@@ -70,6 +75,7 @@ void RegisterEngineSMInstanceElements()
 	UTypedElementRegistry* Registry = UTypedElementRegistry::GetInstance();
 
 	Registry->RegisterElementType<FSMInstanceElementData>(NAME_SMInstance);
+	Registry->RegisterElementInterface<UTypedElementHierarchyInterface>(NAME_SMInstance, NewObject<USMInstanceElementHierarchyInterface>());
 	Registry->RegisterElementInterface<UTypedElementWorldInterface>(NAME_SMInstance, NewObject<USMInstanceElementWorldInterface>());
 	Registry->RegisterElementInterface<UTypedElementSelectionInterface>(NAME_SMInstance, NewObject<USMInstanceElementSelectionInterface>());
 	Registry->RegisterElementInterface<UTypedElementAssetDataInterface>(NAME_SMInstance, NewObject<USMInstanceElementAssetDataInterface>());
diff --git a/Engine/Source/Runtime/Engine/Private/Elements/Framework/TypedElementCommonActions.cpp b/Engine/Source/Runtime/Engine/Private/Elements/Framework/TypedElementCommonActions.cpp
--- a/Engine/Source/Runtime/Engine/Private/Elements/Framework/TypedElementCommonActions.cpp
+++ b/Engine/Source/Runtime/Engine/Private/Elements/Framework/TypedElementCommonActions.cpp
@@ -7,11 +7,6 @@
 
 #include "UObject/GCObjectScopeGuard.h"
 
-void FTypedElementCommonActionsCustomization::GetElementsForAction(const TTypedElement<UTypedElementWorldInterface>& InElementWorldHandle, const UTypedElementList* InElementList, UTypedElementList* OutElementsToDelete)
-{
-	OutElementsToDelete->Add(InElementWorldHandle);
-}
-
 bool FTypedElementCommonActionsCustomization::DeleteElements(UTypedElementWorldInterface* InWorldInterface, TArrayView<const FTypedElementHandle> InElementHandles, UWorld* InWorld, UTypedElementSelectionSet* InSelectionSet, const FTypedElementDeletionOptions& InDeletionOptions)
 {
 	return InWorldInterface->DeleteElements(InElementHandles, InWorld, InSelectionSet, InDeletionOptions);
@@ -23,56 +18,24 @@ void FTypedElementCommonActionsCustomization::DuplicateElements(UTypedElementWor
 }
 
 
-void UTypedElementCommonActions::GetElementsForAction(const UTypedElementList* InElementList, UTypedElementList* OutElementsForAction) const
-{
-	OutElementsForAction->Reset();
-	InElementList->ForEachElement<UTypedElementWorldInterface>([this, InElementList, OutElementsForAction](const TTypedElement<UTypedElementWorldInterface>& InElementWorldHandle)
-	{
-		FTypedElementCommonActionsElement CommonActionsElement(InElementWorldHandle, GetInterfaceCustomizationByTypeId(InElementWorldHandle.GetId().GetTypeId()));
-		check(CommonActionsElement.IsSet());
-		CommonActionsElement.GetElementsForAction(InElementList, OutElementsForAction);
-		return true;
-	});
-}
-
-bool UTypedElementCommonActions::DeleteElements(const TArray<FTypedElementHandle>& ElementHandles, UWorld* World, UTypedElementSelectionSet* SelectionSet, const FTypedElementDeletionOptions& DeletionOptions)
-{
-	return DeleteElements(MakeArrayView(ElementHandles), World, SelectionSet, DeletionOptions);
-}
-
-bool UTypedElementCommonActions::DeleteElements(TArrayView<const FTypedElementHandle> ElementHandles, UWorld* World, UTypedElementSelectionSet* SelectionSet, const FTypedElementDeletionOptions& DeletionOptions)
+bool UTypedElementCommonActions::DeleteSelectedElements(UTypedElementSelectionSet* SelectionSet, UWorld* World, const FTypedElementDeletionOptions& DeletionOptions)
 {
-	TMap<FTypedHandleTypeId, TArray<FTypedElementHandle>> ElementsToDeleteByType;
-	TypedElementUtil::BatchElementsByType(ElementHandles, ElementsToDeleteByType);
-
-	bool bSuccess = false;
-
-	UTypedElementRegistry* Registry = UTypedElementRegistry::GetInstance();
-	UTypedElementRegistry::FDisableElementDestructionOnGC GCGuard(Registry);
-
-	for (const auto& ElementsByTypePair : ElementsToDeleteByType)
-	{
-		FTypedElementCommonActionsCustomization* CommonActionsCustomization = GetInterfaceCustomizationByTypeId(ElementsByTypePair.Key);
-		UTypedElementWorldInterface* WorldInterface = Registry->GetElementInterface<UTypedElementWorldInterface>(ElementsByTypePair.Key);
-		if (CommonActionsCustomization && WorldInterface)
-		{
-			bSuccess |= CommonActionsCustomization->DeleteElements(WorldInterface, ElementsByTypePair.Value, World, SelectionSet, DeletionOptions);
-		}
-	}
-
+	TGCObjectScopeGuard<UTypedElementList> NormalizedElements(SelectionSet->GetNormalizedSelection(FTypedElementSelectionNormalizationOptions()));
+	const bool bSuccess = DeleteNormalizedElements(NormalizedElements.Get(), World, SelectionSet, DeletionOptions);
+	NormalizedElements.Get()->Reset();
 	return bSuccess;
 }
 
-bool UTypedElementCommonActions::DeleteElements(const UTypedElementList* ElementList, UWorld* World, UTypedElementSelectionSet* SelectionSet, const FTypedElementDeletionOptions& DeletionOptions)
+bool UTypedElementCommonActions::DeleteNormalizedElements(const UTypedElementList* ElementList, UWorld* World, UTypedElementSelectionSet* SelectionSet, const FTypedElementDeletionOptions& DeletionOptions)
 {
 	TMap<FTypedHandleTypeId, TArray<FTypedElementHandle>> ElementsToDeleteByType;
 	TypedElementUtil::BatchElementsByType(ElementList, ElementsToDeleteByType);
-
+	
 	bool bSuccess = false;
-
+	
 	UTypedElementRegistry* Registry = UTypedElementRegistry::GetInstance();
 	UTypedElementRegistry::FDisableElementDestructionOnGC GCGuard(Registry);
-
+	
 	for (const auto& ElementsByTypePair : ElementsToDeleteByType)
 	{
 		FTypedElementCommonActionsCustomization* CommonActionsCustomization = GetInterfaceCustomizationByTypeId(ElementsByTypePair.Key);
@@ -82,60 +45,26 @@ bool UTypedElementCommonActions::DeleteElements(const UTypedElementList* Element
 			bSuccess |= CommonActionsCustomization->DeleteElements(WorldInterface, ElementsByTypePair.Value, World, SelectionSet, DeletionOptions);
 		}
 	}
-
-	return bSuccess;
-}
-
-bool UTypedElementCommonActions::DeleteSelectedElements(UTypedElementSelectionSet* SelectionSet, UWorld* World, const FTypedElementDeletionOptions& DeletionOptions)
-{
-	return DeleteElementsInList(SelectionSet->GetElementList(), World, SelectionSet, DeletionOptions);
-}
-
-bool UTypedElementCommonActions::DeleteElementsInList(const UTypedElementList* ElementList, UWorld* World, UTypedElementSelectionSet* SelectionSet, const FTypedElementDeletionOptions& DeletionOptions)
-{
-	TGCObjectScopeGuard<UTypedElementList> ElementsForAction(UTypedElementRegistry::GetInstance()->CreateElementList());
-	GetElementsForAction(ElementList, ElementsForAction.Get());
-
-	const bool bSuccess = DeleteElements(ElementsForAction.Get(), World, SelectionSet, DeletionOptions);
-	ElementsForAction.Get()->Reset();
+	
 	return bSuccess;
 }
 
-TArray<FTypedElementHandle> UTypedElementCommonActions::DuplicateElements(const TArray<FTypedElementHandle>& ElementHandles, UWorld* World, const FVector& LocationOffset)
-{
-	return DuplicateElements(MakeArrayView(ElementHandles), World, LocationOffset);
-}
-
-TArray<FTypedElementHandle> UTypedElementCommonActions::DuplicateElements(TArrayView<const FTypedElementHandle> ElementHandles, UWorld* World, const FVector& LocationOffset)
+TArray<FTypedElementHandle> UTypedElementCommonActions::DuplicateSelectedElements(const UTypedElementSelectionSet* SelectionSet, UWorld* World, const FVector& LocationOffset)
 {
-	TMap<FTypedHandleTypeId, TArray<FTypedElementHandle>> ElementsToDuplicateByType;
-	TypedElementUtil::BatchElementsByType(ElementHandles, ElementsToDuplicateByType);
-
-	TArray<FTypedElementHandle> NewElements;
-	NewElements.Reserve(ElementHandles.Num());
-
-	UTypedElementRegistry* Registry = UTypedElementRegistry::GetInstance();
-	for (const auto& ElementsByTypePair : ElementsToDuplicateByType)
-	{
-		FTypedElementCommonActionsCustomization* CommonActionsCustomization = GetInterfaceCustomizationByTypeId(ElementsByTypePair.Key);
-		UTypedElementWorldInterface* WorldInterface = Registry->GetElementInterface<UTypedElementWorldInterface>(ElementsByTypePair.Key);
-		if (CommonActionsCustomization && WorldInterface)
-		{
-			CommonActionsCustomization->DuplicateElements(WorldInterface, ElementsByTypePair.Value, World, LocationOffset, NewElements);
-		}
-	}
-
+	TGCObjectScopeGuard<UTypedElementList> NormalizedElements(SelectionSet->GetNormalizedSelection(FTypedElementSelectionNormalizationOptions()));
+	TArray<FTypedElementHandle> NewElements = DuplicateNormalizedElements(NormalizedElements.Get(), World, LocationOffset);
+	NormalizedElements.Get()->Reset();
 	return NewElements;
 }
 
-TArray<FTypedElementHandle> UTypedElementCommonActions::DuplicateElements(const UTypedElementList* ElementList, UWorld* World, const FVector& LocationOffset)
+TArray<FTypedElementHandle> UTypedElementCommonActions::DuplicateNormalizedElements(const UTypedElementList* ElementList, UWorld* World, const FVector& LocationOffset)
 {
 	TMap<FTypedHandleTypeId, TArray<FTypedElementHandle>> ElementsToDuplicateByType;
 	TypedElementUtil::BatchElementsByType(ElementList, ElementsToDuplicateByType);
-
+	
 	TArray<FTypedElementHandle> NewElements;
 	NewElements.Reserve(ElementList->Num());
-
+	
 	UTypedElementRegistry* Registry = UTypedElementRegistry::GetInstance();
 	for (const auto& ElementsByTypePair : ElementsToDuplicateByType)
 	{
@@ -146,22 +75,7 @@ TArray<FTypedElementHandle> UTypedElementCommonActions::DuplicateElements(const
 			CommonActionsCustomization->DuplicateElements(WorldInterface, ElementsByTypePair.Value, World, LocationOffset, NewElements);
 		}
 	}
-
-	return NewElements;
-}
-
-TArray<FTypedElementHandle> UTypedElementCommonActions::DuplicateSelectedElements(const UTypedElementSelectionSet* SelectionSet, UWorld* World, const FVector& LocationOffset)
-{
-	return DuplicateElementsInList(SelectionSet->GetElementList(), World, LocationOffset);
-}
-
-TArray<FTypedElementHandle> UTypedElementCommonActions::DuplicateElementsInList(const UTypedElementList* ElementList, UWorld* World, const FVector& LocationOffset)
-{
-	TGCObjectScopeGuard<UTypedElementList> ElementsForAction(UTypedElementRegistry::GetInstance()->CreateElementList());
-	GetElementsForAction(ElementList, ElementsForAction.Get());
-
-	TArray<FTypedElementHandle> NewElements = DuplicateElements(MakeArrayView(ElementsForAction.Get()->GetElementHandles()), World, LocationOffset);
-	ElementsForAction.Get()->Reset();
+	
 	return NewElements;
 }
 
diff --git a/Engine/Source/Runtime/Engine/Private/Elements/SMInstance/SMInstanceElementHierarchyInterface.cpp b/Engine/Source/Runtime/Engine/Private/Elements/SMInstance/SMInstanceElementHierarchyInterface.cpp
new file mode 100644
--- /dev/null
+++ b/Engine/Source/Runtime/Engine/Private/Elements/SMInstance/SMInstanceElementHierarchyInterface.cpp
@@ -0,0 +1,21 @@
+// Copyright Epic Games, Inc. All Rights Reserved.
+
+#include "Elements/SMInstance/SMInstanceElementHierarchyInterface.h"
+#include "Elements/SMInstance/SMInstanceElementData.h"
+#include "Components/InstancedStaticMeshComponent.h"
+
+#include "Elements/Framework/EngineElementsLibrary.h"
+
+FTypedElementHandle USMInstanceElementHierarchyInterface::GetParentElement(const FTypedElementHandle& InElementHandle, const bool bAllowCreate)
+{
+	if (FSMInstanceManager SMInstance = SMInstanceElementDataUtil::GetSMInstanceFromHandle(InElementHandle))
+	{
+#if WITH_EDITOR
+		if (UInstancedStaticMeshComponent* ISMComponent = SMInstance.GetISMComponent())
+		{
+			return UEngineElementsLibrary::AcquireEditorComponentElementHandle(ISMComponent, bAllowCreate);
+		}
+#endif	// WITH_EDITOR
+	}
+	return FTypedElementHandle();
+}
diff --git a/Engine/Source/Runtime/Engine/Private/Elements/SMInstance/SMInstanceElementWorldInterface.cpp b/Engine/Source/Runtime/Engine/Private/Elements/SMInstance/SMInstanceElementWorldInterface.cpp
--- a/Engine/Source/Runtime/Engine/Private/Elements/SMInstance/SMInstanceElementWorldInterface.cpp
+++ b/Engine/Source/Runtime/Engine/Private/Elements/SMInstance/SMInstanceElementWorldInterface.cpp
@@ -63,6 +63,12 @@ bool USMInstanceElementWorldInterface::GetBounds(const FTypedElementHandle& InEl
 	return false;
 }
 
+bool USMInstanceElementWorldInterface::CanMoveElement(const FTypedElementHandle& InElementHandle, const ETypedElementWorldType InWorldType)
+{
+	const FSMInstanceManager SMInstance = SMInstanceElementDataUtil::GetSMInstanceFromHandle(InElementHandle);
+	return SMInstance && SMInstance.CanMoveSMInstance(InWorldType);
+}
+
 bool USMInstanceElementWorldInterface::GetWorldTransform(const FTypedElementHandle& InElementHandle, FTransform& OutTransform)
 {
 	const FSMInstanceManager SMInstance = SMInstanceElementDataUtil::GetSMInstanceFromHandle(InElementHandle);
diff --git a/Engine/Source/Runtime/Engine/Private/InstancedStaticMesh.cpp b/Engine/Source/Runtime/Engine/Private/InstancedStaticMesh.cpp
--- a/Engine/Source/Runtime/Engine/Private/InstancedStaticMesh.cpp
+++ b/Engine/Source/Runtime/Engine/Private/InstancedStaticMesh.cpp
@@ -36,6 +36,7 @@
 
 #include "Elements/Framework/EngineElementsLibrary.h"
 #include "Elements/SMInstance/SMInstanceElementData.h"
+#include "Elements/Interfaces/TypedElementWorldInterface.h"
 
 #if RHI_RAYTRACING
 #include "RayTracingInstance.h"
@@ -1913,6 +1914,19 @@ TStructOnScope<FActorComponentInstanceData> UInstancedStaticMeshComponent::GetCo
 	return InstanceData;
 }
 
+void UInstancedStaticMeshComponent::GetComponentChildElements(TArray<FTypedElementHandle>& OutElementHandles, const bool bAllowCreate)
+{
+#if WITH_EDITOR
+	for (int32 InstanceIndex = 0; InstanceIndex < PerInstanceSMData.Num(); ++InstanceIndex)
+	{
+		if (FTypedElementHandle ElementHandle = UEngineElementsLibrary::AcquireEditorSMInstanceElementHandle(this, InstanceIndex, bAllowCreate))
+		{
+			OutElementHandles.Add(MoveTemp(ElementHandle));
+		}
+	}
+#endif	// WITH_EDITOR
+}
+
 void UInstancedStaticMeshComponent::ApplyComponentInstanceData(FInstancedStaticMeshComponentInstanceData* InstancedMeshData)
 {
 #if WITH_EDITOR
@@ -3960,6 +3974,12 @@ bool UInstancedStaticMeshComponent::CanEditSMInstance(const FSMInstanceId& Insta
 	return IsEditableWhenInherited();
 }
 
+bool UInstancedStaticMeshComponent::CanMoveSMInstance(const FSMInstanceId& InstanceId, const ETypedElementWorldType InWorldType) const
+{
+	check(InstanceId.ISMComponent == this);
+	return InWorldType == ETypedElementWorldType::Editor || InstanceId.ISMComponent->Mobility == EComponentMobility::Movable;
+}
+
 bool UInstancedStaticMeshComponent::GetSMInstanceTransform(const FSMInstanceId& InstanceId, FTransform& OutInstanceTransform, bool bWorldSpace) const
 {
 	check(InstanceId.ISMComponent == this);
diff --git a/Engine/Source/Runtime/Engine/Private/LightWeightInstanceStaticMeshManager.cpp b/Engine/Source/Runtime/Engine/Private/LightWeightInstanceStaticMeshManager.cpp
--- a/Engine/Source/Runtime/Engine/Private/LightWeightInstanceStaticMeshManager.cpp
+++ b/Engine/Source/Runtime/Engine/Private/LightWeightInstanceStaticMeshManager.cpp
@@ -204,6 +204,14 @@ bool ALightWeightInstanceStaticMeshManager::CanEditSMInstance(const FSMInstanceI
 	return InstanceManager->CanEditSMInstance(InstanceId);
 }
 
+bool ALightWeightInstanceStaticMeshManager::CanMoveSMInstance(const FSMInstanceId& InstanceId, const ETypedElementWorldType InWorldType) const
+{
+	check(InstanceId.ISMComponent == InstancedStaticMeshComponent);
+
+	ISMInstanceManager* InstanceManager = InstancedStaticMeshComponent;
+	return InstanceManager->CanMoveSMInstance(InstanceId, InWorldType);
+}
+
 bool ALightWeightInstanceStaticMeshManager::GetSMInstanceTransform(const FSMInstanceId& InstanceId, FTransform& OutInstanceTransform, bool bWorldSpace) const
 {
 	check(InstanceId.ISMComponent == InstancedStaticMeshComponent);
diff --git a/Engine/Source/Runtime/Engine/Public/Elements/Actor/ActorElementHierarchyInterface.h b/Engine/Source/Runtime/Engine/Public/Elements/Actor/ActorElementHierarchyInterface.h
new file mode 100644
--- /dev/null
+++ b/Engine/Source/Runtime/Engine/Public/Elements/Actor/ActorElementHierarchyInterface.h
@@ -0,0 +1,15 @@
+// Copyright Epic Games, Inc. All Rights Reserved.
+
+#pragma once
+
+#include "Elements/Interfaces/TypedElementHierarchyInterface.h"
+#include "ActorElementHierarchyInterface.generated.h"
+
+UCLASS()
+class ENGINE_API UActorElementHierarchyInterface : public UTypedElementHierarchyInterface
+{
+	GENERATED_BODY()
+
+public:
+	virtual void GetChildElements(const FTypedElementHandle& InElementHandle, TArray<FTypedElementHandle>& OutElementHandles, const bool bAllowCreate = true) override;
+};
diff --git a/Engine/Source/Runtime/Engine/Public/Elements/Actor/ActorElementWorldInterface.h b/Engine/Source/Runtime/Engine/Public/Elements/Actor/ActorElementWorldInterface.h
--- a/Engine/Source/Runtime/Engine/Public/Elements/Actor/ActorElementWorldInterface.h
+++ b/Engine/Source/Runtime/Engine/Public/Elements/Actor/ActorElementWorldInterface.h
@@ -17,6 +17,7 @@ public:
 	virtual ULevel* GetOwnerLevel(const FTypedElementHandle& InElementHandle) override;
 	virtual UWorld* GetOwnerWorld(const FTypedElementHandle& InElementHandle) override;
 	virtual bool GetBounds(const FTypedElementHandle& InElementHandle, FBoxSphereBounds& OutBounds) override;
+	virtual bool CanMoveElement(const FTypedElementHandle& InElementHandle, const ETypedElementWorldType InWorldType) override;
 	virtual bool GetWorldTransform(const FTypedElementHandle& InElementHandle, FTransform& OutTransform) override;
 	virtual bool SetWorldTransform(const FTypedElementHandle& InElementHandle, const FTransform& InTransform) override;
 	virtual bool GetRelativeTransform(const FTypedElementHandle& InElementHandle, FTransform& OutTransform) override;
diff --git a/Engine/Source/Runtime/Engine/Public/Elements/Component/ComponentElementHierarchyInterface.h b/Engine/Source/Runtime/Engine/Public/Elements/Component/ComponentElementHierarchyInterface.h
new file mode 100644
--- /dev/null
+++ b/Engine/Source/Runtime/Engine/Public/Elements/Component/ComponentElementHierarchyInterface.h
@@ -0,0 +1,16 @@
+// Copyright Epic Games, Inc. All Rights Reserved.
+
+#pragma once
+
+#include "Elements/Interfaces/TypedElementHierarchyInterface.h"
+#include "ComponentElementHierarchyInterface.generated.h"
+
+UCLASS()
+class ENGINE_API UComponentElementHierarchyInterface : public UTypedElementHierarchyInterface
+{
+	GENERATED_BODY()
+
+public:
+	virtual FTypedElementHandle GetParentElement(const FTypedElementHandle& InElementHandle, const bool bAllowCreate = true) override;
+	virtual void GetChildElements(const FTypedElementHandle& InElementHandle, TArray<FTypedElementHandle>& OutElementHandles, const bool bAllowCreate = true) override;
+};
diff --git a/Engine/Source/Runtime/Engine/Public/Elements/Component/ComponentElementWorldInterface.h b/Engine/Source/Runtime/Engine/Public/Elements/Component/ComponentElementWorldInterface.h
--- a/Engine/Source/Runtime/Engine/Public/Elements/Component/ComponentElementWorldInterface.h
+++ b/Engine/Source/Runtime/Engine/Public/Elements/Component/ComponentElementWorldInterface.h
@@ -16,6 +16,7 @@ public:
 	virtual ULevel* GetOwnerLevel(const FTypedElementHandle& InElementHandle) override;
 	virtual UWorld* GetOwnerWorld(const FTypedElementHandle& InElementHandle) override;
 	virtual bool GetBounds(const FTypedElementHandle& InElementHandle, FBoxSphereBounds& OutBounds) override;
+	virtual bool CanMoveElement(const FTypedElementHandle& InElementHandle, const ETypedElementWorldType InWorldType) override;
 	virtual bool GetWorldTransform(const FTypedElementHandle& InElementHandle, FTransform& OutTransform) override;
 	virtual bool SetWorldTransform(const FTypedElementHandle& InElementHandle, const FTransform& InTransform) override;
 	virtual bool GetRelativeTransform(const FTypedElementHandle& InElementHandle, FTransform& OutTransform) override;
diff --git a/Engine/Source/Runtime/Engine/Public/Elements/Framework/TypedElementCommonActions.h b/Engine/Source/Runtime/Engine/Public/Elements/Framework/TypedElementCommonActions.h
--- a/Engine/Source/Runtime/Engine/Public/Elements/Framework/TypedElementCommonActions.h
+++ b/Engine/Source/Runtime/Engine/Public/Elements/Framework/TypedElementCommonActions.h
@@ -18,7 +18,6 @@ public:
 	virtual ~FTypedElementCommonActionsCustomization() = default;
 
 	//~ See UTypedElementCommonActions for API docs
-	virtual void GetElementsForAction(const TTypedElement<UTypedElementWorldInterface>& InElementWorldHandle, const UTypedElementList* InElementList, UTypedElementList* OutElementsToDelete);
 	virtual bool DeleteElements(UTypedElementWorldInterface* InWorldInterface, TArrayView<const FTypedElementHandle> InElementHandles, UWorld* InWorld, UTypedElementSelectionSet* InSelectionSet, const FTypedElementDeletionOptions& InDeletionOptions);
 	virtual void DuplicateElements(UTypedElementWorldInterface* InWorldInterface, TArrayView<const FTypedElementHandle> InElementHandles, UWorld* InWorld, const FVector& InLocationOffset, TArray<FTypedElementHandle>& OutNewElements);
 };
@@ -55,7 +54,6 @@ public:
 	}
 
 	//~ See UTypedElementCommonActions for API docs
-	void GetElementsForAction(const UTypedElementList* InElementList, UTypedElementList* OutElementsToDelete) { CommonActionsCustomization->GetElementsForAction(ElementWorldHandle, InElementList, OutElementsToDelete); }
 
 private:
 	TTypedElement<UTypedElementWorldInterface> ElementWorldHandle;
@@ -73,28 +71,32 @@ class ENGINE_API UTypedElementCommonActions : public UObject, public TTypedEleme
 
 public:
 	/**
-	 * Get the elements from the given element list that a common action should operate on.
-	 * @note This allows coarse filtering of the elements based on the selection state (eg, favoring components over actors), however it doesn't mean that the returned elements will actually be valid to perform a given action on.
+	 * Delete any elements from the given selection set that can be deleted.
+	 * @note Internally this just calls DeleteNormalizedElements on the result of UTypedElementSelectionSet::GetNormalizedSelection.
 	 */
-	void GetElementsForAction(const UTypedElementList* InElementList, UTypedElementList* OutElementsForAction) const;
-
-	UFUNCTION(BlueprintCallable, Category="TypedElementFramework|Common")
-	bool DeleteElements(const TArray<FTypedElementHandle>& ElementHandles, UWorld* World, UTypedElementSelectionSet* SelectionSet, const FTypedElementDeletionOptions& DeletionOptions);
-	bool DeleteElements(TArrayView<const FTypedElementHandle> ElementHandles, UWorld* World, UTypedElementSelectionSet* SelectionSet, const FTypedElementDeletionOptions& DeletionOptions);
-	bool DeleteElements(const UTypedElementList* ElementList, UWorld* World, UTypedElementSelectionSet* SelectionSet, const FTypedElementDeletionOptions& DeletionOptions);
-
 	UFUNCTION(BlueprintCallable, Category="TypedElementFramework|Common")
 	bool DeleteSelectedElements(UTypedElementSelectionSet* SelectionSet, UWorld* World, const FTypedElementDeletionOptions& DeletionOptions);
-	bool DeleteElementsInList(const UTypedElementList* ElementList, UWorld* World, UTypedElementSelectionSet* InSelectionSet, const FTypedElementDeletionOptions& DeletionOptions);
-
+	
+	/**
+	 * Delete any elements from the given list that can be deleted.
+	 * @note This list should have been pre-normalized via UTypedElementSelectionSet::GetNormalizedSelection or UTypedElementSelectionSet::GetNormalizedElementList.
+	 */
 	UFUNCTION(BlueprintCallable, Category="TypedElementFramework|Common")
-	TArray<FTypedElementHandle> DuplicateElements(const TArray<FTypedElementHandle>& ElementHandles, UWorld* World, const FVector& LocationOffset);
-	TArray<FTypedElementHandle> DuplicateElements(TArrayView<const FTypedElementHandle> ElementHandles, UWorld* World, const FVector& LocationOffset);
-	TArray<FTypedElementHandle> DuplicateElements(const UTypedElementList* ElementList, UWorld* World, const FVector& LocationOffset);
+	bool DeleteNormalizedElements(const UTypedElementList* ElementList, UWorld* World, UTypedElementSelectionSet* InSelectionSet, const FTypedElementDeletionOptions& DeletionOptions);
 
+	/**
+	 * Duplicate any elements from the given selection set that can be duplicated.
+	 * @note Internally this just calls DuplicateNormalizedElements on the result of UTypedElementSelectionSet::GetNormalizedSelection.
+	 */
 	UFUNCTION(BlueprintCallable, Category="TypedElementFramework|Common")
 	TArray<FTypedElementHandle> DuplicateSelectedElements(const UTypedElementSelectionSet* SelectionSet, UWorld* World, const FVector& LocationOffset);
-	TArray<FTypedElementHandle> DuplicateElementsInList(const UTypedElementList* ElementList, UWorld* World, const FVector& LocationOffset);
+	
+	/**
+	 * Duplicate any elements from the given list that can be duplicated.
+	 * @note This list should have been pre-normalized via UTypedElementSelectionSet::GetNormalizedSelection or UTypedElementSelectionSet::GetNormalizedElementList.
+	 */
+	UFUNCTION(BlueprintCallable, Category="TypedElementFramework|Common")
+	TArray<FTypedElementHandle> DuplicateNormalizedElements(const UTypedElementList* ElementList, UWorld* World, const FVector& LocationOffset);
 
 private:
 	/**
diff --git a/Engine/Source/Runtime/Engine/Public/Elements/Interfaces/TypedElementWorldInterface.h b/Engine/Source/Runtime/Engine/Public/Elements/Interfaces/TypedElementWorldInterface.h
--- a/Engine/Source/Runtime/Engine/Public/Elements/Interfaces/TypedElementWorldInterface.h
+++ b/Engine/Source/Runtime/Engine/Public/Elements/Interfaces/TypedElementWorldInterface.h
@@ -10,6 +10,13 @@ class UWorld;
 class UTypedElementSelectionSet;
 struct FCollisionShape;
 
+UENUM()
+enum class ETypedElementWorldType : uint8
+{
+	Game,
+	Editor,
+};
+
 USTRUCT(BlueprintType)
 struct FTypedElementDeletionOptions
 {
@@ -72,6 +79,12 @@ public:
 	UFUNCTION(BlueprintPure, Category="TypedElementInterfaces|World")
 	virtual bool GetBounds(const FTypedElementHandle& InElementHandle, FBoxSphereBounds& OutBounds) { return false; }
 
+	/**
+	 * Can the given element be moved within the world?
+	 */
+	UFUNCTION(BlueprintPure, Category="TypedElementInterfaces|World")
+	virtual bool CanMoveElement(const FTypedElementHandle& InElementHandle, const ETypedElementWorldType InWorldType) { return false; }
+
 	/**
 	 * Get the transform of this element within its owner world, if any.
 	 */
@@ -161,7 +174,7 @@ public:
 
 	/**
 	 * Delete the given set of elements.
-	 * @note If you want to delete an array of elements that are potentially different types, you probably want to use the higher-level UTypedElementCommonActions::DeleteElements function instead.
+	 * @note If you want to delete an array of elements that are potentially different types, you probably want to use the higher-level UTypedElementCommonActions::DeleteNormalizedElements function instead.
 	 */
 	virtual bool DeleteElements(TArrayView<const FTypedElementHandle> InElementHandles, UWorld* InWorld, UTypedElementSelectionSet* InSelectionSet, const FTypedElementDeletionOptions& InDeletionOptions) { return false; }
 
@@ -185,7 +198,7 @@ public:
 
 	/**
 	 * Duplicate the given set of elements.
-	 * @note If you want to duplicate an array of elements that are potentially different types, you probably want to use the higher-level UTypedElementCommonActions::DuplicateElements function instead.
+	 * @note If you want to duplicate an array of elements that are potentially different types, you probably want to use the higher-level UTypedElementCommonActions::DuplicateNormalizedElements function instead.
 	 */
 	virtual void DuplicateElements(TArrayView<const FTypedElementHandle> InElementHandles, UWorld* InWorld, const FVector& InLocationOffset, TArray<FTypedElementHandle>& OutNewElements)
 	{
@@ -200,6 +213,7 @@ struct TTypedElement<UTypedElementWorldInterface> : public TTypedElementBase<UTy
 	ULevel* GetOwnerLevel() const { return InterfacePtr->GetOwnerLevel(*this); }
 	UWorld* GetOwnerWorld() const { return InterfacePtr->GetOwnerWorld(*this); }
 	bool GetBounds(FBoxSphereBounds& OutBounds) const { return InterfacePtr->GetBounds(*this, OutBounds); }
+	bool CanMoveElement(const ETypedElementWorldType InWorldType) const { return InterfacePtr->CanMoveElement(*this, InWorldType); }
 	bool GetWorldTransform(FTransform& OutTransform) const { return InterfacePtr->GetWorldTransform(*this, OutTransform); }
 	bool SetWorldTransform(const FTransform& InTransform) const { return InterfacePtr->SetWorldTransform(*this, InTransform); }
 	bool GetRelativeTransform(FTransform& OutTransform) const { return InterfacePtr->GetRelativeTransform(*this, OutTransform); }
diff --git a/Engine/Source/Runtime/Engine/Public/Elements/SMInstance/SMInstanceElementHierarchyInterface.h b/Engine/Source/Runtime/Engine/Public/Elements/SMInstance/SMInstanceElementHierarchyInterface.h
new file mode 100644
--- /dev/null
+++ b/Engine/Source/Runtime/Engine/Public/Elements/SMInstance/SMInstanceElementHierarchyInterface.h
@@ -0,0 +1,15 @@
+// Copyright Epic Games, Inc. All Rights Reserved.
+
+#pragma once
+
+#include "Elements/Interfaces/TypedElementHierarchyInterface.h"
+#include "SMInstanceElementHierarchyInterface.generated.h"
+
+UCLASS()
+class ENGINE_API USMInstanceElementHierarchyInterface : public UTypedElementHierarchyInterface
+{
+	GENERATED_BODY()
+
+public:
+	virtual FTypedElementHandle GetParentElement(const FTypedElementHandle& InElementHandle, const bool bAllowCreate = true) override;
+};
diff --git a/Engine/Source/Runtime/Engine/Public/Elements/SMInstance/SMInstanceElementWorldInterface.h b/Engine/Source/Runtime/Engine/Public/Elements/SMInstance/SMInstanceElementWorldInterface.h
--- a/Engine/Source/Runtime/Engine/Public/Elements/SMInstance/SMInstanceElementWorldInterface.h
+++ b/Engine/Source/Runtime/Engine/Public/Elements/SMInstance/SMInstanceElementWorldInterface.h
@@ -16,6 +16,7 @@ public:
 	virtual ULevel* GetOwnerLevel(const FTypedElementHandle& InElementHandle) override;
 	virtual UWorld* GetOwnerWorld(const FTypedElementHandle& InElementHandle) override;
 	virtual bool GetBounds(const FTypedElementHandle& InElementHandle, FBoxSphereBounds& OutBounds) override;
+	virtual bool CanMoveElement(const FTypedElementHandle& InElementHandle, const ETypedElementWorldType InWorldType) override;
 	virtual bool GetWorldTransform(const FTypedElementHandle& InElementHandle, FTransform& OutTransform) override;
 	virtual bool SetWorldTransform(const FTypedElementHandle& InElementHandle, const FTransform& InTransform) override;
 	virtual bool GetRelativeTransform(const FTypedElementHandle& InElementHandle, FTransform& OutTransform) override;
diff --git a/Engine/Source/Runtime/Engine/Public/Elements/SMInstance/SMInstanceManager.h b/Engine/Source/Runtime/Engine/Public/Elements/SMInstance/SMInstanceManager.h
--- a/Engine/Source/Runtime/Engine/Public/Elements/SMInstance/SMInstanceManager.h
+++ b/Engine/Source/Runtime/Engine/Public/Elements/SMInstance/SMInstanceManager.h
@@ -8,6 +8,8 @@
 #include "Elements/SMInstance/SMInstanceElementId.h"
 #include "SMInstanceManager.generated.h"
 
+enum class ETypedElementWorldType : uint8;
+
 /**
  * An interface for actors that manage static mesh instances.
  * This exists so that actors that directly manage instances can inject custom logic around the manipulation of the underlying ISM component.
@@ -29,6 +31,12 @@ public:
 	 */
 	virtual bool CanEditSMInstance(const FSMInstanceId& InstanceId) const = 0;
 
+	/**
+	 * Can the given static mesh instance be moved in the world?
+	 * @return True if it can be moved, false otherwise.
+	 */
+	virtual bool CanMoveSMInstance(const FSMInstanceId& InstanceId, const ETypedElementWorldType InWorldType) const = 0;
+
 	/**
 	 * Attempt to get the transform of the given static mesh instance.
 	 * @note The transform should be in the local space of the owner ISM component, unless bWorldSpace is set.
@@ -164,6 +172,7 @@ public:
 
 	//~ ISMInstanceManager interface
 	bool CanEditSMInstance() const { return InstanceManager->CanEditSMInstance(InstanceId); }
+	bool CanMoveSMInstance(const ETypedElementWorldType InWorldType) const { return InstanceManager->CanMoveSMInstance(InstanceId, InWorldType); }
 	bool GetSMInstanceTransform(FTransform& OutInstanceTransform, bool bWorldSpace = false) const { return InstanceManager->GetSMInstanceTransform(InstanceId, OutInstanceTransform, bWorldSpace); }
 	bool SetSMInstanceTransform(const FTransform& InstanceTransform, bool bWorldSpace = false, bool bMarkRenderStateDirty = false, bool bTeleport = false) const { return InstanceManager->SetSMInstanceTransform(InstanceId, InstanceTransform, bWorldSpace, bMarkRenderStateDirty, bTeleport); }
 	void NotifySMInstanceMovementStarted() { return InstanceManager->NotifySMInstanceMovementStarted(InstanceId); }
diff --git a/Engine/Source/Runtime/Foliage/Private/InstancedFoliage.cpp b/Engine/Source/Runtime/Foliage/Private/InstancedFoliage.cpp
--- a/Engine/Source/Runtime/Foliage/Private/InstancedFoliage.cpp
+++ b/Engine/Source/Runtime/Foliage/Private/InstancedFoliage.cpp
@@ -5196,6 +5196,18 @@ bool AInstancedFoliageActor::CanEditSMInstance(const FSMInstanceId& InstanceId)
 #endif
 }
 
+bool AInstancedFoliageActor::CanMoveSMInstance(const FSMInstanceId& InstanceId, const ETypedElementWorldType InWorldType) const
+{
+#if WITH_EDITOR
+	const FFoliageInfo* FoliageInfo = GetFoliageInfoForSMInstance(InstanceId);
+	check(FoliageInfo);
+
+	return true;
+#else
+	return false;
+#endif
+}
+
 bool AInstancedFoliageActor::GetSMInstanceTransform(const FSMInstanceId& InstanceId, FTransform& OutInstanceTransform, bool bWorldSpace) const
 {
 #if WITH_EDITOR
diff --git a/Engine/Source/Runtime/Foliage/Public/InstancedFoliageActor.h b/Engine/Source/Runtime/Foliage/Public/InstancedFoliageActor.h
--- a/Engine/Source/Runtime/Foliage/Public/InstancedFoliageActor.h
+++ b/Engine/Source/Runtime/Foliage/Public/InstancedFoliageActor.h
@@ -62,6 +62,7 @@ protected:
 
 	//~ ISMInstanceManager interface
 	virtual bool CanEditSMInstance(const FSMInstanceId& InstanceId) const override;
+	virtual bool CanMoveSMInstance(const FSMInstanceId& InstanceId, const ETypedElementWorldType InWorldType) const override;
 	virtual bool GetSMInstanceTransform(const FSMInstanceId& InstanceId, FTransform& OutInstanceTransform, bool bWorldSpace = false) const override;
 	virtual bool SetSMInstanceTransform(const FSMInstanceId& InstanceId, const FTransform& InstanceTransform, bool bWorldSpace = false, bool bMarkRenderStateDirty = false, bool bTeleport = false) override;
 	virtual void NotifySMInstanceMovementStarted(const FSMInstanceId& InstanceId) override;
diff --git a/Engine/Source/Runtime/TypedElementRuntime/Private/Elements/Framework/TypedElementSelectionSet.cpp b/Engine/Source/Runtime/TypedElementRuntime/Private/Elements/Framework/TypedElementSelectionSet.cpp
--- a/Engine/Source/Runtime/TypedElementRuntime/Private/Elements/Framework/TypedElementSelectionSet.cpp
+++ b/Engine/Source/Runtime/TypedElementRuntime/Private/Elements/Framework/TypedElementSelectionSet.cpp
@@ -2,10 +2,31 @@
 
 #include "Elements/Framework/TypedElementSelectionSet.h"
 #include "Elements/Framework/TypedElementRegistry.h"
+#include "Elements/Interfaces/TypedElementHierarchyInterface.h"
 
+#include "UObject/GCObjectScopeGuard.h"
 #include "Serialization/ObjectReader.h"
 #include "Serialization/ObjectWriter.h"
 
+void FTypedElementSelectionCustomization::GetNormalizedElements(const TTypedElement<UTypedElementSelectionInterface>& InElementSelectionHandle, const UTypedElementList* InSelectionSet, const FTypedElementSelectionNormalizationOptions& InNormalizationOptions, UTypedElementList* OutNormalizedElements)
+{
+	// Don't include this element in the normalized selection if it has a parent element that is also selected
+	{
+		FTypedElementHandle ElementToTest = InElementSelectionHandle;
+		while (TTypedElement<UTypedElementHierarchyInterface> ElementHierarchyHandle = InSelectionSet->GetElement<UTypedElementHierarchyInterface>(ElementToTest))
+		{
+			ElementToTest = ElementHierarchyHandle.GetParentElement();
+			if (ElementToTest && InSelectionSet->Contains(ElementToTest))
+			{
+				return;
+			}
+		}
+	}
+
+	OutNormalizedElements->Add(InElementSelectionHandle);
+}
+
+
 UTypedElementSelectionSet::UTypedElementSelectionSet()
 {
 	if (!HasAnyFlags(RF_ClassDefaultObject))
@@ -144,8 +165,33 @@ bool UTypedElementSelectionSet::CanDeselectElement(const FTypedElementHandle& In
 
 bool UTypedElementSelectionSet::SelectElement(const FTypedElementHandle& InElementHandle, const FTypedElementSelectionOptions InSelectionOptions)
 {
-	FTypedElementSelectionSetElement SelectionSetElement = ResolveSelectionSetElement(InElementHandle);
-	return SelectionSetElement && SelectionSetElement.CanSelectElement(InSelectionOptions) && SelectionSetElement.SelectElement(InSelectionOptions);
+	bool bSelectionChanged = false;
+
+	if (FTypedElementSelectionSetElement SelectionSetElement = ResolveSelectionSetElement(InElementHandle))
+	{
+		if (SelectionSetElement.CanSelectElement(InSelectionOptions))
+		{
+			bSelectionChanged |= SelectionSetElement.SelectElement(InSelectionOptions);
+		}
+	}
+
+	if (InSelectionOptions.GetChildElementInclusionMethod() != ETypedElementChildInclusionMethod::None)
+	{
+		if (TTypedElement<UTypedElementHierarchyInterface> ElementHierarchyHandle = ElementList->GetElement<UTypedElementHierarchyInterface>(InElementHandle))
+		{
+			const FTypedElementSelectionOptions ChildSelectionOptions = FTypedElementSelectionOptions(InSelectionOptions)
+				.SetChildElementInclusionMethod(InSelectionOptions.GetChildElementInclusionMethod() == ETypedElementChildInclusionMethod::Immediate ? ETypedElementChildInclusionMethod::None : InSelectionOptions.GetChildElementInclusionMethod());
+
+			TArray<FTypedElementHandle> ChildElementHandles;
+			ElementHierarchyHandle.GetChildElements(ChildElementHandles, /*bAllowCreate*/true);
+			for (const FTypedElementHandle& ChildElementHandle : ChildElementHandles)
+			{
+				SelectElement(ChildElementHandle, ChildSelectionOptions);
+			}
+		}
+	}
+
+	return bSelectionChanged;
 }
 
 bool UTypedElementSelectionSet::SelectElements(const TArray<FTypedElementHandle>& InElementHandles, const FTypedElementSelectionOptions InSelectionOptions)
@@ -169,8 +215,33 @@ bool UTypedElementSelectionSet::SelectElements(TArrayView<const FTypedElementHan
 
 bool UTypedElementSelectionSet::DeselectElement(const FTypedElementHandle& InElementHandle, const FTypedElementSelectionOptions InSelectionOptions)
 {
-	FTypedElementSelectionSetElement SelectionSetElement = ResolveSelectionSetElement(InElementHandle);
-	return SelectionSetElement && SelectionSetElement.CanDeselectElement(InSelectionOptions) && SelectionSetElement.DeselectElement(InSelectionOptions);
+	bool bSelectionChanged = false;
+
+	if (FTypedElementSelectionSetElement SelectionSetElement = ResolveSelectionSetElement(InElementHandle))
+	{
+		if (SelectionSetElement.CanDeselectElement(InSelectionOptions))
+		{
+			bSelectionChanged |= SelectionSetElement.DeselectElement(InSelectionOptions);
+		}
+	}
+
+	if (InSelectionOptions.GetChildElementInclusionMethod() != ETypedElementChildInclusionMethod::None)
+	{
+		if (TTypedElement<UTypedElementHierarchyInterface> ElementHierarchyHandle = ElementList->GetElement<UTypedElementHierarchyInterface>(InElementHandle))
+		{
+			const FTypedElementSelectionOptions ChildSelectionOptions = FTypedElementSelectionOptions(InSelectionOptions)
+				.SetChildElementInclusionMethod(InSelectionOptions.GetChildElementInclusionMethod() == ETypedElementChildInclusionMethod::Immediate ? ETypedElementChildInclusionMethod::None : InSelectionOptions.GetChildElementInclusionMethod());
+
+			TArray<FTypedElementHandle> ChildElementHandles;
+			ElementHierarchyHandle.GetChildElements(ChildElementHandles, /*bAllowCreate*/false);
+			for (const FTypedElementHandle& ChildElementHandle : ChildElementHandles)
+			{
+				DeselectElement(ChildElementHandle, ChildSelectionOptions);
+			}
+		}
+	}
+
+	return bSelectionChanged;
 }
 
 bool UTypedElementSelectionSet::DeselectElements(const TArray<FTypedElementHandle>& InElementHandles, const FTypedElementSelectionOptions InSelectionOptions)
@@ -251,6 +322,37 @@ FTypedElementHandle UTypedElementSelectionSet::GetSelectionElement(const FTypedE
 	return SelectionSetElement ? SelectionSetElement.GetSelectionElement(InSelectionMethod) : FTypedElementHandle();
 }
 
+UTypedElementList* UTypedElementSelectionSet::GetNormalizedSelection(const FTypedElementSelectionNormalizationOptions InNormalizationOptions) const
+{
+	return GetNormalizedElementList(ElementList, InNormalizationOptions);
+}
+
+void UTypedElementSelectionSet::GetNormalizedSelection(const FTypedElementSelectionNormalizationOptions& InNormalizationOptions, UTypedElementList* OutNormalizedElements) const
+{
+	GetNormalizedElementList(ElementList, InNormalizationOptions, OutNormalizedElements);
+}
+
+UTypedElementList* UTypedElementSelectionSet::GetNormalizedElementList(const UTypedElementList* InElementList, const FTypedElementSelectionNormalizationOptions InNormalizationOptions) const
+{
+	UTypedElementRegistry* Registry = UTypedElementRegistry::GetInstance();
+	TGCObjectScopeGuard<UTypedElementList> NormalizedElementList(Registry->CreateElementList());
+	GetNormalizedElementList(InElementList, InNormalizationOptions, NormalizedElementList.Get());
+	return NormalizedElementList.Get();
+}
+
+void UTypedElementSelectionSet::GetNormalizedElementList(const UTypedElementList* InElementList, const FTypedElementSelectionNormalizationOptions& InNormalizationOptions, UTypedElementList* OutNormalizedElements) const
+{
+	OutNormalizedElements->Reset();
+	InElementList->ForEachElement<UTypedElementSelectionInterface>([this, InElementList, &InNormalizationOptions, OutNormalizedElements](const TTypedElement<UTypedElementSelectionInterface>& InElementSelectionHandle)
+	{
+		// Note: FTypedElementSelectionSetElement needs a non-const UTypedElementList, but we won't call anything that will modify it...
+		FTypedElementSelectionSetElement SelectionSetElement(InElementSelectionHandle, const_cast<UTypedElementList*>(InElementList), GetInterfaceCustomizationByTypeId(InElementSelectionHandle.GetId().GetTypeId()));
+		check(SelectionSetElement.IsSet());
+		SelectionSetElement.GetNormalizedElements(InNormalizationOptions, OutNormalizedElements);
+		return true;
+	});
+}
+
 FTypedElementSelectionSetState UTypedElementSelectionSet::GetCurrentSelectionState() const
 {
 	FTypedElementSelectionSetState CurrentState;
diff --git a/Engine/Source/Runtime/TypedElementRuntime/Public/Elements/Framework/TypedElementSelectionSet.h b/Engine/Source/Runtime/TypedElementRuntime/Public/Elements/Framework/TypedElementSelectionSet.h
--- a/Engine/Source/Runtime/TypedElementRuntime/Public/Elements/Framework/TypedElementSelectionSet.h
+++ b/Engine/Source/Runtime/TypedElementRuntime/Public/Elements/Framework/TypedElementSelectionSet.h
@@ -8,6 +8,26 @@
 #include "Elements/Interfaces/TypedElementSelectionInterface.h"
 #include "TypedElementSelectionSet.generated.h"
 
+USTRUCT(BlueprintType)
+struct FTypedElementSelectionNormalizationOptions
+{
+	GENERATED_BODY()
+
+public:
+	FTypedElementSelectionNormalizationOptions& SetExpandGroups(const bool InExpandGroups) { bExpandGroups = InExpandGroups; return *this; }
+	bool ExpandGroups() const { return bExpandGroups; }
+
+	FTypedElementSelectionNormalizationOptions& SetFollowAttachment(const bool InFollowAttachment) { bFollowAttachment = InFollowAttachment; return *this; }
+	bool FollowAttachment() const { return bFollowAttachment; }
+	
+private:
+	UPROPERTY(BlueprintReadWrite, Category="TypedElementInterfaces|Selection|NormalizationOptions", meta=(AllowPrivateAccess=true))
+	bool bExpandGroups = false;
+
+	UPROPERTY(BlueprintReadWrite, Category="TypedElementInterfaces|Selection|NormalizationOptions", meta=(AllowPrivateAccess=true))
+	bool bFollowAttachment = false;
+};
+
 /**
  * Customization type used to allow asset editors (such as the level editor) to override the base behavior of element selection,
  * by injecting extra pre/post selection logic around the call into the selection interface for an element type.
@@ -25,6 +45,7 @@ public:
 	virtual bool DeselectElement(const TTypedElement<UTypedElementSelectionInterface>& InElementSelectionHandle, UTypedElementList* InSelectionSet, const FTypedElementSelectionOptions& InSelectionOptions) { return InElementSelectionHandle.DeselectElement(InSelectionSet, InSelectionOptions); }
 	virtual bool AllowSelectionModifiers(const TTypedElement<UTypedElementSelectionInterface>& InElementSelectionHandle, const UTypedElementList* InSelectionSet) { return InElementSelectionHandle.AllowSelectionModifiers(InSelectionSet); }
 	virtual FTypedElementHandle GetSelectionElement(const TTypedElement<UTypedElementSelectionInterface>& InElementSelectionHandle, const UTypedElementList* InCurrentSelection, const ETypedElementSelectionMethod InSelectionMethod) { return InElementSelectionHandle.GetSelectionElement(InCurrentSelection, InSelectionMethod); }
+	virtual void GetNormalizedElements(const TTypedElement<UTypedElementSelectionInterface>& InElementSelectionHandle, const UTypedElementList* InSelectionSet, const FTypedElementSelectionNormalizationOptions& InNormalizationOptions, UTypedElementList* OutNormalizedElements);
 };
 
 /**
@@ -68,6 +89,7 @@ public:
 	bool DeselectElement(const FTypedElementSelectionOptions& InSelectionOptions) const { return SelectionCustomization->DeselectElement(ElementSelectionHandle, ElementList, InSelectionOptions); }
 	bool AllowSelectionModifiers() const { return SelectionCustomization->AllowSelectionModifiers(ElementSelectionHandle, ElementList); }
 	FTypedElementHandle GetSelectionElement(const ETypedElementSelectionMethod InSelectionMethod) const { return SelectionCustomization->GetSelectionElement(ElementSelectionHandle, ElementList, InSelectionMethod); }
+	void GetNormalizedElements(const FTypedElementSelectionNormalizationOptions& InNormalizationOptions, UTypedElementList* OutNormalizedElements) const { return SelectionCustomization->GetNormalizedElements(ElementSelectionHandle, ElementList, InNormalizationOptions, OutNormalizedElements); }
 
 private:
 	TTypedElement<UTypedElementSelectionInterface> ElementSelectionHandle;
@@ -215,6 +237,22 @@ public:
 	UFUNCTION(BlueprintPure, Category="TypedElementFramework|Selection")
 	FTypedElementHandle GetSelectionElement(const FTypedElementHandle& InElementHandle, const ETypedElementSelectionMethod InSelectionMethod) const;
 
+	/**
+	 * Get a normalized version of this selection set that can be used to perform operations like gizmo manipulation, deletion, copying, etc.
+	 * This will do things like expand out groups, and resolve any parent<->child elements so that duplication operations aren't performed on both the parent and the child.
+	 */
+	UFUNCTION(BlueprintPure, Category="TypedElementFramework|Selection")
+	UTypedElementList* GetNormalizedSelection(const FTypedElementSelectionNormalizationOptions InNormalizationOptions) const;
+	void GetNormalizedSelection(const FTypedElementSelectionNormalizationOptions& InNormalizationOptions, UTypedElementList* OutNormalizedElements) const;
+
+	/**
+	 * Get a normalized version of the given element list that can be used to perform operations like gizmo manipulation, deletion, copying, etc.
+	 * This will do things like expand out groups, and resolve any parent<->child elements so that duplication operations aren't performed on both the parent and the child.
+	 */
+	UFUNCTION(BlueprintPure, Category="TypedElementFramework|Selection")
+	UTypedElementList* GetNormalizedElementList(const UTypedElementList* InElementList, const FTypedElementSelectionNormalizationOptions InNormalizationOptions) const;
+	void GetNormalizedElementList(const UTypedElementList* InElementList, const FTypedElementSelectionNormalizationOptions& InNormalizationOptions, UTypedElementList* OutNormalizedElements) const;
+
 	/**
 	 * Get the number of selected elements.
 	 */
diff --git a/Engine/Source/Runtime/TypedElementRuntime/Public/Elements/Interfaces/TypedElementHierarchyInterface.h b/Engine/Source/Runtime/TypedElementRuntime/Public/Elements/Interfaces/TypedElementHierarchyInterface.h
new file mode 100644
--- /dev/null
+++ b/Engine/Source/Runtime/TypedElementRuntime/Public/Elements/Interfaces/TypedElementHierarchyInterface.h
@@ -0,0 +1,36 @@
+// Copyright Epic Games, Inc. All Rights Reserved.
+
+#pragma once
+
+#include "Elements/Framework/TypedElementHandle.h"
+#include "TypedElementHierarchyInterface.generated.h"
+
+UCLASS(Abstract)
+class TYPEDELEMENTRUNTIME_API UTypedElementHierarchyInterface : public UTypedElementInterface
+{
+	GENERATED_BODY()
+
+public:
+	/**
+	 * Get the logical parent of this element, if any.
+	 * eg) A component might return its actor, or a static mesh instance might return its ISM component.
+	 */
+	UFUNCTION(BlueprintPure, Category="TypedElementInterfaces|Hierarchy")
+	virtual FTypedElementHandle GetParentElement(const FTypedElementHandle& InElementHandle, const bool bAllowCreate = true) { return FTypedElementHandle(); }
+
+	/**
+	 * Get the logical children of this element, if any.
+	 * eg) An actor might return its component, or an ISM component might return its static mesh instances.
+	 *
+	 * @note Appends to OutElementHandles.
+	 */
+	UFUNCTION(BlueprintPure, Category="TypedElementInterfaces|Hierarchy")
+	virtual void GetChildElements(const FTypedElementHandle& InElementHandle, TArray<FTypedElementHandle>& OutElementHandles, const bool bAllowCreate = true) {}
+};
+
+template <>
+struct TTypedElement<UTypedElementHierarchyInterface> : public TTypedElementBase<UTypedElementHierarchyInterface>
+{
+	FTypedElementHandle GetParentElement(const bool bAllowCreate = true) const { return InterfacePtr->GetParentElement(*this, bAllowCreate); }
+	void GetChildElements(TArray<FTypedElementHandle>& OutElementHandles, const bool bAllowCreate = true) const { return InterfacePtr->GetChildElements(*this, OutElementHandles, bAllowCreate); }
+};
diff --git a/Engine/Source/Runtime/TypedElementRuntime/Public/Elements/Interfaces/TypedElementSelectionInterface.h b/Engine/Source/Runtime/TypedElementRuntime/Public/Elements/Interfaces/TypedElementSelectionInterface.h
--- a/Engine/Source/Runtime/TypedElementRuntime/Public/Elements/Interfaces/TypedElementSelectionInterface.h
+++ b/Engine/Source/Runtime/TypedElementRuntime/Public/Elements/Interfaces/TypedElementSelectionInterface.h
@@ -10,12 +10,23 @@ class UTypedElementList;
 UENUM()
 enum class ETypedElementSelectionMethod : uint8
 {
-	/** Select the "primary" element (eg, a component favor selecting its owner actor) */
+	/** Select the "primary" element (eg, a component would favor selecting its owner actor) */
 	Primary,
 	/** Select the "secondary" element (eg, a component would favor selecting itself) */
 	Secondary,
 };
 
+UENUM()
+enum class ETypedElementChildInclusionMethod : uint8
+{
+	/** Do not include child elements */
+	None,
+	/** Include the immediate child elements, but do not recurse */
+	Immediate,
+	/** Include the immediate child elements, and recurse into their children too */
+	Recursive,
+};
+
 USTRUCT(BlueprintType)
 struct FTypedElementIsSelectedOptions
 {
@@ -47,6 +58,9 @@ public:
 
 	FTypedElementSelectionOptions& SetWarnIfLocked(const bool InWarnIfLocked) { bWarnIfLocked = InWarnIfLocked; return *this; }
 	bool WarnIfLocked() const { return bWarnIfLocked; }
+
+	FTypedElementSelectionOptions& SetChildElementInclusionMethod(const ETypedElementChildInclusionMethod InChildElementInclusionMethod) { ChildElementInclusionMethod = InChildElementInclusionMethod; return *this; }
+	ETypedElementChildInclusionMethod GetChildElementInclusionMethod() const { return ChildElementInclusionMethod; }
 	
 private:
 	UPROPERTY(BlueprintReadWrite, Category="TypedElementInterfaces|Selection|SelectionOptions", meta=(AllowPrivateAccess=true))
@@ -60,6 +74,9 @@ private:
 
 	UPROPERTY(BlueprintReadWrite, Category="TypedElementInterfaces|Selection|SelectionOptions", meta=(AllowPrivateAccess=true))
 	bool bWarnIfLocked = false;
+
+	UPROPERTY(BlueprintReadWrite, Category="TypedElementInterfaces|Selection|SelectionOptions", meta=(AllowPrivateAccess=true))
+	ETypedElementChildInclusionMethod ChildElementInclusionMethod = ETypedElementChildInclusionMethod::None;
 };
 
 class ITypedElementTransactedElement
