Texture: Read texture sources directly for the new texture build flow

#rnx

diff --git a/Engine/Source/Developer/TextureBuild/Private/TextureBuildFunction.cpp b/Engine/Source/Developer/TextureBuild/Private/TextureBuildFunction.cpp
--- a/Engine/Source/Developer/TextureBuild/Private/TextureBuildFunction.cpp
+++ b/Engine/Source/Developer/TextureBuild/Private/TextureBuildFunction.cpp
@@ -167,24 +167,24 @@ static ERawImageFormat::Type ComputeRawImageFormat(ETextureSourceFormat SourceFo
 	}
 }
 
-static bool TryReadTextureSourceFromCompactBinary(const FCbObjectView& Object, UE::DerivedData::FBuildContext& Context, TArray<FImage>& OutMips)
+static bool TryReadTextureSourceFromCompactBinary(FCbFieldView Source, UE::DerivedData::FBuildContext& Context, TArray<FImage>& OutMips)
 {
-	FSharedBuffer InputBuffer = Context.FindInput(Object["Input"].AsUuid().ToString());
+	FSharedBuffer InputBuffer = Context.FindInput(FUTF8ToTCHAR(Source.GetName()));
 	if (!InputBuffer)
 	{
-		UE_LOG(LogTextureBuildFunction, Error, TEXT("Missing input %s."), *Object["Input"].AsUuid().ToString());
+		UE_LOG(LogTextureBuildFunction, Error, TEXT("Missing input %.*s."), Source.GetName().Len(), Source.GetName().GetData());
 		return false;
 	}
 
-	ETextureSourceCompressionFormat CompressionFormat = (ETextureSourceCompressionFormat)Object["CompressionFormat"].AsUInt8();
-	ETextureSourceFormat SourceFormat = (ETextureSourceFormat)Object["SourceFormat"].AsUInt8();
+	ETextureSourceCompressionFormat CompressionFormat = (ETextureSourceCompressionFormat)Source["CompressionFormat"].AsUInt8();
+	ETextureSourceFormat SourceFormat = (ETextureSourceFormat)Source["SourceFormat"].AsUInt8();
 
 	ERawImageFormat::Type RawImageFormat = ComputeRawImageFormat(SourceFormat);
 
-	EGammaSpace GammaSpace = (EGammaSpace)Object["GammaSpace"].AsUInt8();
-	int32 NumSlices = Object["NumSlices"].AsInt32();
-	int32 SizeX = Object["SizeX"].AsInt32();
-	int32 SizeY = Object["SizeY"].AsInt32();
+	EGammaSpace GammaSpace = (EGammaSpace)Source["GammaSpace"].AsUInt8();
+	int32 NumSlices = Source["NumSlices"].AsInt32();
+	int32 SizeX = Source["SizeX"].AsInt32();
+	int32 SizeY = Source["SizeY"].AsInt32();
 	int32 MipSizeX = SizeX;
 	int32 MipSizeY = SizeY;
 
@@ -217,7 +217,7 @@ static bool TryReadTextureSourceFromCompactBinary(const FCbObjectView& Object, U
 		InputBuffer.Reset();
 	}
 
-	FCbArrayView MipsCbArrayView = Object["Mips"].AsArrayView();
+	FCbArrayView MipsCbArrayView = Source["Mips"].AsArrayView();
 	OutMips.Reserve(MipsCbArrayView.Num());
 	for (FCbFieldView MipsCbArrayIt : MipsCbArrayView)
 	{
@@ -288,13 +288,13 @@ void FTextureBuildFunction::Build(UE::DerivedData::FBuildContext& Context) const
 	ReadOutputSettingsFromCompactBinary(Settings["Output"].AsObjectView(), NumInlineMips, MipKeyPrefix);
 
 	TArray<FImage> SourceMips;
-	if (!TryReadTextureSourceFromCompactBinary(Settings["Source"].AsObjectView(), Context, SourceMips))
+	if (!TryReadTextureSourceFromCompactBinary(Settings["Source"], Context, SourceMips))
 	{
 		return;
 	}
 
 	TArray<FImage> AssociatedNormalSourceMips;
-	if (const FCbObjectView CompositeSource = Settings["CompositeSource"].AsObjectView();
+	if (FCbFieldView CompositeSource = Settings["CompositeSource"];
 		CompositeSource && !TryReadTextureSourceFromCompactBinary(CompositeSource, Context, AssociatedNormalSourceMips))
 	{
 		return;
diff --git a/Engine/Source/Runtime/Engine/Private/TextureDerivedDataBuildUtils.cpp b/Engine/Source/Runtime/Engine/Private/TextureDerivedDataBuildUtils.cpp
--- a/Engine/Source/Runtime/Engine/Private/TextureDerivedDataBuildUtils.cpp
+++ b/Engine/Source/Runtime/Engine/Private/TextureDerivedDataBuildUtils.cpp
@@ -183,9 +183,8 @@ static void WriteSource(FCbWriter& Writer, const UTexture& Texture, int32 LayerI
 
 	Writer.BeginObject();
 
-	Writer.AddUuid("Input", Source.GetId());
 	Writer.AddInteger("CompressionFormat", Source.GetSourceCompression());
-	Writer.AddInteger("SourceFormat", Source.GetFormat());
+	Writer.AddInteger("SourceFormat", Source.GetFormat(LayerIndex));
 	Writer.AddInteger("GammaSpace", static_cast<uint8>(GammaSpace));
 	Writer.AddInteger("NumSlices", Source.GetNumSlices());
 	Writer.AddInteger("SizeX", Source.GetSizeX());
diff --git a/Engine/Source/Runtime/Engine/Private/TextureDerivedDataTask.cpp b/Engine/Source/Runtime/Engine/Private/TextureDerivedDataTask.cpp
--- a/Engine/Source/Runtime/Engine/Private/TextureDerivedDataTask.cpp
+++ b/Engine/Source/Runtime/Engine/Private/TextureDerivedDataTask.cpp
@@ -447,20 +447,36 @@ void FTextureCacheDerivedDataWorker::BuildTexture(bool bReplaceExistingDDC)
 
 			DefinitionBuilder.AddConstant(TEXT("Settings"_SV),
 				SaveTextureBuildSettings(KeySuffix, Texture, BuildSettingsPerLayer[0], 0, NUM_INLINE_DERIVED_MIPS));
-
-			for (const FBuildInputRecord& InputRecord : BuildInputRecords)
+			DefinitionBuilder.AddInputBulkData(TEXT("Source"_SV), Texture.Source.GetId());
+			if (Texture.CompositeTexture)
 			{
-				DefinitionBuilder.AddInputBulkData(InputRecord.Id.ToString(), InputRecord.Id);
+				DefinitionBuilder.AddInputBulkData(TEXT("CompositeSource"_SV), Texture.CompositeTexture->Source.GetId());
 			}
 
 			class FTextureBuildInputResolver final : public IBuildInputResolver
 			{
 			public:
-				explicit FTextureBuildInputResolver(TConstArrayView<FBuildInputRecord> InRecords)
-					: Records(InRecords)
+				explicit FTextureBuildInputResolver(UTexture& InTexture)
+					: Texture(InTexture)
 				{
 				}
 
+				const FCompressedBuffer& FindSource(FCompressedBuffer& Buffer, FTextureSource& Source, const FGuid& BulkDataId)
+				{
+					if (Source.GetId() != BulkDataId)
+					{
+						return FCompressedBuffer::Null;
+					}
+					if (!Buffer)
+					{
+						Source.OperateOnLoadedBulkData([&Buffer](const FSharedBuffer& BulkDataBuffer)
+						{
+							Buffer = FCompressedBuffer::Compress(NAME_Default, BulkDataBuffer);
+						});
+					}
+					return Buffer;
+				}
+
 				FRequest ResolveInputMeta(
 					const FBuildDefinition& Definition,
 					EPriority Priority,
@@ -471,10 +487,13 @@ void FTextureCacheDerivedDataWorker::BuildTexture(bool bReplaceExistingDDC)
 					TArray<FBuildInputMetaByKey> Inputs;
 					Definition.IterateInputBulkData([this, &Status, &InputKeys, &Inputs](FStringView Key, const FGuid& BulkDataId)
 					{
-						if (const FBuildInputRecord* Match = Algo::FindBy(Records, BulkDataId, &FBuildInputRecord::Id))
+						const FCompressedBuffer& Buffer = Key == TEXT("Source"_SV)
+							? FindSource(SourceBuffer, Texture.Source, BulkDataId)
+							: FindSource(CompositeSourceBuffer, Texture.CompositeTexture->Source, BulkDataId);
+						if (Buffer)
 						{
 							InputKeys.Emplace(Key);
-							Inputs.Add({InputKeys.Last(), Match->Data.GetRawHash(), Match->Data.GetRawSize()});
+							Inputs.Add({InputKeys.Last(), Buffer.GetRawHash(), Buffer.GetRawSize()});
 						}
 						else
 						{
@@ -498,10 +517,13 @@ void FTextureCacheDerivedDataWorker::BuildTexture(bool bReplaceExistingDDC)
 					{
 						if (!Filter || Filter(Key))
 						{
-							if (const FBuildInputRecord* Match = Algo::FindBy(Records, BulkDataId, &FBuildInputRecord::Id))
+							const FCompressedBuffer& Buffer = Key == TEXT("Source"_SV)
+								? FindSource(SourceBuffer, Texture.Source, BulkDataId)
+								: FindSource(CompositeSourceBuffer, Texture.CompositeTexture->Source, BulkDataId);
+							if (Buffer)
 							{
 								InputKeys.Emplace(Key);
-								Inputs.Add({InputKeys.Last(), Match->Data});
+								Inputs.Add({InputKeys.Last(), Buffer});
 							}
 							else
 							{
@@ -514,10 +536,12 @@ void FTextureCacheDerivedDataWorker::BuildTexture(bool bReplaceExistingDDC)
 				}
 
 			private:
-				TConstArrayView<FBuildInputRecord> Records;
+				UTexture& Texture;
+				FCompressedBuffer SourceBuffer;
+				FCompressedBuffer CompositeSourceBuffer;
 			};
 
-			FTextureBuildInputResolver InputResolver(BuildInputRecords);
+			FTextureBuildInputResolver InputResolver(Texture);
 			FBuildSession Session = Build.CreateSession(TexturePath, &InputResolver);
 			Session.Build(DefinitionBuilder.Build(), EBuildPolicy::Default, EPriority::Blocking,
 				[this, &TexturePath, bReplaceExistingDDC] (FBuildCompleteParams&& Params)
@@ -846,7 +870,6 @@ void FTextureCacheDerivedDataWorker::DoWork()
 		{
 			TextureData.GetSourceMips(Texture.Source, ImageWrapper);
 			BuildExporter.ExportTextureSourceBulkData(Texture.Source);
-			ConditionalAddBuildFunctionInput(Texture.Source);
 			bHasTextureSourceMips = true;
 		}
 
@@ -855,7 +878,6 @@ void FTextureCacheDerivedDataWorker::DoWork()
 		{
 			CompositeTextureData.GetSourceMips(Texture.CompositeTexture->Source, ImageWrapper);
 			BuildExporter.ExportCompositeTextureSourceBulkData(Texture.CompositeTexture->Source);
-			ConditionalAddBuildFunctionInput(Texture.CompositeTexture->Source);
 			bHasCompositeTextureSourceMips = true;
 		}
 
@@ -863,7 +885,6 @@ void FTextureCacheDerivedDataWorker::DoWork()
 		{
 			TextureData.GetAsyncSourceMips(ImageWrapper);
 			BuildExporter.ExportTextureSourceBulkData(TextureData.AsyncSource);
-			ConditionalAddBuildFunctionInput(TextureData.AsyncSource);
 			TextureData.AsyncSource.RemoveBulkData();
 		}
 
@@ -873,7 +894,6 @@ void FTextureCacheDerivedDataWorker::DoWork()
 			if ((bool)Texture.CompositeTexture)
 			{
 				BuildExporter.ExportCompositeTextureSourceBulkData(CompositeTextureData.AsyncSource);
-				ConditionalAddBuildFunctionInput(CompositeTextureData.AsyncSource);
 			}
 			CompositeTextureData.AsyncSource.RemoveBulkData();
 		}
diff --git a/Engine/Source/Runtime/Engine/Private/TextureDerivedDataTask.h b/Engine/Source/Runtime/Engine/Private/TextureDerivedDataTask.h
--- a/Engine/Source/Runtime/Engine/Private/TextureDerivedDataTask.h
+++ b/Engine/Source/Runtime/Engine/Private/TextureDerivedDataTask.h
@@ -108,12 +108,6 @@ struct FTextureSourceData
  */
 class FTextureCacheDerivedDataWorker : public FNonAbandonableTask
 {
-	struct FBuildInputRecord
-	{
-		FGuid Id;
-		FCompressedBuffer Data;
-	};
-
 	/** Texture compressor module, must be loaded in the game thread. see FModuleManager::WarnIfItWasntSafeToLoadHere() */
 	ITextureCompressorModule* Compressor;
 	/** Image wrapper module, must be loaded in the game thread. see FModuleManager::WarnIfItWasntSafeToLoadHere() */
@@ -132,8 +126,6 @@ class FTextureCacheDerivedDataWorker : public FNonAbandonableTask
 	FTextureSourceData CompositeTextureData;
 	/** DDC2 build function name to use to build this texture (if DDC2 is enabled and the target texture type has a DDC2 build function, empty otherwise) */
 	FString BuildFunctionName;
-	/** DDC2 build input records for putting on the action (if DDC2 is enabled and the target texture type has a DDC2 build function, empty otherwise) */
-	TArray<FBuildInputRecord> BuildInputRecords;
 	/** Exporter that can optionally write out build actions and reference outputs for the purpose of testing remote execution of builds */
 	FTextureDerivedDataBuildExporter BuildExporter;
 	/** Texture cache flags. */
@@ -151,19 +143,6 @@ class FTextureCacheDerivedDataWorker : public FNonAbandonableTask
 	/** Build the texture. This function is safe to call from any thread. */
 	void BuildTexture(bool bReplaceExistingDDC = false);
 
-	void ConditionalAddBuildFunctionInput(FTextureSource& Source)
-	{
-		if (!BuildFunctionName.IsEmpty())
-		{
-			FBuildInputRecord& NewRef = BuildInputRecords.AddDefaulted_GetRef();
-			NewRef.Id = Source.GetId();
-			Source.OperateOnLoadedBulkData([&NewRef](const FSharedBuffer& BulkDataBuffer)
-			{
-				NewRef.Data = FCompressedBuffer::Compress(NAME_Default, BulkDataBuffer);
-			});
-		}
-	}
-
 	void ConsumeBuildFunctionOutput(const UE::DerivedData::FBuildOutput& BuildOutput, const FString& TexturePath, bool bReplaceExistingDDC);
 
 public:
