World Partition HLOD
*Also include the HLOD material in the HLOD hash computation for HLODLayers of type EHLODLayerType::MeshApproximate

diff --git a/Engine/Source/Runtime/Engine/Private/WorldPartition/HLOD/HLODLayer.cpp b/Engine/Source/Runtime/Engine/Private/WorldPartition/HLOD/HLODLayer.cpp
--- a/Engine/Source/Runtime/Engine/Private/WorldPartition/HLOD/HLODLayer.cpp
+++ b/Engine/Source/Runtime/Engine/Private/WorldPartition/HLOD/HLODLayer.cpp
@@ -68,7 +68,7 @@ uint32 UHLODLayer::GetCRC() const
 
 	uint32 Hash = Ar.GetCrc();
 
-	const bool bUseHLODMaterial = LayerType == EHLODLayerType::MeshMerge || LayerType == EHLODLayerType::MeshSimplify;
+	const bool bUseHLODMaterial = LayerType != EHLODLayerType::Instancing;
 	if (bUseHLODMaterial && !HLODMaterial.IsNull())
 	{
 		UMaterialInterface* Material = HLODMaterial.LoadSynchronous();
