Fix external actor packages not getting properly cleaned up if the actor is not present in the level actor list for one reason of another

diff --git a/Engine/Source/Runtime/Engine/Private/Level.cpp b/Engine/Source/Runtime/Engine/Private/Level.cpp
--- a/Engine/Source/Runtime/Engine/Private/Level.cpp
+++ b/Engine/Source/Runtime/Engine/Private/Level.cpp
@@ -372,17 +372,20 @@ void ULevel::CleanupLevel()
 {
 	OnCleanupLevel.Broadcast();
 	// if the level contains any actor with an external package, clear their metadata standalone flag so that the packages can be properly unloaded.
-	for (AActor* Actor : Actors)
-	{
-		if (Actor && Actor->GetExternalPackage())
+	// Do so for any actors outered to the level and not just from the level actors array
+	ForEachObjectWithOuter(this, [](UObject* InObject)
 		{
-			ForEachObjectWithPackage(Actor->GetExternalPackage(), [](UObject* Object)
+			// Currently only actors would have external package, 
+			// but we can ask directly on all objects since the tests validate against object flags before doing a hash lookup
+			if (UPackage* ExternalPackage = InObject->GetExternalPackage())
 			{
-				Object->ClearFlags(RF_Standalone);
-				return true;
-			}, false);
-		}
-	}
+				ForEachObjectWithPackage(ExternalPackage, [](UObject* Object)
+					{
+						Object->ClearFlags(RF_Standalone);
+						return true;
+					}, false);
+			}
+		}, false);
 
 	if (UWorldPartition* WorldPartition = GetWorldPartition())
 	{
