Added per-actor building of texture streaming build data. Primitive's StreamingTextureBuildInfo texture indices will be relative to an array stored in the actor (UActorTextureStreamingBuildDataComponent).
The feature is currently activated for WorldPartitionHLOD actors since the Renderer is initialized when generating HLODs (also required when building texture streaming).
At cook time, this information is then converted to the level of the actor.
On a big project, saves ~50% of FRenderAssetStreamingManager::IncrementalUpdate (inclusive average time goes from 1.2ms down to 0.6ms).

diff --git a/Engine/Source/Runtime/Engine/Classes/Components/InstancedStaticMeshComponent.h b/Engine/Source/Runtime/Engine/Classes/Components/InstancedStaticMeshComponent.h
--- a/Engine/Source/Runtime/Engine/Classes/Components/InstancedStaticMeshComponent.h
+++ b/Engine/Source/Runtime/Engine/Classes/Components/InstancedStaticMeshComponent.h
@@ -216,7 +216,7 @@ class ENGINE_API UInstancedStaticMeshComponent : public UStaticMeshComponent, pu
 	/** Get material, UV density and bounds for a given material index. */
 	virtual bool GetMaterialStreamingData(int32 MaterialIndex, FPrimitiveMaterialInfo& MaterialData) const override;
 	/** Build the data to compute accuracte StreaminTexture data. */
-	virtual bool BuildTextureStreamingData(ETextureStreamingBuildType BuildType, EMaterialQualityLevel::Type QualityLevel, ERHIFeatureLevel::Type FeatureLevel, TSet<FGuid>& DependentResources) override;
+	virtual bool BuildTextureStreamingDataImpl(ETextureStreamingBuildType BuildType, EMaterialQualityLevel::Type QualityLevel, ERHIFeatureLevel::Type FeatureLevel, TSet<FGuid>& DependentResources, bool& bOutSupportsBuildTextureStreamingData) override;
 	/** Get the StreaminTexture data. */
 	virtual void GetStreamingRenderAssetInfo(FStreamingTextureLevelContext& LevelContext, TArray<FStreamingRenderAssetPrimitiveInfo>& OutStreamingRenderAssets) const override;
 
diff --git a/Engine/Source/Runtime/Engine/Classes/Components/PrimitiveComponent.h b/Engine/Source/Runtime/Engine/Classes/Components/PrimitiveComponent.h
--- a/Engine/Source/Runtime/Engine/Classes/Components/PrimitiveComponent.h
+++ b/Engine/Source/Runtime/Engine/Classes/Components/PrimitiveComponent.h
@@ -281,8 +281,17 @@ public:
 	/** Determines how the geometry of a component will be incorporated in proxy (simplified) HLODs. */
 	UPROPERTY(EditAnywhere, AdvancedDisplay, BlueprintReadWrite, Category = HLOD, meta = (DisplayName = "HLOD Batching Policy", EditConditionHides, EditCondition = "bEnableAutoLODGeneration"))
 	EHLODBatchingPolicy HLODBatchingPolicy;
+
+
+	/** Indicates that the texture streaming built data is local to the Actor (see UActorTextureStreamingBuildDataComponent). */
+	UPROPERTY()
+	uint8 bIsActorTextureStreamingBuiltData : 1;
 #endif 
 
+	/** Indicates to the texture streaming wether it can use the pre-built texture streaming data (even if empty). */
+	UPROPERTY()
+	uint8 bIsValidTextureStreamingBuiltData : 1;
+
 	/**
 	 * When enabled this object will not be culled by distance. This is ignored if a child of a HLOD.
 	 */
@@ -1752,6 +1761,12 @@ public:
 
 	/** Add the used GUIDs from UMapBuildDataRegistry::MeshBuildData. Used to preserve hidden level data in lighting scenario. */
 	virtual void AddMapBuildDataGUIDs(TSet<FGuid>& InGUIDs) const {}
+
+	/**
+	 *	Remaps the texture streaming built data that was built for the actor back to the level.
+	 *	
+	 */
+	virtual bool RemapActorTextureStreamingBuiltDataToLevel(const class UActorTextureStreamingBuildDataComponent* InActorTextureBuildData) { return false; }
 #endif // WITH_EDITOR
 
 	/**
@@ -1783,7 +1798,18 @@ public:
 	 *	@param	DependentResources [out]	The resource the build depends on.
 	 *	@return								Returns false if some data needs rebuild but couldn't be rebuilt (because of the build type).
 	 */
-	virtual bool BuildTextureStreamingData(ETextureStreamingBuildType BuildType, EMaterialQualityLevel::Type QualityLevel, ERHIFeatureLevel::Type FeatureLevel, TSet<FGuid>& DependentResources) { return true; }
+	bool BuildTextureStreamingData(ETextureStreamingBuildType BuildType, EMaterialQualityLevel::Type QualityLevel, ERHIFeatureLevel::Type FeatureLevel, TSet<FGuid>& DependentResources);
+
+	/**
+	 *	Component type implementation of updating the streaming data of this component.
+	 *
+	 *	@param	BuildType		[in]		The type of build. Affects what the build is allowed to do.
+	 *	@param	QualityLevel	[in]		The quality level being used in the texture streaming build.
+	 *	@param	FeatureLevel	[in]		The feature level being used in the texture streaming build.
+	 *	@param	DependentResources [out]	The resource the build depends on.
+	 *	@return								Returns false if some data needs rebuild but couldn't be rebuilt (because of the build type).
+	 */
+	virtual bool BuildTextureStreamingDataImpl(ETextureStreamingBuildType BuildType, EMaterialQualityLevel::Type QualityLevel, ERHIFeatureLevel::Type FeatureLevel, TSet<FGuid>& DependentResources, bool& bOutSupportsBuildTextureStreamingData);
 
 	/**
 	 * Determines the DPG the primitive's primary elements are drawn in.
diff --git a/Engine/Source/Runtime/Engine/Classes/Components/StaticMeshComponent.h b/Engine/Source/Runtime/Engine/Classes/Components/StaticMeshComponent.h
--- a/Engine/Source/Runtime/Engine/Classes/Components/StaticMeshComponent.h
+++ b/Engine/Source/Runtime/Engine/Classes/Components/StaticMeshComponent.h
@@ -520,9 +520,12 @@ public:
 	/** Get material, UV density and bounds for a given material index. */
 	virtual bool GetMaterialStreamingData(int32 MaterialIndex, FPrimitiveMaterialInfo& MaterialData) const override;
 	/** Build the data to compute accuracte StreaminTexture data. */
-	virtual bool BuildTextureStreamingData(ETextureStreamingBuildType BuildType, EMaterialQualityLevel::Type QualityLevel, ERHIFeatureLevel::Type FeatureLevel, TSet<FGuid>& DependentResources) override;
+	virtual bool BuildTextureStreamingDataImpl(ETextureStreamingBuildType BuildType, EMaterialQualityLevel::Type QualityLevel, ERHIFeatureLevel::Type FeatureLevel, TSet<FGuid>& DependentResources, bool& bOutSupportsBuildTextureStreamingData) override;
 	/** Get the StreaminTexture data. */
 	virtual void GetStreamingRenderAssetInfo(FStreamingTextureLevelContext& LevelContext, TArray<FStreamingRenderAssetPrimitiveInfo>& OutStreamingRenderAssets) const override;
+#if WITH_EDITOR
+	virtual bool RemapActorTextureStreamingBuiltDataToLevel(const class UActorTextureStreamingBuildDataComponent* InActorTextureBuildData) override;
+#endif 
 
 	virtual class UBodySetup* GetBodySetup() override;
 	virtual bool CanEditSimulatePhysics() override;
diff --git a/Engine/Source/Runtime/Engine/Classes/Engine/Level.h b/Engine/Source/Runtime/Engine/Classes/Engine/Level.h
--- a/Engine/Source/Runtime/Engine/Classes/Engine/Level.h
+++ b/Engine/Source/Runtime/Engine/Classes/Engine/Level.h
@@ -14,6 +14,7 @@
 #include "Templates/ScopedCallback.h"
 #include "Misc/WorldCompositionUtility.h"
 #include "Engine/MaterialMerging.h"
+#include "Engine/TextureStreamingTypes.h"
 #include <atomic>
 
 #include "Level.generated.h"
@@ -436,7 +437,7 @@ enum class EActorsLoadingStrategy: uint8
  * @see UActor
  */
 UCLASS(MinimalAPI)
-class ULevel : public UObject, public IInterface_AssetUserData
+class ULevel : public UObject, public IInterface_AssetUserData, public ITextureStreamingContainer
 {
 	GENERATED_BODY()
 
@@ -532,10 +533,18 @@ public:
 	UPROPERTY()
 	TArray<FVector> StaticNavigableGeometry;
 
-	/** The Guid of each texture refered by FStreamingTextureBuildInfo::TextureLevelIndex	*/
+	/** The Guid of each streamable texture refered by FStreamingTextureBuildInfo::TextureLevelIndex	*/
 	UPROPERTY()
 	TArray<FGuid> StreamingTextureGuids;
 
+	/** The name of each streamable texture referred by FStreamingTextureBuildInfo::TextureLevelIndex */
+	UPROPERTY()
+	TArray<FName> StreamingTextures;
+
+	/** Packed quality level and feature level used when building texture streaming data. This is used by runtime to determine if built data can be used or not. */
+	UPROPERTY()
+	uint32 PackedTextureStreamingQualityLevelFeatureLevel;
+
 	/** Data structures for holding the tick functions **/
 	class FTickTaskLevel*						TickTaskLevel;
 
@@ -828,6 +837,14 @@ public:
 	static void AddReferencedObjects(UObject* InThis, FReferenceCollector& Collector);
 	//~ End UObject Interface.
 
+#if	WITH_EDITOR
+	//~Begin ITextureStreamingContainer Interface.
+	virtual void InitializeTextureStreamingContainer(uint32 InPackedTextureStreamingQualityLevelFeatureLevel) override;
+	virtual uint16 RegisterStreamableTexture(UTexture* InTexture) override;
+	//~End ITextureStreamingContainer Interface.
+	ENGINE_API uint16 RegisterStreamableTexture(const FString& InTextureName, const FGuid& InTextureGuid);
+#endif
+
 	/**
 	 * Flag this level instance for destruction.
 	 * This is called by UWorld::CleanupWorld to flag the level and its owned packages for destruction.
diff --git a/Engine/Source/Runtime/Engine/Classes/Engine/TextureStreamingTypes.h b/Engine/Source/Runtime/Engine/Classes/Engine/TextureStreamingTypes.h
--- a/Engine/Source/Runtime/Engine/Classes/Engine/TextureStreamingTypes.h
+++ b/Engine/Source/Runtime/Engine/Classes/Engine/TextureStreamingTypes.h
@@ -96,6 +96,18 @@ struct FStreamingRenderAssetPrimitiveInfo
 	}
 };
 
+/**
+ * Interface for texture streaming container
+ */
+struct ENGINE_API ITextureStreamingContainer
+{
+#if WITH_EDITOR
+	virtual void InitializeTextureStreamingContainer(uint32 InPackedTextureStreamingQualityLevelFeatureLevel) = 0;
+	virtual uint16 RegisterStreamableTexture(UTexture* InTexture) = 0;
+	virtual bool GetStreamableTexture(uint16 InTextureIndex, FString& OutTextureName, FGuid& OutTextureGuid) const { return false; }
+#endif
+};
+
 /** 
  * This struct holds the result of TextureStreaming Build for each component texture, as referred by its used materials.
  * It is possible that the entry referred by this data is not actually relevant in a given quality / target.
@@ -138,15 +150,16 @@ struct FStreamingTextureBuildInfo
 	{
 	}
 
+#if WITH_EDITOR
 	/**
 	 *	Set this struct to match the unpacked params.
 	 *
-	 *	@param	LevelTextures	[in,out]	The list of textures referred by all component of a level. The array index maps to UTexture::LevelIndex.
-	 *	@param	RefBounds		[in]		The reference bounds used to compute the packed relative box.
-	 *	@param	Info			[in]		The unpacked params.
+	 *	@param	TextureStreamingContainer	[in,out]	Contains the list of registered streamable textures referred by components.
+	 *	@param	RefBounds					[in]		The reference bounds used to compute the packed relative box.
+	 *	@param	Info						[in]		The unpacked params.
 	 */
-	ENGINE_API void PackFrom(ULevel* Level, const FBoxSphereBounds& RefBounds, const FStreamingRenderAssetPrimitiveInfo& Info);
-
+	ENGINE_API void PackFrom(ITextureStreamingContainer* TextureStreamingContainer, const FBoxSphereBounds& RefBounds, const FStreamingRenderAssetPrimitiveInfo& Info);
+#endif
 };
 
 // The max number of uv channels processed in the texture streaming build.
@@ -177,6 +190,7 @@ struct FPrimitiveMaterialInfo
 enum ETextureStreamingBuildType
 {
 	TSB_MapBuild,
+	TSB_ActorBuild,
 	TSB_ValidationOnly,
 	TSB_ViewMode
 };
@@ -192,7 +206,7 @@ enum ETextureStreamingBuildType
  * This requires the logic to not submit a streaming entry for precomputed data, as well as submit fallback data for 
  * texture that were referenced in the texture streaming build.
  */
-class FStreamingTextureLevelContext
+class ENGINE_API FStreamingTextureLevelContext
 {
 	/** Reversed lookup for ULevel::StreamingTextureGuids. */
 	const TMap<FGuid, int32>* TextureGuidToLevelIndex;
@@ -206,6 +220,9 @@ class FStreamingTextureLevelContext
 	/** The last bound component texture streaming build data. */
 	const TArray<FStreamingTextureBuildInfo>* ComponentBuildData;
 
+	/** Level's StreamingTextures array resolved version of FName to UTexture*. */
+	TArray<UTexture*> LevelStreamingTextures;
+
 	struct FTextureBoundState
 	{
 		FTextureBoundState() {}
@@ -228,8 +245,14 @@ class FStreamingTextureLevelContext
 
 	EMaterialQualityLevel::Type QualityLevel;
 	ERHIFeatureLevel::Type FeatureLevel;
+	bool bIsBuiltDataValid;
+
+	int32* GetBuildDataIndexRef(UTexture* Texture, bool bForceUpdate = false);
+	void UpdateQualityAndFeatureLevel(EMaterialQualityLevel::Type InQualityLevel, ERHIFeatureLevel::Type InFeatureLevel, const ULevel* InLevel = nullptr);
 
-	int32* GetBuildDataIndexRef(UTexture* Texture);
+	/** Console command used to turn on/off usage of texture streaming built data. */
+	static class FAutoConsoleCommand UseTextureStreamingBuiltDataCommand;
+	static bool UseTextureStreamingBuiltData;
 
 public:
 
@@ -239,8 +262,10 @@ public:
 	FStreamingTextureLevelContext(EMaterialQualityLevel::Type InQualityLevel, const UPrimitiveComponent* Primitive);
 	~FStreamingTextureLevelContext();
 
+	void UpdateContext(EMaterialQualityLevel::Type InQualityLevel, const ULevel* InLevel, const TMap<FGuid, int32>* InTextureGuidToLevelIndex);
 	void BindBuildData(const TArray<FStreamingTextureBuildInfo>* PreBuiltData);
-	void ProcessMaterial(const FBoxSphereBounds& ComponentBounds, const FPrimitiveMaterialInfo& MaterialData, float ComponentScaling, TArray<FStreamingRenderAssetPrimitiveInfo>& OutStreamingTextures);
+	void ProcessMaterial(const FBoxSphereBounds& ComponentBounds, const FPrimitiveMaterialInfo& MaterialData, float ComponentScaling, TArray<FStreamingRenderAssetPrimitiveInfo>& OutStreamingTextures, bool bIsComponentBuildDataValid = false, const UPrimitiveComponent* DebugComponent = nullptr);
+	bool CanUseTextureStreamingBuiltData() const;
 
 	EMaterialQualityLevel::Type GetQualityLevel() { return QualityLevel; }
 	ERHIFeatureLevel::Type GetFeatureLevel() { return FeatureLevel; }
@@ -270,6 +295,14 @@ typedef TMap<UMaterialInterface*, TArray<FMaterialTextureInfo> > FTexCoordScaleM
 /** A mapping between used material and levels for refering primitives. */
 typedef TMap<UMaterialInterface*, TArray<ULevel*> > FMaterialToLevelsMap;
 
+#if WITH_EDITOR
+/** Build the texture streaming component data for an actor and save common data in a UActorTextureStreamingBuildDataComponent. */
+ENGINE_API void BuildActorTextureStreamingData(AActor* InActor, EMaterialQualityLevel::Type InQualityLevel, ERHIFeatureLevel::Type InFeatureLevel);
+
+/** Build the texture streaming component data for a level based on actor's UActorTextureStreamingBuildDataComponent. */
+ENGINE_API bool BuildLevelTextureStreamingComponentDataFromActors(ULevel* InLevel);
+#endif
+
 /** Build the shaders required for the texture streaming build. Returns whether or not the action was successful. */
 ENGINE_API bool BuildTextureStreamingComponentData(UWorld* InWorld, EMaterialQualityLevel::Type QualityLevel, ERHIFeatureLevel::Type FeatureLevel, bool bFullRebuild, FSlowTask& BuildTextureStreamingTask);
 
diff --git a/Engine/Source/Runtime/Engine/Private/Components/MeshComponent.cpp b/Engine/Source/Runtime/Engine/Private/Components/MeshComponent.cpp
--- a/Engine/Source/Runtime/Engine/Private/Components/MeshComponent.cpp
+++ b/Engine/Source/Runtime/Engine/Private/Components/MeshComponent.cpp
@@ -453,7 +453,7 @@ void UMeshComponent::GetStreamingTextureInfoInner(FStreamingTextureLevelContext&
 		FPrimitiveMaterialInfo MaterialData;
 		if (GetMaterialStreamingData(MaterialIndex, MaterialData))
 		{
-			LevelContext.ProcessMaterial(Bounds, MaterialData, ComponentScaling, OutStreamingTextures);
+			LevelContext.ProcessMaterial(Bounds, MaterialData, ComponentScaling, OutStreamingTextures, bIsValidTextureStreamingBuiltData, this);
 		}
 	}
 }
diff --git a/Engine/Source/Runtime/Engine/Private/Components/PrimitiveComponent.cpp b/Engine/Source/Runtime/Engine/Private/Components/PrimitiveComponent.cpp
--- a/Engine/Source/Runtime/Engine/Private/Components/PrimitiveComponent.cpp
+++ b/Engine/Source/Runtime/Engine/Private/Components/PrimitiveComponent.cpp
@@ -338,6 +338,7 @@ UPrimitiveComponent::UPrimitiveComponent(const FObjectInitializer& ObjectInitial
 	bUseMaxLODAsImposter_DEPRECATED = false;
 	bBatchImpostersAsInstances_DEPRECATED = false;
 #endif
+	bIsValidTextureStreamingBuiltData = false;
 	bNeverDistanceCull = false;
 
 	bUseEditorCompositing = false;
@@ -444,7 +445,7 @@ void UPrimitiveComponent::GetStreamingRenderAssetInfo(FStreamingTextureLevelCont
 				if (MaterialInterface)
 				{
 					MaterialData.Material = MaterialInterface;
-					LevelContext.ProcessMaterial(Bounds, MaterialData, Bounds.SphereRadius, OutStreamingRenderAssets);
+					LevelContext.ProcessMaterial(Bounds, MaterialData, Bounds.SphereRadius, OutStreamingRenderAssets, bIsValidTextureStreamingBuiltData, this);
 				}
 				// Remove all instances of this material in case there were duplicates.
 				UsedMaterials.RemoveSwap(MaterialInterface);
@@ -453,6 +454,37 @@ void UPrimitiveComponent::GetStreamingRenderAssetInfo(FStreamingTextureLevelCont
 	}
 }
 
+bool UPrimitiveComponent::BuildTextureStreamingDataImpl(ETextureStreamingBuildType BuildType, EMaterialQualityLevel::Type QualityLevel, ERHIFeatureLevel::Type FeatureLevel, TSet<FGuid>& DependentResources, bool& bOutSupportsBuildTextureStreamingData)
+{
+	// Default implementation marks component as having invalid texture streaming built data
+	bOutSupportsBuildTextureStreamingData = false;
+	return true;
+}
+
+bool UPrimitiveComponent::BuildTextureStreamingData(ETextureStreamingBuildType BuildType, EMaterialQualityLevel::Type QualityLevel, ERHIFeatureLevel::Type FeatureLevel, TSet<FGuid>& DependentResources)
+{
+	bool bSupportsBuildTextureStreamingData = false;
+	bIsValidTextureStreamingBuiltData = false;
+	if (!FApp::CanEverRender())
+	{
+		return false;
+	}
+	bool Result = BuildTextureStreamingDataImpl(BuildType, QualityLevel, FeatureLevel, DependentResources, bSupportsBuildTextureStreamingData);
+	if (Result && bSupportsBuildTextureStreamingData)
+	{
+		if (BuildType == TSB_MapBuild)
+		{
+			bIsValidTextureStreamingBuiltData = true;
+		}
+#if WITH_EDITOR
+		else if (BuildType == TSB_ActorBuild)
+		{
+			bIsActorTextureStreamingBuiltData = true;
+		}
+#endif
+	}
+	return Result;
+}
 
 void UPrimitiveComponent::GetStreamingRenderAssetInfoWithNULLRemoval(FStreamingTextureLevelContext& LevelContext, TArray<struct FStreamingRenderAssetPrimitiveInfo>& OutStreamingRenderAssets) const
 {
diff --git a/Engine/Source/Runtime/Engine/Private/Components/StaticMeshComponent.cpp b/Engine/Source/Runtime/Engine/Private/Components/StaticMeshComponent.cpp
--- a/Engine/Source/Runtime/Engine/Private/Components/StaticMeshComponent.cpp
+++ b/Engine/Source/Runtime/Engine/Private/Components/StaticMeshComponent.cpp
@@ -28,6 +28,7 @@
 #include "IHierarchicalLODUtilities.h"
 #include "HierarchicalLODUtilitiesModule.h"
 #include "Rendering/StaticLightingSystemInterface.h"
+#include "Streaming/ActorTextureStreamingBuildDataComponent.h"
 #endif
 #include "LightMap.h"
 #include "ShadowMap.h"
@@ -745,12 +746,12 @@ static int32 GetNumberOfElements(const TIndirectArray<FStaticMeshLODResources>&
 /**
  *	Pack the texture into data ready for saving. Also ensures a single entry per texture.
  *
- *	@param	LevelTextures			[in,out]	The list of textures referred by all component of a level. The array index maps to UTexture2D::LevelIndex.
- *	@param	UnpackedData			[in,out]	The unpacked data, emptied after the function executes.
- *	@param	StreamingTextureData	[out]		The resulting packed data.
- *	@param	RefBounds				[in]		The reference bounds used to packed the relative bounds.
+ *	@param	TextureStreamingContainer [in,out]	Contains the list of textures referred by all components. The array index maps to UTexture2D::LevelIndex.
+ *	@param	UnpackedData			  [in,out]	The unpacked data, emptied after the function executes.
+ *	@param	StreamingTextureData	  [out]		The resulting packed data.
+ *	@param	RefBounds				  [in]		The reference bounds used to packed the relative bounds.
  */
-static void PackStreamingTextureData(ULevel* Level, TArray<FStreamingRenderAssetPrimitiveInfo>& UnpackedData, TArray<FStreamingTextureBuildInfo>& StreamingTextureData, const FBoxSphereBounds& RefBounds)
+static void PackStreamingTextureData(ITextureStreamingContainer* TextureStreamingContainer, TArray<FStreamingRenderAssetPrimitiveInfo>& UnpackedData, TArray<FStreamingTextureBuildInfo>& StreamingTextureData, const FBoxSphereBounds& RefBounds)
 {
 	StreamingTextureData.Empty();
 
@@ -776,7 +777,7 @@ static void PackStreamingTextureData(ULevel* Level, TArray<FStreamingRenderAsset
 		}
 
 		FStreamingTextureBuildInfo PackedInfo;
-		PackedInfo.PackFrom(Level, RefBounds, Info);
+		PackedInfo.PackFrom(TextureStreamingContainer, RefBounds, Info);
 		StreamingTextureData.Push(PackedInfo);
 	}
 }
@@ -798,8 +799,37 @@ bool UStaticMeshComponent::GetMaterialStreamingData(int32 MaterialIndex, FPrimit
 	return MaterialData.IsValid();
 }
 
-bool UStaticMeshComponent::BuildTextureStreamingData(ETextureStreamingBuildType BuildType, EMaterialQualityLevel::Type QualityLevel, ERHIFeatureLevel::Type FeatureLevel, TSet<FGuid>& DependentResources)
+#if WITH_EDITOR
+bool UStaticMeshComponent::RemapActorTextureStreamingBuiltDataToLevel(const UActorTextureStreamingBuildDataComponent* InActorTextureBuildData)
+{
+	check(InActorTextureBuildData);
+	check(InActorTextureBuildData->GetOwner() == GetOwner())
+
+	ULevel* Level = GetOwner()->GetLevel();
+	if (!Level || !bIsActorTextureStreamingBuiltData)
+	{
+		return false;
+	}
+
+	FString TextureName;
+	FGuid TextureGuid;
+	for (FStreamingTextureBuildInfo& BuildInfo : StreamingTextureData)
+	{
+		if (!InActorTextureBuildData->GetStreamableTexture(BuildInfo.TextureLevelIndex, TextureName, TextureGuid))
+		{
+			return false;
+		}
+		// Update BuildInfo's TextureLevelIndex
+		BuildInfo.TextureLevelIndex = Level->RegisterStreamableTexture(TextureName, TextureGuid);
+	}
+	return true;
+}
+#endif
+
+bool UStaticMeshComponent::BuildTextureStreamingDataImpl(ETextureStreamingBuildType BuildType, EMaterialQualityLevel::Type QualityLevel, ERHIFeatureLevel::Type FeatureLevel, TSet<FGuid>& DependentResources, bool& bOutSupportsBuildTextureStreamingData)
 {
+	bOutSupportsBuildTextureStreamingData = false;
+
 	bool bBuildDataValid = true;
 
 #if WITH_EDITORONLY_DATA // Only rebuild the data in editor 
@@ -813,7 +843,7 @@ bool UStaticMeshComponent::BuildTextureStreamingData(ETextureStreamingBuildType
 			const int32 NumMaterials = GetNumMaterials();
 
 			// Build the material bounds if in full rebuild or if the data is incomplete.
-			if (BuildType == TSB_MapBuild || (BuildType == TSB_ViewMode && MaterialStreamingRelativeBoxes.Num() != NumMaterials))
+			if ((BuildType == TSB_MapBuild) || (BuildType == TSB_ActorBuild) || (BuildType == TSB_ViewMode && MaterialStreamingRelativeBoxes.Num() != NumMaterials))
 			{
 				// Build the material bounds.
 				MaterialStreamingRelativeBoxes.Empty(NumMaterials);
@@ -831,16 +861,28 @@ bool UStaticMeshComponent::BuildTextureStreamingData(ETextureStreamingBuildType
 			}
 
 			// The texture build data can only be recomputed on a map build because of how the the level StreamingTextureGuids are handled.
-			if (BuildType == TSB_MapBuild)
+			if ((BuildType == TSB_MapBuild) || (BuildType == TSB_ActorBuild))
 			{
-				ULevel* Level = ComponentActor ? ComponentActor->GetLevel() : nullptr;
-				if (Level)
+				ITextureStreamingContainer* TextureStreamingContainer = nullptr;
+				if (ComponentActor)
+				{
+					if (BuildType == TSB_ActorBuild)
+					{
+						TextureStreamingContainer = ComponentActor->FindComponentByClass<UActorTextureStreamingBuildDataComponent>();
+					}
+					else
+					{
+						TextureStreamingContainer = ComponentActor->GetLevel();
+					}
+				}
+				if (TextureStreamingContainer)
 				{
 					// Get the data without any component scaling as the built data does not include scale.
 					FStreamingTextureLevelContext LevelContext(QualityLevel, FeatureLevel, true); // Use the boxes that were just computed!
 					TArray<FStreamingRenderAssetPrimitiveInfo> UnpackedData;
 					GetStreamingTextureInfoInner(LevelContext, nullptr, 1.f, UnpackedData);
-					PackStreamingTextureData(Level, UnpackedData, StreamingTextureData, Bounds);
+					PackStreamingTextureData(TextureStreamingContainer, UnpackedData, StreamingTextureData, Bounds);
+					bOutSupportsBuildTextureStreamingData = true;
 				}
 			}
 			else if (StreamingTextureData.Num() == 0)
diff --git a/Engine/Source/Runtime/Engine/Private/InstancedStaticMesh.cpp b/Engine/Source/Runtime/Engine/Private/InstancedStaticMesh.cpp
--- a/Engine/Source/Runtime/Engine/Private/InstancedStaticMesh.cpp
+++ b/Engine/Source/Runtime/Engine/Private/InstancedStaticMesh.cpp
@@ -3427,12 +3427,12 @@ bool UInstancedStaticMeshComponent::GetMaterialStreamingData(int32 MaterialIndex
 	return MaterialData.IsValid();
 }
 
-bool UInstancedStaticMeshComponent::BuildTextureStreamingData(ETextureStreamingBuildType BuildType, EMaterialQualityLevel::Type QualityLevel, ERHIFeatureLevel::Type FeatureLevel, TSet<FGuid>& DependentResources)
+bool UInstancedStaticMeshComponent::BuildTextureStreamingDataImpl(ETextureStreamingBuildType BuildType, EMaterialQualityLevel::Type QualityLevel, ERHIFeatureLevel::Type FeatureLevel, TSet<FGuid>& DependentResources, bool& bOutSupportsBuildTextureStreamingData)
 {
 #if WITH_EDITORONLY_DATA // Only rebuild the data in editor 
 	if (GetInstanceCount() > 0)
 	{
-		return Super::BuildTextureStreamingData(BuildType, QualityLevel, FeatureLevel, DependentResources);
+		return Super::BuildTextureStreamingDataImpl(BuildType, QualityLevel, FeatureLevel, DependentResources, bOutSupportsBuildTextureStreamingData);
 	}
 #endif
 	return true;
diff --git a/Engine/Source/Runtime/Engine/Private/Level.cpp b/Engine/Source/Runtime/Engine/Private/Level.cpp
--- a/Engine/Source/Runtime/Engine/Private/Level.cpp
+++ b/Engine/Source/Runtime/Engine/Private/Level.cpp
@@ -774,6 +774,11 @@ void ULevel::PreSave(FObjectPreSaveContext ObjectSaveContext)
 				Actor->ClearCrossLevelReferences();
 			}
 		}
+
+		if (ObjectSaveContext.IsCooking())
+		{
+			BuildLevelTextureStreamingComponentDataFromActors(this);
+		}
 	}
 #endif // WITH_EDITOR
 }
@@ -1669,6 +1674,38 @@ void ULevel::CreateModelComponents()
 		ModelComp->InvalidateCollisionData();
 	}
 }
+
+void ULevel::InitializeTextureStreamingContainer(uint32 InPackedTextureStreamingQualityLevelFeatureLevel)
+{
+	PackedTextureStreamingQualityLevelFeatureLevel = InPackedTextureStreamingQualityLevelFeatureLevel;
+	bTextureStreamingRotationChanged = false;
+	StreamingTextureGuids.Empty();
+	StreamingTextures.Empty();
+	TextureStreamingResourceGuids.Empty();
+	NumTextureStreamingDirtyResources = 0; // This is persistent in order to be able to notify if a rebuild is required when running a cooked build.
+}
+
+uint16 ULevel::RegisterStreamableTexture(UTexture* InTexture)
+{
+	if (InTexture->LevelIndex == INDEX_NONE)
+	{
+		// If this is the first time this texture gets processed in the packing process, encode it.
+		InTexture->LevelIndex = (int32)RegisterStreamableTexture(InTexture->GetPathName(), InTexture->GetLightingGuid());
+	}
+	check(StreamingTextureGuids.IsValidIndex(InTexture->LevelIndex));
+	check(StreamingTextureGuids[InTexture->LevelIndex] == InTexture->GetLightingGuid());
+	return (uint16)InTexture->LevelIndex;
+}
+
+uint16 ULevel::RegisterStreamableTexture(const FString& InTextureName, const FGuid& InTextureGuid)
+{
+	check(StreamingTextures.Num() == StreamingTextureGuids.Num());
+	uint16 Index = StreamingTextureGuids.AddUnique(InTextureGuid);
+	StreamingTextures.AddUnique(FName(*InTextureName));
+	check(StreamingTextures.Num() == StreamingTextureGuids.Num());
+	return Index;
+}
+
 #endif
 
 void ULevel::UpdateModelComponents()
diff --git a/Engine/Source/Runtime/Engine/Private/Streaming/ActorTextureStreamingBuildDataComponent.cpp b/Engine/Source/Runtime/Engine/Private/Streaming/ActorTextureStreamingBuildDataComponent.cpp
new file mode 100644
--- /dev/null
+++ b/Engine/Source/Runtime/Engine/Private/Streaming/ActorTextureStreamingBuildDataComponent.cpp
@@ -0,0 +1,51 @@
+// Copyright Epic Games, Inc. All Rights Reserved.
+
+/*=============================================================================
+	TextureStreamingBuiltDataComponent.cpp
+=============================================================================*/
+
+#include "Streaming/ActorTextureStreamingBuildDataComponent.h"
+
+#if WITH_EDITOR
+
+#include "Engine/TextureStreamingTypes.h"
+#include "Engine/Texture.h"
+
+void UActorTextureStreamingBuildDataComponent::InitializeTextureStreamingContainer(uint32 InPackedTextureStreamingQualityLevelFeatureLevel)
+{
+	PackedTextureStreamingQualityLevelFeatureLevel = InPackedTextureStreamingQualityLevelFeatureLevel;
+	StreamableTextures.Empty();
+}
+
+uint16 UActorTextureStreamingBuildDataComponent::RegisterStreamableTexture(UTexture* InTexture)
+{
+	const FString& TextureName = InTexture->GetPathName();
+	const FGuid& TextureGuid = InTexture->GetLightingGuid();
+	
+	uint16 Index = 0;
+	for (const FStreamableTexture& StreamableTexture : StreamableTextures)
+	{
+		if (StreamableTexture.Name == TextureName)
+		{
+			check(StreamableTexture.Guid == TextureGuid);
+			return Index;
+		}
+		++Index;
+	}
+
+	StreamableTextures.Emplace(TextureName, TextureGuid);
+	return Index;
+}
+
+bool UActorTextureStreamingBuildDataComponent::GetStreamableTexture(uint16 InTextureIndex, FString& OutTextureName, FGuid& OutTextureGuid) const
+{
+	if (StreamableTextures.IsValidIndex(InTextureIndex))
+	{
+		OutTextureName = StreamableTextures[InTextureIndex].Name;
+		OutTextureGuid = StreamableTextures[InTextureIndex].Guid;
+		return true;
+	}
+	return false;
+}
+
+#endif
\ No newline at end of file
diff --git a/Engine/Source/Runtime/Engine/Private/Streaming/LevelTextureManager.cpp b/Engine/Source/Runtime/Engine/Private/Streaming/LevelTextureManager.cpp
--- a/Engine/Source/Runtime/Engine/Private/Streaming/LevelTextureManager.cpp
+++ b/Engine/Source/Runtime/Engine/Private/Streaming/LevelTextureManager.cpp
@@ -125,7 +125,7 @@ void FLevelRenderAssetManager::IncrementalBuild(FDynamicRenderAssetInstanceManag
 		BuildStep = EStaticBuildStep::ProcessActors;
 
 		// Update the level context with the texture guid map. This is required in case the incremental build runs more steps.
-		LevelContext = FStreamingTextureLevelContext(EMaterialQualityLevel::Num, Level, &TextureGuidToLevelIndex);
+		LevelContext.UpdateContext(EMaterialQualityLevel::Num, Level, &TextureGuidToLevelIndex);
 		break;
 	}
 	case EStaticBuildStep::ProcessActors:
diff --git a/Engine/Source/Runtime/Engine/Private/Streaming/TextureStreamingBuild.cpp b/Engine/Source/Runtime/Engine/Private/Streaming/TextureStreamingBuild.cpp
--- a/Engine/Source/Runtime/Engine/Private/Streaming/TextureStreamingBuild.cpp
+++ b/Engine/Source/Runtime/Engine/Private/Streaming/TextureStreamingBuild.cpp
@@ -18,6 +18,7 @@ TextureStreamingBuild.cpp : Contains definitions to build texture streaming data
 #include "ShaderCompiler.h"
 #include "Engine/StaticMesh.h"
 #include "Streaming/TextureStreamingHelpers.h"
+#include "Streaming/ActorTextureStreamingBuildDataComponent.h"
 #include "UObject/UObjectIterator.h"
 #include "Logging/MessageLog.h"
 #include "Misc/UObjectToken.h"
@@ -25,6 +26,114 @@ TextureStreamingBuild.cpp : Contains definitions to build texture streaming data
 DEFINE_LOG_CATEGORY(TextureStreamingBuild);
 #define LOCTEXT_NAMESPACE "TextureStreamingBuild"
 
+#if WITH_EDITOR
+
+static const uint32 GInvalidPackedTextureStreamingQualityLevelFeatureLevel = 0xFFFFFFFF;
+
+ENGINE_API uint32 GetPackedTextureStreamingQualityLevelFeatureLevel(EMaterialQualityLevel::Type InQualityLevel, ERHIFeatureLevel::Type InFeatureLevel)
+{
+	return ((uint32)InQualityLevel << 16) | ((uint32)InFeatureLevel & 0xFFFF);
+}
+
+ENGINE_API void BuildActorTextureStreamingData(AActor* InActor, EMaterialQualityLevel::Type InQualityLevel, ERHIFeatureLevel::Type InFeatureLevel)
+{
+	if (!InActor || InActor->IsTemplate())
+	{
+		return;
+	}
+
+	TSet<FGuid> DummyResourceGuids;
+	UActorTextureStreamingBuildDataComponent* BuiltDataComponent = InActor->FindComponentByClass<UActorTextureStreamingBuildDataComponent>();
+	if (!BuiltDataComponent)
+	{
+		BuiltDataComponent = NewObject<UActorTextureStreamingBuildDataComponent>(InActor, NAME_None, RF_Transactional);
+		InActor->AddInstanceComponent(BuiltDataComponent);
+	}
+	BuiltDataComponent->InitializeTextureStreamingContainer(GetPackedTextureStreamingQualityLevelFeatureLevel(InQualityLevel, InFeatureLevel));
+	
+	TInlineComponentArray<UPrimitiveComponent*> Primitives;
+	InActor->GetComponents<UPrimitiveComponent>(Primitives);
+	for (UPrimitiveComponent* Primitive : Primitives)
+	{
+		if (Primitive)
+		{
+			Primitive->BuildTextureStreamingData(TSB_ActorBuild, InQualityLevel, InFeatureLevel, DummyResourceGuids);
+		}
+	}
+}
+
+ENGINE_API bool BuildLevelTextureStreamingComponentDataFromActors(ULevel* InLevel)
+{
+	// If Level already contains texture streaming built data, keep it as it was built in the editor.
+	if (!InLevel->StreamingTextureGuids.IsEmpty())
+	{
+		return true;
+	}
+
+	// Find actors with texture streaming built data (built with BuildActorTextureStreamingData) and remap this data to the level
+	bool bIsLevelInitialized = false;
+	for (AActor* Actor : InLevel->Actors)
+	{
+		UActorTextureStreamingBuildDataComponent* BuiltDataComponent = Actor ? Actor->FindComponentByClass<UActorTextureStreamingBuildDataComponent>() : nullptr;
+		if (!BuiltDataComponent)
+		{
+			continue;
+		}
+
+		uint32 ComponentPackedInfo = BuiltDataComponent->GetPackedTextureStreamingQualityLevelFeatureLevel();
+		if (!bIsLevelInitialized)
+		{
+			InLevel->InitializeTextureStreamingContainer(ComponentPackedInfo);
+			bIsLevelInitialized = true;
+		}
+		else if (InLevel->PackedTextureStreamingQualityLevelFeatureLevel != ComponentPackedInfo)
+		{
+			// We don't support mixed versions of quality/feature levels
+			InLevel->PackedTextureStreamingQualityLevelFeatureLevel = GInvalidPackedTextureStreamingQualityLevelFeatureLevel;
+			return false;
+		}
+
+		TInlineComponentArray<UPrimitiveComponent*> Primitives;
+		Actor->GetComponents<UPrimitiveComponent>(Primitives);
+		for (UPrimitiveComponent* Primitive : Primitives)
+		{
+			if (Primitive && Primitive->bIsActorTextureStreamingBuiltData)
+			{
+				// Remap primitive's texture streaming built data from actor based to level based
+				Primitive->bIsValidTextureStreamingBuiltData = Primitive->RemapActorTextureStreamingBuiltDataToLevel(BuiltDataComponent);
+			}
+		}
+	}
+	return true;
+}
+
+#else
+
+static bool GetUnpackedTextureStreamingQualityLevelFeatureLevel(uint32 InPackedValue, EMaterialQualityLevel::Type& OutQualityLevel, ERHIFeatureLevel::Type& OutFeatureLevel)
+{
+	uint32 QualityLevel = InPackedValue >> 16;
+	uint32 FeatureLevel = InPackedValue & 0xFFFF;
+	bool bIsValidQualityLevel = QualityLevel <= (uint32)EMaterialQualityLevel::Num;
+	bool bIsValidFeatureLevel = FeatureLevel <= (uint32)ERHIFeatureLevel::Num;
+	if (bIsValidQualityLevel && bIsValidFeatureLevel)
+	{
+		OutQualityLevel = (EMaterialQualityLevel::Type)QualityLevel;
+		OutFeatureLevel = (ERHIFeatureLevel::Type)FeatureLevel;
+		if (OutQualityLevel == EMaterialQualityLevel::Num)
+		{
+			OutQualityLevel = GetCachedScalabilityCVars().MaterialQualityLevel;
+		}
+		if (OutFeatureLevel == ERHIFeatureLevel::Num)
+		{
+			OutFeatureLevel = GMaxRHIFeatureLevel;
+		}
+		return true;
+	}
+	return false;
+}
+
+#endif
+
 ENGINE_API bool BuildTextureStreamingComponentData(UWorld* InWorld, EMaterialQualityLevel::Type QualityLevel, ERHIFeatureLevel::Type FeatureLevel, bool bFullRebuild, FSlowTask& BuildTextureStreamingTask)
 {
 #if WITH_EDITORONLY_DATA
@@ -75,10 +184,7 @@ ENGINE_API bool BuildTextureStreamingComponentData(UWorld* InWorld, EMaterialQua
 		// This allows to keep track of full rebuild requirements.
 		if (bFullRebuild)
 		{
-			Level->bTextureStreamingRotationChanged = false;
-			Level->StreamingTextureGuids.Empty();
-			Level->TextureStreamingResourceGuids.Empty();
-			Level->NumTextureStreamingDirtyResources = 0; // This is persistent in order to be able to notify if a rebuild is required when running a cooked build.
+			Level->InitializeTextureStreamingContainer(GetPackedTextureStreamingQualityLevelFeatureLevel(QualityLevel, FeatureLevel));
 		}
 
 		TSet<FGuid> ResourceGuids;
@@ -235,50 +341,43 @@ void UnpackRelativeBox(const FBoxSphereBounds& InRefBounds, uint32 InPackedRelBo
 	}
 }
 
-void FStreamingTextureBuildInfo::PackFrom(ULevel* Level, const FBoxSphereBounds& RefBounds, const FStreamingRenderAssetPrimitiveInfo& Info)
+#if WITH_EDITOR
+void FStreamingTextureBuildInfo::PackFrom(ITextureStreamingContainer* TextureStreamingContainer, const FBoxSphereBounds& RefBounds, const FStreamingRenderAssetPrimitiveInfo& Info)
 {
-	check(Level);
+	check(TextureStreamingContainer);
 
 	PackedRelativeBox = PackRelativeBox(RefBounds.Origin, RefBounds.BoxExtent, Info.Bounds.Origin, Info.Bounds.BoxExtent);
 
 	UTexture* Texture = Cast<UTexture>(Info.RenderAsset);
 	check(Texture);
-	if (Texture->LevelIndex == INDEX_NONE)
-	{
-		// If this is the first time this texture gets processed in the packing process, encode it.
-		Texture->LevelIndex = Level->StreamingTextureGuids.Add(Texture->GetLightingGuid());
-	}
-	TextureLevelIndex = (uint16)Texture->LevelIndex;
-
+	TextureLevelIndex = TextureStreamingContainer->RegisterStreamableTexture(Texture);
 	TexelFactor = Info.TexelFactor;
 }
+#endif
+
+bool FStreamingTextureLevelContext::UseTextureStreamingBuiltData = true;
+FAutoConsoleCommand FStreamingTextureLevelContext::UseTextureStreamingBuiltDataCommand(
+	TEXT("r.Streaming.UseTextureStreamingBuiltData"),
+	TEXT("Turn on/off usage of texture streaming built data (0 to turn off)."),
+	FConsoleCommandWithArgsDelegate::CreateLambda([](const TArray<FString>& Args) { FStreamingTextureLevelContext::UseTextureStreamingBuiltData = (Args.Num() != 1) || (Args[0] != TEXT("0")); }));
 
 FStreamingTextureLevelContext::FStreamingTextureLevelContext(EMaterialQualityLevel::Type InQualityLevel, const UPrimitiveComponent* Primitive)
-: TextureGuidToLevelIndex(nullptr)
-, bUseRelativeBoxes(false)
-, BuildDataTimestamp(0)
-, ComponentBuildData(nullptr)
-, QualityLevel(InQualityLevel)
-, FeatureLevel(GMaxRHIFeatureLevel)
+	: TextureGuidToLevelIndex(nullptr)
+	, bUseRelativeBoxes(false)
+	, BuildDataTimestamp(0)
+	, ComponentBuildData(nullptr)
 {
-	if (Primitive)
-	{
-		const UWorld* World = Primitive->GetWorld();
-		if (World)
-		{
-			FeatureLevel = World->FeatureLevel;
-		}
-	}
+	ULevel* Level = Primitive && Primitive->GetOwner() ? Primitive->GetOwner()->GetLevel() : nullptr;
+	UpdateQualityAndFeatureLevel(InQualityLevel, GMaxRHIFeatureLevel, Level);
 }
 
 FStreamingTextureLevelContext::FStreamingTextureLevelContext(EMaterialQualityLevel::Type InQualityLevel, ERHIFeatureLevel::Type InFeatureLevel, bool InUseRelativeBoxes)
-: TextureGuidToLevelIndex(nullptr)
-, bUseRelativeBoxes(InUseRelativeBoxes)
-, BuildDataTimestamp(0)
-, ComponentBuildData(nullptr)
-, QualityLevel(InQualityLevel)
-, FeatureLevel(InFeatureLevel)
+	: TextureGuidToLevelIndex(nullptr)
+	, bUseRelativeBoxes(InUseRelativeBoxes)
+	, BuildDataTimestamp(0)
+	, ComponentBuildData(nullptr)
 {
+	UpdateQualityAndFeatureLevel(InQualityLevel, InFeatureLevel);
 }
 
 FStreamingTextureLevelContext::FStreamingTextureLevelContext(EMaterialQualityLevel::Type InQualityLevel, const ULevel* InLevel, const TMap<FGuid, int32>* InTextureGuidToLevelIndex) 
@@ -286,28 +385,87 @@ FStreamingTextureLevelContext::FStreamingTextureLevelContext(EMaterialQualityLev
 , bUseRelativeBoxes(false)
 , BuildDataTimestamp(0)
 , ComponentBuildData(nullptr)
-, QualityLevel(InQualityLevel)
-, FeatureLevel(GMaxRHIFeatureLevel)
 {
-	if (InLevel)
+	UpdateContext(InQualityLevel, InLevel, InTextureGuidToLevelIndex);
+}
+
+void FStreamingTextureLevelContext::UpdateContext(EMaterialQualityLevel::Type InQualityLevel, const ULevel* InLevel, const TMap<FGuid, int32>* InTextureGuidToLevelIndex)
+{
+	TextureGuidToLevelIndex = nullptr;
+	bUseRelativeBoxes = false;
+	BuildDataTimestamp = 0;
+	ComponentBuildData = nullptr;
+	LevelStreamingTextures.Reset();
+	for (const FTextureBoundState& BoundState : BoundStates)
 	{
-		const UWorld* World = InLevel->GetWorld();
-		if (World)
+		if (BoundState.Texture)
 		{
-			FeatureLevel = World->FeatureLevel;
+			BoundState.Texture->LevelIndex = INDEX_NONE;
 		}
+	}
+	BoundStates.Reset();
+
+	UpdateQualityAndFeatureLevel(InQualityLevel, GMaxRHIFeatureLevel, InLevel);
 
-		if (InLevel && InTextureGuidToLevelIndex && InLevel->StreamingTextureGuids.Num() > 0 && InLevel->StreamingTextureGuids.Num() == InTextureGuidToLevelIndex->Num())
+	if (InLevel)
+	{
+		if (CanUseTextureStreamingBuiltData() && InTextureGuidToLevelIndex && InLevel->StreamingTextureGuids.Num() > 0 && InLevel->StreamingTextureGuids.Num() == InTextureGuidToLevelIndex->Num())
 		{
 			bUseRelativeBoxes = !InLevel->bTextureStreamingRotationChanged;
 			TextureGuidToLevelIndex = InTextureGuidToLevelIndex;
 
 			// Extra transient data for each texture.
 			BoundStates.AddZeroed(InLevel->StreamingTextureGuids.Num());
+
+			if (InLevel->StreamingTextures.Num() == InLevel->StreamingTextureGuids.Num())
+			{
+				LevelStreamingTextures.SetNumZeroed(InLevel->StreamingTextures.Num());
+				for (int TextureLevelIndex = 0; TextureLevelIndex < InLevel->StreamingTextures.Num(); ++TextureLevelIndex)
+				{
+					if (UTexture* Texture = FindObject<UTexture>(nullptr, *InLevel->StreamingTextures[TextureLevelIndex].ToString()))
+					{
+						LevelStreamingTextures[TextureLevelIndex] = Texture;
+						GetBuildDataIndexRef(Texture, /*ForceUpdate*/ true);
+						check(Texture->LevelIndex == TextureLevelIndex);
+					}
+				}
+			}
 		}
 	}
 }
 
+void FStreamingTextureLevelContext::UpdateQualityAndFeatureLevel(EMaterialQualityLevel::Type InQualityLevel, ERHIFeatureLevel::Type InFeatureLevel, const ULevel* InLevel)
+{
+	QualityLevel = InQualityLevel;
+	FeatureLevel = InFeatureLevel;
+	bIsBuiltDataValid = false;
+
+	if (InLevel)
+	{
+		if (const UWorld* World = InLevel->GetWorld())
+		{
+			FeatureLevel = World->FeatureLevel;
+		}
+#if !WITH_EDITOR
+		// Detect if quality level and feature level used to build texture streaming data matches the one of the context
+		EMaterialQualityLevel::Type BuiltDataQualityLevel;
+		ERHIFeatureLevel::Type BuiltDataFeatureLevel;
+		if (GetUnpackedTextureStreamingQualityLevelFeatureLevel(InLevel->PackedTextureStreamingQualityLevelFeatureLevel, BuiltDataQualityLevel, BuiltDataFeatureLevel))
+		{
+			EMaterialQualityLevel::Type ContextQualityLevel = (QualityLevel == EMaterialQualityLevel::Num) ? GetCachedScalabilityCVars().MaterialQualityLevel : QualityLevel;
+			ERHIFeatureLevel::Type ContextFeatureLevel = (FeatureLevel == ERHIFeatureLevel::Num) ? GMaxRHIFeatureLevel : FeatureLevel;
+
+			bIsBuiltDataValid = (ContextQualityLevel == BuiltDataQualityLevel) && (ContextFeatureLevel == BuiltDataFeatureLevel);
+		}
+#endif
+	}
+}
+
+bool FStreamingTextureLevelContext::CanUseTextureStreamingBuiltData() const
+{
+	return bIsBuiltDataValid && FStreamingTextureLevelContext::UseTextureStreamingBuiltData;
+}
+
 FStreamingTextureLevelContext::~FStreamingTextureLevelContext()
 {
 	// Reset the level indices for the next use.
@@ -329,7 +487,7 @@ void FStreamingTextureLevelContext::BindBuildData(const TArray<FStreamingTexture
 	if (TextureGuidToLevelIndex && CVarStreamingUseNewMetrics.GetValueOnGameThread() != 0) // No point in binding data if there is no possible remapping.
 	{
 		// Process the build data in order to be able to map a texture object to the build data entry.
-		ComponentBuildData = BuildData;
+		ComponentBuildData = CanUseTextureStreamingBuiltData() ? BuildData : nullptr;
 		if (BuildData && BoundStates.Num() > 0)
 		{
 			for (int32 Index = 0; Index < BuildData->Num(); ++Index)
@@ -350,9 +508,9 @@ void FStreamingTextureLevelContext::BindBuildData(const TArray<FStreamingTexture
 	}
 }
 
-int32* FStreamingTextureLevelContext::GetBuildDataIndexRef(UTexture* Texture)
+int32* FStreamingTextureLevelContext::GetBuildDataIndexRef(UTexture* Texture, bool bForceUpdate)
 {
-	if (ComponentBuildData) // If there is some build data to map to.
+	if (ComponentBuildData || bForceUpdate) // If there is some build data to map to.
 	{
 		if (Texture->LevelIndex == INDEX_NONE)
 		{
@@ -396,12 +554,36 @@ int32* FStreamingTextureLevelContext::GetBuildDataIndexRef(UTexture* Texture)
 	return nullptr;
 }
 
-void FStreamingTextureLevelContext::ProcessMaterial(const FBoxSphereBounds& ComponentBounds, const FPrimitiveMaterialInfo& MaterialData, float ComponentScaling, TArray<FStreamingRenderAssetPrimitiveInfo>& OutStreamingTextures)
+void FStreamingTextureLevelContext::ProcessMaterial(const FBoxSphereBounds& ComponentBounds, const FPrimitiveMaterialInfo& MaterialData, float ComponentScaling, TArray<FStreamingRenderAssetPrimitiveInfo>& OutStreamingTextures, bool bIsComponentBuildDataValid, const UPrimitiveComponent* DebugComponent)
 {
 	ensure(MaterialData.IsValid());
 
 	TArray<UTexture*> Textures;
-	MaterialData.Material->GetUsedTextures(Textures, QualityLevel, false, FeatureLevel, false);
+
+#if !WITH_EDITOR
+	// Use pre-built texture streaming data is possible
+	if (CanUseTextureStreamingBuiltData() && (ComponentBuildData || bIsComponentBuildDataValid))
+	{
+		// bIsComponentBuildDataValid can be true, but no streamable textures were found when building texture streaming data
+		if (ComponentBuildData)
+		{
+			for (const FStreamingTextureBuildInfo& BuildInfo : *ComponentBuildData)
+			{
+				check(LevelStreamingTextures.IsValidIndex(BuildInfo.TextureLevelIndex));
+				if (UTexture* Texture = LevelStreamingTextures[BuildInfo.TextureLevelIndex])
+				{
+					check(Texture->IsStreamable());
+					check(Texture->LevelIndex != INDEX_NONE);
+					Textures.Add(Texture);
+				}
+			}
+		}
+	}
+	else
+#endif
+	{
+		MaterialData.Material->GetUsedTextures(Textures, QualityLevel, false, FeatureLevel, false);
+	}
 
 	for (UTexture* Texture : Textures)
 	{
diff --git a/Engine/Source/Runtime/Engine/Private/WorldPartition/HLOD/HLODActor.cpp b/Engine/Source/Runtime/Engine/Private/WorldPartition/HLOD/HLODActor.cpp
--- a/Engine/Source/Runtime/Engine/Private/WorldPartition/HLOD/HLODActor.cpp
+++ b/Engine/Source/Runtime/Engine/Private/WorldPartition/HLOD/HLODActor.cpp
@@ -4,6 +4,7 @@
 #include "WorldPartition/HLOD/HLODSubsystem.h"
 #include "Components/PrimitiveComponent.h"
 #include "UObject/UE5MainStreamObjectVersion.h"
+#include "UObject/ObjectSaveContext.h"
 
 #if WITH_EDITOR
 #include "WorldPartition/WorldPartition.h"
@@ -14,6 +15,7 @@
 #include "Modules/ModuleManager.h"
 #include "IWorldPartitionHLODUtilities.h"
 #include "WorldPartitionHLODUtilitiesModule.h"
+#include "Engine/TextureStreamingTypes.h"
 #endif
 
 AWorldPartitionHLOD::AWorldPartitionHLOD(const FObjectInitializer& ObjectInitializer)
@@ -69,6 +71,22 @@ void AWorldPartitionHLOD::Serialize(FArchive& Ar)
 void AWorldPartitionHLOD::RerunConstructionScripts()
 {}
 
+void AWorldPartitionHLOD::PreSave(FObjectPreSaveContext ObjectSaveContext)
+{
+	Super::PreSave(ObjectSaveContext);
+
+#if WITH_EDITOR
+	// When generating WorldPartition HLODs, we have the renderer initialized.
+	// Take advantage of this and generate texture streaming built data (local to the actor).
+	// This built data will be used by the cooking (it will convert it to level texture streaming built data).
+	if (!ObjectSaveContext.IsCooking())
+	{
+		// Use same quality level and feature level as FEditorBuildUtils::EditorBuildTextureStreaming
+		BuildActorTextureStreamingData(this, EMaterialQualityLevel::High, GMaxRHIFeatureLevel);
+	}
+#endif
+}
+
 #if WITH_EDITOR
 
 bool AWorldPartitionHLOD::IsHiddenEd() const
diff --git a/Engine/Source/Runtime/Engine/Public/Streaming/ActorTextureStreamingBuildDataComponent.h b/Engine/Source/Runtime/Engine/Public/Streaming/ActorTextureStreamingBuildDataComponent.h
new file mode 100644
--- /dev/null
+++ b/Engine/Source/Runtime/Engine/Public/Streaming/ActorTextureStreamingBuildDataComponent.h
@@ -0,0 +1,66 @@
+// Copyright Epic Games, Inc. All Rights Reserved.
+
+#pragma once
+
+#include "CoreMinimal.h"
+#include "UObject/ObjectMacros.h"
+#include "Components/ActorComponent.h"
+#include "Misc/Guid.h"
+#include "Engine/TextureStreamingTypes.h"
+#include "ActorTextureStreamingBuildDataComponent.generated.h"
+
+class UTexture;
+
+USTRUCT()
+struct ENGINE_API FStreamableTexture
+{
+	GENERATED_BODY()
+
+	FStreamableTexture() {}
+
+	FStreamableTexture(const FString& InName, const FGuid& InGuid)
+#if WITH_EDITORONLY_DATA
+		: Name(InName)
+		, Guid(InGuid)
+#endif
+	{}
+
+#if WITH_EDITORONLY_DATA
+	UPROPERTY()
+	FString Name;
+
+	UPROPERTY()
+	FGuid Guid;
+#endif
+};
+
+UCLASS()
+class ENGINE_API UActorTextureStreamingBuildDataComponent : public UActorComponent, public ITextureStreamingContainer
+{
+	GENERATED_BODY()
+
+public:
+
+	//~Begin UObject Interface.
+	virtual bool IsEditorOnly() const override { return true; }
+	//~End UObject Interface.
+
+#if WITH_EDITOR
+	//~Begin ITextureStreamingContainer Interface.
+	virtual void InitializeTextureStreamingContainer(uint32 InPackedTextureStreamingQualityLevelFeatureLevel) override;
+	virtual uint16 RegisterStreamableTexture(UTexture* InTexture) override;
+	virtual bool GetStreamableTexture(uint16 InTextureIndex, FString& OutTextureName, FGuid& OutTextureGuid) const override;
+	//~End ITextureStreamingContainer Interface.
+
+	uint32 GetPackedTextureStreamingQualityLevelFeatureLevel() const { return PackedTextureStreamingQualityLevelFeatureLevel; }
+#endif
+
+private:
+#if WITH_EDITORONLY_DATA
+	UPROPERTY()
+	TArray<FStreamableTexture> StreamableTextures;
+
+	UPROPERTY()
+	uint32 PackedTextureStreamingQualityLevelFeatureLevel;
+#endif
+};
\ No newline at end of file
diff --git a/Engine/Source/Runtime/Engine/Public/WorldPartition/HLOD/HLODActor.h b/Engine/Source/Runtime/Engine/Public/WorldPartition/HLOD/HLODActor.h
--- a/Engine/Source/Runtime/Engine/Public/WorldPartition/HLOD/HLODActor.h
+++ b/Engine/Source/Runtime/Engine/Public/WorldPartition/HLOD/HLODActor.h
@@ -64,6 +64,7 @@ protected:
 	//~ Begin UObject Interface.
 	virtual void Serialize(FArchive& Ar) override;
 	virtual void RerunConstructionScripts() override;
+	virtual void PreSave(FObjectPreSaveContext ObjectSaveContext) override;
 	//~ End UObject Interface.
 
 	//~ Begin AActor Interface.
