PR #7627: Animation Curves: adding a Get with validity checking (Contributed by [REDACTED_PERSON])

diff --git a/Engine/Source/Runtime/Engine/Public/Animation/AnimCurveTypes.h b/Engine/Source/Runtime/Engine/Public/Animation/AnimCurveTypes.h
--- a/Engine/Source/Runtime/Engine/Public/Animation/AnimCurveTypes.h
+++ b/Engine/Source/Runtime/Engine/Public/Animation/AnimCurveTypes.h
@@ -445,6 +445,22 @@ struct FBaseBlendedCurve
 		return 0.f;
 	}
 
+	/** Get Value of InUID with validation and default value */
+	float Get(USkeleton::AnimCurveUID InUid, bool& OutIsValid, float InDefaultValue=0.f) const
+	{
+		check(bInitialized);
+
+		const int32 ArrayIndex = GetArrayIndexByUID(InUid);
+		if ((ArrayIndex != INDEX_NONE) && ValidCurveWeights[ArrayIndex])
+		{
+			OutIsValid = true;
+			return CurveWeights[ArrayIndex];
+		}
+
+		OutIsValid = false;
+		return InDefaultValue;
+	}
+
 	/** Get Array Index by UID */
 	int32 GetArrayIndexByUID(USkeleton::AnimCurveUID InUid) const
 	{
