Texture: Fixed the build exporter to use the correct input names

#rnx

diff --git a/Engine/Source/Runtime/Engine/Private/TextureDerivedDataBuildExporter.cpp b/Engine/Source/Runtime/Engine/Private/TextureDerivedDataBuildExporter.cpp
--- a/Engine/Source/Runtime/Engine/Private/TextureDerivedDataBuildExporter.cpp
+++ b/Engine/Source/Runtime/Engine/Private/TextureDerivedDataBuildExporter.cpp
@@ -169,10 +169,10 @@ void FTextureDerivedDataBuildExporter::ExportTextureBuild(const UTexture& Textur
 	ActionBuilder.AddConstant(TEXT("Settings"_SV),
 		SaveTextureBuildSettings(KeySuffix, Texture, BuildSettings, LayerIndex, NumInlineMips));
 
-	ActionBuilder.AddInput(Texture.Source.GetId().ToString(), ExportedTextureBulkDataHash, ExportedTextureBulkDataSize);
+	ActionBuilder.AddInput(TEXT("Source"_SV), ExportedTextureBulkDataHash, ExportedTextureBulkDataSize);
 	if ((bool)Texture.CompositeTexture && !ExportedCompositeTextureBulkDataHash.IsZero())
 	{
-		ActionBuilder.AddInput(Texture.CompositeTexture->Source.GetId().ToString(), ExportedCompositeTextureBulkDataHash, ExportedCompositeTextureBulkDataSize);
+		ActionBuilder.AddInput(TEXT("CompositeSource"_SV), ExportedCompositeTextureBulkDataHash, ExportedCompositeTextureBulkDataSize);
 	}
 
 	if (TUniquePtr<FArchive> Ar{IFileManager::Get().CreateFileWriter(*(ExportRoot / TEXT("build.uddba")))})
