Fixed nonunity test build.

diff --git a/Engine/Source/Runtime/Engine/Private/Streaming/TextureStreamingBuild.cpp b/Engine/Source/Runtime/Engine/Private/Streaming/TextureStreamingBuild.cpp
--- a/Engine/Source/Runtime/Engine/Private/Streaming/TextureStreamingBuild.cpp
+++ b/Engine/Source/Runtime/Engine/Private/Streaming/TextureStreamingBuild.cpp
@@ -22,6 +22,7 @@ TextureStreamingBuild.cpp : Contains definitions to build texture streaming data
 #include "UObject/UObjectIterator.h"
 #include "Logging/MessageLog.h"
 #include "Misc/UObjectToken.h"
+#include "UnrealEngine.h"
 
 DEFINE_LOG_CATEGORY(TextureStreamingBuild);
 #define LOCTEXT_NAMESPACE "TextureStreamingBuild"
