World Partition
* Hide bEnableWorldPartition property (was in read only) - it was more confusion from what we have seen on forums for EA.
* Note : the whole flow of how enabling WP is going to change in the future to be more user friendly.

diff --git a/Engine/Source/Runtime/Engine/Classes/GameFramework/WorldSettings.h b/Engine/Source/Runtime/Engine/Classes/GameFramework/WorldSettings.h
--- a/Engine/Source/Runtime/Engine/Classes/GameFramework/WorldSettings.h
+++ b/Engine/Source/Runtime/Engine/Classes/GameFramework/WorldSettings.h
@@ -477,7 +477,7 @@ public:
 	 * Enables tools for creating a partitioned world. 
 	 * Level has to be saved and all sub-levels removed before enabling this option.
 	 */
-	UPROPERTY(EditAnywhere, BlueprintReadOnly, Category=WorldPartitionSetup)
+	UPROPERTY(BlueprintReadOnly, Category=WorldPartitionSetup)
 	uint8 bEnableWorldPartition:1;
 
 	/** Class of WorldPartitionStreamingPolicy to be used to manage world partition streaming. */
diff --git a/Engine/Source/Runtime/Engine/Private/WorldSettings.cpp b/Engine/Source/Runtime/Engine/Private/WorldSettings.cpp
--- a/Engine/Source/Runtime/Engine/Private/WorldSettings.cpp
+++ b/Engine/Source/Runtime/Engine/Private/WorldSettings.cpp
@@ -602,10 +602,6 @@ bool AWorldSettings::CanEditChange(const FProperty* InProperty) const
 				return LightmassSettings.EnvironmentIntensity > 0;
 			}
 		}
-		else if (PropertyName == GET_MEMBER_NAME_STRING_CHECKED(AWorldSettings, bEnableWorldPartition))
-		{
-			return false;
-		}
 		else if (PropertyName == GET_MEMBER_NAME_STRING_CHECKED(AWorldSettings, bEnableWorldComposition))
 		{
 			return !bEnableWorldPartition;
