Add non-blendable version of TransformAnimationAttribute, FloatAnimationAttribute and IntegerAnimationAttribute

diff --git a/Engine/Source/Runtime/Engine/Classes/Animation/BuiltInAttributeTypes.h b/Engine/Source/Runtime/Engine/Classes/Animation/BuiltInAttributeTypes.h
--- a/Engine/Source/Runtime/Engine/Classes/Animation/BuiltInAttributeTypes.h
+++ b/Engine/Source/Runtime/Engine/Classes/Animation/BuiltInAttributeTypes.h
@@ -141,6 +141,24 @@ struct FTransformAnimationAttribute
 	}
 };
 
+USTRUCT()
+struct FNonBlendableTransformAnimationAttribute : public FTransformAnimationAttribute
+{
+	GENERATED_BODY()
+};
+
+USTRUCT()
+struct FNonBlendableFloatAnimationAttribute : public FFloatAnimationAttribute
+{
+	GENERATED_BODY()
+};
+
+USTRUCT()
+struct FNonBlendableIntegerAnimationAttribute : public FIntegerAnimationAttribute
+{
+	GENERATED_BODY()
+};
+
 namespace UE
 {
 	namespace Anim
@@ -165,6 +183,34 @@ namespace UE
 			};
 		};
 
+		/** Non blendable types*/
+		template<>
+		struct TAttributeTypeTraits<FNonBlendableTransformAnimationAttribute> : public TAttributeTypeTraitsBase<FNonBlendableTransformAnimationAttribute>
+		{
+			enum
+			{
+				IsBlendable = false,
+			};
+		};
+		
+		template<>
+		struct TAttributeTypeTraits<FNonBlendableFloatAnimationAttribute> : public TAttributeTypeTraitsBase<FNonBlendableFloatAnimationAttribute>
+		{
+			enum
+			{
+				IsBlendable = false,
+			};
+		};
+
+		template<>
+		struct TAttributeTypeTraits<FNonBlendableIntegerAnimationAttribute> : public TAttributeTypeTraitsBase<FNonBlendableIntegerAnimationAttribute>
+		{
+			enum
+			{
+				IsBlendable = false,
+			};
+		};
+
 #if WITH_EDITOR
 		/** Helper functionality allowing the user to add an attribute with a typed value array */
 		template<typename AttributeType, typename ValueType>
