Add Neck to regular expressions for the mirror data table

diff --git a/Engine/Source/Runtime/Engine/Private/Animation/AnimationSettings.cpp b/Engine/Source/Runtime/Engine/Private/Animation/AnimationSettings.cpp
--- a/Engine/Source/Runtime/Engine/Private/Animation/AnimationSettings.cpp
+++ b/Engine/Source/Runtime/Engine/Private/Animation/AnimationSettings.cpp
@@ -30,7 +30,7 @@ UAnimationSettings::UAnimationSettings(const FObjectInitializer& ObjectInitializ
 		FMirrorFindReplaceExpression("_R", "_L", EMirrorFindReplaceMethod::Suffix), FMirrorFindReplaceExpression("_L", "_R", EMirrorFindReplaceMethod::Suffix),
 		FMirrorFindReplaceExpression("right", "left", EMirrorFindReplaceMethod::Prefix), FMirrorFindReplaceExpression("left", "right",  EMirrorFindReplaceMethod::Prefix),
 		FMirrorFindReplaceExpression("Right", "Left", EMirrorFindReplaceMethod::Prefix), FMirrorFindReplaceExpression("Left", "Right",  EMirrorFindReplaceMethod::Prefix),
-		FMirrorFindReplaceExpression("((?:^[sS]pine|^[rR]oot|^[pP]elvis|^[hH]ead).*)", "$1", EMirrorFindReplaceMethod::RegularExpression)
+		FMirrorFindReplaceExpression("((?:^[sS]pine|^[rR]oot|^[pP]elvis|^[nN]eck|^[hH]ead).*)", "$1", EMirrorFindReplaceMethod::RegularExpression)
 	};
 }
 
