Improvements to anim node function calls

Split 'node context' into an 'evaluation context' and a 'node reference' type and updated regular function signatures when creating node functions. This way we can seperate concerns (and potentially reference anim nodes standalone in the future).
Moved function caller code into function ref file to stop AnimNodeBase bloat.
Fixed functions not being called on root nodes.

diff --git a/Engine/Source/Editor/AnimGraph/Private/AnimGraphNode_Base.cpp b/Engine/Source/Editor/AnimGraph/Private/AnimGraphNode_Base.cpp
--- a/Engine/Source/Editor/AnimGraph/Private/AnimGraphNode_Base.cpp
+++ b/Engine/Source/Editor/AnimGraph/Private/AnimGraphNode_Base.cpp
@@ -214,15 +214,15 @@ void UAnimGraphNode_Base::ValidateAnimNodeDuringCompilation(USkeleton* ForSkelet
 			UFunction* Function = InRef.ResolveMember<UFunction>(GetAnimBlueprint()->SkeletonGeneratedClass);
 			if(Function == nullptr)
 			{
-				MessageLog.Error(*InError.ToString());
+				MessageLog.Error(*InError.ToString(), this);
 			}
 		}
 	};
 	
-	ValidateFunctionRef(InitializeFunction, LOCTEXT("MissingInitializeFunction", "Could not resolve initialize function"));
-	ValidateFunctionRef(BecomeRelevantFunction, LOCTEXT("MissingBecomeRelevantFunction", "Could not resolve become relevant function"));
-	ValidateFunctionRef(UpdateFunction, LOCTEXT("MissingUpdateFunction", "Could not resolve update function"));
-	ValidateFunctionRef(EvaluateFunction, LOCTEXT("MissingEvaluateFunction", "Could not resolve evaluate function"));
+	ValidateFunctionRef(InitializeFunction, LOCTEXT("MissingInitializeFunction", "Could not resolve initialize function for @@"));
+	ValidateFunctionRef(BecomeRelevantFunction, LOCTEXT("MissingBecomeRelevantFunction", "Could not resolve become relevant function for @@"));
+	ValidateFunctionRef(UpdateFunction, LOCTEXT("MissingUpdateFunction", "Could not resolve update function for @@"));
+	ValidateFunctionRef(EvaluateFunction, LOCTEXT("MissingEvaluateFunction", "Could not resolve evaluate function for @@"));
 }
 
 void UAnimGraphNode_Base::CopyTermDefaultsToDefaultObject(IAnimBlueprintCopyTermDefaultsContext& InCompilationContext, IAnimBlueprintNodeCopyTermDefaultsContext& InPerNodeContext, IAnimBlueprintGeneratedClassCompiledData& OutCompiledData)
diff --git a/Engine/Source/Editor/AnimGraph/Public/AnimGraphNode_Base.h b/Engine/Source/Editor/AnimGraph/Public/AnimGraphNode_Base.h
--- a/Engine/Source/Editor/AnimGraph/Public/AnimGraphNode_Base.h
+++ b/Engine/Source/Editor/AnimGraph/Public/AnimGraphNode_Base.h
@@ -210,19 +210,19 @@ class ANIMGRAPH_API UAnimGraphNode_Base : public UK2Node
 	EBlueprintUsage BlueprintUsage;
 
 	// Function called when the node is initialized
-	UPROPERTY(EditAnywhere, Category = Functions, meta=(FunctionReference, PrototypeFunction="/Script/AnimGraphRuntime.AnimNodeLibrary.Prototype_ThreadSafeAnimNodeCall"), DisplayName="On Initialize")
+	UPROPERTY(EditAnywhere, Category = Functions, meta=(FunctionReference, PrototypeFunction="/Script/AnimGraphRuntime.AnimExecutionContextLibrary.Prototype_ThreadSafeAnimNodeCall"), DisplayName="On Initialize")
 	FMemberReference InitializeFunction;
 
 	// Function called when the node becomes relevant
-	UPROPERTY(EditAnywhere, Category = Functions, meta=(FunctionReference, PrototypeFunction="/Script/AnimGraphRuntime.AnimNodeLibrary.Prototype_ThreadSafeAnimNodeCall"), DisplayName="On Become Relevant")
+	UPROPERTY(EditAnywhere, Category = Functions, meta=(FunctionReference, PrototypeFunction="/Script/AnimGraphRuntime.AnimExecutionContextLibrary.Prototype_ThreadSafeAnimNodeCall"), DisplayName="On Become Relevant")
 	FMemberReference BecomeRelevantFunction;
 
 	// Function called when the node is updated
-	UPROPERTY(EditAnywhere, Category = Functions, meta=(FunctionReference, PrototypeFunction="/Script/AnimGraphRuntime.AnimNodeLibrary.Prototype_ThreadSafeAnimNodeCall"), DisplayName="On Update")
+	UPROPERTY(EditAnywhere, Category = Functions, meta=(FunctionReference, PrototypeFunction="/Script/AnimGraphRuntime.AnimExecutionContextLibrary.Prototype_ThreadSafeAnimNodeCall"), DisplayName="On Update")
 	FMemberReference UpdateFunction;
 
 	// Function called when the node is evaluated
-	UPROPERTY(EditAnywhere, Category = Functions, meta=(FunctionReference, PrototypeFunction="/Script/AnimGraphRuntime.AnimNodeLibrary.Prototype_ThreadSafeAnimNodeCall"), DisplayName="On Evaluate")
+	UPROPERTY(EditAnywhere, Category = Functions, meta=(FunctionReference, PrototypeFunction="/Script/AnimGraphRuntime.AnimExecutionContextLibrary.Prototype_ThreadSafeAnimNodeCall"), DisplayName="On Evaluate")
 	FMemberReference EvaluateFunction;
 	
 	// UObject interface
diff --git a/Engine/Source/Runtime/AnimGraphRuntime/Private/AnimExecutionContextLibrary.cpp b/Engine/Source/Runtime/AnimGraphRuntime/Private/AnimExecutionContextLibrary.cpp
new file mode 100644
--- /dev/null
+++ b/Engine/Source/Runtime/AnimGraphRuntime/Private/AnimExecutionContextLibrary.cpp
@@ -0,0 +1,22 @@
+// Copyright Epic Games, Inc. All Rights Reserved.
+
+#include "AnimExecutionContextLibrary.h"
+#include "Animation/AnimExecutionContext.h"
+#include "Animation/AnimInstance.h"
+#include "Animation/AnimNodeBase.h"
+#include "Animation/AnimClassInterface.h"
+
+UAnimInstance* UAnimExecutionContextLibrary::GetAnimInstance(const FAnimExecutionContext& Context)
+{
+	return CastChecked<UAnimInstance>(Context.GetBaseContext()->GetAnimInstanceObject());
+}
+
+FAnimNodeReference UAnimExecutionContextLibrary::GetAnimNodeReference(UAnimInstance* Instance, int32 Index)
+{
+	IAnimClassInterface* AnimClassInterface = IAnimClassInterface::GetFromClass(Instance->GetClass());
+	const TArray<FStructProperty*>& AnimNodeProperties = AnimClassInterface->GetAnimNodeProperties();
+
+	// As the index is patched during compilation, it needs to be reversed here
+	int32 ReverseIndex = AnimNodeProperties.Num() - 1 - Index;
+	return FAnimNodeReference(Instance, ReverseIndex);
+}
\ No newline at end of file
diff --git a/Engine/Source/Runtime/AnimGraphRuntime/Private/AnimNodeContextLibrary.cpp b/Engine/Source/Runtime/AnimGraphRuntime/Private/AnimNodeContextLibrary.cpp
deleted file mode 100644
diff --git a/Engine/Source/Runtime/AnimGraphRuntime/Private/LinkedAnimGraphLibrary.cpp b/Engine/Source/Runtime/AnimGraphRuntime/Private/LinkedAnimGraphLibrary.cpp
new file mode 100644
--- /dev/null
+++ b/Engine/Source/Runtime/AnimGraphRuntime/Private/LinkedAnimGraphLibrary.cpp
@@ -0,0 +1,38 @@
+// Copyright Epic Games, Inc. All Rights Reserved.
+
+#include "LinkedAnimGraphLibrary.h"
+#include "Animation/AnimInstance.h"
+#include "Animation/AnimNode_LinkedAnimGraph.h"
+
+FLinkedAnimGraphReference ULinkedAnimGraphLibrary::ConvertToLinkedAnimGraphContext(const FAnimNodeReference& Node, EAnimNodeReferenceConversionResult& Result)
+{
+	return FAnimNodeReference::ConvertToType<FLinkedAnimGraphReference>(Node, Result);
+}
+
+bool ULinkedAnimGraphLibrary::HasLinkedAnimInstance(const FLinkedAnimGraphReference& Node)
+{
+	bool bResult = false;
+	
+	Node.CallAnimNodeFunction<FAnimNode_LinkedAnimGraph>(
+		TEXT("HasLinkedAnimInstance"),
+		[&bResult](FAnimNode_LinkedAnimGraph& LinkedAnimGraphNode)
+		{
+			bResult = LinkedAnimGraphNode.GetTargetInstance<UAnimInstance>() != nullptr;
+		});
+
+	return bResult;
+}
+
+UAnimInstance* ULinkedAnimGraphLibrary::GetLinkedAnimInstance(const FLinkedAnimGraphReference& Node)
+{
+	UAnimInstance* Instance = nullptr;
+	
+	Node.CallAnimNodeFunction<FAnimNode_LinkedAnimGraph>(
+		TEXT("GetLinkedAnimInstance"),
+		[&Instance](FAnimNode_LinkedAnimGraph& LinkedAnimGraphNode)
+		{
+			Instance = LinkedAnimGraphNode.GetTargetInstance<UAnimInstance>();
+		});
+
+	return Instance;
+}
\ No newline at end of file
diff --git a/Engine/Source/Runtime/AnimGraphRuntime/Private/SequencePlayerLibrary.cpp b/Engine/Source/Runtime/AnimGraphRuntime/Private/SequencePlayerLibrary.cpp
--- a/Engine/Source/Runtime/AnimGraphRuntime/Private/SequencePlayerLibrary.cpp
+++ b/Engine/Source/Runtime/AnimGraphRuntime/Private/SequencePlayerLibrary.cpp
@@ -1,49 +1,45 @@
 // Copyright Epic Games, Inc. All Rights Reserved.
 
 #include "SequencePlayerLibrary.h"
-#include "Animation/AnimNodeContext.h"
 #include "Animation/AnimNode_SequencePlayer.h"
 
-DEFINE_LOG_CATEGORY_STATIC(LogSequencePlayerLibrary, Verbose, All);
+FSequencePlayerReference USequencePlayerLibrary::ConvertToSequencePlayerContext(const FAnimNodeReference& Node, EAnimNodeReferenceConversionResult& Result)
+{
+	return FAnimNodeReference::ConvertToType<FSequencePlayerReference>(Node, Result);
+}
 
-void USequencePlayerLibrary::SetAccumulatedTime(const FAnimNodeContext& NodeContext, float Time)
+FSequencePlayerReference USequencePlayerLibrary::SetAccumulatedTime(const FSequencePlayerReference& SequencePlayerContext, float Time)
 {
-	if(FAnimNode_SequencePlayer* SequencePlayer = NodeContext.GetAnimNode<FAnimNode_SequencePlayer>())
-	{
-		SequencePlayer->SetAccumulatedTime(Time);
-	}
-	else
-	{
-		UE_LOG(LogSequencePlayerLibrary, Warning, TEXT("SetAccumulatedTime called on a non-sequence player node"));
-	}
+	SequencePlayerContext.CallAnimNodeFunction<FAnimNode_SequencePlayer>(
+		TEXT("SetAccumulatedTime"),
+		[Time](FAnimNode_SequencePlayer& SequencePlayer)
+		{
+			SequencePlayer.SetAccumulatedTime(Time);
+		});
+
+	return SequencePlayerContext;
 }
 
-void USequencePlayerLibrary::SetStartPosition(const FAnimNodeContext& NodeContext, float StartPosition)
+FSequencePlayerReference USequencePlayerLibrary::SetStartPosition(const FSequencePlayerReference& SequencePlayerContext, float StartPosition)
 {
-	if(FAnimNode_SequencePlayer* SequencePlayer = NodeContext.GetAnimNode<FAnimNode_SequencePlayer>())
-	{
-		if(!SequencePlayer->SetStartPosition(StartPosition))
+	SequencePlayerContext.CallAnimNodeFunction<FAnimNode_SequencePlayer>(
+		TEXT("SetStartPosition"),
+		[StartPosition](FAnimNode_SequencePlayer& SequencePlayer)
 		{
-			UE_LOG(LogSequencePlayerLibrary, Warning, TEXT("Could not set start time on sequence player, value is not dynamic. Set it as Always Dynamic or expose it as a pin"));
-		}
-	}
-	else
-	{
-		UE_LOG(LogSequencePlayerLibrary, Warning, TEXT("SetStartTime called on a non-sequence player node"));
-	}
+			SequencePlayer.SetStartPosition(StartPosition);
+		});
+
+	return SequencePlayerContext;
 }
 
-void USequencePlayerLibrary::SetPlayRate(const FAnimNodeContext& NodeContext, float PlayRate)
+FSequencePlayerReference USequencePlayerLibrary::SetPlayRate(const FSequencePlayerReference& SequencePlayerContext, float PlayRate)
 {
-	if(FAnimNode_SequencePlayer* SequencePlayer = NodeContext.GetAnimNode<FAnimNode_SequencePlayer>())
-	{
-		if(!SequencePlayer->SetPlayRate(PlayRate))
+	SequencePlayerContext.CallAnimNodeFunction<FAnimNode_SequencePlayer>(
+		TEXT("SetPlayRate"),
+		[PlayRate](FAnimNode_SequencePlayer& SequencePlayer)
 		{
-			UE_LOG(LogSequencePlayerLibrary, Warning, TEXT("Could not set start time on sequence player, value is not dynamic. Set it as Always Dynamic or expose it as a pin"));
-		}
-	}
-	else
-	{
-		UE_LOG(LogSequencePlayerLibrary, Warning, TEXT("SetStartTime called on a non-sequence player node"));
-	}
+			SequencePlayer.SetPlayRate(PlayRate);
+		});
+
+	return SequencePlayerContext;
 }
\ No newline at end of file
diff --git a/Engine/Source/Runtime/AnimGraphRuntime/Public/AnimExecutionContextLibrary.h b/Engine/Source/Runtime/AnimGraphRuntime/Public/AnimExecutionContextLibrary.h
new file mode 100644
--- /dev/null
+++ b/Engine/Source/Runtime/AnimGraphRuntime/Public/AnimExecutionContextLibrary.h
@@ -0,0 +1,33 @@
+// Copyright Epic Games, Inc. All Rights Reserved.
+
+#pragma once
+
+#include "CoreMinimal.h"
+#include "Kismet/BlueprintFunctionLibrary.h"
+#include "Animation/AnimNodeReference.h"
+#include "Animation/AnimExecutionContext.h"
+#include "AnimExecutionContextLibrary.generated.h"
+
+class UAnimInstance;
+
+// Exposes operations to be performed on anim node contexts
+UCLASS()
+class ANIMGRAPHRUNTIME_API UAnimExecutionContextLibrary : public UBlueprintFunctionLibrary
+{
+	GENERATED_BODY()
+	
+public:
+#if WITH_EDITOR
+	/** Prototype function for thread-safe anim node calls */
+	UFUNCTION(BlueprintInternalUseOnly, meta=(BlueprintThreadSafe))
+	void Prototype_ThreadSafeAnimNodeCall(const FAnimExecutionContext& Context, const FAnimNodeReference& Node) {}
+#endif
+	
+	/** Get the anim instance that hosts this anim node */
+	UFUNCTION(BlueprintPure, Category = "Anim Node", meta=(BlueprintThreadSafe))
+	static UAnimInstance* GetAnimInstance(const FAnimExecutionContext& Context);
+
+	/** Internal compiler use only - Get a reference to an anim node by index */
+	UFUNCTION(BlueprintPure, BlueprintInternalUseOnly, meta=(BlueprintThreadSafe))
+	static FAnimNodeReference GetAnimNodeReference(UAnimInstance* Instance, int32 Index);
+};
\ No newline at end of file
diff --git a/Engine/Source/Runtime/AnimGraphRuntime/Public/AnimNodeContextLibrary.h b/Engine/Source/Runtime/AnimGraphRuntime/Public/AnimNodeContextLibrary.h
deleted file mode 100644
diff --git a/Engine/Source/Runtime/AnimGraphRuntime/Public/LinkedAnimGraphLibrary.h b/Engine/Source/Runtime/AnimGraphRuntime/Public/LinkedAnimGraphLibrary.h
new file mode 100644
--- /dev/null
+++ b/Engine/Source/Runtime/AnimGraphRuntime/Public/LinkedAnimGraphLibrary.h
@@ -0,0 +1,38 @@
+// Copyright Epic Games, Inc. All Rights Reserved.
+
+#pragma once
+
+#include "CoreMinimal.h"
+#include "Kismet/BlueprintFunctionLibrary.h"
+#include "Animation/AnimNodeReference.h"
+#include "LinkedAnimGraphLibrary.generated.h"
+
+struct FAnimNode_LinkedAnimGraph;
+
+USTRUCT(BlueprintType)
+struct FLinkedAnimGraphReference : public FAnimNodeReference
+{
+	GENERATED_BODY()
+
+	typedef FAnimNode_LinkedAnimGraph FInternalNodeType;
+};
+
+// Exposes operations to be performed on anim node contexts
+UCLASS()
+class ANIMGRAPHRUNTIME_API ULinkedAnimGraphLibrary : public UBlueprintFunctionLibrary
+{
+	GENERATED_BODY()
+
+public:
+	/** Get a sequence player context from an anim node context */
+	UFUNCTION(BlueprintCallable, Category = "Linked Anim Graph", meta=(BlueprintThreadSafe, ExpandEnumAsExecs = "Result"))
+	static FLinkedAnimGraphReference ConvertToLinkedAnimGraphContext(const FAnimNodeReference& Node, EAnimNodeReferenceConversionResult& Result);
+	
+	/** Returns whether the node hosts an instance (e.g. linked anim graph or layer) */
+	UFUNCTION(BlueprintPure, Category = "Linked Anim Graph", meta=(BlueprintThreadSafe))
+	static bool HasLinkedAnimInstance(const FLinkedAnimGraphReference& Node);
+
+	/** Get the linked instance is hosted by this node. If the node does not host an instance then HasLinkedAnimInstance will return false */
+	UFUNCTION(BlueprintPure, Category = "Linked Anim Graph", meta=(BlueprintThreadSafe))
+	static UAnimInstance* GetLinkedAnimInstance(const FLinkedAnimGraphReference& Node);
+};
\ No newline at end of file
diff --git a/Engine/Source/Runtime/AnimGraphRuntime/Public/SequencePlayerLibrary.h b/Engine/Source/Runtime/AnimGraphRuntime/Public/SequencePlayerLibrary.h
--- a/Engine/Source/Runtime/AnimGraphRuntime/Public/SequencePlayerLibrary.h
+++ b/Engine/Source/Runtime/AnimGraphRuntime/Public/SequencePlayerLibrary.h
@@ -4,9 +4,19 @@
 
 #include "CoreMinimal.h"
 #include "Kismet/BlueprintFunctionLibrary.h"
-#include "Animation/AnimNodeContext.h"
+#include "Animation/AnimNodeReference.h"
 #include "SequencePlayerLibrary.generated.h"
 
+struct FAnimNode_SequencePlayer;
+
+USTRUCT(BlueprintType)
+struct FSequencePlayerReference : public FAnimNodeReference
+{
+	GENERATED_BODY()
+
+	typedef FAnimNode_SequencePlayer FInternalNodeType;
+};
+
 // Exposes operations to be performed on a sequence player anim node
 // Note: Experimental and subject to change!
 UCLASS(Experimental)
@@ -14,14 +24,19 @@ class ANIMGRAPHRUNTIME_API USequencePlayerLibrary : public UBlueprintFunctionLib
 {
 	GENERATED_BODY()
 
+public:
+	/** Get a sequence player context from an anim node context */
+	UFUNCTION(BlueprintCallable, Category = "Sequence Player", meta=(BlueprintThreadSafe, ExpandEnumAsExecs = "Result"))
+	static FSequencePlayerReference ConvertToSequencePlayerContext(const FAnimNodeReference& Node, EAnimNodeReferenceConversionResult& Result);
+	
 	/** Set the current accumulated time of the sequence player */
 	UFUNCTION(BlueprintCallable, Category = "Sequence Player", meta=(BlueprintThreadSafe))
-	static void SetAccumulatedTime(const FAnimNodeContext& NodeContext, float Time);
+	static FSequencePlayerReference SetAccumulatedTime(const FSequencePlayerReference& SequencePlayer, float Time);
 
 	/** Set the start position of the sequence player */
 	UFUNCTION(BlueprintCallable, Category = "Sequence Player", meta=(BlueprintThreadSafe))
-	static void SetStartPosition(const FAnimNodeContext& NodeContext, float StartPosition);
+	static FSequencePlayerReference SetStartPosition(const FSequencePlayerReference& SequencePlayer, float StartPosition);
 
 	UFUNCTION(BlueprintCallable, Category = "Sequence Player", meta=(BlueprintThreadSafe))
-	static void SetPlayRate(const FAnimNodeContext& NodeContext, float PlayRate);
+	static FSequencePlayerReference SetPlayRate(const FSequencePlayerReference& SequencePlayer, float PlayRate);
 };
\ No newline at end of file
diff --git a/Engine/Source/Runtime/Engine/Classes/Animation/AnimExecutionContext.h b/Engine/Source/Runtime/Engine/Classes/Animation/AnimExecutionContext.h
new file mode 100644
--- /dev/null
+++ b/Engine/Source/Runtime/Engine/Classes/Animation/AnimExecutionContext.h
@@ -0,0 +1,202 @@
+// Copyright Epic Games, Inc. All Rights Reserved.
+
+#pragma once
+
+#include "CoreMinimal.h"
+#include "AnimExecutionContext.generated.h"
+
+struct FAnimationBaseContext;
+struct FAnimationInitializeContext;
+struct FAnimationUpdateContext;
+struct FPoseContext;
+struct FComponentSpacePoseContext;
+
+// The result of an anim node context conversion 
+UENUM(BlueprintType)
+enum class EAnimExecutionContextConversionResult : uint8
+{
+	Succeeded = 1,
+	Failed = 0,
+};
+
+// Context used to expose anim graph execution to BP function libraries
+USTRUCT(BlueprintType)
+struct ENGINE_API FAnimExecutionContext
+{
+	GENERATED_BODY()
+
+public:
+	// Internal data, weakly referenced
+	struct FData
+	{
+	public:
+		FData(const FAnimationInitializeContext& InContext);
+
+		FData(const FAnimationUpdateContext& InContext);
+
+		FData(FPoseContext& InContext);
+
+		FData(FComponentSpacePoseContext& InContext);
+
+	private:
+		friend struct FAnimExecutionContext;
+		friend struct FAnimInitializationContext;
+		friend struct FAnimUpdateContext;
+		friend struct FAnimPoseContext;
+		friend struct FAnimComponentSpacePoseContext;
+		
+		enum class EContextType
+		{
+			None,
+			Initialize,
+			Update,
+			Pose,
+			ComponentSpacePose,
+		};
+	
+		// The context used when executing this node, e.g. FAnimationUpdateContext, FPoseContext etc.
+		FAnimationBaseContext* Context = nullptr;
+
+		// The phase we are in
+		EContextType ContextType = EContextType::None;	
+	};
+	
+	FAnimExecutionContext() = default;
+	FAnimExecutionContext(const FAnimExecutionContext& InOther) = default;
+	FAnimExecutionContext& operator=(const FAnimExecutionContext& InOther) = default;
+	
+	FAnimExecutionContext(TWeakPtr<FData> InData)
+		: Data(InData)
+	{}
+
+	// Is this a valid context? 
+	bool IsValid() const
+	{
+		return Data.IsValid();
+	}
+
+	// Convert to a derived type
+	template<typename OtherContextType>
+	static OtherContextType ConvertToType(const FAnimExecutionContext& InContext, EAnimExecutionContextConversionResult& OutResult)
+	{
+		static_assert(TIsDerivedFrom<OtherContextType, FAnimExecutionContext>::IsDerived, "Argument ContextType must derive from FAnimExecutionContext");
+		
+		if(TSharedPtr<FData> PinnedData = InContext.Data.Pin())
+		{
+			if(OtherContextType::InternalContextTypeId == PinnedData->ContextType)
+			{
+				OutResult = EAnimExecutionContextConversionResult::Succeeded;
+				
+				OtherContextType Context;
+				Context.Data = InContext.Data;
+				return Context;
+			}
+		}
+
+		OutResult = EAnimExecutionContextConversionResult::Failed;
+		
+		return OtherContextType();
+	}
+
+	// Access internal context. Will return nullptr if invalid or an incorrect type is requested
+	template<typename InternalContextType>
+	InternalContextType* GetContext() const
+	{
+		if(TSharedPtr<FData> PinnedData = Data.Pin())
+		{
+			if(InternalContextType::InternalContextTypeId == PinnedData->ContextType)
+			{
+				return static_cast<InternalContextType*>(PinnedData->Context);
+			}
+		}
+
+		return nullptr;
+	}
+
+	// Access internal context. Will return nullptr if invalid
+	FAnimationBaseContext* GetBaseContext() const
+	{
+		if(TSharedPtr<FData> PinnedData = Data.Pin())
+		{
+			return PinnedData->Context;
+		}
+
+		return nullptr;
+	}
+	
+protected:
+	// Internal data
+	TWeakPtr<FData> Data; 
+};
+
+USTRUCT()
+struct FAnimInitializationContext : public FAnimExecutionContext
+{
+	GENERATED_BODY()
+
+public:
+	static const FAnimExecutionContext::FData::EContextType InternalContextTypeId = FData::EContextType::Initialize;
+
+	FAnimInitializationContext() = default;
+	
+	FAnimInitializationContext(TWeakPtr<FData> InData)
+		: FAnimExecutionContext(InData)
+	{
+		check(InData.IsValid());
+		check(InData.Pin()->ContextType == InternalContextTypeId);
+	}
+};
+
+USTRUCT()
+struct FAnimUpdateContext : public FAnimExecutionContext
+{
+	GENERATED_BODY()
+
+public:
+	static const FData::EContextType InternalContextTypeId = FData::EContextType::Update;
+
+	FAnimUpdateContext() = default;
+	
+	FAnimUpdateContext(TWeakPtr<FData> InData)
+		: FAnimExecutionContext(InData)
+	{
+		check(InData.IsValid());
+		check(InData.Pin()->ContextType == InternalContextTypeId);
+	}
+};
+
+USTRUCT()
+struct FAnimPoseContext : public FAnimExecutionContext
+{
+	GENERATED_BODY()
+
+public:
+	static const FData::EContextType InternalContextTypeId = FData::EContextType::Pose;
+
+	FAnimPoseContext() = default;
+	
+	FAnimPoseContext(TWeakPtr<FData> InData)
+		: FAnimExecutionContext(InData)
+	{
+		check(InData.IsValid());
+		check(InData.Pin()->ContextType == InternalContextTypeId);
+	}
+};
+
+USTRUCT()
+struct FAnimComponentSpacePoseContext : public FAnimExecutionContext
+{
+	GENERATED_BODY()
+
+public:
+	static const FData::EContextType InternalContextTypeId = FData::EContextType::ComponentSpacePose;
+
+	FAnimComponentSpacePoseContext() = default;
+	
+	FAnimComponentSpacePoseContext(TWeakPtr<FData> InData)
+		: FAnimExecutionContext(InData)
+	{
+		check(InData.IsValid());
+		check(InData.Pin()->ContextType == InternalContextTypeId);
+	}
+};
\ No newline at end of file
diff --git a/Engine/Source/Runtime/Engine/Classes/Animation/AnimNodeContext.h b/Engine/Source/Runtime/Engine/Classes/Animation/AnimNodeContext.h
deleted file mode 100644
diff --git a/Engine/Source/Runtime/Engine/Classes/Animation/AnimNodeFunctionRef.h b/Engine/Source/Runtime/Engine/Classes/Animation/AnimNodeFunctionRef.h
--- a/Engine/Source/Runtime/Engine/Classes/Animation/AnimNodeFunctionRef.h
+++ b/Engine/Source/Runtime/Engine/Classes/Animation/AnimNodeFunctionRef.h
@@ -5,6 +5,12 @@
 #include "CoreMinimal.h"
 #include "AnimNodeFunctionRef.generated.h"
 
+struct FAnimNode_Base;
+struct FAnimationInitializeContext;
+struct FAnimationUpdateContext;
+struct FPoseContext;
+struct FComponentSpacePoseContext;
+
 /**
  * Cached function name/ptr that is resolved at init time
  */
@@ -43,4 +49,29 @@ private:
 	// The function to call, recovered by looking for a function of name FunctionName
 	UPROPERTY(Transient)
 	UFunction* Function = nullptr;
+};
+
+// Wrapper used to call anim node functions
+struct FAnimNodeFunctionCaller
+{
+private:
+	friend struct FPoseLinkBase;
+	friend struct FPoseLink;
+	friend struct FComponentSpacePoseLink;
+	friend struct FAnimInstanceProxy;
+	
+	// Call the Initialize function of this node
+	static void Initialize(const FAnimationInitializeContext& InContext, FAnimNode_Base& InNode);
+
+	// Call the BecomeRelevant function of this node
+	static void BecomeRelevant(const FAnimationUpdateContext& InContext, FAnimNode_Base& InNode);
+
+	// Call the Update function of this node
+	static void Update(const FAnimationUpdateContext& InContext, FAnimNode_Base& InNode);
+
+	// Call the Evaluate function of this node
+	static void Evaluate(FPoseContext& InContext, FAnimNode_Base& InNode);
+
+	// Call the EvaluateComponentSpace function of this node
+	static void EvaluateComponentSpace(FComponentSpacePoseContext& InContext, FAnimNode_Base& InNode);
 };
\ No newline at end of file
diff --git a/Engine/Source/Runtime/Engine/Classes/Animation/AnimNodeReference.h b/Engine/Source/Runtime/Engine/Classes/Animation/AnimNodeReference.h
new file mode 100644
--- /dev/null
+++ b/Engine/Source/Runtime/Engine/Classes/Animation/AnimNodeReference.h
@@ -0,0 +1,99 @@
+// Copyright Epic Games, Inc. All Rights Reserved.
+
+#pragma once
+
+#include "CoreMinimal.h"
+#include "EngineLogs.h"
+#include "AnimNodeReference.generated.h"
+
+struct FAnimNode_Base;
+class IAnimClassInterface;
+class UAnimInstance;
+
+// The result of an anim node reference conversion 
+UENUM(BlueprintType)
+enum class EAnimNodeReferenceConversionResult : uint8
+{
+	Succeeded = 1,
+	Failed = 0,
+};
+
+// A reference to an anim node. Does not persist, only valid for the call in which it was retrieved.
+USTRUCT(BlueprintType)
+struct ENGINE_API FAnimNodeReference
+{
+	GENERATED_BODY()
+
+public:
+	typedef FAnimNode_Base FInternalNodeType; 
+		
+	FAnimNodeReference() = default;
+
+	FAnimNodeReference(UAnimInstance* InAnimInstance, FAnimNode_Base& InNode);
+
+	FAnimNodeReference(UAnimInstance* InAnimInstance, int32 InIndex);
+	
+	// Get the node we wrap. If the context is invalid or the node is not of the specified type then this will return nullptr.
+	template<typename NodeType>
+	NodeType* GetAnimNodePtr() const
+	{
+		if(AnimNode && AnimNodeStruct && NodeType::StaticStruct()->IsChildOf(AnimNodeStruct))
+		{
+			return static_cast<NodeType*>(AnimNode);
+		}
+
+		return nullptr;
+	}
+
+	// Get the node we wrap. If the reference is invalid or node is not of the specified type then this will return assert.
+	template<typename NodeType>
+	NodeType& GetAnimNode() const
+	{
+		check(AnimNodeStruct);
+		check(AnimNode);
+		check(NodeType::StaticStruct()->IsChildOf(AnimNodeStruct));
+		return *static_cast<NodeType*>(AnimNode);
+	}
+
+	// Call a function if this context is valid
+	template<typename NodeType>
+	void CallAnimNodeFunction(const TCHAR* InFunctionNameForErrorReporting, TFunctionRef<void(NodeType&)> InFunction) const
+	{
+		if(NodeType* NodePtr = GetAnimNodePtr<NodeType>())
+		{
+			InFunction(*NodePtr);
+		}
+		else
+		{
+			UE_LOG(LogAnimation, Warning, TEXT("%s called on an invalid context or with an invalid type"), InFunctionNameForErrorReporting);
+		}
+	}
+	
+	// Convert to a derived type
+	template<typename OtherContextType>
+	static OtherContextType ConvertToType(const FAnimNodeReference& InReference, EAnimNodeReferenceConversionResult& OutResult)
+	{
+		static_assert(TIsDerivedFrom<OtherContextType, FAnimNodeReference>::IsDerived, "Argument ContextType must derive from FAnimNodeReference");
+		
+		if(InReference.AnimNodeStruct && OtherContextType::FInternalNodeType::StaticStruct()->IsChildOf(InReference.AnimNodeStruct))
+		{
+			OutResult = EAnimNodeReferenceConversionResult::Succeeded;
+			
+			OtherContextType Context;
+			Context.AnimNode = InReference.AnimNode;
+			Context.AnimNodeStruct = InReference.AnimNodeStruct;
+			return Context;
+		}
+
+		OutResult = EAnimNodeReferenceConversionResult::Failed;
+		
+		return OtherContextType();
+	}
+
+private:
+	// The node we wrap
+	FAnimNode_Base* AnimNode = nullptr;
+
+	// The struct type of the anim node
+	UScriptStruct* AnimNodeStruct = nullptr;
+};
\ No newline at end of file
diff --git a/Engine/Source/Runtime/Engine/Private/Animation/AnimExecutionContext.cpp b/Engine/Source/Runtime/Engine/Private/Animation/AnimExecutionContext.cpp
new file mode 100644
--- /dev/null
+++ b/Engine/Source/Runtime/Engine/Private/Animation/AnimExecutionContext.cpp
@@ -0,0 +1,28 @@
+// Copyright Epic Games, Inc. All Rights Reserved.
+
+#include "Animation/AnimExecutionContext.h"
+#include "Animation/AnimNodeBase.h"
+
+FAnimExecutionContext::FData::FData(const FAnimationInitializeContext& InContext)
+{
+	Context = const_cast<FAnimationInitializeContext*>(&InContext);
+	ContextType = EContextType::Initialize;
+}
+
+FAnimExecutionContext::FData::FData(const FAnimationUpdateContext& InContext)
+{
+	Context = const_cast<FAnimationUpdateContext*>(&InContext);
+	ContextType = EContextType::Update;
+}
+
+FAnimExecutionContext::FData::FData(FPoseContext& InContext)
+{
+	Context = &InContext;
+	ContextType = EContextType::Pose;
+}
+
+FAnimExecutionContext::FData::FData(FComponentSpacePoseContext& InContext)
+{
+	Context = &InContext;
+	ContextType = EContextType::ComponentSpacePose;
+}
diff --git a/Engine/Source/Runtime/Engine/Private/Animation/AnimInstanceProxy.cpp b/Engine/Source/Runtime/Engine/Private/Animation/AnimInstanceProxy.cpp
--- a/Engine/Source/Runtime/Engine/Private/Animation/AnimInstanceProxy.cpp
+++ b/Engine/Source/Runtime/Engine/Private/Animation/AnimInstanceProxy.cpp
@@ -58,7 +58,9 @@ void FAnimInstanceProxy::UpdateAnimationNode_WithRoot(const FAnimationUpdateCont
 		{
 			UpdateCounter.Increment();
 		}
-		
+
+		FAnimNodeFunctionCaller::BecomeRelevant(InContext, *InRootNode);
+		FAnimNodeFunctionCaller::Update(InContext, *InRootNode);
 		InRootNode->Update_AnyThread(InContext);
 
 		// We've updated the graph, now update the fractured saved pose sections
@@ -295,6 +297,8 @@ void FAnimInstanceProxy::InitializeRootNode_WithRoot(FAnimNode_Base* InRootNode)
 		FAnimationUpdateSharedContext SharedContext;
 		FAnimationInitializeContext InitContext(this, &SharedContext);
 
+		FAnimNodeFunctionCaller::Initialize(InitContext, *InRootNode);
+		
 		if(InRootNode == RootNode)
 		{
 			InitializationCounter.Increment();
@@ -1289,6 +1293,8 @@ void FAnimInstanceProxy::EvaluateAnimationNode_WithRoot(FPoseContext& Output, FA
 
 		TRACE_SCOPED_ANIM_NODE(Output);
 
+		FAnimNodeFunctionCaller::Evaluate(Output, *InRootNode);
+		
 		InRootNode->Evaluate_AnyThread(Output);
 	}
 	else
diff --git a/Engine/Source/Runtime/Engine/Private/Animation/AnimNodeBase.cpp b/Engine/Source/Runtime/Engine/Private/Animation/AnimNodeBase.cpp
--- a/Engine/Source/Runtime/Engine/Private/Animation/AnimNodeBase.cpp
+++ b/Engine/Source/Runtime/Engine/Private/Animation/AnimNodeBase.cpp
@@ -6,11 +6,7 @@
 #include "Animation/AnimTrace.h"
 #include "UObject/CoreObjectVersion.h"
 #include "Animation/AnimAttributes.h"
-#include "Animation/AnimNode_LinkedAnimGraph.h"
 #include "Animation/AnimSubsystem_Base.h"
-#include "Animation/AnimRootMotionProvider.h"
-#include "Animation/AnimSubsystem_NodeRelevancy.h"
-#include "Animation/AnimNodeContext.h"
 
 /////////////////////////////////////////////////////
 // FAnimationBaseContext
@@ -173,62 +169,6 @@ const FAnimNodeFunctionRef& FAnimNode_Base::GetEvaluateFunction() const
 	return GET_ANIM_NODE_DATA(FAnimNodeFunctionRef, EvaluateFunction);
 }
 
-// Private impl used to call anim node functions from pose links
-struct FAnimNodeFunctionCaller
-{
-private:
-	template<typename ContextType>
-	static void CallFunctionHelper(const FAnimNodeFunctionRef& InFunction, ContextType InContext, FAnimNode_Base& InNode)
-	{
-		if(InFunction.IsValid())
-		{
-			IAnimClassInterface* AnimBlueprintClass = InContext.GetAnimClass();
-			check(AnimBlueprintClass);
-
-			const TArray<FStructProperty*>& AnimNodeProperties = AnimBlueprintClass->GetAnimNodeProperties();
-			UScriptStruct* ScriptStruct = AnimNodeProperties[InNode.GetNodeIndex()]->Struct;
-			TSharedRef<FAnimNodeContext::FData> ContextData = MakeShared<FAnimNodeContext::FData>(InContext, InNode, ScriptStruct);
-			FAnimNodeContext NodeContext(ContextData);
-			InFunction.Call(InContext.GetAnimInstanceObject(), &NodeContext);
-		}
-	}
-
-public:
-	static void Initialize(const FAnimationInitializeContext& InContext, FAnimNode_Base& InNode)
-	{
-		CallFunctionHelper(InNode.GetInitializeFunction(), InContext, InNode);
-	}
-	
-	static void BecomeRelevant(const FAnimationUpdateContext& InContext, FAnimNode_Base& InNode)
-	{
-		const FAnimNodeFunctionRef& Function = InNode.GetBecomeRelevantFunction();
-		if(Function.IsValid())
-		{
-			FAnimSubsystemInstance_NodeRelevancy& RelevancySubsystem = CastChecked<UAnimInstance>(InContext.GetAnimInstanceObject())->GetSubsystem<FAnimSubsystemInstance_NodeRelevancy>();
-			FAnimNodeRelevancyStatus Status = RelevancySubsystem.UpdateNodeRelevancy(InContext, InNode);
-			if(Status.HasJustBecomeRelevant())
-			{
-				CallFunctionHelper(Function, InContext, InNode);
-			}
-		}
-	}
-	
-	static void Update(const FAnimationUpdateContext& InContext, FAnimNode_Base& InNode)
-	{
-		CallFunctionHelper(InNode.GetUpdateFunction(), InContext, InNode);
-	}
-	
-	static void Evaluate(FPoseContext& InContext, FAnimNode_Base& InNode)
-	{
-		CallFunctionHelper(InNode.GetEvaluateFunction(), InContext, InNode);
-	}
-
-	static void EvaluateComponentSpace(FComponentSpacePoseContext& InContext, FAnimNode_Base& InNode)
-	{
-		CallFunctionHelper(InNode.GetEvaluateFunction(), InContext, InNode);
-	}
-};
-
 /////////////////////////////////////////////////////
 // FPoseLinkBase
 
diff --git a/Engine/Source/Runtime/Engine/Private/Animation/AnimNodeContext.cpp b/Engine/Source/Runtime/Engine/Private/Animation/AnimNodeContext.cpp
deleted file mode 100644
diff --git a/Engine/Source/Runtime/Engine/Private/Animation/AnimNodeFunctionRef.cpp b/Engine/Source/Runtime/Engine/Private/Animation/AnimNodeFunctionRef.cpp
--- a/Engine/Source/Runtime/Engine/Private/Animation/AnimNodeFunctionRef.cpp
+++ b/Engine/Source/Runtime/Engine/Private/Animation/AnimNodeFunctionRef.cpp
@@ -1,6 +1,11 @@
 // Copyright Epic Games, Inc. All Rights Reserved.
 
 #include "Animation/AnimNodeFunctionRef.h"
+#include "Animation/AnimNodeBase.h"
+#include "Animation/AnimExecutionContext.h"
+#include "Animation/AnimNodeReference.h"
+#include "Animation/AnimSubsystem_NodeRelevancy.h"
+#include "Animation/AnimInstance.h"
 
 void FAnimNodeFunctionRef::Initialize(const UClass* InClass)
 {
@@ -16,4 +21,59 @@ void FAnimNodeFunctionRef::Call(UObject* InObject, void* InParameters) const
 	{
 		InObject->ProcessEvent(Function, InParameters);
 	}
+}
+
+template<typename ContextType>
+static void CallFunctionHelper(const FAnimNodeFunctionRef& InFunction, ContextType InContext, FAnimNode_Base& InNode)
+{
+	if(InFunction.IsValid())
+	{
+		UAnimInstance* AnimInstance = CastChecked<UAnimInstance>(InContext.GetAnimInstanceObject());
+		
+		TSharedRef<FAnimExecutionContext::FData> ContextData = MakeShared<FAnimExecutionContext::FData>(InContext);
+			
+		struct FAnimNodeFunctionParams
+		{
+			FAnimExecutionContext ExecutionContext;
+			FAnimNodeReference NodeReference;
+		};
+			
+		FAnimNodeFunctionParams Params = { FAnimExecutionContext(ContextData), FAnimNodeReference(AnimInstance, InNode) };
+			
+		InFunction.Call(AnimInstance, &Params);
+	}
+}
+
+void FAnimNodeFunctionCaller::Initialize(const FAnimationInitializeContext& InContext, FAnimNode_Base& InNode)
+{
+	CallFunctionHelper(InNode.GetInitializeFunction(), InContext, InNode);
+}
+	
+void FAnimNodeFunctionCaller::BecomeRelevant(const FAnimationUpdateContext& InContext, FAnimNode_Base& InNode)
+{
+	const FAnimNodeFunctionRef& Function = InNode.GetBecomeRelevantFunction();
+	if(Function.IsValid())
+	{
+		FAnimSubsystemInstance_NodeRelevancy& RelevancySubsystem = CastChecked<UAnimInstance>(InContext.GetAnimInstanceObject())->GetSubsystem<FAnimSubsystemInstance_NodeRelevancy>();
+		FAnimNodeRelevancyStatus Status = RelevancySubsystem.UpdateNodeRelevancy(InContext, InNode);
+		if(Status.HasJustBecomeRelevant())
+		{
+			CallFunctionHelper(Function, InContext, InNode);
+		}
+	}
+}
+	
+void FAnimNodeFunctionCaller::Update(const FAnimationUpdateContext& InContext, FAnimNode_Base& InNode)
+{
+	CallFunctionHelper(InNode.GetUpdateFunction(), InContext, InNode);
+}
+
+void FAnimNodeFunctionCaller::Evaluate(FPoseContext& InContext, FAnimNode_Base& InNode)
+{
+	CallFunctionHelper(InNode.GetEvaluateFunction(), InContext, InNode);
+}
+
+void FAnimNodeFunctionCaller::EvaluateComponentSpace(FComponentSpacePoseContext& InContext, FAnimNode_Base& InNode)
+{
+	CallFunctionHelper(InNode.GetEvaluateFunction(), InContext, InNode);
 }
\ No newline at end of file
diff --git a/Engine/Source/Runtime/Engine/Private/Animation/AnimNodeReference.cpp b/Engine/Source/Runtime/Engine/Private/Animation/AnimNodeReference.cpp
new file mode 100644
--- /dev/null
+++ b/Engine/Source/Runtime/Engine/Private/Animation/AnimNodeReference.cpp
@@ -0,0 +1,24 @@
+// Copyright Epic Games, Inc. All Rights Reserved.
+
+#include "Animation/AnimNodeReference.h"
+#include "Animation/AnimNodeBase.h"
+#include "Animation/AnimClassInterface.h"
+
+FAnimNodeReference::FAnimNodeReference(UAnimInstance* InAnimInstance, FAnimNode_Base& InNode)
+	: AnimNode(&InNode)
+{
+	IAnimClassInterface* AnimClassInterface = IAnimClassInterface::GetFromClass(InAnimInstance->GetClass());
+	const TArray<FStructProperty*>& AnimNodeProperties = AnimClassInterface->GetAnimNodeProperties();
+	AnimNodeStruct = AnimNodeProperties[InNode.GetNodeIndex()]->Struct;
+}
+
+FAnimNodeReference::FAnimNodeReference(UAnimInstance* InAnimInstance, int32 InIndex)
+{
+	IAnimClassInterface* AnimClassInterface = IAnimClassInterface::GetFromClass(InAnimInstance->GetClass());
+	const TArray<FStructProperty*>& AnimNodeProperties = AnimClassInterface->GetAnimNodeProperties();
+	if(AnimNodeProperties.IsValidIndex(InIndex))
+	{
+		AnimNodeStruct = AnimNodeProperties[InIndex]->Struct;
+		AnimNode = AnimNodeProperties[InIndex]->ContainerPtrToValuePtr<FAnimNode_Base>(InAnimInstance);
+	}
+}
\ No newline at end of file
