Remove UE_USE_VIRTUALBULKDATA and the associated code path, there is now no easy way to revert back to old bulkdata format for StaticMesh/Textures.

#rnx

* Macro removal
- UE_USE_VIRTUALBULKDATA: Asssumed to be 1.
- UE_VBD_TO_OLD_BULKDATA_PATH: Assumed to be 0 and code paths entirely deleted.
- UE_VBD_CONST: Assumed to be'const'
- Now that these macros are no longer in use some includes of VirtualizationManager.h can be removed (or moved to the cpp).

*VirtualizedBulkData
- No longer need need to include VirtualizationManager.h in the bulkdata header file so that is moved to the cpp.
- Removed ::GetBulkDataSize which was only added to keep compatibility to older code paths. This method has not yet been released publicly(it's not in EA) so we don't need to worry that anyone is actually using it.
- Add a new method ::HasPayloadData that allows the caller to poll if the payload exists and if it has a valid data length rather than calling ::GetPayloadSize (I'd like to move away from the need for calling code to know how large the payload is until the payload is actually pulled as it will make eventual thread safety code easier to do in the future).
- Added a method ::SetCompressionFormat that allows the compression format to be changed. Normally this should only be done when ::UpdatePayload is called but we need a way to fix old textures that were saved with the wrong compression flags for a time and to do that we need to be able to 'reset' the compression format in some cases. The method is documented to try and make it clear that this is not really intended for general use. At some point in the future if the engine mandates that packages be resaved to a specific version we can look into deprecating this call.

*StaticMeshSourceData
- MESHDATAKEY_STATICMESH_DERIVEDDATA_VER is currently 1 character shorter than it should be, but that doesn't seem worth changing the key to fix.

*Texture
- When accessing data from IImageWrapper we can take the data pointer directly from the TArray and give it's ownership to the FSharedBuffer being returned rather than cloning it.
- Deprecated FTextureSource::LoadBulkDataWithFileReader as with VBD it will always return true.
- FMipAllocation cannot yet be deleted as it is still required for Lock/Unlock but we can remove the support for old style FByteBulkData for now. There is further clean up work to be done here but we would have to make bigger changes to the FTextureSource api.
- FMipData can now store the data in FSharedBuffer format. In the future we can allow people to access mips via views rather than always taking copies to reduce unnecessary data transfer. This also eliminates the need to pass in a lock.
- FTextureSource::InitLayered still supports the parameter 'NewData' potentially being nullptr (in which case we set the texture data to uninitialized memory) this is another candidate for a future code clean up.
- Added a ::HasPayloadData to use instead of ::GetSizeOnDisk as in many cases the calling code is only interested if there is data at all and not what the exact size is.

* Testing
- Ran benchmarks before/after for cold cooks of an internal project showing overall cook times going from 1:54:28 down to 1:49:36 on average.

diff --git a/Engine/Source/Runtime/CoreUObject/Private/Virtualization/VirtualizedBulkData.cpp b/Engine/Source/Runtime/CoreUObject/Private/Virtualization/VirtualizedBulkData.cpp
--- a/Engine/Source/Runtime/CoreUObject/Private/Virtualization/VirtualizedBulkData.cpp
+++ b/Engine/Source/Runtime/CoreUObject/Private/Virtualization/VirtualizedBulkData.cpp
@@ -14,6 +14,7 @@
 #include "UObject/Object.h"
 #include "UObject/ObjectSaveContext.h"
 #include "Virtualization/IVirtualizationSourceControlUtilities.h"
+#include "Virtualization/VirtualizationManager.h"
 
 //#if WITH_EDITORONLY_DATA
 
@@ -157,36 +158,6 @@ void FVirtualizedUntypedBulkData::CreateFromBulkData(FUntypedBulkData& InBulkDat
 	}
 }
 
-#if UE_VBD_TO_OLD_BULKDATA_PATH
-void FVirtualizedUntypedBulkData::ConvertToOldBulkData(FUntypedBulkData& BulkData)
-{
-	TRACE_CPUPROFILER_EVENT_SCOPE(FVirtualizedUntypedBulkData::ConvertToOldBulkData);
-
-	FSharedBuffer UncompressedPayload = GetDataInternal().Decompress();
-
-	checkf(!IsDataVirtualized(), TEXT("Trying to convert data that has been virtualized, this system should've been removed before this could occur!"));
-	checkf(UncompressedPayload.GetSize() % BulkData.GetElementSize() == 0, TEXT("Payload does not match the element size for bulkdata!"));
-
-	const int64 ElementCount = UncompressedPayload.GetSize() / BulkData.GetElementSize();
-
-	BulkData.Lock(LOCK_READ_WRITE);
-
-	// It would be faster if we could 'release' the FSharedBuffer data to bulkdata but neither API supports
-	// this sort of thing. Give this is throw away code that we don't expect to have run it is not worth 
-	// adding support if this is the only usecase.
-	void* BulkDataPtr = BulkData.Realloc(ElementCount);
-	FMemory::Memcpy(BulkDataPtr, UncompressedPayload.GetData(), UncompressedPayload.GetSize());
-
-	BulkData.Unlock();
-
-	// If the virtual bulkdata has not disabled compression we should apply it to the old bulkdata
-	if (!EnumHasAnyFlags(Flags, EFlags::DisablePayloadCompression))
-	{
-		BulkData.StoreCompressedOnDisk(NAME_Zlib);
-	}
-}
-#endif // UE_VBD_TO_OLD_BULKDATA_PATH
-
 void FVirtualizedUntypedBulkData::Serialize(FArchive& Ar, UObject* Owner)
 {
 	TRACE_CPUPROFILER_EVENT_SCOPE(FVirtualizedUntypedBulkData::Serialize);
@@ -643,7 +614,7 @@ bool FVirtualizedUntypedBulkData::SerializeData(FArchive& Ar, FCompressedBuffer&
 	}
 	else if (Ar.IsLoading()) // Loading from old bulkdata format
 	{
-		const int64 Size = GetBulkDataSize();
+		const int64 Size = GetPayloadSize();
 		FUniqueBuffer LoadPayload = FUniqueBuffer::Alloc(Size);
 
 		if (EnumHasAnyFlags(PayloadFlags, EFlags::LegacyFileIsCompressed))
@@ -853,14 +824,8 @@ void FVirtualizedUntypedBulkData::UpdatePayload(FSharedBuffer InPayload, FName I
 							EFlags::LegacyFileIsCompressed |
 							EFlags::LegacyKeyWasGuidDerived);
 
-	//TODO: Should we validate now or let FCompressedBuffer do that later?
-	CompressionFormatToUse = InCompressionFormat;
-
-	if (InCompressionFormat == NAME_None)
-	{
-		EnumAddFlags(Flags, EFlags::DisablePayloadCompression);
-	}
-
+	SetCompressionFormat(InCompressionFormat);
+	
 	PackagePath.Empty();
 	PackageSegment = EPackageSegment::Header;
 	OffsetInFile = INDEX_NONE;
@@ -873,6 +838,17 @@ void FVirtualizedUntypedBulkData::UpdatePayload(FSharedBuffer InPayload, FName I
 	PayloadContentId = FPayloadId(Payload);
 }
 
+void FVirtualizedUntypedBulkData::SetCompressionFormat(FName InCompressionFormat)
+{
+	//TODO: Should we validate now or let FCompressedBuffer do that later?
+	CompressionFormatToUse = InCompressionFormat;
+
+	if (InCompressionFormat == NAME_None)
+	{
+		EnumAddFlags(Flags, EFlags::DisablePayloadCompression);
+	}
+}
+
 FCustomVersionContainer FVirtualizedUntypedBulkData::GetCustomVersions(FArchive& InlineArchive)
 {
 	return InlineArchive.GetCustomVersions();
diff --git a/Engine/Source/Runtime/CoreUObject/Public/Virtualization/VirtualizationManager.h b/Engine/Source/Runtime/CoreUObject/Public/Virtualization/VirtualizationManager.h
--- a/Engine/Source/Runtime/CoreUObject/Public/Virtualization/VirtualizationManager.h
+++ b/Engine/Source/Runtime/CoreUObject/Public/Virtualization/VirtualizationManager.h
@@ -7,35 +7,10 @@
 #include "Logging/LogMacros.h"
 #include "Virtualization/PayloadId.h"
 
-#define UE_USE_VIRTUALBULKDATA 1
-
-// Enables a code path that will allow VBD to be converted back to old style bulkdata if we need to revert for some reason
-#if !UE_USE_VIRTUALBULKDATA
-	#define UE_VBD_TO_OLD_BULKDATA_PATH 1
-
-	// When reverting to old bulkdata we will need to add a new version to FUE5MainStreamObjectVersion, but since it is unlikely
-	// that we will ever do this, I'd prefer not to submit it to FUE5MainStreamObjectVersion in a disabled form.
-	// This static assert will remind anyone turning on this path that the version needs to be updated and hopefully we can just
-	// delete this without ever enabling it once we take the plunge and strip out UE_USE_VIRTUALBULKDATA.
-	static_assert(false, "Add the following two lines to FUE5MainStreamObjectVersion!");
-	//	// Disabled content virtualization due to crippling problems, sadly StaticMesh and Textures saved with virtualization will need reverting.
-	//	DisabledVirtualization,
-#else
-	#define UE_VBD_TO_OLD_BULKDATA_PATH 0
-#endif //UE_VBD_TO_OLD_BULKDATA_PATH
-
 // TODO: Do we want to keep this header public for UE5 release? If the only interaction should be
 //		 via FVirtualizedUntypedBulkData or other such classes we might not want to expose this 
 //		 at all.
 
-// TODO: Remove when UE_USE_VIRTUALBULKDATA is removed  
-/** Some methods can be turned to const when we are using VirtualBulkData, this macro will help with the upgrade path */
-#if UE_USE_VIRTUALBULKDATA
-	#define UE_VBD_CONST const
-#else
-	#define UE_VBD_CONST
-#endif //UE_USE_VIRTUALBULKDATA
-
 /**
  * Configuring the backend hierarchy
  * 
diff --git a/Engine/Source/Runtime/CoreUObject/Public/Virtualization/VirtualizedBulkData.h b/Engine/Source/Runtime/CoreUObject/Public/Virtualization/VirtualizedBulkData.h
--- a/Engine/Source/Runtime/CoreUObject/Public/Virtualization/VirtualizedBulkData.h
+++ b/Engine/Source/Runtime/CoreUObject/Public/Virtualization/VirtualizedBulkData.h
@@ -3,11 +3,12 @@
 #pragma once
 
 #include "Async/Future.h"
+#include "Compression/CompressedBuffer.h"
 #include "HAL/Platform.h"
 #include "Memory/SharedBuffer.h"
 #include "Misc/Guid.h"
 #include "Misc/PackagePath.h"
-#include "Virtualization/VirtualizationManager.h"
+#include "Virtualization/PayloadId.h"
 
 class FArchive;
 class UObject;
@@ -67,19 +68,14 @@ public:
 	 */
 	void CreateFromBulkData(FUntypedBulkData& BulkData, const FGuid& Guid);
 
-#if UE_VBD_TO_OLD_BULKDATA_PATH
-	/**
-	 * Convenience method to make it easier to convert from FVirtualizedBulkData to BulkData. This is not very
-	 * efficient as it will load the payload from disk and the old bulkdata object will hold it in memory. It 
-	 * is provided in case virtualized bulkdata needs to be disabled once it goes live and we need a way to 
-	 * re-save packages back to the older format.
+	/** 
+	 * Used to serialize the bulkdata to/from a FArchive
 	 * 
-	 * @param BulkData	The bulkdata object to be provided with the payload that the virtualized bulkdata object
-	 *					represents.
+	 * @param Ar	The archive to serialize the bulkdata.
+	 * @param Owner	The UObject that contains the bulkdata object, if this is a nullptr then the bulkdata will
+	 *				assume that it must serialize the payload immediately to memory as it will not be able to
+	 *				identify it's package path.
 	 */
-	void ConvertToOldBulkData(FUntypedBulkData& BulkData);
-#endif //UE_VBD_TO_OLD_BULKDATA_PATH
-
 	void Serialize(FArchive& Ar, UObject* Owner);
 
 	/** Reset to a truly empty state */
@@ -100,14 +96,14 @@ public:
 	 */
 	FGuid GetIdentifier() const;
 
-	/** Returns an unique identifier for the content of the payload */
+	/** Returns an unique identifier for the content of the payload. */
 	const FPayloadId& GetPayloadId() const { return PayloadContentId; }
 
-	/** Returns the size of the payload in bytes */
+	/** Returns the size of the payload in bytes. */
 	int64 GetPayloadSize() const { return PayloadSize; }
-	
-	/** Temp method, to make it easier to transition older code to virtualized bulkdata, remove when we remove UE_USE_VIRTUALBULKDATA */
-	FORCEINLINE int64 GetBulkDataSize() const { return GetPayloadSize(); }
+
+	/** Returns true if the bulkdata object contains a valid payload greater than zero bytes in size. */
+	bool HasPayloadData() const { return PayloadSize > 0; }
 
 	// TODO: This (IsDataLoaded) is only needed for TextureDerivedData.cpp (which assumes that the data 
 	// needs to be loaded in order to be able to run on a background thread. Since FVirtualizedUntypedBulkData 
@@ -137,15 +133,36 @@ public:
 	 * by using 'FSharedBuffer::MakeView') then a clone of the data will be created internally and assigned
 	 * to the bulkdata object.
 	 *
-	 * @param InPayload The payload to update the bulkdata with
-	 * @param CompressionFormat (optional) The compression format to use, NAME_None indicates that the
-	 * payload is already in a compressed format and will not gain from being compressed again. These
-	 * payloads will never be compressed. NAME_Default will apply which ever compression format that the
-	 * underlying code deems appropriate. Other specific compression formats may be allowed, see the
-	 * documentation of FCompressedBuffer for details.
+	 * @param InPayload				The payload to update the bulkdata with
+	 * @param InCompressionFormat	The compression format to use, NAME_None indicates that the
+	 *								payload is already in a compressed format and will not gain from being 
+	 *								compressed again. These payloads will never be compressed. NAME_Default 
+	 *								will apply whichever compression format that the underlying code deems 
+	 *								appropriate. Other specific compression formats may be allowed, see the 
+	 *								documentation of FCompressedBuffer for details.
 	 */
 	void UpdatePayload(FSharedBuffer InPayload, FName CompressionFormat = NAME_Default);
 
+	/** 
+	 * Allows the compression format to be specified that will be applied the next time that the bulkdata
+	 * is saved. For non-virtualized payloads this will occur when the package is next saved. For 
+	 * virtualized payloads this will only ever be applied when it is pushed to the local cache. If the
+	 * payload has already been pushed to long term storage backends then the compression is not likely to
+	 * be changed.
+	 * 
+	 * This method is only exposed as part of the public api so that a long standing bug with FTextureSource
+	 * where older textures may have been saved with the wrong compression setting that needs to be fixed. 
+	 * There shouldn't generally be a need to call this.
+	 * 
+	 * @param InCompressionFormat	The compression format to use, NAME_None indicates that the
+	 *								payload is already in a compressed format and will not gain from being
+	 *								compressed again. These payloads will never be compressed. NAME_Default
+	 *								will apply whichever compression format that the underlying code deems
+	 *								appropriate. Other specific compression formats may be allowed, see the
+	 *								documentation of FCompressedBuffer for details.
+	 */
+	void SetCompressionFormat(FName InCompressionFormat);
+
 	/**
 	* Get the CustomVersions used in the file containing the payload. Currently this is assumed
 	* to always be the versions in the InlineArchive
diff --git a/Engine/Source/Runtime/Engine/Classes/Engine/Texture.h b/Engine/Source/Runtime/Engine/Classes/Engine/Texture.h
--- a/Engine/Source/Runtime/Engine/Classes/Engine/Texture.h
+++ b/Engine/Source/Runtime/Engine/Classes/Engine/Texture.h
@@ -272,46 +272,16 @@ struct FTextureSource
 	FORCEINLINE int32 GetNumBlocks() const { return Blocks.Num() + 1; }
 	FORCEINLINE ETextureSourceFormat GetFormat(int32 LayerIndex = 0) const { return (LayerIndex == 0) ? Format : LayerFormat[LayerIndex]; }
 	FORCEINLINE bool IsPNGCompressed() const { return bPNGCompressed; }
-	FORCEINLINE int64 GetSizeOnDisk() const 
-	{ 
-#if UE_USE_VIRTUALBULKDATA
-		return BulkData.GetPayloadSize(); 
-#else
-		return BulkData.GetBulkDataSize(); 
-#endif //UE_USE_VIRTUALBULKDATA
-	}
-	
-	FORCEINLINE bool IsBulkDataLoaded() const
-	{
-#if UE_USE_VIRTUALBULKDATA
-		return BulkData.IsDataLoaded();
-#else
-		return BulkData.IsBulkDataLoaded();
-#endif //UE_USE_VIRTUALBULKDATA
-	}
+	FORCEINLINE int64 GetSizeOnDisk() const { return BulkData.GetPayloadSize(); }
+	inline bool HasPayloadData() const { return BulkData.HasPayloadData(); }
+	FORCEINLINE bool IsBulkDataLoaded() const { return BulkData.IsDataLoaded(); }
 
 	ENGINE_API void OperateOnLoadedBulkData(TFunctionRef<void (const FSharedBuffer& BulkDataBuffer)> Operation);
 
-	FORCEINLINE bool LoadBulkDataWithFileReader()
-	{
-		// TODO: When we remove UE_USE_VIRTUALBULKDATA and go all in, places that call this
-		// should be changed.
-#if UE_USE_VIRTUALBULKDATA
-		return true; // Currently a NOP
-#else
-		return BulkData.LoadBulkDataWithFileReader(); // Kick off an async load if needed
-#endif
-	}
-
-	FORCEINLINE void RemoveBulkData()
-	{
-#if UE_USE_VIRTUALBULKDATA
-		BulkData.UnloadData();
-#else
-		BulkData.RemoveBulkData();
-#endif //UE_USE_VIRTUALBULKDATA
-	}
+	UE_DEPRECATED(5.00, "There is no longer a need to call LoadBulkDataWithFileReader, FTextureSource::BulkData can now load the data on demand without it.")
+	FORCEINLINE bool LoadBulkDataWithFileReader() { return true; }
 
+	FORCEINLINE void RemoveBulkData() { BulkData.UnloadData(); }
 	
 	/** Sets the GUID to use, and whether that GUID is actually a hash of some data. */
 	ENGINE_API void SetId(const FGuid& InId, bool bInGuidIsHash);
@@ -382,11 +352,6 @@ struct FTextureSource
 	/** Structure that encapsulates the decompressed texture data and can be accessed per mip */
 	struct ENGINE_API FMipData
 	{
-		~FMipData();
-
-		/** Get a copy of a given texture mip, to be stored in OutMipData */
-		bool GetMipData(TArray64<uint8>& OutMipData, int32 BlockIndex, int32 LayerIndex, int32 MipIndex) const;
-
 		/** Allow the copy constructor by rvalue*/
 		FMipData(FMipData&& Other)
 			: TextureSource(Other.TextureSource)
@@ -394,21 +359,25 @@ struct FTextureSource
 			MipData = MoveTemp(Other.MipData);
 		}
 
+		~FMipData() = default;
+
 		/** Disallow everything else so we don't get duplicates */
 		FMipData() = delete;
 		FMipData(const FMipData&) = delete;
 		FMipData& operator=(const FMipData&) = delete;
 		FMipData& operator=(FMipData&& Other) = delete;
 
+		/** Get a copy of a given texture mip, to be stored in OutMipData */
+		bool GetMipData(TArray64<uint8>& OutMipData, int32 BlockIndex, int32 LayerIndex, int32 MipIndex) const;
+
 	private:
 		// We only want to allow FTextureSource to create FMipData objects
 		friend struct FTextureSource;
 
-		FMipData(const FTextureSource& InSource, FMipAllocation&& InMipData, FRWLock* InReadLock = nullptr);
+		FMipData(const FTextureSource& InSource, FSharedBuffer InData);
 
 		const FTextureSource& TextureSource;
-		FMipAllocation MipData;
-		FRWLock* ReadLock;
+		FSharedBuffer MipData;
 	};
 #endif // WITH_EDITOR
 
@@ -425,7 +394,7 @@ private:
 	TDontCopy<FRWLock> BulkDataLock;
 #endif
 	/** The bulk source data. */
-	TChooseClass<UE_USE_VIRTUALBULKDATA, UE::Virtualization::FByteVirtualizedBulkData, FByteBulkData>::Result BulkData;
+	UE::Virtualization::FByteVirtualizedBulkData BulkData;
 	
 	/** Number of mips that are locked. */
 	uint32 NumLockedMips;
@@ -447,11 +416,11 @@ private:
 	uint8* LockMipInternal(int32 BlockIndex, int32 LayerIndex, int32 MipIndex, ELockState RequestedLockState);
 	
 	/** Returns the source data fully decompressed */
-	FMipAllocation Decompress(class IImageWrapperModule* ImageWrapperModule) UE_VBD_CONST;
-
-	FMipAllocation TryDecompressPngData(IImageWrapperModule* ImageWrapperModule) UE_VBD_CONST;
+	FSharedBuffer Decompress(class IImageWrapperModule* ImageWrapperModule) const;
+	/** Attempt to decompress the source data from a compressed png format. All failures will be logged and result in the method returning false */
+	FSharedBuffer TryDecompressPngData(IImageWrapperModule* ImageWrapperModule) const;
 	/** Attempt to decompress the source data from Jpeg format. All failures will be logged and result in the method returning false */
-	FMipAllocation TryDecompressJpegData(IImageWrapperModule* ImageWrapperModule) UE_VBD_CONST;
+	FSharedBuffer TryDecompressJpegData(IImageWrapperModule* ImageWrapperModule) const;
 	
 	/** Return true if the source art is not png compressed but could be. */
 	bool CanPNGCompress() const;
diff --git a/Engine/Source/Runtime/Engine/Private/StaticMeshSourceData.cpp b/Engine/Source/Runtime/Engine/Private/StaticMeshSourceData.cpp
--- a/Engine/Source/Runtime/Engine/Private/StaticMeshSourceData.cpp
+++ b/Engine/Source/Runtime/Engine/Private/StaticMeshSourceData.cpp
@@ -360,11 +360,7 @@ void FStaticMeshSourceModel::FillMaterialName(TMap<int32, FName>& OutMaterialMap
 // This is the key which is used for DDC data when legacy RawMesh is converted to MeshDescription.
 // If static mesh derived data needs to be rebuilt (new format, serialization differences, etc.) replace the version GUID below with a new one.
 // In case of merge conflicts with DDC versions, you MUST generate a new GUID and set this new GUID as the version.
-#if UE_USE_VIRTUALBULKDATA
 #define MESHDATAKEY_STATICMESH_DERIVEDDATA_VER TEXT("97E1F285E9824AD83693BB8BF976461")
-#else
-#define MESHDATAKEY_STATICMESH_DERIVEDDATA_VER TEXT("B3CE4F8E59E542D4A721EE56068AD214")
-#endif
 
 static bool GetMeshDataKey(FString& OutKey, const FRawMeshBulkData* RawMeshBulkData, int LodIndex)
 {
diff --git a/Engine/Source/Runtime/Engine/Private/Texture.cpp b/Engine/Source/Runtime/Engine/Private/Texture.cpp
--- a/Engine/Source/Runtime/Engine/Private/Texture.cpp
+++ b/Engine/Source/Runtime/Engine/Private/Texture.cpp
@@ -470,29 +470,14 @@ void UTexture::Serialize(FArchive& Ar)
 
 		if (Ar.IsLoading() && Ar.CustomVer(FUE5MainStreamObjectVersion::GUID) < FUE5MainStreamObjectVersion::TextureSourceVirtualization)
 		{
-#if UE_USE_VIRTUALBULKDATA
 			FByteBulkData TempBulkData;
 			TempBulkData.Serialize(Ar, this);
 
 			Source.BulkData.CreateFromBulkData(TempBulkData, Source.GetId());
-#else
-			Source.BulkData.Serialize(Ar, this);
-#endif //UE_USE_VIRTUALBULKDATA
 		}
 		else
 		{
-#if !UE_USE_VIRTUALBULKDATA && UE_VBD_TO_OLD_BULKDATA_PATH
-			if (Ar.IsLoading() && Ar.CustomVer(FUE5MainStreamObjectVersion::GUID) < FUE5MainStreamObjectVersion::DisabledVirtualization)
-			{
-				UE::Virtualization::FByteVirtualizedBulkData TempVirtualBulkData;
-				TempVirtualBulkData.Serialize(Ar, this);
-				TempVirtualBulkData.ConvertToOldBulkData(Source.BulkData);
-			}
-			else
-#endif // !UE_USE_VIRTUALBULKDATA && UE_VBD_TO_OLD_BULKDATA_PATH
-			{
-				Source.BulkData.Serialize(Ar, this);
-			}
+			Source.BulkData.Serialize(Ar, this);
 		}
 	}
 
@@ -1108,13 +1093,8 @@ void FTextureSource::InitBlocked(const ETextureSourceFormat* InLayerFormats,
 
 	checkf(LockState == ELockState::None, TEXT("InitBlocked shouldn't be called in-between LockMip/UnlockMip"));
 
-#if UE_USE_VIRTUALBULKDATA
 	FUniqueBuffer Buffer = FUniqueBuffer::Alloc(TotalBytes);
 	uint8* DataPtr = (uint8*)Buffer.GetData();
-#else
-	BulkData.Lock(LOCK_READ_WRITE);
-	uint8* DataPtr = (uint8*)BulkData.Realloc(TotalBytes);
-#endif //UE_USE_VIRTUALBULKDATA
 
 	if (InDataPerBlock)
 	{
@@ -1129,11 +1109,7 @@ void FTextureSource::InitBlocked(const ETextureSourceFormat* InLayerFormats,
 		}
 	}
 
-#if UE_USE_VIRTUALBULKDATA
 	BulkData.UpdatePayload(Buffer.MoveToShared());
-#else
-	BulkData.Unlock();
-#endif //UE_USE_VIRTUALBULKDATA
 }
 
 void FTextureSource::InitLayered(
@@ -1166,27 +1142,17 @@ void FTextureSource::InitLayered(
 
 	checkf(LockState == ELockState::None, TEXT("InitLayered shouldn't be called in-between LockMip/UnlockMip"));
 
-#if UE_USE_VIRTUALBULKDATA
-	// TODO: Might be able to just use FUniqueBuffer::Clone, if we can be sure that
-	// TotalBytes and NewData always match (the null check onNewData before memcpy is
-	// suspicious)
-	FUniqueBuffer Buffer = FUniqueBuffer::Alloc(TotalBytes);
-	uint8* DestData = (uint8*)Buffer.GetData();
-#else
-	BulkData.Lock(LOCK_READ_WRITE);
-	uint8* DestData = (uint8*)BulkData.Realloc(TotalBytes);
-#endif //UE_USE_VIRTUALBULKDATA
-
-	if (NewData)
+	// TODO: Allocating an empty buffer if there is no data to copy from seems like an odd choice but the 
+	// code logic has been doing this for almost a decade so I don't want to change it until I am sure that
+	// it serves no purpose. Given a choice I'd assert on NewData == nullptr instead.
+	if (NewData != nullptr)
 	{
-		FMemory::Memcpy(DestData, NewData, TotalBytes);
+		BulkData.UpdatePayload(FSharedBuffer::Clone(NewData, TotalBytes));
+	}
+	else
+	{
+		BulkData.UpdatePayload(FUniqueBuffer::Alloc(TotalBytes).MoveToShared());
 	}
-
-#if UE_USE_VIRTUALBULKDATA
-	BulkData.UpdatePayload(Buffer.MoveToShared());
-#else
-	BulkData.Unlock();
-#endif //UE_USE_VIRTUALBULKDATA
 }
 
 void FTextureSource::Init(
@@ -1257,19 +1223,9 @@ void FTextureSource::InitWithCompressedSourceData(
 
 	checkf(LockState == ELockState::None, TEXT("InitWithCompressedSourceData shouldn't be called in-between LockMip/UnlockMip"));
 	
-#if UE_USE_VIRTUALBULKDATA
 	// Disable the internal bulkdata compression if the source data is already compressed
 	const FName CompressionName = CompressionFormat == TSCF_None ? NAME_Default : NAME_None;
 	BulkData.UpdatePayload(FSharedBuffer::Clone(NewData.GetData(), NewData.Num()), CompressionName);
-#else
-	BulkData.Lock(LOCK_READ_WRITE);
-	uint8* DestData = (uint8*)BulkData.Realloc(NewData.Num());
-	if (NewData.GetData() != nullptr)
-	{
-		FMemory::Memcpy(DestData, NewData.GetData(), NewData.Num());
-	}
-	BulkData.Unlock();
-#endif //UE_USE_VIRTUALBULKDATA
 }
 
 void FTextureSource::Compress()
@@ -1282,31 +1238,19 @@ void FTextureSource::Compress()
 
 	if (CanPNGCompress()) // Note that this will return false if the data is already a compressed PNG
 	{
-#if UE_USE_VIRTUALBULKDATA
 		FSharedBuffer Payload = BulkData.GetPayload().Get();
-		const void* BulkDataPtr = Payload.GetData();
-		const int64 BulkDataSize = Payload.GetSize();
-#else
-		uint8* BulkDataPtr = (uint8*)BulkData.Lock(LOCK_READ_WRITE);
-		const int64 BulkDataSize = BulkData.GetBulkDataSize();
-#endif //UE_USE_VIRTUALBULKDATA
 
 		IImageWrapperModule& ImageWrapperModule = FModuleManager::LoadModuleChecked<IImageWrapperModule>( FName("ImageWrapper") );
 		TSharedPtr<IImageWrapper> ImageWrapper = ImageWrapperModule.CreateImageWrapper( EImageFormat::PNG );
 		// TODO: TSF_BGRA8 is stored as RGBA, so the R and B channels are swapped in the internal png. Should we fix this?
 		ERGBFormat RawFormat = (Format == TSF_G8 || Format == TSF_G16) ? ERGBFormat::Gray : ERGBFormat::RGBA;
-		if ( ImageWrapper.IsValid() && ImageWrapper->SetRaw( BulkDataPtr, BulkDataSize, SizeX, SizeY, RawFormat, (Format == TSF_G16 || Format == TSF_RGBA16) ? 16 : 8 ) )
+		if ( ImageWrapper.IsValid() && ImageWrapper->SetRaw(Payload.GetData(), Payload.GetSize(), SizeX, SizeY, RawFormat, (Format == TSF_G16 || Format == TSF_RGBA16) ? 16 : 8 ) )
 		{
-			const TArray64<uint8> CompressedData = ImageWrapper->GetCompressed();
+			TArray64<uint8> CompressedData = ImageWrapper->GetCompressed();
 			if ( CompressedData.Num() > 0 )
-			{
-#if UE_USE_VIRTUALBULKDATA	
-				BulkData.UpdatePayload(FSharedBuffer::Clone(CompressedData.GetData(), CompressedData.Num()), NAME_None);
-#else
-				BulkDataPtr = (uint8*)BulkData.Realloc(CompressedData.Num());
-				FMemory::Memcpy(BulkDataPtr, CompressedData.GetData(), CompressedData.Num());
-				BulkData.Unlock();
-#endif //UE_USE_VIRTUALBULKDATA
+			{			
+				BulkData.UpdatePayload(MakeSharedBufferFromArray(MoveTemp(CompressedData)), NAME_None);
+
 				bPNGCompressed = true;
 				CompressionFormat = TSCF_PNG;
 			}
@@ -1319,14 +1263,12 @@ void FTextureSource::Compress()
 		CompressionFormat = TSCF_PNG;
 	}
 
-	// TODO: Cannot fix incorrectly assigned compression types if UE_USE_VIRTUALBULKDATA is enabled. Need a fix for that.
-#if !UE_USE_VIRTUALBULKDATA
-	// Disable the internal bulkdata compression if the source data is already compressed	
-	BulkData.StoreCompressedOnDisk(CompressionFormat == TSCF_None ? NAME_Zlib : NAME_None);
-#endif
+	// There was a brief period when textures could be saved with the wrong compression flag based on their storage format
+	// so we force set it here so that they will be fixed on resave.
+	BulkData.SetCompressionFormat(CompressionFormat == TSCF_None ? NAME_Default : NAME_None);
 }
 
-FTextureSource::FMipAllocation FTextureSource::Decompress(IImageWrapperModule* ImageWrapperModule) UE_VBD_CONST
+FSharedBuffer FTextureSource::Decompress(IImageWrapperModule* ImageWrapperModule) const
 {
 	TRACE_CPUPROFILER_EVENT_SCOPE(FTextureSource::Decompress);
 
@@ -1340,13 +1282,7 @@ FTextureSource::FMipAllocation FTextureSource::Decompress(IImageWrapperModule* I
 	}
 	else
 	{
-#if UE_USE_VIRTUALBULKDATA
-		FMipAllocation Payload(BulkData.GetPayload().Get());
-#else
-		FMipAllocation Payload(BulkData);
-#endif //UE_USE_VIRTUALBULKDATA
-
-		return Payload;
+		return BulkData.GetPayload().Get();
 	}
 }
 
@@ -1406,7 +1342,6 @@ void FTextureSource::UnlockMip(int32 BlockIndex, int32 LayerIndex, int32 MipInde
 	--NumLockedMips;
 	if (NumLockedMips == 0u)
 	{
-#if UE_USE_VIRTUALBULKDATA
 		// If the lock was for Read/Write then we need to assume that the decompressed copy
 		// we returned (LockedMipData) was updated and should update the payload accordingly.
 		// This will wipe the compression format that we used to have.
@@ -1419,40 +1354,11 @@ void FTextureSource::UnlockMip(int32 BlockIndex, int32 LayerIndex, int32 MipInde
 			bPNGCompressed = false;
 			CompressionFormat = TSCF_None;
 		}
-#else
-		// If the source data is stored in a compressed format then we returned a copy when the mip was
-		// locked. We now need to deal with that copy.
-		if (CompressionFormat == TSCF_JPEG || bPNGCompressed)
-		{
-			// If we had a write lock then we need to update the bulkdata with the copy as we have
-			// to assume that it contains changes. This means we will no longer be storing it 
-			// in a compressed format.
-			// If the lock was read only then we can just delete the locked mip data as we know
-			// that there were no changes to the data.
-			if (LockState == ELockState::ReadWrite)
-			{
-				UE_CLOG(CompressionFormat == TSCF_JPEG, LogTexture, Warning, TEXT("Call to FTextureSource::UnlockMip will cause texture source to lose it's jpeg storage format"));
-
-				check(BlockIndex == 0);
-				check(LayerIndex == 0);
-				check(MipIndex == 0);
-				int32 MipSize = CalcMipSize(0, 0, 0);
-
-				uint8* UncompressedData = (uint8*)BulkData.Realloc(MipSize);
-				FMemory::Memcpy(UncompressedData, LockedMipData.GetDataReadOnly().GetData(), MipSize);
-
-				bPNGCompressed = false;
-				CompressionFormat = TSCF_None;
-			}
-		}
-
-		BulkData.Unlock();	
-#endif //UE_USE_VIRTUALBULKDATA
 
 		LockState = ELockState::None;
 		LockedMipData.Reset();
 
-		// TODO: Only call this if the LockState was ELockState::ReadWrite
+		// TODO: Only call this if the LockState was ELockState::ReadWrite (this would need testing)
 		ForceGenerateGuid();
 	}
 }
@@ -1465,7 +1371,7 @@ bool FTextureSource::GetMipData(TArray64<uint8>& OutMipData, int32 BlockIndex, i
 
 	bool bSuccess = false;
 
-	if (BlockIndex < GetNumBlocks() && LayerIndex < NumLayers && MipIndex < NumMips && GetSizeOnDisk() > 0)
+	if (BlockIndex < GetNumBlocks() && LayerIndex < NumLayers && MipIndex < NumMips && HasPayloadData())
 	{
 #if WITH_EDITOR
 		FWriteScopeLock BulkDataExclusiveScope(BulkDataLock.Get());
@@ -1474,9 +1380,9 @@ bool FTextureSource::GetMipData(TArray64<uint8>& OutMipData, int32 BlockIndex, i
 		checkf(NumLockedMips == 0, TEXT("Attempting to access a locked FTextureSource"));
 		// LockedMipData should only be allocated if NumLockedMips > 0 so the following assert should have been caught
 		// by the one above. If it fires then it indicates that there is a lock/unlock mismatch as well as invalid access!
-		checkf(LockedMipData.IsNull(), TEXT("Attempting to access mip data while locked mip data is still allocated")); 
+		checkf(LockedMipData.IsNull(), TEXT("Attempting to access mip data while locked mip data is still allocated"));
 
-		FMipAllocation DecompressedData = Decompress(ImageWrapperModule);
+		FSharedBuffer DecompressedData = Decompress(ImageWrapperModule);
 
 		if (!DecompressedData.IsNull())
 		{
@@ -1489,17 +1395,13 @@ bool FTextureSource::GetMipData(TArray64<uint8>& OutMipData, int32 BlockIndex, i
 				OutMipData.AddUninitialized(MipSize);
 				FMemory::Memcpy(
 					OutMipData.GetData(),
-					(const uint8*)DecompressedData.GetDataReadOnly().GetData() + MipOffset,
+					(const uint8*)DecompressedData.GetData() + MipOffset,
 					MipSize
 				);
 
 				bSuccess = true;
 			}
-		}
-
-#if !UE_USE_VIRTUALBULKDATA
-		BulkData.Unlock();
-#endif //!UE_USE_VIRTUALBULKDATA		
+		}	
 	}
 	
 	return bSuccess;
@@ -1515,16 +1417,11 @@ FTextureSource::FMipData FTextureSource::GetMipData(IImageWrapperModule* ImageWr
 	check(NumLockedMips == 0);
 
 #if WITH_EDITOR
-	// When UE_USE_VIRTUALBULKDATA is removed, replace this with a FReadScopeLock,
-	// we only need the expanded scope for old style bulkdata
-	BulkDataLock.Get().ReadLock();
+	FReadScopeLock _(BulkDataLock.Get());
+#endif //WITH_EDITOR
 
-	FMipAllocation DecompressedData = Decompress(ImageWrapperModule);
-	return FMipData(*this, MoveTemp(DecompressedData), &BulkDataLock.Get());
-#else
-	FMipAllocation DecompressedData = Decompress(ImageWrapperModule);
-	return FMipData(*this, MoveTemp(DecompressedData));
-#endif
+	FSharedBuffer DecompressedData = Decompress(ImageWrapperModule);
+	return FMipData(*this, DecompressedData);
 }
 
 int64 FTextureSource::CalcMipSize(int32 BlockIndex, int32 LayerIndex, int32 MipIndex) const
@@ -1554,7 +1451,7 @@ bool FTextureSource::IsPowerOfTwo(int32 BlockIndex) const
 bool FTextureSource::IsValid() const
 {
 	return SizeX > 0 && SizeY > 0 && NumSlices > 0 && NumLayers > 0 && NumMips > 0 &&
-		Format != TSF_Invalid && GetSizeOnDisk() > 0;
+		Format != TSF_Invalid && HasPayloadData();
 }
 
 void FTextureSource::GetBlock(int32 Index, FTextureSourceBlock& OutBlock) const
@@ -1636,16 +1533,12 @@ FString FTextureSource::GetSourceCompressionAsString() const
 	return StaticEnum<ETextureSourceCompressionFormat>()->GetDisplayNameTextByValue(GetSourceCompression()).ToString();
 }
 
-FTextureSource::FMipAllocation FTextureSource::TryDecompressPngData(IImageWrapperModule* ImageWrapperModule) UE_VBD_CONST
+FSharedBuffer FTextureSource::TryDecompressPngData(IImageWrapperModule* ImageWrapperModule) const
 {
 	bool bCanPngCompressFormat = (Format == TSF_G8 || Format == TSF_G16 || Format == TSF_RGBA8 || Format == TSF_BGRA8 || Format == TSF_RGBA16);
 	check(Blocks.Num() == 0 && NumLayers == 1 && NumSlices == 1 && bCanPngCompressFormat);
 
-#if UE_USE_VIRTUALBULKDATA
-	FMipAllocation Payload(BulkData.GetPayload().Get());
-#else
-	FMipAllocation Payload(BulkData);
-#endif //UE_USE_VIRTUALBULKDATA
+	FSharedBuffer Payload = BulkData.GetPayload().Get();
 
 	if (ImageWrapperModule == nullptr) // Optional if called from the gamethread, see FModuleManager::WarnIfItWasntSafeToLoadHere()
 	{
@@ -1653,8 +1546,10 @@ FTextureSource::FMipAllocation FTextureSource::TryDecompressPngData(IImageWrappe
 	}
 
 	TSharedPtr<IImageWrapper> ImageWrapper = ImageWrapperModule->CreateImageWrapper(EImageFormat::PNG);
-	if (ImageWrapper.IsValid() && ImageWrapper->SetCompressed(Payload.GetDataReadOnly().GetData(), Payload.GetSize()))
+	if (ImageWrapper.IsValid() && ImageWrapper->SetCompressed(Payload.GetData(), Payload.GetSize()))
 	{
+		Payload.Reset(); // Image wrapper takes a copy so we can discard the payload now
+
 		check(ImageWrapper->GetWidth() == SizeX);
 		check(ImageWrapper->GetHeight() == SizeY);
 
@@ -1662,23 +1557,23 @@ FTextureSource::FMipAllocation FTextureSource::TryDecompressPngData(IImageWrappe
 		// TODO: TSF_BGRA8 is stored as RGBA, so the R and B channels are swapped in the internal png. Should we fix this?
 		ERGBFormat RawFormat = (Format == TSF_G8 || Format == TSF_G16) ? ERGBFormat::Gray : ERGBFormat::RGBA;
 		if (ImageWrapper->GetRaw(RawFormat, (Format == TSF_G16 || Format == TSF_RGBA16) ? 16 : 8, RawData) && RawData.Num() > 0)
-		{
-			return FMipAllocation(RawData.GetData(), RawData.Num());
+		{		
+			return MakeSharedBufferFromArray(MoveTemp(RawData));
 		}
 		else
 		{
 			UE_LOG(LogTexture, Warning, TEXT("PNG decompression of source art failed"));
-			return FMipAllocation();
+			return FSharedBuffer();
 		}
 	}
 	else
 	{
 		UE_LOG(LogTexture, Log, TEXT("Only pngs are supported"));
-		return FMipAllocation();
+		return FSharedBuffer();
 	}
 }
 
-FTextureSource::FMipAllocation FTextureSource::TryDecompressJpegData(IImageWrapperModule* ImageWrapperModule) UE_VBD_CONST
+FSharedBuffer FTextureSource::TryDecompressJpegData(IImageWrapperModule* ImageWrapperModule) const
 {
 	if (NumLayers == 1 && NumSlices == 1 && Blocks.Num() == 0)
 	{
@@ -1687,39 +1582,37 @@ FTextureSource::FMipAllocation FTextureSource::TryDecompressJpegData(IImageWrapp
 			ImageWrapperModule = &FModuleManager::LoadModuleChecked<IImageWrapperModule>(FName("ImageWrapper"));
 		}
 
-#if UE_USE_VIRTUALBULKDATA
-		FMipAllocation Payload(BulkData.GetPayload().Get());
-#else
-		FMipAllocation Payload(BulkData);
-#endif //UE_USE_VIRTUALBULKDATA
+		FSharedBuffer Payload = BulkData.GetPayload().Get();
 
 		TSharedPtr<IImageWrapper> ImageWrapper = ImageWrapperModule->CreateImageWrapper(EImageFormat::JPEG);
-		if (ImageWrapper.IsValid() && ImageWrapper->SetCompressed(Payload.GetDataReadOnly().GetData(), Payload.GetSize()))
+		if (ImageWrapper.IsValid() && ImageWrapper->SetCompressed(Payload.GetData(), Payload.GetSize()))
 		{
+			Payload.Reset(); // Image wrapper takes a copy so we can discard the payload now
+
 			TArray64<uint8> RawData;
 			// The two formats we support for JPEG imports, see UTextureFactory::ImportImage
 			const ERGBFormat JpegFormat = Format == TSF_G8 ? ERGBFormat::Gray : ERGBFormat::BGRA;
 			if (ImageWrapper->GetRaw(JpegFormat, 8, RawData))
 			{
-				return FMipAllocation(RawData.GetData(), RawData.Num());
+				return MakeSharedBufferFromArray(MoveTemp(RawData));
 			}
 			else
 			{
 				UE_LOG(LogTexture, Warning, TEXT("JPEG decompression of source art failed to return uncompressed data"));
-				return FMipAllocation();
+				return FSharedBuffer();
 			}
 		}
 		else
 		{
 			UE_LOG(LogTexture, Warning, TEXT("JPEG decompression of source art failed initialization"));
-			return FMipAllocation();
+			return FSharedBuffer();
 		}
 	}
 	else
 	{
 		UE_LOG(LogTexture, Warning, TEXT("JPEG compressed source art is in an invalid format NumLayers:(%d) NumSlices:(%d) NumBlocks:(%d)"),
 			NumLayers, NumSlices, Blocks.Num());
-		return FMipAllocation();
+		return FSharedBuffer();
 	}	
 }
 
@@ -1734,7 +1627,7 @@ bool FTextureSource::CanPNGCompress() const
 		Blocks.Num() == 0 &&
 		SizeX > 4 &&
 		SizeY > 4 &&
-		GetSizeOnDisk() > 0 &&
+		HasPayloadData() &&
 		bCanPngCompressFormat &&
 		CompressionFormat == TSCF_None)
 	{
@@ -1752,16 +1645,7 @@ void FTextureSource::ForceGenerateGuid()
 void FTextureSource::ReleaseSourceMemory()
 {
 	bHasHadBulkDataCleared = true;
-
-#if UE_USE_VIRTUALBULKDATA
 	BulkData.UnloadData();
-#else
-	if (BulkData.IsLocked())
-	{
-		BulkData.Unlock();
-	}
-	BulkData.RemoveBulkData();
-#endif //UE_USE_VIRTUALBULKDATA
 }
 
 void FTextureSource::RemoveSourceData()
@@ -1779,16 +1663,8 @@ void FTextureSource::RemoveSourceData()
 	LockedMipData.Reset();
 	NumLockedMips = 0u;
 	LockState = ELockState::None;
-
-#if UE_USE_VIRTUALBULKDATA
+	
 	BulkData.UnloadData();
-#else
-	if (BulkData.IsLocked())
-	{
-		BulkData.Unlock();
-	}
-	BulkData.RemoveBulkData();
-#endif //UE_USE_VIRTUALBULKDATA
 
 	ForceGenerateGuid();
 }
@@ -1858,21 +1734,10 @@ int64 FTextureSource::CalcMipOffset(int32 BlockIndex, int32 LayerIndex, int32 Mi
 
 void FTextureSource::UseHashAsGuid()
 {
-	if (GetSizeOnDisk() > 0)
+	if (HasPayloadData())
 	{
 		bGuidIsHash = true;
-
-#if UE_USE_VIRTUALBULKDATA
 		Id = BulkData.GetPayloadId().ToGuid();
-#else
-		uint32 Hash[5] = {};
-
-		const void* Buffer = BulkData.LockReadOnly();
-		FSHA1::HashBuffer(Buffer, BulkData.GetBulkDataSize(), (uint8*)Hash);
-		BulkData.Unlock();
-
-		Id = FGuid(Hash[0] ^ Hash[4], Hash[1], Hash[2], Hash[3]);
-#endif //UE_USE_VIRTUALBULKDATA
 	}
 	else
 	{
@@ -1888,17 +1753,8 @@ void FTextureSource::OperateOnLoadedBulkData(TFunctionRef<void(const FSharedBuff
 	FReadScopeLock BulkDataExclusiveScope(BulkDataLock.Get());
 #endif
 
-#if UE_USE_VIRTUALBULKDATA
 	FSharedBuffer Payload = BulkData.GetPayload().Get();
-#else
-	uint8* BulkDataPtr = (uint8*)BulkData.Lock(LOCK_READ_ONLY);
-	const int64 BulkDataSize = BulkData.GetBulkDataSize();
-	FSharedBuffer Payload = FSharedBuffer::MakeView(BulkDataPtr, BulkDataSize);
-#endif //UE_USE_VIRTUALBULKDATA
 	Operation(Payload);
-#if !UE_USE_VIRTUALBULKDATA
-	BulkData.Unlock();
-#endif //!UE_USE_VIRTUALBULKDATA
 }
 
 void FTextureSource::SetId(const FGuid& InId, bool bInGuidIsHash)
@@ -2306,23 +2162,10 @@ void UTexture::NotifyMaterials(const ENotifyMaterialsEffectOnShaders EffectOnSha
 
 #if WITH_EDITOR
 
-FTextureSource::FMipData::FMipData(const FTextureSource& InSource, FMipAllocation&& InMipData, FRWLock* InReadLock)
+FTextureSource::FMipData::FMipData(const FTextureSource& InSource, FSharedBuffer InData)
 	: TextureSource(InSource)
-	, MipData(MoveTemp(InMipData))
-	, ReadLock(InReadLock)
-{
-}
-
-FTextureSource::FMipData::~FMipData()
+	, MipData(InData)
 {
-#if !UE_USE_VIRTUALBULKDATA
-	TextureSource.BulkData.Unlock();
-#endif //!UE_USE_VIRTUALBULKDATA
-
-	if (ReadLock != nullptr)
-	{
-		ReadLock->ReadUnlock();
-	}
 }
 
 bool FTextureSource::FMipData::GetMipData(TArray64<uint8>& OutMipData, int32 BlockIndex, int32 LayerIndex, int32 MipIndex) const
@@ -2338,7 +2181,7 @@ bool FTextureSource::FMipData::GetMipData(TArray64<uint8>& OutMipData, int32 Blo
 			OutMipData.AddUninitialized(MipSize);
 			FMemory::Memcpy(
 				OutMipData.GetData(),
-				(const uint8*)MipData.GetDataReadOnly().GetData() + MipOffset,
+				(const uint8*)MipData.GetData() + MipOffset,
 				MipSize
 			);
 		}
diff --git a/Engine/Source/Runtime/Engine/Private/TextureDerivedDataTask.cpp b/Engine/Source/Runtime/Engine/Private/TextureDerivedDataTask.cpp
--- a/Engine/Source/Runtime/Engine/Private/TextureDerivedDataTask.cpp
+++ b/Engine/Source/Runtime/Engine/Private/TextureDerivedDataTask.cpp
@@ -249,12 +249,9 @@ void FTextureSourceData::GetSourceMips(FTextureSource& Source, IImageWrapperModu
 
 void FTextureSourceData::GetAsyncSourceMips(IImageWrapperModule* InImageWrapper)
 {
-	if (bValid && !Blocks[0].MipsPerLayer[0].Num() && AsyncSource.GetSizeOnDisk())
+	if (bValid && !Blocks[0].MipsPerLayer[0].Num() && AsyncSource.HasPayloadData())
 	{
-		if (AsyncSource.LoadBulkDataWithFileReader())
-		{
-			GetSourceMips(AsyncSource, InImageWrapper);
-		}
+		GetSourceMips(AsyncSource, InImageWrapper);
 	}
 }
 void FTextureCacheDerivedDataWorker::ConsumeBuildFunctionOutput(const UE::DerivedData::FBuildOutput& BuildOutput, const FString& TexturePath, bool bReplaceExistingDDC)
diff --git a/Engine/Source/Runtime/MeshDescription/Private/MeshDescription.cpp b/Engine/Source/Runtime/MeshDescription/Private/MeshDescription.cpp
--- a/Engine/Source/Runtime/MeshDescription/Private/MeshDescription.cpp
+++ b/Engine/Source/Runtime/MeshDescription/Private/MeshDescription.cpp
@@ -2008,28 +2008,13 @@ void FMeshDescriptionBulkData::Serialize( FArchive& Ar, UObject* Owner )
 	FByteBulkData TempBulkData;
 	if (Ar.IsLoading() && Ar.CustomVer(FUE5MainStreamObjectVersion::GUID) < FUE5MainStreamObjectVersion::MeshDescriptionVirtualization)
 	{
-#if UE_USE_VIRTUALBULKDATA
 		// Serialize the old BulkData format and mark that we require a conversion after the guid has been serialized
 		TempBulkData.Serialize(Ar, Owner);
 		bSerializedOldDataTypes = true;
-#else
-		BulkData.Serialize(Ar, Owner);
-#endif //UE_USE_VIRTUALBULKDATA
 	}
 	else
 	{
-#if !UE_USE_VIRTUALBULKDATA && UE_VBD_TO_OLD_BULKDATA_PATH
-		if (Ar.IsLoading() && Ar.CustomVer(FUE5MainStreamObjectVersion::GUID) < FUE5MainStreamObjectVersion::DisabledVirtualization)
-		{
-			UE::Virtualization::FByteVirtualizedBulkData TempVirtualBulkData;
-			TempVirtualBulkData.Serialize(Ar, Owner);
-			TempVirtualBulkData.ConvertToOldBulkData(BulkData);
-		}
-		else
-#endif // !UE_USE_VIRTUALBULKDATA && UE_VBD_TO_OLD_BULKDATA_PATH
-		{
-			BulkData.Serialize(Ar, Owner);
-		}
+		BulkData.Serialize(Ar, Owner);
 	}
 
 	if( Ar.IsLoading() && Ar.CustomVer( FEditorObjectVersion::GUID ) < FEditorObjectVersion::MeshDescriptionBulkDataGuid )
@@ -2058,13 +2043,11 @@ void FMeshDescriptionBulkData::Serialize( FArchive& Ar, UObject* Owner )
 		CustomVersions = BulkData.GetCustomVersions(Ar);
 	}
 
-#if UE_USE_VIRTUALBULKDATA
 	// Needs to be after the guid is serialized
 	if (bSerializedOldDataTypes)
 	{
 		BulkData.CreateFromBulkData(TempBulkData, Guid);
 	}
-#endif //UE_USE_VIRTUALBULKDATA
 }
 
 
@@ -2076,21 +2059,12 @@ void FMeshDescriptionBulkData::SaveMeshDescription( FMeshDescription& MeshDescri
 	FRWScopeLock ScopeLock(BulkDataLock, SLT_Write);
 #endif
 
-#if UE_USE_VIRTUALBULKDATA
 	BulkData.Reset();
-#else
-	BulkData.RemoveBulkData();
-#endif //UE_USE_VIRTUALBULKDATA
 
 	if( !MeshDescription.IsEmpty() )
 	{
 		const bool bIsPersistent = true;
-#if UE_USE_VIRTUALBULKDATA
 		UE::Virtualization::FVirtualizedBulkDataWriter Ar(BulkData, bIsPersistent);
-#else
-		FBulkDataWriter Ar(BulkData, bIsPersistent);
-#endif //UE_USE_VIRTUALBULKDATA
-
 		Ar << MeshDescription;
 
 		// Preserve CustomVersions at save time so we can reuse the same ones when reloading direct from memory
@@ -2117,7 +2091,6 @@ void FMeshDescriptionBulkData::LoadMeshDescription( FMeshDescription& MeshDescri
 
 	MeshDescription.Empty();
 
-#if UE_USE_VIRTUALBULKDATA
 	if (BulkData.GetPayloadSize() > 0)
 	{
 #if WITH_EDITOR
@@ -2136,55 +2109,11 @@ void FMeshDescriptionBulkData::LoadMeshDescription( FMeshDescription& MeshDescri
 
 		BulkData.UnloadData();
 	}
-#else
-	if (BulkData.GetBulkDataSize() > 0)
-	{
-		// Get a lock on the bulk data and read it into the mesh description
-#if WITH_EDITOR
-		// A lock is required so we can clone the mesh description from multiple threads
-		FRWScopeLock ScopeLock(BulkDataLock, SLT_Write);
-
-		// This allows any thread to be able to clone a mesh description directly
-		// from disk so we can unload bulk data from memory.
-		bool bHasBeenLoadedFromFileReader = false;
-		if (BulkData.IsAsyncLoadingComplete() && !BulkData.IsBulkDataLoaded())
-		{
-			// This will return false under -game mode for now because we're not allowed to load bulk data outside of EDL.
-			bHasBeenLoadedFromFileReader = BulkData.LoadBulkDataWithFileReader();
-		}
-#endif //WITH_EDITOR
-		// This is in a scope because the FBulkDataReader need to be destroyed in order
-		// to unlock the BulkData and allow UnloadBulkData to actually do its job.
-		{
-			const bool bIsPersistent = true;
-			FBulkDataReader Ar(BulkData, bIsPersistent);
-
-			// Propagate the custom version information from the package to the bulk data, so that the MeshDescription
-			// is serialized with the same versioning.
-			Ar.SetCustomVersions( CustomVersions );
-			Ar << MeshDescription;
-		} // Unlock bulk data when we leave scope
-
-#if WITH_EDITOR
-		// Throw away the bulk data allocation only in the case we can safely reload it from disk
-		// and if BulkData.LoadBulkDataWithFileReader() is allowed to work from any thread.
-		// This saves a significant amount of memory during map loading of Nanite Meshes.
-		if (bHasBeenLoadedFromFileReader)
-		{
-			verify(BulkData.UnloadBulkData());
-		}
-#endif //WITH_EDITOR
-	}
-#endif // UE_USE_VIRTUALBULKDATA
 }
 
 void FMeshDescriptionBulkData::Empty()
 {
-#if UE_USE_VIRTUALBULKDATA
 	BulkData.Reset();
-#else
-	BulkData.RemoveBulkData();
-#endif //UE_USE_VIRTUALBULKDATA
 }
 
 FString FMeshDescriptionBulkData::GetIdString() const
@@ -2199,19 +2128,9 @@ FString FMeshDescriptionBulkData::GetIdString() const
 
 FGuid FMeshDescriptionBulkData::GetHash() const
 {
-	if (BulkData.GetBulkDataSize() > 0)
+	if (BulkData.GetPayloadSize() > 0)
 	{
-#if UE_USE_VIRTUALBULKDATA
 		return BulkData.GetPayloadId().ToGuid();
-#else
-		uint32 Hash[5] = {};
-
-		const void* Buffer = BulkData.LockReadOnly();
-		FSHA1::HashBuffer(Buffer, BulkData.GetBulkDataSize(), (uint8*)Hash);
-		BulkData.Unlock();
-
-		return FGuid(Hash[0] ^ Hash[4], Hash[1], Hash[2], Hash[3]);
-#endif //UE_USE_VIRTUALBULKDATA
 	}
 
 	return FGuid();
@@ -2219,7 +2138,7 @@ FGuid FMeshDescriptionBulkData::GetHash() const
 
 void FMeshDescriptionBulkData::UseHashAsGuid()
 {
-	if (BulkData.GetBulkDataSize() > 0)
+	if (BulkData.GetPayloadSize() > 0)
 	{
 		bGuidIsHash = true;
 
diff --git a/Engine/Source/Runtime/MeshDescription/Public/MeshDescription.h b/Engine/Source/Runtime/MeshDescription/Public/MeshDescription.h
--- a/Engine/Source/Runtime/MeshDescription/Public/MeshDescription.h
+++ b/Engine/Source/Runtime/MeshDescription/Public/MeshDescription.h
@@ -1290,9 +1290,6 @@ public:
 		: bBulkDataUpdated(false)
 		, bGuidIsHash(false)
 	{
-#if !UE_USE_VIRTUALBULKDATA
-		BulkData.SetBulkDataFlags(BULKDATA_SerializeCompressed | BULKDATA_SerializeCompressedBitWindow);
-#endif //!UE_USE_VIRTUALBULKDATA
 	}
 
 	/** Serialization */
@@ -1308,7 +1305,7 @@ public:
 	void Empty();
 
 	/** Returns true if there is no bulk data available */
-	bool IsEmpty() const { return BulkData.GetBulkDataSize() == 0; }
+	bool IsEmpty() const { return BulkData.GetPayloadSize() == 0; }
 
 	/** Return unique ID string for this bulk data */
 	FString GetIdString() const;
@@ -1317,7 +1314,7 @@ public:
 	void UseHashAsGuid();
 
 	/** Gets the size of the serialized bulk data */
-	int64 GetBulkDataSize() { return BulkData.GetBulkDataSize(); }
+	int64 GetBulkDataSize() { return BulkData.GetPayloadSize(); }
 
 private:
 #if WITH_EDITOR
@@ -1325,7 +1322,7 @@ private:
 	FRWLock       BulkDataLock;
 #endif
 	/** Internally store bulk data as bytes */
-	TChooseClass<UE_USE_VIRTUALBULKDATA, UE::Virtualization::FByteVirtualizedBulkData, FByteBulkData>::Result BulkData;
+	UE::Virtualization::FByteVirtualizedBulkData BulkData;
 
 	/** GUID associated with the data stored herein. */
 	FGuid Guid;
