Intermediate fix for crash due to uninitialized instance data

diff --git a/Engine/Source/Runtime/Engine/Private/InstancedStaticMesh.cpp b/Engine/Source/Runtime/Engine/Private/InstancedStaticMesh.cpp
--- a/Engine/Source/Runtime/Engine/Private/InstancedStaticMesh.cpp
+++ b/Engine/Source/Runtime/Engine/Private/InstancedStaticMesh.cpp
@@ -1283,6 +1283,15 @@ void FInstancedStaticMeshSceneProxy::SetupProxy(UInstancedStaticMeshComponent* I
 
 		for (int32 InInstanceIndex = 0; InInstanceIndex < InstanceSceneData.Num(); ++InInstanceIndex)
 		{
+			// Make sure the instance is initialized, regardless of below remapping
+			{
+				FPrimitiveInstance& TmpSceneData = InstanceSceneData[InInstanceIndex];
+				TmpSceneData.LocalToPrimitive = FRenderTransform::Identity;
+				TmpSceneData.LocalBounds = InComponent->GetStaticMesh()->GetBounds();
+				TmpSceneData.NaniteHierarchyOffset = NANITE_INVALID_HIERARCHY_OFFSET;
+				TmpSceneData.Flags = InstanceDataFlags;
+			}
+
 			int32 OutInstanceIndex = InInstanceIndex;
 			// GPUCULL_TODO: After deleting instances in a HISM the InstanceReorderTable often contains nonsense, this is then corrected
 			// by the async build, which re-creates the proxy in a nearby future frame. All this should be removed in favour of GPU-side culling.
@@ -1295,9 +1304,6 @@ void FInstancedStaticMeshSceneProxy::SetupProxy(UInstancedStaticMeshComponent* I
 
 			FPrimitiveInstance& SceneData = InstanceSceneData[OutInstanceIndex];
 
-			SceneData.Flags = InstanceDataFlags;
-			SceneData.NaniteHierarchyOffset = NANITE_INVALID_HIERARCHY_OFFSET;
-			SceneData.LocalBounds = InComponent->GetStaticMesh()->GetBounds();
 
 			FTransform InstanceTransform;
 			InComponent->GetInstanceTransform(InInstanceIndex, InstanceTransform);
@@ -1360,6 +1366,13 @@ void FInstancedStaticMeshSceneProxy::CreateRenderThreadResources()
 			if (InstanceSceneData.Num() == 0)
 			{
 				InstanceSceneData.SetNum(InstanceBuffer.GetNumInstances());
+				for (int32 InstanceIndex = 0; InstanceIndex < InstanceSceneData.Num(); ++InstanceIndex)
+				{
+					FPrimitiveInstance& SceneData = InstanceSceneData[InstanceIndex];
+					SceneData.NaniteHierarchyOffset = NANITE_INVALID_HIERARCHY_OFFSET;
+					// TODO: Probably need to set  these flags up properly?
+					SceneData.Flags = 0U;
+				}
 			}
 
 			// NOTE: we set up partial data in the construction of ISM proxy (yep, awful but the equally awful way the InstanceBuffer is maintained means complete data is not available)
