Correctly namespaced anim node function caller

diff --git a/Engine/Source/Runtime/Engine/Classes/Animation/AnimNodeBase.h b/Engine/Source/Runtime/Engine/Classes/Animation/AnimNodeBase.h
--- a/Engine/Source/Runtime/Engine/Classes/Animation/AnimNodeBase.h
+++ b/Engine/Source/Runtime/Engine/Classes/Animation/AnimNodeBase.h
@@ -975,7 +975,7 @@ private:
 	friend class IAnimClassInterface;
 	friend class UAnimBlueprintGeneratedClass;
 	friend struct UE::Anim::FNodeDataId;
-	friend struct FAnimNodeFunctionCaller;
+	friend struct UE::Anim::FNodeFunctionCaller;
 	friend class UAnimGraphNode_Base;
 
 	// Set the cached ptr to the constant/folded data for this node
diff --git a/Engine/Source/Runtime/Engine/Classes/Animation/AnimNodeFunctionRef.h b/Engine/Source/Runtime/Engine/Classes/Animation/AnimNodeFunctionRef.h
--- a/Engine/Source/Runtime/Engine/Classes/Animation/AnimNodeFunctionRef.h
+++ b/Engine/Source/Runtime/Engine/Classes/Animation/AnimNodeFunctionRef.h
@@ -5,11 +5,15 @@
 #include "CoreMinimal.h"
 #include "AnimNodeFunctionRef.generated.h"
 
+struct FPoseLink;
+struct FPoseLinkBase;
+struct FComponentSpacePoseLink;
 struct FAnimNode_Base;
 struct FAnimationInitializeContext;
 struct FAnimationUpdateContext;
 struct FPoseContext;
 struct FComponentSpacePoseContext;
+struct FAnimInstanceProxy;
 
 /**
  * Cached function name/ptr that is resolved at init time
@@ -51,14 +55,16 @@ private:
 	UFunction* Function = nullptr;
 };
 
+namespace UE { namespace Anim {
+
 // Wrapper used to call anim node functions
-struct FAnimNodeFunctionCaller
+struct FNodeFunctionCaller
 {
 private:
-	friend struct FPoseLinkBase;
-	friend struct FPoseLink;
-	friend struct FComponentSpacePoseLink;
-	friend struct FAnimInstanceProxy;
+	friend struct ::FPoseLinkBase;
+	friend struct ::FPoseLink;
+	friend struct ::FComponentSpacePoseLink;
+	friend struct ::FAnimInstanceProxy;
 	
 	// Call the Initialize function of this node
 	static void Initialize(const FAnimationInitializeContext& InContext, FAnimNode_Base& InNode);
@@ -74,4 +80,6 @@ private:
 
 	// Call the EvaluateComponentSpace function of this node
 	static void EvaluateComponentSpace(FComponentSpacePoseContext& InContext, FAnimNode_Base& InNode);
-};
\ No newline at end of file
+};
+
+}}
diff --git a/Engine/Source/Runtime/Engine/Private/Animation/AnimInstanceProxy.cpp b/Engine/Source/Runtime/Engine/Private/Animation/AnimInstanceProxy.cpp
--- a/Engine/Source/Runtime/Engine/Private/Animation/AnimInstanceProxy.cpp
+++ b/Engine/Source/Runtime/Engine/Private/Animation/AnimInstanceProxy.cpp
@@ -59,8 +59,8 @@ void FAnimInstanceProxy::UpdateAnimationNode_WithRoot(const FAnimationUpdateCont
 			UpdateCounter.Increment();
 		}
 
-		FAnimNodeFunctionCaller::BecomeRelevant(InContext, *InRootNode);
-		FAnimNodeFunctionCaller::Update(InContext, *InRootNode);
+		UE::Anim::FNodeFunctionCaller::BecomeRelevant(InContext, *InRootNode);
+		UE::Anim::FNodeFunctionCaller::Update(InContext, *InRootNode);
 		InRootNode->Update_AnyThread(InContext);
 
 		// We've updated the graph, now update the fractured saved pose sections
@@ -297,7 +297,7 @@ void FAnimInstanceProxy::InitializeRootNode_WithRoot(FAnimNode_Base* InRootNode)
 		FAnimationUpdateSharedContext SharedContext;
 		FAnimationInitializeContext InitContext(this, &SharedContext);
 
-		FAnimNodeFunctionCaller::Initialize(InitContext, *InRootNode);
+		UE::Anim::FNodeFunctionCaller::Initialize(InitContext, *InRootNode);
 		
 		if(InRootNode == RootNode)
 		{
@@ -1293,7 +1293,7 @@ void FAnimInstanceProxy::EvaluateAnimationNode_WithRoot(FPoseContext& Output, FA
 
 		TRACE_SCOPED_ANIM_NODE(Output);
 
-		FAnimNodeFunctionCaller::Evaluate(Output, *InRootNode);
+		UE::Anim::FNodeFunctionCaller::Evaluate(Output, *InRootNode);
 		
 		InRootNode->Evaluate_AnyThread(Output);
 	}
diff --git a/Engine/Source/Runtime/Engine/Private/Animation/AnimNodeBase.cpp b/Engine/Source/Runtime/Engine/Private/Animation/AnimNodeBase.cpp
--- a/Engine/Source/Runtime/Engine/Private/Animation/AnimNodeBase.cpp
+++ b/Engine/Source/Runtime/Engine/Private/Animation/AnimNodeBase.cpp
@@ -214,7 +214,7 @@ void FPoseLinkBase::Initialize(const FAnimationInitializeContext& InContext)
 		FAnimationInitializeContext LinkContext(InContext);
 		LinkContext.SetNodeId(LinkID);
 		TRACE_SCOPED_ANIM_NODE(LinkContext);
-		FAnimNodeFunctionCaller::Initialize(LinkContext, *LinkedNode);
+		UE::Anim::FNodeFunctionCaller::Initialize(LinkContext, *LinkedNode);
 		LinkedNode->Initialize_AnyThread(LinkContext);
 	}
 }
@@ -325,8 +325,8 @@ void FPoseLinkBase::Update(const FAnimationUpdateContext& InContext)
 	{
 		FAnimationUpdateContext LinkContext(InContext.WithNodeId(LinkID));
 		TRACE_SCOPED_ANIM_NODE(LinkContext);
-		FAnimNodeFunctionCaller::BecomeRelevant(LinkContext, *LinkedNode);
-		FAnimNodeFunctionCaller::Update(LinkContext, *LinkedNode);
+		UE::Anim::FNodeFunctionCaller::BecomeRelevant(LinkContext, *LinkedNode);
+		UE::Anim::FNodeFunctionCaller::Update(LinkContext, *LinkedNode);
 		LinkedNode->Update_AnyThread(LinkContext);
 	}
 }
@@ -376,7 +376,7 @@ void FPoseLink::Evaluate(FPoseContext& Output)
 		{
 			Output.SetNodeId(LinkID);
 			TRACE_SCOPED_ANIM_NODE(Output);
-			FAnimNodeFunctionCaller::Evaluate(Output, *LinkedNode);
+			UE::Anim::FNodeFunctionCaller::Evaluate(Output, *LinkedNode);
 			LinkedNode->Evaluate_AnyThread(Output);
 			TRACE_ANIM_NODE_BLENDABLE_ATTRIBUTES(Output, SourceID, LinkID);
 		}
@@ -487,7 +487,7 @@ void FComponentSpacePoseLink::EvaluateComponentSpace(FComponentSpacePoseContext&
 		{
 			Output.SetNodeId(LinkID);
 			TRACE_SCOPED_ANIM_NODE(Output);
-			FAnimNodeFunctionCaller::EvaluateComponentSpace(Output, *LinkedNode);
+			UE::Anim::FNodeFunctionCaller::EvaluateComponentSpace(Output, *LinkedNode);
 			LinkedNode->EvaluateComponentSpace_AnyThread(Output);
 			TRACE_ANIM_NODE_BLENDABLE_ATTRIBUTES(Output, SourceID, LinkID);
 		}
diff --git a/Engine/Source/Runtime/Engine/Private/Animation/AnimNodeFunctionRef.cpp b/Engine/Source/Runtime/Engine/Private/Animation/AnimNodeFunctionRef.cpp
--- a/Engine/Source/Runtime/Engine/Private/Animation/AnimNodeFunctionRef.cpp
+++ b/Engine/Source/Runtime/Engine/Private/Animation/AnimNodeFunctionRef.cpp
@@ -23,6 +23,11 @@ void FAnimNodeFunctionRef::Call(UObject* InObject, void* InParameters) const
 	}
 }
 
+namespace UE
+{
+namespace Anim
+{
+
 template<typename ContextType>
 static void CallFunctionHelper(const FAnimNodeFunctionRef& InFunction, ContextType InContext, FAnimNode_Base& InNode)
 {
@@ -44,12 +49,12 @@ static void CallFunctionHelper(const FAnimNodeFunctionRef& InFunction, ContextTy
 	}
 }
 
-void FAnimNodeFunctionCaller::Initialize(const FAnimationInitializeContext& InContext, FAnimNode_Base& InNode)
+void FNodeFunctionCaller::Initialize(const FAnimationInitializeContext& InContext, FAnimNode_Base& InNode)
 {
 	CallFunctionHelper(InNode.GetInitializeFunction(), InContext, InNode);
 }
 	
-void FAnimNodeFunctionCaller::BecomeRelevant(const FAnimationUpdateContext& InContext, FAnimNode_Base& InNode)
+void FNodeFunctionCaller::BecomeRelevant(const FAnimationUpdateContext& InContext, FAnimNode_Base& InNode)
 {
 	const FAnimNodeFunctionRef& Function = InNode.GetBecomeRelevantFunction();
 	if(Function.IsValid())
@@ -63,17 +68,19 @@ void FAnimNodeFunctionCaller::BecomeRelevant(const FAnimationUpdateContext& InCo
 	}
 }
 	
-void FAnimNodeFunctionCaller::Update(const FAnimationUpdateContext& InContext, FAnimNode_Base& InNode)
+void FNodeFunctionCaller::Update(const FAnimationUpdateContext& InContext, FAnimNode_Base& InNode)
 {
 	CallFunctionHelper(InNode.GetUpdateFunction(), InContext, InNode);
 }
 
-void FAnimNodeFunctionCaller::Evaluate(FPoseContext& InContext, FAnimNode_Base& InNode)
+void FNodeFunctionCaller::Evaluate(FPoseContext& InContext, FAnimNode_Base& InNode)
 {
 	CallFunctionHelper(InNode.GetEvaluateFunction(), InContext, InNode);
 }
 
-void FAnimNodeFunctionCaller::EvaluateComponentSpace(FComponentSpacePoseContext& InContext, FAnimNode_Base& InNode)
+void FNodeFunctionCaller::EvaluateComponentSpace(FComponentSpacePoseContext& InContext, FAnimNode_Base& InNode)
 {
 	CallFunctionHelper(InNode.GetEvaluateFunction(), InContext, InNode);
-}
\ No newline at end of file
+}
+
+}}
\ No newline at end of file
