FISMComponentDescriptor - Include HLODBatchingPolicy

diff --git a/Engine/Source/Runtime/Engine/Private/ISMPartition/ISMComponentDescriptor.cpp b/Engine/Source/Runtime/Engine/Private/ISMPartition/ISMComponentDescriptor.cpp
--- a/Engine/Source/Runtime/Engine/Private/ISMPartition/ISMComponentDescriptor.cpp
+++ b/Engine/Source/Runtime/Engine/Private/ISMPartition/ISMComponentDescriptor.cpp
@@ -49,6 +49,7 @@ void FISMComponentDescriptor::InitFrom(const UStaticMeshComponent* Template, boo
 	TranslucencySortPriority = Template->TranslucencySortPriority;
 	OverriddenLightMapRes = Template->OverriddenLightMapRes;
 	CustomDepthStencilValue = Template->CustomDepthStencilValue;
+	HLODBatchingPolicy = Template->HLODBatchingPolicy;
 	bCastShadow = Template->CastShadow;
 	bCastStaticShadow = Template->bCastStaticShadow;
 	bCastDynamicShadow = Template->bCastDynamicShadow;
@@ -111,6 +112,7 @@ bool FISMComponentDescriptor::operator==(const FISMComponentDescriptor& Other) c
 	TranslucencySortPriority == Other.TranslucencySortPriority &&
 	OverriddenLightMapRes == Other.OverriddenLightMapRes &&
 	CustomDepthStencilValue == Other.CustomDepthStencilValue &&
+	HLODBatchingPolicy == Other.HLODBatchingPolicy &&
 	bCastShadow == Other.bCastShadow &&
 	bCastStaticShadow == Other.bCastStaticShadow &&
 	bCastDynamicShadow == Other.bCastDynamicShadow &&
@@ -175,6 +177,7 @@ void FISMComponentDescriptor::InitComponent(UInstancedStaticMeshComponent* ISMCo
 	ISMComponent->TranslucencySortPriority = TranslucencySortPriority;
 	ISMComponent->OverriddenLightMapRes = OverriddenLightMapRes;
 	ISMComponent->CustomDepthStencilValue = CustomDepthStencilValue;
+	ISMComponent->HLODBatchingPolicy = HLODBatchingPolicy;
 	ISMComponent->CastShadow = bCastShadow;
 	ISMComponent->bCastStaticShadow = bCastStaticShadow;
 	ISMComponent->bCastDynamicShadow = bCastDynamicShadow;
diff --git a/Engine/Source/Runtime/Engine/Public/ISMPartition/ISMComponentDescriptor.h b/Engine/Source/Runtime/Engine/Public/ISMPartition/ISMComponentDescriptor.h
--- a/Engine/Source/Runtime/Engine/Public/ISMPartition/ISMComponentDescriptor.h
+++ b/Engine/Source/Runtime/Engine/Public/ISMPartition/ISMComponentDescriptor.h
@@ -108,6 +108,9 @@ public:
 
 	UPROPERTY()
 	int32 CustomDepthStencilValue;
+
+	UPROPERTY()
+	EHLODBatchingPolicy HLODBatchingPolicy;
 	
 	UPROPERTY()
 	uint8 bCastShadow : 1;
