Fixed bug in Level Streaming when package is already or still loaded. For Editor world type, level must be added to owning world.
Apply same logic to level as the regular code path of ULevelStreaming::AsyncLevelLoadComplete on a loaded level.

diff --git a/Engine/Source/Runtime/Engine/Classes/Engine/LevelStreaming.h b/Engine/Source/Runtime/Engine/Classes/Engine/LevelStreaming.h
--- a/Engine/Source/Runtime/Engine/Classes/Engine/LevelStreaming.h
+++ b/Engine/Source/Runtime/Engine/Classes/Engine/LevelStreaming.h
@@ -537,6 +537,9 @@ protected:
 	 */
 	virtual bool RequestLevel(UWorld* PersistentWorld, bool bAllowLevelLoadRequests, EReqLevelBlock BlockPolicy);
 
+	/** Sets loaded level, fixups for PIE, notifies level is loaded, apply necessary modifications on level once loaded. */
+	void PrepareLoadedLevel(ULevel* InLevel, UPackage* InLevelPackage, int32 InPIEInstanceID);
+
 	/** Sets the value of LoadedLevel */
 	virtual void SetLoadedLevel(ULevel* Level);
 
diff --git a/Engine/Source/Runtime/Engine/Private/LevelStreaming.cpp b/Engine/Source/Runtime/Engine/Private/LevelStreaming.cpp
--- a/Engine/Source/Runtime/Engine/Private/LevelStreaming.cpp
+++ b/Engine/Source/Runtime/Engine/Private/LevelStreaming.cpp
@@ -955,6 +955,70 @@ bool ULevelStreaming::IsDesiredLevelLoaded() const
 	return false;
 }
 
+void ULevelStreaming::PrepareLoadedLevel(ULevel* InLevel, UPackage* InLevelPackage, int32 InPIEInstanceID)
+{
+	check(InLevel);
+	UWorld* LevelOwningWorld = InLevel->OwningWorld;
+
+	if (ensure(LevelOwningWorld))
+	{
+		ULevel* PendingLevelVisOrInvis = (LevelOwningWorld->GetCurrentLevelPendingVisibility() ? LevelOwningWorld->GetCurrentLevelPendingVisibility() : LevelOwningWorld->GetCurrentLevelPendingInvisibility());
+		if (PendingLevelVisOrInvis && PendingLevelVisOrInvis == LoadedLevel)
+		{
+			// We can't change current loaded level if it's still processing visibility request
+			// On next UpdateLevelStreaming call this loaded package will be found in memory by RequestLevel function in case visibility request has finished
+			UE_LOG(LogLevelStreaming, Verbose, TEXT("Delaying setting result of async load new level %s, because current loaded level still processing visibility request"), *InLevelPackage->GetName());
+		}
+		else
+		{
+			check(PendingUnloadLevel == nullptr);
+
+#if WITH_EDITOR
+			if (InPIEInstanceID != INDEX_NONE)
+			{
+				InLevel->FixupForPIE(InPIEInstanceID);
+			}
+#endif
+			SetLoadedLevel(InLevel);
+			// Broadcast level loaded event to blueprints
+			OnLevelLoaded.Broadcast();
+		}
+	}
+
+	InLevel->HandleLegacyMapBuildData();
+
+	// Notify the streamer to start building incrementally the level streaming data.
+	IStreamingManager::Get().AddLevel(InLevel);
+
+	// Make sure this level will start to render only when it will be fully added to the world
+	if (ShouldRequireFullVisibilityToRender())
+	{
+		InLevel->bRequireFullVisibilityToRender = true;
+		// LOD levels should not be visible on server
+		if (LODPackageNames.Num() > 0)
+		{
+			InLevel->bClientOnlyVisible = LODPackageNames.Contains(InLevelPackage->GetFName());
+		}
+	}
+
+	// Apply streaming level property to level
+	InLevel->bClientOnlyVisible |= bClientOnlyVisible;
+
+	// In the editor levels must be in the levels array regardless of whether they are visible or not
+	if (ensure(LevelOwningWorld) && LevelOwningWorld->WorldType == EWorldType::Editor)
+	{
+		LevelOwningWorld->AddLevel(InLevel);
+#if WITH_EDITOR
+		// We should also at this point, apply the level's editor transform
+		if (!InLevel->bAlreadyMovedActors)
+		{
+			FLevelUtils::ApplyEditorTransform(this, false);
+			InLevel->bAlreadyMovedActors = true;
+		}
+#endif // WITH_EDITOR
+	}
+}
+
 bool ULevelStreaming::RequestLevel(UWorld* PersistentWorld, bool bAllowLevelLoadRequests, EReqLevelBlock BlockPolicy)
 {
 	// Quit early in case load request already issued
@@ -1121,20 +1185,11 @@ bool ULevelStreaming::RequestLevel(UWorld* PersistentWorld, bool bAllowLevelLoad
 #endif
 			if (World->PersistentLevel != LoadedLevel)
 			{
-#if WITH_EDITOR
-				if (PIEInstanceID != INDEX_NONE)
-				{
-					World->PersistentLevel->FixupForPIE(PIEInstanceID);
-				}
-#endif
-
 				// Level already exists but may have the wrong type due to being inactive before, so copy data over
 				World->WorldType = PersistentWorld->WorldType;
 				World->PersistentLevel->OwningWorld = PersistentWorld;
 
-				SetLoadedLevel(World->PersistentLevel);
-				// Broadcast level loaded event to blueprints
-				OnLevelLoaded.Broadcast();
+				PrepareLoadedLevel(World->PersistentLevel, LevelPackage, PIEInstanceID);
 			}
 			
 			return true;
@@ -1238,69 +1293,9 @@ void ULevelStreaming::AsyncLevelLoadComplete(const FName& InPackageName, UPackag
 
 		if (World)
 		{
-			ULevel* Level = World->PersistentLevel;
-			if (Level)
+			if (ULevel* Level = World->PersistentLevel)
 			{
-				UWorld* LevelOwningWorld = Level->OwningWorld;
-				if (LevelOwningWorld)
-				{
-					ULevel* PendingLevelVisOrInvis = (LevelOwningWorld->GetCurrentLevelPendingVisibility() ? LevelOwningWorld->GetCurrentLevelPendingVisibility() : LevelOwningWorld->GetCurrentLevelPendingInvisibility());
-					if (PendingLevelVisOrInvis && PendingLevelVisOrInvis == LoadedLevel)
-					{
-						// We can't change current loaded level if it's still processing visibility request
-						// On next UpdateLevelStreaming call this loaded package will be found in memory by RequestLevel function in case visibility request has finished
-						UE_LOG(LogLevelStreaming, Verbose, TEXT("Delaying setting result of async load new level %s, because current loaded level still processing visibility request"), *LevelPackage->GetName());
-					}
-					else
-					{
-						check(PendingUnloadLevel == nullptr);
-					
-#if WITH_EDITOR
-						int32 PIEInstanceID = GetOutermost()->PIEInstanceID;
-						if (PIEInstanceID != INDEX_NONE)
-						{
-							World->PersistentLevel->FixupForPIE(PIEInstanceID);
-						}
-#endif
-
-						SetLoadedLevel(Level);
-						// Broadcast level loaded event to blueprints
-						OnLevelLoaded.Broadcast();
-					}
-				}
-
-				Level->HandleLegacyMapBuildData();
-
-				// Notify the streamer to start building incrementally the level streaming data.
-				IStreamingManager::Get().AddLevel(Level);
-
-				// Make sure this level will start to render only when it will be fully added to the world
-				if (ShouldRequireFullVisibilityToRender())
-				{
-					Level->bRequireFullVisibilityToRender = true;
-					// LOD levels should not be visible on server
-					if (LODPackageNames.Num() > 0)
-					{
-						Level->bClientOnlyVisible = LODPackageNames.Contains(InLoadedPackage->GetFName());
-					}
-				}
-
-				// Apply streaming level property to level
-				Level->bClientOnlyVisible |= bClientOnlyVisible;
-			
-				// In the editor levels must be in the levels array regardless of whether they are visible or not
-				if (ensure(LevelOwningWorld) && LevelOwningWorld->WorldType == EWorldType::Editor)
-				{
-					LevelOwningWorld->AddLevel(Level);
-#if WITH_EDITOR
-					// We should also at this point, apply the level's editor transform
-					if (!Level->bAlreadyMovedActors)
-					{
-						FLevelUtils::ApplyEditorTransform(this, false);
-						Level->bAlreadyMovedActors = true;
-					}
-#endif // WITH_EDITOR
-				}
+				PrepareLoadedLevel(Level, LevelPackage, GetOutermost()->PIEInstanceID);
 			}
 			else
 			{
