Texture: Fixed serialization of bitfields to compact binary

#rnx

diff --git a/Engine/Source/Runtime/Engine/Private/TextureDerivedDataBuildUtils.cpp b/Engine/Source/Runtime/Engine/Private/TextureDerivedDataBuildUtils.cpp
--- a/Engine/Source/Runtime/Engine/Private/TextureDerivedDataBuildUtils.cpp
+++ b/Engine/Source/Runtime/Engine/Private/TextureDerivedDataBuildUtils.cpp
@@ -104,26 +104,26 @@ static void WriteBuildSettings(FCbWriter& Writer, const FTextureBuildSettings& B
 	WriteCbFieldWithDefault(Writer, "TextureFormatName", WriteToString<64>(BuildSettings.TextureFormatName).ToView(), TEXT(""_SV));
 	WriteCbFieldWithDefault(Writer, "bHDRSource", BuildSettings.bHDRSource, DefaultSettings.bHDRSource);
 	WriteCbFieldWithDefault(Writer, "MipGenSettings", BuildSettings.MipGenSettings, DefaultSettings.MipGenSettings);
-	WriteCbFieldWithDefault(Writer, "bCubemap", BuildSettings.bCubemap, DefaultSettings.bCubemap);
-	WriteCbFieldWithDefault(Writer, "bTextureArray", BuildSettings.bTextureArray, DefaultSettings.bTextureArray);
-	WriteCbFieldWithDefault(Writer, "bVolume", BuildSettings.bVolume, DefaultSettings.bVolume);
-	WriteCbFieldWithDefault(Writer, "bLongLatSource", BuildSettings.bLongLatSource, DefaultSettings.bLongLatSource);
-	WriteCbFieldWithDefault(Writer, "bSRGB", BuildSettings.bSRGB, DefaultSettings.bSRGB);
-	WriteCbFieldWithDefault(Writer, "bUseLegacyGamma", BuildSettings.bUseLegacyGamma, DefaultSettings.bUseLegacyGamma);
-	WriteCbFieldWithDefault(Writer, "bPreserveBorder", BuildSettings.bPreserveBorder, DefaultSettings.bPreserveBorder);
-	WriteCbFieldWithDefault(Writer, "bForceNoAlphaChannel", BuildSettings.bForceNoAlphaChannel, DefaultSettings.bForceNoAlphaChannel);
-	WriteCbFieldWithDefault(Writer, "bForceAlphaChannel", BuildSettings.bForceAlphaChannel, DefaultSettings.bForceAlphaChannel);
-	WriteCbFieldWithDefault(Writer, "bDitherMipMapAlpha", BuildSettings.bDitherMipMapAlpha, DefaultSettings.bDitherMipMapAlpha);
-	WriteCbFieldWithDefault(Writer, "bComputeBokehAlpha", BuildSettings.bComputeBokehAlpha, DefaultSettings.bComputeBokehAlpha);
-	WriteCbFieldWithDefault(Writer, "bReplicateRed", BuildSettings.bReplicateRed, DefaultSettings.bReplicateRed);
-	WriteCbFieldWithDefault(Writer, "bReplicateAlpha", BuildSettings.bReplicateAlpha, DefaultSettings.bReplicateAlpha);
-	WriteCbFieldWithDefault(Writer, "bDownsampleWithAverage", BuildSettings.bDownsampleWithAverage, DefaultSettings.bDownsampleWithAverage);
-	WriteCbFieldWithDefault(Writer, "bSharpenWithoutColorShift", BuildSettings.bSharpenWithoutColorShift, DefaultSettings.bSharpenWithoutColorShift);
-	WriteCbFieldWithDefault(Writer, "bBorderColorBlack", BuildSettings.bBorderColorBlack, DefaultSettings.bBorderColorBlack);
-	WriteCbFieldWithDefault(Writer, "bFlipGreenChannel", BuildSettings.bFlipGreenChannel, DefaultSettings.bFlipGreenChannel);
-	WriteCbFieldWithDefault(Writer, "bApplyYCoCgBlockScale", BuildSettings.bApplyYCoCgBlockScale, DefaultSettings.bApplyYCoCgBlockScale);
-	WriteCbFieldWithDefault(Writer, "bApplyKernelToTopMip", BuildSettings.bApplyKernelToTopMip, DefaultSettings.bApplyKernelToTopMip);
-	WriteCbFieldWithDefault(Writer, "bRenormalizeTopMip", BuildSettings.bRenormalizeTopMip, DefaultSettings.bRenormalizeTopMip);
+	WriteCbFieldWithDefault<bool>(Writer, "bCubemap", BuildSettings.bCubemap, DefaultSettings.bCubemap);
+	WriteCbFieldWithDefault<bool>(Writer, "bTextureArray", BuildSettings.bTextureArray, DefaultSettings.bTextureArray);
+	WriteCbFieldWithDefault<bool>(Writer, "bVolume", BuildSettings.bVolume, DefaultSettings.bVolume);
+	WriteCbFieldWithDefault<bool>(Writer, "bLongLatSource", BuildSettings.bLongLatSource, DefaultSettings.bLongLatSource);
+	WriteCbFieldWithDefault<bool>(Writer, "bSRGB", BuildSettings.bSRGB, DefaultSettings.bSRGB);
+	WriteCbFieldWithDefault<bool>(Writer, "bUseLegacyGamma", BuildSettings.bUseLegacyGamma, DefaultSettings.bUseLegacyGamma);
+	WriteCbFieldWithDefault<bool>(Writer, "bPreserveBorder", BuildSettings.bPreserveBorder, DefaultSettings.bPreserveBorder);
+	WriteCbFieldWithDefault<bool>(Writer, "bForceNoAlphaChannel", BuildSettings.bForceNoAlphaChannel, DefaultSettings.bForceNoAlphaChannel);
+	WriteCbFieldWithDefault<bool>(Writer, "bForceAlphaChannel", BuildSettings.bForceAlphaChannel, DefaultSettings.bForceAlphaChannel);
+	WriteCbFieldWithDefault<bool>(Writer, "bDitherMipMapAlpha", BuildSettings.bDitherMipMapAlpha, DefaultSettings.bDitherMipMapAlpha);
+	WriteCbFieldWithDefault<bool>(Writer, "bComputeBokehAlpha", BuildSettings.bComputeBokehAlpha, DefaultSettings.bComputeBokehAlpha);
+	WriteCbFieldWithDefault<bool>(Writer, "bReplicateRed", BuildSettings.bReplicateRed, DefaultSettings.bReplicateRed);
+	WriteCbFieldWithDefault<bool>(Writer, "bReplicateAlpha", BuildSettings.bReplicateAlpha, DefaultSettings.bReplicateAlpha);
+	WriteCbFieldWithDefault<bool>(Writer, "bDownsampleWithAverage", BuildSettings.bDownsampleWithAverage, DefaultSettings.bDownsampleWithAverage);
+	WriteCbFieldWithDefault<bool>(Writer, "bSharpenWithoutColorShift", BuildSettings.bSharpenWithoutColorShift, DefaultSettings.bSharpenWithoutColorShift);
+	WriteCbFieldWithDefault<bool>(Writer, "bBorderColorBlack", BuildSettings.bBorderColorBlack, DefaultSettings.bBorderColorBlack);
+	WriteCbFieldWithDefault<bool>(Writer, "bFlipGreenChannel", BuildSettings.bFlipGreenChannel, DefaultSettings.bFlipGreenChannel);
+	WriteCbFieldWithDefault<bool>(Writer, "bApplyYCoCgBlockScale", BuildSettings.bApplyYCoCgBlockScale, DefaultSettings.bApplyYCoCgBlockScale);
+	WriteCbFieldWithDefault<bool>(Writer, "bApplyKernelToTopMip", BuildSettings.bApplyKernelToTopMip, DefaultSettings.bApplyKernelToTopMip);
+	WriteCbFieldWithDefault<bool>(Writer, "bRenormalizeTopMip", BuildSettings.bRenormalizeTopMip, DefaultSettings.bRenormalizeTopMip);
 	WriteCbFieldWithDefault(Writer, "CompositeTextureMode", BuildSettings.CompositeTextureMode, DefaultSettings.CompositeTextureMode);
 	WriteCbFieldWithDefault(Writer, "CompositePower", BuildSettings.CompositePower, DefaultSettings.CompositePower);
 	WriteCbFieldWithDefault(Writer, "LODBias", BuildSettings.LODBias, DefaultSettings.LODBias);
@@ -131,9 +131,9 @@ static void WriteBuildSettings(FCbWriter& Writer, const FTextureBuildSettings& B
 	WriteCbFieldWithDefault(Writer, "TopMipSize", BuildSettings.TopMipSize, DefaultSettings.TopMipSize);
 	WriteCbFieldWithDefault(Writer, "VolumeSizeZ", BuildSettings.VolumeSizeZ, DefaultSettings.VolumeSizeZ);
 	WriteCbFieldWithDefault(Writer, "ArraySlices", BuildSettings.ArraySlices, DefaultSettings.ArraySlices);
-	WriteCbFieldWithDefault(Writer, "bStreamable", BuildSettings.bStreamable, DefaultSettings.bStreamable);
-	WriteCbFieldWithDefault(Writer, "bVirtualStreamable", BuildSettings.bVirtualStreamable, DefaultSettings.bVirtualStreamable);
-	WriteCbFieldWithDefault(Writer, "bChromaKeyTexture", BuildSettings.bChromaKeyTexture, DefaultSettings.bChromaKeyTexture);
+	WriteCbFieldWithDefault<bool>(Writer, "bStreamable", BuildSettings.bStreamable, DefaultSettings.bStreamable);
+	WriteCbFieldWithDefault<bool>(Writer, "bVirtualStreamable", BuildSettings.bVirtualStreamable, DefaultSettings.bVirtualStreamable);
+	WriteCbFieldWithDefault<bool>(Writer, "bChromaKeyTexture", BuildSettings.bChromaKeyTexture, DefaultSettings.bChromaKeyTexture);
 	WriteCbFieldWithDefault(Writer, "PowerOfTwoMode", BuildSettings.PowerOfTwoMode, DefaultSettings.PowerOfTwoMode);
 	WriteCbFieldWithDefault(Writer, "PaddingColor", BuildSettings.PaddingColor, DefaultSettings.PaddingColor);
 	WriteCbFieldWithDefault(Writer, "ChromaKeyColor", BuildSettings.ChromaKeyColor, DefaultSettings.ChromaKeyColor);
@@ -146,8 +146,8 @@ static void WriteBuildSettings(FCbWriter& Writer, const FTextureBuildSettings& B
 	WriteCbFieldWithDefault(Writer, "VirtualAddressingModeY", BuildSettings.VirtualAddressingModeY, DefaultSettings.VirtualAddressingModeY);
 	WriteCbFieldWithDefault(Writer, "VirtualTextureTileSize", BuildSettings.VirtualTextureTileSize, DefaultSettings.VirtualTextureTileSize);
 	WriteCbFieldWithDefault(Writer, "VirtualTextureBorderSize", BuildSettings.VirtualTextureBorderSize, DefaultSettings.VirtualTextureBorderSize);
-	WriteCbFieldWithDefault(Writer, "bVirtualTextureEnableCompressZlib", BuildSettings.bVirtualTextureEnableCompressZlib, DefaultSettings.bVirtualTextureEnableCompressZlib);
-	WriteCbFieldWithDefault(Writer, "bVirtualTextureEnableCompressCrunch", BuildSettings.bVirtualTextureEnableCompressCrunch, DefaultSettings.bVirtualTextureEnableCompressCrunch);
+	WriteCbFieldWithDefault<bool>(Writer, "bVirtualTextureEnableCompressZlib", BuildSettings.bVirtualTextureEnableCompressZlib, DefaultSettings.bVirtualTextureEnableCompressZlib);
+	WriteCbFieldWithDefault<bool>(Writer, "bVirtualTextureEnableCompressCrunch", BuildSettings.bVirtualTextureEnableCompressCrunch, DefaultSettings.bVirtualTextureEnableCompressCrunch);
 	WriteCbFieldWithDefault(Writer, "bHasEditorOnlyData", BuildSettings.bHasEditorOnlyData, DefaultSettings.bHasEditorOnlyData);
 
 	Writer.EndObject();
