Fix crash calling anim node functions when running standalone nodes

They have no node data, so cannot call node functions.

diff --git a/Engine/Source/Runtime/Engine/Private/Animation/AnimNodeFunctionRef.cpp b/Engine/Source/Runtime/Engine/Private/Animation/AnimNodeFunctionRef.cpp
--- a/Engine/Source/Runtime/Engine/Private/Animation/AnimNodeFunctionRef.cpp
+++ b/Engine/Source/Runtime/Engine/Private/Animation/AnimNodeFunctionRef.cpp
@@ -51,36 +51,51 @@ static void CallFunctionHelper(const FAnimNodeFunctionRef& InFunction, ContextTy
 
 void FNodeFunctionCaller::Initialize(const FAnimationInitializeContext& InContext, FAnimNode_Base& InNode)
 {
-	CallFunctionHelper(InNode.GetInitializeFunction(), InContext, InNode);
+	if(InNode.NodeData != nullptr)
+	{
+		CallFunctionHelper(InNode.GetInitializeFunction(), InContext, InNode);
+	}
 }
 	
 void FNodeFunctionCaller::BecomeRelevant(const FAnimationUpdateContext& InContext, FAnimNode_Base& InNode)
 {
-	const FAnimNodeFunctionRef& Function = InNode.GetBecomeRelevantFunction();
-	if(Function.IsValid())
+	if(InNode.NodeData != nullptr)
 	{
-		FAnimSubsystemInstance_NodeRelevancy& RelevancySubsystem = CastChecked<UAnimInstance>(InContext.GetAnimInstanceObject())->GetSubsystem<FAnimSubsystemInstance_NodeRelevancy>();
-		FAnimNodeRelevancyStatus Status = RelevancySubsystem.UpdateNodeRelevancy(InContext, InNode);
-		if(Status.HasJustBecomeRelevant())
+		const FAnimNodeFunctionRef& Function = InNode.GetBecomeRelevantFunction();
+		if(Function.IsValid())
 		{
-			CallFunctionHelper(Function, InContext, InNode);
+			FAnimSubsystemInstance_NodeRelevancy& RelevancySubsystem = CastChecked<UAnimInstance>(InContext.GetAnimInstanceObject())->GetSubsystem<FAnimSubsystemInstance_NodeRelevancy>();
+			FAnimNodeRelevancyStatus Status = RelevancySubsystem.UpdateNodeRelevancy(InContext, InNode);
+			if(Status.HasJustBecomeRelevant())
+			{
+				CallFunctionHelper(Function, InContext, InNode);
+			}
 		}
 	}
 }
 	
 void FNodeFunctionCaller::Update(const FAnimationUpdateContext& InContext, FAnimNode_Base& InNode)
 {
-	CallFunctionHelper(InNode.GetUpdateFunction(), InContext, InNode);
+	if(InNode.NodeData != nullptr)
+	{
+		CallFunctionHelper(InNode.GetUpdateFunction(), InContext, InNode);
+	}
 }
 
 void FNodeFunctionCaller::Evaluate(FPoseContext& InContext, FAnimNode_Base& InNode)
 {
-	CallFunctionHelper(InNode.GetEvaluateFunction(), InContext, InNode);
+	if(InNode.NodeData != nullptr)
+	{
+		CallFunctionHelper(InNode.GetEvaluateFunction(), InContext, InNode);
+	}
 }
 
 void FNodeFunctionCaller::EvaluateComponentSpace(FComponentSpacePoseContext& InContext, FAnimNode_Base& InNode)
 {
-	CallFunctionHelper(InNode.GetEvaluateFunction(), InContext, InNode);
+	if(InNode.NodeData != nullptr)
+	{
+		CallFunctionHelper(InNode.GetEvaluateFunction(), InContext, InNode);
+	}
 }
 
 }}
\ No newline at end of file
