Chaos: Implement constraint breaking event

diff --git a/Engine/Source/Runtime/Engine/Classes/PhysicsEngine/PhysicsHandleComponent.h b/Engine/Source/Runtime/Engine/Classes/PhysicsEngine/PhysicsHandleComponent.h
--- a/Engine/Source/Runtime/Engine/Classes/PhysicsEngine/PhysicsHandleComponent.h
+++ b/Engine/Source/Runtime/Engine/Classes/PhysicsEngine/PhysicsHandleComponent.h
@@ -82,7 +82,7 @@ protected:
 	bool bPendingConstraint;
 
 	FPhysicsUserData PhysicsUserData;
-	FConstraintInstanceBase ConstraintInstance;
+	FConstraintInstance ConstraintInstance;
 	FPhysicsActorHandle GrabbedHandle;
 	FPhysicsActorHandle KinematicHandle;
 	FPhysicsConstraintHandle ConstraintHandle;
diff --git a/Engine/Source/Runtime/Engine/Private/PhysicsEngine/Experimental/PhysInterface_Chaos.cpp b/Engine/Source/Runtime/Engine/Private/PhysicsEngine/Experimental/PhysInterface_Chaos.cpp
--- a/Engine/Source/Runtime/Engine/Private/PhysicsEngine/Experimental/PhysInterface_Chaos.cpp
+++ b/Engine/Source/Runtime/Engine/Private/PhysicsEngine/Experimental/PhysInterface_Chaos.cpp
@@ -637,7 +637,7 @@ private:
 		{
 			Chaos::FJointConstraint* Constraint = static_cast<Chaos::FJointConstraint*>(InConstraintHandle->Constraint);
 
-			FConstraintInstanceBase* ConstraintInstance = (Constraint) ? FPhysicsUserData_Chaos::Get<FConstraintInstanceBase>(Constraint->GetUserData()) : nullptr;
+			FConstraintInstance* ConstraintInstance = (Constraint) ? FPhysicsUserData_Chaos::Get<FConstraintInstance>(Constraint->GetUserData()) : nullptr;
 			if (ConstraintInstance)
 			{
 				return ConstraintInstance->GetPhysicsScene();
diff --git a/Engine/Source/Runtime/Engine/Private/PhysicsEngine/Experimental/PhysScene_Chaos.cpp b/Engine/Source/Runtime/Engine/Private/PhysicsEngine/Experimental/PhysScene_Chaos.cpp
--- a/Engine/Source/Runtime/Engine/Private/PhysicsEngine/Experimental/PhysScene_Chaos.cpp
+++ b/Engine/Source/Runtime/Engine/Private/PhysicsEngine/Experimental/PhysScene_Chaos.cpp
@@ -31,6 +31,7 @@
 #include "PhysicsProxy/SingleParticlePhysicsProxy.h"
 #include "PhysicsProxy/SkeletalMeshPhysicsProxy.h"
 #include "PhysicsProxy/StaticMeshPhysicsProxy.h"
+#include "PhysicsProxy/JointConstraintProxy.h"
 #include "Chaos/Real.h"
 #include "Chaos/UniformGrid.h"
 #include "Chaos/BoundingVolume.h"
@@ -1615,7 +1616,7 @@ void FPhysScene_Chaos::OnSyncBodies(Chaos::FPhysicsSolverBase* Solver)
 	TArray<FPhysScenePendingComponentTransform_Chaos> PendingTransforms;
 	TSet<FGeometryCollectionPhysicsProxy*> GCProxies;
 
-	Solver->PullPhysicsStateForEachDirtyProxy_External([&PendingTransforms](FSingleParticlePhysicsProxy* Proxy)
+	auto RigidLambda = [&PendingTransforms](FSingleParticlePhysicsProxy* Proxy)
 	{
 		FPBDRigidParticle* DirtyParticle = Proxy->GetRigidParticleUnsafe();
 
@@ -1653,7 +1654,27 @@ void FPhysScene_Chaos::OnSyncBodies(Chaos::FPhysicsSolverBase* Solver)
 				}
 			}
 		}
-	});
+	};
+
+	auto ConstraintLambda = [&PendingTransforms](FJointConstraintPhysicsProxy* Proxy)
+	{
+		Chaos::FJointConstraint* Constraint = Proxy->GetConstraint();
+
+		if (Constraint->GetOutputData().bIsBreaking)
+		{
+			FConstraintInstance* ConstraintInstance = (Constraint) ? FPhysicsUserData_Chaos::Get<FConstraintInstance>(Constraint->GetUserData()) : nullptr;
+			{
+				FConstraintBrokenDelegateData CBDD(ConstraintInstance);
+				CBDD.DispatchOnBroken();
+			}
+
+			Constraint->GetOutputData().bIsBreaking = false;
+		}
+
+	};
+
+	Solver->PullPhysicsStateForEachDirtyProxy_External(RigidLambda, ConstraintLambda);
+
 
 	for (const FPhysScenePendingComponentTransform_Chaos& ComponentTransform : PendingTransforms)
 	{
@@ -1694,6 +1715,13 @@ void FPhysScene_Chaos::RemoveSpringConstraint(const FPhysicsConstraintHandle& Co
 	// #todo : Implement
 }
 
+FConstraintBrokenDelegateData::FConstraintBrokenDelegateData(FConstraintInstance* ConstraintInstance)
+	: OnConstraintBrokenDelegate(ConstraintInstance->OnConstraintBrokenDelegate)
+	, ConstraintIndex(ConstraintInstance->ConstraintIndex)
+{
+
+}
+
 void FPhysScene_Chaos::ResimNFrames(const int32 NumFramesRequested)
 {
 	//needs to run on physics thread from a special location
diff --git a/Engine/Source/Runtime/Engine/Public/Physics/Experimental/PhysicsUserData_Chaos.h b/Engine/Source/Runtime/Engine/Public/Physics/Experimental/PhysicsUserData_Chaos.h
--- a/Engine/Source/Runtime/Engine/Public/Physics/Experimental/PhysicsUserData_Chaos.h
+++ b/Engine/Source/Runtime/Engine/Public/Physics/Experimental/PhysicsUserData_Chaos.h
@@ -5,8 +5,9 @@
 #include "EngineDefines.h"
 #include "Physics/PhysicsInterfaceDeclares.h"
 
-struct FConstraintInstanceBase;
+struct FConstraintInstance;
 
+// Note: this must match the enum EChaosUserDataType in PhysicsInterfaceTypesCore
 class FPhysicsUserData_Chaos
 {
 private:
@@ -14,8 +15,9 @@ private:
 	{
 		Invalid,
 		BodyInstance,
+		PhysicalMaterial,
 		PhysScene,
-		ConstraintInstanceBase
+		ConstraintInstance
 	};
 
 	EType Type;
@@ -24,8 +26,9 @@ private:
 public:
 	FPhysicsUserData_Chaos() : Type(EType::Invalid), Payload(nullptr) { }
 	FPhysicsUserData_Chaos(FBodyInstance* InPayload) : Type(EType::BodyInstance), Payload(InPayload) { }
+	FPhysicsUserData_Chaos(UPhysicalMaterial* InPayload) : Type(EType::PhysicalMaterial), Payload(InPayload) { }
 	FPhysicsUserData_Chaos(FPhysScene* InPayload) : Type(EType::PhysScene), Payload(InPayload) { }
-	FPhysicsUserData_Chaos(FConstraintInstanceBase* InPayload) : Type(EType::ConstraintInstanceBase), Payload(InPayload) { }
+	FPhysicsUserData_Chaos(FConstraintInstance* InPayload) : Type(EType::ConstraintInstance), Payload(InPayload) { }
 
 	template <class T> static T* Get(void* UserData);
 	template <class T> static void Set(void* UserData, T* InPayload);
@@ -56,12 +59,17 @@ template<> FORCEINLINE FBodyInstance* FPhysicsUserData_Chaos::Get(void* UserData
 	return Get<FBodyInstance, EType::BodyInstance>(UserData);
 }
 
+template<> FORCEINLINE UPhysicalMaterial* FPhysicsUserData_Chaos::Get(void* UserData)
+{
+	return Get<UPhysicalMaterial, EType::PhysicalMaterial>(UserData);
+}
+
 template<> FORCEINLINE FPhysScene* FPhysicsUserData_Chaos::Get(void* UserData)
 {
 	return Get<FPhysScene, EType::PhysScene>(UserData);
 }
 
-template<> FORCEINLINE FConstraintInstanceBase* FPhysicsUserData_Chaos::Get(void* UserData)
+template<> FORCEINLINE FConstraintInstance* FPhysicsUserData_Chaos::Get(void* UserData)
 {
-	return Get<FConstraintInstanceBase, EType::ConstraintInstanceBase>(UserData);
-}
\ No newline at end of file
+	return Get<FConstraintInstance, EType::ConstraintInstance>(UserData);
+}
diff --git a/Engine/Source/Runtime/Experimental/Chaos/Private/Chaos/PBDJointConstraints.cpp b/Engine/Source/Runtime/Experimental/Chaos/Private/Chaos/PBDJointConstraints.cpp
--- a/Engine/Source/Runtime/Experimental/Chaos/Private/Chaos/PBDJointConstraints.cpp
+++ b/Engine/Source/Runtime/Experimental/Chaos/Private/Chaos/PBDJointConstraints.cpp
@@ -81,6 +81,16 @@ namespace Chaos
 		return ConstraintContainer->IsConstraintEnabled(ConstraintIndex);
 	}
 
+	bool FPBDJointConstraintHandle::IsConstraintBreaking() const
+	{
+		return ConstraintContainer->IsConstraintBreaking(ConstraintIndex);
+	}
+
+	void FPBDJointConstraintHandle::ClearConstraintBreaking()
+	{
+		return ConstraintContainer->ClearConstraintBreaking(ConstraintIndex);
+	}
+
 	FVec3 FPBDJointConstraintHandle::GetLinearImpulse() const
 	{
 		return ConstraintContainer->GetConstraintLinearImpulse(ConstraintIndex);
@@ -519,6 +529,15 @@ namespace Chaos
 		return !ConstraintStates[ConstraintIndex].bDisabled;
 	}
 
+	bool FPBDJointConstraints::IsConstraintBreaking(int32 ConstraintIndex) const
+	{
+		return ConstraintStates[ConstraintIndex].bBreaking;
+	}
+
+	void FPBDJointConstraints::ClearConstraintBreaking(int32 ConstraintIndex)
+	{
+		ConstraintStates[ConstraintIndex].bBreaking = false;
+	}
 
 	void FPBDJointConstraints::SetConstraintEnabled(int32 ConstraintIndex, bool bEnabled)
 	{
@@ -541,10 +560,15 @@ namespace Chaos
 		}
 	}
 
+	void FPBDJointConstraints::SetConstraintBreaking(int32 ConstraintIndex, bool bBreaking)
+	{
+		ConstraintStates[ConstraintIndex].bBreaking = bBreaking;
+	}
 
 	void FPBDJointConstraints::BreakConstraint(int32 ConstraintIndex)
 	{
 		SetConstraintEnabled(ConstraintIndex, false);
+		SetConstraintBreaking(ConstraintIndex, true);
 		if (BreakCallback)
 		{
 			BreakCallback(Handles[ConstraintIndex]);
diff --git a/Engine/Source/Runtime/Experimental/Chaos/Private/PBDRigidsSolver.cpp b/Engine/Source/Runtime/Experimental/Chaos/Private/PBDRigidsSolver.cpp
--- a/Engine/Source/Runtime/Experimental/Chaos/Private/PBDRigidsSolver.cpp
+++ b/Engine/Source/Runtime/Experimental/Chaos/Private/PBDRigidsSolver.cpp
@@ -588,6 +588,8 @@ namespace Chaos
 		FJointConstraintPhysicsProxy* JointProxy = new FJointConstraintPhysicsProxy(GTConstraint, nullptr);
 		JointProxy->SetSolver(this);
 
+		JointConstraintPhysicsProxies_Internal.Add(JointProxy);
+
 		AddDirtyProxy(JointProxy);
 	}
 
@@ -1247,7 +1249,7 @@ namespace Chaos
 	// however.
 	void FPBDRigidsSolver::UpdateGameThreadStructures()
 	{
-		PullPhysicsStateForEachDirtyProxy_External([](auto){});
+		PullPhysicsStateForEachDirtyProxy_External([](auto){}, [](auto) {});
 	}
 
 	int32 FPBDRigidsSolver::NumJointConstraints() const
diff --git a/Engine/Source/Runtime/Experimental/Chaos/Private/PhysicsProxy/JointConstraintProxy.cpp b/Engine/Source/Runtime/Experimental/Chaos/Private/PhysicsProxy/JointConstraintProxy.cpp
--- a/Engine/Source/Runtime/Experimental/Chaos/Private/PhysicsProxy/JointConstraintProxy.cpp
+++ b/Engine/Source/Runtime/Experimental/Chaos/Private/PhysicsProxy/JointConstraintProxy.cpp
@@ -47,11 +47,14 @@ void FJointConstraintPhysicsProxy::BufferPhysicsResults(Chaos::FDirtyJointConstr
 	Buffer.SetProxy(*this);
 	if (Constraint != nullptr && Constraint->IsValid() )
 	{
-		if (Handle != nullptr && Handle->IsValid())
+		if (Handle != nullptr && (Handle->IsValid() || Handle->IsConstraintBreaking()))
 		{
+			Buffer.OutputData.bIsBreaking = Handle->IsConstraintBreaking();
 			Buffer.OutputData.bIsBroken = !Handle->IsConstraintEnabled();
 			Buffer.OutputData.Force = Handle->GetLinearImpulse();
 			Buffer.OutputData.Torque = Handle->GetAngularImpulse();
+
+			Handle->ClearConstraintBreaking(); // it's a single frame event, so reset
 		}
 	}
 }
@@ -61,8 +64,9 @@ bool FJointConstraintPhysicsProxy::PullFromPhysicsState(const Chaos::FDirtyJoint
 {
 	if (Constraint != nullptr && Constraint->IsValid())
 	{
-		if (Handle != nullptr && Handle->IsValid())
+		if (Handle != nullptr && (Handle->IsValid() || Buffer.OutputData.bIsBreaking))
 		{
+			Constraint->GetOutputData().bIsBreaking = Buffer.OutputData.bIsBreaking;
 			Constraint->GetOutputData().bIsBroken = Buffer.OutputData.bIsBroken;
 			Constraint->GetOutputData().Force = Buffer.OutputData.Force;
 			Constraint->GetOutputData().Torque = Buffer.OutputData.Torque;
diff --git a/Engine/Source/Runtime/Experimental/Chaos/Public/Chaos/Framework/PhysicsSolverBase.h b/Engine/Source/Runtime/Experimental/Chaos/Public/Chaos/Framework/PhysicsSolverBase.h
--- a/Engine/Source/Runtime/Experimental/Chaos/Public/Chaos/Framework/PhysicsSolverBase.h
+++ b/Engine/Source/Runtime/Experimental/Chaos/Public/Chaos/Framework/PhysicsSolverBase.h
@@ -324,8 +324,8 @@ namespace Chaos
 		}
 
 		/** Used to update external thread data structures. RigidFunc allows per dirty rigid code to execute. Include PhysicsSolverBaseImpl.h to call this function*/
-		template <typename RigidLambda>
-		void PullPhysicsStateForEachDirtyProxy_External(const RigidLambda& RigidFunc);
+		template <typename RigidLambda, typename ConstraintLambda>
+		void PullPhysicsStateForEachDirtyProxy_External(const RigidLambda& RigidFunc, const ConstraintLambda& ConstraintFunc);
 
 		bool IsUsingAsyncResults() const
 		{
diff --git a/Engine/Source/Runtime/Experimental/Chaos/Public/Chaos/PBDJointConstraintData.h b/Engine/Source/Runtime/Experimental/Chaos/Public/Chaos/PBDJointConstraintData.h
--- a/Engine/Source/Runtime/Experimental/Chaos/Public/Chaos/PBDJointConstraintData.h
+++ b/Engine/Source/Runtime/Experimental/Chaos/Public/Chaos/PBDJointConstraintData.h
@@ -298,6 +298,7 @@ namespace Chaos
 		struct FOutputData
 		{
 			// Output properties
+			bool bIsBreaking = false;
 			bool bIsBroken = false;
 			FVector Force = FVector(0);
 			FVector Torque = FVector(0);
diff --git a/Engine/Source/Runtime/Experimental/Chaos/Public/Chaos/PBDJointConstraints.h b/Engine/Source/Runtime/Experimental/Chaos/Public/Chaos/PBDJointConstraints.h
--- a/Engine/Source/Runtime/Experimental/Chaos/Public/Chaos/PBDJointConstraints.h
+++ b/Engine/Source/Runtime/Experimental/Chaos/Public/Chaos/PBDJointConstraints.h
@@ -40,6 +40,8 @@ namespace Chaos
 		int32 GetConstraintBatch() const;
 
 		bool IsConstraintEnabled() const;
+		bool IsConstraintBreaking() const;
+		void ClearConstraintBreaking();
 		FVec3 GetLinearImpulse() const;
 		FVec3 GetAngularImpulse() const;
 
@@ -64,6 +66,7 @@ namespace Chaos
 		int32 Color;
 		int32 IslandSize;
 		bool bDisabled;
+		bool bBreaking;
 		FVec3 LinearImpulse;
 		FVec3 AngularImpulse;
 	};
@@ -131,11 +134,26 @@ namespace Chaos
 		 */
 		bool IsConstraintEnabled(int32 ConstraintIndex) const;
 
+		/*
+		 * Whether the constraint is breaking this frame
+		 */
+		bool IsConstraintBreaking(int32 ConstraintIndex) const;
+
+		/*
+		 * Clear the constraint braking state
+		 */
+		void ClearConstraintBreaking(int32 ConstraintIndex);
+
 		/*
 		 * Enable or disable a constraints
 		 */
 		void SetConstraintEnabled(int32 ConstraintIndex, bool bEnabled);
 
+		/*
+		 * Set Breaking State
+		 */
+		void SetConstraintBreaking(int32 ConstraintIndex, bool bBreaking);
+
 		/*
 		 * Force a constraints to break
 		 */
diff --git a/Engine/Source/Runtime/Experimental/Chaos/Public/Chaos/PhysicsSolverBaseImpl.h b/Engine/Source/Runtime/Experimental/Chaos/Public/Chaos/PhysicsSolverBaseImpl.h
--- a/Engine/Source/Runtime/Experimental/Chaos/Public/Chaos/PhysicsSolverBaseImpl.h
+++ b/Engine/Source/Runtime/Experimental/Chaos/Public/Chaos/PhysicsSolverBaseImpl.h
@@ -16,8 +16,8 @@
 namespace Chaos
 {
 	// Pulls physics state for each dirty particle and allows caller to do additional work if needed
-	template <typename RigidLambda>
-	void FPhysicsSolverBase::PullPhysicsStateForEachDirtyProxy_External(const RigidLambda& RigidFunc)
+	template <typename RigidLambda, typename ConstraintLambda>
+	void FPhysicsSolverBase::PullPhysicsStateForEachDirtyProxy_External(const RigidLambda& RigidFunc, const ConstraintLambda& ConstraintFunc)
 	{
 		using namespace Chaos;
 
@@ -60,9 +60,9 @@ namespace Chaos
 			const int32 SolverTimestamp = Results.Next ? Results.Next->SolverTimestamp : INDEX_NONE;
 			for (const FChaosRigidInterpolationData& RigidInterp : Results.RigidInterpolations)
 				{
-				if (FSingleParticlePhysicsProxy* Proxy = RigidInterp.Prev.GetProxy())
+					if (FSingleParticlePhysicsProxy* Proxy = RigidInterp.Prev.GetProxy())
 					{
-					if (Proxy->PullFromPhysicsState(RigidInterp.Next, SolverTimestamp))
+						if (Proxy->PullFromPhysicsState(RigidInterp.Next, SolverTimestamp))
 						{
 							RigidFunc(Proxy);
 						}
@@ -98,7 +98,11 @@ namespace Chaos
 			{
 				if (auto Proxy = DirtyData.GetProxy())
 				{
-					Proxy->PullFromPhysicsState(DirtyData, SyncTimestamp);
+					if (Proxy->PullFromPhysicsState(DirtyData, SyncTimestamp))
+					{
+						ConstraintFunc(Proxy);
+					}
+
 				}
 			}
 
@@ -107,4 +111,4 @@ namespace Chaos
 		}
 	}
 
-}
\ No newline at end of file
+}
diff --git a/Engine/Source/Runtime/Experimental/Chaos/Public/Chaos/PullPhysicsDataImp.h b/Engine/Source/Runtime/Experimental/Chaos/Public/Chaos/PullPhysicsDataImp.h
--- a/Engine/Source/Runtime/Experimental/Chaos/Public/Chaos/PullPhysicsDataImp.h
+++ b/Engine/Source/Runtime/Experimental/Chaos/Public/Chaos/PullPhysicsDataImp.h
@@ -62,6 +62,7 @@ struct FDirtyGeometryCollectionData : public TBasePullData<FGeometryCollectionPh
 };
 
 struct FJointConstraintOutputData {
+	bool bIsBreaking = false;
 	bool bIsBroken = false;
 	FVector Force = FVector(0);
 	FVector Torque = FVector(0);
diff --git a/Engine/Source/Runtime/PhysicsCore/Private/ChaosScene.cpp b/Engine/Source/Runtime/PhysicsCore/Private/ChaosScene.cpp
--- a/Engine/Source/Runtime/PhysicsCore/Private/ChaosScene.cpp
+++ b/Engine/Source/Runtime/PhysicsCore/Private/ChaosScene.cpp
@@ -387,7 +387,7 @@ void FChaosScene::StartFrame()
 
 void FChaosScene::OnSyncBodies(Chaos::FPhysicsSolverBase* Solver)
 {
-	Solver->PullPhysicsStateForEachDirtyProxy_External([](auto){});
+	Solver->PullPhysicsStateForEachDirtyProxy_External([](auto){}, [](auto) {});
 }
 
 bool FChaosScene::AreAnyTasksPending() const
diff --git a/Engine/Source/Runtime/PhysicsCore/Public/Chaos/ChaosScene.h b/Engine/Source/Runtime/PhysicsCore/Public/Chaos/ChaosScene.h
--- a/Engine/Source/Runtime/PhysicsCore/Public/Chaos/ChaosScene.h
+++ b/Engine/Source/Runtime/PhysicsCore/Public/Chaos/ChaosScene.h
@@ -168,9 +168,6 @@ protected:
 	virtual void OnSyncBodies(Chaos::FPhysicsSolverBase* Solver);
 	//Engine interface END
 
-	template <typename RigidLambda>
-	void PullPhysicsStateForEachDirtyProxy(const int32 SyncTimestamp, const RigidLambda& DirtyRigidFunc);
-
 	float MDeltaTime;
 
 	UObject* Owner;
@@ -186,4 +183,4 @@ private:
 	FGraphEventArray CompletionEvents;
 
 	FChaosSceneSimCallback* SimCallback;
-};
\ No newline at end of file
+};
