CIS fix for Clang platforms

diff --git a/Engine/Source/Runtime/Engine/Classes/Animation/AnimExecutionContext.h b/Engine/Source/Runtime/Engine/Classes/Animation/AnimExecutionContext.h
--- a/Engine/Source/Runtime/Engine/Classes/Animation/AnimExecutionContext.h
+++ b/Engine/Source/Runtime/Engine/Classes/Animation/AnimExecutionContext.h
@@ -112,7 +112,7 @@ public:
 protected:
 	// Access internal context. Will return nullptr if invalid or an incorrect type is requested
 	template<typename OtherContextType, typename InternalContextType>
-	typename InternalContextType* GetInternalContext() const
+	InternalContextType* GetInternalContext() const
 	{
 		if(TSharedPtr<FData> PinnedData = Data.Pin())
 		{
