Visual Logger: Added circle shape

diff --git a/Engine/Source/Developer/LogVisualizer/Private/VisualLoggerRenderingActor.cpp b/Engine/Source/Developer/LogVisualizer/Private/VisualLoggerRenderingActor.cpp
--- a/Engine/Source/Developer/LogVisualizer/Private/VisualLoggerRenderingActor.cpp
+++ b/Engine/Source/Developer/LogVisualizer/Private/VisualLoggerRenderingActor.cpp
@@ -679,6 +679,45 @@ void AVisualLoggerRenderingActor::GetDebugShapes(const FVisualLogDevice::FVisual
 			}
 		}
 			break;
+
+		case EVisualLoggerShapeElement::Circle:
+		{
+			const bool bDrawLabel = (ElementToDraw->Description.IsEmpty() == false);
+			const int32 NumPoints = ElementToDraw->Points.Num();
+
+			for (int32 Index = 0; Index + 2 < NumPoints; Index += 3)
+			{
+				const FVector Center = ElementToDraw->Points[Index + 0];
+				const FVector UpAxis = ElementToDraw->Points[Index + 1];
+				const float Radius = ElementToDraw->Points[Index + 2].X;
+				const float Thickness = float(ElementToDraw->Thicknes);
+
+				const FQuat Rotation = FQuat::FindBetweenNormals(FVector::UpVector, UpAxis);
+				const FVector XAxis = Rotation.RotateVector(FVector::XAxisVector);
+				const FVector YAxis = Rotation.RotateVector(FVector::YAxisVector);
+				
+				static const int32 CircleDivs = 12;
+				FVector PrevPosition = FVector::ZeroVector;
+				for (int32 Div = 0; Div <= CircleDivs; Div++)
+				{
+					const float Angle = (float)Div / (float)CircleDivs * PI * 2.0f;
+					const FVector Position = Center + (FMath::Cos(Angle) * XAxis + FMath::Sin(Angle) * YAxis) * Radius;
+					if (Div > 0)
+					{
+						DebugShapes.Lines.Add(FDebugRenderSceneProxy::FDebugLine(PrevPosition, Position, Color, Thickness));
+					}
+					PrevPosition = Position;
+				}
+
+				if (bDrawLabel)
+				{
+					const FString PrintString = NumPoints == 3 ? ElementToDraw->Description : FString::Printf(TEXT("%s_%d"), *ElementToDraw->Description, Index / 3);
+					DebugShapes.Texts.Add(FDebugRenderSceneProxy::FText3d(PrintString, Center, Color));
+				}
+			}
+		}
+			break;
+			
 		}
 	}
 }
diff --git a/Engine/Source/Runtime/Engine/Private/VisualLogger/VisualLoggerTypes.cpp b/Engine/Source/Runtime/Engine/Private/VisualLogger/VisualLoggerTypes.cpp
--- a/Engine/Source/Runtime/Engine/Private/VisualLogger/VisualLoggerTypes.cpp
+++ b/Engine/Source/Runtime/Engine/Private/VisualLogger/VisualLoggerTypes.cpp
@@ -198,6 +198,21 @@ void FVisualLogEntry::AddArrow(const FVector& Start, const FVector& End, const F
 	ElementsToDraw.Add(Element);
 }
 
+void FVisualLogEntry::AddCircle(const FVector& Center, const FVector& UpAxis, const float Radius, const FName& CategoryName, ELogVerbosity::Type Verbosity, const FColor& Color, const FString& Description, const uint16 Thickness)
+{
+	FVisualLogShapeElement Element(EVisualLoggerShapeElement::Circle);
+	Element.Category = CategoryName;
+	Element.SetColor(Color);
+	Element.Thicknes = Thickness;
+	Element.Description = Description;
+	Element.Points.Reserve(3);
+	Element.Points.Add(Center);
+	Element.Points.Add(UpAxis);
+	Element.Points.Add(FVector(Radius, 0.0f, 0.0f));
+	Element.Verbosity = Verbosity;
+	ElementsToDraw.Add(Element);
+}
+
 void FVisualLogEntry::AddElement(const FBox& Box, const FMatrix& Matrix, const FName& CategoryName, ELogVerbosity::Type Verbosity, const FColor& Color, const FString& Description, uint16 Thickness)
 {
 	FVisualLogShapeElement Element(Description, Color, Thickness, CategoryName);
diff --git a/Engine/Source/Runtime/Engine/Public/VisualLogger/VisualLogger.h b/Engine/Source/Runtime/Engine/Public/VisualLogger/VisualLogger.h
--- a/Engine/Source/Runtime/Engine/Public/VisualLogger/VisualLogger.h
+++ b/Engine/Source/Runtime/Engine/Public/VisualLogger/VisualLogger.h
@@ -65,6 +65,12 @@
 // Segment with an arrowhead
 #define UE_VLOG_ARROW(LogOwner, CategoryName, Verbosity, SegmentStart, SegmentEnd, Color, Format, ...) if(FVisualLogger::IsRecording()) FVisualLogger::ArrowLogf(LogOwner, CategoryName, ELogVerbosity::Verbosity, SegmentStart, SegmentEnd, Color, Format, ##__VA_ARGS__)
 #define UE_CVLOG_ARROW(Condition, LogOwner, CategoryName, Verbosity, SegmentStart, SegmentEnd, Color, Format, ...) if(FVisualLogger::IsRecording() && Condition) {UE_VLOG_ARROW(LogOwner, CategoryName, Verbosity, SegmentStart, SegmentEnd, Color, Format, ##__VA_ARGS__);} 
+// Circle shape
+#define UE_VLOG_CIRCLE(LogOwner, CategoryName, Verbosity, Center, UpAxis, Radius, Color, Format, ...) if(FVisualLogger::IsRecording()) FVisualLogger::CircleLogf(LogOwner, CategoryName, ELogVerbosity::Verbosity, Center, UpAxis, Radius, Color, 0, Format, ##__VA_ARGS__)
+#define UE_CVLOG_CIRCLE(Condition, LogOwner, CategoryName, Verbosity, Center, UpAxis, Radius, Color, Format, ...)  if(FVisualLogger::IsRecording() && Condition) {UE_VLOG_CIRCLE(LogOwner, CategoryName, Verbosity, Center, UpAxis, Radius, Color, 0, Format, ##__VA_ARGS__);} 
+// Circle shape
+#define UE_VLOG_CIRCLE_THICK(LogOwner, CategoryName, Verbosity, Center, UpAxis, Radius, Color, Thickness, Format, ...) if(FVisualLogger::IsRecording()) FVisualLogger::CircleLogf(LogOwner, CategoryName, ELogVerbosity::Verbosity, Center, UpAxis, Radius, Color, Thickness, Format, ##__VA_ARGS__)
+#define UE_CVLOG_CIRCLE_THICK(Condition, LogOwner, CategoryName, Verbosity, Center, UpAxis, Radius, Color, Thickness, Format, ...)  if(FVisualLogger::IsRecording() && Condition) {UE_VLOG_CIRCLE(LogOwner, CategoryName, Verbosity, Center, UpAxis, Radius, Color, Thickness, Format, ##__VA_ARGS__);} 
 
 #define DECLARE_VLOG_EVENT(EventName) extern FVisualLogEventBase EventName;
 #define DEFINE_VLOG_EVENT(EventName, Verbosity, UserFriendlyDesc) FVisualLogEventBase EventName(TEXT(#EventName), TEXT(UserFriendlyDesc), ELogVerbosity::Verbosity); 
@@ -172,6 +178,8 @@ class ENGINE_API FVisualLogger : public FOutputDevice
 	static void HistogramDataLogfImpl (const UObject* LogOwner, const FName& CategoryName, ELogVerbosity::Type Verbosity, FName GraphName, FName DataName, const FVector2D& Data, const FColor& Color, const TCHAR* Fmt, ...);
 	static void ArrowLogfImpl(const UObject* LogOwner, const FLogCategoryBase& Category, ELogVerbosity::Type Verbosity, const FVector& Start, const FVector& End, const FColor& Color, const TCHAR* Fmt, ...);
 	static void ArrowLogfImpl(const UObject* LogOwner, const FName& CategoryName, ELogVerbosity::Type Verbosity, const FVector& Start, const FVector& End, const FColor& Color, const TCHAR* Fmt, ...);
+	static void CircleLogfImpl(const UObject* LogOwner, const FLogCategoryBase& Category, ELogVerbosity::Type Verbosity, const FVector& Center, const FVector& UpAxis, const float Radius, const FColor& Color, const uint16 Thickness, const TCHAR* Fmt, ...);
+	static void CircleLogfImpl(const UObject* LogOwner, const FName& CategoryName, ELogVerbosity::Type Verbosity, const FVector& Center, const FVector& UpAxis, const float Radius, const FColor& Color, const uint16 Thickness, const TCHAR* Fmt, ...);
 
 public:
 	// Regular text log
@@ -227,6 +235,25 @@ public:
 		ArrowLogfImpl(LogOwner, CategoryName, Verbosity, Start, End, Color, Fmt, Args...);
 	}
 
+
+	// Circle log
+	template <typename FmtType, typename... Types>
+	static void CircleLogf(const UObject* LogOwner, const FLogCategoryBase& Category, ELogVerbosity::Type Verbosity, const FVector& Center, const FVector& UpAxis, const float Radius, const FColor& Color, const uint16 Thickness, const FmtType& Fmt, Types... Args)
+	{
+		static_assert(TIsArrayOrRefOfType<FmtType, TCHAR>::Value, "Formatting string must be a TCHAR array.");
+		static_assert(TAnd<TIsValidVariadicFunctionArg<Types>...>::Value, "Invalid argument(s) passed to FVisualLogger::CircleLogf");
+
+		CircleLogfImpl(LogOwner, Category, Verbosity, Center, UpAxis, Radius, Color, Thickness, Fmt, Args...);
+	}
+	template <typename FmtType, typename... Types>
+	static void CircleLogf(const UObject* LogOwner, const FName& CategoryName, ELogVerbosity::Type Verbosity, const FVector& Center, const FVector& UpAxis, const float Radius, const FColor& Color, const uint16 Thickness, const FmtType& Fmt, Types... Args)
+	{
+		static_assert(TIsArrayOrRefOfType<FmtType, TCHAR>::Value, "Formatting string must be a TCHAR array.");
+		static_assert(TAnd<TIsValidVariadicFunctionArg<Types>...>::Value, "Invalid argument(s) passed to FVisualLogger::CircleLogf");
+
+		CircleLogfImpl(LogOwner, CategoryName, Verbosity, Center, UpAxis, Radius, Color, Thickness, Fmt, Args...);
+	}
+
 	// Location/Sphere log
 	template <typename FmtType, typename... Types>
 	static void GeometryShapeLogf(const UObject* LogOwner, const FLogCategoryBase& Category, ELogVerbosity::Type Verbosity, const FVector& Location, float Radius, const FColor& Color, const FmtType& Fmt, Types... Args)
@@ -653,6 +680,20 @@ inline void FVisualLogger::ArrowLogfImpl(const UObject* Object, const FName& Cat
 	);
 }
 
+inline void FVisualLogger::CircleLogfImpl(const UObject* Object, const FLogCategoryBase& Category, ELogVerbosity::Type Verbosity, const FVector& Center, const FVector& UpAxis, const float Radius, const FColor& Color, const uint16 Thickness, const TCHAR* Fmt, ...)
+{
+	const FName CategoryName = Category.GetCategoryName();
+	COLLAPSED_LOGF(
+		CurrentEntry->AddCircle(Center, UpAxis, Radius, CategoryName, Verbosity, Color, Buffer, Thickness);
+	);
+}
+inline void FVisualLogger::CircleLogfImpl(const UObject* Object, const FName& CategoryName, ELogVerbosity::Type Verbosity, const FVector& Center, const FVector& UpAxis, const float Radius, const FColor& Color, const uint16 Thickness, const TCHAR* Fmt, ...)
+{
+	COLLAPSED_LOGF(
+		CurrentEntry->AddCircle(Center, UpAxis, Radius, CategoryName, Verbosity, Color, Buffer, Thickness);
+	);
+}
+
 inline void FVisualLogger::GeometryShapeLogfImpl(const UObject* Object, const FLogCategoryBase& Category, ELogVerbosity::Type Verbosity, const FVector& Location, float Radius, const FColor& Color, const TCHAR* Fmt, ...)
 {
 	const FName CategoryName = Category.GetCategoryName(); 
diff --git a/Engine/Source/Runtime/Engine/Public/VisualLogger/VisualLoggerTypes.h b/Engine/Source/Runtime/Engine/Public/VisualLogger/VisualLoggerTypes.h
--- a/Engine/Source/Runtime/Engine/Public/VisualLogger/VisualLoggerTypes.h
+++ b/Engine/Source/Runtime/Engine/Public/VisualLogger/VisualLoggerTypes.h
@@ -48,7 +48,8 @@ enum class EVisualLoggerShapeElement : uint8
 	Polygon,
 	Mesh,
 	NavAreaMesh, // convex based mesh with min and max Z values
-	Arrow, 
+	Arrow,
+	Circle,
 	// note that in order to remain backward compatibility in terms of log
 	// serialization new enum values need to be added at the end
 };
@@ -210,6 +211,8 @@ struct ENGINE_API FVisualLogEntry
 	void AddArrow(const FVector& Start, const FVector& End, const FName& CategoryName, ELogVerbosity::Type Verbosity, const FColor& Color = FColor::White, const FString& Description = TEXT(""));
 	// boxes
 	void AddBoxes(const TArray<FBox>& Boxes, const FName& CategoryName, ELogVerbosity::Type Verbosity, const FColor& Color = FColor::White);
+	// circle
+	void AddCircle(const FVector& Center, const FVector& UpAxis, const float Radius, const FName& CategoryName, ELogVerbosity::Type Verbosity, const FColor& Color, const FString& Description = TEXT(""), uint16 Thickness = 0);
 
 	// Custom data block
 	FVisualLogDataBlock& AddDataBlock(const FString& TagName, const TArray<uint8>& BlobDataArray, const FName& CategoryName, ELogVerbosity::Type Verbosity);
