Fix for root bone being disabled when SetAllBodiesBelowPhysicsDisabled is called using a bone name that doesn't exist in the skeleton.

diff --git a/Engine/Source/Runtime/Engine/Private/PhysicsEngine/PhysicsAsset.cpp b/Engine/Source/Runtime/Engine/Private/PhysicsEngine/PhysicsAsset.cpp
--- a/Engine/Source/Runtime/Engine/Private/PhysicsEngine/PhysicsAsset.cpp
+++ b/Engine/Source/Runtime/Engine/Private/PhysicsEngine/PhysicsAsset.cpp
@@ -530,20 +530,23 @@ void UPhysicsAsset::GetBodyIndicesBelow(TArray<int32>& OutBodyIndices, FName InB
 {
 	int32 BaseIndex = SkelMesh->GetRefSkeleton().FindBoneIndex(InBoneName);
 
-	// Iterate over all other bodies, looking for 'children' of this one
-	for(int32 i=0; i<SkeletalBodySetups.Num(); i++)
+	if (BaseIndex != INDEX_NONE)
 	{
-		UBodySetup* BS = SkeletalBodySetups[i];
-		if (!ensure(BS))
+		// Iterate over all other bodies, looking for 'children' of this one
+		for (int32 i = 0; i < SkeletalBodySetups.Num(); i++)
 		{
-			continue;
-		}
-		FName TestName = BS->BoneName;
-		int32 TestIndex = SkelMesh->GetRefSkeleton().FindBoneIndex(TestName);
+			UBodySetup* BS = SkeletalBodySetups[i];
+			if (!ensure(BS))
+			{
+				continue;
+			}
+			FName TestName = BS->BoneName;
+			int32 TestIndex = SkelMesh->GetRefSkeleton().FindBoneIndex(TestName);
 
-		if( (bIncludeParent && TestIndex == BaseIndex) || SkelMesh->GetRefSkeleton().BoneIsChildOf(TestIndex, BaseIndex))
-		{
-			OutBodyIndices.Add(i);
+			if ((bIncludeParent && TestIndex == BaseIndex) || SkelMesh->GetRefSkeleton().BoneIsChildOf(TestIndex, BaseIndex))
+			{
+				OutBodyIndices.Add(i);
+			}
 		}
 	}
 }
diff --git a/Engine/Source/Runtime/Engine/Public/ReferenceSkeleton.h b/Engine/Source/Runtime/Engine/Public/ReferenceSkeleton.h
--- a/Engine/Source/Runtime/Engine/Public/ReferenceSkeleton.h
+++ b/Engine/Source/Runtime/Engine/Public/ReferenceSkeleton.h
@@ -349,20 +349,23 @@ public:
 
 	bool BoneIsChildOf(const int32 ChildBoneIndex, const int32 ParentBoneIndex) const
 	{
-		// Bones are in strictly increasing order.
-		// So child must have an index greater than his parent.
-		if( ChildBoneIndex > ParentBoneIndex )
+		if (ParentBoneIndex != INDEX_NONE)
 		{
-			int32 BoneIndex = GetParentIndex(ChildBoneIndex);
-			do
+			// Bones are in strictly increasing order.
+			// So child must have an index greater than his parent.
+			if (ChildBoneIndex > ParentBoneIndex)
 			{
-				if( BoneIndex == ParentBoneIndex )
+				int32 BoneIndex = GetParentIndex(ChildBoneIndex);
+				do
 				{
-					return true;
-				}
-				BoneIndex = GetParentIndex(BoneIndex);
+					if (BoneIndex == ParentBoneIndex)
+					{
+						return true;
+					}
+					BoneIndex = GetParentIndex(BoneIndex);
 
-			} while (BoneIndex != INDEX_NONE);
+				} while (BoneIndex != INDEX_NONE);
+			}
 		}
 
 		return false;
