Fixed reported typo in volumetriccloudcomponent.

diff --git a/Engine/Source/Runtime/Engine/Classes/Components/VolumetricCloudComponent.h b/Engine/Source/Runtime/Engine/Classes/Components/VolumetricCloudComponent.h
--- a/Engine/Source/Runtime/Engine/Classes/Components/VolumetricCloudComponent.h
+++ b/Engine/Source/Runtime/Engine/Classes/Components/VolumetricCloudComponent.h
@@ -99,7 +99,7 @@ class UVolumetricCloudComponent : public USceneComponent
 	float ShadowTracingDistance;
 
 	/**
-	 * When the mean transmittance is below this threashold, we stop tracing. This is an good way to reduce the ray marched sample count, and thus to increase performance.
+	 * When the mean transmittance is below this threshold, we stop tracing. This is a good way to reduce the ray marched sample count, and thus to increase performance.
 	 */
 	UPROPERTY(EditAnywhere, BlueprintReadOnly, Category = "Cloud Tracing", AdvancedDisplay, meta = (UIMin = "0.0", UIMax = "1.0", ClampMin = "0.0", ClampMax = "1.0", SliderExponent = 5.0))
 	float StopTracingTransmittanceThreshold;
