LISTSTATICMESHES
- now outputs and Instance count column
- new '-game' filter to allow filtering out non game world components/instances
- allows sorting per instance count '-instancesort'

diff --git a/Engine/Source/Runtime/Engine/Private/UnrealEngine.cpp b/Engine/Source/Runtime/Engine/Private/UnrealEngine.cpp
--- a/Engine/Source/Runtime/Engine/Private/UnrealEngine.cpp
+++ b/Engine/Source/Runtime/Engine/Private/UnrealEngine.cpp
@@ -251,6 +251,8 @@ UnrealEngine.cpp: Implements the UEngine class and helpers.
 
 #include "Particles/ParticlePerfStatsManager.h"
 
+#include "Engine/InstancedStaticMesh.h"
+
 DEFINE_LOG_CATEGORY(LogEngine);
 IMPLEMENT_MODULE( FEngineModule, Engine );
 
@@ -3724,14 +3726,15 @@ struct FSortedStaticMesh
 	int32			VertexCountTotalMobile;
 	int32			VertexCountCollision;
 	int32			ShapeCountCollision;
-	int32			UsageCount;
+	int32			ComponentUsageCount;
+	int32			InstanceUsageCount;
 	int32			CPUAccessOverheadKB;
 	bool			bUnknownRef;
 	UStaticMesh*	Mesh;
 	FString			Name;
 
 	/** Constructor, initializing every member variable with passed in values. */
-	FSortedStaticMesh(UStaticMesh* InMesh, int32 InNumKB, int32 InMaxKB, int32 InResKBExc, int32 InResKBInc, int32 InResKBIncMobile, int32 InResKBResident, int32 InDistanceFieldKB, int32 InLodCount, int32 InResidentLodCount, int32 InMobileMinLOD, int32 InVertexCountLod0, int32 InVertexCountLod1, int32 InVertexCountLod2, int32 InVertexCountTotal, int32 InVertexCountTotalMobile, int32 InVertexCountCollision, int32 InShapeCountCollision, int32 InUsageCount, int32 InCPUAccessOverheadKB, bool bInUnknownRef, FString InName)
+	FSortedStaticMesh(UStaticMesh* InMesh, int32 InNumKB, int32 InMaxKB, int32 InResKBExc, int32 InResKBInc, int32 InResKBIncMobile, int32 InResKBResident, int32 InDistanceFieldKB, int32 InLodCount, int32 InResidentLodCount, int32 InMobileMinLOD, int32 InVertexCountLod0, int32 InVertexCountLod1, int32 InVertexCountLod2, int32 InVertexCountTotal, int32 InVertexCountTotalMobile, int32 InVertexCountCollision, int32 InShapeCountCollision, int32 InComponentUsageCount, int32 InInstanceUsageCount, int32 InCPUAccessOverheadKB, bool bInUnknownRef, FString InName)
 		: NumKB(InNumKB)
 		, MaxKB(InMaxKB)
 		, ResKBExc(InResKBExc)
@@ -3749,7 +3752,8 @@ struct FSortedStaticMesh
 		, VertexCountTotalMobile(InVertexCountTotalMobile)
 		, VertexCountCollision(InVertexCountCollision)
 		, ShapeCountCollision(InShapeCountCollision)
-		, UsageCount(InUsageCount)
+		, ComponentUsageCount(InComponentUsageCount)
+		, InstanceUsageCount(InInstanceUsageCount)
 		, CPUAccessOverheadKB(InCPUAccessOverheadKB)
 		, bUnknownRef(bInUnknownRef)
 		, Mesh(InMesh)		
@@ -3758,25 +3762,41 @@ struct FSortedStaticMesh
 };
 struct FCompareFSortedStaticMesh
 {
-	bool bAlphaSort;
-	bool bMobileSort;
-	FCompareFSortedStaticMesh(bool InAlphaSort, bool InMobileSort)
-		: bAlphaSort(InAlphaSort)
-		, bMobileSort(InMobileSort)
+	enum ESortType
+	{
+		Default,
+		InstanceCount,
+		Mobile,
+		Alpha
+	} SortBy;
+
+	FCompareFSortedStaticMesh(ESortType InSortBy)
+		: SortBy(InSortBy)
 	{}
 	FORCEINLINE bool operator()(const FSortedStaticMesh& A, const FSortedStaticMesh& B) const
 	{
-		if (bMobileSort)
+		switch (SortBy)
 		{
-			return ((B.VertexCountTotalMobile) < (A.VertexCountTotalMobile));
-		}
+		case InstanceCount:
+			return B.InstanceUsageCount < A.InstanceUsageCount;
+			
+		case Mobile:
+			return B.VertexCountTotalMobile < A.VertexCountTotalMobile;
 
-		if (bAlphaSort || (B.NumKB + B.ResKBExc) == (A.NumKB + A.ResKBExc))
-		{
+		case Alpha:
 			return A.Name < B.Name;
-		}
+		
+		case Default:
+		default:
+			{
+				if ((B.NumKB + B.ResKBExc) == (A.NumKB + A.ResKBExc))
+				{
+					return A.Name < B.Name; // default to alpha
+				}
 
-		return (B.NumKB + B.ResKBExc) < (A.NumKB + A.ResKBExc);
+				return (B.NumKB + B.ResKBExc) < (A.NumKB + A.ResKBExc);
+			}
+		}
 	}
 };
 
@@ -5618,12 +5638,16 @@ bool UEngine::HandleListTexturesCommand( const TCHAR* Cmd, FOutputDevice& Ar )
 
 bool UEngine::HandleListStaticMeshesCommand(const TCHAR* Cmd, FOutputDevice& Ar)
 {
-	const bool bAlphaSort = FParse::Param(Cmd, TEXT("ALPHASORT"));
-	const bool bMobileSort = FParse::Param(Cmd, TEXT("MOBILESORT"));
+	FCompareFSortedStaticMesh::ESortType SortType = FCompareFSortedStaticMesh::ESortType::Default;
+	SortType = FParse::Param(Cmd, TEXT("ALPHASORT"))? FCompareFSortedStaticMesh::ESortType::Alpha : SortType;
+	SortType = FParse::Param(Cmd, TEXT("MOBILESORT")) ? FCompareFSortedStaticMesh::ESortType::Mobile : SortType;
+	SortType = FParse::Param(Cmd, TEXT("INSTANCESORT")) ? FCompareFSortedStaticMesh::ESortType::InstanceCount : SortType;
+
 	const bool bUsedComponents = FParse::Param(Cmd, TEXT("usedcomponents"));
 	const bool bUnknownRefOnly = FParse::Param(Cmd, TEXT("unknownrefonly"));
+	const bool bGameWorld = FParse::Param(Cmd, TEXT("game"));
 	const bool bUsage = FParse::Param(Cmd, TEXT("-?"));
-
+	
 
 
 	//non-editor builds literally don't have the data to determine mobile lods or vert data.  The data prints out incorrectly and is confusing, just remove it.
@@ -5635,6 +5659,8 @@ bool UEngine::HandleListStaticMeshesCommand(const TCHAR* Cmd, FOutputDevice& Ar)
 			"\n Optional params: \n"
 			"-alphasort: sort alphabetically \n"
 			"-mobilesort: sort by mobile verts \n"
+			"-instancesort: sort by instance count \n"
+			"-game: only consider game world \n"
 			"-usedcomponents: print the all components used by each mesh\n"
 			"-unknownrefonly: list only static meshes using unknown ref heuristic for LOD streaming"));
 	}
@@ -5648,13 +5674,29 @@ bool UEngine::HandleListStaticMeshesCommand(const TCHAR* Cmd, FOutputDevice& Ar)
 
 	//Collect usage counts
 	TMap<UStaticMesh*, TArray<UStaticMeshComponent*>> UsageList;
+	TMap<UStaticMesh*, int32> InstanceCounts;
 	for (TObjectIterator<UStaticMeshComponent> It; It; ++It)
 	{
 		UStaticMeshComponent* MeshComponent = *It;
+		if (bGameWorld && (!MeshComponent->GetWorld() || !MeshComponent->GetWorld()->IsGameWorld()))
+		{
+			continue;
+		}
+		
 		UStaticMesh* Mesh = MeshComponent->GetStaticMesh();
 
 		TArray<UStaticMeshComponent*>& MeshUsageArray = UsageList.FindOrAdd(Mesh);
 		MeshUsageArray.Add(MeshComponent);
+
+		int32& InstanceCount = InstanceCounts.FindOrAdd(Mesh, 0);
+		if (UInstancedStaticMeshComponent* ISM = Cast<UInstancedStaticMeshComponent>(MeshComponent))
+		{
+			InstanceCount += ISM->GetInstanceCount();
+		}
+		else
+		{
+			InstanceCount++;
+		}
 	}
 
 	// Collect meshes.
@@ -5679,7 +5721,15 @@ bool UEngine::HandleListStaticMeshesCommand(const TCHAR* Cmd, FOutputDevice& Ar)
 		}
 
 		const TArray<UStaticMeshComponent*>* MeshUsageList = UsageList.Find(Mesh);
-		int32 UsageCount = MeshUsageList ? MeshUsageList->Num() : 0;
+		int32 ComponentUsageCount = MeshUsageList ? MeshUsageList->Num() : 0;
+
+		int32* InstanceCountPtr = InstanceCounts.Find(Mesh);
+		int32 InstanceUsageCount = InstanceCountPtr ? *InstanceCountPtr : 0;
+
+		if (bGameWorld && !InstanceUsageCount)
+		{
+			continue;
+		}
 
 		FArchiveCountMem Count(Mesh);
 		FResourceSizeEx ResourceSizeExc = FResourceSizeEx(EResourceSizeMode::Exclusive);
@@ -5784,14 +5834,15 @@ bool UEngine::HandleListStaticMeshesCommand(const TCHAR* Cmd, FOutputDevice& Ar)
 			VertexCountTotalMobile,
 			VertexCountCollision,
 			CollisionShapeCount,
-			UsageCount,
+			ComponentUsageCount,
+			InstanceUsageCount,
 			CPUAccessOverheadKB,
 			bUnknownRef,
 			Name);
 	}
 
 	// Sort meshes by cost or alphabetically.
-	SortedMeshes.Sort(FCompareFSortedStaticMesh(bAlphaSort, bMobileSort));
+	SortedMeshes.Sort(FCompareFSortedStaticMesh(SortType));
 
 	// Display.
 	int64 TotalNumKB = 0;
@@ -5804,8 +5855,10 @@ bool UEngine::HandleListStaticMeshesCommand(const TCHAR* Cmd, FOutputDevice& Ar)
 	int32 TotalVertexCount = 0;
 	int32 TotalVertexCountMobile = 0;
 	int64 TotalCPUAccessOverheadKB = 0;
+	int64 TotalComponents = 0;
+	int64 TotalInstances = 0;
 
-	FString HeaderString(TEXT(",       NumKB,       MaxKB,    ResKBExc,    ResKBInc, ResKBResident, DistFieldKB,  LODCount, ResidentLODCount, VertsLOD0, VertsLOD1,  VertsLOD2, Verts Total,  Verts Coll, Coll Shapes,     NumUsed, CPUAccessOverheadKB, UnknownRef"));
+	FString HeaderString(TEXT(",       NumKB,       MaxKB,    ResKBExc,    ResKBInc,   ResKBResi, DistFieldKB,    LODCount,ResiLODCount,   VertsLOD0,   VertsLOD1,   VertsLOD2, Verts Total,  Verts Coll, Coll Shapes,    NumComps,     NumInst,      Nanite, CPUAccessOverheadKB,  UnknownRef"));
 	FString MobileHeaderString = bHasMobileColumns ? FString(TEXT(", ResKBIncMob,Verts Mobile,MobileMinLOD")) : FString();
 	
 	Ar.Logf(TEXT("%s%s, Name"), *HeaderString, *MobileHeaderString);	
@@ -5816,7 +5869,7 @@ bool UEngine::HandleListStaticMeshesCommand(const TCHAR* Cmd, FOutputDevice& Ar)
 
 		if (bHasMobileColumns)
 		{
-			Ar.Logf(TEXT(", %11i, %11i, %11i, %11i, %11i, %11i, %11i, %11i, %11i, %11i, %11i, %11i, %11i, %11i, %11i, %11i, %s, %11i, %11i, %11i, %s"),
+			Ar.Logf(TEXT(", %11i, %11i, %11i, %11i, %11i, %11i, %11i, %11i, %11i, %11i, %11i, %11i, %11i, %11i, %11i, %11i, %s, %19i, %s, %11i, %11i, %11i, %s"),
 				SortedMesh.NumKB,
 				SortedMesh.MaxKB,
 				SortedMesh.ResKBExc,
@@ -5831,9 +5884,11 @@ bool UEngine::HandleListStaticMeshesCommand(const TCHAR* Cmd, FOutputDevice& Ar)
 				SortedMesh.VertexCountTotal,
 				SortedMesh.VertexCountCollision,
 				SortedMesh.ShapeCountCollision,
-				SortedMesh.UsageCount,
+				SortedMesh.ComponentUsageCount,
+				SortedMesh.InstanceUsageCount,
+				SortedMesh.Mesh->HasValidNaniteData() ? TEXT("        Yes") : TEXT("         No"),
 				SortedMesh.CPUAccessOverheadKB,
-				SortedMesh.bUnknownRef ? TEXT("Yes") : TEXT("No"),
+				SortedMesh.bUnknownRef ? TEXT("        Yes") : TEXT("         No"),
 				SortedMesh.ResKBIncMobile,
 				SortedMesh.VertexCountTotalMobile,
 				SortedMesh.MobileMinLOD,
@@ -5841,7 +5896,7 @@ bool UEngine::HandleListStaticMeshesCommand(const TCHAR* Cmd, FOutputDevice& Ar)
 		}
 		else
 		{
-			Ar.Logf(TEXT(", %11i, %11i, %11i, %11i, %11i, %11i, %11i, %11i, %11i, %11i, %11i, %11i, %11i, %11i, %11i, %11i, %s, %s"),
+			Ar.Logf(TEXT(", %11i, %11i, %11i, %11i, %11i, %11i, %11i, %11i, %11i, %11i, %11i, %11i, %11i, %11i, %11i, %11i, %s, %19i, %s, %s"),
 				SortedMesh.NumKB,
 				SortedMesh.MaxKB,
 				SortedMesh.ResKBExc,
@@ -5856,9 +5911,11 @@ bool UEngine::HandleListStaticMeshesCommand(const TCHAR* Cmd, FOutputDevice& Ar)
 				SortedMesh.VertexCountTotal,
 				SortedMesh.VertexCountCollision,
 				SortedMesh.ShapeCountCollision,
-				SortedMesh.UsageCount,
+				SortedMesh.ComponentUsageCount,
+				SortedMesh.InstanceUsageCount,
+				SortedMesh.Mesh->HasValidNaniteData() ? TEXT("        Yes") : TEXT("         No"),
 				SortedMesh.CPUAccessOverheadKB,
-				SortedMesh.bUnknownRef ? TEXT("Yes") : TEXT("No"),
+				SortedMesh.bUnknownRef ? TEXT("        Yes") : TEXT("         No"),
 				*SortedMesh.Name);
 		}
 		
@@ -5872,9 +5929,11 @@ bool UEngine::HandleListStaticMeshesCommand(const TCHAR* Cmd, FOutputDevice& Ar)
 		TotalVertexCount += SortedMesh.VertexCountTotal;
 		TotalVertexCountMobile += SortedMesh.VertexCountTotalMobile;
 		TotalCPUAccessOverheadKB += SortedMesh.CPUAccessOverheadKB;
+		TotalComponents += SortedMesh.ComponentUsageCount;
+		TotalInstances += SortedMesh.InstanceUsageCount;
 	}
 
-	Ar.Logf(TEXT("Total NumKB: %lld KB, Total MaxKB: %lld KB, Total ResKB Exc: %lld KB, Total ResKB Inc %lld KB,  Total ResKB Inc Mobile %lld KB, Total ResKB Resident %lld KB, Total Distance Field %lld KB, Total CPUAccess Overhead %lld KB, Total Vertex Count: %i, Total Vertex Count Mobile: %i, Static Mesh Count=%d"), TotalNumKB, TotalMaxKB, TotalResKBExc, TotalResKBInc, TotalResKBIncMobile, TotalResKBResident, TotalDistanceFieldKB, TotalCPUAccessOverheadKB, TotalVertexCount, TotalVertexCountMobile, SortedMeshes.Num());
+	Ar.Logf(TEXT("Total NumKB: %lld KB, Total MaxKB: %lld KB, Total ResKB Exc: %lld KB, Total ResKB Inc %lld KB,  Total ResKB Inc Mobile %lld KB, Total ResKB Resident %lld KB, Total Distance Field %lld KB, Total CPUAccess Overhead %lld KB, Total Vertex Count: %i, Total Vertex Count Mobile: %i, Component Count: %lld, Instance Count: %lld, Static Mesh Count=%d"), TotalNumKB, TotalMaxKB, TotalResKBExc, TotalResKBInc, TotalResKBIncMobile, TotalResKBResident, TotalDistanceFieldKB, TotalCPUAccessOverheadKB, TotalVertexCount, TotalVertexCountMobile, TotalComponents, TotalInstances, SortedMeshes.Num());
 
 	if (bUsedComponents)
 	{
