Texture: Removed the texture build exporter

Texture builds can be exported (without reference outputs) by running with -ExportBuilds=Texture

#rnx

diff --git a/Engine/Source/Runtime/Engine/Private/TextureDerivedDataBuildExporter.cpp b/Engine/Source/Runtime/Engine/Private/TextureDerivedDataBuildExporter.cpp
deleted file mode 100644
diff --git a/Engine/Source/Runtime/Engine/Private/TextureDerivedDataBuildExporter.h b/Engine/Source/Runtime/Engine/Private/TextureDerivedDataBuildExporter.h
deleted file mode 100644
diff --git a/Engine/Source/Runtime/Engine/Private/TextureDerivedDataTask.cpp b/Engine/Source/Runtime/Engine/Private/TextureDerivedDataTask.cpp
--- a/Engine/Source/Runtime/Engine/Private/TextureDerivedDataTask.cpp
+++ b/Engine/Source/Runtime/Engine/Private/TextureDerivedDataTask.cpp
@@ -431,8 +431,6 @@ void FTextureCacheDerivedDataWorker::BuildTexture(bool bReplaceExistingDDC)
 
 		FOptTexturePlatformData OptData;
 
-		BuildExporter.ExportTextureBuild(Texture, BuildSettingsPerLayer[0], 0, NUM_INLINE_DERIVED_MIPS);
-
 		if (!BuildFunctionName.IsEmpty())
 		{
 			using namespace UE::DerivedData;
@@ -606,8 +604,6 @@ void FTextureCacheDerivedDataWorker::BuildTexture(bool bReplaceExistingDDC)
 				}
 
 				DerivedData->SetOptData(OptData);
-
-				BuildExporter.ExportTextureOutput(*DerivedData, BuildSettingsPerLayer[0]);
 				
 				// Store it in the cache.
 				// @todo: This will remove the streaming bulk data, which we immediately reload below!
@@ -673,7 +669,6 @@ FTextureCacheDerivedDataWorker::FTextureCacheDerivedDataWorker(
 	}
 	UTexture::GetPixelFormatEnum();
 	GetTextureDerivedDataKeySuffix(Texture, InSettingsPerLayer, KeySuffix);
-	BuildExporter.Init(KeySuffix);
 		
 	const bool bAllowAsyncBuild = (CacheFlags & ETextureCacheFlags::AllowAsyncBuild) != 0;
 	const bool bAllowAsyncLoading = (CacheFlags & ETextureCacheFlags::AllowAsyncLoading) != 0;
@@ -866,7 +861,6 @@ void FTextureCacheDerivedDataWorker::DoWork()
 		if (TextureData.IsValid() && Texture.Source.IsBulkDataLoaded())
 		{
 			TextureData.GetSourceMips(Texture.Source, ImageWrapper);
-			BuildExporter.ExportTextureSourceBulkData(Texture.Source);
 			bHasTextureSourceMips = true;
 		}
 
@@ -874,24 +868,18 @@ void FTextureCacheDerivedDataWorker::DoWork()
 		if (CompositeTextureData.IsValid() && Texture.CompositeTexture && Texture.CompositeTexture->Source.IsBulkDataLoaded())
 		{
 			CompositeTextureData.GetSourceMips(Texture.CompositeTexture->Source, ImageWrapper);
-			BuildExporter.ExportCompositeTextureSourceBulkData(Texture.CompositeTexture->Source);
 			bHasCompositeTextureSourceMips = true;
 		}
 
 		if (bAllowAsyncLoading && !bHasTextureSourceMips)
 		{
 			TextureData.GetAsyncSourceMips(ImageWrapper);
-			BuildExporter.ExportTextureSourceBulkData(TextureData.AsyncSource);
 			TextureData.AsyncSource.RemoveBulkData();
 		}
 
 		if (bAllowAsyncLoading && !bHasCompositeTextureSourceMips)
 		{
 			CompositeTextureData.GetAsyncSourceMips(ImageWrapper);
-			if ((bool)Texture.CompositeTexture)
-			{
-				BuildExporter.ExportCompositeTextureSourceBulkData(CompositeTextureData.AsyncSource);
-			}
 			CompositeTextureData.AsyncSource.RemoveBulkData();
 		}
 
diff --git a/Engine/Source/Runtime/Engine/Private/TextureDerivedDataTask.h b/Engine/Source/Runtime/Engine/Private/TextureDerivedDataTask.h
--- a/Engine/Source/Runtime/Engine/Private/TextureDerivedDataTask.h
+++ b/Engine/Source/Runtime/Engine/Private/TextureDerivedDataTask.h
@@ -15,7 +15,6 @@
 #include "IImageWrapperModule.h"
 #include "ImageCore.h"
 #include "TextureCompressorModule.h"
-#include "TextureDerivedDataBuildExporter.h"
 
 #endif // WITH_EDITOR
 
@@ -126,8 +125,6 @@ class FTextureCacheDerivedDataWorker : public FNonAbandonableTask
 	FTextureSourceData CompositeTextureData;
 	/** DDC2 build function name to use to build this texture (if DDC2 is enabled and the target texture type has a DDC2 build function, empty otherwise) */
 	FString BuildFunctionName;
-	/** Exporter that can optionally write out build actions and reference outputs for the purpose of testing remote execution of builds */
-	FTextureDerivedDataBuildExporter BuildExporter;
 	/** Texture cache flags. */
 	uint32 CacheFlags;
 	/** Have many bytes were loaded from DDC or built (for telemetry) */
