Add IsDataValid override to ChildActorComponent to be invalid if there is a cycle detected between the CAC and the class that it's on

#rnx

diff --git a/Engine/Source/Runtime/Engine/Classes/Components/ChildActorComponent.h b/Engine/Source/Runtime/Engine/Classes/Components/ChildActorComponent.h
--- a/Engine/Source/Runtime/Engine/Classes/Components/ChildActorComponent.h
+++ b/Engine/Source/Runtime/Engine/Classes/Components/ChildActorComponent.h
@@ -153,6 +153,7 @@ public:
 	virtual void PostEditImport() override;
 	virtual void PostEditUndo() override;
 	virtual void PostLoad() override;
+	virtual EDataValidationResult IsDataValid(TArray<FText>& ValidationErrors) override;
 #endif
 	virtual void Serialize(FArchive& Ar) override;
 	virtual void BeginDestroy() override;
diff --git a/Engine/Source/Runtime/Engine/Private/Components/ChildActorComponent.cpp b/Engine/Source/Runtime/Engine/Private/Components/ChildActorComponent.cpp
--- a/Engine/Source/Runtime/Engine/Private/Components/ChildActorComponent.cpp
+++ b/Engine/Source/Runtime/Engine/Private/Components/ChildActorComponent.cpp
@@ -568,6 +568,20 @@ void UChildActorComponent::PostLoad()
 	}
 
 }
+
+EDataValidationResult UChildActorComponent::IsDataValid(TArray<FText>& ValidationErrors)
+{
+	EDataValidationResult Result = CombineDataValidationResults(Super::IsDataValid(ValidationErrors), EDataValidationResult::Valid);
+
+	// If a CAC is set to the class of the blueprint it is currently on, then there will be a cycle and the data is invalid
+	const UClass* const Outer = Cast<UClass>(GetOuter());
+	if (Outer && Outer == ChildActorClass)
+	{
+		Result = EDataValidationResult::Invalid;
+		ValidationErrors.Add(FText::Format(NSLOCTEXT("ChildActorComponent", "ChildActorCycle", "A Child Actor Component's class cannot be set to its owner! '{0}' is an invalid class choice for '{1}'."), FText::FromString(ChildActorClass->GetName()), FText::FromString(GetPathName())));
+	}
+	return Result;
+}
 #endif
 
 void UChildActorComponent::CreateChildActor()
