Implement ray tracing min LOD clamp for skeletal meshes

* Each GPU skeletal mesh object now has 2 skin cache entries, one for raster and one for ray tracing. Skin cache usage is decoupled from user settings and ray tracing forces creation of a dedicated skin cache entry.
* RayTracingMinLOD can be set on a per-asset basis (which requires re-cooking to change) and a global cvar r.RayTracing.Geometry.SkeletalMeshes.LODBias allows changing between different LOD levels at runtime for quick tests.
* Ray tracing support for skeletal meshes remains disabled until a pass over all the meshes setting proper ray tracing min LOD is done to avoid OOM crashes.

diff --git a/Engine/Source/Runtime/Engine/Classes/Engine/SkeletalMesh.h b/Engine/Source/Runtime/Engine/Classes/Engine/SkeletalMesh.h
--- a/Engine/Source/Runtime/Engine/Classes/Engine/SkeletalMesh.h
+++ b/Engine/Source/Runtime/Engine/Classes/Engine/SkeletalMesh.h
@@ -1959,6 +1959,12 @@ public:
 	UPROPERTY(EditAnywhere, Category = RayTracing)
 	uint8 bSupportRayTracing : 1;
 
+	/**
+	 * LOD bias for ray tracing. When non-zero, a different LOD level other than the predicted LOD level will be used for ray tracing. Advanced features like morph targets and cloth simulation may not work properly.
+	 */
+	UPROPERTY(EditAnywhere, Category = RayTracing)
+	int32 RayTracingMinLOD;
+
 	UE_DEPRECATED(4.27, "Please do not access this member directly; use USkeletalMesh::GetMorphTargets() or USkeletalMesh::SetMorphTargets().")
 	UPROPERTY(BlueprintGetter = GetMorphTargets, BlueprintSetter = SetMorphTargets, Category = Mesh)
 	TArray<TObjectPtr<UMorphTarget>> MorphTargets;
diff --git a/Engine/Source/Runtime/Engine/Private/Components/SkinnedMeshComponent.cpp b/Engine/Source/Runtime/Engine/Private/Components/SkinnedMeshComponent.cpp
--- a/Engine/Source/Runtime/Engine/Private/Components/SkinnedMeshComponent.cpp
+++ b/Engine/Source/Runtime/Engine/Private/Components/SkinnedMeshComponent.cpp
@@ -1481,13 +1481,13 @@ bool USkinnedMeshComponent::IsSkinCacheAllowed(int32 LodIdx) const
 
 	if (!SkeletalMesh)
 	{
-		return GEnableGPUSkinCache && (bIsRayTracing || bGlobalDefault);
+		return GEnableGPUSkinCache && bGlobalDefault;
 	}
 
 	FSkeletalMeshLODInfo* LodInfo = SkeletalMesh->GetLODInfo(LodIdx);
 	if (!LodInfo)
 	{
-		return GEnableGPUSkinCache && (bIsRayTracing || bGlobalDefault);
+		return GEnableGPUSkinCache && bGlobalDefault;
 	}
 
 	bool bLodEnabled = LodInfo->SkinCacheUsage == ESkinCacheUsage::Auto ?
@@ -1496,14 +1496,14 @@ bool USkinnedMeshComponent::IsSkinCacheAllowed(int32 LodIdx) const
 
 	if (!SkinCacheUsage.IsValidIndex(LodIdx))
 	{
-		return GEnableGPUSkinCache && (bIsRayTracing || bLodEnabled);
+		return GEnableGPUSkinCache && bLodEnabled;
 	}
 
 	bool bComponentEnabled = SkinCacheUsage[LodIdx] == ESkinCacheUsage::Auto ? 
 		bLodEnabled :
 		SkinCacheUsage[LodIdx] == ESkinCacheUsage::Enabled;
 
-	return GEnableGPUSkinCache && (bIsRayTracing || bComponentEnabled);
+	return GEnableGPUSkinCache && bComponentEnabled;
 }
 
 void USkinnedMeshComponent::GetBoneNames(TArray<FName>& BoneNames)
diff --git a/Engine/Source/Runtime/Engine/Private/GPUSkinCache.cpp b/Engine/Source/Runtime/Engine/Private/GPUSkinCache.cpp
--- a/Engine/Source/Runtime/Engine/Private/GPUSkinCache.cpp
+++ b/Engine/Source/Runtime/Engine/Private/GPUSkinCache.cpp
@@ -222,12 +222,13 @@ enum class EGPUSkinCacheDispatchFlags
 class FGPUSkinCacheEntry
 {
 public:
-	FGPUSkinCacheEntry(FGPUSkinCache* InSkinCache, FSkeletalMeshObjectGPUSkin* InGPUSkin, FGPUSkinCache::FRWBuffersAllocation* InPositionAllocation)
+	FGPUSkinCacheEntry(FGPUSkinCache* InSkinCache, FSkeletalMeshObjectGPUSkin* InGPUSkin, FGPUSkinCache::FRWBuffersAllocation* InPositionAllocation, int32 InLOD, EGPUSkinCacheEntryMode InMode)
 		: PositionAllocation(InPositionAllocation)
 		, SkinCache(InSkinCache)
 		, GPUSkin(InGPUSkin)
 		, MorphBuffer(0)
-		, LOD(InGPUSkin->GetLOD())
+		, LOD(InLOD)
+		, Mode(InMode)
 	{
 		
 		const TArray<FSkelMeshRenderSection>& Sections = InGPUSkin->GetRenderSections(LOD);
@@ -394,9 +395,9 @@ public:
 		return SectionData.SourceVertexFactory == SourceVertexFactory;
 	}
 
-	bool IsValid(FSkeletalMeshObjectGPUSkin* InSkin) const
+	bool IsValid(FSkeletalMeshObjectGPUSkin* InSkin, int32 InLOD) const
 	{
-		return GPUSkin == InSkin && GPUSkin->GetLOD() == LOD;
+		return GPUSkin == InSkin && LOD == InLOD;
 	}
 
 	void UpdateSkinWeightBuffer()
@@ -431,7 +432,6 @@ public:
 		Data.SectionIndex = SectionIndex;
 		Data.Section = Section;
 
-		check(GPUSkin->GetLOD() == LOD);
 		FSkeletalMeshRenderData& SkelMeshRenderData = GPUSkin->GetSkeletalMeshRenderData();
 		FSkeletalMeshLODRenderData& LodData = SkelMeshRenderData.LODRenderData[LOD];
 		check(Data.SectionIndex == LodData.FindSectionIndex(*Section));
@@ -513,6 +513,7 @@ public:
 	TArray<FSectionDispatchData>& GetDispatchData() { return DispatchData; }
 
 protected:
+	EGPUSkinCacheEntryMode Mode;
 	FGPUSkinCache::FRWBuffersAllocation* PositionAllocation;
 	FGPUSkinCache* SkinCache;
 	TArray<FGPUSkinBatchElementUserData> BatchElementsUserData;
@@ -1477,6 +1478,7 @@ void FGPUSkinCache::DoDispatch(FRHICommandListImmediate& RHICmdList, FGPUSkinCac
 }
 
 bool FGPUSkinCache::ProcessEntry(
+	EGPUSkinCacheEntryMode Mode,
 	FRHICommandListImmediate& RHICmdList, 
 	FGPUBaseSkinVertexFactory* VertexFactory,
 	FGPUSkinPassthroughVertexFactory* TargetVertexFactory, 
@@ -1489,7 +1491,8 @@ bool FGPUSkinCache::ProcessEntry(
 	const FMatrix44f& ClothLocalToWorld, 
 	float ClothBlendWeight, 
 	uint32 RevisionNumber, 
-	int32 Section, 
+	int32 Section,
+	int32 LODIndex,
 	FGPUSkinCacheEntry*& InOutEntry
 	)
 {
@@ -1500,7 +1503,6 @@ bool FGPUSkinCache::ProcessEntry(
 	const uint32 InputStreamStart = BatchElement.BaseVertexIndex;
 
 	FSkeletalMeshRenderData& SkelMeshRenderData = Skin->GetSkeletalMeshRenderData();
-	int32 LODIndex = Skin->GetLOD();
 	FSkeletalMeshLODRenderData& LodData = SkelMeshRenderData.LODRenderData[LODIndex];
 
 	if (FlushCounter < GGPUSkinCacheFlushCounter)
@@ -1512,7 +1514,7 @@ bool FGPUSkinCache::ProcessEntry(
 	if (InOutEntry)
 	{
 		// If the LOD changed, the entry has to be invalidated
-		if (!InOutEntry->IsValid(Skin))
+		if (!InOutEntry->IsValid(Skin, LODIndex))
 		{
 			Release(InOutEntry);
 			InOutEntry = nullptr;
@@ -1583,7 +1585,7 @@ bool FGPUSkinCache::ProcessEntry(
 			return false;
 		}
 
-		InOutEntry = new FGPUSkinCacheEntry(this, Skin, NewPositionAllocation);
+		InOutEntry = new FGPUSkinCacheEntry(this, Skin, NewPositionAllocation, LODIndex, Mode);
 		InOutEntry->GPUSkin = Skin;
 
 		InOutEntry->SetupSection(Section, NewPositionAllocation, &LodData.RenderSections[Section], MorphVertexBuffer, ClothVertexBuffer, NumVertices, InputStreamStart, VertexFactory, TargetVertexFactory);
@@ -1849,6 +1851,12 @@ void FGPUSkinCache::EndBatchDispatch(FRHICommandListImmediate& RHICmdList)
 			FDispatchEntry& DispatchItem = BatchDispatches[Index];
 
 			FGPUSkinCacheEntry* SkinCacheEntry = DispatchItem.SkinCacheEntry;
+
+			if (SkinCacheEntry->Mode != EGPUSkinCacheEntryMode::RayTracing)
+			{
+				continue;
+			}
+
 			FSkeletalMeshLODRenderData& LODModel = *DispatchItem.LODModel;
 
 			if (SkinCacheEntriesProcessed.Contains(SkinCacheEntry))
diff --git a/Engine/Source/Runtime/Engine/Private/SkeletalMesh.cpp b/Engine/Source/Runtime/Engine/Private/SkeletalMesh.cpp
--- a/Engine/Source/Runtime/Engine/Private/SkeletalMesh.cpp
+++ b/Engine/Source/Runtime/Engine/Private/SkeletalMesh.cpp
@@ -134,6 +134,14 @@ static TAutoConsoleVariable<int32> CVarRayTracingSupportSkeletalMeshes(
 	TEXT("This setting is read-only at runtime. (default: 1)"),
 	ECVF_ReadOnly);
 
+static TAutoConsoleVariable<int32> CVarRayTracingSkeletalMeshLODBias(
+	TEXT("r.RayTracing.Geometry.SkeletalMeshes.LODBias"),
+	0,
+	TEXT("Global LOD bias for skeletal meshes in ray tracing.\n")
+	TEXT("When non-zero, a different LOD level other than the predicted LOD level will be used for ray tracing. Advanced features like morph targets and cloth simulation may not work properly.\n")
+	TEXT("Final LOD level to use in ray tracing is the sum of this global bias and the bias set on each skeletal mesh asset."),
+	ECVF_RenderThreadSafe);
+
 #if WITH_APEX_CLOTHING
 /*-----------------------------------------------------------------------------
 	utility functions for apex clothing 
@@ -491,6 +499,7 @@ PRAGMA_ENABLE_DEPRECATION_WARNINGS
 	SetMinLod(FPerPlatformInt(0));
 	SetDisableBelowMinLodStripping(FPerPlatformBool(false));
 	bSupportRayTracing = true;
+	RayTracingMinLOD = 0;
 }
 
 USkeletalMesh::USkeletalMesh(FVTableHelper& Helper)
@@ -6033,9 +6042,9 @@ void FSkeletalMeshSceneProxy::GetMeshElementsConditionallySelectable(const TArra
 #endif
 }
 
-void FSkeletalMeshSceneProxy::CreateBaseMeshBatch(const FSceneView* View, const FSkeletalMeshLODRenderData& LODData, const int32 LODIndex, const int32 SectionIndex, const FSectionElementInfo& SectionElementInfo, FMeshBatch& Mesh) const
+void FSkeletalMeshSceneProxy::CreateBaseMeshBatch(const FSceneView* View, const FSkeletalMeshLODRenderData& LODData, const int32 LODIndex, const int32 SectionIndex, const FSectionElementInfo& SectionElementInfo, FMeshBatch& Mesh, ESkinVertexFactoryMode VFMode) const
 {
-	Mesh.VertexFactory = MeshObject->GetSkinVertexFactory(View, LODIndex, SectionIndex);
+	Mesh.VertexFactory = MeshObject->GetSkinVertexFactory(View, LODIndex, SectionIndex, VFMode);
 	Mesh.MaterialRenderProxy = SectionElementInfo.Material->GetRenderProxy();
 #if RHI_RAYTRACING
 	Mesh.SegmentIndex = SectionIndex;
@@ -6047,7 +6056,7 @@ void FSkeletalMeshSceneProxy::CreateBaseMeshBatch(const FSceneView* View, const
 	BatchElement.IndexBuffer = LODData.MultiSizeIndexContainer.GetIndexBuffer();
 	BatchElement.MinVertexIndex = LODData.RenderSections[SectionIndex].GetVertexBufferIndex();
 	BatchElement.MaxVertexIndex = LODData.RenderSections[SectionIndex].GetVertexBufferIndex() + LODData.RenderSections[SectionIndex].GetNumVertices() - 1;
-	BatchElement.VertexFactoryUserData = FGPUSkinCache::GetFactoryUserData(MeshObject->SkinCacheEntry, SectionIndex);
+	BatchElement.VertexFactoryUserData = FGPUSkinCache::GetFactoryUserData(VFMode == ESkinVertexFactoryMode::RayTracing ? MeshObject->SkinCacheEntryForRayTracing : MeshObject->SkinCacheEntry, SectionIndex);
 	BatchElement.PrimitiveUniformBuffer = GetUniformBuffer();
 	BatchElement.NumPrimitives = LODData.RenderSections[SectionIndex].NumTriangles;
 }
@@ -6207,7 +6216,7 @@ void FSkeletalMeshSceneProxy::GetDynamicRayTracingInstances(FRayTracingMaterialG
 			RayTracingInstance.Geometry = MeshObject->GetRayTracingGeometry();
 
 			// Setup materials for each segment
-			const int32 LODIndex = MeshObject->GetLOD();
+			const int32 LODIndex = MeshObject->GetRayTracingLOD();
 			check(LODIndex < SkeletalMeshRenderData->LODRenderData.Num());
 			const FSkeletalMeshLODRenderData& LODData = SkeletalMeshRenderData->LODRenderData[LODIndex];
 
@@ -6241,7 +6250,7 @@ void FSkeletalMeshSceneProxy::GetDynamicRayTracingInstances(FRayTracingMaterialG
 				const FSectionElementInfo& SectionElementInfo = Iter.GetSectionElementInfo();
 
 				FMeshBatch MeshBatch;
-				CreateBaseMeshBatch(Context.ReferenceView, LODData, LODIndex, SectionIndex, SectionElementInfo, MeshBatch);
+				CreateBaseMeshBatch(Context.ReferenceView, LODData, LODIndex, SectionIndex, SectionElementInfo, MeshBatch, ESkinVertexFactoryMode::RayTracing);
 
 				RayTracingInstance.Materials.Add(MeshBatch);
 			}
@@ -6259,7 +6268,7 @@ void FSkeletalMeshSceneProxy::GetDynamicRayTracingInstances(FRayTracingMaterialG
 				RayTracingInstance.InstanceTransforms.Add(GetLocalToWorld());
 			}
 
-			const FVertexFactory* VertexFactory = MeshObject->GetSkinVertexFactory(Context.ReferenceView, LODIndex, 0);
+			const FVertexFactory* VertexFactory = MeshObject->GetSkinVertexFactory(Context.ReferenceView, LODIndex, 0, ESkinVertexFactoryMode::RayTracing);
 			const FVertexFactoryType* VertexFactoryType = VertexFactory->GetType();
 			if (bAnySegmentUsesWorldPositionOffset 
 				&& ensureMsgf(VertexFactoryType->SupportsRayTracingDynamicGeometry(),
diff --git a/Engine/Source/Runtime/Engine/Private/SkeletalRender.cpp b/Engine/Source/Runtime/Engine/Private/SkeletalRender.cpp
--- a/Engine/Source/Runtime/Engine/Private/SkeletalRender.cpp
+++ b/Engine/Source/Runtime/Engine/Private/SkeletalRender.cpp
@@ -42,6 +42,7 @@ FSkeletalMeshObject::FSkeletalMeshObject(USkinnedMeshComponent* InMeshComponent,
 , bRequireRecreatingRayTracingGeometry(false)
 , bSupportRayTracing(IsSkeletalMeshRayTracingSupported() && InMeshComponent->SkeletalMesh->bSupportRayTracing)
 , bHiddenMaterialVisibilityDirtyForRayTracing(false)
+, RayTracingMinLOD(InMeshComponent->SkeletalMesh->RayTracingMinLOD)
 #endif
 #if !UE_BUILD_SHIPPING
 , DebugName(InMeshComponent->SkeletalMesh->GetFName())
@@ -55,6 +56,7 @@ FSkeletalMeshObject::FSkeletalMeshObject(USkinnedMeshComponent* InMeshComponent,
 ,	SkeletalMeshRenderData(InSkelMeshRenderData)
 ,	SkeletalMeshLODInfo(InMeshComponent->SkeletalMesh->GetLODInfoArray())
 ,	SkinCacheEntry(nullptr)
+,	SkinCacheEntryForRayTracing(nullptr)
 ,	LastFrameNumber(0)
 ,	bUsePerBoneMotionBlur(InMeshComponent->bPerBoneMotionBlur)
 ,	StatId(InMeshComponent->SkeletalMesh->GetStatID(true))
diff --git a/Engine/Source/Runtime/Engine/Private/SkeletalRenderCPUSkin.cpp b/Engine/Source/Runtime/Engine/Private/SkeletalRenderCPUSkin.cpp
--- a/Engine/Source/Runtime/Engine/Private/SkeletalRenderCPUSkin.cpp
+++ b/Engine/Source/Runtime/Engine/Private/SkeletalRenderCPUSkin.cpp
@@ -287,7 +287,7 @@ void FSkeletalMeshObjectCPUSkin::CacheVertices(int32 LODIndex, bool bForce) cons
 	}
 }
 
-const FVertexFactory* FSkeletalMeshObjectCPUSkin::GetSkinVertexFactory(const FSceneView* View, int32 LODIndex,int32 /*ChunkIdx*/) const
+const FVertexFactory* FSkeletalMeshObjectCPUSkin::GetSkinVertexFactory(const FSceneView* View, int32 LODIndex, int32 ChunkIdx, ESkinVertexFactoryMode VFMode) const
 {
 	check( LODs.IsValidIndex(LODIndex) );
 	return &LODs[LODIndex].VertexFactory;
diff --git a/Engine/Source/Runtime/Engine/Private/SkeletalRenderCPUSkin.h b/Engine/Source/Runtime/Engine/Private/SkeletalRenderCPUSkin.h
--- a/Engine/Source/Runtime/Engine/Private/SkeletalRenderCPUSkin.h
+++ b/Engine/Source/Runtime/Engine/Private/SkeletalRenderCPUSkin.h
@@ -103,7 +103,7 @@ public:
 	virtual void EnableOverlayRendering(bool bEnabled, const TArray<int32>* InBonesOfInterest, const TArray<UMorphTarget*>* InMorphTargetOfInterest) override;
 	virtual void CacheVertices(int32 LODIndex, bool bForce) const override;
 	virtual bool IsCPUSkinned() const override { return true; }
-	virtual const FVertexFactory* GetSkinVertexFactory(const FSceneView* View, int32 LODIndex, int32 ChunkIdx) const override;
+	virtual const FVertexFactory* GetSkinVertexFactory(const FSceneView* View, int32 LODIndex, int32 ChunkIdx, ESkinVertexFactoryMode VFMode = ESkinVertexFactoryMode::Default) const override;
 	virtual TArray<FTransform>* GetComponentSpaceTransforms() const override;
 	virtual const TArray<FMatrix44f>& GetReferenceToLocalMatrices() const override;
 
diff --git a/Engine/Source/Runtime/Engine/Private/SkeletalRenderGPUSkin.cpp b/Engine/Source/Runtime/Engine/Private/SkeletalRenderGPUSkin.cpp
--- a/Engine/Source/Runtime/Engine/Private/SkeletalRenderGPUSkin.cpp
+++ b/Engine/Source/Runtime/Engine/Private/SkeletalRenderGPUSkin.cpp
@@ -69,6 +69,12 @@ static FAutoConsoleVariableRef CVarMorphTargetWeightThreshold(
 	ECVF_Default
 );
 
+int32 GetRayTracingSkeletalMeshGlobalLODBias()
+{
+	static const auto RayTracingSkeletalMeshLODBiasVar = IConsoleManager::Get().FindTConsoleVariableDataInt(TEXT("r.RayTracing.Geometry.SkeletalMeshes.LODBias"));
+
+	return RayTracingSkeletalMeshLODBiasVar != nullptr ? (RayTracingSkeletalMeshLODBiasVar->GetValueOnGameThread()) : 0;
+}
 
 /*-----------------------------------------------------------------------------
 FMorphVertexBuffer
@@ -217,14 +223,17 @@ void FSkeletalMeshObjectGPUSkin::ReleaseResources()
 	FSkeletalMeshObjectGPUSkin* MeshObject = this;
 	FGPUSkinCacheEntry** PtrSkinCacheEntry = &SkinCacheEntry;
 	ENQUEUE_RENDER_COMMAND(WaitRHIThreadFenceForDynamicData)(
-		[MeshObject, PtrSkinCacheEntry](FRHICommandList& RHICmdList)
+		[MeshObject, PtrSkinCacheEntry, &SkinCacheEntryForRayTracing = SkinCacheEntryForRayTracing](FRHICommandList& RHICmdList)
 		{
 			FGPUSkinCacheEntry*& LocalSkinCacheEntry = *PtrSkinCacheEntry;
 			FGPUSkinCache::Release(LocalSkinCacheEntry);
+			FGPUSkinCacheEntry* LocalSkinCacheEntryForRayTracing = SkinCacheEntryForRayTracing;
+			FGPUSkinCache::Release(LocalSkinCacheEntryForRayTracing);
 
 			FScopeCycleCounter Context(MeshObject->GetStatId());
 			MeshObject->WaitForRHIThreadFenceForDynamicData();
 			*PtrSkinCacheEntry = nullptr;
+			SkinCacheEntryForRayTracing = nullptr;
 		}
 	);
 
@@ -295,7 +304,7 @@ void FSkeletalMeshObjectGPUSkin::Update(int32 LODIndex,USkinnedMeshComponent* In
 	// create the new dynamic data for use by the rendering thread
 	// this data is only deleted when another update is sent
 	FDynamicSkelMeshObjectDataGPUSkin* NewDynamicData = FDynamicSkelMeshObjectDataGPUSkin::AllocDynamicSkelMeshObjectDataGPUSkin();		
-	NewDynamicData->InitDynamicSkelMeshObjectDataGPUSkin(InMeshComponent,SkeletalMeshRenderData,LODIndex,ActiveMorphTargets, MorphTargetWeights, PreviousBoneTransformUpdateMode);
+	NewDynamicData->InitDynamicSkelMeshObjectDataGPUSkin(InMeshComponent,SkeletalMeshRenderData,this,LODIndex,ActiveMorphTargets, MorphTargetWeights, PreviousBoneTransformUpdateMode);
 
 	// We prepare the next frame but still have the value from the last one
 	uint32 FrameNumberToPrepare = GFrameNumber + 1;
@@ -350,6 +359,15 @@ void FSkeletalMeshObjectGPUSkin::UpdateSkinWeightBuffer(USkinnedMeshComponent* I
 						GPUSkinCache->UpdateSkinWeightBuffer(SkinCacheEntryToUpdate);
 					});
 				}
+
+				if (GPUSkinCache && SkinCacheEntryForRayTracing)
+				{
+					ENQUEUE_RENDER_COMMAND(UpdateSkinCacheSkinWeightBuffer)(
+						[GPUSkinCache, SkinCacheEntryForRayTracing = SkinCacheEntryForRayTracing](FRHICommandListImmediate& RHICmdList)
+					{
+						GPUSkinCache->UpdateSkinWeightBuffer(SkinCacheEntryForRayTracing);
+					});
+				}
 			}
 		}
 	}
@@ -378,9 +396,10 @@ void FSkeletalMeshObjectGPUSkin::UpdateDynamicData_RenderThread(FGPUSkinCache* G
 	// When switching to a higher detailed LOD, skin cache destroys the old entry, if there isn't enough memory to create the new entry because higher LOD costs more memory, SkinCacheEntry becomes null.
 	// So make sure to set bRequireRecreatingRayTracingGeometry to true to account for this, so RT geometry is updated correctly.
 	bRequireRecreatingRayTracingGeometry = (DynamicData == nullptr || RayTracingGeometry.Initializer.Segments.Num() == 0 || // Newly created
-		(DynamicData != nullptr && DynamicData->LODIndex != InDynamicData->LODIndex) || // LOD level changed
+		(DynamicData != nullptr && DynamicData->RayTracingLODIndex != InDynamicData->RayTracingLODIndex) || // LOD level changed
 		SkinCacheEntry == nullptr ||
-		bHiddenMaterialVisibilityDirtyForRayTracing);
+		bHiddenMaterialVisibilityDirtyForRayTracing ||
+		SkinCacheEntryForRayTracing == nullptr);
 
 	bHiddenMaterialVisibilityDirtyForRayTracing = false;
 	
@@ -413,16 +432,21 @@ void FSkeletalMeshObjectGPUSkin::UpdateDynamicData_RenderThread(FGPUSkinCache* G
 	}
 	else
 	{
-		ProcessUpdatedDynamicData(GPUSkinCache, RHICmdList, FrameNumberToPrepare, RevisionNumber, bMorphNeedsUpdate);
+		ProcessUpdatedDynamicData(EGPUSkinCacheEntryMode::Raster, GPUSkinCache, RHICmdList, FrameNumberToPrepare, RevisionNumber, bMorphNeedsUpdate, DynamicData->LODIndex);
+
+		if (FGPUSkinCache::IsGPUSkinCacheRayTracingSupported() && GPUSkinCache && SkeletalMeshRenderData->bSupportRayTracing)
+		{
+			ProcessUpdatedDynamicData(EGPUSkinCacheEntryMode::RayTracing, GPUSkinCache, RHICmdList, FrameNumberToPrepare, RevisionNumber, bMorphNeedsUpdate, DynamicData->RayTracingLODIndex);
+		}
 	}
 
 #if RHI_RAYTRACING
-	if (IsRayTracingEnabled() && GEnableGPUSkinCache && GPUSkinCache && !GPUSkinCache->IsBatchingDispatch() && SkeletalMeshRenderData->bSupportRayTracing)
+	if (FGPUSkinCache::IsGPUSkinCacheRayTracingSupported() && GPUSkinCache && SkeletalMeshRenderData->bSupportRayTracing && !GPUSkinCache->IsBatchingDispatch())
 	{
-		if (SkinCacheEntry)
+		if (SkinCacheEntryForRayTracing)
 		{
-			FSkeletalMeshLODRenderData& LODModel = this->SkeletalMeshRenderData->LODRenderData[DynamicData->LODIndex];
-			GPUSkinCache->ProcessRayTracingGeometryToUpdate(RHICmdList, SkinCacheEntry, LODModel, bRequireRecreatingRayTracingGeometry, DynamicData->bAnySegmentUsesWorldPositionOffset);
+			FSkeletalMeshLODRenderData& LODModel = this->SkeletalMeshRenderData->LODRenderData[DynamicData->RayTracingLODIndex];
+			GPUSkinCache->ProcessRayTracingGeometryToUpdate(RHICmdList, SkinCacheEntryForRayTracing, LODModel, bRequireRecreatingRayTracingGeometry, DynamicData->bAnySegmentUsesWorldPositionOffset);
 		}
 		else
 		{
@@ -438,7 +462,7 @@ void FSkeletalMeshObjectGPUSkin::PreGDMECallback(FGPUSkinCache* GPUSkinCache, ui
 {
 	if (bNeedsUpdateDeferred)
 	{
-		ProcessUpdatedDynamicData(GPUSkinCache, FRHICommandListExecutor::GetImmediateCommandList(), FrameNumber, LastBoneTransformRevisionNumber, bMorphNeedsUpdateDeferred);
+		ProcessUpdatedDynamicData(EGPUSkinCacheEntryMode::Raster, GPUSkinCache, FRHICommandListExecutor::GetImmediateCommandList(), FrameNumber, LastBoneTransformRevisionNumber, bMorphNeedsUpdateDeferred, DynamicData->LODIndex);
 	}
 }
 
@@ -452,20 +476,20 @@ void FSkeletalMeshObjectGPUSkin::WaitForRHIThreadFenceForDynamicData()
 	}
 }
 
-void FSkeletalMeshObjectGPUSkin::ProcessUpdatedDynamicData(FGPUSkinCache* GPUSkinCache, FRHICommandListImmediate& RHICmdList, uint32 FrameNumberToPrepare, uint32 RevisionNumber, bool bMorphNeedsUpdate)
+void FSkeletalMeshObjectGPUSkin::ProcessUpdatedDynamicData(EGPUSkinCacheEntryMode Mode, FGPUSkinCache* GPUSkinCache, FRHICommandListImmediate& RHICmdList, uint32 FrameNumberToPrepare, uint32 RevisionNumber, bool bMorphNeedsUpdate, int32 LODIndex)
 {
 	QUICK_SCOPE_CYCLE_COUNTER(STAT_FSkeletalMeshObjectGPUSkin_ProcessUpdatedDynamicData);
 	bNeedsUpdateDeferred = false;
 	bMorphNeedsUpdateDeferred = false;
 
-	FSkeletalMeshObjectLOD& LOD = LODs[DynamicData->LODIndex];
+	FSkeletalMeshObjectLOD& LOD = LODs[LODIndex];
 
 	// if hasn't been updated, force update again
 	bMorphNeedsUpdate = LOD.MorphVertexBuffer.bHasBeenUpdated ? bMorphNeedsUpdate : true;
 	bMorphNeedsUpdate |= (GForceUpdateMorphTargets != 0);
 
-	const FSkeletalMeshLODRenderData& LODData = SkeletalMeshRenderData->LODRenderData[DynamicData->LODIndex];
-	const TArray<FSkelMeshRenderSection>& Sections = GetRenderSections(DynamicData->LODIndex);
+	const FSkeletalMeshLODRenderData& LODData = SkeletalMeshRenderData->LODRenderData[LODIndex];
+	const TArray<FSkelMeshRenderSection>& Sections = GetRenderSections(LODIndex);
 
 	// Only consider morphs with active curves and data to deform.
 	const bool bMorph = DynamicData->NumWeightedActiveMorphTargets > 0 && LODData.GetNumVertices() > 0;
@@ -475,7 +499,7 @@ void FSkeletalMeshObjectGPUSkin::ProcessUpdatedDynamicData(FGPUSkinCache* GPUSki
 
 	bool bDataPresent = false;
 
-	bool bGPUSkinCacheEnabled = GPUSkinCache && GEnableGPUSkinCache && (FeatureLevel >= ERHIFeatureLevel::SM5) && DynamicData->bIsSkinCacheAllowed;
+	bool bGPUSkinCacheEnabled = GPUSkinCache && GEnableGPUSkinCache && (FeatureLevel >= ERHIFeatureLevel::SM5) && (DynamicData->bIsSkinCacheAllowed || Mode == EGPUSkinCacheEntryMode::RayTracing);
 
 	if (DynamicData->PreSkinningOffsets.Num() > 0)
 	{
@@ -621,15 +645,17 @@ void FSkeletalMeshObjectGPUSkin::ProcessUpdatedDynamicData(FGPUSkinCache* GPUSki
 				}
 			}
 
+			bool bShouldUseSeparateMatricesForRayTracing = Mode == EGPUSkinCacheEntryMode::RayTracing && DynamicData->RayTracingLODIndex != DynamicData->LODIndex;
+
 			// if we have previous reference to loca, we also update to previous frame
 			if (DynamicData->PreviousReferenceToLocal.Num() > 0)
 			{
-				TArray<FMatrix44f>& PreviousReferenceToLocalMatrices = DynamicData->PreviousReferenceToLocal;
+				TArray<FMatrix44f>& PreviousReferenceToLocalMatrices = bShouldUseSeparateMatricesForRayTracing ? DynamicData->PreviousReferenceToLocalForRayTracing : DynamicData->PreviousReferenceToLocal;
 				ShaderData.UpdateBoneData(RHICmdList, PreviousReferenceToLocalMatrices, Section.BoneMap, RevisionNumber, true, FeatureLevel, bUseSkinCache);
 			}
 
 			// Create a uniform buffer from the bone transforms.
-			TArray<FMatrix44f>& ReferenceToLocalMatrices = DynamicData->ReferenceToLocal;
+			TArray<FMatrix44f>& ReferenceToLocalMatrices = bShouldUseSeparateMatricesForRayTracing ? DynamicData->ReferenceToLocalForRayTracing : DynamicData->ReferenceToLocal;
 			bool bNeedFence = ShaderData.UpdateBoneData(RHICmdList, ReferenceToLocalMatrices, Section.BoneMap, RevisionNumber, false, FeatureLevel, bUseSkinCache);
 
 #if WITH_APEX_CLOTHING || WITH_CHAOS_CLOTHING
@@ -659,6 +685,7 @@ void FSkeletalMeshObjectGPUSkin::ProcessUpdatedDynamicData(FGPUSkinCache* GPUSki
 				if (bSkinCacheResult)
 				{
 					bSkinCacheResult = GPUSkinCache->ProcessEntry(
+						Mode,
 						RHICmdList,
 						VertexFactory,
 						VertexFactoryData.PassthroughVertexFactories[SectionIdx].Get(),
@@ -672,7 +699,8 @@ void FSkeletalMeshObjectGPUSkin::ProcessUpdatedDynamicData(FGPUSkinCache* GPUSki
 						DynamicData->ClothBlendWeight,
 						RevisionNumber,
 						SectionIdx,
-						SkinCacheEntry
+						LODIndex,
+						Mode == EGPUSkinCacheEntryMode::RayTracing ? SkinCacheEntryForRayTracing : SkinCacheEntry
 						);
 				}
 			}
@@ -1107,7 +1135,7 @@ void FSkeletalMeshObjectGPUSkin::FSkeletalMeshObjectLOD::UpdateMorphVertexBuffer
 	}
 }
 
-const FVertexFactory* FSkeletalMeshObjectGPUSkin::GetSkinVertexFactory(const FSceneView* View, int32 LODIndex,int32 ChunkIdx) const
+const FVertexFactory* FSkeletalMeshObjectGPUSkin::GetSkinVertexFactory(const FSceneView* View, int32 LODIndex, int32 ChunkIdx, ESkinVertexFactoryMode VFMode) const
 {
 	checkSlow( LODs.IsValidIndex(LODIndex) );
 	checkSlow( DynamicData );
@@ -1115,6 +1143,15 @@ const FVertexFactory* FSkeletalMeshObjectGPUSkin::GetSkinVertexFactory(const FSc
 	const FSkelMeshObjectLODInfo& MeshLODInfo = LODInfo[LODIndex];
 	const FSkeletalMeshObjectLOD& LOD = LODs[LODIndex];
 
+	// Return the passthrough vertex factory if it is requested (by ray tracing)
+	if (VFMode == ESkinVertexFactoryMode::RayTracing)
+	{
+		check(SkinCacheEntryForRayTracing);
+		check(FGPUSkinCache::IsEntryValid(SkinCacheEntryForRayTracing, ChunkIdx));
+
+		return LOD.GPUSkinVertexFactories.PassthroughVertexFactories[ChunkIdx].Get();
+	}
+
 	// If the GPU skinning cache was used, return the passthrough vertex factory
 	if (SkinCacheEntry && FGPUSkinCache::IsEntryValid(SkinCacheEntry, ChunkIdx) && DynamicData->bIsSkinCacheAllowed)
 	{
@@ -1879,6 +1916,7 @@ FDynamicSkelMeshObjectDataGPUSkin
 void FDynamicSkelMeshObjectDataGPUSkin::Clear()
 {
 	ReferenceToLocal.Reset();
+	ReferenceToLocalForRayTracing.Reset();
 #if !(UE_BUILD_SHIPPING || UE_BUILD_TEST) 
 	MeshComponentSpaceTransforms.Reset();
 #endif
@@ -1980,6 +2018,7 @@ void FDynamicSkelMeshObjectDataGPUSkin::FreeDynamicSkelMeshObjectDataGPUSkin(FDy
 void FDynamicSkelMeshObjectDataGPUSkin::InitDynamicSkelMeshObjectDataGPUSkin(
 	USkinnedMeshComponent* InMeshComponent,
 	FSkeletalMeshRenderData* InSkeletalMeshRenderData,
+	FSkeletalMeshObjectGPUSkin* InMeshObject,
 	int32 InLODIndex,
 	const TArray<FActiveMorphTarget>& InActiveMorphTargets,
 	const TArray<float>& InMorphTargetWeights, 
@@ -1998,19 +2037,34 @@ void FDynamicSkelMeshObjectDataGPUSkin::InitDynamicSkelMeshObjectDataGPUSkin(
 	FSkeletalMeshSceneProxy* SkeletalMeshProxy = (FSkeletalMeshSceneProxy*)InMeshComponent->SceneProxy;
 	const TArray<FBoneIndexType>* ExtraRequiredBoneIndices = SkeletalMeshProxy ? &SkeletalMeshProxy->GetSortedShadowBoneIndices() : nullptr;
 
+	RayTracingLODIndex = FMath::Clamp(FMath::Max(LODIndex + GetRayTracingSkeletalMeshGlobalLODBias(), InMeshObject->RayTracingMinLOD), LODIndex, InSkeletalMeshRenderData->LODRenderData.Num() - 1);
+
 	// update ReferenceToLocal
 	UpdateRefToLocalMatrices( ReferenceToLocal, InMeshComponent, InSkeletalMeshRenderData, LODIndex, ExtraRequiredBoneIndices );
+	if (RayTracingLODIndex != LODIndex)
+	{
+		UpdateRefToLocalMatrices(ReferenceToLocalForRayTracing, InMeshComponent, InSkeletalMeshRenderData, RayTracingLODIndex, ExtraRequiredBoneIndices);
+	}
 	switch(PreviousBoneTransformUpdateMode)
 	{
 	case EPreviousBoneTransformUpdateMode::None:
 		// otherwise, clear it, it will use previous buffer
 		PreviousReferenceToLocal.Reset();
+		PreviousReferenceToLocalForRayTracing.Reset();
 		break;
 	case EPreviousBoneTransformUpdateMode::UpdatePrevious:
 		UpdatePreviousRefToLocalMatrices(PreviousReferenceToLocal, InMeshComponent, InSkeletalMeshRenderData, LODIndex, ExtraRequiredBoneIndices);
+		if (RayTracingLODIndex != LODIndex)
+		{
+			UpdatePreviousRefToLocalMatrices(PreviousReferenceToLocalForRayTracing, InMeshComponent, InSkeletalMeshRenderData, RayTracingLODIndex, ExtraRequiredBoneIndices);
+		}
 		break;
 	case EPreviousBoneTransformUpdateMode::DuplicateCurrentToPrevious:
 		UpdateRefToLocalMatrices(PreviousReferenceToLocal, InMeshComponent, InSkeletalMeshRenderData, LODIndex, ExtraRequiredBoneIndices);
+		if (RayTracingLODIndex != LODIndex)
+		{
+			UpdateRefToLocalMatrices(PreviousReferenceToLocalForRayTracing, InMeshComponent, InSkeletalMeshRenderData, RayTracingLODIndex, ExtraRequiredBoneIndices);
+		}
 		break;
 	}
 
diff --git a/Engine/Source/Runtime/Engine/Private/SkeletalRenderGPUSkin.h b/Engine/Source/Runtime/Engine/Private/SkeletalRenderGPUSkin.h
--- a/Engine/Source/Runtime/Engine/Private/SkeletalRenderGPUSkin.h
+++ b/Engine/Source/Runtime/Engine/Private/SkeletalRenderGPUSkin.h
@@ -57,6 +57,7 @@ public:
 	void InitDynamicSkelMeshObjectDataGPUSkin(
 		USkinnedMeshComponent* InMeshComponent,
 		FSkeletalMeshRenderData* InSkeletalMeshRenderData,
+		FSkeletalMeshObjectGPUSkin* InMeshObject,
 		int32 InLODIndex,
 		const TArray<FActiveMorphTarget>& InActiveMorphTargets,
 		const TArray<float>& InMorphTargetWeights, 
@@ -65,9 +66,11 @@ public:
 
 	/** ref pose to local space transforms */
 	TArray<FMatrix44f> ReferenceToLocal;
+	TArray<FMatrix44f> ReferenceToLocalForRayTracing;
 
 	/** Previous ref pose to local space transform */
 	TArray<FMatrix44f> PreviousReferenceToLocal;
+	TArray<FMatrix44f> PreviousReferenceToLocalForRayTracing;
 
 #if !(UE_BUILD_SHIPPING || UE_BUILD_TEST) 
 	/** component space bone transforms*/
@@ -76,12 +79,14 @@ public:
 
 	/** currently LOD for bones being updated */
 	int32 LODIndex;
+	int32 RayTracingLODIndex;
 	/** current morph targets active on this mesh */
 	TArray<FActiveMorphTarget> ActiveMorphTargets;
 	/** All morph target weights on this mesh */
 	TArray<float> MorphTargetWeights;
 	/** All section ID impacted by active morph target on this mesh */
 	TArray<int32> SectionIdsUseByActiveMorphTargets;
+	TArray<int32> SectionIdsUseByActiveMorphTargetsForRayTracing;
 	/** number of active morph targets with weights > 0 */
 	int32 NumWeightedActiveMorphTargets;
 
@@ -369,7 +374,7 @@ public:
 	virtual void Update(int32 LODIndex,USkinnedMeshComponent* InMeshComponent,const TArray<FActiveMorphTarget>& ActiveMorphTargets, const TArray<float>& MorphTargetWeights, EPreviousBoneTransformUpdateMode PreviousBoneTransformUpdateMode) override;
 	void UpdateDynamicData_RenderThread(FGPUSkinCache* GPUSkinCache, FRHICommandListImmediate& RHICmdList, FDynamicSkelMeshObjectDataGPUSkin* InDynamicData, FSceneInterface* Scene, uint32 FrameNumberToPrepare, uint32 RevisionNumber);
 	virtual void PreGDMECallback(FGPUSkinCache* GPUSkinCache, uint32 FrameNumber) override;
-	virtual const FVertexFactory* GetSkinVertexFactory(const FSceneView* View, int32 LODIndex,int32 ChunkIdx) const override;
+	virtual const FVertexFactory* GetSkinVertexFactory(const FSceneView* View, int32 LODIndex,int32 ChunkIdx, ESkinVertexFactoryMode VFMode = ESkinVertexFactoryMode::Default) const override;
 	virtual void CacheVertices(int32 LODIndex, bool bForce) const override {}
 	virtual bool IsCPUSkinned() const override { return false; }
 	virtual TArray<FTransform>* GetComponentSpaceTransforms() const override;
@@ -385,6 +390,18 @@ public:
 
 	/** Return the internal vertex buffer only when initialized otherwise used the shared vertex buffer - needs to be updated every frame */
 	virtual FRWBuffer* GetRayTracingDynamicVertexBuffer() { return RayTracingDynamicVertexBuffer.NumBytes > 0 ? &RayTracingDynamicVertexBuffer : nullptr; }
+
+	virtual int32 GetRayTracingLOD() const override
+	{
+		if (DynamicData)
+		{
+			return DynamicData->RayTracingLODIndex;
+		}
+		else
+		{
+			return 0;
+		}
+	}
 #endif // RHI_RAYTRACING
 
 	virtual int32 GetLOD() const override
@@ -638,7 +655,7 @@ private:
 	*/
 	void ReleaseMorphResources();
 
-	void ProcessUpdatedDynamicData(FGPUSkinCache* GPUSkinCache, FRHICommandListImmediate& RHICmdList, uint32 FrameNumberToPrepare, uint32 RevisionNumber, bool bMorphNeedsUpdate);
+	void ProcessUpdatedDynamicData(EGPUSkinCacheEntryMode Mode, FGPUSkinCache* GPUSkinCache, FRHICommandListImmediate& RHICmdList, uint32 FrameNumberToPrepare, uint32 RevisionNumber, bool bMorphNeedsUpdate, int32 LODIndex);
 
 	void WaitForRHIThreadFenceForDynamicData();
 
diff --git a/Engine/Source/Runtime/Engine/Private/SkeletalRenderStatic.cpp b/Engine/Source/Runtime/Engine/Private/SkeletalRenderStatic.cpp
--- a/Engine/Source/Runtime/Engine/Private/SkeletalRenderStatic.cpp
+++ b/Engine/Source/Runtime/Engine/Private/SkeletalRenderStatic.cpp
@@ -153,7 +153,7 @@ void FSkeletalMeshObjectStatic::ReleaseResources()
 	}
 }
 
-const FVertexFactory* FSkeletalMeshObjectStatic::GetSkinVertexFactory(const FSceneView* View, int32 LODIndex, int32 ChunkIdx) const
+const FVertexFactory* FSkeletalMeshObjectStatic::GetSkinVertexFactory(const FSceneView* View, int32 LODIndex, int32 ChunkIdx, ESkinVertexFactoryMode VFMode) const
 {
 	check(LODs.IsValidIndex(LODIndex));
 	return &LODs[LODIndex].VertexFactory; 
diff --git a/Engine/Source/Runtime/Engine/Private/SkeletalRenderStatic.h b/Engine/Source/Runtime/Engine/Private/SkeletalRenderStatic.h
--- a/Engine/Source/Runtime/Engine/Private/SkeletalRenderStatic.h
+++ b/Engine/Source/Runtime/Engine/Private/SkeletalRenderStatic.h
@@ -31,7 +31,7 @@ public:
 	virtual void EnableOverlayRendering(bool bEnabled, const TArray<int32>* InBonesOfInterest, const TArray<UMorphTarget*>* InMorphTargetOfInterest) override {};
 	virtual void CacheVertices(int32 LODIndex, bool bForce) const override {};
 	virtual bool IsCPUSkinned() const override { return true; }
-	virtual const FVertexFactory* GetSkinVertexFactory(const FSceneView* View, int32 LODIndex, int32 ChunkIdx) const override;
+	virtual const FVertexFactory* GetSkinVertexFactory(const FSceneView* View, int32 LODIndex, int32 ChunkIdx, ESkinVertexFactoryMode VFMode = ESkinVertexFactoryMode::Default) const override;
 	virtual TArray<FTransform>* GetComponentSpaceTransforms() const override;
 	virtual const TArray<FMatrix44f>& GetReferenceToLocalMatrices() const override;
 
diff --git a/Engine/Source/Runtime/Engine/Public/GPUSkinCache.h b/Engine/Source/Runtime/Engine/Public/GPUSkinCache.h
--- a/Engine/Source/Runtime/Engine/Public/GPUSkinCache.h
+++ b/Engine/Source/Runtime/Engine/Public/GPUSkinCache.h
@@ -116,6 +116,12 @@ struct FCachedGeometry
 	FRDGBufferRef DeformedPositionBuffer = nullptr;
 };
 
+enum class EGPUSkinCacheEntryMode
+{
+	Raster,
+	RayTracing
+};
+
 class FGPUSkinCache
 {
 public:
@@ -155,6 +161,7 @@ public:
 	void UpdateSkinWeightBuffer(FGPUSkinCacheEntry* Entry);
 
 	bool ProcessEntry(
+		EGPUSkinCacheEntryMode Mode,
 		FRHICommandListImmediate& RHICmdList, 
 		FGPUBaseSkinVertexFactory* VertexFactory,
 		FGPUSkinPassthroughVertexFactory* TargetVertexFactory, 
@@ -167,7 +174,8 @@ public:
 		const FMatrix44f& ClothLocalToWorld,
 		float ClothBlendWeight, 
 		uint32 RevisionNumber, 
-		int32 Section, 
+		int32 Section,
+		int32 LOD,
 		FGPUSkinCacheEntry*& InOutEntry
 		);
 
diff --git a/Engine/Source/Runtime/Engine/Public/SkeletalMeshTypes.h b/Engine/Source/Runtime/Engine/Public/SkeletalMeshTypes.h
--- a/Engine/Source/Runtime/Engine/Public/SkeletalMeshTypes.h
+++ b/Engine/Source/Runtime/Engine/Public/SkeletalMeshTypes.h
@@ -46,7 +46,11 @@ enum class ESkinVertexColorChannel : uint8
 	None = Alpha
 };
 
-
+enum class ESkinVertexFactoryMode
+{
+	Default,
+	RayTracing
+};
 
 /**
  * A structure for holding mesh-to-mesh triangle influences to skin one mesh to another (similar to a wrap deformer)
@@ -304,7 +308,7 @@ protected:
 	/** Only call on render thread timeline */
 	uint8 GetCurrentFirstLODIdx_Internal() const;
 private:
-	void CreateBaseMeshBatch(const FSceneView* View, const FSkeletalMeshLODRenderData& LODData, const int32 LODIndex, const int32 SectionIndex, const FSectionElementInfo& SectionElementInfo, FMeshBatch& Mesh) const;
+	void CreateBaseMeshBatch(const FSceneView* View, const FSkeletalMeshLODRenderData& LODData, const int32 LODIndex, const int32 SectionIndex, const FSectionElementInfo& SectionElementInfo, FMeshBatch& Mesh, ESkinVertexFactoryMode VFMode = ESkinVertexFactoryMode::Default) const;
 };
 
 /** Used to recreate all skinned mesh components for a given skeletal mesh */
diff --git a/Engine/Source/Runtime/Engine/Public/SkeletalRenderPublic.h b/Engine/Source/Runtime/Engine/Public/SkeletalRenderPublic.h
--- a/Engine/Source/Runtime/Engine/Public/SkeletalRenderPublic.h
+++ b/Engine/Source/Runtime/Engine/Public/SkeletalRenderPublic.h
@@ -84,7 +84,7 @@ public:
 	 * @param	ChunkIdx - not used
 	 * @return	vertex factory for rendering the LOD, 0 to suppress rendering
 	 */
-	virtual const FVertexFactory* GetSkinVertexFactory(const FSceneView* View, int32 LODIndex,int32 ChunkIdx) const = 0;
+	virtual const FVertexFactory* GetSkinVertexFactory(const FSceneView* View, int32 LODIndex,int32 ChunkIdx, ESkinVertexFactoryMode VFMode = ESkinVertexFactoryMode::Default) const = 0;
 
 	/**
 	 * Re-skin cached vertices for an LOD and update the vertex buffer. Note that this
@@ -204,6 +204,7 @@ public:
 	virtual FRayTracingGeometry* GetRayTracingGeometry() { return nullptr; }
 	virtual const FRayTracingGeometry* GetRayTracingGeometry() const { return nullptr; }
 	virtual FRWBuffer* GetRayTracingDynamicVertexBuffer() { return nullptr; }
+	virtual int32 GetRayTracingLOD() const { return GetLOD(); }
 #endif // RHI_RAYTRACING
 
 	/** Called to notify clothing data that component transform has changed */
@@ -245,6 +246,7 @@ public:
 	bool bRequireRecreatingRayTracingGeometry;
 	bool bSupportRayTracing;
 	bool bHiddenMaterialVisibilityDirtyForRayTracing;
+	int32 RayTracingMinLOD;
 #endif
 
 #if UE_BUILD_SHIPPING
@@ -286,6 +288,7 @@ protected:
 	TArray<FSkeletalMeshLODInfo> SkeletalMeshLODInfo;
 
 	class FGPUSkinCacheEntry* SkinCacheEntry;
+	FGPUSkinCacheEntry* SkinCacheEntryForRayTracing;
 
 	/** Used to keep track of the first call to UpdateMinDesiredLODLevel each frame. from ViewFamily.FrameNumber */
 	uint32 LastFrameNumber;
