WorldPartition: Log new streaming sources + teleports for debugging

diff --git a/Engine/Source/Runtime/Engine/Private/WorldPartition/WorldPartitionStreamingPolicy.cpp b/Engine/Source/Runtime/Engine/Private/WorldPartition/WorldPartitionStreamingPolicy.cpp
--- a/Engine/Source/Runtime/Engine/Private/WorldPartition/WorldPartitionStreamingPolicy.cpp
+++ b/Engine/Source/Runtime/Engine/Private/WorldPartition/WorldPartitionStreamingPolicy.cpp
@@ -148,7 +148,7 @@ void UWorldPartitionStreamingPolicy::UpdateStreamingSources()
 		if (!StreamingSource.Name.IsNone())
 		{
 			ValidStreamingSources.Add(StreamingSource.Name);
-			FStreamingSourceVelocity& SourceVelocity = StreamingSourcesVelocity.FindOrAdd(StreamingSource.Name);
+			FStreamingSourceVelocity& SourceVelocity = StreamingSourcesVelocity.FindOrAdd(StreamingSource.Name, FStreamingSourceVelocity(StreamingSource.Name));
 			StreamingSource.Velocity = SourceVelocity.GetAverageVelocity(StreamingSource.Location, CurrentTime);
 		}
 	}
@@ -600,8 +600,9 @@ void UWorldPartitionStreamingPolicy::DrawRuntimeHash3D()
  * FStreamingSourceVelocity Implementation
  */
 
-FStreamingSourceVelocity::FStreamingSourceVelocity()
-: LastIndex(INDEX_NONE)
+FStreamingSourceVelocity::FStreamingSourceVelocity(const FName& InSourceName)
+: SourceName(InSourceName)
+, LastIndex(INDEX_NONE)
 , LastUpdateTime(-1.0)
 , VelocitiesHistorySum(0.f)
 {
@@ -617,6 +618,7 @@ float FStreamingSourceVelocity::GetAverageVelocity(const FVector& NewPosition, c
 	const float Distance = bIsFirstCall ? 0.f : ((NewPosition - LastPosition) * 0.01f).Size();
 	if (bIsFirstCall)
 	{
+		UE_LOG(LogWorldPartition, Log, TEXT("New Streaming Source: %s -> Position: %s"), *SourceName.ToString(), *NewPosition.ToString());
 		LastIndex = 0;
 	}
 
@@ -629,6 +631,7 @@ float FStreamingSourceVelocity::GetAverageVelocity(const FVector& NewPosition, c
 	// Handle invalid cases
 	if (bIsFirstCall || (DeltaSeconds <= 0.f) || (DeltaSeconds > MaxDeltaSeconds) || (Distance > TeleportDistance))
 	{
+		UE_CLOG(Distance > TeleportDistance, LogWorldPartition, Log, TEXT("Detected Streaming Source Teleport: %s -> Last Position: %s -> New Position: %s"), *SourceName.ToString(), *LastPosition.ToString(), *NewPosition.ToString());
 		return 0.f;
 	}
 
diff --git a/Engine/Source/Runtime/Engine/Public/WorldPartition/WorldPartitionStreamingPolicy.h b/Engine/Source/Runtime/Engine/Public/WorldPartition/WorldPartitionStreamingPolicy.h
--- a/Engine/Source/Runtime/Engine/Public/WorldPartition/WorldPartitionStreamingPolicy.h
+++ b/Engine/Source/Runtime/Engine/Public/WorldPartition/WorldPartitionStreamingPolicy.h
@@ -23,11 +23,12 @@ class UWorldPartition;
  */
 struct FStreamingSourceVelocity
 {
-	FStreamingSourceVelocity();
+	FStreamingSourceVelocity(const FName& InSourceName);
 	float GetAverageVelocity(const FVector& NewPosition, const float CurrentTime);
 
 private:
 	enum { VELOCITY_HISTORY_SAMPLE_COUNT = 16 };
+	FName SourceName;
 	int32 LastIndex;
 	float LastUpdateTime;
 	FVector LastPosition;
