Strata better default tangent.
This is in case it is accessed in forward to build local space (can happen for isntance with the new rough refraction model)

diff --git a/Engine/Source/Runtime/Engine/Private/Materials/HLSLMaterialTranslator.cpp b/Engine/Source/Runtime/Engine/Private/Materials/HLSLMaterialTranslator.cpp
--- a/Engine/Source/Runtime/Engine/Private/Materials/HLSLMaterialTranslator.cpp
+++ b/Engine/Source/Runtime/Engine/Private/Materials/HLSLMaterialTranslator.cpp
@@ -1731,7 +1731,7 @@ void FHLSLMaterialTranslator::GetMaterialEnvironment(EShaderPlatform InPlatform,
 				}
 				else
 				{
-					StrataPixelNormalInitializerValues += FString::Printf(TEXT("\tParameters.SharedLocalBases.Tangents[%u] = float3(1,0,0);\n"), FinalSharedLocalBasisIndex);
+					StrataPixelNormalInitializerValues += FString::Printf(TEXT("\tParameters.SharedLocalBases.Tangents[%u] = Parameters.TangentToWorld[0];\n"), FinalSharedLocalBasisIndex);
 				}
 				StrataPixelNormalInitializerValues += FString::Printf(TEXT("\t#if MATERIAL_TANGENTSPACENORMAL\n"));
 				StrataPixelNormalInitializerValues += FString::Printf(TEXT("\tParameters.SharedLocalBases.Tangents[%u] *= Parameters.TwoSidedSign;\n"), FinalSharedLocalBasisIndex);
