Fix CIS issues and non-ray tracing platforms

diff --git a/Engine/Source/Runtime/Engine/Private/GPUSkinCache.cpp b/Engine/Source/Runtime/Engine/Private/GPUSkinCache.cpp
--- a/Engine/Source/Runtime/Engine/Private/GPUSkinCache.cpp
+++ b/Engine/Source/Runtime/Engine/Private/GPUSkinCache.cpp
@@ -223,12 +223,12 @@ class FGPUSkinCacheEntry
 {
 public:
 	FGPUSkinCacheEntry(FGPUSkinCache* InSkinCache, FSkeletalMeshObjectGPUSkin* InGPUSkin, FGPUSkinCache::FRWBuffersAllocation* InPositionAllocation, int32 InLOD, EGPUSkinCacheEntryMode InMode)
-		: PositionAllocation(InPositionAllocation)
+		: Mode(InMode)
+		, PositionAllocation(InPositionAllocation)
 		, SkinCache(InSkinCache)
 		, GPUSkin(InGPUSkin)
 		, MorphBuffer(0)
 		, LOD(InLOD)
-		, Mode(InMode)
 	{
 		
 		const TArray<FSkelMeshRenderSection>& Sections = InGPUSkin->GetRenderSections(LOD);
diff --git a/Engine/Source/Runtime/Engine/Private/SkeletalRenderGPUSkin.cpp b/Engine/Source/Runtime/Engine/Private/SkeletalRenderGPUSkin.cpp
--- a/Engine/Source/Runtime/Engine/Private/SkeletalRenderGPUSkin.cpp
+++ b/Engine/Source/Runtime/Engine/Private/SkeletalRenderGPUSkin.cpp
@@ -434,10 +434,12 @@ void FSkeletalMeshObjectGPUSkin::UpdateDynamicData_RenderThread(FGPUSkinCache* G
 	{
 		ProcessUpdatedDynamicData(EGPUSkinCacheEntryMode::Raster, GPUSkinCache, RHICmdList, FrameNumberToPrepare, RevisionNumber, bMorphNeedsUpdate, DynamicData->LODIndex);
 
+	#if RHI_RAYTRACING
 		if (FGPUSkinCache::IsGPUSkinCacheRayTracingSupported() && GPUSkinCache && SkeletalMeshRenderData->bSupportRayTracing)
 		{
 			ProcessUpdatedDynamicData(EGPUSkinCacheEntryMode::RayTracing, GPUSkinCache, RHICmdList, FrameNumberToPrepare, RevisionNumber, bMorphNeedsUpdate, DynamicData->RayTracingLODIndex);
 		}
+	#endif
 	}
 
 #if RHI_RAYTRACING
@@ -645,7 +647,11 @@ void FSkeletalMeshObjectGPUSkin::ProcessUpdatedDynamicData(EGPUSkinCacheEntryMod
 				}
 			}
 
+		#if RHI_RAYTRACING
 			bool bShouldUseSeparateMatricesForRayTracing = Mode == EGPUSkinCacheEntryMode::RayTracing && DynamicData->RayTracingLODIndex != DynamicData->LODIndex;
+		#else
+			bool bShouldUseSeparateMatricesForRayTracing = false;
+		#endif
 
 			// if we have previous reference to loca, we also update to previous frame
 			if (DynamicData->PreviousReferenceToLocal.Num() > 0)
@@ -2037,14 +2043,18 @@ void FDynamicSkelMeshObjectDataGPUSkin::InitDynamicSkelMeshObjectDataGPUSkin(
 	FSkeletalMeshSceneProxy* SkeletalMeshProxy = (FSkeletalMeshSceneProxy*)InMeshComponent->SceneProxy;
 	const TArray<FBoneIndexType>* ExtraRequiredBoneIndices = SkeletalMeshProxy ? &SkeletalMeshProxy->GetSortedShadowBoneIndices() : nullptr;
 
+#if RHI_RAYTRACING
 	RayTracingLODIndex = FMath::Clamp(FMath::Max(LODIndex + GetRayTracingSkeletalMeshGlobalLODBias(), InMeshObject->RayTracingMinLOD), LODIndex, InSkeletalMeshRenderData->LODRenderData.Num() - 1);
+#endif
 
 	// update ReferenceToLocal
 	UpdateRefToLocalMatrices( ReferenceToLocal, InMeshComponent, InSkeletalMeshRenderData, LODIndex, ExtraRequiredBoneIndices );
+#if RHI_RAYTRACING
 	if (RayTracingLODIndex != LODIndex)
 	{
 		UpdateRefToLocalMatrices(ReferenceToLocalForRayTracing, InMeshComponent, InSkeletalMeshRenderData, RayTracingLODIndex, ExtraRequiredBoneIndices);
 	}
+#endif
 	switch(PreviousBoneTransformUpdateMode)
 	{
 	case EPreviousBoneTransformUpdateMode::None:
@@ -2054,17 +2064,21 @@ void FDynamicSkelMeshObjectDataGPUSkin::InitDynamicSkelMeshObjectDataGPUSkin(
 		break;
 	case EPreviousBoneTransformUpdateMode::UpdatePrevious:
 		UpdatePreviousRefToLocalMatrices(PreviousReferenceToLocal, InMeshComponent, InSkeletalMeshRenderData, LODIndex, ExtraRequiredBoneIndices);
+	#if RHI_RAYTRACING
 		if (RayTracingLODIndex != LODIndex)
 		{
 			UpdatePreviousRefToLocalMatrices(PreviousReferenceToLocalForRayTracing, InMeshComponent, InSkeletalMeshRenderData, RayTracingLODIndex, ExtraRequiredBoneIndices);
 		}
+	#endif
 		break;
 	case EPreviousBoneTransformUpdateMode::DuplicateCurrentToPrevious:
 		UpdateRefToLocalMatrices(PreviousReferenceToLocal, InMeshComponent, InSkeletalMeshRenderData, LODIndex, ExtraRequiredBoneIndices);
+	#if RHI_RAYTRACING
 		if (RayTracingLODIndex != LODIndex)
 		{
 			UpdateRefToLocalMatrices(PreviousReferenceToLocalForRayTracing, InMeshComponent, InSkeletalMeshRenderData, RayTracingLODIndex, ExtraRequiredBoneIndices);
 		}
+	#endif
 		break;
 	}
 
diff --git a/Engine/Source/Runtime/Engine/Private/SkeletalRenderGPUSkin.h b/Engine/Source/Runtime/Engine/Private/SkeletalRenderGPUSkin.h
--- a/Engine/Source/Runtime/Engine/Private/SkeletalRenderGPUSkin.h
+++ b/Engine/Source/Runtime/Engine/Private/SkeletalRenderGPUSkin.h
@@ -79,7 +79,9 @@ public:
 
 	/** currently LOD for bones being updated */
 	int32 LODIndex;
+#if RHI_RAYTRACING
 	int32 RayTracingLODIndex;
+#endif
 	/** current morph targets active on this mesh */
 	TArray<FActiveMorphTarget> ActiveMorphTargets;
 	/** All morph target weights on this mesh */
