Fix CIS issues by changing back to using base class of contraint in user pointer

#okfirgithub public, internal

diff --git a/Engine/Source/Runtime/Engine/Classes/PhysicsEngine/ConstraintInstance.h b/Engine/Source/Runtime/Engine/Classes/PhysicsEngine/ConstraintInstance.h
--- a/Engine/Source/Runtime/Engine/Classes/PhysicsEngine/ConstraintInstance.h
+++ b/Engine/Source/Runtime/Engine/Classes/PhysicsEngine/ConstraintInstance.h
@@ -161,6 +161,17 @@ struct ENGINE_API FConstraintInstanceBase
 
 	FPhysScene* GetPhysicsScene() { return PhysScene; }
 	const FPhysScene* GetPhysicsScene() const { return PhysScene; }
+
+	/** Set the constraint broken delegate. */
+	void SetConstraintBrokenDelegate(FOnConstraintBroken InConstraintBrokenDelegate);
+
+	protected:
+
+		FOnConstraintBroken OnConstraintBrokenDelegate;
+
+		friend struct FConstraintBrokenDelegateData;
+		friend struct FConstraintBrokenDelegateWrapper;
+
 };
 
 /** Container for a physics representation of an object. */
@@ -793,9 +804,6 @@ public:
 	/** Create physics engine constraint using physx actors. */
 	void InitConstraint_AssumesLocked(const FPhysicsActorHandle& ActorRef1, const FPhysicsActorHandle& ActorRef2, float InScale, FOnConstraintBroken InConstraintBrokenDelegate = FOnConstraintBroken());
 
-	/** Set teh constraint broken delegate. */
-	void SetConstraintBrokenDelegate(FOnConstraintBroken InConstraintBrokenDelegate);
-
 	/** Terminate physics engine constraint */
 	void TermConstraint();
 
@@ -931,10 +939,6 @@ private:
 	void UpdatePlasticity();
 	void UpdateDriveTarget();
 
-	FOnConstraintBroken OnConstraintBrokenDelegate;
-
-	friend struct FConstraintBrokenDelegateData;
-
 public:
 
 	///////////////////////////// DEPRECATED
diff --git a/Engine/Source/Runtime/Engine/Private/PhysicsEngine/ConstraintInstance.cpp b/Engine/Source/Runtime/Engine/Private/PhysicsEngine/ConstraintInstance.cpp
--- a/Engine/Source/Runtime/Engine/Private/PhysicsEngine/ConstraintInstance.cpp
+++ b/Engine/Source/Runtime/Engine/Private/PhysicsEngine/ConstraintInstance.cpp
@@ -219,6 +219,11 @@ void FConstraintInstanceBase::Reset()
 	PhysScene = nullptr;
 }
 
+void FConstraintInstanceBase::SetConstraintBrokenDelegate(FOnConstraintBroken InConstraintBrokenDelegate)
+{
+	OnConstraintBrokenDelegate = InConstraintBrokenDelegate;
+}
+
 
 /** Constructor **/
 FConstraintInstance::FConstraintInstance()
@@ -514,11 +519,6 @@ void FConstraintInstance::InitConstraint_AssumesLocked(const FPhysicsActorHandle
 	}
 }
 
-void FConstraintInstance::SetConstraintBrokenDelegate(FOnConstraintBroken InConstraintBrokenDelegate)
-{
-	OnConstraintBrokenDelegate = InConstraintBrokenDelegate;
-}
-
 void FConstraintProfileProperties::Update_AssumesLocked(const FPhysicsConstraintHandle& InConstraintRef, float AverageMass, float UseScale) const
 {
 	// flags and projection settings
diff --git a/Engine/Source/Runtime/Engine/Private/PhysicsEngine/Experimental/PhysInterface_Chaos.cpp b/Engine/Source/Runtime/Engine/Private/PhysicsEngine/Experimental/PhysInterface_Chaos.cpp
--- a/Engine/Source/Runtime/Engine/Private/PhysicsEngine/Experimental/PhysInterface_Chaos.cpp
+++ b/Engine/Source/Runtime/Engine/Private/PhysicsEngine/Experimental/PhysInterface_Chaos.cpp
@@ -637,7 +637,7 @@ private:
 		{
 			Chaos::FJointConstraint* Constraint = static_cast<Chaos::FJointConstraint*>(InConstraintHandle->Constraint);
 
-			FConstraintInstance* ConstraintInstance = (Constraint) ? FPhysicsUserData_Chaos::Get<FConstraintInstance>(Constraint->GetUserData()) : nullptr;
+			FConstraintInstanceBase* ConstraintInstance = (Constraint) ? FPhysicsUserData_Chaos::Get<FConstraintInstanceBase>(Constraint->GetUserData()) : nullptr;
 			if (ConstraintInstance)
 			{
 				return ConstraintInstance->GetPhysicsScene();
diff --git a/Engine/Source/Runtime/Engine/Private/PhysicsEngine/Experimental/PhysScene_Chaos.cpp b/Engine/Source/Runtime/Engine/Private/PhysicsEngine/Experimental/PhysScene_Chaos.cpp
--- a/Engine/Source/Runtime/Engine/Private/PhysicsEngine/Experimental/PhysScene_Chaos.cpp
+++ b/Engine/Source/Runtime/Engine/Private/PhysicsEngine/Experimental/PhysScene_Chaos.cpp
@@ -37,6 +37,7 @@
 #include "Chaos/BoundingVolume.h"
 #include "Chaos/Framework/DebugSubstep.h"
 #include "Chaos/PBDSpringConstraints.h"
+#include "Chaos/PBDJointConstraints.h"
 #include "Chaos/PerParticleGravity.h"
 #include "PBDRigidActiveParticlesBuffer.h"
 #include "Chaos/GeometryParticlesfwd.h"
@@ -1662,10 +1663,10 @@ void FPhysScene_Chaos::OnSyncBodies(Chaos::FPhysicsSolverBase* Solver)
 
 		if (Constraint->GetOutputData().bIsBreaking)
 		{
-			FConstraintInstance* ConstraintInstance = (Constraint) ? FPhysicsUserData_Chaos::Get<FConstraintInstance>(Constraint->GetUserData()) : nullptr;
+			if (FConstraintInstanceBase* ConstraintInstance = (Constraint) ? FPhysicsUserData_Chaos::Get<FConstraintInstanceBase>(Constraint->GetUserData()) : nullptr)
 			{
-				FConstraintBrokenDelegateData CBDD(ConstraintInstance);
-				CBDD.DispatchOnBroken();
+				FConstraintBrokenDelegateWrapper CBD(ConstraintInstance);
+				CBD.DispatchOnBroken();
 			}
 
 			Constraint->GetOutputData().bIsBreaking = false;
@@ -1715,13 +1716,18 @@ void FPhysScene_Chaos::RemoveSpringConstraint(const FPhysicsConstraintHandle& Co
 	// #todo : Implement
 }
 
-FConstraintBrokenDelegateData::FConstraintBrokenDelegateData(FConstraintInstance* ConstraintInstance)
+FConstraintBrokenDelegateWrapper::FConstraintBrokenDelegateWrapper(FConstraintInstanceBase* ConstraintInstance)
 	: OnConstraintBrokenDelegate(ConstraintInstance->OnConstraintBrokenDelegate)
 	, ConstraintIndex(ConstraintInstance->ConstraintIndex)
 {
 
 }
 
+void FConstraintBrokenDelegateWrapper::DispatchOnBroken()
+{
+	OnConstraintBrokenDelegate.ExecuteIfBound(ConstraintIndex);
+}
+
 void FPhysScene_Chaos::ResimNFrames(const int32 NumFramesRequested)
 {
 	//needs to run on physics thread from a special location
diff --git a/Engine/Source/Runtime/Engine/Public/Physics/Experimental/PhysScene_Chaos.h b/Engine/Source/Runtime/Engine/Public/Physics/Experimental/PhysScene_Chaos.h
--- a/Engine/Source/Runtime/Engine/Public/Physics/Experimental/PhysScene_Chaos.h
+++ b/Engine/Source/Runtime/Engine/Public/Physics/Experimental/PhysScene_Chaos.h
@@ -37,6 +37,7 @@ class UWorld;
 class AWorldSettings;
 class FPhysicsReplicationFactory;
 class FContactModifyCallbackFactory;
+struct FConstraintInstanceBase;
 
 namespace Chaos
 {
@@ -65,6 +66,17 @@ namespace Chaos
 
 extern int32 GEnableKinematicDeferralStartPhysicsCondition;
 
+struct FConstraintBrokenDelegateWrapper
+{
+	FConstraintBrokenDelegateWrapper(FConstraintInstanceBase* ConstraintInstance);
+
+	void DispatchOnBroken();
+
+	FOnConstraintBroken OnConstraintBrokenDelegate;
+	int32 ConstraintIndex;
+};
+
+
 /**
 * Low level Chaos scene used when building custom simulations that don't exist in the main world physics scene.
 */
diff --git a/Engine/Source/Runtime/Engine/Public/Physics/Experimental/PhysicsUserData_Chaos.h b/Engine/Source/Runtime/Engine/Public/Physics/Experimental/PhysicsUserData_Chaos.h
--- a/Engine/Source/Runtime/Engine/Public/Physics/Experimental/PhysicsUserData_Chaos.h
+++ b/Engine/Source/Runtime/Engine/Public/Physics/Experimental/PhysicsUserData_Chaos.h
@@ -5,7 +5,7 @@
 #include "EngineDefines.h"
 #include "Physics/PhysicsInterfaceDeclares.h"
 
-struct FConstraintInstance;
+struct FConstraintInstanceBase;
 
 // Note: this must match the enum EChaosUserDataType in PhysicsInterfaceTypesCore
 class FPhysicsUserData_Chaos
@@ -17,7 +17,7 @@ private:
 		BodyInstance,
 		PhysicalMaterial,
 		PhysScene,
-		ConstraintInstance
+		ConstraintInstanceBase
 	};
 
 	EType Type;
@@ -28,7 +28,7 @@ public:
 	FPhysicsUserData_Chaos(FBodyInstance* InPayload) : Type(EType::BodyInstance), Payload(InPayload) { }
 	FPhysicsUserData_Chaos(UPhysicalMaterial* InPayload) : Type(EType::PhysicalMaterial), Payload(InPayload) { }
 	FPhysicsUserData_Chaos(FPhysScene* InPayload) : Type(EType::PhysScene), Payload(InPayload) { }
-	FPhysicsUserData_Chaos(FConstraintInstance* InPayload) : Type(EType::ConstraintInstance), Payload(InPayload) { }
+	FPhysicsUserData_Chaos(FConstraintInstanceBase* InPayload) : Type(EType::ConstraintInstanceBase), Payload(InPayload) { }
 
 	template <class T> static T* Get(void* UserData);
 	template <class T> static void Set(void* UserData, T* InPayload);
@@ -69,7 +69,7 @@ template<> FORCEINLINE FPhysScene* FPhysicsUserData_Chaos::Get(void* UserData)
 	return Get<FPhysScene, EType::PhysScene>(UserData);
 }
 
-template<> FORCEINLINE FConstraintInstance* FPhysicsUserData_Chaos::Get(void* UserData)
+template<> FORCEINLINE FConstraintInstanceBase* FPhysicsUserData_Chaos::Get(void* UserData)
 {
-	return Get<FConstraintInstance, EType::ConstraintInstance>(UserData);
+	return Get<FConstraintInstanceBase, EType::ConstraintInstanceBase>(UserData);
 }
