Optimus: Initialized some uninitialized struct members to fix a startup warning (soon to be an error)

diff --git a/Engine/Source/Runtime/Engine/Public/ComputeFramework/ComputeGraph.h b/Engine/Source/Runtime/Engine/Public/ComputeFramework/ComputeGraph.h
--- a/Engine/Source/Runtime/Engine/Public/ComputeFramework/ComputeGraph.h
+++ b/Engine/Source/Runtime/Engine/Public/ComputeFramework/ComputeGraph.h
@@ -39,8 +39,6 @@ struct FComputeGraphEdge
 {
 	GENERATED_BODY()
 
-	UPROPERTY()
-	bool bKernelInput;
 	UPROPERTY()
 	int32 KernelIndex;
 	UPROPERTY()
@@ -49,11 +47,21 @@ struct FComputeGraphEdge
 	int32 DataInterfaceIndex;
 	UPROPERTY()
 	int32 DataInterfaceBindingIndex;
+	UPROPERTY()
+	bool bKernelInput;
 
 	// Optional name to use for the proxy generation function, in case the kernel expects
 	// something other than the interface's bind name. Leave empty to go with the default. 
 	UPROPERTY()
 	FString BindingFunctionNameOverride;
+
+	FComputeGraphEdge()
+		: KernelIndex(0)
+		, KernelBindingIndex(0)
+		, DataInterfaceIndex(0)
+		, DataInterfaceBindingIndex(0)
+		, bKernelInput(false)
+	{}
 };
 
 /** 
diff --git a/Engine/Source/Runtime/Engine/Public/ComputeFramework/ShaderParamTypeDefinition.h b/Engine/Source/Runtime/Engine/Public/ComputeFramework/ShaderParamTypeDefinition.h
--- a/Engine/Source/Runtime/Engine/Public/ComputeFramework/ShaderParamTypeDefinition.h
+++ b/Engine/Source/Runtime/Engine/Public/ComputeFramework/ShaderParamTypeDefinition.h
@@ -162,7 +162,7 @@ struct ENGINE_API FShaderValueType
 	FString GetTypeDeclaration() const;
 
 	UPROPERTY()
-	EShaderFundamentalType Type;
+	EShaderFundamentalType Type = EShaderFundamentalType::Bool;
 
 	UPROPERTY()
 	EShaderFundamentalDimensionType DimensionType = EShaderFundamentalDimensionType::Scalar;
@@ -232,19 +232,19 @@ public:
 	FShaderValueTypeHandle ValueType;
 
 	UPROPERTY()
-	uint16 ArrayElementCount; // 0 indicates not an array. >= 1 indicates an array
+	uint16 ArrayElementCount = 0; // 0 indicates not an array. >= 1 indicates an array
 
 	UPROPERTY()
-	EShaderFundamentalType FundamentalType;
+	EShaderFundamentalType FundamentalType = EShaderFundamentalType::Bool;
 
 	UPROPERTY()
-	EShaderFundamentalDimensionType	DimType;
+	EShaderFundamentalDimensionType	DimType = EShaderFundamentalDimensionType::Scalar;
 
 	UPROPERTY()
-	EShaderParamBindingType BindingType;
+	EShaderParamBindingType BindingType = EShaderParamBindingType::None;
 
 	UPROPERTY()
-	EShaderResourceType	ResourceType;
+	EShaderResourceType	ResourceType = EShaderResourceType::None;
 
 	union
 	{
@@ -326,5 +326,5 @@ public:
 
 	/** If set to true then the first element of ParamTypes is the return type. Otherwise return type is void. */
 	UPROPERTY()
-	bool bHasReturnType;
+	bool bHasReturnType = false;
 };
