Don't apply lod bias to clamp at large texture size for virtual textures

diff --git a/Engine/Source/Runtime/Engine/Private/TextureLODSettings.cpp b/Engine/Source/Runtime/Engine/Private/TextureLODSettings.cpp
--- a/Engine/Source/Runtime/Engine/Private/TextureLODSettings.cpp
+++ b/Engine/Source/Runtime/Engine/Private/TextureLODSettings.cpp
@@ -184,7 +184,7 @@ int32 UTextureLODSettings::CalculateLODBias(int32 Width, int32 Height, int32 Max
 
 	// Don't apply LOD group mip clamping to VT, since VT max size isn't constrained by any GPU capabilities
 	const int32 MinLOD	= bVirtualTexture ? 0 : LODGroupInfo.MinLODMipCount;
-	const int32 MaxLOD	= bVirtualTexture ? VIRTUALTEXTURE_DATA_MAXMIPS : LODGroupInfo.MaxLODMipCount;
+	const int32 MaxLOD	= bVirtualTexture ? TextureMaxLOD : LODGroupInfo.MaxLODMipCount;
 
 	int32 WantedMaxLOD	= FMath::Clamp( TextureMaxLOD - UsedLODBias, MinLOD, MaxLOD );
 	WantedMaxLOD		= FMath::Clamp( WantedMaxLOD, 0, TextureMaxLOD );
