Make sure ULevel::FixupForPIE use serialization to discover inner objects instead of going through GetObjectsWithOuter, as it skips async loading objects in some cases when instancing levels.

diff --git a/Engine/Source/Runtime/Engine/Private/Level.cpp b/Engine/Source/Runtime/Engine/Private/Level.cpp
--- a/Engine/Source/Runtime/Engine/Private/Level.cpp
+++ b/Engine/Source/Runtime/Engine/Private/Level.cpp
@@ -2760,8 +2760,9 @@ void ULevel::FixupForPIE(int32 InPIEInstanceID, TFunctionRef<void(int32, FSoftOb
 {
 	struct FSoftPathPIEFixupSerializer : public FArchiveUObject
 	{
-		FSoftPathPIEFixupSerializer(int32 InPIEInstanceID, TFunctionRef<void(int32, FSoftObjectPath&)> InCustomFixupFunction)
-			: PIEInstanceID(InPIEInstanceID)
+		FSoftPathPIEFixupSerializer(UObject* InRoot, int32 InPIEInstanceID, TFunctionRef<void(int32, FSoftObjectPath&)> InCustomFixupFunction)
+			: Root(InRoot)
+			, PIEInstanceID(InPIEInstanceID)
 			, CustomFixupFunction(InCustomFixupFunction)
 		{
 			this->SetIsSaving(true);
@@ -2773,29 +2774,36 @@ void ULevel::FixupForPIE(int32 InPIEInstanceID, TFunctionRef<void(int32, FSoftOb
 			return InProperty->IsA<FMulticastDelegateProperty>() || FArchiveUObject::ShouldSkipProperty(InProperty);
 		}
 
+		virtual FArchive& operator<<(UObject*& Object) override
+		{
+			if (Object && Object->IsIn(Root) && !VisitedObjects.Contains(Object))
+			{
+				VisitedObjects.Add(Object);
+
+				// Skip instanced static mesh component as their impact on serialization is enormous and they don't contain lazy ptrs.
+				if (!Cast<UInstancedStaticMeshComponent>(Object))
+				{
+					Object->Serialize(*this);
+				}
+			}
+			return *this;
+		}
+
 		FArchive& operator<<(FSoftObjectPath& Value)
 		{
 			Value.FixupForPIE(PIEInstanceID, CustomFixupFunction);
 			return *this;
 		}
 
+		UObject* Root;
+		TSet<UObject*> VisitedObjects;
+
 		int32 PIEInstanceID;
 		TFunctionRef<void(int32, FSoftObjectPath&)> CustomFixupFunction;
 	};
 
-	FSoftPathPIEFixupSerializer FixupSerializer(InPIEInstanceID, InCustomFixupFunction);
-
-	TArray<UObject*> SubObjects;
-	GetObjectsWithOuter(this, SubObjects);
-
-	for (UObject* Object : SubObjects)
-	{
-		// Skip instanced static mesh component as their impact on serialization is enormous and they don't contain soft ptrs.
-		if (!Cast<UInstancedStaticMeshComponent>(Object))
-		{
-			Object->Serialize(FixupSerializer);
-		}
-	}
+	FSoftPathPIEFixupSerializer FixupSerializer(this, InPIEInstanceID, InCustomFixupFunction);
+	Serialize(FixupSerializer);
 }
 
 #endif	//WITH_EDITOR
