Fix for explicit shader code compression for in process shader compilation.  Addresses crash with Niagara shaders when r.Shaders.AllowCompilingThroughWorkers=0.

diff --git a/Engine/Source/Runtime/Engine/Private/ShaderCompiler/ShaderCompiler.cpp b/Engine/Source/Runtime/Engine/Private/ShaderCompiler/ShaderCompiler.cpp
--- a/Engine/Source/Runtime/Engine/Private/ShaderCompiler/ShaderCompiler.cpp
+++ b/Engine/Source/Runtime/Engine/Private/ShaderCompiler/ShaderCompiler.cpp
@@ -2308,7 +2308,10 @@ void FShaderCompileUtilities::ExecuteShaderCompileJob(FShaderCommonCompileJob& J
 			// Generate a hash of the output and compress the code
 			// The shader processing this output will use the heash to search for existing FShaderResources
 			SingleJob->Output.GenerateOutputHash();
-			SingleJob->Output.CompressOutput(NAME_LZ4);
+			if (SingleJob->Input.CompressionFormat != NAME_None)
+			{
+				SingleJob->Output.CompressOutput(SingleJob->Input.CompressionFormat);
+			}
 		}
 	}
 	else
