Implements r.AntiAliasingMethod

diff --git a/Engine/Plugins/MovieScene/MovieRenderPipeline/Source/MovieRenderPipelineCore/Private/MoviePipelineAntiAliasingSetting.cpp b/Engine/Plugins/MovieScene/MovieRenderPipeline/Source/MovieRenderPipelineCore/Private/MoviePipelineAntiAliasingSetting.cpp
--- a/Engine/Plugins/MovieScene/MovieRenderPipeline/Source/MovieRenderPipelineCore/Private/MoviePipelineAntiAliasingSetting.cpp
+++ b/Engine/Plugins/MovieScene/MovieRenderPipeline/Source/MovieRenderPipelineCore/Private/MoviePipelineAntiAliasingSetting.cpp
@@ -13,7 +13,7 @@ void UMoviePipelineAntiAliasingSetting::ValidateStateImpl()
 		NumTAASamples = AntiAliasingCVar->GetInt();
 	}
 	
-	if (UE::MovieRenderPipeline::GetEffectiveAntiAliasingMethod(this) == EAntiAliasingMethod::AAM_TemporalAA)
+	if (IsTemporalAccumulationBasedMethod(UE::MovieRenderPipeline::GetEffectiveAntiAliasingMethod(this)))
 	{
 		if ((TemporalSampleCount*SpatialSampleCount) > NumTAASamples)
 		{
diff --git a/Engine/Plugins/MovieScene/MovieRenderPipeline/Source/MovieRenderPipelineCore/Private/MoviePipelineRendering.cpp b/Engine/Plugins/MovieScene/MovieRenderPipeline/Source/MovieRenderPipelineCore/Private/MoviePipelineRendering.cpp
--- a/Engine/Plugins/MovieScene/MovieRenderPipeline/Source/MovieRenderPipelineCore/Private/MoviePipelineRendering.cpp
+++ b/Engine/Plugins/MovieScene/MovieRenderPipeline/Source/MovieRenderPipelineCore/Private/MoviePipelineRendering.cpp
@@ -293,7 +293,7 @@ void UMoviePipeline::RenderFrame()
 
 				// Now to check if we have to force it off (at which point we warn the user).
 				bool bMultipleTiles = (TileCount.X > 1) || (TileCount.Y > 1);
-				if (bMultipleTiles && AntiAliasingMethod == EAntiAliasingMethod::AAM_TemporalAA)
+				if (bMultipleTiles && IsTemporalAccumulationBasedMethod(AntiAliasingMethod))
 				{
 					// Temporal Anti-Aliasing isn't supported when using tiled rendering because it relies on having history, and
 					// the tiles use the previous tile as the history which is incorrect. 
diff --git a/Engine/Plugins/MovieScene/MovieRenderPipeline/Source/MovieRenderPipelineCore/Private/MoviePipelineUtils.cpp b/Engine/Plugins/MovieScene/MovieRenderPipeline/Source/MovieRenderPipelineCore/Private/MoviePipelineUtils.cpp
--- a/Engine/Plugins/MovieScene/MovieRenderPipeline/Source/MovieRenderPipelineCore/Private/MoviePipelineUtils.cpp
+++ b/Engine/Plugins/MovieScene/MovieRenderPipeline/Source/MovieRenderPipelineCore/Private/MoviePipelineUtils.cpp
@@ -67,7 +67,7 @@ namespace UE
 		{
 			EAntiAliasingMethod AntiAliasingMethod = EAntiAliasingMethod::AAM_None;
 
-			IConsoleVariable* AntiAliasingCVar = IConsoleManager::Get().FindConsoleVariable(TEXT("r.DefaultFeature.AntiAliasing"));
+			IConsoleVariable* AntiAliasingCVar = IConsoleManager::Get().FindConsoleVariable(TEXT("r.AntiAliasingMethod"));
 			if (AntiAliasingCVar)
 			{
 				int32 Value = AntiAliasingCVar->GetInt();
diff --git a/Engine/Plugins/MovieScene/MovieRenderPipeline/Source/MovieRenderPipelineRenderPasses/Private/MoviePipelineImagePassBase.cpp b/Engine/Plugins/MovieScene/MovieRenderPipeline/Source/MovieRenderPipelineRenderPasses/Private/MoviePipelineImagePassBase.cpp
--- a/Engine/Plugins/MovieScene/MovieRenderPipeline/Source/MovieRenderPipelineRenderPasses/Private/MoviePipelineImagePassBase.cpp
+++ b/Engine/Plugins/MovieScene/MovieRenderPipeline/Source/MovieRenderPipelineRenderPasses/Private/MoviePipelineImagePassBase.cpp
@@ -159,7 +159,7 @@ TSharedPtr<FSceneViewFamilyContext> UMoviePipelineImagePassBase::CalculateViewFa
 	{
 		// If we're not using Temporal Anti-Aliasing or Path Tracing we will apply the View Matrix projection jitter. Normally TAA sets this
 		// inside FSceneRenderer::PreVisibilityFrameSetup. Path Tracing does its own anti-aliasing internally.
-		if (View->AntiAliasingMethod != EAntiAliasingMethod::AAM_TemporalAA && !OutViewFamily->EngineShowFlags.PathTracing)
+		if (!IsTemporalAccumulationBasedMethod(View->AntiAliasingMethod) && !OutViewFamily->EngineShowFlags.PathTracing)
 		{
 			View->ViewMatrices.HackAddTemporalAAProjectionJitter(InOutSampleState.ProjectionMatrixJitterAmount);
 		}
diff --git a/Engine/Plugins/Runtime/Oculus/OculusVR/Source/OculusEditor/Private/OculusToolWidget.cpp b/Engine/Plugins/Runtime/Oculus/OculusVR/Source/OculusEditor/Private/OculusToolWidget.cpp
--- a/Engine/Plugins/Runtime/Oculus/OculusVR/Source/OculusEditor/Private/OculusToolWidget.cpp
+++ b/Engine/Plugins/Runtime/Oculus/OculusVR/Source/OculusEditor/Private/OculusToolWidget.cpp
@@ -799,7 +799,7 @@ EVisibility SOculusToolWidget::AndroidQuestArchVisibility(FName tag) const
 FReply SOculusToolWidget::AntiAliasingEnable(bool text)
 {
 	URendererSettings* Settings = GetMutableDefault<URendererSettings>();
-	Settings->DefaultFeatureAntiAliasing = EAntiAliasingMethod::AAM_TemporalAA;
+	Settings->DefaultFeatureAntiAliasing = EAntiAliasingMethod::AAM_TemporalAA; // TODO(TSR)
 	Settings->UpdateSinglePropertyInConfigFile(Settings->GetClass()->FindPropertyByName(GET_MEMBER_NAME_CHECKED(URendererSettings, DefaultFeatureAntiAliasing)), Settings->GetDefaultConfigFilename());
 	return FReply::Handled();
 }
diff --git a/Engine/Plugins/Runtime/Oculus/OculusVR/Source/OculusHMD/Private/OculusHMD_DynamicResolutionState.cpp b/Engine/Plugins/Runtime/Oculus/OculusVR/Source/OculusHMD/Private/OculusHMD_DynamicResolutionState.cpp
--- a/Engine/Plugins/Runtime/Oculus/OculusVR/Source/OculusHMD/Private/OculusHMD_DynamicResolutionState.cpp
+++ b/Engine/Plugins/Runtime/Oculus/OculusVR/Source/OculusHMD/Private/OculusHMD_DynamicResolutionState.cpp
@@ -105,8 +105,12 @@ void FDynamicResolutionState::SetupMainViewFamily(class FSceneViewFamily& ViewFa
 		MinResolutionFraction = FMath::Max(MinResolutionFraction, FSceneViewScreenPercentageConfig::kMinResolutionFraction);
 		MaxResolutionFraction = FMath::Min(MaxResolutionFraction, FSceneViewScreenPercentageConfig::kMaxResolutionFraction);
 
-		// Temporal upsample has a smaller resolution fraction range.
-		if (View.AntiAliasingMethod == AAM_TemporalAA)
+		if (View.AntiAliasingMethod == AAM_TSR)
+		{
+			MinResolutionFraction = FMath::Max(MinResolutionFraction, FSceneViewScreenPercentageConfig::kMinTSRResolutionFraction);
+			MaxResolutionFraction = FMath::Min(MaxResolutionFraction, FSceneViewScreenPercentageConfig::kMaxTSRResolutionFraction);
+		}
+		else if (View.AntiAliasingMethod == AAM_TemporalAA)
 		{
 			MinResolutionFraction = FMath::Max(MinResolutionFraction, FSceneViewScreenPercentageConfig::kMinTAAUpsampleResolutionFraction);
 			MaxResolutionFraction = FMath::Min(MaxResolutionFraction, FSceneViewScreenPercentageConfig::kMaxTAAUpsampleResolutionFraction);
diff --git a/Engine/Source/Developer/FunctionalTesting/Private/AutomationBlueprintFunctionLibrary.cpp b/Engine/Source/Developer/FunctionalTesting/Private/AutomationBlueprintFunctionLibrary.cpp
--- a/Engine/Source/Developer/FunctionalTesting/Private/AutomationBlueprintFunctionLibrary.cpp
+++ b/Engine/Source/Developer/FunctionalTesting/Private/AutomationBlueprintFunctionLibrary.cpp
@@ -232,7 +232,7 @@ private:
 };
 
 FAutomationTestScreenshotEnvSetup::FAutomationTestScreenshotEnvSetup()
-	: DefaultFeature_AntiAliasing(TEXT("r.DefaultFeature.AntiAliasing"))
+	: DefaultFeature_AntiAliasing(TEXT("r.AntiAliasingMethod"))
 	, DefaultFeature_AutoExposure(TEXT("r.DefaultFeature.AutoExposure"))
 	, DefaultFeature_MotionBlur(TEXT("r.DefaultFeature.MotionBlur"))
 	, PostProcessAAQuality(TEXT("r.PostProcessAAQuality"))
diff --git a/Engine/Source/Editor/HardwareTargeting/Private/HardwareTargetingModule.cpp b/Engine/Source/Editor/HardwareTargeting/Private/HardwareTargetingModule.cpp
--- a/Engine/Source/Editor/HardwareTargeting/Private/HardwareTargetingModule.cpp
+++ b/Engine/Source/Editor/HardwareTargeting/Private/HardwareTargetingModule.cpp
@@ -109,11 +109,13 @@ FText FMetaSettingGatherer::ValueToString(EAntiAliasingMethod Value)
 	case AAM_None:
 		return LOCTEXT("AA_None", "None");
 	case AAM_FXAA:
-		return LOCTEXT("AA_FXAA", "FXAA");
+		return LOCTEXT("AA_FXAA", "Fast Approximate Anti-Aliasing (FXAA)");
 	case AAM_TemporalAA:
-		return LOCTEXT("AA_TemporalAA", "Temporal AA");
+		return LOCTEXT("AA_TemporalAA", "Temporal Anti-Aliasing (TAA)");
 	case AAM_MSAA:
-		return LOCTEXT("AA_MSAA", "MSAA");
+		return LOCTEXT("AA_MSAA", "Multisample Anti-Aliasing (MSAA)");
+	case AAM_TSR:
+		return LOCTEXT("AAM_TSR", "Temporal Super-Resolution (TSR)");
 	default:
 		return FText::AsNumber((int32)Value);
 	}
@@ -239,7 +241,7 @@ void FHardwareTargetingModule::GatherSettings(FMetaSettingGatherer& Builder)
 		// DOF and AA work on mobile but are expensive, keeping them off by default
 		//@TODO: DOF setting doesn't exist yet
 		// UE_META_SETTING_ENTRY(Builder, URendererSettings, bDefaultFeatureDepthOfField, bHighEndPC);
-		UE_META_SETTING_ENTRY(Builder, URendererSettings, DefaultFeatureAntiAliasing, bHighEndPC ? AAM_TemporalAA : AAM_None);
+		UE_META_SETTING_ENTRY(Builder, URendererSettings, DefaultFeatureAntiAliasing, bHighEndPC ? AAM_TSR : AAM_None);
 	}
 
 	{
diff --git a/Engine/Source/Runtime/Engine/Classes/Engine/RendererSettings.h b/Engine/Source/Runtime/Engine/Classes/Engine/RendererSettings.h
--- a/Engine/Source/Runtime/Engine/Classes/Engine/RendererSettings.h
+++ b/Engine/Source/Runtime/Engine/Classes/Engine/RendererSettings.h
@@ -537,14 +537,8 @@ class ENGINE_API URendererSettings : public UDeveloperSettings
 	uint32 bTemporalUpsampling : 1;
 
 	UPROPERTY(config, EditAnywhere, Category = DefaultSettings, meta = (
-		EditCondition = "DefaultFeatureAntiAliasing == EAntiAliasingMethod::AAM_TemporalAA",
-		ConsoleVariable = "r.TemporalAA.Algorithm", DisplayName = "Temporal Super Resolution",
-		ToolTip = "Enable UE5's Temporal Super Resolution algorithm. This new technique is mainly built for handling Nanite amount of details on screen and higher output fidelity with lower ScreenPercentage/Dynamic Resolution on desktops and consoles."))
-	uint32 bUseGen5TAA : 1;
-
-	UPROPERTY(config, EditAnywhere, Category = DefaultSettings, meta = (
-		ConsoleVariable = "r.DefaultFeature.AntiAliasing", DisplayName = "Anti-Aliasing Method",
-		ToolTip = "Which anti-aliasing mode is used by default"))
+		ConsoleVariable = "r.AntiAliasingMethod", DisplayName = "Anti-Aliasing Method",
+		ToolTip = "Selects the anti-aliasing method to use."))
 	TEnumAsByte<EAntiAliasingMethod> DefaultFeatureAntiAliasing;
 
 	UPROPERTY(config, EditAnywhere, Category = DefaultSettings, meta = (
diff --git a/Engine/Source/Runtime/Engine/Private/SceneUtils.cpp b/Engine/Source/Runtime/Engine/Private/SceneUtils.cpp
--- a/Engine/Source/Runtime/Engine/Private/SceneUtils.cpp
+++ b/Engine/Source/Runtime/Engine/Private/SceneUtils.cpp
@@ -61,8 +61,8 @@ ENGINE_API EAntiAliasingMethod GetDefaultAntiAliasingMethod(const FStaticFeature
 	}
 	else
 	{
-		static auto* DefaultAntiAliasingCvar = IConsoleManager::Get().FindTConsoleVariableDataInt(TEXT("r.DefaultFeature.AntiAliasing"));
-		AntiAliasingMethod = (EAntiAliasingMethod)DefaultAntiAliasingCvar->GetValueOnAnyThread();
+		static auto* DefaultAntiAliasingCvar = IConsoleManager::Get().FindTConsoleVariableDataInt(TEXT("r.AntiAliasingMethod"));
+		AntiAliasingMethod = EAntiAliasingMethod(FMath::Clamp(DefaultAntiAliasingCvar->GetValueOnAnyThread(), 0, AAM_MAX));
 	}
 
 	static auto* MSAACountCVar = IConsoleManager::Get().FindTConsoleVariableDataInt(TEXT("r.MSAACount"));
@@ -71,12 +71,37 @@ ENGINE_API EAntiAliasingMethod GetDefaultAntiAliasingMethod(const FStaticFeature
 
 	EShaderPlatform ShaderPlatform = GetFeatureLevelShaderPlatform(InFeatureLevel);
 
-	if (AntiAliasingMethod == EAntiAliasingMethod::AAM_MSAA)
+	if (AntiAliasingMethod == EAntiAliasingMethod::AAM_None)
+	{
+		// NOP
+	}
+	else if (AntiAliasingMethod == EAntiAliasingMethod::AAM_FXAA)
+	{
+		int32 Quality = PostProcessAAQualityCVar->GetValueOnAnyThread();
+
+		if (Quality <= 0)
+		{
+			AntiAliasingMethod = EAntiAliasingMethod::AAM_None;
+		}
+	}
+	else if (AntiAliasingMethod == EAntiAliasingMethod::AAM_TemporalAA)
+	{
+		int32 Quality = PostProcessAAQualityCVar->GetValueOnAnyThread();
+
+		if (Quality <= 0)
+		{
+			AntiAliasingMethod = EAntiAliasingMethod::AAM_None;
+		}
+		else if (Quality == 1 || Quality == 2)
+		{
+			AntiAliasingMethod = AAM_FXAA;
+		}
+	}
+	else if (AntiAliasingMethod == EAntiAliasingMethod::AAM_MSAA)
 	{
 		if (MSAACountCVar->GetValueOnAnyThread() <= 0)
 		{
-			// Fallback to temporal AA so we can easily toggle methods with r.MSAACount
-			AntiAliasingMethod = EAntiAliasingMethod::AAM_TemporalAA;
+			AntiAliasingMethod = EAntiAliasingMethod::AAM_None;
 		}
 		else if ((InFeatureLevel >= ERHIFeatureLevel::SM5 && !IsForwardShadingEnabled(ShaderPlatform))
 			|| (IsMobilePlatform(ShaderPlatform) && IsMobileDeferredShadingEnabled(ShaderPlatform)))
@@ -85,17 +110,17 @@ ENGINE_API EAntiAliasingMethod GetDefaultAntiAliasingMethod(const FStaticFeature
 			AntiAliasingMethod = EAntiAliasingMethod::AAM_None;
 		}
 	}
-
-	int32 Quality = FMath::Clamp(PostProcessAAQualityCVar->GetValueOnAnyThread(), 0, 6);
-
-	if (Quality <= 0)
+	else if (AntiAliasingMethod == EAntiAliasingMethod::AAM_TSR)
 	{
-		AntiAliasingMethod = EAntiAliasingMethod::AAM_None;
+		if (!SupportsTSR(ShaderPlatform))
+		{
+			// Fallback to UE4's TAA if TSR isn't supported on that platform
+			AntiAliasingMethod = AAM_TemporalAA;
+		}
 	}
-
-	if (AntiAliasingMethod == AAM_TemporalAA && Quality < 3)
+	else
 	{
-		AntiAliasingMethod = AAM_FXAA;
+		unimplemented();
 	}
 
 	return AntiAliasingMethod;
diff --git a/Engine/Source/Runtime/Engine/Private/SceneView.cpp b/Engine/Source/Runtime/Engine/Private/SceneView.cpp
--- a/Engine/Source/Runtime/Engine/Private/SceneView.cpp
+++ b/Engine/Source/Runtime/Engine/Private/SceneView.cpp
@@ -205,13 +205,14 @@ static TAutoConsoleVariable<int32> CVarDefaultLensFlare(
 
 // see EAntiAliasingMethod
 static TAutoConsoleVariable<int32> CVarDefaultAntiAliasing(
-	TEXT("r.DefaultFeature.AntiAliasing"),
-	2,
+	TEXT("r.AntiAliasingMethod"),
+	4,
 	TEXT("Engine default (project setting) for AntiAliasingMethod is (postprocess volume/camera/game setting still can override)\n")
 	TEXT(" 0: off (no anti-aliasing)\n")
-	TEXT(" 1: FXAA (faster than TemporalAA but much more shimmering for non static cases)\n")
-	TEXT(" 2: TemporalAA (default)\n")
-	TEXT(" 3: MSAA (Forward shading only)"),
+	TEXT(" 1: Fast Approximate Anti-Aliasing (FXAA)\n")
+	TEXT(" 2: Temporal Anti-Aliasing (TAA)\n")
+	TEXT(" 3: Multisample Anti-Aliasing (MSAA, Only available on the desktop forward renderer)\n")
+	TEXT(" 4: Temporal Super-Resolution (TSR, Default)"),
 	ECVF_RenderThreadSafe);
 
 // see ELightUnits
@@ -844,8 +845,15 @@ FSceneView::FSceneView(const FSceneViewInitOptions& InitOptions)
 
 	SetupAntiAliasingMethod();
 
-	if ((AntiAliasingMethod == AAM_TemporalAA) &&
-		(CVarEnableTemporalUpsample.GetValueOnAnyThread() || (Family && Family->GetTemporalUpscalerInterface() != nullptr)))
+	if (AntiAliasingMethod == AAM_TSR)
+	{
+		PrimaryScreenPercentageMethod = EPrimaryScreenPercentageMethod::TemporalUpscale;
+	}
+	else if (AntiAliasingMethod == AAM_TemporalAA && CVarEnableTemporalUpsample.GetValueOnAnyThread() != 0)
+	{
+		PrimaryScreenPercentageMethod = EPrimaryScreenPercentageMethod::TemporalUpscale;
+	}
+	else if (Family && Family->GetTemporalUpscalerInterface() != nullptr)
 	{
 		PrimaryScreenPercentageMethod = EPrimaryScreenPercentageMethod::TemporalUpscale;
 	}
@@ -873,12 +881,14 @@ FSceneView::FSceneView(const FSceneViewInitOptions& InitOptions)
 #if DO_CHECK || USING_CODE_ANALYSIS
 bool FSceneView::VerifyMembersChecks() const
 {
+	bool bIsTemporalAccumulation = IsTemporalAccumulationBasedMethod(AntiAliasingMethod);
+
 	if (PrimaryScreenPercentageMethod == EPrimaryScreenPercentageMethod::TemporalUpscale)
 	{
-		checkf(AntiAliasingMethod == AAM_TemporalAA, TEXT("ScreenPercentageMethod == EPrimaryScreenPercentageMethod::TemporalUpscale requires AntiAliasingMethod == AAM_TemporalAA"));
+		checkf(bIsTemporalAccumulation, TEXT("ScreenPercentageMethod == EPrimaryScreenPercentageMethod::TemporalUpscale requires TAA, TSR or custom Temporal Upscaler"));
 	}
 
-	if (AntiAliasingMethod == AAM_TemporalAA)
+	if (bIsTemporalAccumulation)
 	{
 		checkf(State, TEXT("TemporalAA requires the view to have a valid state."));
 	}
@@ -914,16 +924,24 @@ void FSceneView::SetupAntiAliasingMethod()
 				AntiAliasingMethod = AAM_FXAA;
 			}
 		}
+		else if (AntiAliasingMethod == AAM_TSR)
+		{
+			// TODO(TSR): Support TSR with bRealtimeUpdate
+			if (!Family->EngineShowFlags.TemporalAA || !Family->bRealtimeUpdate || !SupportsTSR(GetShaderPlatform()))
+			{
+				AntiAliasingMethod = AAM_FXAA;
+			}
+		}
 
 		// Overides the anti aliasing method to temporal AA when using a custom temporal upscaler.
-		if (Family->GetTemporalUpscalerInterface() != nullptr)
+		if (Family->GetTemporalUpscalerInterface() != nullptr && Family->bRealtimeUpdate)
 		{
 			AntiAliasingMethod = AAM_TemporalAA;
 		}
 	}
 
-    // TemporalAA requires view state for history.
-	if (AntiAliasingMethod == AAM_TemporalAA && !State)
+	// TemporalAA requires view state for history.
+	if (IsTemporalAccumulationBasedMethod(AntiAliasingMethod) && !State)
 	{
 		AntiAliasingMethod = AAM_None;
 	}
diff --git a/Engine/Source/Runtime/Engine/Public/SceneUtils.h b/Engine/Source/Runtime/Engine/Public/SceneUtils.h
--- a/Engine/Source/Runtime/Engine/Public/SceneUtils.h
+++ b/Engine/Source/Runtime/Engine/Public/SceneUtils.h
@@ -35,13 +35,20 @@ UENUM()
 enum EAntiAliasingMethod
 {
 	AAM_None UMETA(DisplayName = "None"),
-	AAM_FXAA UMETA(DisplayName = "FXAA"),
-	AAM_TemporalAA UMETA(DisplayName = "TemporalAA"),
+	AAM_FXAA UMETA(DisplayName = "Fast Approximate Anti-Aliasing (FXAA)"),
+	AAM_TemporalAA UMETA(DisplayName = "Temporal Anti-Aliasing (TAA)"),
 	/** Only supported with forward shading.  MSAA sample count is controlled by r.MSAACount. */
-	AAM_MSAA UMETA(DisplayName = "MSAA"),
+	AAM_MSAA UMETA(DisplayName = "Multisample Anti-Aliasing (MSAA)"),
+	AAM_TSR UMETA(DisplayName = "Temporal Super-Resolution (TSR)"),
 	AAM_MAX,
 };
 
+/** Returns whether the anti-aliasing method use a temporal accumulation */
+inline bool IsTemporalAccumulationBasedMethod(EAntiAliasingMethod AntiAliasingMethod)
+{
+	return AntiAliasingMethod == AAM_TemporalAA || AntiAliasingMethod == AAM_TSR;
+}
+
 /** True if HDR is enabled for the mobile renderer. */
 ENGINE_API bool IsMobileHDR();
 
diff --git a/Engine/Source/Runtime/Engine/Public/SceneView.h b/Engine/Source/Runtime/Engine/Public/SceneView.h
--- a/Engine/Source/Runtime/Engine/Public/SceneView.h
+++ b/Engine/Source/Runtime/Engine/Public/SceneView.h
@@ -1421,7 +1421,11 @@ struct FSceneViewScreenPercentageConfig
 	static constexpr float kMinResolutionFraction = 0.01f;
 	static constexpr float kMaxResolutionFraction = 4.0f;
 
-	// Sets the minimal and maximal screen percentage for TAA upsample.
+	// Sets the minimal and maximal screen percentage for TSR.
+	static constexpr float kMinTSRResolutionFraction = 0.25f;
+	static constexpr float kMaxTSRResolutionFraction = 2.0f;
+
+	// Sets the minimal and maximal screen percentage for TAAU.
 	static constexpr float kMinTAAUpsampleResolutionFraction = 0.5f;
 	static constexpr float kMaxTAAUpsampleResolutionFraction = 2.0f;
 
diff --git a/Engine/Source/Runtime/RenderCore/Private/RenderUtils.cpp b/Engine/Source/Runtime/RenderCore/Private/RenderUtils.cpp
--- a/Engine/Source/Runtime/RenderCore/Private/RenderUtils.cpp
+++ b/Engine/Source/Runtime/RenderCore/Private/RenderUtils.cpp
@@ -1151,6 +1151,11 @@ RENDERCORE_API bool SupportsGen4TAA(const FStaticShaderPlatform Platform)
 	return true;
 }
 
+RENDERCORE_API bool SupportsTSR(const FStaticShaderPlatform Platform)
+{
+	return FDataDrivenShaderPlatformInfo::GetSupportsGen5TemporalAA(Platform);
+}
+
 template<typename Type>
 Type FShaderPlatformCachedIniValue<Type>::Get(EShaderPlatform ShaderPlatform)
 {
diff --git a/Engine/Source/Runtime/RenderCore/Public/RenderUtils.h b/Engine/Source/Runtime/RenderCore/Public/RenderUtils.h
--- a/Engine/Source/Runtime/RenderCore/Public/RenderUtils.h
+++ b/Engine/Source/Runtime/RenderCore/Public/RenderUtils.h
@@ -505,6 +505,8 @@ RENDERCORE_API bool IsMobileDistanceFieldShadowingEnabled(const FStaticShaderPla
 
 RENDERCORE_API bool SupportsGen4TAA(const FStaticShaderPlatform Platform);
 
+RENDERCORE_API bool SupportsTSR(const FStaticShaderPlatform Platform);
+
 RENDERCORE_API bool PlatformSupportsVelocityRendering(const FStaticShaderPlatform Platform);
 
 RENDERCORE_API bool IsUsingDBuffers(const FStaticShaderPlatform Platform);
diff --git a/Engine/Source/Runtime/Renderer/Private/CompositionLighting/CompositionLighting.cpp b/Engine/Source/Runtime/Renderer/Private/CompositionLighting/CompositionLighting.cpp
--- a/Engine/Source/Runtime/Renderer/Private/CompositionLighting/CompositionLighting.cpp
+++ b/Engine/Source/Runtime/Renderer/Private/CompositionLighting/CompositionLighting.cpp
@@ -146,7 +146,7 @@ static FSSAOCommonParameters GetSSAOCommonParameters(
 
 	// If there is no temporal upsampling, we need a smooth pass to get rid of the grid pattern.
 	// Pixel shader version has relatively smooth result so no need to do extra work.
-	CommonParameters.bNeedSmoothingPass = CommonParameters.FullscreenType != ESSAOType::EPS && View.AntiAliasingMethod != AAM_TemporalAA && CVarSSAOSmoothPass.GetValueOnRenderThread();
+	CommonParameters.bNeedSmoothingPass = CommonParameters.FullscreenType != ESSAOType::EPS && !IsTemporalAccumulationBasedMethod(View.AntiAliasingMethod) && CVarSSAOSmoothPass.GetValueOnRenderThread();
 
 	return CommonParameters;
 }
diff --git a/Engine/Source/Runtime/Renderer/Private/CustomDepthRendering.cpp b/Engine/Source/Runtime/Renderer/Private/CustomDepthRendering.cpp
--- a/Engine/Source/Runtime/Renderer/Private/CustomDepthRendering.cpp
+++ b/Engine/Source/Runtime/Renderer/Private/CustomDepthRendering.cpp
@@ -192,7 +192,7 @@ bool FSceneRenderer::RenderCustomDepthPass(FRDGBuilder& GraphBuilder, const FCus
 			PassParameters->SceneTextures = SceneTextures;
 
 			// User requested jitter-free custom depth.
-			if (CVarCustomDepthTemporalAAJitter.GetValueOnRenderThread() == 0 && View.AntiAliasingMethod == AAM_TemporalAA)
+			if (CVarCustomDepthTemporalAAJitter.GetValueOnRenderThread() == 0 && IsTemporalAccumulationBasedMethod(View.AntiAliasingMethod))
 			{
 				PassParameters->View = CreateViewShaderParametersWithoutJitter(View);
 			}
diff --git a/Engine/Source/Runtime/Renderer/Private/LightShaftRendering.cpp b/Engine/Source/Runtime/Renderer/Private/LightShaftRendering.cpp
--- a/Engine/Source/Runtime/Renderer/Private/LightShaftRendering.cpp
+++ b/Engine/Source/Runtime/Renderer/Private/LightShaftRendering.cpp
@@ -345,7 +345,7 @@ FScreenPassTexture AddTemporalAAPass(
 	FTemporalAAHistory* HistoryState,
 	FScreenPassTexture LightShafts)
 {
-	if (View.AntiAliasingMethod == AAM_TemporalAA && HistoryState && GLightShaftAllowTAA)
+	if (IsTemporalAccumulationBasedMethod(View.AntiAliasingMethod) && HistoryState && GLightShaftAllowTAA)
 	{
 		const FSceneTextureParameters SceneTextureParameters = GetSceneTextureParameters(GraphBuilder, SceneTextures);
 
diff --git a/Engine/Source/Runtime/Renderer/Private/PostProcess/DiaphragmDOF.cpp b/Engine/Source/Runtime/Renderer/Private/PostProcess/DiaphragmDOF.cpp
--- a/Engine/Source/Runtime/Renderer/Private/PostProcess/DiaphragmDOF.cpp
+++ b/Engine/Source/Runtime/Renderer/Private/PostProcess/DiaphragmDOF.cpp
@@ -1434,7 +1434,7 @@ FRDGTextureRef DiaphragmDOF::AddPasses(
 	FIntPoint PreprocessViewSize = FIntPoint::DivideAndRoundUp(FullResViewSize, PrefilteringResolutionDivisor);
 	FIntPoint GatheringViewSize = PreprocessViewSize;
 
-	if (View.AntiAliasingMethod == AAM_TemporalAA && ViewState)
+	if (IsTemporalAccumulationBasedMethod(View.AntiAliasingMethod) && ViewState)
 	{
 		PreprocessViewSize = FIntPoint::DivideAndRoundUp(TAAParameters.OutputViewRect.Size(), PrefilteringResolutionDivisor);
 	}
@@ -1615,7 +1615,7 @@ FRDGTextureRef DiaphragmDOF::AddPasses(
 	}
 	
 	// TAA the setup for the convolution to be temporally stable.
-	if (View.AntiAliasingMethod == AAM_TemporalAA && ViewState)
+	if (IsTemporalAccumulationBasedMethod(View.AntiAliasingMethod) && ViewState)
 	{
 		TAAParameters.SceneDepthTexture = SceneTextures.SceneDepthTexture;
 		TAAParameters.SceneVelocityTexture = SceneTextures.GBufferVelocityTexture;
diff --git a/Engine/Source/Runtime/Renderer/Private/PostProcess/PostProcessCompositeEditorPrimitives.cpp b/Engine/Source/Runtime/Renderer/Private/PostProcess/PostProcessCompositeEditorPrimitives.cpp
--- a/Engine/Source/Runtime/Renderer/Private/PostProcess/PostProcessCompositeEditorPrimitives.cpp
+++ b/Engine/Source/Runtime/Renderer/Private/PostProcess/PostProcessCompositeEditorPrimitives.cpp
@@ -258,7 +258,7 @@ const FViewInfo* CreateEditorPrimitiveView(const FViewInfo& ParentView, FIntRect
 	// Disable decals so that we don't do a SetDepthStencilState() in TMobileBasePassDrawingPolicy::SetupPipelineState()
 	EditorView->bSceneHasDecals = false;
 
-	if (EditorView->AntiAliasingMethod == AAM_TemporalAA)
+	if (IsTemporalAccumulationBasedMethod(EditorView->AntiAliasingMethod))
 	{
 		EditorView->ViewMatrices.HackRemoveTemporalAAProjectionJitter();
 	}
@@ -363,7 +363,7 @@ FScreenPassTexture AddEditorPrimitivePass(
 			const int32 DepthUpsampleFactor = FMath::Clamp(CVarEditorTemporalUpsampleDepth.GetValueOnRenderThread(), 0, ComputeMaxUpsampleFactor);
 
 			// Upsample the depth at higher resolution to reduce depth intersection instability of editor primitives.
-			if (DepthUpsampleFactor > 0 && View.ViewState && View.AntiAliasingMethod == AAM_TemporalAA)
+			if (DepthUpsampleFactor > 0 && View.ViewState && IsTemporalAccumulationBasedMethod(View.AntiAliasingMethod))
 			{
 				FScreenPassTexture History;
 				{
diff --git a/Engine/Source/Runtime/Renderer/Private/PostProcess/PostProcessing.cpp b/Engine/Source/Runtime/Renderer/Private/PostProcess/PostProcessing.cpp
--- a/Engine/Source/Runtime/Renderer/Private/PostProcess/PostProcessing.cpp
+++ b/Engine/Source/Runtime/Renderer/Private/PostProcess/PostProcessing.cpp
@@ -1629,6 +1629,7 @@ void AddMobilePostProcessingPasses(FRDGBuilder& GraphBuilder, FScene* Scene, con
 			View.FinalPostProcessSettings.AutoExposureMinBrightness < View.FinalPostProcessSettings.AutoExposureMaxBrightness;
 
 		bool bUseTAA = View.AntiAliasingMethod == AAM_TemporalAA;
+		ensure(View.AntiAliasingMethod != AAM_TSR);
 
 		bool bUseDistortion = IsMobileDistortionActive(View);
 
@@ -2219,7 +2220,7 @@ FRDGTextureRef AddProcessPlanarReflectionPass(
 	FSceneViewState* ViewState = View.ViewState;
 	const EAntiAliasingMethod AntiAliasingMethod = View.AntiAliasingMethod;
 
-	if (AntiAliasingMethod == AAM_TemporalAA)
+	if (IsTemporalAccumulationBasedMethod(AntiAliasingMethod))
 	{
 		check(ViewState);
 
diff --git a/Engine/Source/Runtime/Renderer/Private/PostProcess/TemporalAA.cpp b/Engine/Source/Runtime/Renderer/Private/PostProcess/TemporalAA.cpp
--- a/Engine/Source/Runtime/Renderer/Private/PostProcess/TemporalAA.cpp
+++ b/Engine/Source/Runtime/Renderer/Private/PostProcess/TemporalAA.cpp
@@ -17,13 +17,6 @@ namespace
 const int32 GTemporalAATileSizeX = 8;
 const int32 GTemporalAATileSizeY = 8;
 
-TAutoConsoleVariable<int32> CVarTAAAlgorithm(
-	TEXT("r.TemporalAA.Algorithm"), 1,
-	TEXT("Algorithm to use for Temporal AA\n")
-	TEXT(" 0: Temporal AA Upsampling\n")
-	TEXT(" 1: Temporal Super Resolution (default)"),
-	ECVF_RenderThreadSafe);
-
 TAutoConsoleVariable<float> CVarTemporalAAFilterSize(
 	TEXT("r.TemporalAAFilterSize"),
 	1.0f,
@@ -798,7 +791,7 @@ static void AddGen4MainTemporalAAPasses(
 	FRDGTextureRef* OutSceneColorHalfResTexture,
 	FIntRect* OutSceneColorHalfResViewRect)
 {
-	check(View.AntiAliasingMethod == AAM_TemporalAA && View.ViewState);
+	check(View.ViewState);
 
 	FTAAPassParameters TAAParameters(View);
 
@@ -906,7 +899,7 @@ public:
 		FRDGTextureRef* OutSceneColorHalfResTexture,
 		FIntRect* OutSceneColorHalfResViewRect) const final
 	{
-		if (CVarTAAAlgorithm.GetValueOnRenderThread() && DoesPlatformSupportTSR(View.GetShaderPlatform()))
+		if (ITemporalUpscaler::GetMainTAAPassConfig(View) == EMainTAAPassConfig::TSR)
 		{
 			*OutSceneColorHalfResTexture = nullptr;
 			//*OutSceneColorHalfResViewRect; // TODO.
@@ -955,7 +948,7 @@ EMainTAAPassConfig ITemporalUpscaler::GetMainTAAPassConfig(const FViewInfo& View
 	{
 		return EMainTAAPassConfig::Disabled;
 	}
-	else if (View.AntiAliasingMethod != AAM_TemporalAA)
+	else if (!IsTemporalAccumulationBasedMethod(View.AntiAliasingMethod))
 	{
 		return EMainTAAPassConfig::Disabled;
 	}
@@ -966,8 +959,10 @@ EMainTAAPassConfig ITemporalUpscaler::GetMainTAAPassConfig(const FViewInfo& View
 	{
 		return EMainTAAPassConfig::ThirdParty;
 	}
-	else if (CVarTAAAlgorithm.GetValueOnRenderThread() && DoesPlatformSupportTSR(View.GetShaderPlatform()))
+	else if (View.AntiAliasingMethod == AAM_TSR && !IsPostProcessingWithAlphaChannelSupported())
 	{
+		// TODO(TSR): alpha channel is not supported yet
+		check(SupportsTSR(View.GetShaderPlatform()));
 		return EMainTAAPassConfig::TSR;
 	}
 	else
diff --git a/Engine/Source/Runtime/Renderer/Private/PostProcess/TemporalSuperResolution.cpp b/Engine/Source/Runtime/Renderer/Private/PostProcess/TemporalSuperResolution.cpp
--- a/Engine/Source/Runtime/Renderer/Private/PostProcess/TemporalSuperResolution.cpp
+++ b/Engine/Source/Runtime/Renderer/Private/PostProcess/TemporalSuperResolution.cpp
@@ -13,16 +13,6 @@
 
 #define COMPILE_TSR_DEBUG_PASSES (!UE_BUILD_SHIPPING)
 
-bool DoesPlatformSupportTSR(EShaderPlatform Platform)
-{
-	// TODO(TSR): alpha channel is not supported yet
-	if (IsPostProcessingWithAlphaChannelSupported())
-	{
-		return false;
-	}
-	return FDataDrivenShaderPlatformInfo::GetSupportsGen5TemporalAA(Platform);
-}
-
 namespace
 {
 
@@ -143,7 +133,7 @@ public:
 
 	static bool ShouldCompilePermutation(const FGlobalShaderPermutationParameters& Parameters)
 	{
-		return DoesPlatformSupportTSR(Parameters.Platform);
+		return SupportsTSR(Parameters.Platform);
 	}
 
 	static void ModifyCompilationEnvironment(const FGlobalShaderPermutationParameters& Parameters, FShaderCompilerEnvironment& OutEnvironment)
diff --git a/Engine/Source/Runtime/Renderer/Private/SceneRendering.cpp b/Engine/Source/Runtime/Renderer/Private/SceneRendering.cpp
--- a/Engine/Source/Runtime/Renderer/Private/SceneRendering.cpp
+++ b/Engine/Source/Runtime/Renderer/Private/SceneRendering.cpp
@@ -1552,7 +1552,7 @@ void FViewInfo::SetupUniformBufferParameters(
 	}
 
 	{
-		ensureMsgf(TemporalJitterSequenceLength == 1 || AntiAliasingMethod == AAM_TemporalAA,
+		ensureMsgf(TemporalJitterSequenceLength == 1 || IsTemporalAccumulationBasedMethod(AntiAliasingMethod),
 			TEXT("TemporalJitterSequenceLength = %i is invalid"), TemporalJitterSequenceLength);
 		ensureMsgf(TemporalJitterIndex >= 0 && TemporalJitterIndex < TemporalJitterSequenceLength,
 			TEXT("TemporalJitterIndex = %i is invalid (TemporalJitterSequenceLength = %i)"), TemporalJitterIndex, TemporalJitterSequenceLength);
@@ -2574,7 +2574,7 @@ void FSceneRenderer::PrepareViewRectsForRendering()
 		{
 			// Tenmporal upsample is supported only if TAA is turned on.
 			if (View.PrimaryScreenPercentageMethod == EPrimaryScreenPercentageMethod::TemporalUpscale &&
-				(View.AntiAliasingMethod != AAM_TemporalAA ||
+				(!IsTemporalAccumulationBasedMethod(View.AntiAliasingMethod) ||
 				 ViewFamily.EngineShowFlags.VisualizeBuffer))
 			{
 				View.PrimaryScreenPercentageMethod = EPrimaryScreenPercentageMethod::SpatialUpscale;
diff --git a/Engine/Source/Runtime/Renderer/Private/SceneVisibility.cpp b/Engine/Source/Runtime/Renderer/Private/SceneVisibility.cpp
--- a/Engine/Source/Runtime/Renderer/Private/SceneVisibility.cpp
+++ b/Engine/Source/Runtime/Renderer/Private/SceneVisibility.cpp
@@ -3269,7 +3269,7 @@ void FSceneRenderer::PreVisibilityFrameSetup(FRDGBuilder& GraphBuilder, const FS
 		bool bTemporalUpsampling = View.PrimaryScreenPercentageMethod == EPrimaryScreenPercentageMethod::TemporalUpscale;
 		
 		// Apply a sub pixel offset to the view.
-		if (View.AntiAliasingMethod == AAM_TemporalAA && ViewState && (CVarTemporalAASamplesValue > 0 || bTemporalUpsampling) && View.bAllowTemporalJitter)
+		if (IsTemporalAccumulationBasedMethod(View.AntiAliasingMethod) && ViewState && (CVarTemporalAASamplesValue > 0 || bTemporalUpsampling) && View.bAllowTemporalJitter)
 		{
 			float EffectivePrimaryResolutionFraction = float(View.ViewRect.Width()) / float(View.GetSecondaryViewRectSize().X);
 
diff --git a/Engine/Source/Runtime/Renderer/Private/ScreenSpaceRayTracing.cpp b/Engine/Source/Runtime/Renderer/Private/ScreenSpaceRayTracing.cpp
--- a/Engine/Source/Runtime/Renderer/Private/ScreenSpaceRayTracing.cpp
+++ b/Engine/Source/Runtime/Renderer/Private/ScreenSpaceRayTracing.cpp
@@ -170,7 +170,7 @@ bool IsSSRTemporalPassRequired(const FViewInfo& View)
 	{
 		return false;
 	}
-	return View.AntiAliasingMethod != AAM_TemporalAA || CVarSSRTemporal.GetValueOnRenderThread() != 0;
+	return !IsTemporalAccumulationBasedMethod(View.AntiAliasingMethod) || CVarSSRTemporal.GetValueOnRenderThread() != 0;
 }
 
 
@@ -304,7 +304,7 @@ FLinearColor ComputeSSRParams(const FViewInfo& View, ESSRQuality SSRQuality, boo
 
 	if(View.ViewState)
 	{
-		bool bTemporalAAIsOn = View.AntiAliasingMethod == AAM_TemporalAA;
+		bool bTemporalAAIsOn = IsTemporalAccumulationBasedMethod(View.AntiAliasingMethod);
 
 		if(bTemporalAAIsOn)
 		{
diff --git a/Engine/Source/Runtime/Renderer/Private/VelocityRendering.cpp b/Engine/Source/Runtime/Renderer/Private/VelocityRendering.cpp
--- a/Engine/Source/Runtime/Renderer/Private/VelocityRendering.cpp
+++ b/Engine/Source/Runtime/Renderer/Private/VelocityRendering.cpp
@@ -168,7 +168,7 @@ bool FDeferredShadingSceneRenderer::ShouldRenderVelocities() const
 	{
 		const FViewInfo& View = Views[ViewIndex];
 
-		bool bTemporalAA = (View.AntiAliasingMethod == AAM_TemporalAA) && !View.bCameraCut;
+		bool bTemporalAA = IsTemporalAccumulationBasedMethod(View.AntiAliasingMethod) && !View.bCameraCut;
 		bool bMotionBlur = IsMotionBlurEnabled(View);
 		bool bDistanceFieldAO = ShouldPrepareForDistanceFieldAO();
 
@@ -200,7 +200,7 @@ bool FMobileSceneRenderer::ShouldRenderVelocities() const
 	{
 		const FViewInfo& View = Views[ViewIndex];
 
-		bool bTemporalAA = (View.AntiAliasingMethod == AAM_TemporalAA) && !View.bCameraCut;
+		bool bTemporalAA = IsTemporalAccumulationBasedMethod(View.AntiAliasingMethod) && !View.bCameraCut;
 
 		bNeedsVelocity |= bTemporalAA;
 	}
diff --git a/Templates/TP_AEC_ArchvisBP/Config/DefaultEngine.ini b/Templates/TP_AEC_ArchvisBP/Config/DefaultEngine.ini
--- a/Templates/TP_AEC_ArchvisBP/Config/DefaultEngine.ini
+++ b/Templates/TP_AEC_ArchvisBP/Config/DefaultEngine.ini
@@ -8,7 +8,6 @@ AppliedDefaultGraphicsPerformance=Maximum
 [/Script/Engine.RendererSettings]
 ; exposure control and tone mapping settings
 r.Tonemapper.Sharpen=2
-r.TemporalAA.Upsampling=True
 r.DefaultFeature.AutoExposure.ExtendDefaultLuminanceRange=True
 
 ;dynamic gi settings
diff --git a/Templates/TP_AEC_BlankBP/Config/DefaultEngine.ini b/Templates/TP_AEC_BlankBP/Config/DefaultEngine.ini
--- a/Templates/TP_AEC_BlankBP/Config/DefaultEngine.ini
+++ b/Templates/TP_AEC_BlankBP/Config/DefaultEngine.ini
@@ -8,7 +8,6 @@ AppliedDefaultGraphicsPerformance=Maximum
 [/Script/Engine.RendererSettings]
 ; exposure control and tone mapping settings
 r.Tonemapper.Sharpen=2
-r.TemporalAA.Upsampling=True
 r.DefaultFeature.AutoExposure.ExtendDefaultLuminanceRange=True
 
 ;dynamic gi settings
diff --git a/Templates/TP_PhotoStudioBP/Config/DefaultEngine.ini b/Templates/TP_PhotoStudioBP/Config/DefaultEngine.ini
--- a/Templates/TP_PhotoStudioBP/Config/DefaultEngine.ini
+++ b/Templates/TP_PhotoStudioBP/Config/DefaultEngine.ini
@@ -8,7 +8,6 @@ AppliedDefaultGraphicsPerformance=Maximum
 [/Script/Engine.RendererSettings]
 ; exposure control and tone mapping settings
 r.Tonemapper.Sharpen=2
-r.TemporalAA.Upsampling=True
 r.DefaultFeature.AutoExposure.ExtendDefaultLuminanceRange=True
 
 ;dynamic gi settings
