Copy correct RHI resource

diff --git a/Engine/Source/Runtime/Engine/Public/Rendering/SkinWeightVertexBuffer.h b/Engine/Source/Runtime/Engine/Public/Rendering/SkinWeightVertexBuffer.h
--- a/Engine/Source/Runtime/Engine/Public/Rendering/SkinWeightVertexBuffer.h
+++ b/Engine/Source/Runtime/Engine/Public/Rendering/SkinWeightVertexBuffer.h
@@ -475,7 +475,7 @@ public:
 	void InitRHIForStreaming(const FSkinWeightRHIInfo& RHIInfo, TRHIResourceUpdateBatcher<MaxNumUpdates>& Batcher)
 	{
 		DataVertexBuffer.InitRHIForStreaming(RHIInfo.DataVertexBufferRHI, Batcher);
-		LookupVertexBuffer.InitRHIForStreaming(RHIInfo.DataVertexBufferRHI, Batcher);
+		LookupVertexBuffer.InitRHIForStreaming(RHIInfo.LookupVertexBufferRHI, Batcher);
 	}
 
 	template <uint32 MaxNumUpdates>
