Remove some unneeded calls to FetchActor.

diff --git a/Engine/Source/Runtime/Engine/Private/Components/CharacterMovementComponent.cpp b/Engine/Source/Runtime/Engine/Private/Components/CharacterMovementComponent.cpp
--- a/Engine/Source/Runtime/Engine/Private/Components/CharacterMovementComponent.cpp
+++ b/Engine/Source/Runtime/Engine/Private/Components/CharacterMovementComponent.cpp
@@ -6858,9 +6858,9 @@ void UCharacterMovementComponent::HandleImpact(const FHitResult& Impact, float T
 		PFAgent->OnMoveBlockedBy(Impact);
 	}
 
-	APawn* OtherPawn = Impact.HitObjectHandle.FetchActor<APawn>();
-	if (OtherPawn)
+	if (Impact.HitObjectHandle.DoesRepresentClass(APawn::StaticClass()))
 	{
+		APawn* OtherPawn = Impact.HitObjectHandle.FetchActor<APawn>();
 		NotifyBumpedPawn(OtherPawn);
 	}
 
diff --git a/Engine/Source/Runtime/Engine/Private/Components/PrimitiveComponent.cpp b/Engine/Source/Runtime/Engine/Private/Components/PrimitiveComponent.cpp
--- a/Engine/Source/Runtime/Engine/Private/Components/PrimitiveComponent.cpp
+++ b/Engine/Source/Runtime/Engine/Private/Components/PrimitiveComponent.cpp
@@ -2582,7 +2582,7 @@ void UPrimitiveComponent::DispatchBlockingHit(AActor& Owner, FHitResult const& B
 		if (!BlockingHitComponent->IsPendingKill())
 		{
 			// BlockingHit.GetActor() could be marked for deletion in DispatchBlockingHit(), which would make the weak pointer return NULL.
-			if (AActor* const BlockingHitActor = BlockingHit.HitObjectHandle.FetchActor())
+			if (AActor* const BlockingHitActor = BlockingHit.HitObjectHandle.GetManagingActor())
 			{
 				BlockingHitActor->DispatchBlockingHit(BlockingHitComponent, this, false, BlockingHit);
 			}
