Updated FHitResult to take a MyItem parameter. This new parameter is the same as Item but for the object we are currently dealing with.
FHitResult::GetReversedHit now swaps the contents of Item with MyItem and BoneName with MyBoneName.
Fixed a bug in light weight instances that could occur when you try to convert the same instance to an actor multiple times in the same frame.

diff --git a/Engine/Source/Runtime/Engine/Classes/Engine/EngineTypes.h b/Engine/Source/Runtime/Engine/Classes/Engine/EngineTypes.h
--- a/Engine/Source/Runtime/Engine/Classes/Engine/EngineTypes.h
+++ b/Engine/Source/Runtime/Engine/Classes/Engine/EngineTypes.h
@@ -2097,8 +2097,12 @@ struct ENGINE_API FActorInstanceHandle
 	template <typename T>
 	T* FetchActor() const;
 
+	/* Returns the index used internally by the manager */
 	FORCEINLINE int32 GetInstanceIndex() const { return InstanceIndex; }
 
+	/* Returns the index used by rendering and collision */
+	int32 GetRenderingInstanceIndex() const;
+
 	FActorInstanceHandle& operator=(const FActorInstanceHandle& Other) = default;
 	FActorInstanceHandle& operator=(FActorInstanceHandle&& Other) = default;
 	FActorInstanceHandle& operator=(AActor* OtherActor);
@@ -2218,6 +2222,10 @@ struct ENGINE_API FHitResult
 	UPROPERTY()
 	float PenetrationDepth;
 
+	/** If the hit result is from a collision this will have extra info about the item that hit the second item. */
+	UPROPERTY()
+	int32 MyItem;
+
 	/** Extra data about item that was hit (hit primitive specific). */
 	UPROPERTY()
 	int32 Item;
@@ -2294,6 +2302,7 @@ struct ENGINE_API FHitResult
 		FMemory::Memzero(this, sizeof(FHitResult));
 		HitObjectHandle = FActorInstanceHandle();
 		Time = 1.f;
+		MyItem = INDEX_NONE;
 	}
 
 	/** Initialize empty hit result with given time, TraceStart, and TraceEnd */
@@ -2304,6 +2313,7 @@ struct ENGINE_API FHitResult
 		Time = 1.f;
 		TraceStart = Start;
 		TraceEnd = End;
+		MyItem = INDEX_NONE;
 	}
 
 	/** Ctor for easily creating "fake" hits from limited data. */
@@ -2396,6 +2406,14 @@ struct ENGINE_API FHitResult
 		FHitResult Result(Hit);
 		Result.Normal = -Result.Normal;
 		Result.ImpactNormal = -Result.ImpactNormal;
+
+		int32 TempItem = Result.Item;
+		Result.Item = Result.MyItem;
+		Result.MyItem = TempItem;
+
+		FName TempBoneName = Result.BoneName;
+		Result.BoneName = Result.MyBoneName;
+		Result.MyBoneName = Result.BoneName;
 		return Result;
 	}
 
diff --git a/Engine/Source/Runtime/Engine/Classes/GameFramework/LightWeightInstanceManager.h b/Engine/Source/Runtime/Engine/Classes/GameFramework/LightWeightInstanceManager.h
--- a/Engine/Source/Runtime/Engine/Classes/GameFramework/LightWeightInstanceManager.h
+++ b/Engine/Source/Runtime/Engine/Classes/GameFramework/LightWeightInstanceManager.h
@@ -98,8 +98,12 @@ public:
 	// Returns a handle to a light weight instance representing the same object as InActor and calls destroy on InActor if successful.
 	FActorInstanceHandle ConvertActorToLightWeightInstance(AActor* InActor);
 
+	// Returns the index used internally by the light weight instance manager that is associated with the instance referred to by InIndex used by collision and rendering
 	virtual int32 ConvertCollisionIndexToLightWeightIndex(int32 InIndex) const;
 
+	// Returns the index used by collision and rendering that is associated with the instance referred to by InIndex
+	virtual int32 ConvertLightWeightIndexToCollisionIndex(int32 InIndex) const;
+
 	template<typename U>
 	bool IsInterfaceSupported() const
 	{
diff --git a/Engine/Source/Runtime/Engine/Classes/GameFramework/LightWeightInstanceStaticMeshManager.h b/Engine/Source/Runtime/Engine/Classes/GameFramework/LightWeightInstanceStaticMeshManager.h
--- a/Engine/Source/Runtime/Engine/Classes/GameFramework/LightWeightInstanceStaticMeshManager.h
+++ b/Engine/Source/Runtime/Engine/Classes/GameFramework/LightWeightInstanceStaticMeshManager.h
@@ -26,6 +26,8 @@ class ENGINE_API ALightWeightInstanceStaticMeshManager : public ALightWeightInst
 
 	virtual int32 ConvertCollisionIndexToLightWeightIndex(int32 InIndex) const override;
 
+	virtual int32 ConvertLightWeightIndexToCollisionIndex(int32 InIndex) const override;
+
 protected:
 	virtual void AddNewInstanceAt(FLWIData* InitData, int32 Index) override;
 
@@ -33,6 +35,8 @@ protected:
 
 	void RemoveInstanceFromRendering(int32 DataIndex);
 
+	void PostRemoveInstanceFromRendering(int32 RenderingIndex);
+
 	// sets the parameters on the instanced static mesh component
 	virtual void SetInstancedStaticMeshParams();
 
diff --git a/Engine/Source/Runtime/Engine/Private/Actor.cpp b/Engine/Source/Runtime/Engine/Private/Actor.cpp
--- a/Engine/Source/Runtime/Engine/Private/Actor.cpp
+++ b/Engine/Source/Runtime/Engine/Private/Actor.cpp
@@ -4595,6 +4595,7 @@ void AActor::DispatchPhysicsCollisionHit(const FRigidBodyCollisionInfo& MyInfo,
 	Result.PhysMaterial = ContactInfo.PhysMaterial[1];
 	Result.HitObjectHandle = FActorInstanceHandle(OtherInfo.Actor.Get());
 	Result.Component = OtherInfo.Component;
+	Result.MyItem = MyInfo.BodyIndex;
 	Result.Item = OtherInfo.BodyIndex;
 	Result.BoneName = OtherInfo.BoneName;
 	Result.MyBoneName = MyInfo.BoneName;
diff --git a/Engine/Source/Runtime/Engine/Private/EngineTypes.cpp b/Engine/Source/Runtime/Engine/Private/EngineTypes.cpp
--- a/Engine/Source/Runtime/Engine/Private/EngineTypes.cpp
+++ b/Engine/Source/Runtime/Engine/Private/EngineTypes.cpp
@@ -471,6 +471,12 @@ AActor* FActorInstanceHandle::FetchActor() const
 	return FLightWeightInstanceSubsystem::Get().FetchActor(*this);
 }
 
+int32 FActorInstanceHandle::GetRenderingInstanceIndex() const
+{
+	const ALightWeightInstanceManager* Manager = FLightWeightInstanceSubsystem::Get().GetManagerAt(ManagerIndex);
+	return Manager ? Manager->ConvertLightWeightIndexToCollisionIndex(InstanceIndex) : INDEX_NONE;
+}
+
 UObject* FActorInstanceHandle::GetActorAsUObject()
 {
 	if (IsActorValid())
diff --git a/Engine/Source/Runtime/Engine/Private/LightWeightInstanceManager.cpp b/Engine/Source/Runtime/Engine/Private/LightWeightInstanceManager.cpp
--- a/Engine/Source/Runtime/Engine/Private/LightWeightInstanceManager.cpp
+++ b/Engine/Source/Runtime/Engine/Private/LightWeightInstanceManager.cpp
@@ -169,6 +169,11 @@ int32 ALightWeightInstanceManager::ConvertCollisionIndexToLightWeightIndex(int32
 	return InIndex;
 }
 
+int32 ALightWeightInstanceManager::ConvertLightWeightIndexToCollisionIndex(int32 InIndex) const
+{
+	return InIndex;
+}
+
 void ALightWeightInstanceManager::SetSpawnParameters(FActorSpawnParameters& SpawnParams)
 {
 	SpawnParams.OverrideLevel = GetLevel();
diff --git a/Engine/Source/Runtime/Engine/Private/LightWeightInstanceStaticMeshManager.cpp b/Engine/Source/Runtime/Engine/Private/LightWeightInstanceStaticMeshManager.cpp
--- a/Engine/Source/Runtime/Engine/Private/LightWeightInstanceStaticMeshManager.cpp
+++ b/Engine/Source/Runtime/Engine/Private/LightWeightInstanceStaticMeshManager.cpp
@@ -49,6 +49,11 @@ int32 ALightWeightInstanceStaticMeshManager::ConvertCollisionIndexToLightWeightI
 	return RenderingIndicesToDataIndices[InIndex];
 }
 
+int32 ALightWeightInstanceStaticMeshManager::ConvertLightWeightIndexToCollisionIndex(int32 InIndex) const
+{
+	return DataIndicesToRenderingIndices[InIndex];
+}
+
 void ALightWeightInstanceStaticMeshManager::AddNewInstanceAt(FLWIData* InitData, int32 Index)
 {
 	Super::AddNewInstanceAt(InitData, Index);
@@ -85,7 +90,6 @@ void ALightWeightInstanceStaticMeshManager::RemoveInstanceFromRendering(int32 Da
 	if (IsIndexValid(DataIndex))
 	{
 		const int32 RenderingIndex = DataIndicesToRenderingIndices[DataIndex];
-		DataIndicesToRenderingIndices[DataIndex] = INDEX_NONE;
 
 		if (RenderingIndex != INDEX_NONE)
 		{
@@ -94,19 +98,24 @@ void ALightWeightInstanceStaticMeshManager::RemoveInstanceFromRendering(int32 Da
 				InstancedStaticMeshComponent->RemoveInstance(RenderingIndex);
 			}
 
-			// update the map to match what is now in the instanced static mesh component
-			RenderingIndicesToDataIndices.RemoveAtSwap(RenderingIndex);
+			// We can't remove the indices right away because we could have a situation where we receive multiple requests to convert an instance to an actor in the same frame			
+			GetWorld()->GetTimerManager().SetTimerForNextTick([this, RenderingIndex]() { PostRemoveInstanceFromRendering(RenderingIndex); });
+		}
+	}
+}
+
+void ALightWeightInstanceStaticMeshManager::PostRemoveInstanceFromRendering(int32 RenderingIndex)
+{
+	RenderingIndicesToDataIndices.RemoveAtSwap(RenderingIndex);
 
-			// fix up the other side of the map to match the change we just made
-			// if we removed the last element than nothing was moved so we're done
-			if (RenderingIndex < RenderingIndicesToDataIndices.Num())
-			{
-				// find the data index that corresponds with the changed rendering index
-				const int32 ShiftedDataIndex = RenderingIndicesToDataIndices[RenderingIndex];
+	// fix up the other side of the map to match the change we just made
+	// if we removed the last element than nothing was moved so we're done
+	if (RenderingIndex < RenderingIndicesToDataIndices.Num())
+	{
+		// find the data index that corresponds with the changed rendering index
+		const int32 ShiftedDataIndex = RenderingIndicesToDataIndices[RenderingIndex];
 
-				DataIndicesToRenderingIndices[ShiftedDataIndex] = RenderingIndex;
-			}
-		}
+		DataIndicesToRenderingIndices[ShiftedDataIndex] = RenderingIndex;
 	}
 }
 
