Fix: RandomPointInBoundingBox is available globally but not available on the random stream

diff --git a/Engine/Config/BaseEngine.ini b/Engine/Config/BaseEngine.ini
--- a/Engine/Config/BaseEngine.ini
+++ b/Engine/Config/BaseEngine.ini
@@ -484,6 +484,9 @@ NearClipPlane=10.0
 +FunctionRedirects=(OldName="RandomUnitVectorInConeWithYawAndPitch",NewName="/Script/Engine.KismetMathLibrary.RandomUnitVectorInEllipticalConeInDegrees")
 +PropertyRedirects=(OldName="KismetMathLibrary.RandomUnitVectorInConeInRadians.ConeHalfAngle",NewName="ConeHalfAngleInRadians")
 
++PropertyRedirects=(OldName="KismetMathLibrary.RandomPointInBoundingBox.Origin",NewName="Center")
++PropertyRedirects=(OldName="KismetMathLibrary.RandomPointInBoundingBox.BoxExtent",NewName="HalfSize")
+
 +PropertyRedirects=(OldName="Widget.Visiblity",NewName="Widget.Visibility")
 +PropertyRedirects=(OldName="WidgetBlueprintLibrary.SetInputMode_UIOnlyEx.Target",NewName="PlayerController")
 +PropertyRedirects=(OldName="WidgetBlueprintLibrary.SetInputMode_GameAndUIEx.Target",NewName="PlayerController")
diff --git a/Engine/Source/Runtime/Core/Public/Math/RandomStream.h b/Engine/Source/Runtime/Core/Public/Math/RandomStream.h
--- a/Engine/Source/Runtime/Core/Public/Math/RandomStream.h
+++ b/Engine/Source/Runtime/Core/Public/Math/RandomStream.h
@@ -220,6 +220,13 @@ public:
 		return GetUnitVector();
 	}
 
+	FORCEINLINE FVector RandPointInBox(const FBox& Box) const
+	{
+		return FVector(	FRandRange(Box.Min.X, Box.Max.X),
+						FRandRange(Box.Min.Y, Box.Max.Y),
+						FRandRange(Box.Min.Z, Box.Max.Z) );
+	}
+
 	/**
 	 * Returns a random unit vector, uniformly distributed, within the specified cone.
 	 *
diff --git a/Engine/Source/Runtime/Engine/Classes/Kismet/KismetMathLibrary.h b/Engine/Source/Runtime/Engine/Classes/Kismet/KismetMathLibrary.h
--- a/Engine/Source/Runtime/Engine/Classes/Kismet/KismetMathLibrary.h
+++ b/Engine/Source/Runtime/Engine/Classes/Kismet/KismetMathLibrary.h
@@ -1982,7 +1982,7 @@ class ENGINE_API UKismetMathLibrary : public UBlueprintFunctionLibrary
 
 	/** Returns a random point within the specified bounding box using the first vector as an origin and the second as the box extents. */
 	UFUNCTION(BlueprintPure, Category = "Math|Random", meta=(ScriptMethod = "RandomPointInBoxExtents", NotBlueprintThreadSafe))
-	static FVector RandomPointInBoundingBox(FVector Origin, FVector BoxExtent);
+	static FVector RandomPointInBoundingBox(const FVector Center, const FVector HalfSize);
 
 	/** 
 	 * Returns a random vector with length of 1, within the specified cone, with uniform random distribution.
@@ -3961,6 +3961,10 @@ class ENGINE_API UKismetMathLibrary : public UBlueprintFunctionLibrary
 	UFUNCTION(BlueprintPure, Category="Math|Random")
 	static FVector RandomUnitVectorFromStream(const FRandomStream& Stream);
 
+	/** Returns a random point within the specified bounding box using the first vector as an origin and the second as the half size of the AABB. */
+	UFUNCTION(BlueprintPure, Category="Math|Random")
+	static FVector RandomPointInBoundingBoxFromStream(const FVector Center, const FVector HalfSize, const FRandomStream& Stream);
+
 	/** Create a random rotation */
 	UFUNCTION(BlueprintPure, Category="Math|Random")
 	static FRotator RandomRotatorFromStream(bool bRoll, const FRandomStream& Stream);
diff --git a/Engine/Source/Runtime/Engine/Classes/Kismet/KismetMathLibrary.inl b/Engine/Source/Runtime/Engine/Classes/Kismet/KismetMathLibrary.inl
--- a/Engine/Source/Runtime/Engine/Classes/Kismet/KismetMathLibrary.inl
+++ b/Engine/Source/Runtime/Engine/Classes/Kismet/KismetMathLibrary.inl
@@ -369,7 +369,7 @@ int64 UKismetMathLibrary::SignOfInteger64(int64 A)
 KISMET_MATH_FORCEINLINE
 int64 UKismetMathLibrary::RandomInteger64(int64 A)
 {
-	return FMath::RandHelper(A);
+	return FMath::RandHelper64(A);
 }
 
 KISMET_MATH_FORCEINLINE
@@ -1768,10 +1768,10 @@ FVector UKismetMathLibrary::Vector_Reciprocal(const FVector& A)
 }	
 
 KISMET_MATH_INLINE
-FVector UKismetMathLibrary::RandomPointInBoundingBox(FVector Origin, FVector BoxExtent)
+FVector UKismetMathLibrary::RandomPointInBoundingBox(const FVector Center, const FVector HalfSize)
 {
-	const FVector BoxMin = Origin - BoxExtent;
-	const FVector BoxMax = Origin + BoxExtent;
+	const FVector BoxMin = Center - HalfSize;
+	const FVector BoxMax = Center + HalfSize;
 	return FMath::RandPointInBox(FBox(BoxMin, BoxMax));
 }
 
diff --git a/Engine/Source/Runtime/Engine/Private/KismetMathLibrary.cpp b/Engine/Source/Runtime/Engine/Private/KismetMathLibrary.cpp
--- a/Engine/Source/Runtime/Engine/Private/KismetMathLibrary.cpp
+++ b/Engine/Source/Runtime/Engine/Private/KismetMathLibrary.cpp
@@ -1008,6 +1008,13 @@ FVector UKismetMathLibrary::RandomUnitVectorFromStream(const FRandomStream& Stre
 	return Stream.VRand();
 }
 
+FVector UKismetMathLibrary::RandomPointInBoundingBoxFromStream(const FVector Center, const FVector HalfSize, const FRandomStream& Stream)
+{
+	const FVector BoxMin = Center - HalfSize;
+	const FVector BoxMax = Center + HalfSize;
+	return Stream.RandPointInBox(FBox(BoxMin, BoxMax));
+}
+
 FRotator UKismetMathLibrary::RandomRotatorFromStream(bool bRoll, const FRandomStream& Stream)
 {
 	FRotator RRot;
