LWC: Replaced Promote/DemoteArrayType with a single ConvertArrayType now that conversion is possible via explicit cast.

diff --git a/Engine/Plugins/Enterprise/DatasmithImporter/Source/DatasmithNativeTranslator/Private/DatasmithNativeTranslator.cpp b/Engine/Plugins/Enterprise/DatasmithImporter/Source/DatasmithNativeTranslator/Private/DatasmithNativeTranslator.cpp
--- a/Engine/Plugins/Enterprise/DatasmithImporter/Source/DatasmithNativeTranslator/Private/DatasmithNativeTranslator.cpp
+++ b/Engine/Plugins/Enterprise/DatasmithImporter/Source/DatasmithNativeTranslator/Private/DatasmithNativeTranslator.cpp
@@ -138,7 +138,7 @@ bool FDatasmithNativeTranslator::LoadStaticMesh(const TSharedRef<IDatasmithMeshE
 				SourceModel.RawMeshBulkData.LoadRawMesh( RawMesh );
 				if ( RawMesh.VertexPositions.Num() > 0 )
 				{
-					OutMeshPayload.CollisionPointCloud = LWC::PromoteArrayType<FVector>(RawMesh.VertexPositions);	// LWC_TODO: Perf pessimization. Was MoveTemp.
+					OutMeshPayload.CollisionPointCloud = LWC::ConvertArrayType<FVector>(RawMesh.VertexPositions);	// LWC_TODO: Perf pessimization. Was MoveTemp.
 					break;
 				}
 				ExtractionFailure++;
diff --git a/Engine/Source/Editor/UnrealEd/Private/ImportUtils/StaticMeshImportUtils.cpp b/Engine/Source/Editor/UnrealEd/Private/ImportUtils/StaticMeshImportUtils.cpp
--- a/Engine/Source/Editor/UnrealEd/Private/ImportUtils/StaticMeshImportUtils.cpp
+++ b/Engine/Source/Editor/UnrealEd/Private/ImportUtils/StaticMeshImportUtils.cpp
@@ -579,7 +579,7 @@ bool StaticMeshImportUtils::AddConvexGeomFromVertices( const TArray<FVector3f>&
 	}
 
 	FKConvexElem* ConvexElem = new(AggGeom->ConvexElems) FKConvexElem();
-	ConvexElem->VertexData = LWC::PromoteArrayType<FVector>(Verts);	// LWC_TODO: Perf pessimization
+	ConvexElem->VertexData = LWC::ConvertArrayType<FVector>(Verts);	// LWC_TODO: Perf pessimization
 	ConvexElem->UpdateElemBox();
 
 	return true;
diff --git a/Engine/Source/Runtime/Core/Public/Math/UnrealMathUtility.h b/Engine/Source/Runtime/Core/Public/Math/UnrealMathUtility.h
--- a/Engine/Source/Runtime/Core/Public/Math/UnrealMathUtility.h
+++ b/Engine/Source/Runtime/Core/Public/Math/UnrealMathUtility.h
@@ -2191,34 +2191,25 @@ public:
 namespace LWC
 {
 
-
-// Convert array via constructor (for float -> double conversions)
+// Convert array to a new type
 template<typename TDest, typename TSrc, typename InAllocatorType>
-TArray<TDest, InAllocatorType> PromoteArrayType(const TArray<TSrc, InAllocatorType>& From)
+TArray<TDest, InAllocatorType> ConvertArrayType(const TArray<TSrc, InAllocatorType>& From)
 {
-	//static_assert(!std::is_same<TDest, TSrc>::value, "Redundant call to PromoteArrayType");	// Unavoidable if supporting LWC toggle, but a useful check once LWC is locked to enabled.
-	TArray<TDest, InAllocatorType> Converted;
-	Converted.Reserve(From.Num());
-	for (const TSrc& Item : From)
+	//static_assert(!std::is_same<TDest, TSrc>::value, "Redundant call to ConvertArrayType");	// Unavoidable if supporting LWC toggle, but a useful check once LWC is locked to enabled.
+	if constexpr (std::is_same<TDest, TSrc>::value)
 	{
-		Converted.Add(Item);
+		return From;
 	}
-	return Converted;
-}
-
-
-// Convert array via ToFloatType converter (For double -> float conversions)
-template<typename TDest, typename TSrc, typename InAllocatorType>
-TArray<TDest, InAllocatorType> DemoteArrayType(const TArray<TSrc, InAllocatorType>& From)
-{
-	//static_assert(!std::is_same<TDest, TSrc>::value, "Redundant call to DemoteArrayType");	// Unavoidable if supporting LWC toggle, but a useful check once LWC is locked to enabled.
-	TArray<TDest, InAllocatorType> Converted;
-	Converted.Reserve(From.Num());
-	for (const TSrc& Item : From)
+	else
 	{
-		Converted.Add(static_cast<TDest>(Item));
+		TArray<TDest, InAllocatorType> Converted;
+		Converted.Reserve(From.Num());
+		for (const TSrc& Item : From)
+		{
+			Converted.Add(static_cast<TDest>(Item));
+		}
+		return Converted;
 	}
-	return Converted;
 }
 
 }
\ No newline at end of file
diff --git a/Engine/Source/Runtime/Engine/Private/AI/Navigation/NavigationModifier.cpp b/Engine/Source/Runtime/Engine/Private/AI/Navigation/NavigationModifier.cpp
--- a/Engine/Source/Runtime/Engine/Private/AI/Navigation/NavigationModifier.cpp
+++ b/Engine/Source/Runtime/Engine/Private/AI/Navigation/NavigationModifier.cpp
@@ -345,7 +345,7 @@ void FAreaNavModifier::InitializeConvex(const TNavStatArray<FVector>& InPoints,
 	check(InPoints.IsValidIndex(FirstIndex) && InPoints.IsValidIndex(LastIndex-1));
 
 	Init(InAreaClass);
-	SetConvex(LWC::DemoteArrayType<FVector3f>(InPoints).GetData(), FirstIndex, LastIndex, ENavigationCoordSystem::Unreal, LocalToWorld);	// LWC_TODO: Perf pessimization
+	SetConvex(LWC::ConvertArrayType<FVector3f>(InPoints).GetData(), FirstIndex, LastIndex, ENavigationCoordSystem::Unreal, LocalToWorld);	// LWC_TODO: Perf pessimization
 }
 
 void FAreaNavModifier::InitializePerInstanceConvex(const TNavStatArray<FVector>& InPoints, const int32 FirstIndex, const int32 LastIndex, const TSubclassOf<UNavAreaBase> InAreaClass)
@@ -353,7 +353,7 @@ void FAreaNavModifier::InitializePerInstanceConvex(const TNavStatArray<FVector>&
 	check(InPoints.IsValidIndex(FirstIndex) && InPoints.IsValidIndex(LastIndex - 1));
 
 	Init(InAreaClass);
-	SetPerInstanceConvex(LWC::DemoteArrayType<FVector3f>(InPoints).GetData(), FirstIndex, LastIndex);	// LWC_TODO: Perf pessimization
+	SetPerInstanceConvex(LWC::ConvertArrayType<FVector3f>(InPoints).GetData(), FirstIndex, LastIndex);	// LWC_TODO: Perf pessimization
 }
 
 FAreaNavModifier::FAreaNavModifier(const UBrushComponent* BrushComponent, const TSubclassOf<UNavAreaBase> InAreaClass)
@@ -895,7 +895,7 @@ void FCompositeNavModifier::CreateAreaModifiers(const UPrimitiveComponent* PrimC
 		const FKConvexElem& ConvexElem = BodySetup->AggGeom.ConvexElems[Idx];
 		if (ConvexElem.VertexData.Num() > 0)
 		{
-			FAreaNavModifier AreaMod(LWC::DemoteArrayType<FVector3f>(ConvexElem.VertexData), 0, ConvexElem.VertexData.Num(), ENavigationCoordSystem::Unreal, PrimComp->GetComponentTransform(), AreaClass);
+			FAreaNavModifier AreaMod(LWC::ConvertArrayType<FVector3f>(ConvexElem.VertexData), 0, ConvexElem.VertexData.Num(), ENavigationCoordSystem::Unreal, PrimComp->GetComponentTransform(), AreaClass);
 			Add(AreaMod);
 		}
 		else
diff --git a/Engine/Source/Runtime/Engine/Private/Components/SkinnedMeshComponent.cpp b/Engine/Source/Runtime/Engine/Private/Components/SkinnedMeshComponent.cpp
--- a/Engine/Source/Runtime/Engine/Private/Components/SkinnedMeshComponent.cpp
+++ b/Engine/Source/Runtime/Engine/Private/Components/SkinnedMeshComponent.cpp
@@ -3978,7 +3978,7 @@ void USkinnedMeshComponent::SetPreSkinningOffsets(int32 LODIndex, TArray<FVector
 		uint32 VertexCount = LODData.GetNumVertices();
 		Offsets.SetNumZeroed(VertexCount);
 
-		Info.PreSkinningOffsets = LWC::DemoteArrayType<FVector3f>(Offsets);	// LWC_TODO: Perf pessimization - was MoveTemp!
+		Info.PreSkinningOffsets = LWC::ConvertArrayType<FVector3f>(Offsets);	// LWC_TODO: Perf pessimization - was MoveTemp!
 
 		MarkRenderDynamicDataDirty();
 	}
@@ -4001,7 +4001,7 @@ void USkinnedMeshComponent::SetPostSkinningOffsets(int32 LODIndex, TArray<FVecto
 		uint32 VertexCount = LODData.GetNumVertices();
 		Offsets.SetNumZeroed(VertexCount);
 
-		Info.PostSkinningOffsets = LWC::DemoteArrayType<FVector3f>(Offsets);	// LWC_TODO: Perf pessimization - was MoveTemp!
+		Info.PostSkinningOffsets = LWC::ConvertArrayType<FVector3f>(Offsets);	// LWC_TODO: Perf pessimization - was MoveTemp!
 
 		MarkRenderDynamicDataDirty();
 	}
diff --git a/Engine/Source/Runtime/Engine/Private/PhysicsEngine/PhysAnim.cpp b/Engine/Source/Runtime/Engine/Private/PhysicsEngine/PhysAnim.cpp
--- a/Engine/Source/Runtime/Engine/Private/PhysicsEngine/PhysAnim.cpp
+++ b/Engine/Source/Runtime/Engine/Private/PhysicsEngine/PhysAnim.cpp
@@ -707,7 +707,7 @@ void USkeletalMeshComponent::UpdateKinematicBonesToAnim(const TArray<FTransform>
 						}
 					}
 				}
-				BodyInstance.UpdateTriMeshVertices(LWC::PromoteArrayType<FVector>(NewPositions));
+				BodyInstance.UpdateTriMeshVertices(LWC::ConvertArrayType<FVector>(NewPositions));
 			}
 			
 			BodyInstance.SetBodyTransform(CurrentLocalToWorld, Teleport);
