Fixed ISM per-instance selection rendering

It used to draw every instance as selected, rather than only the instances that were selected. This seems to be due to it enabling dynamic relevance when selected, but without disabling static relevance (as HISM does, which renders correctly).

diff --git a/Engine/Source/Runtime/Engine/Classes/Engine/InstancedStaticMesh.h b/Engine/Source/Runtime/Engine/Classes/Engine/InstancedStaticMesh.h
--- a/Engine/Source/Runtime/Engine/Classes/Engine/InstancedStaticMesh.h
+++ b/Engine/Source/Runtime/Engine/Classes/Engine/InstancedStaticMesh.h
@@ -558,6 +558,7 @@ public:
 			if( bHasSelectedInstances )
 			{
 				Result.bDynamicRelevance = true;
+				Result.bStaticRelevance = false;
 			}
 #endif
 		}
diff --git a/Engine/Source/Runtime/Engine/Private/InstancedStaticMesh.cpp b/Engine/Source/Runtime/Engine/Private/InstancedStaticMesh.cpp
--- a/Engine/Source/Runtime/Engine/Private/InstancedStaticMesh.cpp
+++ b/Engine/Source/Runtime/Engine/Private/InstancedStaticMesh.cpp
@@ -3929,7 +3929,6 @@ void UInstancedStaticMeshComponent::SelectInstance(bool bInSelected, int32 InIns
 				{
 					// Record if the instance is selected
 					FColor HitProxyColor(ForceInit);
-					bool bSelected = SelectedInstances[InstanceIndex] != 0;
 					if (PerInstanceRenderData->HitProxies.IsValidIndex(InstanceIndex))
 					{
 						HitProxyColor = PerInstanceRenderData->HitProxies[InstanceIndex]->Id.GetColor();
@@ -3938,7 +3937,7 @@ void UInstancedStaticMeshComponent::SelectInstance(bool bInSelected, int32 InIns
 					const int32 RenderIndex = GetRenderIndex(InstanceIndex);
 					if (RenderIndex != INDEX_NONE)
 					{
-						InstanceUpdateCmdBuffer.SetEditorData(RenderIndex, HitProxyColor, bSelected);
+						InstanceUpdateCmdBuffer.SetEditorData(RenderIndex, HitProxyColor, bInSelected);
 					}
 				}
 			}			
