UWorld::CleanupWorldInternal should not clear standalone flags in itself or other packages when running in commandlet; that will interfere with the load/save of the other package, and is unnecessary in the cooker. This is a fix for the recent change that sets bCleanupResources=true when CleanupWorld is called from CleanupPhysicsSceneThatWasInitializedForSave in the cook commandlet.
#rnx

diff --git a/Engine/Source/Editor/UnrealEd/Private/EditorEngine.cpp b/Engine/Source/Editor/UnrealEd/Private/EditorEngine.cpp
--- a/Engine/Source/Editor/UnrealEd/Private/EditorEngine.cpp
+++ b/Engine/Source/Editor/UnrealEd/Private/EditorEngine.cpp
@@ -4228,7 +4228,8 @@ void UEditorEngine::CleanupPhysicsSceneThatWasInitializedForSave(UWorld* World,
 
 	if(bForceInitialized)
 	{
-		World->CleanupWorld();
+		// Cleanup the world, but with bCleanupResources=false so that we don't try to remove the world and dependent packages from memory and break future saving
+		World->CleanupWorld(true /* bSessionEnded=default */, false /* bCleanupResources */);
 	}
 
 	World->SetPhysicsScene(nullptr);
diff --git a/Engine/Source/Runtime/Engine/Classes/Engine/Level.h b/Engine/Source/Runtime/Engine/Classes/Engine/Level.h
--- a/Engine/Source/Runtime/Engine/Classes/Engine/Level.h
+++ b/Engine/Source/Runtime/Engine/Classes/Engine/Level.h
@@ -848,8 +848,10 @@ public:
 	/**
 	 * Flag this level instance for destruction.
 	 * This is called by UWorld::CleanupWorld to flag the level and its owned packages for destruction.
+	 *
+	 * @param bCleanupResources Whether to also remove RF_Standalone flags in the editor on level-specific subassets so that they will be garbage collected
 	 */
-	ENGINE_API void CleanupLevel();
+	ENGINE_API void CleanupLevel(bool bCleanupResources);
 
 	/**
 	 * Clears all components of actors associated with this level (aka in Actors array) and 
diff --git a/Engine/Source/Runtime/Engine/Private/Level.cpp b/Engine/Source/Runtime/Engine/Private/Level.cpp
--- a/Engine/Source/Runtime/Engine/Private/Level.cpp
+++ b/Engine/Source/Runtime/Engine/Private/Level.cpp
@@ -368,24 +368,27 @@ void ULevel::AddReferencedObjects(UObject* InThis, FReferenceCollector& Collecto
 	Super::AddReferencedObjects( This, Collector );
 }
 
-void ULevel::CleanupLevel()
+void ULevel::CleanupLevel(bool bCleanupResources)
 {
-	OnCleanupLevel.Broadcast();
-	// if the level contains any actor with an external package, clear their metadata standalone flag so that the packages can be properly unloaded.
-	// Do so for any actors outered to the level and not just from the level actors array
-	ForEachObjectWithOuter(this, [](UObject* InObject)
-		{
-			// Currently only actors would have external package, 
-			// but we can ask directly on all objects since the tests validate against object flags before doing a hash lookup
-			if (UPackage* ExternalPackage = InObject->GetExternalPackage())
+	if (bCleanupResources)
+	{
+		OnCleanupLevel.Broadcast();
+		// if the level contains any actor with an external package, clear their metadata standalone flag so that the packages can be properly unloaded.
+		// Do so for any actors outered to the level and not just from the level actors array
+		ForEachObjectWithOuter(this, [](UObject* InObject)
 			{
-				ForEachObjectWithPackage(ExternalPackage, [](UObject* Object)
-					{
-						Object->ClearFlags(RF_Standalone);
-						return true;
-					}, false);
-			}
-		}, false);
+				// Currently only actors would have external package, 
+				// but we can ask directly on all objects since the tests validate against object flags before doing a hash lookup
+				if (UPackage* ExternalPackage = InObject->GetExternalPackage())
+				{
+					ForEachObjectWithPackage(ExternalPackage, [](UObject* Object)
+						{
+							Object->ClearFlags(RF_Standalone);
+							return true;
+						}, false);
+				}
+			}, false);
+	}
 
 	if (UWorldPartition* WorldPartition = GetWorldPartition())
 	{
diff --git a/Engine/Source/Runtime/Engine/Private/World.cpp b/Engine/Source/Runtime/Engine/Private/World.cpp
--- a/Engine/Source/Runtime/Engine/Private/World.cpp
+++ b/Engine/Source/Runtime/Engine/Private/World.cpp
@@ -3133,7 +3133,7 @@ void FLevelStreamingGCHelper::PrepareStreamedOutLevelsForGC()
 			// Make sure that this package has been unloaded after GC pass.
 			LevelPackageNames.Add( LevelPackage->GetFName() );
 
-			Level->CleanupLevel();
+			Level->CleanupLevel(true /* bCleanupResources */);
 
 			// Mark world and all other package subobjects as pending kill
 			// This will destroy metadata objects and any other objects left behind
@@ -4758,42 +4758,45 @@ void UWorld::CleanupWorldInternal(bool bSessionEnded, bool bCleanupResources, UW
 #if WITH_EDITOR
 	// Clear standalone flag when switching maps in the Editor. This causes resources placed in the map
 	// package to be garbage collected together with the world.
-	if( GIsEditor && !IsTemplate() && bCleanupResources && this != NewWorld )
+	if( GIsEditor && !IsTemplate() && this != NewWorld )
 	{
-		// Iterate over all objects to find ones that reside in the same package as the world.
-		ForEachObjectWithPackage(GetOutermost(), [this](UObject* CurrentObject)
+		if (bCleanupResources)
 		{
-			if ( CurrentObject != this )
+			// Iterate over all objects to find ones that reside in the same package as the world.
+			ForEachObjectWithPackage(GetOutermost(), [this](UObject* CurrentObject)
 			{
-				CurrentObject->ClearFlags( RF_Standalone );
-			}
-			return true;
-		});
+				if ( CurrentObject != this )
+				{
+					CurrentObject->ClearFlags( RF_Standalone );
+				}
+				return true;
+			});
 
-		if (WorldType != EWorldType::PIE)
-		{
-			if (PersistentLevel && PersistentLevel->MapBuildData)
+			if (WorldType != EWorldType::PIE)
 			{
-				PersistentLevel->MapBuildData->ClearFlags(RF_Standalone);
-
-				// Iterate over all objects to find ones that reside in the same package as the MapBuildData.
-				// Specifically the PackageMetaData
-				ForEachObjectWithPackage(PersistentLevel->MapBuildData->GetOutermost(), [this](UObject* CurrentObject)
+				if (PersistentLevel && PersistentLevel->MapBuildData)
 				{
-					if (CurrentObject != this)
+					PersistentLevel->MapBuildData->ClearFlags(RF_Standalone);
+
+					// Iterate over all objects to find ones that reside in the same package as the MapBuildData.
+					// Specifically the PackageMetaData
+					ForEachObjectWithPackage(PersistentLevel->MapBuildData->GetOutermost(), [this](UObject* CurrentObject)
 					{
-						CurrentObject->ClearFlags(RF_Standalone);
-					}
-					return true;
-				});
+						if (CurrentObject != this)
+						{
+							CurrentObject->ClearFlags(RF_Standalone);
+						}
+						return true;
+					});
+				}
 			}
 		}
-				
-		// Cleanup Persistent level outside of following loop because unitialized worlds don't have a valid Levels array
+
+		// Cleanup Persistent level outside of following loop because uninitialized worlds don't have a valid Levels array
 		// StreamingLevels are not initialized.
 		if (PersistentLevel)
 		{
-			PersistentLevel->CleanupLevel();
+			PersistentLevel->CleanupLevel(bCleanupResources);
 		}
 
 		if (GetNumLevels() > 1)
@@ -4801,7 +4804,7 @@ void UWorld::CleanupWorldInternal(bool bSessionEnded, bool bCleanupResources, UW
 			check(GetLevel(0) == PersistentLevel);
 			for (int32 LevelIndex = 1; LevelIndex < GetNumLevels(); ++LevelIndex)
 			{
-				GetLevel(LevelIndex)->CleanupLevel();
+				GetLevel(LevelIndex)->CleanupLevel(bCleanupResources);
 			}
 		}
 	}
@@ -4830,7 +4833,7 @@ void UWorld::CleanupWorldInternal(bool bSessionEnded, bool bCleanupResources, UW
 		{
 			if (Level)
 			{
-			UWorld* const LevelWorld = CastChecked<UWorld>(Level->GetOuter());
+				UWorld* const LevelWorld = CastChecked<UWorld>(Level->GetOuter());
 				LevelWorld->CleanupWorldInternal(bSessionEnded, bCleanupResources, NewWorld);
 			}
 		}
