Ensure that DDC fill operation pulls VTChunks for Textures.

diff --git a/Engine/Source/Runtime/Engine/Private/TextureDerivedDataTask.cpp b/Engine/Source/Runtime/Engine/Private/TextureDerivedDataTask.cpp
--- a/Engine/Source/Runtime/Engine/Private/TextureDerivedDataTask.cpp
+++ b/Engine/Source/Runtime/Engine/Private/TextureDerivedDataTask.cpp
@@ -777,6 +777,23 @@ void FTextureCacheDerivedDataWorker::DoWork()
 		{
 			bSucceeded = DerivedData->TryLoadMips(0, nullptr, &Texture);
 
+			if (bForVirtualTextureStreamingBuild)
+			{
+				if (DerivedData->VTData != nullptr &&
+					DerivedData->VTData->IsInitialized())
+				{
+					TArray<FString, TInlineAllocator<16>> ChunkKeys;
+					for (const FVirtualTextureDataChunk& Chunk : DerivedData->VTData->Chunks)
+					{
+						if (!Chunk.DerivedDataKey.IsEmpty())
+						{
+							ChunkKeys.Add(Chunk.DerivedDataKey);
+						}
+					}
+					GetDerivedDataCacheRef().TryToPrefetch(ChunkKeys, TEXT("DerivedVTChunks"_SV));
+				}
+			}
+
 			if (!bSucceeded)
 			{
 				UE_LOG(LogTexture, Display, TEXT("Texture %s is missing mips. The texture will be rebuilt."), *Texture.GetFullName());
