Fixed ISM instances being left at the wrong location after undo/redo

diff --git a/Engine/Source/Runtime/Engine/Private/InstancedStaticMesh.cpp b/Engine/Source/Runtime/Engine/Private/InstancedStaticMesh.cpp
--- a/Engine/Source/Runtime/Engine/Private/InstancedStaticMesh.cpp
+++ b/Engine/Source/Runtime/Engine/Private/InstancedStaticMesh.cpp
@@ -3891,6 +3891,9 @@ void UInstancedStaticMeshComponent::PostEditUndo()
 	Super::PostEditUndo();
 
 	FNavigationSystem::UpdateComponentData(*this);
+
+	InstanceUpdateCmdBuffer.Edit();
+	MarkRenderStateDirty();
 }
 #endif
 
