Allow DDC fill to pull shipping textures optionally via commandline switch.

diff --git a/Engine/Source/Runtime/Engine/Private/TextureDerivedData.cpp b/Engine/Source/Runtime/Engine/Private/TextureDerivedData.cpp
--- a/Engine/Source/Runtime/Engine/Private/TextureDerivedData.cpp
+++ b/Engine/Source/Runtime/Engine/Private/TextureDerivedData.cpp
@@ -738,6 +738,7 @@ void FTexturePlatformData::Cache(
 	uint32 Flags = InFlags;
 
 	static bool bForDDC = FString(FCommandLine::Get()).Contains(TEXT("Run=DerivedDataCache"));
+	static bool bDDCShippingTextures = FParse::Param(FCommandLine::Get(), TEXT("DDCShippingTextures"));
 	if (bForDDC)
 	{
 		Flags |= ETextureCacheFlags::ForDDCBuild;
@@ -747,7 +748,7 @@ void FTexturePlatformData::Cache(
 	bool bAsync = (Flags & ETextureCacheFlags::Async) != 0;
 	GetTextureDerivedDataKey(InTexture, InSettingsPerLayer, DerivedDataKey);
 
-	if (!bForDDC && !bForceRebuild && InSettingsPerLayer[0].bHasEditorOnlyData)
+	if ((!bForDDC || bDDCShippingTextures) && !bForceRebuild && InSettingsPerLayer[0].bHasEditorOnlyData)
 	{
 		// If we are running in the Editor, generate a shipping derived data key, which can be used to load
 		// a texture cooked for the shipping build from the cache, if available.
