Disable calling post build mesh data during cooking/caching of skeletal mesh data because memory overhead to store the prebuild BLAS data is too high right now and it's not used

diff --git a/Engine/Source/Runtime/Engine/Private/SkeletalMeshRenderData.cpp b/Engine/Source/Runtime/Engine/Private/SkeletalMeshRenderData.cpp
--- a/Engine/Source/Runtime/Engine/Private/SkeletalMeshRenderData.cpp
+++ b/Engine/Source/Runtime/Engine/Private/SkeletalMeshRenderData.cpp
@@ -42,6 +42,11 @@ static TAutoConsoleVariable<int32> CVarSkeletalMeshKeepMobileMinLODSettingOnDesk
 	0,
 	TEXT("If non-zero, mobile setting for MinLOD will be stored in the cooked data for desktop platforms"));
 
+static TAutoConsoleVariable<int32> CVarSkeletalMeshPostBuildMeshData(
+	TEXT("r.SkeletalMesh.PostBuildMeshData"),
+	0,
+	TEXT("If non-zero, the platform specific post build mesh data will be called during caching the data (used to generate prebuild BLAS data on supported platforms)"));
+
 #if WITH_EDITOR
 
 /** 
@@ -186,7 +191,7 @@ static void SerializeLODInfoForDDC(USkeletalMesh* SkeletalMesh, FString& KeySuff
 // differences, etc.) replace the version GUID below with a new one.
 // In case of merge conflicts with DDC versions, you MUST generate a new GUID
 // and set this new GUID as the version.
-#define SKELETALMESH_DERIVEDDATA_VER TEXT("23BBACC58782442DAAE9DCBFC4715C65")
+#define SKELETALMESH_DERIVEDDATA_VER TEXT("D6FFBA9B10194A0B82AFA70334487134")
 
 const FString& GetSkeletalMeshDerivedDataVersion()
 {
@@ -547,10 +552,12 @@ void FSkeletalMeshRenderData::Cache(const ITargetPlatform* TargetPlatform, USkel
 				}
 			}
 
-			IMeshBuilderModule& MeshBuilderModule = IMeshBuilderModule::GetForPlatform(TargetPlatform);
-
-			MeshBuilderModule.PostBuildSkeletalMesh(this, Owner);
-
+			if (CVarSkeletalMeshPostBuildMeshData.GetValueOnAnyThread() != 0)
+			{
+				IMeshBuilderModule& MeshBuilderModule = IMeshBuilderModule::GetForPlatform(TargetPlatform);
+				MeshBuilderModule.PostBuildSkeletalMesh(this, Owner);
+			}
+			
 			//Serialize the render data
 			Serialize(Ar, Owner);
 			for (int32 LODIndex = 0; LODIndex < LODRenderData.Num(); ++LODIndex)
