Fix a bug where linked anim layers did not work with WasAnimNotify functions

diff --git a/Engine/Source/Runtime/Engine/Classes/Animation/AnimInstance.h b/Engine/Source/Runtime/Engine/Classes/Animation/AnimInstance.h
--- a/Engine/Source/Runtime/Engine/Classes/Animation/AnimInstance.h
+++ b/Engine/Source/Runtime/Engine/Classes/Animation/AnimInstance.h
@@ -1402,6 +1402,8 @@ private:
 	TMap<FName, FMontageActiveSlotTracker> SlotWeightTracker;
 	TMap<FName, FSimpleMulticastDelegate> ExternalNotifyHandlers;
 
+	bool CheckOnInstanceAndMainInstance(TFunctionRef<bool (FAnimInstanceProxy* )> ProxyLambdaFunc);
+
 public:
 	/** 
 	 * Recalculate Required Bones [RequiredBones]
diff --git a/Engine/Source/Runtime/Engine/Private/Animation/AnimInstance.cpp b/Engine/Source/Runtime/Engine/Private/Animation/AnimInstance.cpp
--- a/Engine/Source/Runtime/Engine/Private/Animation/AnimInstance.cpp
+++ b/Engine/Source/Runtime/Engine/Private/Animation/AnimInstance.cpp
@@ -3481,49 +3481,101 @@ float UAnimInstance::GetRelevantAnimTimeFraction(int32 MachineIndex, int32 State
 	return GetProxyOnAnyThread<FAnimInstanceProxy>().GetRelevantAnimTimeFraction(MachineIndex, StateIndex);
 }
 
+bool UAnimInstance::CheckOnInstanceAndMainInstance(TFunctionRef<bool (FAnimInstanceProxy* )> ProxyLambdaFunc)
+{
+	FAnimInstanceProxy& InstanceProxy = GetProxyOnAnyThread<FAnimInstanceProxy>();
+	if (ProxyLambdaFunc(&InstanceProxy))
+	{
+		return true;
+	}
+
+	FAnimInstanceProxy* MainInstanceProxy = InstanceProxy.GetMainInstanceProxy();
+	if (MainInstanceProxy && MainInstanceProxy != &InstanceProxy)
+	{
+		return ProxyLambdaFunc(MainInstanceProxy);
+	}
+	return false;
+}
+
 bool UAnimInstance::WasAnimNotifyStateActiveInAnyState(TSubclassOf<UAnimNotifyState> AnimNotifyStateType)
 {
-	return GetProxyOnAnyThread<FAnimInstanceProxy>().WasAnimNotifyStateActiveInAnyState(AnimNotifyStateType);
+	return CheckOnInstanceAndMainInstance([&AnimNotifyStateType](FAnimInstanceProxy* Proxy)
+		{
+			return Proxy->WasAnimNotifyStateActiveInAnyState(AnimNotifyStateType);
+		}
+	);
 }
 
 bool UAnimInstance::WasAnimNotifyStateActiveInStateMachine(int32 MachineIndex, TSubclassOf<UAnimNotifyState> AnimNotifyStateType)
 {
-	return GetProxyOnAnyThread<FAnimInstanceProxy>().WasAnimNotifyStateActiveInStateMachine(MachineIndex, AnimNotifyStateType);
+	return CheckOnInstanceAndMainInstance([&MachineIndex, &AnimNotifyStateType](FAnimInstanceProxy* Proxy)
+		{
+			return Proxy->WasAnimNotifyStateActiveInStateMachine(MachineIndex, AnimNotifyStateType);
+		}
+	);
 }
 
 bool UAnimInstance::WasAnimNotifyStateActiveInSourceState(int32 MachineIndex, int32 StateIndex, TSubclassOf<UAnimNotifyState> AnimNotifyStateType)
 {
-	return GetProxyOnAnyThread<FAnimInstanceProxy>().WasAnimNotifyStateActiveInSourceState(MachineIndex, StateIndex, AnimNotifyStateType);
+	return CheckOnInstanceAndMainInstance([&MachineIndex, &StateIndex, &AnimNotifyStateType](FAnimInstanceProxy* Proxy)
+		{
+			return Proxy->WasAnimNotifyStateActiveInSourceState(MachineIndex, StateIndex, AnimNotifyStateType);
+		}
+	);
 }
 
 bool UAnimInstance::WasAnimNotifyTriggeredInSourceState(int32 MachineIndex, int32 StateIndex,  TSubclassOf<UAnimNotify> AnimNotifyType)
 {
-	return GetProxyOnAnyThread<FAnimInstanceProxy>().WasAnimNotifyTriggeredInSourceState(MachineIndex, StateIndex, AnimNotifyType);
+	return CheckOnInstanceAndMainInstance([&MachineIndex, &StateIndex, &AnimNotifyType](FAnimInstanceProxy* Proxy)
+		{
+			return Proxy->WasAnimNotifyTriggeredInSourceState(MachineIndex, StateIndex, AnimNotifyType);
+		}
+	);
 }
 
 bool UAnimInstance::WasAnimNotifyNameTriggeredInSourceState(int32 MachineIndex, int32 StateIndex, FName NotifyName)
 {
-	return GetProxyOnAnyThread<FAnimInstanceProxy>().WasAnimNotifyNameTriggeredInSourceState(MachineIndex, StateIndex, NotifyName);
+	return CheckOnInstanceAndMainInstance([&MachineIndex, &StateIndex, &NotifyName](FAnimInstanceProxy* Proxy)
+		{
+			return Proxy->WasAnimNotifyNameTriggeredInSourceState(MachineIndex, StateIndex, NotifyName);
+		}
+	);
 }
 
 bool UAnimInstance::WasAnimNotifyTriggeredInStateMachine(int32 MachineIndex, TSubclassOf<UAnimNotify> AnimNotifyType)
 {
-	return GetProxyOnAnyThread<FAnimInstanceProxy>().WasAnimNotifyTriggeredInStateMachine(MachineIndex, AnimNotifyType);
+	return CheckOnInstanceAndMainInstance([&MachineIndex, &AnimNotifyType](FAnimInstanceProxy* Proxy)
+		{
+			return Proxy->WasAnimNotifyTriggeredInStateMachine(MachineIndex, AnimNotifyType);
+		}
+	);
 }
 
 bool UAnimInstance::WasAnimNotifyNameTriggeredInStateMachine(int32 MachineIndex, FName NotifyName)
 {
-	return GetProxyOnAnyThread<FAnimInstanceProxy>().WasAnimNotifyNameTriggeredInStateMachine(MachineIndex, NotifyName);
+	return CheckOnInstanceAndMainInstance([&MachineIndex, &NotifyName](FAnimInstanceProxy* Proxy)
+		{
+			return Proxy->WasAnimNotifyNameTriggeredInStateMachine(MachineIndex, NotifyName);
+		}
+	);
 }
 
 bool UAnimInstance::WasAnimNotifyTriggeredInAnyState(TSubclassOf<UAnimNotify> AnimNotifyType)
 {
-	return GetProxyOnAnyThread<FAnimInstanceProxy>().WasAnimNotifyTriggeredInAnyState(AnimNotifyType);
+	return CheckOnInstanceAndMainInstance([&AnimNotifyType](FAnimInstanceProxy* Proxy)
+		{
+			return Proxy->WasAnimNotifyTriggeredInAnyState(AnimNotifyType);
+		}
+	);
 }
 
 bool UAnimInstance::WasAnimNotifyNameTriggeredInAnyState(FName NotifyName)
 {
-	return GetProxyOnAnyThread<FAnimInstanceProxy>().WasAnimNotifyNameTriggeredInAnyState(NotifyName);
+	return CheckOnInstanceAndMainInstance([&NotifyName](FAnimInstanceProxy* Proxy)
+		{
+			return Proxy->WasAnimNotifyNameTriggeredInAnyState(NotifyName);
+		}
+	);
 }
 
 const FAnimNode_StateMachine* UAnimInstance::GetStateMachineInstance(int32 MachineIndex) const
