Fix NonUnity build

diff --git a/Engine/Source/Runtime/Engine/Private/LightWeightInstanceStaticMeshManager.cpp b/Engine/Source/Runtime/Engine/Private/LightWeightInstanceStaticMeshManager.cpp
--- a/Engine/Source/Runtime/Engine/Private/LightWeightInstanceStaticMeshManager.cpp
+++ b/Engine/Source/Runtime/Engine/Private/LightWeightInstanceStaticMeshManager.cpp
@@ -3,6 +3,7 @@
 #include "GameFramework/LightWeightInstanceStaticMeshManager.h"
 #include "Components/HierarchicalInstancedStaticMeshComponent.h"
 #include "Net/UnrealNetwork.h"
+#include "TimerManager.h"
 
 ALightWeightInstanceStaticMeshManager::ALightWeightInstanceStaticMeshManager(const FObjectInitializer& ObjectInitializer)
 {
