Remove ensure when bone doesn't exist on skeleton

diff --git a/Engine/Source/Runtime/Engine/Private/Animation/AnimSequence.cpp b/Engine/Source/Runtime/Engine/Private/Animation/AnimSequence.cpp
--- a/Engine/Source/Runtime/Engine/Private/Animation/AnimSequence.cpp
+++ b/Engine/Source/Runtime/Engine/Private/Animation/AnimSequence.cpp
@@ -4069,9 +4069,9 @@ void UAnimSequence::BakeTrackCurvesToRawAnimationTracks(TArray<FRawAnimSequenceT
 				{
 					TrackIndex = InsertNewBoneTrack(BoneName);
 					// if it still didn't find, something went horribly wrong
-					if (ensure(TrackIndex != INDEX_NONE) == false)
+					if(TrackIndex == INDEX_NONE)
 					{
-						UE_LOG(LogAnimation, Warning, TEXT("Animation Baking : Error adding %s track."), *BoneName.ToString());
+						UE_LOG(LogAnimation, Log, TEXT("Animation Track Curve Baking : Error adding %s track to %s. Bone does not exist. Skipping"), *BoneName.ToString(), *GetName());
 						// I can't do anything about it
 						continue;
 					}
