Chaos : Fix some of the compilation problem when compiling with PhysX

diff --git a/Engine/Plugins/Experimental/ChaosCaching/Source/ChaosCaching/Private/Chaos/Adapters/GeometryCollectionComponentCacheAdapter.cpp b/Engine/Plugins/Experimental/ChaosCaching/Source/ChaosCaching/Private/Chaos/Adapters/GeometryCollectionComponentCacheAdapter.cpp
--- a/Engine/Plugins/Experimental/ChaosCaching/Source/ChaosCaching/Private/Chaos/Adapters/GeometryCollectionComponentCacheAdapter.cpp
+++ b/Engine/Plugins/Experimental/ChaosCaching/Source/ChaosCaching/Private/Chaos/Adapters/GeometryCollectionComponentCacheAdapter.cpp
@@ -53,7 +53,7 @@ namespace Chaos
 		UGeometryCollectionComponent*    Comp  = CastChecked<UGeometryCollectionComponent>(InComp);
 		FGeometryCollectionPhysicsProxy* Proxy = Comp->GetPhysicsProxy();
 		
-		FPhysScene_Chaos& Scene = *(Comp->GetWorld()->GetPhysicsScene());
+		FPhysScene& Scene = *(Comp->GetWorld()->GetPhysicsScene());
 		
 		if(!Proxy)
 		{
diff --git a/Engine/Source/Runtime/Engine/Public/Physics/PhysicsInterfacePhysX.h b/Engine/Source/Runtime/Engine/Public/Physics/PhysicsInterfacePhysX.h
--- a/Engine/Source/Runtime/Engine/Public/Physics/PhysicsInterfacePhysX.h
+++ b/Engine/Source/Runtime/Engine/Public/Physics/PhysicsInterfacePhysX.h
@@ -238,6 +238,8 @@ struct ENGINE_API FPhysicsInterface_PhysX : public FGenericPhysicsInterface, pub
 	// NOTE: Adding this temporarily, until we phase out Handle.IsValid().
 	static bool IsValid(const FPhysicsActorHandle& Handle) { return Handle.IsValid(); }
 
+	static bool IsInScene(const FPhysicsActorHandle& InActorReference);
+
 	template<typename AllocatorType>
 	static int32 GetAllShapes_AssumedLocked(const FPhysicsActorHandle_PhysX& InHandle, TArray<FPhysicsShapeHandle_PhysX, AllocatorType>& OutShapes);
 	static int32 GetNumShapes(const FPhysicsActorHandle_PhysX& InHandle);
@@ -255,7 +257,6 @@ struct ENGINE_API FPhysicsInterface_PhysX : public FGenericPhysicsInterface, pub
 	static bool IsKinematic_AssumesLocked(const FPhysicsActorHandle_PhysX& InHandle);
 	static bool IsSleeping(const FPhysicsActorHandle_PhysX& InHandle);
 	static bool IsCcdEnabled(const FPhysicsActorHandle_PhysX& InHandle);
-	static bool IsCcdEnabled(const FPhysicsActorHandle_PhysX& InHandle);
 	static bool IsDisabled(const FPhysicsActorHandle_PhysX& InHandle);
 	static void SetDisabled(const FPhysicsActorHandle_PhysX& InHandle, bool bSetDisabled);
 	static FPhysScene* GetCurrentScene(const FPhysicsActorHandle_PhysX& InHandle);
diff --git a/Engine/Source/Runtime/Experimental/ChaosSolverEngine/Private/Chaos/ChaosGameplayEventDispatcher.cpp b/Engine/Source/Runtime/Experimental/ChaosSolverEngine/Private/Chaos/ChaosGameplayEventDispatcher.cpp
--- a/Engine/Source/Runtime/Experimental/ChaosSolverEngine/Private/Chaos/ChaosGameplayEventDispatcher.cpp
+++ b/Engine/Source/Runtime/Experimental/ChaosSolverEngine/Private/Chaos/ChaosGameplayEventDispatcher.cpp
@@ -438,7 +438,7 @@ void UChaosGameplayEventDispatcher::HandleBreakingEvents(const Chaos::FBreakingE
 
 void UChaosGameplayEventDispatcher::HandleSleepingEvents(const Chaos::FSleepingEventData& SleepingData)
 {
-	FPhysScene_Chaos& Scene = *(GetWorld()->GetPhysicsScene());
+	FPhysScene& Scene = *(GetWorld()->GetPhysicsScene());
 
 	const Chaos::FSleepingDataArray& SleepingArray = SleepingData.SleepingData;
 
