Texture: Converted the derived data task to an interface that can be implemented separately for the new derived data build interface

#rnx

diff --git a/Engine/Source/Runtime/Engine/Private/TextureCompiler.cpp b/Engine/Source/Runtime/Engine/Private/TextureCompiler.cpp
--- a/Engine/Source/Runtime/Engine/Private/TextureCompiler.cpp
+++ b/Engine/Source/Runtime/Engine/Private/TextureCompiler.cpp
@@ -274,7 +274,7 @@ void FTextureCompilingManager::FinishCompilation(TArrayView<UTexture* const> InT
 
 	if (PendingTextures.Num())
 	{
-		class FCompilableTexture : public AsyncCompilationHelpers::TCompilableAsyncTask<FTextureAsyncCacheDerivedDataTask>
+		class FCompilableTexture final : public AsyncCompilationHelpers::ICompilable
 		{
 		public:
 			FCompilableTexture(UTexture* InTexture)
@@ -282,7 +282,7 @@ void FTextureCompilingManager::FinishCompilation(TArrayView<UTexture* const> InT
 			{
 			}
 
-			FTextureAsyncCacheDerivedDataTask* GetAsyncTask() override
+			FTextureAsyncCacheDerivedDataTask* GetAsyncTask()
 			{
 				if (FTexturePlatformData** PlatformDataPtr = Texture->GetRunningPlatformData())
 				{
@@ -295,8 +295,26 @@ void FTextureCompilingManager::FinishCompilation(TArrayView<UTexture* const> InT
 				return nullptr;
 			}
 
-			TStrongObjectPtr<UTexture> Texture;
+			void Reschedule(FQueuedThreadPool* InThreadPool, EQueuedWorkPriority InPriority) final
+			{
+				if (FTextureAsyncCacheDerivedDataTask* AsyncTask = GetAsyncTask())
+				{
+					AsyncTask->SetPriority(InPriority);
+				}
+			}
+		
+			bool WaitCompletionWithTimeout(float TimeLimitSeconds) final
+			{
+				if (FTextureAsyncCacheDerivedDataTask* AsyncTask = GetAsyncTask())
+				{
+					return AsyncTask->WaitWithTimeout(TimeLimitSeconds);
+				}
+				return true;
+			}
+
 			FName GetName() override { return Texture->GetOutermost()->GetFName(); }
+
+			TStrongObjectPtr<UTexture> Texture;
 		};
 
 		TArray<UTexture*> UniqueTextures(PendingTextures.Array());
@@ -455,7 +473,7 @@ bool FTextureCompilingManager::RequestPriorityChange(UTexture* InTexture, EQueue
 				EQueuedWorkPriority OldPriority = AsyncTask->GetPriority();
 				if (OldPriority != InPriority)
 				{
-					if (AsyncTask->Reschedule(GetThreadPool(), InPriority))
+					if (AsyncTask->SetPriority(InPriority))
 					{
 						UE_LOG(
 							LogTexture,
diff --git a/Engine/Source/Runtime/Engine/Private/TextureDerivedData.cpp b/Engine/Source/Runtime/Engine/Private/TextureDerivedData.cpp
--- a/Engine/Source/Runtime/Engine/Private/TextureDerivedData.cpp
+++ b/Engine/Source/Runtime/Engine/Private/TextureDerivedData.cpp
@@ -780,8 +780,9 @@ void FTexturePlatformData::Cache(
 		FQueuedThreadPool*  TextureThreadPool = FTextureCompilingManager::Get().GetThreadPool();
 		EQueuedWorkPriority BasePriority      = FTextureCompilingManager::Get().GetBasePriority(&InTexture);
 
-		AsyncTask = new FTextureAsyncCacheDerivedDataTask(Compressor, this, &InTexture, InSettingsPerLayer, Flags);
-		AsyncTask->StartBackgroundTask(TextureThreadPool, BasePriority, EQueuedWorkFlags::DoNotRunInsideBusyWait, AsyncTask->GetTask().GetRequiredMemoryEstimate());
+		FTextureAsyncCacheDerivedDataWorkerTask* LocalTask = new FTextureAsyncCacheDerivedDataWorkerTask(TextureThreadPool, Compressor, this, &InTexture, InSettingsPerLayer, Flags);
+		AsyncTask = LocalTask;
+		LocalTask->StartBackgroundTask(TextureThreadPool, BasePriority, EQueuedWorkFlags::DoNotRunInsideBusyWait, LocalTask->GetTask().GetRequiredMemoryEstimate());
 	}
 	else
 	{
@@ -797,16 +798,12 @@ void FTexturePlatformData::Cache(
 
 bool FTexturePlatformData::TryCancelCache()
 {
-	if (AsyncTask)
+	if (AsyncTask && AsyncTask->Cancel())
 	{
-		if (AsyncTask->IsDone() || AsyncTask->Cancel())
-		{
-			delete AsyncTask;
-			AsyncTask = nullptr;
-		}
+		delete AsyncTask;
+		AsyncTask = nullptr;
 	}
-
-	return AsyncTask == nullptr;
+	return !AsyncTask;
 }
 
 void FTexturePlatformData::CancelCache()
@@ -822,7 +819,7 @@ void FTexturePlatformData::CancelCache()
 
 bool FTexturePlatformData::IsAsyncWorkComplete() const
 {
-	return AsyncTask == nullptr || AsyncTask->IsWorkDone();
+	return !AsyncTask || AsyncTask->Poll();
 }
 
 void FTexturePlatformData::FinishCache()
@@ -832,13 +829,14 @@ void FTexturePlatformData::FinishCache()
 		TRACE_CPUPROFILER_EVENT_SCOPE(FTexturePlatformData::FinishCache)
 		{
 			COOK_STAT(auto Timer = TextureCookStats::UsageStats.TimeAsyncWait());
-			AsyncTask->EnsureCompletion();
-			FTextureCacheDerivedDataWorker& Worker = AsyncTask->GetTask();
-			Worker.Finalize();
-			COOK_STAT(Timer.AddHitOrMiss(Worker.WasLoadedFromDDC() ? FCookStats::CallStats::EHitOrMiss::Hit : FCookStats::CallStats::EHitOrMiss::Miss, Worker.GetBytesCached()));
+			bool bFoundInCache = false;
+			uint64 ProcessedByteCount = 0;
+			AsyncTask->Wait();
+			AsyncTask->Finalize(bFoundInCache, ProcessedByteCount);
+			COOK_STAT(Timer.AddHitOrMiss(bFoundInCache ? FCookStats::CallStats::EHitOrMiss::Hit : FCookStats::CallStats::EHitOrMiss::Miss, int64(ProcessedByteCount)));
 		}
 		delete AsyncTask;
-		AsyncTask = NULL;
+		AsyncTask = nullptr;
 	}
 }
 
@@ -995,7 +993,7 @@ FTexturePlatformData::FTexturePlatformData()
 	, PixelFormat(PF_Unknown)
 	, VTData(nullptr)
 #if WITH_EDITORONLY_DATA
-	, AsyncTask(NULL)
+	, AsyncTask(nullptr)
 #endif // #if WITH_EDITORONLY_DATA
 {
 }
@@ -1005,9 +1003,9 @@ FTexturePlatformData::~FTexturePlatformData()
 #if WITH_EDITOR
 	if (AsyncTask)
 	{
-		AsyncTask->EnsureCompletion();
+		AsyncTask->Wait();
 		delete AsyncTask;
-		AsyncTask = NULL;
+		AsyncTask = nullptr;
 	}
 #endif
 	if (VTData) delete VTData;
@@ -1832,30 +1830,32 @@ void UTexture::ClearAllCachedCookedPlatformData()
 	}
 }
 
-bool UTexture::IsCachedCookedPlatformDataLoaded( const ITargetPlatform* TargetPlatform ) 
+bool UTexture::IsCachedCookedPlatformDataLoaded(const ITargetPlatform* TargetPlatform)
 { 
-	const TMap<FString, FTexturePlatformData*> *CookedPlatformDataPtr = GetCookedPlatformData();
-	if ( CookedPlatformDataPtr == NULL )
-		return true; // we should always have cookedplatformDataPtr in the case of WITH_EDITOR
+	const TMap<FString, FTexturePlatformData*>* CookedPlatformDataPtr = GetCookedPlatformData();
+	if (!CookedPlatformDataPtr)
+	{
+		return true; // we should always have CookedPlatformDataPtr in the case of WITH_EDITOR
+	}
 
 	FTextureBuildSettings BuildSettings;
 	TArray<FTexturePlatformData*> PlatformDataToSerialize;
 	GetTextureBuildSettings(*this, TargetPlatform->GetTextureLODSettings(), *TargetPlatform, BuildSettings);
 
-	TArray< TArray<FTextureBuildSettings> > BuildSettingsToCache;
+	TArray<TArray<FTextureBuildSettings>> BuildSettingsToCache;
 	GetBuildSettingsPerFormat(*this, BuildSettings, TargetPlatform, BuildSettingsToCache);
 
-	for (int32 SettingIndex = 0; SettingIndex < BuildSettingsToCache.Num(); SettingIndex++)
+	for (const TArray<FTextureBuildSettings>& PlatformBuildSettings : BuildSettingsToCache)
 	{
-		check(BuildSettingsToCache[SettingIndex].Num() == Source.GetNumLayers());
+		check(PlatformBuildSettings.Num() == Source.GetNumLayers());
 
 		FString DerivedDataKey;
-		GetTextureDerivedDataKey(*this, BuildSettingsToCache[SettingIndex].GetData(), DerivedDataKey);
+		GetTextureDerivedDataKey(*this, PlatformBuildSettings.GetData(), DerivedDataKey);
 
-		FTexturePlatformData *PlatformData= (*CookedPlatformDataPtr).FindRef(DerivedDataKey);
+		FTexturePlatformData* PlatformData = (*CookedPlatformDataPtr).FindRef(DerivedDataKey);
 
 		// begin cache hasn't been called
-		if ( !PlatformData )
+		if (!PlatformData)
 		{
 			if (!HasAnyFlags(RF_ClassDefaultObject) && Source.SizeX != 0 && Source.SizeY != 0)
 			{
@@ -1865,12 +1865,12 @@ bool UTexture::IsCachedCookedPlatformDataLoaded( const ITargetPlatform* TargetPl
 			return false;
 		}
 
-		if ( (PlatformData->AsyncTask != NULL) && ( PlatformData->AsyncTask->IsWorkDone() == true ) )
+		if (PlatformData->AsyncTask && PlatformData->AsyncTask->Poll())
 		{
 			PlatformData->FinishCache();
 		}
 
-		if ( PlatformData->AsyncTask)
+		if (PlatformData->AsyncTask)
 		{
 			return false;
 		}
@@ -1881,28 +1881,24 @@ bool UTexture::IsCachedCookedPlatformDataLoaded( const ITargetPlatform* TargetPl
 
 bool UTexture::IsAsyncCacheComplete() const
 {
-	const FTexturePlatformData* const* RunningPlatformDataPtr = const_cast<UTexture*>(this)->GetRunningPlatformData();
-	if (RunningPlatformDataPtr)
+	if (const FTexturePlatformData* const* RunningPlatformDataPtr = const_cast<UTexture*>(this)->GetRunningPlatformData())
 	{
-		const FTexturePlatformData* RunningPlatformData = *RunningPlatformDataPtr;
-		if (RunningPlatformData)
-		{
-			if (RunningPlatformData->AsyncTask != nullptr && !RunningPlatformData->AsyncTask->IsWorkDone())
+		if (const FTexturePlatformData* PlatformData = *RunningPlatformDataPtr)
 		{
+			if (PlatformData->AsyncTask && !PlatformData->AsyncTask->Poll())
+			{
 				return false;
 			}
 		}
 	}
 
-	const TMap<FString,FTexturePlatformData*>* CookedPlatformDataPtr = const_cast<UTexture*>(this)->GetCookedPlatformData();
-	if (CookedPlatformDataPtr)
+	if (const TMap<FString, FTexturePlatformData*>* CookedPlatformDataPtr = const_cast<UTexture*>(this)->GetCookedPlatformData())
 	{
 		for (const TTuple<FString, FTexturePlatformData*>& Kvp : *CookedPlatformDataPtr)
 		{
-			const FTexturePlatformData* PlatformData = Kvp.Value;
-			if (PlatformData)
+			if (const FTexturePlatformData* PlatformData = Kvp.Value)
 			{
-				if (PlatformData->AsyncTask != nullptr && !PlatformData->AsyncTask->IsWorkDone())
+				if (PlatformData->AsyncTask && !PlatformData->AsyncTask->Poll())
 				{
 					return false;
 				}
diff --git a/Engine/Source/Runtime/Engine/Private/TextureDerivedDataTask.h b/Engine/Source/Runtime/Engine/Private/TextureDerivedDataTask.h
--- a/Engine/Source/Runtime/Engine/Private/TextureDerivedDataTask.h
+++ b/Engine/Source/Runtime/Engine/Private/TextureDerivedDataTask.h
@@ -182,9 +182,23 @@ public:
 	}
 };
 
-struct FTextureAsyncCacheDerivedDataTask : public FAsyncTask<FTextureCacheDerivedDataWorker>
+struct FTextureAsyncCacheDerivedDataTask
 {
-	FTextureAsyncCacheDerivedDataTask(
+	virtual ~FTextureAsyncCacheDerivedDataTask() = default;
+	virtual void Finalize(bool& bOutFoundInCache, uint64& OutProcessedByteCount) = 0;
+	virtual EQueuedWorkPriority GetPriority() const = 0;
+	virtual bool SetPriority(EQueuedWorkPriority InQueuedWorkPriority) = 0;
+	virtual bool Cancel() = 0;
+	virtual void Wait() = 0;
+	virtual bool WaitWithTimeout(float TimeLimitSeconds) = 0;
+	virtual bool Poll() const = 0;
+};
+
+class FTextureAsyncCacheDerivedDataWorkerTask final : public FTextureAsyncCacheDerivedDataTask, public FAsyncTask<FTextureCacheDerivedDataWorker>
+{
+public:
+	FTextureAsyncCacheDerivedDataWorkerTask(
+		FQueuedThreadPool* InQueuedPool,
 		ITextureCompressorModule* InCompressor,
 		FTexturePlatformData* InDerivedData,
 		UTexture* InTexture,
@@ -198,8 +212,49 @@ struct FTextureAsyncCacheDerivedDataTask : public FAsyncTask<FTextureCacheDerive
 			InSettingsPerLayer,
 			InCacheFlags
 			)
+		, QueuedPool(InQueuedPool)
+	{
+	}
+
+	void Finalize(bool& bOutFoundInCache, uint64& OutProcessedByteCount) final
+	{
+		GetTask().Finalize();
+		bOutFoundInCache = GetTask().WasLoadedFromDDC();
+		OutProcessedByteCount = GetTask().GetBytesCached();
+	}
+
+	EQueuedWorkPriority GetPriority() const final
+	{
+		return FAsyncTask<FTextureCacheDerivedDataWorker>::GetPriority();
+	}
+
+	bool SetPriority(EQueuedWorkPriority InQueuedWorkPriority) final
 	{
+		return FAsyncTask<FTextureCacheDerivedDataWorker>::Reschedule(QueuedPool, InQueuedWorkPriority);
 	}
+
+	bool Cancel() final
+	{
+		return FAsyncTask<FTextureCacheDerivedDataWorker>::IsDone() || FAsyncTask<FTextureCacheDerivedDataWorker>::Cancel();
+	}
+
+	void Wait() final
+	{
+		FAsyncTask<FTextureCacheDerivedDataWorker>::EnsureCompletion();
+	}
+
+	bool WaitWithTimeout(float TimeLimitSeconds) final
+	{
+		return FAsyncTask<FTextureCacheDerivedDataWorker>::WaitCompletionWithTimeout(TimeLimitSeconds);
+	}
+
+	bool Poll() const final
+	{
+		return FAsyncTask<FTextureCacheDerivedDataWorker>::IsWorkDone();
+	}
+
+private:
+	FQueuedThreadPool* QueuedPool;
 };
 
 #endif // WITH_EDITOR
