SceneUtils: Fix build break due to auto casting to double

diff --git a/Engine/Source/Runtime/Engine/Private/SceneUtils.cpp b/Engine/Source/Runtime/Engine/Private/SceneUtils.cpp
--- a/Engine/Source/Runtime/Engine/Private/SceneUtils.cpp
+++ b/Engine/Source/Runtime/Engine/Private/SceneUtils.cpp
@@ -62,7 +62,7 @@ ENGINE_API EAntiAliasingMethod GetDefaultAntiAliasingMethod(const FStaticFeature
 	else
 	{
 		static auto* DefaultAntiAliasingCvar = IConsoleManager::Get().FindTConsoleVariableDataInt(TEXT("r.AntiAliasingMethod"));
-		AntiAliasingMethod = EAntiAliasingMethod(FMath::Clamp(DefaultAntiAliasingCvar->GetValueOnAnyThread(), 0, AAM_MAX));
+		AntiAliasingMethod = EAntiAliasingMethod(FMath::Clamp<int32>(DefaultAntiAliasingCvar->GetValueOnAnyThread(), 0, AAM_MAX));
 	}
 
 	static auto* MSAACountCVar = IConsoleManager::Get().FindTConsoleVariableDataInt(TEXT("r.MSAACount"));
