Optimus: Ensure we can serialize invalid shader value handles (because of CDOs and such).

#rnx

diff --git a/Engine/Source/Runtime/Engine/Private/ComputeFramework/ShaderParamTypeDefinition.cpp b/Engine/Source/Runtime/Engine/Private/ComputeFramework/ShaderParamTypeDefinition.cpp
--- a/Engine/Source/Runtime/Engine/Private/ComputeFramework/ShaderParamTypeDefinition.cpp
+++ b/Engine/Source/Runtime/Engine/Private/ComputeFramework/ShaderParamTypeDefinition.cpp
@@ -393,27 +393,35 @@ FArchive& operator<<(FArchive& InArchive, FShaderValueTypeHandle& InHandle)
 		ValueTypePtr = const_cast<FShaderValueType *>(InHandle.ValueTypePtr);
 	}
 
-	InArchive << ValueTypePtr->Type;
-
-	if (ValueTypePtr->Type == EShaderFundamentalType::Struct)
-	{
-		InArchive << ValueTypePtr->Name;
-		InArchive << ValueTypePtr->StructElements;
-	}
-	else
+	if (ValueTypePtr)
 	{
-		InArchive << ValueTypePtr->DimensionType;
-
-		if (ValueTypePtr->DimensionType == EShaderFundamentalDimensionType::Vector)
+		InArchive << ValueTypePtr->Type;
+		
+		if (ValueTypePtr->Type == EShaderFundamentalType::Struct)
 		{
-			InArchive << ValueTypePtr->VectorElemCount;
+			InArchive << ValueTypePtr->Name;
+			InArchive << ValueTypePtr->StructElements;
 		}
-		else if (ValueTypePtr->DimensionType == EShaderFundamentalDimensionType::Matrix)
+		else
 		{
-			InArchive << ValueTypePtr->MatrixRowCount;
-			InArchive << ValueTypePtr->MatrixColumnCount;
+			InArchive << ValueTypePtr->DimensionType;
+
+			if (ValueTypePtr->DimensionType == EShaderFundamentalDimensionType::Vector)
+			{
+				InArchive << ValueTypePtr->VectorElemCount;
+			}
+			else if (ValueTypePtr->DimensionType == EShaderFundamentalDimensionType::Matrix)
+			{
+				InArchive << ValueTypePtr->MatrixRowCount;
+				InArchive << ValueTypePtr->MatrixColumnCount;
+			}
 		}
 	}
+	else if (InArchive.IsSaving())
+	{
+		EShaderFundamentalType Type = EShaderFundamentalType::None;
+		InArchive << Type; 
+	}
 
 	if (InArchive.IsLoading())
 	{
diff --git a/Engine/Source/Runtime/Engine/Public/ComputeFramework/ShaderParamTypeDefinition.h b/Engine/Source/Runtime/Engine/Public/ComputeFramework/ShaderParamTypeDefinition.h
--- a/Engine/Source/Runtime/Engine/Public/ComputeFramework/ShaderParamTypeDefinition.h
+++ b/Engine/Source/Runtime/Engine/Public/ComputeFramework/ShaderParamTypeDefinition.h
@@ -14,11 +14,12 @@ struct FShaderParamTypeDefinition;
 UENUM()
 enum class EShaderFundamentalType : uint8
 {
-	Bool,
-	Int,
-	Uint,
-	Float,
-	Struct
+	Bool	= 0,
+	Int		= 1,
+	Uint	= 2,
+	Float	= 3,
+	Struct	= 4,
+	None	= 255u,
 };
 
 /*
