Texture: Changed ETextureCacheFlags to an enum class

#rnx

diff --git a/Engine/Source/Runtime/Engine/Private/TextureDerivedData.cpp b/Engine/Source/Runtime/Engine/Private/TextureDerivedData.cpp
--- a/Engine/Source/Runtime/Engine/Private/TextureDerivedData.cpp
+++ b/Engine/Source/Runtime/Engine/Private/TextureDerivedData.cpp
@@ -735,7 +735,7 @@ void FTexturePlatformData::Cache(
 	// Flush any existing async task and ignore results.
 	CancelCache();
 
-	uint32 Flags = InFlags;
+	ETextureCacheFlags Flags = ETextureCacheFlags(InFlags);
 
 	static bool bForDDC = FString(FCommandLine::Get()).Contains(TEXT("Run=DerivedDataCache"));
 	static bool bDDCShippingTextures = FParse::Param(FCommandLine::Get(), TEXT("DDCShippingTextures"));
@@ -744,8 +744,8 @@ void FTexturePlatformData::Cache(
 		Flags |= ETextureCacheFlags::ForDDCBuild;
 	}
 
-	bool bForceRebuild = (Flags & ETextureCacheFlags::ForceRebuild) != 0;
-	bool bAsync = (Flags & ETextureCacheFlags::Async) != 0;
+	bool bForceRebuild = EnumHasAnyFlags(Flags, ETextureCacheFlags::ForceRebuild);
+	bool bAsync = EnumHasAnyFlags(Flags, ETextureCacheFlags::Async);
 	GetTextureDerivedDataKey(InTexture, InSettingsPerLayer, DerivedDataKey);
 
 	if ((!bForDDC || bDDCShippingTextures) && !bForceRebuild && InSettingsPerLayer[0].bHasEditorOnlyData)
@@ -1671,12 +1671,12 @@ void UTexture::CachePlatformData(bool bAsyncCache, bool bAllowAsyncBuild, bool b
 				{
 					PlatformDataLink = new FTexturePlatformData();
 				}
-				int32 CacheFlags = 
+				ETextureCacheFlags CacheFlags = 
 					(bAsyncCache ? ETextureCacheFlags::Async : ETextureCacheFlags::None) |
 					(bAllowAsyncBuild? ETextureCacheFlags::AllowAsyncBuild : ETextureCacheFlags::None) |
 					(bAllowAsyncLoading? ETextureCacheFlags::AllowAsyncLoading : ETextureCacheFlags::None);
 
-				PlatformDataLink->Cache(*this, BuildSettings.GetData(), CacheFlags, Compressor);
+				PlatformDataLink->Cache(*this, BuildSettings.GetData(), uint32(CacheFlags), Compressor);
 			}
 		}
 		else if (PlatformDataLink == NULL)
@@ -1755,7 +1755,7 @@ void UTexture::BeginCacheForCookedPlatformData( const ITargetPlatform *TargetPla
 			PlatformDataToCache->Cache(
 				*this,
 				BuildSettingsToCache[SettingsIndex].GetData(),
-				ETextureCacheFlags::Async | ETextureCacheFlags::InlineMips | ETextureCacheFlags::AllowAsyncBuild | ETextureCacheFlags::AllowAsyncLoading,
+				uint32(ETextureCacheFlags::Async | ETextureCacheFlags::InlineMips | ETextureCacheFlags::AllowAsyncBuild | ETextureCacheFlags::AllowAsyncLoading),
 				nullptr
 				);
 			CookedPlatformData.Add(BuildSettingsCacheKeys[SettingsIndex], PlatformDataToCache);
@@ -1994,7 +1994,7 @@ void UTexture::ForceRebuildPlatformData()
 		PlatformDataLink->Cache(
 			*this,
 			BuildSettings.GetData(),
-			ETextureCacheFlags::ForceRebuild,
+			uint32(ETextureCacheFlags::ForceRebuild),
 			nullptr
 			);
 	}
@@ -2124,7 +2124,7 @@ void UTexture::SerializeCookedPlatformData(FArchive& Ar)
 					if (PlatformDataPtr == NULL)
 					{
 						PlatformDataPtr = new FTexturePlatformData();
-						PlatformDataPtr->Cache(*this, BuildSettingsToCache[SettingIndex].GetData(), ETextureCacheFlags::InlineMips | ETextureCacheFlags::Async, nullptr);
+						PlatformDataPtr->Cache(*this, BuildSettingsToCache[SettingIndex].GetData(), uint32(ETextureCacheFlags::InlineMips | ETextureCacheFlags::Async), nullptr);
 
 						CookedPlatformDataPtr->Add(DerivedDataKey, PlatformDataPtr);
 
diff --git a/Engine/Source/Runtime/Engine/Private/TextureDerivedDataTask.cpp b/Engine/Source/Runtime/Engine/Private/TextureDerivedDataTask.cpp
--- a/Engine/Source/Runtime/Engine/Private/TextureDerivedDataTask.cpp
+++ b/Engine/Source/Runtime/Engine/Private/TextureDerivedDataTask.cpp
@@ -278,7 +278,7 @@ void FTextureCacheDerivedDataWorker::ConsumeBuildFunctionOutput(const UE::Derive
 
 	if (DerivedData->Mips.Num())
 	{
-		const bool bInlineMips = (CacheFlags & ETextureCacheFlags::InlineMips) != 0;
+		const bool bInlineMips = EnumHasAnyFlags(CacheFlags, ETextureCacheFlags::InlineMips);
 
 		for (int32 MipIndex = 0; MipIndex < DerivedData->Mips.Num(); ++MipIndex)
 		{
@@ -334,7 +334,7 @@ void FTextureCacheDerivedDataWorker::BuildTexture(bool bReplaceExistingDDC)
 	TRACE_CPUPROFILER_EVENT_SCOPE(FTextureCacheDerivedDataWorker::BuildTexture);
 
 	const bool bHasValidMip0 = TextureData.Blocks.Num() && TextureData.Blocks[0].MipsPerLayer.Num() && TextureData.Blocks[0].MipsPerLayer[0].Num();
-	const bool bForVirtualTextureStreamingBuild = (CacheFlags & ETextureCacheFlags::ForVirtualTextureStreamingBuild) != 0;	
+	const bool bForVirtualTextureStreamingBuild = EnumHasAnyFlags(CacheFlags, ETextureCacheFlags::ForVirtualTextureStreamingBuild);
 
 	FFormatNamedArguments Args;
 	Args.Add(TEXT("TextureName"), FText::FromString(Texture.GetName()));
@@ -394,7 +394,7 @@ void FTextureCacheDerivedDataWorker::BuildTexture(bool bReplaceExistingDDC)
 
 			if (DerivedData->VTData->Chunks.Num())
 			{
-				const bool bInlineMips = (CacheFlags & ETextureCacheFlags::InlineMips) != 0;
+				const bool bInlineMips = EnumHasAnyFlags(CacheFlags, ETextureCacheFlags::InlineMips);
 				bSucceeded = !bInlineMips || DerivedData->TryInlineMipData(BuildSettingsPerLayer[0].LODBiasWithCinematicMips, &Texture);
 				if (!bSucceeded)
 				{
@@ -622,7 +622,7 @@ void FTextureCacheDerivedDataWorker::BuildTexture(bool bReplaceExistingDDC)
 
 			if (DerivedData->Mips.Num())
 			{
-				const bool bInlineMips = (CacheFlags & ETextureCacheFlags::InlineMips) != 0;
+				const bool bInlineMips = EnumHasAnyFlags(CacheFlags, ETextureCacheFlags::InlineMips);
 				bSucceeded = !bInlineMips || DerivedData->TryInlineMipData(BuildSettingsPerLayer[0].LODBiasWithCinematicMips, &Texture);
 				if (!bSucceeded)
 				{
@@ -642,7 +642,7 @@ FTextureCacheDerivedDataWorker::FTextureCacheDerivedDataWorker(
 	FTexturePlatformData* InDerivedData,
 	UTexture* InTexture,
 	const FTextureBuildSettings* InSettingsPerLayer,
-	uint32 InCacheFlags
+	ETextureCacheFlags InCacheFlags
 	)
 	: Compressor(InCompressor)
 	, ImageWrapper(nullptr)
@@ -678,9 +678,9 @@ FTextureCacheDerivedDataWorker::FTextureCacheDerivedDataWorker(
 	UTexture::GetPixelFormatEnum();
 	GetTextureDerivedDataKeySuffix(Texture, InSettingsPerLayer, KeySuffix);
 		
-	const bool bAllowAsyncBuild = (CacheFlags & ETextureCacheFlags::AllowAsyncBuild) != 0;
-	const bool bAllowAsyncLoading = (CacheFlags & ETextureCacheFlags::AllowAsyncLoading) != 0;
-	const bool bForVirtualTextureStreamingBuild = (CacheFlags & ETextureCacheFlags::ForVirtualTextureStreamingBuild) != 0;
+	const bool bAllowAsyncBuild = EnumHasAnyFlags(CacheFlags, ETextureCacheFlags::AllowAsyncBuild);
+	const bool bAllowAsyncLoading = EnumHasAnyFlags(CacheFlags, ETextureCacheFlags::AllowAsyncLoading);
+	const bool bForVirtualTextureStreamingBuild = EnumHasAnyFlags(CacheFlags, ETextureCacheFlags::ForVirtualTextureStreamingBuild);
 
 	static const bool bBuildFunctionEnabled = FParse::Param(FCommandLine::Get(), TEXT("DDC2TextureBuilds"));
 	if (bBuildFunctionEnabled && !bForVirtualTextureStreamingBuild && (BuildSettingsPerLayer.Num() == 1))
@@ -740,10 +740,10 @@ void FTextureCacheDerivedDataWorker::DoWork()
 {
 	TRACE_CPUPROFILER_EVENT_SCOPE(FTextureCacheDerivedDataWorker::DoWork);
 
-	const bool bForceRebuild = (CacheFlags & ETextureCacheFlags::ForceRebuild) != 0;
-	const bool bAllowAsyncBuild = (CacheFlags & ETextureCacheFlags::AllowAsyncBuild) != 0;
-	const bool bAllowAsyncLoading = (CacheFlags & ETextureCacheFlags::AllowAsyncLoading) != 0;
-	const bool bForVirtualTextureStreamingBuild = (CacheFlags & ETextureCacheFlags::ForVirtualTextureStreamingBuild) != 0;
+	const bool bForceRebuild = EnumHasAnyFlags(CacheFlags, ETextureCacheFlags::ForceRebuild);
+	const bool bAllowAsyncBuild = EnumHasAnyFlags(CacheFlags, ETextureCacheFlags::AllowAsyncBuild);
+	const bool bAllowAsyncLoading = EnumHasAnyFlags(CacheFlags, ETextureCacheFlags::AllowAsyncLoading);
+	const bool bForVirtualTextureStreamingBuild = EnumHasAnyFlags(CacheFlags, ETextureCacheFlags::ForVirtualTextureStreamingBuild);
 	bool bInvalidVirtualTextureCompression = false;
 
 	TArray<uint8> RawDerivedData;
@@ -765,8 +765,8 @@ void FTextureCacheDerivedDataWorker::DoWork()
 
 	if (bLoadedFromDDC)
 	{
-		const bool bInlineMips = (CacheFlags & ETextureCacheFlags::InlineMips) != 0;
-		const bool bForDDC = (CacheFlags & ETextureCacheFlags::ForDDCBuild) != 0;
+		const bool bInlineMips = EnumHasAnyFlags(CacheFlags, ETextureCacheFlags::InlineMips);
+		const bool bForDDC = EnumHasAnyFlags(CacheFlags, ETextureCacheFlags::ForDDCBuild);
 
 		BytesCached = RawDerivedData.Num();
 		FMemoryReader Ar(RawDerivedData, /*bIsPersistent=*/ true);
diff --git a/Engine/Source/Runtime/Engine/Private/TextureDerivedDataTask.h b/Engine/Source/Runtime/Engine/Private/TextureDerivedDataTask.h
--- a/Engine/Source/Runtime/Engine/Private/TextureDerivedDataTask.h
+++ b/Engine/Source/Runtime/Engine/Private/TextureDerivedDataTask.h
@@ -14,6 +14,7 @@
 #include "Engine/Texture2D.h"
 #include "IImageWrapperModule.h"
 #include "ImageCore.h"
+#include "Misc/EnumClassFlags.h"
 #include "TextureCompressorModule.h"
 
 #endif // WITH_EDITOR
@@ -31,22 +32,21 @@ namespace UE::DerivedData { class FBuildOutput; }
 void GetTextureDerivedDataKeySuffix(const UTexture& Texture, const FTextureBuildSettings* BuildSettingsPerLayer, FString& OutKeySuffix);
 uint32 PutDerivedDataInCache(FTexturePlatformData* DerivedData, const FString& DerivedDataKeySuffix, const FStringView& TextureName, bool bForceAllMipsToBeInlined, bool bReplaceExistingDDC);
 
-namespace ETextureCacheFlags
+enum class ETextureCacheFlags : uint32
 {
-	enum Type
-	{
-		None			= 0x00,
-		Async			= 0x01,
-		ForceRebuild	= 0x02,
-		InlineMips		= 0x08,
-		AllowAsyncBuild	= 0x10,
-		ForDDCBuild		= 0x20,
-		RemoveSourceMipDataAfterCache = 0x40,
-		AllowAsyncLoading = 0x80,
-		ForVirtualTextureStreamingBuild = 0x100
-	};
+	None			= 0x00,
+	Async			= 0x01,
+	ForceRebuild	= 0x02,
+	InlineMips		= 0x08,
+	AllowAsyncBuild	= 0x10,
+	ForDDCBuild		= 0x20,
+	RemoveSourceMipDataAfterCache = 0x40,
+	AllowAsyncLoading = 0x80,
+	ForVirtualTextureStreamingBuild = 0x100,
 };
 
+ENUM_CLASS_FLAGS(ETextureCacheFlags);
+
 // Everything required to get the texture source data.
 struct FTextureSourceLayerData
 {
@@ -126,7 +126,7 @@ class FTextureCacheDerivedDataWorker : public FNonAbandonableTask
 	/** DDC2 build function name to use to build this texture (if DDC2 is enabled and the target texture type has a DDC2 build function, empty otherwise) */
 	FString BuildFunctionName;
 	/** Texture cache flags. */
-	uint32 CacheFlags;
+	ETextureCacheFlags CacheFlags;
 	/** Have many bytes were loaded from DDC or built (for telemetry) */
 	uint32 BytesCached = 0;
 	/** Estimate of the peak amount of memory required to complete this task. */
@@ -150,7 +150,7 @@ public:
 		FTexturePlatformData* InDerivedData,
 		UTexture* InTexture,
 		const FTextureBuildSettings* InSettingsPerLayer,
-		uint32 InCacheFlags);
+		ETextureCacheFlags InCacheFlags);
 
 	/** Does the work to cache derived data. Safe to call from any thread. */
 	void DoWork();
@@ -203,7 +203,7 @@ public:
 		FTexturePlatformData* InDerivedData,
 		UTexture* InTexture,
 		const FTextureBuildSettings* InSettingsPerLayer,
-		uint32 InCacheFlags
+		ETextureCacheFlags InCacheFlags
 		)
 		: FAsyncTask<FTextureCacheDerivedDataWorker>(
 			InCompressor,
