Work around texture streaming check by simply soft failing and avoid streaming the texture in.

- This will be revisted in the future.

diff --git a/Engine/Source/Runtime/Engine/Private/Streaming/TextureStreamingBuild.cpp b/Engine/Source/Runtime/Engine/Private/Streaming/TextureStreamingBuild.cpp
--- a/Engine/Source/Runtime/Engine/Private/Streaming/TextureStreamingBuild.cpp
+++ b/Engine/Source/Runtime/Engine/Private/Streaming/TextureStreamingBuild.cpp
@@ -570,12 +570,14 @@ void FStreamingTextureLevelContext::ProcessMaterial(const FBoxSphereBounds& Comp
 		{
 			for (const FStreamingTextureBuildInfo& BuildInfo : *ComponentBuildData)
 			{
-				check(LevelStreamingTextures.IsValidIndex(BuildInfo.TextureLevelIndex));
-				if (UTexture* Texture = LevelStreamingTextures[BuildInfo.TextureLevelIndex])
+				if (LevelStreamingTextures.IsValidIndex(BuildInfo.TextureLevelIndex))
 				{
-					check(Texture->IsStreamable());
-					check(Texture->LevelIndex != INDEX_NONE);
-					Textures.Add(Texture);
+					if (UTexture* Texture = LevelStreamingTextures[BuildInfo.TextureLevelIndex])
+					{
+						check(Texture->IsStreamable());
+						check(Texture->LevelIndex != INDEX_NONE);
+						Textures.Add(Texture);
+					}
 				}
 			}
 		}
