Texture: Only attempt the new build path for texture formats that support it

#rnx

diff --git a/Engine/Source/Runtime/Engine/Private/TextureDerivedDataBuildUtils.cpp b/Engine/Source/Runtime/Engine/Private/TextureDerivedDataBuildUtils.cpp
--- a/Engine/Source/Runtime/Engine/Private/TextureDerivedDataBuildUtils.cpp
+++ b/Engine/Source/Runtime/Engine/Private/TextureDerivedDataBuildUtils.cpp
@@ -3,6 +3,8 @@
 #include "TextureDerivedDataBuildUtils.h"
 
 #if WITH_EDITOR
+#include "DerivedDataBuild.h"
+#include "DerivedDataBuildFunctionRegistry.h"
 #include "Engine/Texture.h"
 #include "Interfaces/ITextureFormat.h"
 #include "Interfaces/ITextureFormatManagerModule.h"
@@ -205,7 +207,7 @@ static void WriteSource(FCbWriter& Writer, const UTexture& Texture, int32 LayerI
 	Writer.EndObject();
 }
 
-FString GetTextureBuildFunctionName(const FTextureBuildSettings& BuildSettings)
+bool TryFindTextureBuildFunction(FStringBuilderBase& OutFunctionName, const FTextureBuildSettings& BuildSettings)
 {
 	FName TextureFormatModuleName;
 
@@ -214,19 +216,22 @@ FString GetTextureBuildFunctionName(const FTextureBuildSettings& BuildSettings)
 		ITextureFormatModule* TextureFormatModule = nullptr;
 		if (!TFM->FindTextureFormatAndModule(BuildSettings.TextureFormatName, TextureFormatModuleName, TextureFormatModule))
 		{
-			return FString();
+			return false;
 		}
 	}
 
+	const int32 NameIndex = OutFunctionName.Len();
+
 	// Texture format modules are inconsistent in their naming, e.g., TextureFormatUncompressed, <Platform>TextureFormat.
 	// Attempt to unify the naming of build functions as <Format>Texture.
-	TStringBuilder<64> FunctionName;
-	FunctionName << TextureFormatModuleName << TEXT("Texture"_SV);
-	if (int32 Index = UE::String::FindFirst(FunctionName, TEXT("TextureFormat"_SV)); Index != INDEX_NONE)
+	OutFunctionName << TextureFormatModuleName << TEXT("Texture"_SV);
+	if (int32 Index = UE::String::FindFirst(OutFunctionName, TEXT("TextureFormat"_SV)); Index != INDEX_NONE)
 	{
-		FunctionName.RemoveAt(Index, TEXT("TextureFormat"_SV).Len());
+		OutFunctionName.RemoveAt(Index, TEXT("TextureFormat"_SV).Len());
 	}
-	return FString(FunctionName);
+
+	const FStringView FunctionName = OutFunctionName.ToView().RightChop(NameIndex);
+	return GetDerivedDataBuildRef().GetFunctionRegistry().FindFunctionVersion(FunctionName).IsValid();
 }
 
 FCbObject SaveTextureBuildSettings(const FString& KeySuffix, const UTexture& Texture, const FTextureBuildSettings& BuildSettings, int32 LayerIndex, int32 NumInlineMips)
diff --git a/Engine/Source/Runtime/Engine/Private/TextureDerivedDataBuildUtils.h b/Engine/Source/Runtime/Engine/Private/TextureDerivedDataBuildUtils.h
--- a/Engine/Source/Runtime/Engine/Private/TextureDerivedDataBuildUtils.h
+++ b/Engine/Source/Runtime/Engine/Private/TextureDerivedDataBuildUtils.h
@@ -5,13 +5,14 @@
 #if WITH_EDITOR
 
 #include "CoreTypes.h"
+#include "Containers/StringFwd.h"
 
 class FCbObject;
 class FString;
 class UTexture;
 struct FTextureBuildSettings;
 
-FString GetTextureBuildFunctionName(const FTextureBuildSettings& BuildSettings);
+bool TryFindTextureBuildFunction(FStringBuilderBase& OutFunctionName, const FTextureBuildSettings& BuildSettings);
 FCbObject SaveTextureBuildSettings(const FString& KeySuffix, const UTexture& Texture, const FTextureBuildSettings& BuildSettings, int32 LayerIndex, int32 NumInlineMips);
 
 #endif // WITH_EDITOR
diff --git a/Engine/Source/Runtime/Engine/Private/TextureDerivedDataTask.cpp b/Engine/Source/Runtime/Engine/Private/TextureDerivedDataTask.cpp
--- a/Engine/Source/Runtime/Engine/Private/TextureDerivedDataTask.cpp
+++ b/Engine/Source/Runtime/Engine/Private/TextureDerivedDataTask.cpp
@@ -685,7 +685,11 @@ FTextureCacheDerivedDataWorker::FTextureCacheDerivedDataWorker(
 	static const bool bBuildFunctionEnabled = FParse::Param(FCommandLine::Get(), TEXT("DDC2TextureBuilds"));
 	if (bBuildFunctionEnabled && !bForVirtualTextureStreamingBuild && (BuildSettingsPerLayer.Num() == 1))
 	{
-		BuildFunctionName = GetTextureBuildFunctionName(BuildSettingsPerLayer[0]);
+		TStringBuilder<64> FunctionName;
+		if (TryFindTextureBuildFunction(FunctionName, BuildSettingsPerLayer[0]))
+		{
+			BuildFunctionName = FunctionName;
+		}
 	}
 
 	// FVirtualTextureDataBuilder always wants to load ImageWrapper module
