Added PIE.StartupTime analytics event
Added DDC Core.Loading analytics

diff --git a/Engine/Source/Runtime/Engine/Private/Analytics/StudioAnalytics.cpp b/Engine/Source/Runtime/Engine/Private/Analytics/StudioAnalytics.cpp
--- a/Engine/Source/Runtime/Engine/Private/Analytics/StudioAnalytics.cpp
+++ b/Engine/Source/Runtime/Engine/Private/Analytics/StudioAnalytics.cpp
@@ -18,6 +18,7 @@
 #include "Interfaces/IAnalyticsProvider.h"
 #include "Templates/SharedPointer.h"
 #include "HAL/PlatformProcess.h"
+#include "ProfilingDebugging/CookStats.h"
 
 bool FStudioAnalytics::bInitialized = false;
 volatile double FStudioAnalytics::TimeEstimation = 0;
@@ -156,10 +157,30 @@ void FStudioAnalytics::FireEvent_Loading(const FString& LoadingName, double Seco
 	if (FStudioAnalytics::IsAvailable())
 	{
 		TArray<FAnalyticsEventAttribute> Attributes;
+
+		
 		Attributes.Emplace(TEXT("LoadingName"), LoadingName);
 		Attributes.Emplace(TEXT("LoadingSeconds"), SecondsSpentLoading);
 		Attributes.Append(InAttributes);
 
 		FStudioAnalytics::GetProvider().RecordEvent(TEXT("Performance.Loading"), Attributes);
+
+#if ENABLE_COOK_STATS
+		// Sends each cook stat to the studio analytics system.
+		auto SendCookStatsToAnalytics = [&Attributes](const FString& StatName, const TArray<FCookStatsManager::StringKeyValue>& StatAttributes)
+		{
+			for (const auto& Attr : StatAttributes)
+			{
+				FString FormattedAttrName = StatName + "." + Attr.Key;
+
+				Attributes.Emplace(FormattedAttrName, Attr.Value);
+			}
+		};
+
+		// Grab the DDC stats
+		FCookStatsManager::LogCookStats(SendCookStatsToAnalytics);
+
+		FStudioAnalytics::GetProvider().RecordEvent(TEXT("Core.Loading"), Attributes);
+#endif
 	}
 }
\ No newline at end of file
