Fix stable sorting between static mesh sections. MeshIdInPrimitive is used in the sort key for translucent objects for internal sorting and was not properly set.

diff --git a/Engine/Source/Runtime/Engine/Private/StaticMeshRender.cpp b/Engine/Source/Runtime/Engine/Private/StaticMeshRender.cpp
--- a/Engine/Source/Runtime/Engine/Private/StaticMeshRender.cpp
+++ b/Engine/Source/Runtime/Engine/Private/StaticMeshRender.cpp
@@ -647,6 +647,7 @@ bool FStaticMeshSceneProxy::GetMeshElement(
 	if (NumPrimitives > 0)
 	{
 		OutMeshBatch.SegmentIndex = SectionIndex;
+		OutMeshBatch.MeshIdInPrimitive = SectionIndex;
 
 		OutMeshBatch.LODIndex = LODIndex;
 #if STATICMESH_ENABLE_DEBUG_RENDERING
@@ -1749,6 +1750,7 @@ void FStaticMeshSceneProxy::GetDynamicRayTracingInstances(FRayTracingMaterialGat
 				}
 
 				Mesh.SegmentIndex = SectionIndex;
+				Mesh.MeshIdInPrimitive = SectionIndex;
 			}
 		}
 
