Fix 2 crashes converting skeletal mesh to static mesh when 'Render Static' is ticked in editor: 1) fail to convert to FSkeletalMeshObjectCPUSkin, 2) access vertex color buffer when it is empty

diff --git a/Engine/Source/Developer/MeshUtilities/Private/MeshUtilities.cpp b/Engine/Source/Developer/MeshUtilities/Private/MeshUtilities.cpp
--- a/Engine/Source/Developer/MeshUtilities/Private/MeshUtilities.cpp
+++ b/Engine/Source/Developer/MeshUtilities/Private/MeshUtilities.cpp
@@ -370,7 +370,7 @@ static void SkinnedMeshToRawMeshes(USkinnedMeshComponent* InSkinnedMeshComponent
 						}
 					}
 
-					if (LODData.StaticVertexBuffers.ColorVertexBuffer.IsInitialized())
+					if (LODData.StaticVertexBuffers.ColorVertexBuffer.IsInitialized() && LODData.StaticVertexBuffers.ColorVertexBuffer.GetNumVertices() > 0)
 					{
 						RawMesh.WedgeColors.Add(LODData.StaticVertexBuffers.ColorVertexBuffer.VertexColor(VertexIndexForWedge));
 						RawMeshTracker.bValidColors = true;
diff --git a/Engine/Source/Runtime/Engine/Private/Components/SkinnedMeshComponent.cpp b/Engine/Source/Runtime/Engine/Private/Components/SkinnedMeshComponent.cpp
--- a/Engine/Source/Runtime/Engine/Private/Components/SkinnedMeshComponent.cpp
+++ b/Engine/Source/Runtime/Engine/Private/Components/SkinnedMeshComponent.cpp
@@ -3276,6 +3276,10 @@ void USkinnedMeshComponent::GetCPUSkinnedVertices(TArray<FFinalSkinVertex>& OutV
 	PoseComponent->UpdateLODStatus();
 	PoseComponent->RefreshBoneTransforms(nullptr);
 	
+	// Turn bRenderStatic off so MeshObject can be switched to FSkeletalMeshObjectCPUSkin
+	const bool bCachedRenderStatic = bRenderStatic;
+	MutableThis->bRenderStatic = false;
+
 	// switch to CPU skinning
 	const bool bCachedCPUSkinning = GetCPUSkinningEnabled();
 	constexpr bool bRecreateRenderStateImmediately = true;
@@ -3287,8 +3291,9 @@ void USkinnedMeshComponent::GetCPUSkinnedVertices(TArray<FFinalSkinVertex>& OutV
 	// Copy our vertices out. We know we are using CPU skinning now, so this cast is safe
 	OutVertices = static_cast<FSkeletalMeshObjectCPUSkin*>(MeshObject)->GetCachedFinalVertices();
 	
-	// switch skinning mode, LOD etc. back
+	// switch skinning mode, LOD, bRenderStatic, etc. back
 	PoseComponent->SetForcedLOD(CachedForcedLOD);
+	MutableThis->bRenderStatic = bCachedRenderStatic;
 	MutableThis->SetCPUSkinningEnabled(bCachedCPUSkinning, bRecreateRenderStateImmediately);
 }
 
