Reduced the stats recorded in the Loading Event to just the DDC.Summary

diff --git a/Engine/Source/Runtime/Engine/Private/Analytics/StudioAnalytics.cpp b/Engine/Source/Runtime/Engine/Private/Analytics/StudioAnalytics.cpp
--- a/Engine/Source/Runtime/Engine/Private/Analytics/StudioAnalytics.cpp
+++ b/Engine/Source/Runtime/Engine/Private/Analytics/StudioAnalytics.cpp
@@ -166,14 +166,17 @@ void FStudioAnalytics::FireEvent_Loading(const FString& LoadingName, double Seco
 		FStudioAnalytics::GetProvider().RecordEvent(TEXT("Performance.Loading"), Attributes);
 
 #if ENABLE_COOK_STATS
-		// Sends each cook stat to the studio analytics system.
+		// Sends each DDC stat to the studio analytics system.
 		auto SendCookStatsToAnalytics = [&Attributes](const FString& StatName, const TArray<FCookStatsManager::StringKeyValue>& StatAttributes)
 		{
-			for (const auto& Attr : StatAttributes)
+			// We'reonly interested in DDC Summary stats
+			if (StatName.Contains("DDC.Summary"))
 			{
-				FString FormattedAttrName = StatName + "." + Attr.Key;
-
-				Attributes.Emplace(FormattedAttrName, Attr.Value);
+				for (const auto& Attr : StatAttributes)
+				{
+					FString FormattedAttrName = StatName + "." + Attr.Key;
+					Attributes.Emplace(FormattedAttrName, Attr.Value);
+				}		
 			}
 		};
 
