Fix dbuffer decal when reading world normal and writing to opacity.

diff --git a/Engine/Source/Runtime/Engine/Private/Materials/HLSLMaterialTranslator.cpp b/Engine/Source/Runtime/Engine/Private/Materials/HLSLMaterialTranslator.cpp
--- a/Engine/Source/Runtime/Engine/Private/Materials/HLSLMaterialTranslator.cpp
+++ b/Engine/Source/Runtime/Engine/Private/Materials/HLSLMaterialTranslator.cpp
@@ -6062,7 +6062,8 @@ int32 FHLSLMaterialTranslator::SceneTextureLookup(int32 ViewportUV, uint32 InSce
 				MaterialProperty == MP_Normal || 
 				MaterialProperty == MP_Roughness || 
 				MaterialProperty == MP_Specular || 
-				MaterialProperty == MP_Metallic;
+				MaterialProperty == MP_Metallic ||
+				MaterialProperty == MP_Opacity;
 
 			if (bUsingDBuffer && bIsDBufferOutput)
 			{
@@ -6170,17 +6171,17 @@ void FHLSLMaterialTranslator::UseSceneTextureId(ESceneTextureId SceneTextureId,
 	MaterialCompilationOutput.SetIsSceneTextureUsed(SceneTextureId);
 
 	if (Material->GetMaterialDomain() == MD_DeferredDecal)
-			{
+	{
 		const bool bSceneTextureSupportsDecal = SceneTextureId == PPI_SceneDepth || SceneTextureId == PPI_WorldNormal || SceneTextureId == PPI_CustomDepth || SceneTextureId == PPI_CustomStencil;
 		if (!bSceneTextureSupportsDecal)
 		{
 			// Note: For DBuffer decals CustomDepth and CustomStencil are not available if r.CustomDepth.Order = 1
 			Errorf(TEXT("Decals can only access SceneDepth, CustomDepth, CustomStencil, and WorldNormal."));
-			}
+		}
 
 		const bool bSceneTextureRequiresSM5 = SceneTextureId == PPI_WorldNormal || SceneTextureId == PPI_CustomDepth || SceneTextureId == PPI_CustomStencil;
 		if (bSceneTextureRequiresSM5)
-			{
+		{
 			ErrorUnlessFeatureLevelSupported(ERHIFeatureLevel::SM5);
 		}
 
