Cook and editor memory usage optimization for target platform static meshes and skeletal meshes
  - Avoid cooking each platform twice and keeping them both in memory due to DDC key mismatch caused by SectionInfoMap postload modifications
  - Get rid of cached platform from memory as soon as they are built and available from local DDC
  - Get rid of cached platform data during cook to reduce memory footprint

  - 133GB -> 82GB when cooking PS5 for special project. This 50GB gain is expected to be per active platform target.

diff --git a/Engine/Source/Runtime/Engine/Classes/Engine/SkeletalMesh.h b/Engine/Source/Runtime/Engine/Classes/Engine/SkeletalMesh.h
--- a/Engine/Source/Runtime/Engine/Classes/Engine/SkeletalMesh.h
+++ b/Engine/Source/Runtime/Engine/Classes/Engine/SkeletalMesh.h
@@ -820,6 +820,7 @@ public:
     /** Warn if the platform supports the minimal number of per vertex bone weights */
 	void ValidateBoneWeights(const ITargetPlatform* TargetPlatform);
 	virtual void BeginCacheForCookedPlatformData(const ITargetPlatform* TargetPlatform) override;
+	virtual void ClearAllCachedCookedPlatformData() override;
 #endif
 
 	/** Get the data to use for rendering. */
diff --git a/Engine/Source/Runtime/Engine/Classes/Engine/StaticMesh.h b/Engine/Source/Runtime/Engine/Classes/Engine/StaticMesh.h
--- a/Engine/Source/Runtime/Engine/Classes/Engine/StaticMesh.h
+++ b/Engine/Source/Runtime/Engine/Classes/Engine/StaticMesh.h
@@ -1464,6 +1464,10 @@ public:
 	ENGINE_API virtual void PostEditChangeProperty(FPropertyChangedEvent& PropertyChangedEvent) override;
 	ENGINE_API virtual void PostEditUndo() override;
 	ENGINE_API virtual void GetAssetRegistryTagMetadata(TMap<FName, FAssetRegistryTagMetadata>& OutMetadata) const override;
+	
+	ENGINE_API virtual void WillNeverCacheCookedPlatformDataAgain() override;
+	ENGINE_API virtual void ClearCachedCookedPlatformData(const ITargetPlatform* TargetPlatform) override;
+	ENGINE_API virtual void ClearAllCachedCookedPlatformData() override;
 	ENGINE_API virtual void BeginCacheForCookedPlatformData(const ITargetPlatform* TargetPlatform) override;
 	ENGINE_API virtual bool IsCachedCookedPlatformDataLoaded(const ITargetPlatform* TargetPlatform) override;
 	ENGINE_API void SetLODGroup(FName NewGroup, bool bRebuildImmediately = true);
@@ -1879,6 +1883,11 @@ public:
 	 */
 	ENGINE_API void CacheDerivedData();
 
+	/**
+	 * Caches derived renderable for cooked platforms currently active.
+	 */
+	ENGINE_API void PrepareDerivedDataForActiveTargetPlatforms();
+
 private:
 	// Filled at CommitDescription time and reused during build
 	TOptional<FBoxSphereBounds> CachedMeshDescriptionBounds;
diff --git a/Engine/Source/Runtime/Engine/Private/SkeletalMesh.cpp b/Engine/Source/Runtime/Engine/Private/SkeletalMesh.cpp
--- a/Engine/Source/Runtime/Engine/Private/SkeletalMesh.cpp
+++ b/Engine/Source/Runtime/Engine/Private/SkeletalMesh.cpp
@@ -4369,6 +4369,11 @@ void USkeletalMesh::BeginCacheForCookedPlatformData(const ITargetPlatform* Targe
 	ValidateBoneWeights(TargetPlatform);
 }
 
+void USkeletalMesh::ClearAllCachedCookedPlatformData()
+{
+	GetResourceForRendering()->NextCachedRenderData.Reset();
+}
+
 extern FString BuildSkeletalMeshDerivedDataKey(const ITargetPlatform* TargetPlatform, USkeletalMesh* SkelMesh);
 FString USkeletalMesh::GetDerivedDataKey()
 {
diff --git a/Engine/Source/Runtime/Engine/Private/StaticMesh.cpp b/Engine/Source/Runtime/Engine/Private/StaticMesh.cpp
--- a/Engine/Source/Runtime/Engine/Private/StaticMesh.cpp
+++ b/Engine/Source/Runtime/Engine/Private/StaticMesh.cpp
@@ -4061,17 +4061,6 @@ void UStaticMesh::GetAssetRegistryTagMetadata(TMap<FName, FAssetRegistryTagMetad
 }
 #endif
 
-#if WITH_EDITOR
-void UStaticMesh::BeginCacheForCookedPlatformData(const ITargetPlatform* TargetPlatform)
-{
-}
-
-bool UStaticMesh::IsCachedCookedPlatformDataLoaded(const ITargetPlatform* TargetPlatform)
-{
-	return !IsCompiling();
-}
-#endif
-
 
 /*------------------------------------------------------------------------------
 	FStaticMeshSourceModel
@@ -4240,6 +4229,33 @@ static FStaticMeshRenderData& GetPlatformStaticMeshRenderData(UStaticMesh* Mesh,
 	return *PlatformRenderData;
 }
 
+void UStaticMesh::WillNeverCacheCookedPlatformDataAgain()
+{
+}
+
+void UStaticMesh::ClearCachedCookedPlatformData(const ITargetPlatform* TargetPlatform)
+{
+}
+
+void UStaticMesh::ClearAllCachedCookedPlatformData()
+{
+	GetRenderData()->NextCachedRenderData.Reset();
+}
+
+void UStaticMesh::BeginCacheForCookedPlatformData(const ITargetPlatform* TargetPlatform)
+{
+}
+
+bool UStaticMesh::IsCachedCookedPlatformDataLoaded(const ITargetPlatform* TargetPlatform)
+{
+	if (IsCompiling())
+	{
+		return false;
+	}
+	
+	GetPlatformStaticMeshRenderData(this, TargetPlatform);
+	return true;
+}
 
 #if WITH_EDITORONLY_DATA
 
@@ -4669,8 +4685,12 @@ void UStaticMesh::CacheDerivedData()
 
 	SetRenderData(MakeUnique<FStaticMeshRenderData>());
 	GetRenderData()->Cache(RunningPlatform, this, LODSettings);
+}
 
-	// Additionally cache derived data for any other platforms we care about.
+void UStaticMesh::PrepareDerivedDataForActiveTargetPlatforms()
+{
+	ITargetPlatformManagerModule& TargetPlatformManager = GetTargetPlatformManagerRef();
+	ITargetPlatform* RunningPlatform = TargetPlatformManager.GetRunningTargetPlatform();
 	const TArray<ITargetPlatform*>& TargetPlatforms = TargetPlatformManager.GetActiveTargetPlatforms();
 	for (int32 PlatformIndex = 0; PlatformIndex < TargetPlatforms.Num(); ++PlatformIndex)
 	{
@@ -4680,6 +4700,10 @@ void UStaticMesh::CacheDerivedData()
 			GetPlatformStaticMeshRenderData(this, Platform);
 		}
 	}
+
+	// Now that they are in local DDC cache, clear them to save memory
+	// next time they are read it will be fast anyway.
+	ClearAllCachedCookedPlatformData();
 }
 
 #endif // #if WITH_EDITORONLY_DATA
@@ -5623,6 +5647,15 @@ void UStaticMesh::ExecutePostLoadInternal(FStaticMeshPostLoadContext& Context)
 			}
 		}
 	}
+
+	// Additionally cache derived data for any other platforms we care about.
+	// This must be done after the sectioninfomap fixups to make sure the DDC key matches
+	// the one the cooker will generate during save.
+	if (!Context.bIsCookedForEditor)
+	{
+		PrepareDerivedDataForActiveTargetPlatforms();
+	}
+
 	ReleaseAsyncProperty(EStaticMeshAsyncProperties::SectionInfoMap);
 	ReleaseAsyncProperty(EStaticMeshAsyncProperties::OriginalSectionInfoMap);
 
diff --git a/Engine/Source/Runtime/Engine/Private/StaticMeshBuild.cpp b/Engine/Source/Runtime/Engine/Private/StaticMeshBuild.cpp
--- a/Engine/Source/Runtime/Engine/Private/StaticMeshBuild.cpp
+++ b/Engine/Source/Runtime/Engine/Private/StaticMeshBuild.cpp
@@ -444,6 +444,7 @@ bool UStaticMesh::ExecuteBuildInternal(bool bInSilent, TArray<FText>* OutErrors)
 
 	// Free existing render data and recache.
 	CacheDerivedData();
+	PrepareDerivedDataForActiveTargetPlatforms();
 
 	// InitResources will send commands to other threads that will
 	// use our RenderData, we must mark it as ready to be used since
