Properly detect invalid actor files, which can happen when force deleting a Blueprint class.
- WorldPartitionResaveActorsBuilder: detect and delete invalid actor files.
- DerivedDataCacheCommandlet: don't crash when encountering an invalid actor file.

diff --git a/Engine/Source/Editor/UnrealEd/Private/Commandlets/DerivedDataCacheCommandlet.cpp b/Engine/Source/Editor/UnrealEd/Private/Commandlets/DerivedDataCacheCommandlet.cpp
--- a/Engine/Source/Editor/UnrealEd/Private/Commandlets/DerivedDataCacheCommandlet.cpp
+++ b/Engine/Source/Editor/UnrealEd/Private/Commandlets/DerivedDataCacheCommandlet.cpp
@@ -350,10 +350,13 @@ void UDerivedDataCacheCommandlet::CacheWorldPackages(UWorld* World, uint8 Packag
 		UWorldPartition* WorldPartition = World->GetWorldPartition();
 		check(WorldPartition);
 
-		FWorldPartitionHelpers::ForEachActorWithLoading(WorldPartition, [this, PackageFilter, &Platforms](AActor* Actor)
+		FWorldPartitionHelpers::ForEachActorWithLoading(WorldPartition, [this, PackageFilter, &Platforms](const FWorldPartitionActorDesc* ActorDesc)
 		{
-			UE_LOG(LogDerivedDataCacheCommandlet, Display, TEXT("Loaded actor %s"), *Actor->GetName());
-			CacheLoadedPackages(Actor->GetPackage(), PackageFilter, Platforms);
+			if (AActor* Actor = ActorDesc->GetActor())
+			{
+				UE_LOG(LogDerivedDataCacheCommandlet, Display, TEXT("Loaded actor %s"), *Actor->GetName());
+				CacheLoadedPackages(Actor->GetPackage(), PackageFilter, Platforms);
+			}
 			return true;
 		});
 	}
diff --git a/Engine/Source/Editor/UnrealEd/Private/WorldPartition/WorldPartitionResaveActorsBuilder.cpp b/Engine/Source/Editor/UnrealEd/Private/WorldPartition/WorldPartitionResaveActorsBuilder.cpp
--- a/Engine/Source/Editor/UnrealEd/Private/WorldPartition/WorldPartitionResaveActorsBuilder.cpp
+++ b/Engine/Source/Editor/UnrealEd/Private/WorldPartition/WorldPartitionResaveActorsBuilder.cpp
@@ -44,10 +44,27 @@ bool UWorldPartitionResaveActorsBuilder::RunInternal(UWorld* World, const FBox&
 		return false;
 	}
 
-	FWorldPartitionHelpers::ForEachActorWithLoading(WorldPartition, ActorClass, [&SaveCount, &FailCount, &SCCHelper](AActor* Actor)
+	TArray<FString> PackagesToDelete;
+	FWorldPartitionHelpers::ForEachActorWithLoading(WorldPartition, ActorClass, [&SaveCount, &FailCount, &SCCHelper, &PackagesToDelete](const FWorldPartitionActorDesc* ActorDesc)
 	{
-		if (UPackage* Package = Actor->GetExternalPackage())
+		ON_SCOPE_EXIT
 		{
+			UE_LOG(LogWorldPartitionResaveActorsBuilder, Display, TEXT("Processed %d packages (%d Saved / %d Failed)"), SaveCount + FailCount, SaveCount, FailCount);
+		};
+
+		AActor* Actor = ActorDesc->GetActor();
+
+		if (!Actor)
+		{
+			PackagesToDelete.Add(ActorDesc->GetActorPackage().ToString());
+			++FailCount;
+			return true;
+		}
+		else
+		{		
+			UPackage* Package = Actor->GetExternalPackage();
+			check(Package);
+
 			if (SCCHelper.Checkout(Package))
 			{
 				// Save package
@@ -62,6 +79,7 @@ bool UWorldPartitionResaveActorsBuilder::RunInternal(UWorld* World, const FBox&
 				// It is possible the resave can't checkout everything. Continue processing.
 				UE_LOG(LogWorldPartitionResaveActorsBuilder, Display, TEXT("Saved package %s."), *Package->GetName());
 				++SaveCount;
+				return true;
 			}
 			else
 			{
@@ -70,8 +88,6 @@ bool UWorldPartitionResaveActorsBuilder::RunInternal(UWorld* World, const FBox&
 				++FailCount;
 				return true;
 			}
-
-			UE_LOG(LogWorldPartitionResaveActorsBuilder, Display, TEXT("Processed %d packages (%d Saved / %d Failed)"), SaveCount + FailCount, SaveCount, FailCount);
 		}
 
 		return true;
@@ -79,6 +95,15 @@ bool UWorldPartitionResaveActorsBuilder::RunInternal(UWorld* World, const FBox&
 
 	UPackage::WaitForAsyncFileWrites();
 
+	for (const FString& PackageToDelete : PackagesToDelete)
+	{
+		if (SCCHelper.Delete(PackageToDelete))
+		{
+			UE_LOG(LogWorldPartitionResaveActorsBuilder, Error, TEXT("Error deleting package %s."), *PackageToDelete);
+			return true;
+		}
+	}
+
 	return true;
 }
 
diff --git a/Engine/Source/Runtime/Engine/Private/WorldPartition/WorldPartitionHelpers.cpp b/Engine/Source/Runtime/Engine/Private/WorldPartition/WorldPartitionHelpers.cpp
--- a/Engine/Source/Runtime/Engine/Private/WorldPartition/WorldPartitionHelpers.cpp
+++ b/Engine/Source/Runtime/Engine/Private/WorldPartition/WorldPartitionHelpers.cpp
@@ -29,7 +29,7 @@ void FWorldPartitionHelpers::ForEachActorDesc(UWorldPartition* WorldPartition, T
 	}
 }
 
-void FWorldPartitionHelpers::ForEachActorWithLoading(UWorldPartition* WorldPartition, TSubclassOf<AActor> ActorClass, TFunctionRef<bool(AActor*)> Predicate)
+void FWorldPartitionHelpers::ForEachActorWithLoading(UWorldPartition* WorldPartition, TSubclassOf<AActor> ActorClass, TFunctionRef<bool(const FWorldPartitionActorDesc*)> Predicate)
 {
 	// Recursive loading of references 
 	TFunction<void(const FGuid&, TMap<FGuid, FWorldPartitionReference>&)> LoadReferences = [WorldPartition, &LoadReferences](const FGuid& ActorGuid, TMap<FGuid, FWorldPartitionReference>& InOutActorReferences)
@@ -55,35 +55,40 @@ void FWorldPartitionHelpers::ForEachActorWithLoading(UWorldPartition* WorldParti
 	TMap<FGuid, FWorldPartitionReference> ActorReferences;
 	TArray<FGuid> ActorsToProcess;
 
+	auto DoCollectGarbage = [&ActorReferences]()
+	{
+		ActorReferences.Empty();
+
+		const FPlatformMemoryStats MemStatsBefore = FPlatformMemory::GetStats();
+		CollectGarbage(RF_NoFlags, true);
+		const FPlatformMemoryStats MemStatsAfter = FPlatformMemory::GetStats();
+
+		UE_LOG(LogWorldPartition, Log, TEXT("GC Performed - Available Physical: %.2fGB, Available Virtual: %.2fGB"),
+			(int64)MemStatsAfter.AvailablePhysical / (1024.0 * 1024.0 * 1024.0),
+			(int64)MemStatsAfter.AvailableVirtual / (1024.0 * 1024.0 * 1024.0)
+		);
+	};
+
 	ForEachActorDesc(WorldPartition, ActorClass, [&](const FWorldPartitionActorDesc* ActorDesc)
 	{
 		const FGuid& ActorGuid = ActorDesc->GetGuid();
 		LoadReferences(ActorGuid, ActorReferences);
 
 		FWorldPartitionReference ActorReference(WorldPartition, ActorGuid);
-		AActor* Actor = ActorReference.Get()->GetActor();
-
-		if (!Predicate(Actor))
+		if (!Predicate(ActorReference.Get()))
 		{
 			return false;
 		}
 
 		if (HasExceededMaxMemory())
 		{
-			ActorReferences.Empty();
-
-			const FPlatformMemoryStats MemStatsBefore = FPlatformMemory::GetStats();
-			CollectGarbage(RF_NoFlags, true);
-			const FPlatformMemoryStats MemStatsAfter = FPlatformMemory::GetStats();
-
-			UE_LOG(LogWorldPartition, Log, TEXT("GC Performed - Available Physical: %.2fGB, Available Virtual: %.2fGB"),
-				(int64)MemStatsAfter.AvailablePhysical / (1024.0 * 1024.0 * 1024.0),
-				(int64)MemStatsAfter.AvailableVirtual / (1024.0 * 1024.0 * 1024.0)
-			);
+			DoCollectGarbage();
 		}
 
 		return true;
 	});
+
+	DoCollectGarbage();
 }
 
 bool FWorldPartitionHelpers::HasExceededMaxMemory()
diff --git a/Engine/Source/Runtime/Engine/Public/WorldPartition/WorldPartitionHelpers.h b/Engine/Source/Runtime/Engine/Public/WorldPartition/WorldPartitionHelpers.h
--- a/Engine/Source/Runtime/Engine/Public/WorldPartition/WorldPartitionHelpers.h
+++ b/Engine/Source/Runtime/Engine/Public/WorldPartition/WorldPartitionHelpers.h
@@ -25,14 +25,14 @@ public:
 	}
 
 	template<class ActorClass = AActor>
-	static void ForEachActorWithLoading(UWorldPartition* WorldPartition, TFunctionRef<bool(AActor*)> Func)
+	static void ForEachActorWithLoading(UWorldPartition* WorldPartition, TFunctionRef<bool(const FWorldPartitionActorDesc*)> Func)
 	{
 		ForEachActorWithLoading(WorldPartition, ActorClass::StaticClass(), Func);
 	}
 
 	static void ForEachIntersectingActorDesc(UWorldPartition* WorldPartition, const FBox& Box, TSubclassOf<AActor> ActorClass, TFunctionRef<bool(const FWorldPartitionActorDesc*)> Func);
 	static void ForEachActorDesc(UWorldPartition* WorldPartition, TSubclassOf<AActor> ActorClass, TFunctionRef<bool(const FWorldPartitionActorDesc*)> Func);
-	static void ForEachActorWithLoading(UWorldPartition* WorldPartition, TSubclassOf<AActor> ActorClass, TFunctionRef<bool(AActor*)> Func);
+	static void ForEachActorWithLoading(UWorldPartition* WorldPartition, TSubclassOf<AActor> ActorClass, TFunctionRef<bool(const FWorldPartitionActorDesc*)> Func);
 
 	static bool HasExceededMaxMemory();
 };
