WorldPartition: support loading of invalid actor files.
Temp fix until we save the full class path in the asset registry or in the actor descriptors to detect invalid actor classes before regisdtration.

diff --git a/Engine/Source/Runtime/Engine/Private/WorldPartition/WorldPartition.cpp b/Engine/Source/Runtime/Engine/Private/WorldPartition/WorldPartition.cpp
--- a/Engine/Source/Runtime/Engine/Private/WorldPartition/WorldPartition.cpp
+++ b/Engine/Source/Runtime/Engine/Private/WorldPartition/WorldPartition.cpp
@@ -835,15 +835,15 @@ bool UWorldPartition::UpdateEditorCells(TFunctionRef<bool(TArray<UWorldPartition
 		// Only prompt if the actor will get unloaded by the unloading cells
 		if (ActorDesc->GetHardRefCount() == Pair.Value)
 		{
-			AActor* LoadedActor = ActorDesc->GetActor();
-			checkf(LoadedActor, TEXT("Can't find loaded actor %s"), *ActorDesc->ToString());
-
-			UPackage* ActorPackage = LoadedActor->GetExternalPackage();
-			if (ActorPackage && ActorPackage->IsDirty())
+			if (AActor* LoadedActor = ActorDesc->GetActor())
 			{
-				ModifiedPackages.Add(ActorPackage);
+				UPackage* ActorPackage = LoadedActor->GetExternalPackage();
+				if (ActorPackage && ActorPackage->IsDirty())
+				{
+					ModifiedPackages.Add(ActorPackage);
+				}
+				bIsUpdateUnloadingActors = true;
 			}
-			bIsUpdateUnloadingActors = true;
 		}
 	}
 
@@ -955,10 +955,6 @@ void UWorldPartition::UpdateLoadingEditorCell(UWorldPartitionEditorCell* Cell, b
 		{
 			AActor* Actor = ActorHandle->GetActor();
 
-			// The actor could be either loaded but with no cells loaded (in the case of a reference from another actor, for example)
-			// or directly referenced by a loaded cell.
-			check(Actor || !ActorHandle->GetHardRefCount());
-
 			// Filter actor against DataLayers
 			if (ShouldActorBeLoadedByEditorCells(*ActorHandle))
 			{
diff --git a/Engine/Source/Runtime/Engine/Private/WorldPartition/WorldPartitionActorDesc.cpp b/Engine/Source/Runtime/Engine/Private/WorldPartition/WorldPartitionActorDesc.cpp
--- a/Engine/Source/Runtime/Engine/Private/WorldPartition/WorldPartitionActorDesc.cpp
+++ b/Engine/Source/Runtime/Engine/Private/WorldPartition/WorldPartitionActorDesc.cpp
@@ -274,7 +274,7 @@ AActor* FWorldPartitionActorDesc::Load() const
 		if (Package)
 		{
 			ActorPtr = FindObject<AActor>(nullptr, *ActorPath.ToString());
-			check(ActorPtr.IsValid());
+			UE_CLOG(!ActorPtr.IsValid(), LogWorldPartition, Warning, TEXT("Can't load actor %s"), *GetActorName().ToString());
 		}
 	}
 
