Pick better default colors for spline component visualizers

diff --git a/Engine/Source/Runtime/Engine/Private/Components/SplineComponent.cpp b/Engine/Source/Runtime/Engine/Private/Components/SplineComponent.cpp
--- a/Engine/Source/Runtime/Engine/Private/Components/SplineComponent.cpp
+++ b/Engine/Source/Runtime/Engine/Private/Components/SplineComponent.cpp
@@ -10,6 +10,7 @@
 #include "PrimitiveSceneProxy.h"
 #include "SceneManagement.h"
 #include "UnrealEngine.h"
+#include "Styling/StyleColors.h"
 
 #if WITH_EDITOR
 #include "Settings/LevelEditorViewportSettings.h"
@@ -40,9 +41,9 @@ USplineComponent::USplineComponent(const FObjectInitializer& ObjectInitializer)
 	, bClosedLoop(false)
 	, DefaultUpVector(FVector::UpVector)
 #if WITH_EDITORONLY_DATA
-	, EditorUnselectedSplineSegmentColor(FLinearColor(1.0f, 1.0f, 1.0f))
-	, EditorSelectedSplineSegmentColor(FLinearColor(0.828f, 0.364f, 0.003f))
-	, EditorTangentColor(FLinearColor(1.0f, 1.0f, 1.0f))
+	, EditorUnselectedSplineSegmentColor(FStyleColors::White.GetSpecifiedColor())
+	, EditorSelectedSplineSegmentColor(FStyleColors::AccentOrange.GetSpecifiedColor())
+	, EditorTangentColor(FLinearColor(0.718f, 0.589f, 0.921f))
 	, bAllowDiscontinuousSpline(false)
 	, bShouldVisualizeScale(false)
 	, ScaleVisualizationWidth(30.0f)
