Update display names and tooltips of spline component color properties

diff --git a/Engine/Source/Runtime/Engine/Classes/Components/SplineComponent.h b/Engine/Source/Runtime/Engine/Classes/Components/SplineComponent.h
--- a/Engine/Source/Runtime/Engine/Classes/Components/SplineComponent.h
+++ b/Engine/Source/Runtime/Engine/Classes/Components/SplineComponent.h
@@ -267,16 +267,16 @@ public:
 	FVector DefaultUpVector;
 
 #if WITH_EDITORONLY_DATA
-	/** Color of an unselected spline component segment in the editor */
-	UPROPERTY(EditAnywhere, Category = Editor)
+	/** Color of unselected spline component parts in the editor */
+	UPROPERTY(EditAnywhere, Category = Editor, meta = (DisplayName="Editor Spline Unselected Color"))
 	FLinearColor EditorUnselectedSplineSegmentColor;
 
-	/** Color of a selected spline component segment in the editor */
-	UPROPERTY(EditAnywhere, Category = Editor)
+	/** Color of selected spline component parts in the editor */
+	UPROPERTY(EditAnywhere, Category = Editor, meta = (DisplayName="Editor Spline Selected Color"))
 	FLinearColor EditorSelectedSplineSegmentColor;
 
-	/** Color of spline point tangent in the editor */
-	UPROPERTY(EditAnywhere, Category = Editor)
+	/** Color of spline point tangents in the editor */
+	UPROPERTY(EditAnywhere, Category = Editor, meta = (DisplayName="Editor Spline Tangent Color"))
 	FLinearColor EditorTangentColor;
 
 	/** Whether the spline's leave and arrive tangents can be different */
