PR #8216: UE5: Bug fix for blueprint AI systems not being created (Contributed by [REDACTED_PERSON])

diff --git a/Engine/Source/Runtime/Engine/Private/WorldSettings.cpp b/Engine/Source/Runtime/Engine/Private/WorldSettings.cpp
--- a/Engine/Source/Runtime/Engine/Private/WorldSettings.cpp
+++ b/Engine/Source/Runtime/Engine/Private/WorldSettings.cpp
@@ -943,7 +943,7 @@ void AWorldSettings::UpdateBookmarkClass()
 
 FSoftClassPath AWorldSettings::GetAISystemClassName() const
 {
-	return bEnableAISystem ? AISystemClass.Get() : FSoftClassPath();
+	return bEnableAISystem ? FSoftClassPath(AISystemClass.ToString()) : FSoftClassPath();
 }
 
 void AWorldSettings::RewindForReplay()
