Replaced the usage of macro with a function (easier for licensees to overload).

diff --git a/Engine/Source/Runtime/Engine/Private/ShaderCompiler/ShaderCompiler.cpp b/Engine/Source/Runtime/Engine/Private/ShaderCompiler/ShaderCompiler.cpp
--- a/Engine/Source/Runtime/Engine/Private/ShaderCompiler/ShaderCompiler.cpp
+++ b/Engine/Source/Runtime/Engine/Private/ShaderCompiler/ShaderCompiler.cpp
@@ -5973,7 +5973,7 @@ static FString GetGlobalShaderMapKeyString(const FGlobalShaderMapId& ShaderMapId
 	FString ShaderMapKeyString = Format.ToString() + TEXT("_") + FString(FString::FromInt(GetTargetPlatformManagerRef().ShaderFormatVersion(Format))) + TEXT("_");
 	ShaderMapAppendKeyString(Platform, ShaderMapKeyString);
 	ShaderMapId.AppendKeyString(ShaderMapKeyString, Dependencies);
-	return FDerivedDataCacheInterface::BuildCacheKey(TEXT("GSM"), GLOBALSHADERMAP_DERIVEDDATA_VER, *ShaderMapKeyString);
+	return FDerivedDataCacheInterface::BuildCacheKey(TEXT("GSM"), *GetGlobalShaderMapDDCKey(), *ShaderMapKeyString);
 }
 
 /** Saves the platform's shader map to the DDC. */
