Added logic to properly detect and clear invalid texture streaming built data.

diff --git a/Engine/Source/Runtime/Engine/Classes/Components/StaticMeshComponent.h b/Engine/Source/Runtime/Engine/Classes/Components/StaticMeshComponent.h
--- a/Engine/Source/Runtime/Engine/Classes/Components/StaticMeshComponent.h
+++ b/Engine/Source/Runtime/Engine/Classes/Components/StaticMeshComponent.h
@@ -686,6 +686,9 @@ private:
 
 	/** Called when the StaticMesh property gets overwritten without us knowing about it */
 	void OutdatedKnownStaticMeshDetected() const;
+
+	/** Clears texture streaming data. */
+	void ClearStreamingTextureData();
 #endif
 protected:
 
diff --git a/Engine/Source/Runtime/Engine/Private/Components/PrimitiveComponent.cpp b/Engine/Source/Runtime/Engine/Private/Components/PrimitiveComponent.cpp
--- a/Engine/Source/Runtime/Engine/Private/Components/PrimitiveComponent.cpp
+++ b/Engine/Source/Runtime/Engine/Private/Components/PrimitiveComponent.cpp
@@ -465,10 +465,10 @@ bool UPrimitiveComponent::BuildTextureStreamingData(ETextureStreamingBuildType B
 {
 	bool bSupportsBuildTextureStreamingData = false;
 	bIsValidTextureStreamingBuiltData = false;
-	if (!FApp::CanEverRender())
-	{
-		return false;
-	}
+#if WITH_EDITOR
+	bIsActorTextureStreamingBuiltData = false;
+#endif
+	
 	bool Result = BuildTextureStreamingDataImpl(BuildType, QualityLevel, FeatureLevel, DependentResources, bSupportsBuildTextureStreamingData);
 	if (Result && bSupportsBuildTextureStreamingData)
 	{
diff --git a/Engine/Source/Runtime/Engine/Private/Components/StaticMeshComponent.cpp b/Engine/Source/Runtime/Engine/Private/Components/StaticMeshComponent.cpp
--- a/Engine/Source/Runtime/Engine/Private/Components/StaticMeshComponent.cpp
+++ b/Engine/Source/Runtime/Engine/Private/Components/StaticMeshComponent.cpp
@@ -803,7 +803,7 @@ bool UStaticMeshComponent::GetMaterialStreamingData(int32 MaterialIndex, FPrimit
 bool UStaticMeshComponent::RemapActorTextureStreamingBuiltDataToLevel(const UActorTextureStreamingBuildDataComponent* InActorTextureBuildData)
 {
 	check(InActorTextureBuildData);
-	check(InActorTextureBuildData->GetOwner() == GetOwner())
+	check(InActorTextureBuildData->GetOwner() == GetOwner());
 
 	ULevel* Level = GetOwner()->GetLevel();
 	if (!Level || !bIsActorTextureStreamingBuiltData)
@@ -817,6 +817,9 @@ bool UStaticMeshComponent::RemapActorTextureStreamingBuiltDataToLevel(const UAct
 	{
 		if (!InActorTextureBuildData->GetStreamableTexture(BuildInfo.TextureLevelIndex, TextureName, TextureGuid))
 		{
+			// If remapping failed, invalidate built texture streaming data (this should not happen with newly generated texture streaming build data)
+			UE_LOG(LogStaticMesh, Warning, TEXT("Clearing invalid texture streaming built data for %s"), *GetFullName());
+			ClearStreamingTextureData();
 			return false;
 		}
 		// Update BuildInfo's TextureLevelIndex
@@ -837,7 +840,12 @@ bool UStaticMeshComponent::BuildTextureStreamingDataImpl(ETextureStreamingBuildT
 	{
 		AActor* ComponentActor = GetOwner();
 
-		if (!bIgnoreInstanceForTextureStreaming && Mobility == EComponentMobility::Static && GetStaticMesh() && GetStaticMesh()->GetRenderData() && !bHiddenInGame)
+		const bool bCanBuildTextureStreamingData = FApp::CanEverRender();
+		if (!bCanBuildTextureStreamingData)
+		{
+			bBuildDataValid = false;
+		}
+		if (!bIgnoreInstanceForTextureStreaming && Mobility == EComponentMobility::Static && GetStaticMesh() && GetStaticMesh()->GetRenderData() && !bHiddenInGame && bCanBuildTextureStreamingData)
 		{
 			// First generate the bounds. Will be used in the texture streaming build and also in the debug viewmode.
 			const int32 NumMaterials = GetNumMaterials();
@@ -884,6 +892,11 @@ bool UStaticMeshComponent::BuildTextureStreamingDataImpl(ETextureStreamingBuildT
 					PackStreamingTextureData(TextureStreamingContainer, UnpackedData, StreamingTextureData, Bounds);
 					bOutSupportsBuildTextureStreamingData = true;
 				}
+				else
+				{
+					UE_LOG(LogStaticMesh, Warning, TEXT("No texture streaming container found : Can't build texture streaming data for %s"), *GetFullName());
+					bBuildDataValid = false;
+				}
 			}
 			else if (StreamingTextureData.Num() == 0)
 			{
@@ -937,19 +950,32 @@ bool UStaticMeshComponent::BuildTextureStreamingDataImpl(ETextureStreamingBuildT
 		}
 		else // Otherwise clear any data.
 		{
-			StreamingTextureData.Empty();
-
-			if (MaterialStreamingRelativeBoxes.Num())
-			{
-				MaterialStreamingRelativeBoxes.Empty();
-				MarkRenderStateDirty(); // Update since proxy has a copy of the material bounds.
-			}
+			ClearStreamingTextureData();
 		}
 	}
+
+	// Make sure to clear invalid streaming texture data
+	if ((BuildType == TSB_MapBuild || BuildType == TSB_ActorBuild) && !bOutSupportsBuildTextureStreamingData)
+	{
+		ClearStreamingTextureData();
+	}
 #endif
 	return bBuildDataValid;
 }
 
+#if WITH_EDITOR
+void UStaticMeshComponent::ClearStreamingTextureData()
+{
+	StreamingTextureData.Empty();
+
+	if (MaterialStreamingRelativeBoxes.Num())
+	{
+		MaterialStreamingRelativeBoxes.Empty();
+		MarkRenderStateDirty(); // Update since proxy has a copy of the material bounds.
+	}
+}
+#endif
+
 float UStaticMeshComponent::GetTextureStreamingTransformScale() const
 {
 	return GetComponentTransform().GetMaximumAxisScale();
diff --git a/Engine/Source/Runtime/Engine/Private/Streaming/TextureStreamingBuild.cpp b/Engine/Source/Runtime/Engine/Private/Streaming/TextureStreamingBuild.cpp
--- a/Engine/Source/Runtime/Engine/Private/Streaming/TextureStreamingBuild.cpp
+++ b/Engine/Source/Runtime/Engine/Private/Streaming/TextureStreamingBuild.cpp
@@ -570,7 +570,7 @@ void FStreamingTextureLevelContext::ProcessMaterial(const FBoxSphereBounds& Comp
 		{
 			for (const FStreamingTextureBuildInfo& BuildInfo : *ComponentBuildData)
 			{
-				if (LevelStreamingTextures.IsValidIndex(BuildInfo.TextureLevelIndex))
+				if (ensure(LevelStreamingTextures.IsValidIndex(BuildInfo.TextureLevelIndex)))
 				{
 					if (UTexture* Texture = LevelStreamingTextures[BuildInfo.TextureLevelIndex])
 					{
