WorldPartition: Add bBlockOnSlowLoading to Streaming Source structure so that prefetching streaming sources don't block loading

diff --git a/Engine/Source/Runtime/Engine/Private/WorldPartition/WorldPartitionRuntimeSpatialHash.cpp b/Engine/Source/Runtime/Engine/Private/WorldPartition/WorldPartitionRuntimeSpatialHash.cpp
--- a/Engine/Source/Runtime/Engine/Private/WorldPartition/WorldPartitionRuntimeSpatialHash.cpp
+++ b/Engine/Source/Runtime/Engine/Private/WorldPartition/WorldPartitionRuntimeSpatialHash.cpp
@@ -1161,17 +1161,21 @@ EWorldPartitionStreamingPerformance UWorldPartitionRuntimeSpatialHash::GetStream
 	check(StreamingGrid && *StreamingGrid);
 	
 	const float LoadingRange = (*StreamingGrid)->LoadingRange;
-	const float Distance = FMath::Sqrt(StreamingCell->CachedSourceMinDistanceSquare) - ((float)(*StreamingGrid)->GetCellSize(StreamingCell->Level) / 2.f);
 
-	const float Ratio = Distance / LoadingRange;
-
-	if (Ratio < GBlockOnSlowStreamingRatio)
+	if (StreamingCell->CachedIsBlockingSource)
 	{
-		return EWorldPartitionStreamingPerformance::Critical;
-	}
-	else if (Ratio < BlockOnSlowStreamingWarningRatio)
-	{
-		return EWorldPartitionStreamingPerformance::Slow;
+		const float Distance = FMath::Sqrt(StreamingCell->CachedBlockingMinDistanceSquare) - ((float)(*StreamingGrid)->GetCellSize(StreamingCell->Level) / 2.f);
+
+		const float Ratio = Distance / LoadingRange;
+
+		if (Ratio < GBlockOnSlowStreamingRatio)
+		{
+			return EWorldPartitionStreamingPerformance::Critical;
+		}
+		else if (Ratio < BlockOnSlowStreamingWarningRatio)
+		{
+			return EWorldPartitionStreamingPerformance::Slow;
+		}
 	}
 
 	return EWorldPartitionStreamingPerformance::Good;
diff --git a/Engine/Source/Runtime/Engine/Private/WorldPartition/WorldPartitionRuntimeSpatialHashCell.cpp b/Engine/Source/Runtime/Engine/Private/WorldPartition/WorldPartitionRuntimeSpatialHashCell.cpp
--- a/Engine/Source/Runtime/Engine/Private/WorldPartition/WorldPartitionRuntimeSpatialHashCell.cpp
+++ b/Engine/Source/Runtime/Engine/Private/WorldPartition/WorldPartitionRuntimeSpatialHashCell.cpp
@@ -17,7 +17,8 @@ static FAutoConsoleVariableRef CVarRuntimeSpatialHashCellToSourceAngleContributi
 UWorldPartitionRuntimeSpatialHashCell::UWorldPartitionRuntimeSpatialHashCell(const FObjectInitializer& ObjectInitializer)
 : Super(ObjectInitializer)
 , Level(0)
-, CachedSourceMinDistanceSquare(0.f)
+, CachedIsBlockingSource(false)
+, CachedBlockingMinDistanceSquare(MAX_FLT)
 , CachedSourceMinSortDistance(0.f)
 {}
 
@@ -40,7 +41,12 @@ void UWorldPartitionRuntimeSpatialHashCell::PostDuplicate(bool bDuplicateForPIE)
 bool UWorldPartitionRuntimeSpatialHashCell::CacheStreamingSourceInfo(const FWorldPartitionStreamingSource& Source) const
 {
 	const bool bWasCacheDirtied = Super::CacheStreamingSourceInfo(Source);
-	
+	if (bWasCacheDirtied)
+	{
+		CachedIsBlockingSource = false;
+		CachedBlockingMinDistanceSquare = MAX_FLT;
+	}
+
 	const float AngleContribution = FMath::Clamp(GRuntimeSpatialHashCellToSourceAngleContributionToCellImportance, 0.f, 1.f);
 	const float SqrDistance = FVector::DistSquared(Source.Location, Position);
 	float AngleFactor = 1.f;
@@ -57,7 +63,13 @@ bool UWorldPartitionRuntimeSpatialHashCell::CacheStreamingSourceInfo(const FWorl
 
 	// If cache was dirtied, use value, else use minimum with existing cached value
 	CachedSourceMinSortDistance = bWasCacheDirtied ? ModulatedDistance : FMath::Min(ModulatedDistance, CachedSourceMinSortDistance);
-	CachedSourceMinDistanceSquare = bWasCacheDirtied ? SqrDistance : FMath::Min(SqrDistance, CachedSourceMinDistanceSquare);
+	
+	// Only consider blocking sources
+	if (Source.bBlockOnSlowLoading)
+	{
+		CachedIsBlockingSource = true;
+		CachedBlockingMinDistanceSquare = FMath::Min(SqrDistance, CachedBlockingMinDistanceSquare);
+	}
 
 	return bWasCacheDirtied;
 }
diff --git a/Engine/Source/Runtime/Engine/Private/WorldPartition/WorldPartitionStreamingPolicy.cpp b/Engine/Source/Runtime/Engine/Private/WorldPartition/WorldPartitionStreamingPolicy.cpp
--- a/Engine/Source/Runtime/Engine/Private/WorldPartition/WorldPartitionStreamingPolicy.cpp
+++ b/Engine/Source/Runtime/Engine/Private/WorldPartition/WorldPartitionStreamingPolicy.cpp
@@ -94,7 +94,7 @@ void UWorldPartitionStreamingPolicy::UpdateStreamingSources()
 		const FVector ViewLocationLocal = WorldToLocal.TransformPosition(ViewLocation);
 		const FRotator ViewRotationLocal = WorldToLocal.TransformRotation(ViewRotation.Quaternion()).Rotator();
 		static const FName NAME_SIE(TEXT("SIE"));
-		StreamingSources.Add(FWorldPartitionStreamingSource(NAME_SIE, ViewLocationLocal, ViewRotationLocal, EStreamingSourceTargetState::Activated));
+		StreamingSources.Add(FWorldPartitionStreamingSource(NAME_SIE, ViewLocationLocal, ViewRotationLocal, EStreamingSourceTargetState::Activated, /*bBlockOnSlowLoading=*/false));
 	}
 	else
 #endif
@@ -116,7 +116,7 @@ void UWorldPartitionStreamingPolicy::UpdateStreamingSources()
 					// Transform to Local
 					ViewLocation = WorldToLocal.TransformPosition(ViewLocation);
 					ViewRotation = WorldToLocal.TransformRotation(ViewRotation.Quaternion()).Rotator();
-					StreamingSources.Add(FWorldPartitionStreamingSource(Player->GetFName(), ViewLocation, ViewRotation, EStreamingSourceTargetState::Activated));
+					StreamingSources.Add(FWorldPartitionStreamingSource(Player->GetFName(), ViewLocation, ViewRotation, EStreamingSourceTargetState::Activated, /*bBlockOnSlowLoading=*/true));
 				}
 			}
 		}
diff --git a/Engine/Source/Runtime/Engine/Public/WorldPartition/WorldPartitionRuntimeSpatialHashCell.h b/Engine/Source/Runtime/Engine/Public/WorldPartition/WorldPartitionRuntimeSpatialHashCell.h
--- a/Engine/Source/Runtime/Engine/Public/WorldPartition/WorldPartitionRuntimeSpatialHashCell.h
+++ b/Engine/Source/Runtime/Engine/Public/WorldPartition/WorldPartitionRuntimeSpatialHashCell.h
@@ -24,8 +24,11 @@ class UWorldPartitionRuntimeSpatialHashCell : public UWorldPartitionRuntimeCell
 	UPROPERTY()
 	int32 Level;
 		
+	// Used to determine if cell was requested by blocking source
+	mutable bool CachedIsBlockingSource;
+
 	// Used to determine if the streaming of cells is lagging
-	mutable float CachedSourceMinDistanceSquare;
+	mutable float CachedBlockingMinDistanceSquare;
 
 	// Computed and cached value used by SortCompare to sort Cells
 	mutable float CachedSourceMinSortDistance;
diff --git a/Engine/Source/Runtime/Engine/Public/WorldPartition/WorldPartitionStreamingSource.h b/Engine/Source/Runtime/Engine/Public/WorldPartition/WorldPartitionStreamingSource.h
--- a/Engine/Source/Runtime/Engine/Public/WorldPartition/WorldPartitionStreamingSource.h
+++ b/Engine/Source/Runtime/Engine/Public/WorldPartition/WorldPartitionStreamingSource.h
@@ -84,15 +84,17 @@ enum class EStreamingSourcePriority : int32
 struct ENGINE_API FWorldPartitionStreamingSource
 {
 	FWorldPartitionStreamingSource()
-		: Priority(EStreamingSourcePriority::Default)
+		: bBlockOnSlowLoading(false)
+		, Priority(EStreamingSourcePriority::Default)
 		, Velocity(0.f)
 	{}
 
-	FWorldPartitionStreamingSource(FName InName, const FVector& InLocation, const FRotator& InRotation, EStreamingSourceTargetState InTargetState, EStreamingSourcePriority InPriority = EStreamingSourcePriority::Default, float InVelocity = 0.f)
+	FWorldPartitionStreamingSource(FName InName, const FVector& InLocation, const FRotator& InRotation, EStreamingSourceTargetState InTargetState, bool bInBlockOnSlowLoading, EStreamingSourcePriority InPriority = EStreamingSourcePriority::Default, float InVelocity = 0.f)
 		: Name(InName)
 		, Location(InLocation)
 		, Rotation(InRotation)
 		, TargetState(InTargetState)
+		, bBlockOnSlowLoading(bInBlockOnSlowLoading)
 		, Priority(InPriority)
 		, Velocity(InVelocity)
 	{
@@ -104,6 +106,7 @@ struct ENGINE_API FWorldPartitionStreamingSource
 	FVector Location;
 	FRotator Rotation;
 	EStreamingSourceTargetState TargetState;
+	bool bBlockOnSlowLoading;
 	EStreamingSourcePriority Priority;
 	float Velocity; // Computed automatically
 };
