Introduces a EMobileAntiAliasingMethod so that TSR doesn't show up in the Mobile Anti-Aliasing Method renderer settings

diff --git a/Engine/Plugins/Runtime/Oculus/OculusVR/Source/OculusEditor/Private/OculusToolWidget.cpp b/Engine/Plugins/Runtime/Oculus/OculusVR/Source/OculusEditor/Private/OculusToolWidget.cpp
--- a/Engine/Plugins/Runtime/Oculus/OculusVR/Source/OculusEditor/Private/OculusToolWidget.cpp
+++ b/Engine/Plugins/Runtime/Oculus/OculusVR/Source/OculusEditor/Private/OculusToolWidget.cpp
@@ -705,7 +705,7 @@ EVisibility SOculusToolWidget::MobileMultiViewVisibility(FName tag) const
 FReply SOculusToolWidget::MobileMSAAEnable(bool text)
 {
 	URendererSettings* Settings = GetMutableDefault<URendererSettings>();
-	Settings->MobileAntiAliasing = EAntiAliasingMethod::AAM_MSAA;
+	Settings->MobileAntiAliasing = EMobileAntiAliasingMethod::MSAA;
 	Settings->UpdateSinglePropertyInConfigFile(Settings->GetClass()->FindPropertyByName(GET_MEMBER_NAME_CHECKED(URendererSettings, MobileAntiAliasing)), Settings->GetDefaultConfigFilename());
 	Settings->MSAASampleCount = ECompositingSampleCount::Four;
 	Settings->UpdateSinglePropertyInConfigFile(Settings->GetClass()->FindPropertyByName(GET_MEMBER_NAME_CHECKED(URendererSettings, MSAASampleCount)), Settings->GetDefaultConfigFilename());
diff --git a/Engine/Source/Runtime/Engine/Classes/Engine/RendererSettings.h b/Engine/Source/Runtime/Engine/Classes/Engine/RendererSettings.h
--- a/Engine/Source/Runtime/Engine/Classes/Engine/RendererSettings.h
+++ b/Engine/Source/Runtime/Engine/Classes/Engine/RendererSettings.h
@@ -155,6 +155,19 @@ namespace EFixedFoveationLevels
 	};
 }
 
+UENUM()
+namespace EMobileAntiAliasingMethod
+{
+	enum Type
+	{
+		None = AAM_None UMETA(DisplayName = "None"),
+		FXAA = AAM_FXAA UMETA(DisplayName = "Fast Approximate Anti-Aliasing (FXAA)"),
+		TemporalAA = AAM_TemporalAA UMETA(DisplayName = "Temporal Anti-Aliasing (TAA)"),
+		/** Only supported with forward shading.  MSAA sample count is controlled by r.MSAACount. */
+		MSAA = AAM_MSAA UMETA(DisplayName = "Multisample Anti-Aliasing (MSAA)"),
+	};
+}
+
 namespace EDefaultBackBufferPixelFormat
 {
 	ENGINE_API EPixelFormat Convert2PixelFormat(EDefaultBackBufferPixelFormat::Type InDefaultBackBufferPixelFormat);
@@ -208,7 +221,7 @@ class ENGINE_API URendererSettings : public UDeveloperSettings
 	UPROPERTY(config, EditAnywhere, Category = Mobile, meta = (
 		ConsoleVariable = "r.Mobile.AntiAliasing", DisplayName = "Mobile Anti-Aliasing Method",
 		ToolTip = "The mobile default anti-aliasing method."))
-	TEnumAsByte<EAntiAliasingMethod> MobileAntiAliasing;
+	TEnumAsByte<EMobileAntiAliasingMethod::Type> MobileAntiAliasing;
 
 	UPROPERTY(config, EditAnywhere, Category = Mobile, meta=(
 		ConsoleVariable="r.Mobile.AllowDitheredLODTransition", DisplayName="Allow Dithered LOD Transition",
