SceneComponent bounds calculation optimizations:
- Added bComputeFastLocalBounds: when set, the component will use its current bounds transformed back into local space instead of calling CalcBounds with an identity transform.
- Added bComputeBoundsOnceDuringCook: when set, this component will cache its bounds during cooking and never recompute it for cooked builds. This is for components that are known to be static.
- Set both of theses flags automatically to all HLODs components; this greatly reduce the cost of recomputing bounds for HLODs (GT and RT).

diff --git a/Engine/Source/Developer/WorldPartitionHLODUtilities/Private/HLODBuilder.cpp b/Engine/Source/Developer/WorldPartitionHLODUtilities/Private/HLODBuilder.cpp
--- a/Engine/Source/Developer/WorldPartitionHLODUtilities/Private/HLODBuilder.cpp
+++ b/Engine/Source/Developer/WorldPartitionHLODUtilities/Private/HLODBuilder.cpp
@@ -27,6 +27,14 @@ void FHLODBuilder::Build(AWorldPartitionHLOD* InHLODActor, const UHLODLayer* InH
 	TArray<UPrimitiveComponent*> HLODPrimitives = CreateComponents(InHLODActor, InHLODLayer, SubComponents);
 	HLODPrimitives.RemoveSwap(nullptr);
 
+	for (UPrimitiveComponent* HLODPrimitive : HLODPrimitives)
+	{
+		DisableCollisions(HLODPrimitive);
+
+		HLODPrimitive->bComputeFastLocalBounds = true;
+		HLODPrimitive->bComputeBoundsOnceDuringCook = true;
+	}
+
 	if (!HLODPrimitives.IsEmpty())
 	{
 		InHLODActor->Modify();
diff --git a/Engine/Source/Developer/WorldPartitionHLODUtilities/Private/HLODBuilderInstancing.cpp b/Engine/Source/Developer/WorldPartitionHLODUtilities/Private/HLODBuilderInstancing.cpp
--- a/Engine/Source/Developer/WorldPartitionHLODUtilities/Private/HLODBuilderInstancing.cpp
+++ b/Engine/Source/Developer/WorldPartitionHLODUtilities/Private/HLODBuilderInstancing.cpp
@@ -116,9 +116,6 @@ TArray<UPrimitiveComponent*> FHLODBuilder_Instancing::CreateComponents(AWorldPar
 
 		UInstancedStaticMeshComponent* Component = ISMComponentDescriptor.CreateComponent(InHLODActor);
 		Component->SetForcedLodModel(Component->GetStaticMesh()->GetNumLODs());
-
-		DisableCollisions(Component);
-
 		Component->NumCustomDataFloats = EntryInstancingData.NumCustomDataFloats;
 		Component->AddInstances(EntryInstancingData.InstancesTransforms, /*bShouldReturnIndices*/false, /*bWorldSpace*/true);
 		Component->PerInstanceSMCustomData = MoveTemp(EntryInstancingData.InstancesCustomData);
diff --git a/Engine/Source/Developer/WorldPartitionHLODUtilities/Private/HLODBuilderMeshApproximate.cpp b/Engine/Source/Developer/WorldPartitionHLODUtilities/Private/HLODBuilderMeshApproximate.cpp
--- a/Engine/Source/Developer/WorldPartitionHLODUtilities/Private/HLODBuilderMeshApproximate.cpp
+++ b/Engine/Source/Developer/WorldPartitionHLODUtilities/Private/HLODBuilderMeshApproximate.cpp
@@ -135,7 +135,6 @@ TArray<UPrimitiveComponent*> FHLODBuilder_MeshApproximate::CreateComponents(AWor
 		{
 			UStaticMeshComponent* Component = NewObject<UStaticMeshComponent>(InHLODActor);
 			Component->SetStaticMesh(StaticMesh);
-			DisableCollisions(Component);
 
 			Components.Add(Component);
 		}
diff --git a/Engine/Source/Developer/WorldPartitionHLODUtilities/Private/HLODBuilderMeshMerge.cpp b/Engine/Source/Developer/WorldPartitionHLODUtilities/Private/HLODBuilderMeshMerge.cpp
--- a/Engine/Source/Developer/WorldPartitionHLODUtilities/Private/HLODBuilderMeshMerge.cpp
+++ b/Engine/Source/Developer/WorldPartitionHLODUtilities/Private/HLODBuilderMeshMerge.cpp
@@ -35,7 +35,6 @@ TArray<UPrimitiveComponent*> FHLODBuilder_MeshMerge::CreateComponents(AWorldPart
 			Component = NewObject<UStaticMeshComponent>(InHLODActor);
 			Component->SetStaticMesh(static_cast<UStaticMesh*>(Asset));
 			Component->SetWorldLocation(MergedActorLocation);
-			DisableCollisions(Component);
 		}
 	});
 
diff --git a/Engine/Source/Developer/WorldPartitionHLODUtilities/Private/HLODBuilderMeshSimplify.cpp b/Engine/Source/Developer/WorldPartitionHLODUtilities/Private/HLODBuilderMeshSimplify.cpp
--- a/Engine/Source/Developer/WorldPartitionHLODUtilities/Private/HLODBuilderMeshSimplify.cpp
+++ b/Engine/Source/Developer/WorldPartitionHLODUtilities/Private/HLODBuilderMeshSimplify.cpp
@@ -67,7 +67,6 @@ TArray<UPrimitiveComponent*> FHLODBuilder_MeshSimplify::CreateComponents(AWorldP
 		{
 			Component = NewObject<UStaticMeshComponent>(InHLODActor);
 			Component->SetStaticMesh(static_cast<UStaticMesh*>(Asset));
-			DisableCollisions(Component);
 		}
 	});
 
diff --git a/Engine/Source/Developer/WorldPartitionHLODUtilities/Private/WorldPartitionHLODUtilities.cpp b/Engine/Source/Developer/WorldPartitionHLODUtilities/Private/WorldPartitionHLODUtilities.cpp
--- a/Engine/Source/Developer/WorldPartitionHLODUtilities/Private/WorldPartitionHLODUtilities.cpp
+++ b/Engine/Source/Developer/WorldPartitionHLODUtilities/Private/WorldPartitionHLODUtilities.cpp
@@ -74,7 +74,7 @@ static uint32 ComputeHLODHash(AWorldPartitionHLOD* InHLODActor, const TArray<FWo
 	FArchiveCrc32 Ar;
 
 	// Base key, changing this will force a rebuild of all HLODs
-	FString HLODBaseKey = "E7DBB5E1B39947FB84E72A38074B7711";
+	FString HLODBaseKey = "E0745F6E31AC49EC89BDD6D1F6A87968";
 	Ar << HLODBaseKey;
 
 	// HLOD Layer
diff --git a/Engine/Source/Runtime/Core/Public/UObject/UE5PrivateFrostyStreamObjectVersion.h b/Engine/Source/Runtime/Core/Public/UObject/UE5PrivateFrostyStreamObjectVersion.h
--- a/Engine/Source/Runtime/Core/Public/UObject/UE5PrivateFrostyStreamObjectVersion.h
+++ b/Engine/Source/Runtime/Core/Public/UObject/UE5PrivateFrostyStreamObjectVersion.h
@@ -15,6 +15,9 @@ struct CORE_API FUE5PrivateFrostyStreamObjectVersion
 		// Added HLODBatchingPolicy member to UPrimitiveComponent, which replaces the confusing bUseMaxLODAsImposter & bBatchImpostersAsInstances.
 		HLODBatchingPolicy,
 
+		// Serialize scene components static bounds
+		SerializeSceneComponentStaticBounds,
+
 		// -----<new versions can be added above this line>-------------------------------------------------
 		VersionPlusOne,
 		LatestVersion = VersionPlusOne - 1
diff --git a/Engine/Source/Runtime/Engine/Classes/Components/SceneComponent.h b/Engine/Source/Runtime/Engine/Classes/Components/SceneComponent.h
--- a/Engine/Source/Runtime/Engine/Classes/Components/SceneComponent.h
+++ b/Engine/Source/Runtime/Engine/Classes/Components/SceneComponent.h
@@ -230,6 +230,17 @@ public:
 	UPROPERTY(EditAnywhere, BlueprintReadWrite, AdvancedDisplay, Category=Rendering)
 	uint8 bUseAttachParentBound:1;
 
+	/** If true, this component will use its current bounds transformed back into local space instead of calling CalcBounds with an identity transform. */
+	UPROPERTY()
+	uint8 bComputeFastLocalBounds : 1;
+
+	/** If true, this component will cache its bounds during cooking and never recompute it for cooked builds. This is for components that are known to be static. */
+	UPROPERTY()
+	uint8 bComputeBoundsOnceDuringCook : 1;
+
+	/** Get the current local bounds of the component */
+	FBoxSphereBounds GetLocalBounds() const;
+
 	/** Clears the skip update overlaps flag. This should be called any time a change to state would prevent the result of UpdateOverlaps. For example attachment, changing collision settings, etc... */
 	void ClearSkipUpdateOverlaps();
 
@@ -881,6 +892,7 @@ public:
 	virtual void PreNetReceive() override;
 	virtual void PostNetReceive() override;
 	virtual void PostRepNotifies() override;
+	virtual void Serialize(FArchive& Ar) override;
 #if WITH_EDITORONLY_DATA
 	static void AddReferencedObjects(UObject* InThis, FReferenceCollector& Collector);
 #endif
@@ -1010,7 +1022,7 @@ public:
 	/** Calculate the local bounds of the component. Default behavior is calling CalcBounds with an identity transform. */
 	virtual FBoxSphereBounds CalcLocalBounds() const 
 	{ 
-		return CalcBounds(FTransform::Identity);
+		return GetLocalBounds();
 	}
 
 	/**
diff --git a/Engine/Source/Runtime/Engine/Private/Components/SceneComponent.cpp b/Engine/Source/Runtime/Engine/Private/Components/SceneComponent.cpp
--- a/Engine/Source/Runtime/Engine/Private/Components/SceneComponent.cpp
+++ b/Engine/Source/Runtime/Engine/Private/Components/SceneComponent.cpp
@@ -28,6 +28,7 @@
 #include "Widgets/Notifications/SNotificationList.h"
 #include "Components/ChildActorComponent.h"
 #include "UObject/UObjectThreadContext.h"
+#include "UObject/UE5PrivateFrostyStreamObjectVersion.h"
 #include "Engine/SCS_Node.h"
 #include "EngineGlobals.h"
 #include "DeviceProfiles/DeviceProfileManager.h"
@@ -1214,9 +1215,12 @@ void USceneComponent::UpdateBounds()
 	}
 	else
 	{
-		SCOPE_CYCLE_COUNTER(STAT_ComponentCalcBounds);
 		// Calculate new bounds
-		Bounds = CalcBounds(GetComponentTransform());
+		if (!bComputeBoundsOnceDuringCook || !FPlatformProperties::RequiresCookedData())
+		{
+			SCOPE_CYCLE_COUNTER(STAT_ComponentCalcBounds);
+			Bounds = CalcBounds(GetComponentTransform());
+		}
 	}
 
 
@@ -3310,6 +3314,27 @@ void USceneComponent::PostRepNotifies()
 	}
 }
 
+void USceneComponent::Serialize(FArchive& Ar)
+{
+	Ar.UsingCustomVersion(FUE5PrivateFrostyStreamObjectVersion::GUID);
+
+	Super::Serialize(Ar);
+
+	if (bComputeBoundsOnceDuringCook)
+	{
+		if(Ar.CustomVer(FUE5PrivateFrostyStreamObjectVersion::GUID) >= FUE5PrivateFrostyStreamObjectVersion::SerializeSceneComponentStaticBounds)
+		{
+			bool bIsCooked = Ar.IsCooking();
+			Ar << bIsCooked;
+
+			if (bIsCooked)
+			{
+				Ar << Bounds;
+			}
+		}
+	}
+}
+
 void USceneComponent::GetLifetimeReplicatedProps(TArray< FLifetimeProperty > & OutLifetimeProps) const
 {
 	Super::GetLifetimeReplicatedProps(OutLifetimeProps);
@@ -3723,6 +3748,15 @@ bool FScopedMovementUpdate::SetWorldLocationAndRotation(FVector NewLocation, con
 	return false;
 }
 
+FBoxSphereBounds USceneComponent::GetLocalBounds() const
+{
+	if (bComputeFastLocalBounds)
+	{
+		return Bounds.TransformBy(ComponentToWorld.Inverse());
+	}
+	return CalcBounds(FTransform::Identity);
+}
+
 void USceneComponent::ClearSkipUpdateOverlaps()
 {
 	if (ShouldSkipUpdateOverlaps())
diff --git a/Engine/Source/Runtime/Renderer/Private/RendererScene.cpp b/Engine/Source/Runtime/Renderer/Private/RendererScene.cpp
--- a/Engine/Source/Runtime/Renderer/Private/RendererScene.cpp
+++ b/Engine/Source/Runtime/Renderer/Private/RendererScene.cpp
@@ -1216,7 +1216,7 @@ void FScene::AddPrimitive(UPrimitiveComponent* Primitive)
 		RenderMatrix,
 		Primitive->Bounds,
 		AttachmentRootPosition,
-		Primitive->CalcBounds(FTransform::Identity)
+		Primitive->GetLocalBounds()
 	};
 
 	// Help track down primitive with bad bounds way before the it gets to the Renderer
@@ -1348,7 +1348,7 @@ void FScene::UpdatePrimitiveTransform(UPrimitiveComponent* Primitive)
 			UpdateParams.WorldBounds = Primitive->Bounds;
 			UpdateParams.LocalToWorld = Primitive->GetRenderMatrix();
 			UpdateParams.AttachmentRootPosition = AttachmentRootPosition;
-			UpdateParams.LocalBounds = Primitive->CalcBounds(FTransform::Identity);
+			UpdateParams.LocalBounds = Primitive->GetLocalBounds();
 			UpdateParams.PreviousTransform = FMotionVectorSimulation::Get().GetPreviousTransform(Primitive);
 
 			// Help track down primitive with bad bounds way before the it gets to the renderer.
