Fixed missing DF shadows

diff --git a/Engine/Source/Runtime/Engine/Private/Components/DirectionalLightComponent.cpp b/Engine/Source/Runtime/Engine/Private/Components/DirectionalLightComponent.cpp
--- a/Engine/Source/Runtime/Engine/Private/Components/DirectionalLightComponent.cpp
+++ b/Engine/Source/Runtime/Engine/Private/Components/DirectionalLightComponent.cpp
@@ -677,7 +677,7 @@ private:
 
 	float GetDistanceFieldShadowDistance() const
 	{
-		if (DoesProjectSupportDistanceFields())
+		if (!DoesProjectSupportDistanceFields())
 		{
 			// Meshes must have distance fields generated
 			return 0;
