Remove the DDCShippingTextures commandline switch and just unify the DDC fill behavior with the editor behavior instead.

diff --git a/Engine/Source/Runtime/Engine/Private/TextureDerivedData.cpp b/Engine/Source/Runtime/Engine/Private/TextureDerivedData.cpp
--- a/Engine/Source/Runtime/Engine/Private/TextureDerivedData.cpp
+++ b/Engine/Source/Runtime/Engine/Private/TextureDerivedData.cpp
@@ -738,7 +738,6 @@ void FTexturePlatformData::Cache(
 	ETextureCacheFlags Flags = ETextureCacheFlags(InFlags);
 
 	static bool bForDDC = FString(FCommandLine::Get()).Contains(TEXT("Run=DerivedDataCache"));
-	static bool bDDCShippingTextures = FParse::Param(FCommandLine::Get(), TEXT("DDCShippingTextures"));
 	if (bForDDC)
 	{
 		Flags |= ETextureCacheFlags::ForDDCBuild;
@@ -748,7 +747,7 @@ void FTexturePlatformData::Cache(
 	bool bAsync = EnumHasAnyFlags(Flags, ETextureCacheFlags::Async);
 	GetTextureDerivedDataKey(InTexture, InSettingsPerLayer, DerivedDataKey);
 
-	if ((!bForDDC || bDDCShippingTextures) && !bForceRebuild && InSettingsPerLayer[0].bHasEditorOnlyData)
+	if (!bForceRebuild && InSettingsPerLayer[0].bHasEditorOnlyData)
 	{
 		// If we are running in the Editor, generate a shipping derived data key, which can be used to load
 		// a texture cooked for the shipping build from the cache, if available.
