Fixed an issue with sync mirroring using blendspaces

diff --git a/Engine/Source/Runtime/Engine/Classes/Animation/BlendSpace.h b/Engine/Source/Runtime/Engine/Classes/Animation/BlendSpace.h
--- a/Engine/Source/Runtime/Engine/Classes/Animation/BlendSpace.h
+++ b/Engine/Source/Runtime/Engine/Classes/Animation/BlendSpace.h
@@ -662,7 +662,7 @@ protected:
 	/** Initialize Per Bone Blend **/
 	void InitializePerBoneBlend();
 
-	void TickFollowerSamples(TArray<FBlendSampleData> &SampleDataList, const int32 HighestWeightIndex, FAnimAssetTickContext &Context, bool bResetMarkerDataOnFollowers) const;
+	void TickFollowerSamples(TArray<FBlendSampleData> &SampleDataList, const int32 HighestWeightIndex, FAnimAssetTickContext &Context, bool bResetMarkerDataOnFollowers, const UMirrorDataTable* MirrorDataTable = nullptr) const;
 
 	/** Utility function to calculate animation length from sample data list **/
 	float GetAnimationLengthFromSampleData(const TArray<FBlendSampleData> & SampleDataList) const;
diff --git a/Engine/Source/Runtime/Engine/Private/Animation/BlendSpace.cpp b/Engine/Source/Runtime/Engine/Private/Animation/BlendSpace.cpp
--- a/Engine/Source/Runtime/Engine/Private/Animation/BlendSpace.cpp
+++ b/Engine/Source/Runtime/Engine/Private/Animation/BlendSpace.cpp
@@ -445,16 +445,16 @@ void UBlendSpace::TickAssetPlayer(FAnimTickRecord& Instance, struct FAnimNotifyQ
 					}
 					else if (!SampleDataItem.MarkerTickRecord.IsValid(Instance.bLooping) && Context.MarkerTickContext.GetMarkerSyncStartPosition().IsValid())
 					{
-						Sample.Animation->GetMarkerIndicesForPosition(Context.MarkerTickContext.GetMarkerSyncStartPosition(), true, SampleDataItem.MarkerTickRecord.PreviousMarker, SampleDataItem.MarkerTickRecord.NextMarker, SampleDataItem.Time, nullptr);
+						Sample.Animation->GetMarkerIndicesForPosition(Context.MarkerTickContext.GetMarkerSyncStartPosition(), true, SampleDataItem.MarkerTickRecord.PreviousMarker, SampleDataItem.MarkerTickRecord.NextMarker, SampleDataItem.Time, Instance.MirrorDataTable);
 					}
 
 					const float NewDeltaTime = Context.GetDeltaTime() * Instance.PlayRateMultiplier * Sample.RateScale * Sample.Animation->RateScale;
 					if (!FMath::IsNearlyZero(NewDeltaTime))
 					{
 						Context.SetLeaderDelta(NewDeltaTime);
-						Sample.Animation->TickByMarkerAsLeader(SampleDataItem.MarkerTickRecord, Context.MarkerTickContext, SampleDataItem.Time, SampleDataItem.PreviousTime, NewDeltaTime, Instance.bLooping, nullptr);
+						Sample.Animation->TickByMarkerAsLeader(SampleDataItem.MarkerTickRecord, Context.MarkerTickContext, SampleDataItem.Time, SampleDataItem.PreviousTime, NewDeltaTime, Instance.bLooping, Instance.MirrorDataTable);
 						check(!Instance.bLooping || Context.MarkerTickContext.IsMarkerSyncStartValid());
-						TickFollowerSamples(SampleDataList, HighestMarkerSyncWeightIndex, Context, bResetMarkerDataOnFollowers);
+						TickFollowerSamples(SampleDataList, HighestMarkerSyncWeightIndex, Context, bResetMarkerDataOnFollowers, Instance.MirrorDataTable);
 					}
 					NormalizedCurrentTime = SampleDataItem.Time / Sample.Animation->GetPlayLength();
 					*Instance.MarkerTickRecord = SampleDataItem.MarkerTickRecord;
@@ -490,7 +490,7 @@ void UBlendSpace::TickAssetPlayer(FAnimTickRecord& Instance, struct FAnimNotifyQ
 							SampleDataItem.Time = NormalizedCurrentTime * Sample.Animation->GetPlayLength();
 						}
 
-						TickFollowerSamples(SampleDataList, -1, Context, false);
+						TickFollowerSamples(SampleDataList, -1, Context, false, Instance.MirrorDataTable);
 					}
 					*Instance.MarkerTickRecord = SampleDataItem.MarkerTickRecord;
 					NormalizedCurrentTime = SampleDataItem.Time / Sample.Animation->GetPlayLength();
@@ -1491,7 +1491,7 @@ void UBlendSpace::InitializePerBoneBlend()
 	}
 }
 
-void UBlendSpace::TickFollowerSamples(TArray<FBlendSampleData>& SampleDataList, const int32 HighestWeightIndex, FAnimAssetTickContext& Context, bool bResetMarkerDataOnFollowers) const
+void UBlendSpace::TickFollowerSamples(TArray<FBlendSampleData>& SampleDataList, const int32 HighestWeightIndex, FAnimAssetTickContext& Context, bool bResetMarkerDataOnFollowers, const UMirrorDataTable* MirrorDataTable) const
 {
 	for (int32 SampleIndex = 0; SampleIndex < SampleDataList.Num(); ++SampleIndex)
 	{
@@ -1506,7 +1506,7 @@ void UBlendSpace::TickFollowerSamples(TArray<FBlendSampleData>& SampleDataList,
 
 			if (Sample.Animation->AuthoredSyncMarkers.Num() > 0) // Update followers who can do marker sync, others will be handled later in TickAssetPlayer
 			{
-				Sample.Animation->TickByMarkerAsFollower(SampleDataItem.MarkerTickRecord, Context.MarkerTickContext, SampleDataItem.Time, SampleDataItem.PreviousTime, Context.GetLeaderDelta(), true, nullptr);
+				Sample.Animation->TickByMarkerAsFollower(SampleDataItem.MarkerTickRecord, Context.MarkerTickContext, SampleDataItem.Time, SampleDataItem.PreviousTime, Context.GetLeaderDelta(), true, MirrorDataTable);
 			}
 		}
 	}
