Fix assert in FSkeletalMeshObjectGPUSkin::GetSkinVertexFactory where SkinCacheEntryForRayTracing is nullptr

diff --git a/Engine/Source/Runtime/Engine/Private/SkeletalMesh.cpp b/Engine/Source/Runtime/Engine/Private/SkeletalMesh.cpp
--- a/Engine/Source/Runtime/Engine/Private/SkeletalMesh.cpp
+++ b/Engine/Source/Runtime/Engine/Private/SkeletalMesh.cpp
@@ -6210,7 +6210,7 @@ void FSkeletalMeshSceneProxy::GetDynamicRayTracingInstances(FRayTracingMaterialG
 	}
 
 	// GetRayTracingGeometry()->IsInitialized() is checked as a workaround for UE-92634. FSkeletalMeshSceneProxy's resources may have already been released, but proxy has not removed yet)
-	if (MeshObject->GetRayTracingGeometry() && MeshObject->GetRayTracingGeometry()->IsInitialized() )
+	if (MeshObject->GetRayTracingGeometry() && MeshObject->GetRayTracingGeometry()->IsInitialized() && MeshObject->SkinCacheEntryForRayTracing)
 	{
 		// #dxr: the only case where RayTracingGeometryRHI is invalid is the very first frame - if that's not the case we have a bug somewhere else
 		if (MeshObject->GetRayTracingGeometry()->RayTracingGeometryRHI.IsValid())
