Added DDC events for Resource Size, BuildCount and MB

diff --git a/Engine/Source/Runtime/Engine/Private/Analytics/StudioAnalytics.cpp b/Engine/Source/Runtime/Engine/Private/Analytics/StudioAnalytics.cpp
--- a/Engine/Source/Runtime/Engine/Private/Analytics/StudioAnalytics.cpp
+++ b/Engine/Source/Runtime/Engine/Private/Analytics/StudioAnalytics.cpp
@@ -166,23 +166,153 @@ void FStudioAnalytics::FireEvent_Loading(const FString& LoadingName, double Seco
 		FStudioAnalytics::GetProvider().RecordEvent(TEXT("Performance.Loading"), Attributes);
 
 #if ENABLE_COOK_STATS
-		// Sends each DDC stat to the studio analytics system.
-		auto SendCookStatsToAnalytics = [&Attributes](const FString& StatName, const TArray<FCookStatsManager::StringKeyValue>& StatAttributes)
+		//// Sends each DDC stat to the studio analytics system.
+		//auto SendDDCStatsToAnalytics = [&Attributes](const FString& StatName, const TArray<FCookStatsManager::StringKeyValue>& StatAttributes)
+		//{
+		//	// We'reonly interested in DDC Summary stats
+		//	if (StatName.Contains("DDC.Summary"))
+		//	{
+		//		for (const auto& Attr : StatAttributes)
+		//		{
+		//			FString FormattedAttrName = StatName + "." + Attr.Key;
+		//			Attributes.Emplace(FormattedAttrName, Attr.Value);
+		//		}		
+		//	}
+		//};
+
+		//// Grab the DDC stats
+		//FCookStatsManager::LogCookStats(SendDDCStatsToAnalytics);
+
+		/** Used for custom logging of DDC Resource usage stats. */
+		struct FDDCResourceUsageStat
 		{
-			// We'reonly interested in DDC Summary stats
-			if (StatName.Contains("DDC.Summary"))
+		public:
+			FDDCResourceUsageStat(FString InAssetType, double InTotalTimeSec, bool bIsGameThreadTime, double InSizeMB, int64 InAssetsBuilt) : AssetType(MoveTemp(InAssetType)), TotalTimeSec(InTotalTimeSec), GameThreadTimeSec(bIsGameThreadTime ? InTotalTimeSec : 0.0), SizeMB(InSizeMB), AssetsBuilt(InAssetsBuilt) {}
+			void Accumulate(const FDDCResourceUsageStat& OtherStat)
 			{
-				for (const auto& Attr : StatAttributes)
+				TotalTimeSec += OtherStat.TotalTimeSec;
+				GameThreadTimeSec += OtherStat.GameThreadTimeSec;
+				SizeMB += OtherStat.SizeMB;
+				AssetsBuilt += OtherStat.AssetsBuilt;
+			}
+			FString AssetType;
+			double TotalTimeSec;
+			double GameThreadTimeSec;
+			double SizeMB;
+			int64 AssetsBuilt;
+		};
+
+		/** Used for custom TSet comparison of DDC Resource usage stats. */
+		struct FDDCResourceUsageStatKeyFuncs : BaseKeyFuncs<FDDCResourceUsageStat, FString, false>
+		{
+			static const FString& GetSetKey(const FDDCResourceUsageStat& Element) { return Element.AssetType; }
+			static bool Matches(const FString& A, const FString& B) { return A == B; }
+			static uint32 GetKeyHash(const FString& Key) { return GetTypeHash(Key); }
+		};
+
+		// instead of printing the usage stats generically, we capture them so we can log a subset of them in an easy-to-read way.
+		TSet<FDDCResourceUsageStat, FDDCResourceUsageStatKeyFuncs> DDCResourceUsageStats;
+		TArray<FCookStatsManager::StringKeyValue> DDCSummaryStats;
+
+		int64 TotalAssetsBuilt = 0;
+		double TotalAssetTimeSec = 0.0;
+		double TotalAssetSizeMB = 0.0;
+		
+		/** this functor will take a collected cooker stat and log it out using some custom formatting based on known stats that are collected.. */
+		auto LogStatsFunc = [&DDCResourceUsageStats, &DDCSummaryStats, &TotalAssetsBuilt, &TotalAssetTimeSec, &TotalAssetSizeMB]
+		(const FString& StatName, const TArray<FCookStatsManager::StringKeyValue>& StatAttributes)
+		{
+			if (StatName.EndsWith(TEXT(".Usage"), ESearchCase::IgnoreCase))
+			{
+				// Anything that ends in .Usage is assumed to be an instance of FCookStats.FDDCResourceUsageStats. We'll log that using custom formatting.
+				FString AssetType = StatName;
+				AssetType.RemoveFromEnd(TEXT(".Usage"), ESearchCase::IgnoreCase);
+				// See if the asset has a subtype (found via the "Node" parameter")
+				const FCookStatsManager::StringKeyValue* AssetSubType = StatAttributes.FindByPredicate([](const FCookStatsManager::StringKeyValue& Item) { return Item.Key == TEXT("Node"); });
+				if (AssetSubType && AssetSubType->Value.Len() > 0)
+				{
+					AssetType += FString::Printf(TEXT(" (%s)"), *AssetSubType->Value);
+				}
+				// Pull the Time and Size attributes and AddOrAccumulate them into the set of stats. Ugly string/container manipulation code courtesy of UE4/C++.
+				const FCookStatsManager::StringKeyValue* AssetTimeSecAttr = StatAttributes.FindByPredicate([](const FCookStatsManager::StringKeyValue& Item) { return Item.Key == TEXT("TimeSec"); });
+				double AssetTimeSec = 0.0;
+				if (AssetTimeSecAttr)
+				{
+					LexFromString(AssetTimeSec, *AssetTimeSecAttr->Value);
+				}
+				const FCookStatsManager::StringKeyValue* AssetSizeMBAttr = StatAttributes.FindByPredicate([](const FCookStatsManager::StringKeyValue& Item) { return Item.Key == TEXT("MB"); });
+				double AssetSizeMB = 0.0;
+				if (AssetSizeMBAttr)
 				{
-					FString FormattedAttrName = StatName + "." + Attr.Key;
-					Attributes.Emplace(FormattedAttrName, Attr.Value);
-				}		
+					LexFromString(AssetSizeMB, *AssetSizeMBAttr->Value);
+				}
+				const FCookStatsManager::StringKeyValue* ThreadNameAttr = StatAttributes.FindByPredicate([](const FCookStatsManager::StringKeyValue& Item) { return Item.Key == TEXT("ThreadName"); });
+				bool bIsGameThreadTime = ThreadNameAttr != nullptr && ThreadNameAttr->Value == TEXT("GameThread");
+
+				const FCookStatsManager::StringKeyValue* HitOrMissAttr = StatAttributes.FindByPredicate([](const FCookStatsManager::StringKeyValue& Item) { return Item.Key == TEXT("HitOrMiss"); });
+				bool bWasMiss = HitOrMissAttr != nullptr && HitOrMissAttr->Value == TEXT("Miss");
+				int64 AssetsBuilt = 0;
+				if (bWasMiss)
+				{
+					const FCookStatsManager::StringKeyValue* CountAttr = StatAttributes.FindByPredicate([](const FCookStatsManager::StringKeyValue& Item) { return Item.Key == TEXT("Count"); });
+					if (CountAttr)
+					{
+						LexFromString(AssetsBuilt, *CountAttr->Value);
+					}
+				}
+
+				TotalAssetsBuilt++;
+				TotalAssetTimeSec += AssetTimeSec;
+				TotalAssetSizeMB += AssetSizeMB;
+
+				FDDCResourceUsageStat Stat(AssetType, AssetTimeSec, bIsGameThreadTime, AssetSizeMB, AssetsBuilt);
+				FDDCResourceUsageStat* ExistingStat = DDCResourceUsageStats.Find(Stat.AssetType);
+				if (ExistingStat)
+				{
+					ExistingStat->Accumulate(Stat);
+				}
+				else
+				{
+					DDCResourceUsageStats.Add(Stat);
+				}
+			}
+			else if (StatName == TEXT("DDC.Summary"))
+			{
+				DDCSummaryStats.Append(StatAttributes);
 			}
 		};
 
 		// Grab the DDC stats
-		FCookStatsManager::LogCookStats(SendCookStatsToAnalytics);
+		FCookStatsManager::LogCookStats(LogStatsFunc);
+
+		for (const FDDCResourceUsageStat& Stat : DDCResourceUsageStats)
+		{
+			{
+				FString AttrName = "DDC.Resource." + Stat.AssetType + ".Built";
+				Attributes.Emplace(AttrName, Stat.AssetsBuilt);
+			}
+			
+			{
+				FString AttrName = "DDC.Resource." + Stat.AssetType + ".TimeSec";
+				Attributes.Emplace(AttrName, Stat.TotalTimeSec);
+			}
 
+			{
+				FString AttrName = "DDC.Resource." + Stat.AssetType + ".SizeMB";
+				Attributes.Emplace(AttrName, Stat.SizeMB);
+			}
+		}
+
+		Attributes.Emplace(TEXT("DDC.Resource.TotalAssetsBuilt"), TotalAssetsBuilt);
+		Attributes.Emplace(TEXT("DDC.Resource.TotalAssetTimeSec"), TotalAssetTimeSec);
+		Attributes.Emplace(TEXT("DDC.Resource.TotalAssetSizeMB"), TotalAssetSizeMB);
+		
+		for (const FCookStatsManager::StringKeyValue& Attr : DDCSummaryStats)
+		{
+			FString FormattedAttrName = "DDC.Summary." + Attr.Key;
+			Attributes.Emplace(FormattedAttrName, Attr.Value);
+		}
+		
 		FStudioAnalytics::GetProvider().RecordEvent(TEXT("Core.Loading"), Attributes);
 #endif
 	}
