Analytics attribute name change

diff --git a/Engine/Source/Runtime/Engine/Private/Analytics/StudioAnalytics.cpp b/Engine/Source/Runtime/Engine/Private/Analytics/StudioAnalytics.cpp
--- a/Engine/Source/Runtime/Engine/Private/Analytics/StudioAnalytics.cpp
+++ b/Engine/Source/Runtime/Engine/Private/Analytics/StudioAnalytics.cpp
@@ -303,9 +303,9 @@ void FStudioAnalytics::FireEvent_Loading(const FString& LoadingName, double Seco
 			}
 		}
 
-		Attributes.Emplace(TEXT("DDC.Resource.TotalAssetsBuilt"), TotalAssetsBuilt);
-		Attributes.Emplace(TEXT("DDC.Resource.TotalAssetTimeSec"), TotalAssetTimeSec);
-		Attributes.Emplace(TEXT("DDC.Resource.TotalAssetSizeMB"), TotalAssetSizeMB);
+		Attributes.Emplace(TEXT("DDC.Summary.TotalAssetsBuilt"), TotalAssetsBuilt);
+		Attributes.Emplace(TEXT("DDC.Summary.TotalAssetTimeSec"), TotalAssetTimeSec);
+		Attributes.Emplace(TEXT("DDC.Summary.TotalAssetSizeMB"), TotalAssetSizeMB);
 		
 		for (const FCookStatsManager::StringKeyValue& Attr : DDCSummaryStats)
 		{
