Implements FilmGrain intensity controls for shadows, midtones and highlights

diff --git a/Engine/Shaders/Private/PostProcessTonemap.usf b/Engine/Shaders/Private/PostProcessTonemap.usf
--- a/Engine/Shaders/Private/PostProcessTonemap.usf
+++ b/Engine/Shaders/Private/PostProcessTonemap.usf
@@ -56,10 +56,17 @@ float4 ColorScale1;
 // from the postprocess settings, x:VignetteIntensity, y:SharpenDiv6
 float4 TonemapperParams;
 
-// Film grain
+// Grain quantization
 float3 GrainRandomFull;
+
+// Film grain
 float FilmGrainMul;
 float FilmGrainAdd;
+float FilmGrainIntensityShadows;
+float FilmGrainIntensityMidtones;
+float FilmGrainIntensityHighlights;
+float FilmGrainShadowsMax;
+float FilmGrainHighlightsMin;
 Texture2D<half3> FilmGrainTexture;
 SamplerState FilmGrainSampler;
 float4 ScreenPosToFilmGrainTextureUV;
@@ -111,6 +118,32 @@ half3 LinearToPostTonemapSpace(half3 lin)
 #endif
 }
 
+
+// Nuke-style Color Correct
+float ComputeFilmGrainIntensity(float3 LinearColorRGBGamut)
+{
+	const float3x3 sRGB_2_AP1 = mul( XYZ_2_AP1_MAT, mul( D65_2_D60_CAT, sRGB_2_XYZ_MAT ) );
+	const float3x3 AP1_2_sRGB = mul( XYZ_2_sRGB_MAT, mul( D60_2_D65_CAT, AP1_2_XYZ_MAT ) );
+
+	float3 ColorAP1 = mul(sRGB_2_AP1, LinearColorRGBGamut);
+	
+	//float Luma = dot(ColorAP1, AP1_RGB2Y);
+	float Luma = dot(LinearColorRGBGamut,  mul(AP1_2_sRGB, AP1_RGB2Y));
+
+	float CCWeightShadows = 1 - smoothstep(0, FilmGrainShadowsMax, Luma);
+	float CCWeightHighlights = smoothstep(FilmGrainHighlightsMin, 1, Luma);
+	float CCWeightMidtones = 1 - CCWeightShadows - CCWeightHighlights;
+	
+	// Blend Shadow, Midtone and Highlight CCs
+	float FilmGrainIntensity = (
+		FilmGrainIntensityShadows * CCWeightShadows +
+		FilmGrainIntensityMidtones * CCWeightMidtones +
+		FilmGrainIntensityHighlights * CCWeightHighlights);
+	
+	return FilmGrainIntensity;
+}
+
+
 // LUT for color grading
 #if USE_VOLUME_LUT == 1
 	Texture3D		ColorGradingLUT;
@@ -564,12 +597,14 @@ float4 TonemapCommonPS(
 
 		#if USE_FILM_GRAIN
 		{
+			half FilmGrainIntensity = ComputeFilmGrainIntensity(LinearColor);
+
 			float2 FilmGrainUV = ApplyScreenTransform(ScreenPos, ScreenPosToFilmGrainTextureUV);
 
 			half3 RawGrain = FilmGrainTexture.SampleLevel(FilmGrainSampler, FilmGrainUV, 0);
 			half3 FinalGrain = RawGrain * half(FilmGrainMul) + half(FilmGrainAdd);
 
-			LinearColor.rgb *= FinalGrain;
+			LinearColor.rgb *= lerp(1.0, FinalGrain, FilmGrainIntensity);
 		}
 		#endif
 
diff --git a/Engine/Source/Runtime/Engine/Classes/Engine/Scene.h b/Engine/Source/Runtime/Engine/Classes/Engine/Scene.h
--- a/Engine/Source/Runtime/Engine/Classes/Engine/Scene.h
+++ b/Engine/Source/Runtime/Engine/Classes/Engine/Scene.h
@@ -936,6 +936,21 @@ struct FPostProcessSettings
 	UPROPERTY(EditAnywhere, BlueprintReadWrite, Category=Overrides, meta=(PinHiddenByDefault, InlineEditConditionToggle))
 	uint8 bOverride_FilmGrainIntensity:1;
 
+	UPROPERTY(EditAnywhere, BlueprintReadWrite, Category = Overrides, meta = (PinHiddenByDefault, InlineEditConditionToggle))
+	uint8 bOverride_FilmGrainIntensityShadows : 1;
+
+	UPROPERTY(EditAnywhere, BlueprintReadWrite, Category = Overrides, meta = (PinHiddenByDefault, InlineEditConditionToggle))
+	uint8 bOverride_FilmGrainIntensityMidtones : 1;
+
+	UPROPERTY(EditAnywhere, BlueprintReadWrite, Category = Overrides, meta = (PinHiddenByDefault, InlineEditConditionToggle))
+	uint8 bOverride_FilmGrainIntensityHighlights : 1;
+
+	UPROPERTY(EditAnywhere, BlueprintReadWrite, Category = Overrides, meta = (PinHiddenByDefault, InlineEditConditionToggle))
+	uint8 bOverride_FilmGrainShadowsMax : 1;
+
+	UPROPERTY(EditAnywhere, BlueprintReadWrite, Category = Overrides, meta = (PinHiddenByDefault, InlineEditConditionToggle))
+	uint8 bOverride_FilmGrainHighlightsMin : 1;
+
 	UPROPERTY(EditAnywhere, BlueprintReadWrite, Category=Overrides, meta=(PinHiddenByDefault, InlineEditConditionToggle))
 	uint8 bOverride_FilmGrainTexelSize:1;
 
@@ -1712,23 +1727,38 @@ struct FPostProcessSettings
 	float GrainIntensity_DEPRECATED;
 
 	/** 0..1 Film grain intensity to apply. LinearSceneColor *= lerp(1.0, DecodedFilmGrainTexture, FilmGrainIntensity) */
-	UPROPERTY(interp, BlueprintReadWrite, Category = "Film|Grain", meta = (UIMin = "0.0", UIMax = "1.0", editcondition = "bOverride_FilmGrainIntensity"))
+	UPROPERTY(interp, BlueprintReadWrite, Category = "Film Grain", meta = (UIMin = "0.0", UIMax = "1.0", editcondition = "bOverride_FilmGrainIntensity"))
 	float FilmGrainIntensity;
 
+	UPROPERTY(interp, BlueprintReadWrite, Category = "Film Grain", meta = (UIMin = "0.0", UIMax = "1.0", editcondition = "bOverride_FilmGrainIntensityShadows"))
+	float FilmGrainIntensityShadows;
+
+	UPROPERTY(interp, BlueprintReadWrite, Category = "Film Grain", meta = (UIMin = "0.0", UIMax = "1.0", editcondition = "bOverride_FilmGrainIntensityMidtones"))
+	float FilmGrainIntensityMidtones;
+
+	UPROPERTY(interp, BlueprintReadWrite, Category = "Film Grain", meta = (UIMin = "0.0", UIMax = "1.0", editcondition = "bOverride_FilmGrainIntensityHighlights"))
+	float FilmGrainIntensityHighlights;
+
+	UPROPERTY(interp, BlueprintReadWrite, Category = "Film Grain", meta = (UIMin = "0.0", UIMax = "1.0", editcondition = "bOverride_FilmGrainShadowsMax"))
+	float FilmGrainShadowsMax;
+	
+	UPROPERTY(interp, BlueprintReadWrite, Category = "Film Grain", meta = (UIMin = "0.0", UIMax = "1.0", editcondition = "bOverride_FilmGrainHighlightsMin"))
+	float FilmGrainHighlightsMin;
+
 	/** Size of texel of FilmGrainTexture on screen. */
-	UPROPERTY(interp, BlueprintReadWrite, Category = "Film|Grain", meta = (UIMin = "0.0", UIMax = "4.0", editcondition = "bOverride_FilmGrainTexelSize"))
+	UPROPERTY(interp, BlueprintReadWrite, Category = "Film Grain", meta = (UIMin = "0.0", UIMax = "4.0", editcondition = "bOverride_FilmGrainTexelSize"))
 	float FilmGrainTexelSize;
 
 	/** Defines film grain texture to use. */
-	UPROPERTY(EditAnywhere, BlueprintReadWrite, Category = "Film|Grain", meta = (editcondition = "bOverride_FilmGrainTexture"))
+	UPROPERTY(EditAnywhere, BlueprintReadWrite, Category = "Film Grain", meta = (editcondition = "bOverride_FilmGrainTexture"))
 	TObjectPtr<class UTexture2D> FilmGrainTexture;
 
 	/** Multiply the FilmGrainTexture. DecodedFilmGrainTexture = FilmGrainTexture * FilmGrainDecodeMultiply + FilmGrainDecodeAdd */
-	UPROPERTY(EditAnywhere, BlueprintReadWrite, Category = "Film|Grain", meta = (editcondition = "bOverride_FilmGrainTexture"))
+	UPROPERTY(EditAnywhere, BlueprintReadWrite, Category = "Film Grain", meta = (editcondition = "bOverride_FilmGrainTexture"))
 	float FilmGrainDecodeMultiply;
 
 	/** Add to the FilmGrainTexture. DecodedFilmGrainTexture = FilmGrainTexture * FilmGrainDecodeMultiply + FilmGrainDecodeAdd */
-	UPROPERTY(EditAnywhere, BlueprintReadWrite, Category = "Film|Grain", meta = (editcondition = "bOverride_FilmGrainTexture"))
+	UPROPERTY(EditAnywhere, BlueprintReadWrite, Category = "Film Grain", meta = (editcondition = "bOverride_FilmGrainTexture"))
 	float FilmGrainDecodeAdd;
 
 	/** 0..1 0=off/no ambient occlusion .. 1=strong ambient occlusion, defines how much it affects the non direct lighting after base pass */
diff --git a/Engine/Source/Runtime/Engine/Private/Scene.cpp b/Engine/Source/Runtime/Engine/Private/Scene.cpp
--- a/Engine/Source/Runtime/Engine/Private/Scene.cpp
+++ b/Engine/Source/Runtime/Engine/Private/Scene.cpp
@@ -351,6 +351,11 @@ static void DoPostProcessSettingsSanityCheck()
 		RenameMap.Add(TEXT("FilmShadowTintAmount"), TEXT(""));
 		// Film Grain are ignored
 		RenameMap.Add(TEXT("FilmGrainIntensity"), TEXT(""));
+		RenameMap.Add(TEXT("FilmGrainIntensityShadows"), TEXT(""));
+		RenameMap.Add(TEXT("FilmGrainIntensityMidtones"), TEXT(""));
+		RenameMap.Add(TEXT("FilmGrainIntensityHighlights"), TEXT(""));
+		RenameMap.Add(TEXT("FilmGrainShadowsMax"), TEXT(""));
+		RenameMap.Add(TEXT("FilmGrainHighlightsMin"), TEXT(""));
 		RenameMap.Add(TEXT("FilmGrainTexelSize"), TEXT(""));
 		RenameMap.Add(TEXT("FilmGrainTexture"), TEXT(""));
 		RenameMap.Add(TEXT("FilmGrainDecodeMultiply"), TEXT(""));
@@ -501,10 +506,18 @@ FPostProcessSettings::FPostProcessSettings()
 	VignetteIntensity = 0.4f;
 	GrainIntensity_DEPRECATED = 0.0f;
 	GrainJitter_DEPRECATED = 0.0f;
+
+	// Film Grain
 	FilmGrainIntensity = 0.0f;
+	FilmGrainIntensityShadows = 1.0f;
+	FilmGrainIntensityMidtones = 1.0f;
+	FilmGrainIntensityHighlights = 1.0f;
+	FilmGrainShadowsMax = 0.09f;
+	FilmGrainHighlightsMin = 0.5f;
 	FilmGrainTexelSize = 1.0f;
 	FilmGrainDecodeMultiply = 1.0f;
 	FilmGrainDecodeAdd = 0.0f;
+
 	// next value might get overwritten by r.DefaultFeature.AmbientOcclusion
 	AmbientOcclusionIntensity = .5f;
 	// next value might get overwritten by r.DefaultFeature.AmbientOcclusionStaticFraction
@@ -707,6 +720,11 @@ FPostProcessSettings::FPostProcessSettings(const FPostProcessSettings& Settings)
 	, bOverride_GrainIntensity_DEPRECATED(Settings.bOverride_GrainIntensity_DEPRECATED)
 	, bOverride_GrainJitter_DEPRECATED(Settings.bOverride_GrainJitter_DEPRECATED)
 	, bOverride_FilmGrainIntensity(Settings.bOverride_FilmGrainIntensity)
+	, bOverride_FilmGrainIntensityShadows(Settings.bOverride_FilmGrainIntensityShadows)
+	, bOverride_FilmGrainIntensityMidtones(Settings.bOverride_FilmGrainIntensityMidtones)
+	, bOverride_FilmGrainIntensityHighlights(Settings.bOverride_FilmGrainIntensityHighlights)
+	, bOverride_FilmGrainShadowsMax(Settings.bOverride_FilmGrainShadowsMax)
+	, bOverride_FilmGrainHighlightsMin(Settings.bOverride_FilmGrainHighlightsMin)
 	, bOverride_FilmGrainTexelSize(Settings.bOverride_FilmGrainTexelSize)
 	, bOverride_FilmGrainTexture(Settings.bOverride_FilmGrainTexture)
 	, bOverride_AmbientOcclusionIntensity(Settings.bOverride_AmbientOcclusionIntensity)
@@ -902,6 +920,11 @@ FPostProcessSettings::FPostProcessSettings(const FPostProcessSettings& Settings)
 	, GrainJitter_DEPRECATED(Settings.GrainJitter_DEPRECATED)
 	, GrainIntensity_DEPRECATED(Settings.GrainIntensity_DEPRECATED)
 	, FilmGrainIntensity(Settings.FilmGrainIntensity)
+	, FilmGrainIntensityShadows(Settings.FilmGrainIntensityShadows)
+	, FilmGrainIntensityMidtones(Settings.FilmGrainIntensityMidtones)
+	, FilmGrainIntensityHighlights(Settings.FilmGrainIntensityHighlights)
+	, FilmGrainShadowsMax(Settings.FilmGrainShadowsMax)
+	, FilmGrainHighlightsMin(Settings.FilmGrainHighlightsMin)
 	, FilmGrainTexelSize(Settings.FilmGrainTexelSize)
 	, FilmGrainTexture(Settings.FilmGrainTexture)
 	, FilmGrainDecodeMultiply(Settings.FilmGrainDecodeMultiply)
diff --git a/Engine/Source/Runtime/Engine/Private/SceneView.cpp b/Engine/Source/Runtime/Engine/Private/SceneView.cpp
--- a/Engine/Source/Runtime/Engine/Private/SceneView.cpp
+++ b/Engine/Source/Runtime/Engine/Private/SceneView.cpp
@@ -1437,6 +1437,11 @@ void FSceneView::OverridePostProcessSettings(const FPostProcessSettings& Src, fl
 		LERP_PP(LensFlareThreshold);
 		LERP_PP(VignetteIntensity);
 		LERP_PP(FilmGrainIntensity);
+		LERP_PP(FilmGrainIntensityShadows);
+		LERP_PP(FilmGrainIntensityMidtones);
+		LERP_PP(FilmGrainIntensityHighlights);
+		LERP_PP(FilmGrainShadowsMax);
+		LERP_PP(FilmGrainHighlightsMin);
 		LERP_PP(FilmGrainTexelSize);
 		LERP_PP(AmbientOcclusionIntensity);
 		LERP_PP(AmbientOcclusionStaticFraction);
diff --git a/Engine/Source/Runtime/Renderer/Private/PostProcess/PostProcessTonemap.cpp b/Engine/Source/Runtime/Renderer/Private/PostProcess/PostProcessTonemap.cpp
--- a/Engine/Source/Runtime/Renderer/Private/PostProcess/PostProcessTonemap.cpp
+++ b/Engine/Source/Runtime/Renderer/Private/PostProcess/PostProcessTonemap.cpp
@@ -314,6 +314,11 @@ BEGIN_SHADER_PARAMETER_STRUCT(FFilmGrainParameters, )
 	SHADER_PARAMETER(FVector3f, GrainRandomFull)
 	SHADER_PARAMETER(float, FilmGrainMul)
 	SHADER_PARAMETER(float, FilmGrainAdd)
+	SHADER_PARAMETER(float, FilmGrainIntensityShadows)
+	SHADER_PARAMETER(float, FilmGrainIntensityMidtones)
+	SHADER_PARAMETER(float, FilmGrainIntensityHighlights)
+	SHADER_PARAMETER(float, FilmGrainShadowsMax)
+	SHADER_PARAMETER(float, FilmGrainHighlightsMin)
 	SHADER_PARAMETER_TEXTURE(Texture2D, FilmGrainTexture)
 	SHADER_PARAMETER_SAMPLER(SamplerState, FilmGrainSampler)
 	SHADER_PARAMETER(FScreenTransform, ScreenPosToFilmGrainTextureUV)
@@ -564,9 +569,15 @@ FScreenPassTexture AddTonemapPass(FRDGBuilder& GraphBuilder, const FViewInfo& Vi
 	{
 		check(View.FilmGrainTexture);
 
-		// (FilmGrainTexture * FilmGrainDecodeMultiply + FilmGrainDecodeAdd) * View.FinalPostProcessSettings.FilmGrainIntensity + 1.0f - View.FinalPostProcessSettings.FilmGrainIntensity
-		CommonParameters.FilmGrain.FilmGrainMul = View.FinalPostProcessSettings.FilmGrainDecodeMultiply * View.FinalPostProcessSettings.FilmGrainIntensity;
-		CommonParameters.FilmGrain.FilmGrainAdd = View.FinalPostProcessSettings.FilmGrainDecodeAdd * View.FinalPostProcessSettings.FilmGrainIntensity + 1.0f - View.FinalPostProcessSettings.FilmGrainIntensity;
+		// (FilmGrainTexture * FilmGrainDecodeMultiply + FilmGrainDecodeAdd)
+		CommonParameters.FilmGrain.FilmGrainMul = View.FinalPostProcessSettings.FilmGrainDecodeMultiply;
+		CommonParameters.FilmGrain.FilmGrainAdd = View.FinalPostProcessSettings.FilmGrainDecodeAdd;
+		CommonParameters.FilmGrain.FilmGrainIntensityShadows    = View.FinalPostProcessSettings.FilmGrainIntensity * View.FinalPostProcessSettings.FilmGrainIntensityShadows;
+		CommonParameters.FilmGrain.FilmGrainIntensityMidtones   = View.FinalPostProcessSettings.FilmGrainIntensity * View.FinalPostProcessSettings.FilmGrainIntensityMidtones;
+		CommonParameters.FilmGrain.FilmGrainIntensityHighlights = View.FinalPostProcessSettings.FilmGrainIntensity * View.FinalPostProcessSettings.FilmGrainIntensityHighlights;
+		CommonParameters.FilmGrain.FilmGrainShadowsMax = View.FinalPostProcessSettings.FilmGrainShadowsMax;
+		CommonParameters.FilmGrain.FilmGrainHighlightsMin = View.FinalPostProcessSettings.FilmGrainHighlightsMin;
+
 		CommonParameters.FilmGrain.FilmGrainTexture = View.FilmGrainTexture->GetTexture2DRHI();
 		CommonParameters.FilmGrain.FilmGrainSampler = TStaticSamplerState<SF_Bilinear, AM_Wrap, AM_Wrap>::GetRHI();
 
@@ -588,6 +599,11 @@ FScreenPassTexture AddTonemapPass(FRDGBuilder& GraphBuilder, const FViewInfo& Vi
 	{
 		CommonParameters.FilmGrain.FilmGrainMul = 0.0f;
 		CommonParameters.FilmGrain.FilmGrainAdd = 1.0f;
+		CommonParameters.FilmGrain.FilmGrainIntensityShadows = 0.0f;
+		CommonParameters.FilmGrain.FilmGrainIntensityMidtones = 0.0f;
+		CommonParameters.FilmGrain.FilmGrainIntensityHighlights = 0.0f;
+		CommonParameters.FilmGrain.FilmGrainShadowsMax = 0.09f;
+		CommonParameters.FilmGrain.FilmGrainHighlightsMin = 0.5f;
 		CommonParameters.FilmGrain.FilmGrainTexture = GWhiteTexture->TextureRHI;
 		CommonParameters.FilmGrain.FilmGrainSampler = TStaticSamplerState<SF_Bilinear, AM_Wrap, AM_Wrap>::GetRHI();
 		CommonParameters.FilmGrain.ScreenPosToFilmGrainTextureUV = FScreenTransform::ScreenPosToViewportUV;
