Release async properties lock when cancelling compilation in case an unreachable asset is being polled for IsCompiling() during cook.

#rnx

diff --git a/Engine/Source/Runtime/Engine/Private/StaticMesh.cpp b/Engine/Source/Runtime/Engine/Private/StaticMesh.cpp
--- a/Engine/Source/Runtime/Engine/Private/StaticMesh.cpp
+++ b/Engine/Source/Runtime/Engine/Private/StaticMesh.cpp
@@ -3884,6 +3884,7 @@ bool UStaticMesh::TryCancelAsyncTasks()
 		if (AsyncTask->IsDone() || AsyncTask->Cancel())
 		{
 			AsyncTask.Reset();
+			ReleaseAsyncProperty();
 		}
 	}
 
diff --git a/Engine/Source/Runtime/Engine/Private/StaticMeshCompiler.cpp b/Engine/Source/Runtime/Engine/Private/StaticMeshCompiler.cpp
--- a/Engine/Source/Runtime/Engine/Private/StaticMeshCompiler.cpp
+++ b/Engine/Source/Runtime/Engine/Private/StaticMeshCompiler.cpp
@@ -268,6 +268,7 @@ void FStaticMeshCompilingManager::PostCompilation(UStaticMesh* StaticMesh)
 			// Do not do anything else if the staticmesh is being garbage collected
 			if (StaticMesh->IsUnreachable())
 			{
+				StaticMesh->ReleaseAsyncProperty();
 				return;
 			}
 
