Add better contextualization to DDC access for Mips/VT chunks

#rnx

diff --git a/Engine/Source/Runtime/Engine/Classes/Engine/Texture.h b/Engine/Source/Runtime/Engine/Classes/Engine/Texture.h
--- a/Engine/Source/Runtime/Engine/Classes/Engine/Texture.h
+++ b/Engine/Source/Runtime/Engine/Classes/Engine/Texture.h
@@ -680,7 +680,11 @@ public:
 	bool TryCancelCache();
 	void CancelCache();
 	ENGINE_API bool TryInlineMipData(int32 FirstMipToLoad = 0, UTexture* Texture = nullptr);
+	bool AreDerivedMipsAvailable(FStringView Context) const;
+	bool AreDerivedVTChunksAvailable(FStringView Context) const;
+	UE_DEPRECATED(5.00, "Use AreDerivedMipsAvailable with the context instead.")
 	bool AreDerivedMipsAvailable() const;
+	UE_DEPRECATED(5.00, "Use AreDerivedVTChunksAvailable with the context instead.")
 	bool AreDerivedVTChunksAvailable() const;
 #endif
 
diff --git a/Engine/Source/Runtime/Engine/Private/TextureDerivedData.cpp b/Engine/Source/Runtime/Engine/Private/TextureDerivedData.cpp
--- a/Engine/Source/Runtime/Engine/Private/TextureDerivedData.cpp
+++ b/Engine/Source/Runtime/Engine/Private/TextureDerivedData.cpp
@@ -1269,7 +1269,7 @@ EPixelFormat FTexturePlatformData::GetLayerPixelFormat(uint32 LayerIndex) const
 }
 
 #if WITH_EDITOR
-bool FTexturePlatformData::AreDerivedMipsAvailable() const
+bool FTexturePlatformData::AreDerivedMipsAvailable(FStringView Context) const
 {
 	TArray<FString, TInlineAllocator<16>> MipKeys;
 	for (const FTexture2DMipMap& Mip : Mips)
@@ -1290,10 +1290,10 @@ bool FTexturePlatformData::AreDerivedMipsAvailable() const
 		// to fetch the asset from a remote location once
 		// the mips needs to be accessed
 		TRACE_CPUPROFILER_EVENT_SCOPE(PrefetchMips);
-		return GetDerivedDataCacheRef().TryToPrefetch(MipKeys, TEXT("DerivedMips"_SV));
+		return GetDerivedDataCacheRef().TryToPrefetch(MipKeys, Context);
 	}
 }
-bool FTexturePlatformData::AreDerivedVTChunksAvailable() const
+bool FTexturePlatformData::AreDerivedVTChunksAvailable(FStringView Context) const
 {
 	check(VTData);
 	TArray<FString, TInlineAllocator<16>> ChunkKeys;
@@ -1315,9 +1315,20 @@ bool FTexturePlatformData::AreDerivedVTChunksAvailable() const
 		// to fetch the asset from a remote location once
 		// the VT chunks needs to be accessed
 		TRACE_CPUPROFILER_EVENT_SCOPE(PrefetchDerivedVTChunks);
-		return GetDerivedDataCacheRef().TryToPrefetch(ChunkKeys, TEXT("DerivedVTChunks"_SV));
+		return GetDerivedDataCacheRef().TryToPrefetch(ChunkKeys, Context);
 	}
 }
+
+bool FTexturePlatformData::AreDerivedMipsAvailable() const
+{
+	return AreDerivedMipsAvailable(TEXT("DerivedMips"_SV));
+}
+
+bool FTexturePlatformData::AreDerivedVTChunksAvailable() const
+{
+	return AreDerivedVTChunksAvailable(TEXT("DerivedVTChunks"_SV));
+}
+
 #endif // #if WITH_EDITOR
 
 // Transient flags used to control behavior of platform data serialization
diff --git a/Engine/Source/Runtime/Engine/Private/TextureDerivedDataTask.cpp b/Engine/Source/Runtime/Engine/Private/TextureDerivedDataTask.cpp
--- a/Engine/Source/Runtime/Engine/Private/TextureDerivedDataTask.cpp
+++ b/Engine/Source/Runtime/Engine/Private/TextureDerivedDataTask.cpp
@@ -794,7 +794,7 @@ void FTextureCacheDerivedDataWorker::DoWork()
 							ChunkKeys.Add(Chunk.DerivedDataKey);
 						}
 					}
-					GetDerivedDataCacheRef().TryToPrefetch(ChunkKeys, TEXT("DerivedVTChunks"_SV));
+					GetDerivedDataCacheRef().TryToPrefetch(ChunkKeys, Texture.GetPathName());
 				}
 			}
 
@@ -818,7 +818,7 @@ void FTextureCacheDerivedDataWorker::DoWork()
 			{
 				bSucceeded =	DerivedData->VTData != nullptr &&
 								DerivedData->VTData->IsInitialized() &&
-								DerivedData->AreDerivedVTChunksAvailable();
+								DerivedData->AreDerivedVTChunksAvailable(Texture.GetPathName());
 
 				if (!bSucceeded)
 				{
@@ -827,7 +827,7 @@ void FTextureCacheDerivedDataWorker::DoWork()
 			}
 			else
 			{
-				bSucceeded = DerivedData->AreDerivedMipsAvailable();
+				bSucceeded = DerivedData->AreDerivedMipsAvailable(Texture.GetPathName());
 				if (!bSucceeded)
 				{
 					UE_LOG(LogTexture, Display, TEXT("Texture %s is missing derived mips. The texture will be rebuilt."), *Texture.GetFullName());
