fixed RVT strange filtering on ios

diff --git a/Engine/Source/Runtime/Engine/Private/Materials/MaterialUniformExpressions.cpp b/Engine/Source/Runtime/Engine/Private/Materials/MaterialUniformExpressions.cpp
--- a/Engine/Source/Runtime/Engine/Private/Materials/MaterialUniformExpressions.cpp
+++ b/Engine/Source/Runtime/Engine/Private/Materials/MaterialUniformExpressions.cpp
@@ -1158,7 +1158,10 @@ void FUniformExpressionSet::FillUniformBuffer(const FMaterialRenderContext& Mate
 					}
 				}
 			}
-			
+			static const auto CVarVTAnisotropic = IConsoleManager::Get().FindTConsoleVariableDataInt(TEXT("r.VT.AnisotropicFiltering"));
+			const bool VTAnisotropic = CVarVTAnisotropic && CVarVTAnisotropic->GetValueOnAnyThread() != 0;
+			static const auto CVarVTMaxAnisotropic = IConsoleManager::Get().FindTConsoleVariableDataInt(TEXT("r.VT.MaxAnisotropy"));
+			const int32 VTMaxAnisotropic = (VTAnisotropic && CVarVTMaxAnisotropic) ? CVarVTMaxAnisotropic->GetValueOnAnyThread(): 1;
 			// Now check for runtime virtual texture
 			if (!bValidResources)
 			{
@@ -1174,7 +1177,22 @@ void FUniformExpressionSet::FillUniformBuffer(const FMaterialRenderContext& Mate
 						if (PhysicalViewRHI != nullptr)
 						{
 							*ResourceTablePhysicalTexturePtr = PhysicalViewRHI;
-							*ResourceTablePhysicalSamplerPtr = TStaticSamplerState<SF_AnisotropicPoint, AM_Clamp, AM_Clamp, AM_Clamp, 0, 8>::GetRHI();
+							if(VTMaxAnisotropic >= 8)
+							{
+								*ResourceTablePhysicalSamplerPtr = TStaticSamplerState<SF_AnisotropicPoint, AM_Clamp, AM_Clamp, AM_Clamp, 0, 8>::GetRHI();
+							}
+							else if(VTMaxAnisotropic >= 4)
+							{
+								*ResourceTablePhysicalSamplerPtr = TStaticSamplerState<SF_AnisotropicPoint, AM_Clamp, AM_Clamp, AM_Clamp, 0, 4>::GetRHI();
+							}
+							else if (VTMaxAnisotropic >= 2)
+							{
+								*ResourceTablePhysicalSamplerPtr = TStaticSamplerState<SF_AnisotropicPoint, AM_Clamp, AM_Clamp, AM_Clamp, 0, 2>::GetRHI();
+							}
+							else
+							{
+								*ResourceTablePhysicalSamplerPtr = TStaticSamplerState<SF_AnisotropicPoint, AM_Clamp, AM_Clamp, AM_Clamp, 0, 1>::GetRHI();
+							}
 							bValidResources = true;
 						}
 					}
