Physics : add the ability to query the attached body names from a constraints in BP

diff --git a/Engine/Source/Runtime/Engine/Classes/PhysicsEngine/ConstraintInstanceBlueprintLibrary.h b/Engine/Source/Runtime/Engine/Classes/PhysicsEngine/ConstraintInstanceBlueprintLibrary.h
--- a/Engine/Source/Runtime/Engine/Classes/PhysicsEngine/ConstraintInstanceBlueprintLibrary.h
+++ b/Engine/Source/Runtime/Engine/Classes/PhysicsEngine/ConstraintInstanceBlueprintLibrary.h
@@ -14,6 +14,24 @@ class ENGINE_API UConstraintInstanceBlueprintLibrary : public UBlueprintFunction
 	GENERATED_BODY()
 
 public:
+	//---------------------------------------------------------------------------------------------------
+	// 	   
+	// CONSTRAINT BODIES
+	// 
+	//---------------------------------------------------------------------------------------------------
+
+	/** Gets Attached body names 
+	*	@param Accessor		Constraint accessor to query
+	*	@param ParentBody	Parent body name of the constraint
+	*	@param ChildBody	Child body name of the constraint
+	*/
+	UFUNCTION(BlueprintCallable, Category = "Physics|Constraints", meta = (Keywords = "rigid body name name parent child"))
+		static void GetAttachedBodyNames(
+			UPARAM(ref) FConstraintInstanceAccessor& Accessor,
+			FName& ParentBody,
+			FName& ChildBody
+		);
+
 	//---------------------------------------------------------------------------------------------------
 	//
 	// CONSTRAINT BEHAVIOR 
@@ -34,7 +52,7 @@ public:
 	*	@param Accessor		Constraint accessor to query
 	*/
 	UFUNCTION(BlueprintCallable, Category = "Physics|Constraints", meta = (Keywords = "disable enable collision"))
-	static bool GetDisableCollsion(FConstraintInstanceAccessor& Accessor);
+	static bool GetDisableCollsion(UPARAM(ref) FConstraintInstanceAccessor& Accessor);
 
 	/** Sets projection parameters of the constraint
 	*	@param Accessor					Constraint accessor to change
@@ -78,7 +96,7 @@ public:
 	*	@param Accessor Constraint accessor to query
 	*/
 	UFUNCTION(BlueprintCallable, Category = "Physics|Constraints")
-	static bool GetParentDominates(FConstraintInstanceAccessor& Accessor);
+	static bool GetParentDominates(UPARAM(ref) FConstraintInstanceAccessor& Accessor);
 
 	//---------------------------------------------------------------------------------------------------
 	//
diff --git a/Engine/Source/Runtime/Engine/Private/PhysicsEngine/ConstraintInstanceBlueprintLibrary.cpp b/Engine/Source/Runtime/Engine/Private/PhysicsEngine/ConstraintInstanceBlueprintLibrary.cpp
--- a/Engine/Source/Runtime/Engine/Private/PhysicsEngine/ConstraintInstanceBlueprintLibrary.cpp
+++ b/Engine/Source/Runtime/Engine/Private/PhysicsEngine/ConstraintInstanceBlueprintLibrary.cpp
@@ -2,6 +2,23 @@
 
 #include "PhysicsEngine/ConstraintInstanceBlueprintLibrary.h"
 
+//---------------------------------------------------------------------------------------------------
+// 	   
+// CONSTRAINT BODIES
+// 
+//---------------------------------------------------------------------------------------------------
+void UConstraintInstanceBlueprintLibrary::GetAttachedBodyNames(
+	FConstraintInstanceAccessor& Accessor,
+	FName& ParentBody,
+	FName& ChildBody
+)
+{
+	if (FConstraintInstance* ConstraintInstance = Accessor.Get())
+	{
+		ParentBody = ConstraintInstance->GetParentBoneName();
+		ChildBody = ConstraintInstance->GetChildBoneName();
+	}
+}
 //---------------------------------------------------------------------------------------------------
 //
 // CONSTRAINT BEHAVIOR 
