Mannequin Skeletons seem to break in PIE
#fix patch up from previous fix, MaxLODIndex would be default initialised to MinLODIndex and overwritten when MeshObject was valid, now its default init to maximum # of LODs so should use MinLODIndex when setting NewPredictedLODLevel if no valid MeshObject exists

diff --git a/Engine/Source/Runtime/Engine/Private/Components/SkinnedMeshComponent.cpp b/Engine/Source/Runtime/Engine/Private/Components/SkinnedMeshComponent.cpp
--- a/Engine/Source/Runtime/Engine/Private/Components/SkinnedMeshComponent.cpp
+++ b/Engine/Source/Runtime/Engine/Private/Components/SkinnedMeshComponent.cpp
@@ -3115,7 +3115,7 @@ bool USkinnedMeshComponent::UpdateLODStatus_Internal(int32 InMasterPoseComponent
 			// If no MeshObject - just assume lowest LOD.
 			else
 			{
-				NewPredictedLODLevel = MaxLODIndex;
+				NewPredictedLODLevel = MinLodIndex;
 			}
 
 			// now check to see if we have a MinLODLevel and apply it
