Re-enabling DXC on Vulkan Mobile

diff --git a/Engine/Source/Runtime/Engine/Private/ShaderCompiler/ShaderCompiler.cpp b/Engine/Source/Runtime/Engine/Private/ShaderCompiler/ShaderCompiler.cpp
--- a/Engine/Source/Runtime/Engine/Private/ShaderCompiler/ShaderCompiler.cpp
+++ b/Engine/Source/Runtime/Engine/Private/ShaderCompiler/ShaderCompiler.cpp
@@ -944,7 +944,7 @@ static TAutoConsoleVariable<int32> CVarOpenGLForceDXC(
 
 static TAutoConsoleVariable<int32> CVarVulkanForceDXC(
 	TEXT("r.Vulkan.ForceDXC"),
-	1, 
+	3, 
 	TEXT("Forces DirectX Shader Compiler (DXC) to be used for all Vulkan shaders instead of hlslcc.\n")
 	TEXT(" 0: Disable (hlslcc/glslang)\n")
 	TEXT(" 1: Enabled on desktop platforms only (default)\n")
