UStaticMesh::BatchBuild()
Use the static meshes reverse lookup to obtain the components referencing the meshes we are rebuilding

diff --git a/Engine/Source/Runtime/Engine/Private/StaticMeshBuild.cpp b/Engine/Source/Runtime/Engine/Private/StaticMeshBuild.cpp
--- a/Engine/Source/Runtime/Engine/Private/StaticMeshBuild.cpp
+++ b/Engine/Source/Runtime/Engine/Private/StaticMeshBuild.cpp
@@ -19,6 +19,7 @@
 #if WITH_EDITOR
 #include "Templates/UniquePtr.h"
 #include "Async/Async.h"
+#include "ObjectCacheContext.h"
 #include "IMeshBuilderModule.h"
 #include "IMeshReductionManagerModule.h"
 #include "Interfaces/ITargetPlatformManagerModule.h"
@@ -135,6 +136,8 @@ void UStaticMesh::BatchBuild(const TArray<UStaticMesh*>& InStaticMeshes, bool bI
 	check(IsInGameThread());
 	TRACE_CPUPROFILER_EVENT_SCOPE(UStaticMesh::BatchBuild);
 
+	FObjectCacheContextScope ObjectCacheScope;
+
 	TSet<UStaticMesh*> StaticMeshesToProcess;
 	StaticMeshesToProcess.Reserve(InStaticMeshes.Num());
 
@@ -189,27 +192,24 @@ void UStaticMesh::BatchBuild(const TArray<UStaticMesh*>& InStaticMeshes, bool bI
 		
 		TMap<UStaticMesh*, TArray<UStaticMeshComponent*>> StaticMeshComponents;
 		StaticMeshComponents.Reserve(InStaticMeshes.Num());
+
+		TSet<FSceneInterface*> Scenes;
+
 		for (UStaticMesh* StaticMesh : InStaticMeshes)
 		{
 			if (StaticMesh)
 			{
 				StaticMeshComponents.Add(StaticMesh);
-			}
-		}
 
-		// Detach all instances of those static meshes from the scene.
-		TSet<FSceneInterface*> Scenes;
-		for (TObjectIterator<UStaticMeshComponent> It; It; ++It)
-		{
-			UStaticMesh* StaticMesh = It->GetStaticMesh();
-
-			if (StaticMeshesToProcess.Contains(StaticMesh))
-			{
-				if (It->IsRenderStateCreated())
+				for (UStaticMeshComponent* Component : ObjectCacheScope.GetContext().GetStaticMeshComponents(StaticMesh))
 				{
-					It->DestroyRenderState_Concurrent();
-					StaticMeshComponents[StaticMesh].Add(*It);
-					Scenes.Add(It->GetScene());
+					// Detach all instances of those static meshes from the scene.
+					if (Component->IsRenderStateCreated())
+					{
+						Component->DestroyRenderState_Concurrent();
+						StaticMeshComponents[StaticMesh].Add(Component);
+						Scenes.Add(Component->GetScene());
+					}
 				}
 			}
 		}
