Minimize calls to IsInGameThread from FMICReentranceGuard.

diff --git a/Engine/Source/Runtime/Engine/Classes/Materials/MaterialInstance.h b/Engine/Source/Runtime/Engine/Classes/Materials/MaterialInstance.h
--- a/Engine/Source/Runtime/Engine/Classes/Materials/MaterialInstance.h
+++ b/Engine/Source/Runtime/Engine/Classes/Materials/MaterialInstance.h
@@ -365,19 +365,19 @@ class UMaterialInstance : public UMaterialInterface
 	//Cached copies of the base property overrides or the value from the parent to avoid traversing the parent chain for each access.
 	float OpacityMaskClipValue;
 
-	FORCEINLINE bool GetReentrantFlag() const
+	FORCEINLINE bool GetReentrantFlag(bool bIsInGameThread = IsInGameThread()) const
 	{
 #if WITH_EDITOR
-		return ReentrantFlag[IsInGameThread() ? 0 : 1];
+		return ReentrantFlag[bIsInGameThread ? 0 : 1];
 #else
 		return false;
 #endif
 	}
 
-	FORCEINLINE void SetReentrantFlag(const bool bValue)
+	FORCEINLINE void SetReentrantFlag(const bool bValue, bool bIsInGameThread = IsInGameThread())
 	{
 #if WITH_EDITOR
-		ReentrantFlag[IsInGameThread() ? 0 : 1] = bValue;
+		ReentrantFlag[bIsInGameThread ? 0 : 1] = bValue;
 #endif
 	}
 
diff --git a/Engine/Source/Runtime/Engine/Private/Materials/MaterialInstanceSupport.h b/Engine/Source/Runtime/Engine/Private/Materials/MaterialInstanceSupport.h
--- a/Engine/Source/Runtime/Engine/Private/Materials/MaterialInstanceSupport.h
+++ b/Engine/Source/Runtime/Engine/Private/Materials/MaterialInstanceSupport.h
@@ -36,22 +36,24 @@ public:
 #else
 	FMICReentranceGuard(const UMaterialInstance* InMaterial)
 	{
+		bIsInGameThread = IsInGameThread();
 		Material = const_cast<UMaterialInstance*>(InMaterial);
 
-		if (Material->GetReentrantFlag() == true)
+		if (Material->GetReentrantFlag(bIsInGameThread) == true)
 		{
 			UE_LOG(LogMaterial, Warning, TEXT("InMaterial: %s GameThread: %d RenderThread: %d"), *InMaterial->GetFullName(), IsInGameThread(), IsInRenderingThread());
-			check(!Material->GetReentrantFlag());
+			check(!Material->GetReentrantFlag(bIsInGameThread));
 		}
-		Material->SetReentrantFlag(true);
+		Material->SetReentrantFlag(true, bIsInGameThread);
 	}
 
 	~FMICReentranceGuard()
 	{
-		Material->SetReentrantFlag(false);
+		Material->SetReentrantFlag(false, bIsInGameThread);
 	}
 
 private:
+	bool bIsInGameThread;
 	UMaterialInstance* Material;
 #endif // WITH_EDITOR
 };
