Updating listtextures command to report uncompressed texture and non-vt textures which are common sources of memory bloat.

- Adding `listtextures uncompressed` which will dump only uncompressed textures.
- Adding `listtextures nonvt` which will only dump non-virtual textures.
- Added a function to UTexture which checks the compression settings and returns if it is uncompressed.
- Adding "Uncompressed" to the listtextures output.
- Removing leading comma in the CSV output which makes it impossible to copy/paste into a spreadsheet with hand modifying the output.

diff --git a/Engine/Source/Runtime/Engine/Classes/Engine/Texture.h b/Engine/Source/Runtime/Engine/Classes/Engine/Texture.h
--- a/Engine/Source/Runtime/Engine/Classes/Engine/Texture.h
+++ b/Engine/Source/Runtime/Engine/Classes/Engine/Texture.h
@@ -1269,6 +1269,19 @@ public:
 		return (CompressionSettings == TC_Normalmap);
 	}
 
+	/** @return true if the texture has an uncompressed texture setting */
+	bool IsUncompressed() const
+	{
+		return (CompressionSettings == TC_Grayscale ||
+				CompressionSettings == TC_Displacementmap ||
+				CompressionSettings == TC_VectorDisplacementmap ||
+				CompressionSettings == TC_HDR ||
+				CompressionSettings == TC_EditorIcon ||
+				CompressionSettings == TC_DistanceFieldFont ||
+				CompressionSettings == TC_HalfFloat
+			);
+	}
+
 	/**
 	 * Calculates the size of this texture if it had MipCount miplevels streamed in.
 	 *
diff --git a/Engine/Source/Runtime/Engine/Private/UnrealEngine.cpp b/Engine/Source/Runtime/Engine/Private/UnrealEngine.cpp
--- a/Engine/Source/Runtime/Engine/Private/UnrealEngine.cpp
+++ b/Engine/Source/Runtime/Engine/Private/UnrealEngine.cpp
@@ -3673,11 +3673,12 @@ struct FSortedTexture
 	bool		bIsStreaming;
 	bool		bIsUnknownRef;
 	bool		bIsVirtual;
+	bool		bIsUncompressed;
 	int32		UsageCount;
 	int32		NumMips;
 
 	/** Constructor, initializing every member variable with passed in values. */
-	FSortedTexture(	int32 InMaxAllowedSizeX, int32 InMaxAllowedSizeY, EPixelFormat InFormat, int32 InCurSizeX, int32 InCurSizeY, int32 InLODBias, int32 InMaxAllowedSize, int32 InCurrentSize, const FString& InName, int32 InLODGroup, bool bInIsStreaming, bool bInIsUnknownRef, bool bInIsVirtual, int32 InUsageCount, int32 InNumMips )
+	FSortedTexture(	int32 InMaxAllowedSizeX, int32 InMaxAllowedSizeY, EPixelFormat InFormat, int32 InCurSizeX, int32 InCurSizeY, int32 InLODBias, int32 InMaxAllowedSize, int32 InCurrentSize, const FString& InName, int32 InLODGroup, bool bInIsStreaming, bool bInIsUnknownRef, bool bInIsVirtual, int32 InUsageCount, int32 InNumMips, bool bInIsUncompressed)
 		:	MaxAllowedSizeX( InMaxAllowedSizeX )
 		,	MaxAllowedSizeY( InMaxAllowedSizeY )
 		,	Format( InFormat )
@@ -3691,6 +3692,7 @@ struct FSortedTexture
 		,	bIsStreaming( bInIsStreaming )
 		,	bIsUnknownRef(bInIsUnknownRef)
 		,	bIsVirtual( bInIsVirtual )
+		,	bIsUncompressed( bInIsUncompressed )
 		,	UsageCount( InUsageCount )
 		,	NumMips( InNumMips )
 	{}
@@ -5399,12 +5401,14 @@ bool UEngine::HandleListTexturesCommand( const TCHAR* Cmd, FOutputDevice& Ar )
 	const bool bShouldOnlyListStreaming = FParse::Command(&Cmd, TEXT("STREAMING"));
 	const bool bShouldOnlyListNonStreaming = FParse::Command(&Cmd, TEXT("NONSTREAMING")) && !bShouldOnlyListStreaming;
 	const bool bShouldOnlyListForced = FParse::Command(&Cmd, TEXT("FORCED")) && !bShouldOnlyListStreaming && !bShouldOnlyListNonStreaming;
+	const bool bShouldOnlyListUncompressed = FParse::Command(&Cmd, TEXT("UNCOMPRESSED")) && !bShouldOnlyListStreaming && !bShouldOnlyListNonStreaming && !bShouldOnlyListForced;
+	const bool bShouldOnlyListNonVT = FParse::Command(&Cmd, TEXT("NONVT")) && !bShouldOnlyListStreaming && !bShouldOnlyListNonStreaming && !bShouldOnlyListForced && !bShouldOnlyListUncompressed;
 	const bool bAlphaSort = FParse::Param( Cmd, TEXT("ALPHASORT") );
 	const bool bCSV = FParse::Param( Cmd, TEXT("CSV") );
 	const bool bUnknownRefOnly = FParse::Param(Cmd, TEXT("unknownrefonly"));
 
 	Ar.Logf(TEXT("Listing %s%s textures."),
-		bShouldOnlyListForced ? TEXT("forced") : bShouldOnlyListNonStreaming ? TEXT("non streaming") : bShouldOnlyListStreaming ? TEXT("streaming") : TEXT("all"),
+		bShouldOnlyListForced ? TEXT("forced") : bShouldOnlyListNonStreaming ? TEXT("non streaming") : bShouldOnlyListStreaming ? TEXT("streaming") : bShouldOnlyListUncompressed ? TEXT("uncompressed") : bShouldOnlyListNonVT ? TEXT("NONVT") : TEXT("all"),
 		bUnknownRefOnly ? TEXT(" unknown-ref") : TEXT(""));
 
 	// Find out how many times a texture is referenced by primitive components.
@@ -5464,6 +5468,7 @@ bool UEngine::HandleListTexturesCommand( const TCHAR* Cmd, FOutputDevice& Ar )
 		int32				UsageCount			= 0;
 		bool				bIsForced			= false;
 		bool				bIsVirtual			= false;
+		bool				bIsUncompressed		= false;
 
 		bool bUnknownRef = false;
 		if (Streamer)
@@ -5494,6 +5499,7 @@ bool UEngine::HandleListTexturesCommand( const TCHAR* Cmd, FOutputDevice& Ar )
 			UsageCount			= TextureToUsageMap.FindRef(Texture2D);
 			bIsForced			= Texture2D->ShouldMipLevelsBeForcedResident() && bIsStreamingTexture;
 			bIsVirtual			= Texture2D->IsCurrentlyVirtualTextured();
+			bIsUncompressed		= Texture2D->IsUncompressed();
 
 			if ((NumMips >= Texture2D->GetMinTextureResidentMipCount()) && bIsStreamingTexture)
 			{
@@ -5509,8 +5515,10 @@ bool UEngine::HandleListTexturesCommand( const TCHAR* Cmd, FOutputDevice& Ar )
 
 		if( (bShouldOnlyListStreaming && bIsStreamingTexture) ||	
 			(bShouldOnlyListNonStreaming && !bIsStreamingTexture) ||
-			(bShouldOnlyListForced && bIsForced) ||   
-			(!bShouldOnlyListStreaming && !bShouldOnlyListNonStreaming && !bShouldOnlyListForced) )
+			(bShouldOnlyListForced && bIsForced) ||
+			(bShouldOnlyListUncompressed && bIsUncompressed) ||
+			(bShouldOnlyListNonVT && !bIsVirtual) ||
+			(!bShouldOnlyListStreaming && !bShouldOnlyListNonStreaming && !bShouldOnlyListForced && !bShouldOnlyListUncompressed && !bShouldOnlyListNonVT) )
 		{
 			new(SortedTextures) FSortedTexture( 
 				MaxAllowedSizeX,
@@ -5527,7 +5535,8 @@ bool UEngine::HandleListTexturesCommand( const TCHAR* Cmd, FOutputDevice& Ar )
 				bUnknownRef,
 				bIsVirtual,
 				UsageCount,
-				NumMips);
+				NumMips,
+				bIsUncompressed);
 		}
 	}
 
@@ -5555,15 +5564,15 @@ bool UEngine::HandleListTexturesCommand( const TCHAR* Cmd, FOutputDevice& Ar )
 
 	if (bCSV)
 	{
-		Ar.Logf(TEXT(",Max Width,Max Height,Max Size (KB),Bias Authored,Current Width,Current Height,Current Size (KB),Format,LODGroup,Name,Streaming,UnknownRef,VT,Usage Count,NumMips"));
+		Ar.Logf(TEXT("Max Width,Max Height,Max Size (KB),Bias Authored,Current Width,Current Height,Current Size (KB),Format,LODGroup,Name,Streaming,UnknownRef,VT,Usage Count,NumMips,Uncompressed"));
 	}
 	else if (!FPlatformProperties::RequiresCookedData())
 	{
-		Ar.Logf(TEXT("MaxAllowedSize: Width x Height (Size in KB, Authored Bias), Current/InMem: Width x Height (Size in KB), Format, LODGroup, Name, Streaming, UnknownRef, VT, Usage Count, NumMips"));
+		Ar.Logf(TEXT("MaxAllowedSize: Width x Height (Size in KB, Authored Bias), Current/InMem: Width x Height (Size in KB), Format, LODGroup, Name, Streaming, UnknownRef, VT, Usage Count, NumMips, Uncompressed"));
 	}
 	else
 	{
-		Ar.Logf(TEXT("Cooked/OnDisk: Width x Height (Size in KB, Authored Bias), Current/InMem: Width x Height (Size in KB), Format, LODGroup, Name, Streaming, UnknownRef, VT, Usage Count, NumMips"));
+		Ar.Logf(TEXT("Cooked/OnDisk: Width x Height (Size in KB, Authored Bias), Current/InMem: Width x Height (Size in KB), Format, LODGroup, Name, Streaming, UnknownRef, VT, Usage Count, NumMips, Uncompressed"));
 	}
 
 	for( int32 TextureIndex=0; TextureIndex<SortedTextures.Num(); TextureIndex++ )
@@ -5580,7 +5589,7 @@ bool UEngine::HandleListTexturesCommand( const TCHAR* Cmd, FOutputDevice& Ar )
 
 		if (bCSV)
 		{
-			Ar.Logf(TEXT(",%i, %i, %i, %s, %i, %i, %i, %s, %s, %s, %s, %s, %s, %i, %i"),
+			Ar.Logf(TEXT("%i, %i, %i, %s, %i, %i, %i, %s, %s, %s, %s, %s, %s, %i, %i, %s"),
 				SortedTexture.MaxAllowedSizeX, SortedTexture.MaxAllowedSizeY, (SortedTexture.MaxAllowedSize + 512) / 1024,
 				*AuthoredBiasString,
 				SortedTexture.CurSizeX, SortedTexture.CurSizeY, (SortedTexture.CurrentSize + 512) / 1024,
@@ -5591,11 +5600,12 @@ bool UEngine::HandleListTexturesCommand( const TCHAR* Cmd, FOutputDevice& Ar )
 				SortedTexture.bIsUnknownRef ? TEXT("YES") : TEXT("NO"),
 				SortedTexture.bIsVirtual ? TEXT("YES") : TEXT("NO"),
 				SortedTexture.UsageCount,
-				SortedTexture.NumMips);
+				SortedTexture.NumMips,
+				SortedTexture.bIsUncompressed ? TEXT("YES") : TEXT("NO"));
 		}
 		else
 		{
-			Ar.Logf(TEXT("%ix%i (%i KB, %s), %ix%i (%i KB), %s, %s, %s, %s, %s, %s, %i, %i"),
+			Ar.Logf(TEXT("%ix%i (%i KB, %s), %ix%i (%i KB), %s, %s, %s, %s, %s, %s, %i, %i, %s"),
 				SortedTexture.MaxAllowedSizeX, SortedTexture.MaxAllowedSizeY, (SortedTexture.MaxAllowedSize + 512) / 1024,
 				*AuthoredBiasString,
 				SortedTexture.CurSizeX, SortedTexture.CurSizeY, (SortedTexture.CurrentSize + 512) / 1024,
@@ -5606,7 +5616,8 @@ bool UEngine::HandleListTexturesCommand( const TCHAR* Cmd, FOutputDevice& Ar )
 				SortedTexture.bIsUnknownRef ? TEXT("YES") : TEXT("NO"),
 				SortedTexture.bIsVirtual ? TEXT("YES") : TEXT("NO"),
 				SortedTexture.UsageCount,
-				SortedTexture.NumMips);
+				SortedTexture.NumMips,
+				SortedTexture.bIsUncompressed ? TEXT("YES") : TEXT("NO"));
 		}
 
 		if (bValidTextureGroup)
