Fixed exclusion of PackedLevelInstances from HLOD not actually excluded.

diff --git a/Engine/Source/Developer/WorldPartitionHLODUtilities/Private/HLODBuilder.cpp b/Engine/Source/Developer/WorldPartitionHLODUtilities/Private/HLODBuilder.cpp
--- a/Engine/Source/Developer/WorldPartitionHLODUtilities/Private/HLODBuilder.cpp
+++ b/Engine/Source/Developer/WorldPartitionHLODUtilities/Private/HLODBuilder.cpp
@@ -80,6 +80,10 @@ TArray<UPrimitiveComponent*> FHLODBuilder::GatherPrimitiveComponents(const TArra
 	for (const FWorldPartitionReference& ActorRef : InActors)
 	{
 		AActor* Actor = ActorRef->GetActor();
+		if (!Actor->IsHLODRelevant())
+		{
+			continue;
+		}
 
 		// Gather primitives from the Actor
 		GatherPrimitivesFromActor(Actor);
diff --git a/Engine/Source/Runtime/Engine/Private/LevelInstance/Packed/PackedLevelInstanceActor.cpp b/Engine/Source/Runtime/Engine/Private/LevelInstance/Packed/PackedLevelInstanceActor.cpp
--- a/Engine/Source/Runtime/Engine/Private/LevelInstance/Packed/PackedLevelInstanceActor.cpp
+++ b/Engine/Source/Runtime/Engine/Private/LevelInstance/Packed/PackedLevelInstanceActor.cpp
@@ -181,6 +181,12 @@ bool APackedLevelInstance::IsHiddenEd() const
 	return Super::IsHiddenEd() || bEditing || (ChildEditing > 0);
 }
 
+bool APackedLevelInstance::IsHLODRelevant() const
+{
+	// Bypass base class ALevelInstance (because it always returns true). We want the same implementation as AActor.
+	return AActor::IsHLODRelevant();
+}
+
 bool APackedLevelInstance::CanEditChange(const FProperty* InProperty) const
 {
 	if (!Super::CanEditChange(InProperty))
diff --git a/Engine/Source/Runtime/Engine/Public/LevelInstance/Packed/PackedLevelInstanceActor.h b/Engine/Source/Runtime/Engine/Public/LevelInstance/Packed/PackedLevelInstanceActor.h
--- a/Engine/Source/Runtime/Engine/Public/LevelInstance/Packed/PackedLevelInstanceActor.h
+++ b/Engine/Source/Runtime/Engine/Public/LevelInstance/Packed/PackedLevelInstanceActor.h
@@ -52,6 +52,7 @@ public:
 	virtual EActorGridPlacement GetDefaultGridPlacement() const override { return EActorGridPlacement::None; }
 
 	virtual bool IsHiddenEd() const override;
+	virtual bool IsHLODRelevant() const override;
 
 	void DestroyPackedComponents();
 	void GetPackedComponents(TArray<UActorComponent*>& OutPackedComponents) const;
