Fixed a bug where particle beams with LWC were corrupted

diff --git a/Engine/Source/Runtime/Engine/Private/Particles/ParticleSystemRender.cpp b/Engine/Source/Runtime/Engine/Private/Particles/ParticleSystemRender.cpp
--- a/Engine/Source/Runtime/Engine/Private/Particles/ParticleSystemRender.cpp
+++ b/Engine/Source/Runtime/Engine/Private/Particles/ParticleSystemRender.cpp
@@ -2505,11 +2505,11 @@ void FDynamicBeam2EmitterData::RenderDirectLine(const FParticleSystemSceneProxy*
 		DECLARE_PARTICLE_PTR(Particle, Source.DataContainer.ParticleData + Source.ParticleStride * Beam);
 
 		FBeam2TypeDataPayload*	BeamPayloadData		= NULL;
-		FVector3f*				InterpolatedPoints	= NULL;
+		FVector*				InterpolatedPoints	= NULL;
 		float*					NoiseRate			= NULL;
 		float*					NoiseDelta			= NULL;
-		FVector3f*				TargetNoisePoints	= NULL;
-		FVector3f*				NextNoisePoints		= NULL;
+		FVector*				TargetNoisePoints	= NULL;
+		FVector*				NextNoisePoints		= NULL;
 		float*					TaperValues			= NULL;
 
 		BeamPayloadData = (FBeam2TypeDataPayload*)((uint8*)Particle + Source.BeamDataOffset);
@@ -2572,8 +2572,8 @@ void FDynamicBeam2EmitterData::RenderLines(const FParticleSystemSceneProxy* Prox
 			FBeam2TypeDataPayload*	BeamPayloadData		= NULL;
 			float*					NoiseRate			= NULL;
 			float*					NoiseDelta			= NULL;
-			FVector3f*				TargetNoisePoints	= NULL;
-			FVector3f*				NextNoisePoints		= NULL;
+			FVector*				TargetNoisePoints	= NULL;
+			FVector*				NextNoisePoints		= NULL;
 			float*					TaperValues			= NULL;
 			float*					NoiseDistanceScale	= NULL;
 
@@ -2592,11 +2592,11 @@ void FDynamicBeam2EmitterData::RenderLines(const FParticleSystemSceneProxy* Prox
 			}
 			if (Source.TargetNoisePointsOffset != -1)
 			{
-				TargetNoisePoints = (FVector3f*)((uint8*)Particle + Source.TargetNoisePointsOffset);
+				TargetNoisePoints = (FVector*)((uint8*)Particle + Source.TargetNoisePointsOffset);
 			}
 			if (Source.NextNoisePointsOffset != -1)
 			{
-				NextNoisePoints = (FVector3f*)((uint8*)Particle + Source.NextNoisePointsOffset);
+				NextNoisePoints = (FVector*)((uint8*)Particle + Source.NextNoisePointsOffset);
 			}
 			if (Source.TaperValuesOffset != -1)
 			{
@@ -2613,8 +2613,8 @@ void FDynamicBeam2EmitterData::RenderLines(const FParticleSystemSceneProxy* Prox
 				NoiseDistScale = *NoiseDistanceScale;
 			}
 
-			FVector3f* NoisePoints	= TargetNoisePoints;
-			FVector3f* NextNoise		= NextNoisePoints;
+			FVector* NoisePoints	= TargetNoisePoints;
+			FVector* NextNoise		= NextNoisePoints;
 
 			float NoiseRangeScaleFactor = Source.NoiseRangeScale;
 			//@todo. How to handle no noise points?
@@ -2911,7 +2911,7 @@ void FDynamicBeam2EmitterData::RenderLines(const FParticleSystemSceneProxy* Prox
 				DECLARE_PARTICLE_PTR(Particle, Source.DataContainer.ParticleData + Source.ParticleStride * i);
 
 				FBeam2TypeDataPayload*	BeamPayloadData		= NULL;
-				FVector3f*				InterpolatedPoints	= NULL;
+				FVector*				InterpolatedPoints	= NULL;
 
 				BeamPayloadData = (FBeam2TypeDataPayload*)((uint8*)Particle + Source.BeamDataOffset);
 				if (BeamPayloadData->TriangleCount == 0)
@@ -2920,7 +2920,7 @@ void FDynamicBeam2EmitterData::RenderLines(const FParticleSystemSceneProxy* Prox
 				}
 				if (Source.InterpolatedPointsOffset != -1)
 				{
-					InterpolatedPoints = (FVector3f*)((uint8*)Particle + Source.InterpolatedPointsOffset);
+					InterpolatedPoints = (FVector*)((uint8*)Particle + Source.InterpolatedPointsOffset);
 				}
 
 				FVector Location;
@@ -3127,11 +3127,11 @@ int32 FDynamicBeam2EmitterData::FillVertexData_NoNoise(FAsyncBufferFillData& Me)
 			DECLARE_PARTICLE_PTR(Particle, Source.DataContainer.ParticleData + Source.ParticleStride * i);
 
 			FBeam2TypeDataPayload*	BeamPayloadData		= NULL;
-			FVector3f*				InterpolatedPoints	= NULL;
+			FVector*				InterpolatedPoints	= NULL;
 			float*					NoiseRate			= NULL;
 			float*					NoiseDelta			= NULL;
-			FVector3f*				TargetNoisePoints	= NULL;
-			FVector3f*				NextNoisePoints		= NULL;
+			FVector*				TargetNoisePoints	= NULL;
+			FVector*				NextNoisePoints		= NULL;
 			float*					TaperValues			= NULL;
 
 			BeamPayloadData = (FBeam2TypeDataPayload*)((uint8*)Particle + Source.BeamDataOffset);
@@ -3141,7 +3141,7 @@ int32 FDynamicBeam2EmitterData::FillVertexData_NoNoise(FAsyncBufferFillData& Me)
 			}
 			if (Source.InterpolatedPointsOffset != -1)
 			{
-				InterpolatedPoints = (FVector3f*)((uint8*)Particle + Source.InterpolatedPointsOffset);
+				InterpolatedPoints = (FVector*)((uint8*)Particle + Source.InterpolatedPointsOffset);
 			}
 			if (Source.NoiseRateOffset != -1)
 			{
@@ -3153,11 +3153,11 @@ int32 FDynamicBeam2EmitterData::FillVertexData_NoNoise(FAsyncBufferFillData& Me)
 			}
 			if (Source.TargetNoisePointsOffset != -1)
 			{
-				TargetNoisePoints = (FVector3f*)((uint8*)Particle + Source.TargetNoisePointsOffset);
+				TargetNoisePoints = (FVector*)((uint8*)Particle + Source.TargetNoisePointsOffset);
 			}
 			if (Source.NextNoisePointsOffset != -1)
 			{
-				NextNoisePoints = (FVector3f*)((uint8*)Particle + Source.NextNoisePointsOffset);
+				NextNoisePoints = (FVector*)((uint8*)Particle + Source.NextNoisePointsOffset);
 			}
 			if (Source.TaperValuesOffset != -1)
 			{
@@ -3306,11 +3306,11 @@ int32 FDynamicBeam2EmitterData::FillVertexData_NoNoise(FAsyncBufferFillData& Me)
 			DECLARE_PARTICLE_PTR(Particle, Source.DataContainer.ParticleData + Source.ParticleStride * i);
 
 			FBeam2TypeDataPayload*	BeamPayloadData		= NULL;
-			FVector3f*				InterpolatedPoints	= NULL;
+			FVector*				InterpolatedPoints	= NULL;
 			float*					NoiseRate			= NULL;
 			float*					NoiseDelta			= NULL;
-			FVector3f*				TargetNoisePoints	= NULL;
-			FVector3f*				NextNoisePoints		= NULL;
+			FVector*				TargetNoisePoints	= NULL;
+			FVector*				NextNoisePoints		= NULL;
 			float*					TaperValues			= NULL;
 
 			BeamPayloadData = (FBeam2TypeDataPayload*)((uint8*)Particle + Source.BeamDataOffset);
@@ -3320,7 +3320,7 @@ int32 FDynamicBeam2EmitterData::FillVertexData_NoNoise(FAsyncBufferFillData& Me)
 			}
 			if (Source.InterpolatedPointsOffset != -1)
 			{
-				InterpolatedPoints = (FVector3f*)((uint8*)Particle + Source.InterpolatedPointsOffset);
+				InterpolatedPoints = (FVector*)((uint8*)Particle + Source.InterpolatedPointsOffset);
 			}
 			if (Source.NoiseRateOffset != -1)
 			{
@@ -3332,11 +3332,11 @@ int32 FDynamicBeam2EmitterData::FillVertexData_NoNoise(FAsyncBufferFillData& Me)
 			}
 			if (Source.TargetNoisePointsOffset != -1)
 			{
-				TargetNoisePoints = (FVector3f*)((uint8*)Particle + Source.TargetNoisePointsOffset);
+				TargetNoisePoints = (FVector*)((uint8*)Particle + Source.TargetNoisePointsOffset);
 			}
 			if (Source.NextNoisePointsOffset != -1)
 			{
-				NextNoisePoints = (FVector3f*)((uint8*)Particle + Source.NextNoisePointsOffset);
+				NextNoisePoints = (FVector*)((uint8*)Particle + Source.NextNoisePointsOffset);
 			}
 			if (Source.TaperValuesOffset != -1)
 			{
@@ -3582,11 +3582,11 @@ int32 FDynamicBeam2EmitterData::FillData_Noise(FAsyncBufferFillData& Me) const
 
 		// Retrieve the beam data from the particle.
 		FBeam2TypeDataPayload*	BeamPayloadData		= NULL;
-		FVector3f*				InterpolatedPoints	= NULL;
+		FVector*				InterpolatedPoints	= NULL;
 		float*					NoiseRate			= NULL;
 		float*					NoiseDelta			= NULL;
-		FVector3f*				TargetNoisePoints	= NULL;
-		FVector3f*				NextNoisePoints		= NULL;
+		FVector*				TargetNoisePoints	= NULL;
+		FVector*				NextNoisePoints		= NULL;
 		float*					TaperValues			= NULL;
 		float*					NoiseDistanceScale	= NULL;
 
@@ -3597,7 +3597,7 @@ int32 FDynamicBeam2EmitterData::FillData_Noise(FAsyncBufferFillData& Me) const
 		}
 		if (Source.InterpolatedPointsOffset != -1)
 		{
-			InterpolatedPoints = (FVector3f*)((uint8*)Particle + Source.InterpolatedPointsOffset);
+			InterpolatedPoints = (FVector*)((uint8*)Particle + Source.InterpolatedPointsOffset);
 		}
 		if (Source.NoiseRateOffset != -1)
 		{
@@ -3609,11 +3609,11 @@ int32 FDynamicBeam2EmitterData::FillData_Noise(FAsyncBufferFillData& Me) const
 		}
 		if (Source.TargetNoisePointsOffset != -1)
 		{
-			TargetNoisePoints = (FVector3f*)((uint8*)Particle + Source.TargetNoisePointsOffset);
+			TargetNoisePoints = (FVector*)((uint8*)Particle + Source.TargetNoisePointsOffset);
 		}
 		if (Source.NextNoisePointsOffset != -1)
 		{
-			NextNoisePoints = (FVector3f*)((uint8*)Particle + Source.NextNoisePointsOffset);
+			NextNoisePoints = (FVector*)((uint8*)Particle + Source.NextNoisePointsOffset);
 		}
 		if (Source.TaperValuesOffset != -1)
 		{
@@ -3630,8 +3630,8 @@ int32 FDynamicBeam2EmitterData::FillData_Noise(FAsyncBufferFillData& Me) const
 			NoiseDistScale = *NoiseDistanceScale;
 		}
 
-		FVector3f* NoisePoints	= TargetNoisePoints;
-		FVector3f* NextNoise		= NextNoisePoints;
+		FVector* NoisePoints	= TargetNoisePoints;
+		FVector* NextNoise		= NextNoisePoints;
 
 		float NoiseRangeScaleFactor = Source.NoiseRangeScale;
 		//@todo. How to handle no noise points?
@@ -4543,11 +4543,11 @@ int32 FDynamicBeam2EmitterData::FillData_InterpolatedNoise(FAsyncBufferFillData&
 
 		// Retrieve the beam data from the particle.
 		FBeam2TypeDataPayload*	BeamPayloadData		= NULL;
-		FVector3f*				InterpolatedPoints	= NULL;
+		FVector*				InterpolatedPoints	= NULL;
 		float*					NoiseRate			= NULL;
 		float*					NoiseDelta			= NULL;
-		FVector3f*				TargetNoisePoints	= NULL;
-		FVector3f*				NextNoisePoints		= NULL;
+		FVector*				TargetNoisePoints	= NULL;
+		FVector*				NextNoisePoints		= NULL;
 		float*					TaperValues			= NULL;
 		float*					NoiseDistanceScale	= NULL;
 
@@ -4563,7 +4563,7 @@ int32 FDynamicBeam2EmitterData::FillData_InterpolatedNoise(FAsyncBufferFillData&
 
 		if (Source.InterpolatedPointsOffset != -1)
 		{
-			InterpolatedPoints = (FVector3f*)((uint8*)Particle + Source.InterpolatedPointsOffset);
+			InterpolatedPoints = (FVector*)((uint8*)Particle + Source.InterpolatedPointsOffset);
 		}
 		if (Source.NoiseRateOffset != -1)
 		{
@@ -4575,11 +4575,11 @@ int32 FDynamicBeam2EmitterData::FillData_InterpolatedNoise(FAsyncBufferFillData&
 		}
 		if (Source.TargetNoisePointsOffset != -1)
 		{
-			TargetNoisePoints = (FVector3f*)((uint8*)Particle + Source.TargetNoisePointsOffset);
+			TargetNoisePoints = (FVector*)((uint8*)Particle + Source.TargetNoisePointsOffset);
 		}
 		if (Source.NextNoisePointsOffset != -1)
 		{
-			NextNoisePoints = (FVector3f*)((uint8*)Particle + Source.NextNoisePointsOffset);
+			NextNoisePoints = (FVector*)((uint8*)Particle + Source.NextNoisePointsOffset);
 		}
 		if (Source.TaperValuesOffset != -1)
 		{
@@ -4603,8 +4603,8 @@ int32 FDynamicBeam2EmitterData::FillData_InterpolatedNoise(FAsyncBufferFillData&
 		bool bInterpFractionIsZero = false;
 		int32 InterpIndex = FMath::TruncToInt(InterpStepSize);
 
-		FVector3f* NoisePoints	= TargetNoisePoints;
-		FVector3f* NextNoise		= NextNoisePoints;
+		FVector* NoisePoints	= TargetNoisePoints;
+		FVector* NextNoise		= NextNoisePoints;
 
 		// Appropriate checks are made before access access, no need to assert here
 		CA_ASSUME(NoisePoints != NULL);
