Trivial simplification of BlendSpace blending space check

diff --git a/Engine/Source/Runtime/Engine/Classes/Animation/BlendSpace.h b/Engine/Source/Runtime/Engine/Classes/Animation/BlendSpace.h
--- a/Engine/Source/Runtime/Engine/Classes/Animation/BlendSpace.h
+++ b/Engine/Source/Runtime/Engine/Classes/Animation/BlendSpace.h
@@ -704,8 +704,8 @@ private:
 
 public:
 	/**
-	* When you use blend per bone, allows rotation to blend in mesh space. This only works if this does not contain additive animation samples
-	* This is more performance intensive
+	* When you use per bone sample smoothing, this makes blending happen in mesh space and can happen if this 
+	* contains additive animation samples. This is more performance intensive than blending in local space.
 	*/
 	UPROPERTY()
 	bool bRotationBlendInMeshSpace;
@@ -751,7 +751,6 @@ public:
 	UPROPERTY(EditAnywhere, Category = SampleSmoothing, meta = (DisplayName = "Smoothing"))
 	bool bTargetWeightInterpolationEaseInOut = true;
 
-
 #if WITH_EDITORONLY_DATA
 	/** Preview Base pose for additive BlendSpace **/
 	UPROPERTY(EditAnywhere, Category = AdditiveSettings)
diff --git a/Engine/Source/Runtime/Engine/Private/Animation/BlendSpace.cpp b/Engine/Source/Runtime/Engine/Private/Animation/BlendSpace.cpp
--- a/Engine/Source/Runtime/Engine/Private/Animation/BlendSpace.cpp
+++ b/Engine/Source/Runtime/Engine/Private/Animation/BlendSpace.cpp
@@ -873,16 +873,9 @@ void UBlendSpace::GetAnimationPose_Internal(TArray<FBlendSampleData>& BlendSampl
 
 	if (PerBoneBlend.Num() > 0)
 	{
-		if (IsValidAdditive())
+		if (bRotationBlendInMeshSpace)
 		{
-			if (bRotationBlendInMeshSpace)
-			{
-				FAnimationRuntime::BlendPosesTogetherPerBoneInMeshSpace(ChildrenPosesView, ChildrenCurves, ChildrenAttributes, this, BlendSampleDataCache, OutAnimationPoseData);
-			}
-			else
-			{
-				FAnimationRuntime::BlendPosesTogetherPerBone(ChildrenPosesView, ChildrenCurves, ChildrenAttributes, this, BlendSampleDataCache, OutAnimationPoseData);
-			}
+			FAnimationRuntime::BlendPosesTogetherPerBoneInMeshSpace(ChildrenPosesView, ChildrenCurves, ChildrenAttributes, this, BlendSampleDataCache, OutAnimationPoseData);
 		}
 		else
 		{
