Fixed a bug where Cascade GPU systems did not render with LWC enabled

diff --git a/Engine/Source/Runtime/Engine/Private/Particles/FXSystemPrivate.h b/Engine/Source/Runtime/Engine/Private/Particles/FXSystemPrivate.h
--- a/Engine/Source/Runtime/Engine/Private/Particles/FXSystemPrivate.h
+++ b/Engine/Source/Runtime/Engine/Private/Particles/FXSystemPrivate.h
@@ -90,11 +90,11 @@ Injecting particles in to the GPU for simulation.
 struct FNewParticle
 {
 	/** The initial position of the particle. */
-	FVector Position;
+	FVector3f Position;
 	/** The relative time of the particle. */
 	float RelativeTime;
 	/** The initial velocity of the particle. */
-	FVector Velocity;
+	FVector3f Velocity;
 	/** The time scale for the particle. */
 	float TimeScale;
 	/** Initial size of the particle. */
