Texture: Implement a texture build using only the new DDC Build API

The new build path can be run by adding -DDC2AsyncTextureBuilds to the command line. It is missing some of the input validation done by the existing build path, and has not been hooked up to editor texture streaming yet.

#rnx

diff --git a/Engine/Source/Runtime/Engine/Private/TextureDerivedData.cpp b/Engine/Source/Runtime/Engine/Private/TextureDerivedData.cpp
--- a/Engine/Source/Runtime/Engine/Private/TextureDerivedData.cpp
+++ b/Engine/Source/Runtime/Engine/Private/TextureDerivedData.cpp
@@ -737,6 +737,19 @@ void FTexturePlatformData::Cache(
 
 	ETextureCacheFlags Flags = ETextureCacheFlags(InFlags);
 
+	static const bool bUseNewBuildPath = FParse::Param(FCommandLine::Get(), TEXT("DDC2AsyncTextureBuilds"));
+	if (bUseNewBuildPath && InTexture.Source.GetNumLayers() == 1 && !InSettingsPerLayer->bVirtualStreamable)
+	{
+		COOK_STAT(auto Timer = TextureCookStats::UsageStats.TimeSyncWork());
+		COOK_STAT(Timer.TrackCyclesOnly());
+		EQueuedWorkPriority Priority = FTextureCompilingManager::Get().GetBasePriority(&InTexture);
+		AsyncTask = CreateTextureBuildTask(InTexture, *this, *InSettingsPerLayer, Priority, Flags);
+		if (AsyncTask)
+		{
+			return;
+		}
+	}
+
 	static bool bForDDC = FString(FCommandLine::Get()).Contains(TEXT("Run=DerivedDataCache"));
 	if (bForDDC)
 	{
@@ -779,6 +792,8 @@ void FTexturePlatformData::Cache(
 		FQueuedThreadPool*  TextureThreadPool = FTextureCompilingManager::Get().GetThreadPool();
 		EQueuedWorkPriority BasePriority      = FTextureCompilingManager::Get().GetBasePriority(&InTexture);
 
+		COOK_STAT(auto Timer = TextureCookStats::UsageStats.TimeSyncWork());
+		COOK_STAT(Timer.TrackCyclesOnly());
 		FTextureAsyncCacheDerivedDataWorkerTask* LocalTask = new FTextureAsyncCacheDerivedDataWorkerTask(TextureThreadPool, Compressor, this, &InTexture, InSettingsPerLayer, Flags);
 		AsyncTask = LocalTask;
 		LocalTask->StartBackgroundTask(TextureThreadPool, BasePriority, EQueuedWorkFlags::DoNotRunInsideBusyWait, LocalTask->GetTask().GetRequiredMemoryEstimate());
diff --git a/Engine/Source/Runtime/Engine/Private/TextureDerivedDataTask.cpp b/Engine/Source/Runtime/Engine/Private/TextureDerivedDataTask.cpp
--- a/Engine/Source/Runtime/Engine/Private/TextureDerivedDataTask.cpp
+++ b/Engine/Source/Runtime/Engine/Private/TextureDerivedDataTask.cpp
@@ -23,6 +23,7 @@
 
 #if WITH_EDITOR
 
+#include "Algo/Accumulate.h"
 #include "DerivedDataBuild.h"
 #include "DerivedDataBuildAction.h"
 #include "DerivedDataBuildInputResolver.h"
@@ -34,6 +35,9 @@
 #include "Engine/TextureCube.h"
 #include "Engine/VolumeTexture.h"
 #include "GenericPlatform/GenericPlatformMath.h"
+#include "HAL/CriticalSection.h"
+#include "HAL/PlatformProcess.h"
+#include "HAL/PlatformTime.h"
 #include "Interfaces/ITargetPlatform.h"
 #include "Interfaces/ITargetPlatformManagerModule.h"
 #include "Interfaces/ITextureFormat.h"
@@ -41,10 +45,12 @@
 #include "Math/UnrealMathUtility.h"
 #include "Misc/FeedbackContext.h"
 #include "Misc/Optional.h"
+#include "Misc/ScopeRWLock.h"
 #include "ProfilingDebugging/CookStats.h"
 #include "Serialization/BulkDataRegistry.h"
 #include "TextureDerivedDataBuildUtils.h"
 #include "VT/VirtualTextureDataBuilder.h"
+#include <atomic>
 
 static TAutoConsoleVariable<int32> CVarVTValidateCompressionOnLoad(
 	TEXT("r.VT.ValidateCompressionOnLoad"),
@@ -256,6 +262,102 @@ void FTextureSourceData::GetAsyncSourceMips(IImageWrapperModule* InImageWrapper)
 		GetSourceMips(AsyncSource, InImageWrapper);
 	}
 }
+
+namespace UE::TextureDerivedData
+{
+
+using namespace UE::DerivedData;
+
+class FTextureBuildInputResolver final : public IBuildInputResolver
+{
+public:
+	explicit FTextureBuildInputResolver(UTexture& InTexture)
+		: Texture(InTexture)
+	{
+	}
+
+	const FCompressedBuffer& FindSource(FCompressedBuffer& Buffer, FTextureSource& Source, const FGuid& BulkDataId)
+	{
+		if (Source.GetPersistentId() != BulkDataId)
+		{
+			return FCompressedBuffer::Null;
+		}
+		if (!Buffer)
+		{
+			Source.OperateOnLoadedBulkData([&Buffer](const FSharedBuffer& BulkDataBuffer)
+			{
+				Buffer = FCompressedBuffer::Compress(NAME_Default, BulkDataBuffer);
+			});
+		}
+		return Buffer;
+	}
+
+	FRequest ResolveInputMeta(
+		const FBuildDefinition& Definition,
+		EPriority Priority,
+		FOnBuildInputMetaResolved&& OnResolved) final
+	{
+		EStatus Status = EStatus::Ok;
+		TArray<FString> InputKeys;
+		TArray<FBuildInputMetaByKey> Inputs;
+		Definition.IterateInputBulkData([this, &Status, &InputKeys, &Inputs](FStringView Key, const FGuid& BulkDataId)
+		{
+			const FCompressedBuffer& Buffer = Key == TEXT("Source"_SV)
+				? FindSource(SourceBuffer, Texture.Source, BulkDataId)
+				: FindSource(CompositeSourceBuffer, Texture.CompositeTexture->Source, BulkDataId);
+			if (Buffer)
+			{
+				InputKeys.Emplace(Key);
+				Inputs.Add({InputKeys.Last(), Buffer.GetRawHash(), Buffer.GetRawSize()});
+			}
+			else
+			{
+				Status = EStatus::Error;
+			}
+		});
+		OnResolved({Inputs, Status});
+		return FRequest();
+	}
+
+	FRequest ResolveInputData(
+		const FBuildDefinition& Definition,
+		EPriority Priority,
+		FOnBuildInputDataResolved&& OnResolved,
+		FBuildInputFilter&& Filter) final
+	{
+		EStatus Status = EStatus::Ok;
+		TArray<FString> InputKeys;
+		TArray<FBuildInputDataByKey> Inputs;
+		Definition.IterateInputBulkData([this, &Filter, &Status, &InputKeys, &Inputs](FStringView Key, const FGuid& BulkDataId)
+		{
+			if (!Filter || Filter(Key))
+			{
+				const FCompressedBuffer& Buffer = Key == TEXT("Source"_SV)
+					? FindSource(SourceBuffer, Texture.Source, BulkDataId)
+					: FindSource(CompositeSourceBuffer, Texture.CompositeTexture->Source, BulkDataId);
+				if (Buffer)
+				{
+					InputKeys.Emplace(Key);
+					Inputs.Add({InputKeys.Last(), Buffer});
+				}
+				else
+				{
+					Status = EStatus::Error;
+				}
+			}
+		});
+		OnResolved({Inputs, Status});
+		return FRequest();
+	}
+
+private:
+	UTexture& Texture;
+	FCompressedBuffer SourceBuffer;
+	FCompressedBuffer CompositeSourceBuffer;
+};
+
+} // UE::TextureDerivedData
+
 void FTextureCacheDerivedDataWorker::ConsumeBuildFunctionOutput(const UE::DerivedData::FBuildOutput& BuildOutput, const FString& TexturePath, bool bReplaceExistingDDC)
 {
 	using namespace UE::DerivedData;
@@ -450,95 +552,7 @@ void FTextureCacheDerivedDataWorker::BuildTexture(bool bReplaceExistingDDC)
 				DefinitionBuilder.AddInputBulkData(TEXT("CompositeSource"_SV), Texture.CompositeTexture->Source.GetPersistentId());
 			}
 
-			class FTextureBuildInputResolver final : public IBuildInputResolver
-			{
-			public:
-				explicit FTextureBuildInputResolver(UTexture& InTexture)
-					: Texture(InTexture)
-				{
-				}
-
-				const FCompressedBuffer& FindSource(FCompressedBuffer& Buffer, FTextureSource& Source, const FGuid& BulkDataId)
-				{
-					if (Source.GetPersistentId() != BulkDataId)
-					{
-						return FCompressedBuffer::Null;
-					}
-					if (!Buffer)
-					{
-						Source.OperateOnLoadedBulkData([&Buffer](const FSharedBuffer& BulkDataBuffer)
-						{
-							Buffer = FCompressedBuffer::Compress(NAME_Default, BulkDataBuffer);
-						});
-					}
-					return Buffer;
-				}
-
-				FRequest ResolveInputMeta(
-					const FBuildDefinition& Definition,
-					EPriority Priority,
-					FOnBuildInputMetaResolved&& OnResolved) final
-				{
-					EStatus Status = EStatus::Ok;
-					TArray<FString> InputKeys;
-					TArray<FBuildInputMetaByKey> Inputs;
-					Definition.IterateInputBulkData([this, &Status, &InputKeys, &Inputs](FStringView Key, const FGuid& BulkDataId)
-					{
-						const FCompressedBuffer& Buffer = Key == TEXT("Source"_SV)
-							? FindSource(SourceBuffer, Texture.Source, BulkDataId)
-							: FindSource(CompositeSourceBuffer, Texture.CompositeTexture->Source, BulkDataId);
-						if (Buffer)
-						{
-							InputKeys.Emplace(Key);
-							Inputs.Add({InputKeys.Last(), Buffer.GetRawHash(), Buffer.GetRawSize()});
-						}
-						else
-						{
-							Status = EStatus::Error;
-						}
-					});
-					OnResolved({Inputs, Status});
-					return FRequest();
-				}
-
-				FRequest ResolveInputData(
-					const FBuildDefinition& Definition,
-					EPriority Priority,
-					FOnBuildInputDataResolved&& OnResolved,
-					FBuildInputFilter&& Filter) final
-				{
-					EStatus Status = EStatus::Ok;
-					TArray<FString> InputKeys;
-					TArray<FBuildInputDataByKey> Inputs;
-					Definition.IterateInputBulkData([this, &Filter, &Status, &InputKeys, &Inputs](FStringView Key, const FGuid& BulkDataId)
-					{
-						if (!Filter || Filter(Key))
-						{
-							const FCompressedBuffer& Buffer = Key == TEXT("Source"_SV)
-								? FindSource(SourceBuffer, Texture.Source, BulkDataId)
-								: FindSource(CompositeSourceBuffer, Texture.CompositeTexture->Source, BulkDataId);
-							if (Buffer)
-							{
-								InputKeys.Emplace(Key);
-								Inputs.Add({InputKeys.Last(), Buffer});
-							}
-							else
-							{
-								Status = EStatus::Error;
-							}
-						}
-					});
-					OnResolved({Inputs, Status});
-					return FRequest();
-				}
-
-			private:
-				UTexture& Texture;
-				FCompressedBuffer SourceBuffer;
-				FCompressedBuffer CompositeSourceBuffer;
-			};
-
-			TOptional<FTextureBuildInputResolver> PlaceholderResolver;
+			TOptional<UE::TextureDerivedData::FTextureBuildInputResolver> PlaceholderResolver;
 			UE::DerivedData::IBuildInputResolver* InputResolver = GetGlobalBuildInputResolver();
 			if (!InputResolver)
 			{
@@ -963,4 +977,324 @@ void FTextureCacheDerivedDataWorker::Finalize()
 	}
 }
 
+class FTextureBuildTask final : public FTextureAsyncCacheDerivedDataTask
+{
+public:
+	FTextureBuildTask(
+		UTexture& Texture,
+		FStringView FunctionName,
+		FTexturePlatformData& InDerivedData,
+		const FTextureBuildSettings& Settings,
+		EQueuedWorkPriority BasePriority,
+		ETextureCacheFlags Flags)
+		: ActiveRequest(&Request)
+		, DerivedData(InDerivedData)
+		, Priority(EnumHasAnyFlags(Flags, ETextureCacheFlags::Async) ? ConvertPriority(BasePriority) : UE::DerivedData::EPriority::Blocking)
+		, bCacheHit(false)
+		, bInlineMips(EnumHasAnyFlags(Flags, ETextureCacheFlags::InlineMips))
+		, FirstMipToLoad(Settings.LODBiasWithCinematicMips)
+		, InputResolver(Texture)
+	{
+		using namespace UE::DerivedData;
+
+		static bool bLoadedModules = LoadModules();
+
+		FString KeySuffix;
+		GetTextureDerivedDataKeySuffix(Texture, &Settings, KeySuffix);
+
+		TStringBuilder<256> TexturePath;
+		Texture.GetPathName(nullptr, TexturePath);
+
+		IBuild& Build = GetDerivedDataBuildRef();
+		IBuildInputResolver* GlobalResolver = GetGlobalBuildInputResolver();
+		BuildSession = Build.CreateSession(TexturePath, GlobalResolver ? GlobalResolver : &InputResolver);
+
+		FBuildDefinition Definition = CreateDefinition(Build, Texture, TexturePath, FunctionName, KeySuffix, Settings);
+
+		if (!EnumHasAnyFlags(Flags, ETextureCacheFlags::ForceRebuild) && Settings.bHasEditorOnlyData)
+		{
+			ActiveRequest.store(&ShippingRequest, std::memory_order_relaxed);
+			FTextureBuildSettings ShippingSettings = Settings;
+			ShippingSettings.bHasEditorOnlyData = false;
+			FBuildDefinition ShippingDefinition = CreateDefinition(Build, Texture, TexturePath, FunctionName, KeySuffix, ShippingSettings);
+			ShippingRequest = BuildSession.Get().Build(ShippingDefinition, EBuildPolicy::Cache, Priority,
+				[this, Definition = MoveTemp(Definition), Flags](FBuildCompleteParams&& Params)
+				{
+					switch (Params.Status)
+					{
+					default:
+					case EStatus::Ok:
+						return EndBuild(MoveTemp(Params.Output), Params.BuildStatus);
+					case EStatus::Error:
+						return BeginBuild(Definition, Flags);
+					}
+				});
+		}
+		else
+		{
+			ActiveRequest.store(&Request, std::memory_order_relaxed);
+			BeginBuild(Definition, Flags);
+		}
+	}
+
+	static UE::DerivedData::FBuildDefinition CreateDefinition(
+		UE::DerivedData::IBuild& Build,
+		UTexture& Texture,
+		FStringView TexturePath,
+		FStringView FunctionName,
+		const FString& KeySuffix,
+		const FTextureBuildSettings& Settings)
+	{
+		UE::DerivedData::FBuildDefinitionBuilder DefinitionBuilder = Build.CreateDefinition(TexturePath, FunctionName);
+		DefinitionBuilder.AddConstant(TEXT("Settings"_SV),
+			SaveTextureBuildSettings(KeySuffix, Texture, Settings, 0, NUM_INLINE_DERIVED_MIPS));
+		DefinitionBuilder.AddInputBulkData(TEXT("Source"_SV), Texture.Source.GetPersistentId());
+		if (Texture.CompositeTexture)
+		{
+			DefinitionBuilder.AddInputBulkData(TEXT("CompositeSource"_SV), Texture.CompositeTexture->Source.GetPersistentId());
+		}
+		return DefinitionBuilder.Build();
+	}
+
+	void BeginBuild(const UE::DerivedData::FBuildDefinition& Definition, ETextureCacheFlags Flags)
+	{
+		using namespace UE::DerivedData;
+		EBuildPolicy BuildPolicy = EBuildPolicy::Default;
+		if (EnumHasAnyFlags(Flags, ETextureCacheFlags::ForceRebuild))
+		{
+			BuildPolicy &= ~EBuildPolicy::CacheQuery;
+		}
+		FRequest* Active = ActiveRequest.load(std::memory_order_relaxed);
+		Request = BuildSession.Get().Build(Definition, BuildPolicy, Priority,
+			[this](FBuildCompleteParams&& Params) { EndBuild(MoveTemp(Params.Output), Params.BuildStatus); });
+		ActiveRequest.compare_exchange_strong(Active, &Request, std::memory_order_release, std::memory_order_relaxed);
+	}
+
+	void EndBuild(UE::DerivedData::FBuildOutput&& Output, UE::DerivedData::EBuildStatus Status)
+	{
+		using namespace UE::DerivedData;
+		bCacheHit = EnumHasAnyFlags(Status, EBuildStatus::CacheQueryHit);
+		BuildOutputSize = Algo::TransformAccumulate(Output.GetPayloads(),
+			[](const FPayload& Payload) { return Payload.GetData().GetRawSize(); }, uint64(0));
+		WriteDerivedData(MoveTemp(Output));
+		ActiveRequest.store(nullptr, std::memory_order_relaxed);
+	}
+
+	void Finalize(bool& bOutFoundInCache, uint64& OutProcessedByteCount) final
+	{
+		bOutFoundInCache = bCacheHit;
+		OutProcessedByteCount = BuildOutputSize;
+	}
+
+	EQueuedWorkPriority GetPriority() const final
+	{
+		return ConvertPriority(Priority);
+	}
+
+	bool SetPriority(EQueuedWorkPriority QueuedWorkPriority) final
+	{
+		using namespace UE::DerivedData;
+		const EPriority NewPriority = ConvertPriority(QueuedWorkPriority);
+		if (Priority < EPriority::Blocking)
+		{
+			Priority = NewPriority;
+			if (FRequest* Active = ActiveRequest.load(std::memory_order_acquire))
+			{
+				Active->SetPriority(NewPriority);
+			}
+		}
+		return true;
+	}
+
+	bool Cancel() final
+	{
+		using namespace UE::DerivedData;
+		while (FRequest* Active = ActiveRequest.load(std::memory_order_acquire))
+		{
+			Active->Cancel();
+		}
+		return true;
+	}
+
+	void Wait() final
+	{
+		using namespace UE::DerivedData;
+		while (FRequest* Active = ActiveRequest.load(std::memory_order_acquire))
+		{
+			Active->Wait();
+		}
+	}
+
+	bool WaitWithTimeout(float TimeLimitSeconds) final
+	{
+		const double TimeLimit = FPlatformTime::Seconds() + TimeLimitSeconds;
+		if (Poll())
+		{
+			return true;
+		}
+		do
+		{
+			FPlatformProcess::Sleep(0.005);
+			if (Poll())
+			{
+				return true;
+			}
+		}
+		while (FPlatformTime::Seconds() < TimeLimit);
+		return false;
+	}
+
+	bool Poll() const final
+	{
+		using namespace UE::DerivedData;
+		while (FRequest* Active = ActiveRequest.load(std::memory_order_acquire))
+		{
+			if (!Active->Poll())
+			{
+				return false;
+			}
+		}
+		return true;
+	}
+
+private:
+	void WriteDerivedData(UE::DerivedData::FBuildOutput&& Output)
+	{
+		using namespace UE::DerivedData;
+
+		Output.IterateDiagnostics([](const FBuildDiagnostic& Diagnostic)
+		{
+			if (Diagnostic.Level == EBuildDiagnosticLevel::Error)
+			{
+				UE_LOG(LogTexture, Warning, TEXT("[Build Error] %.*s: %.*s"),
+					Diagnostic.Category.Len(), Diagnostic.Category.GetData(),
+					Diagnostic.Message.Len(), Diagnostic.Message.GetData());
+			}
+			else
+			{
+				UE_LOG(LogTexture, Warning, TEXT("[Build Warning] %.*s: %.*s"),
+					Diagnostic.Category.Len(), Diagnostic.Category.GetData(),
+					Diagnostic.Message.Len(), Diagnostic.Message.GetData());
+			}
+		});
+
+		if (Output.HasError())
+		{
+			UE_LOG(LogTexture, Warning, TEXT("Failed to build derived data for build of '%s' by %s."),
+				*WriteToString<128>(Output.GetName()), *WriteToString<32>(Output.GetFunction()));
+			return;
+		}
+
+		if (const FPayload& Payload = Output.GetPayload(FPayloadId::FromName(TEXT("Texture"))))
+		{
+			FSharedBuffer Data = Payload.GetData().Decompress();
+			FMemoryReaderView Ar(Data, /*bIsPersistent*/ true);
+			DerivedData.Serialize(Ar, nullptr);
+		}
+		else
+		{
+			UE_LOG(LogTexture, Warning, TEXT("Missing output for build of '%s' by %s."),
+				*WriteToString<128>(Output.GetName()), *WriteToString<32>(Output.GetFunction()));
+			return;
+		}
+
+		for (int32 MipIndex = 0, MipCount = DerivedData.Mips.Num(); MipIndex < MipCount; ++MipIndex)
+		{
+			FTexture2DMipMap& Mip = DerivedData.Mips[MipIndex];
+			if (Mip.DerivedDataKey.IsEmpty())
+			{
+				break;
+			}
+
+			if (const FPayload& MipPayload = Output.GetPayload(FPayloadId::FromName(WriteToString<8>(TEXT("Mip"), MipIndex))))
+			{
+				if (bInlineMips && MipIndex >= FirstMipToLoad)
+				{
+					FSharedBuffer MipData = MipPayload.GetData().Decompress();
+					FMemoryReaderView Ar(MipData, /*bIsPersistent=*/ true);
+
+					int32 MipSize = 0;
+					Ar << MipSize;
+
+					Mip.BulkData.Lock(LOCK_READ_WRITE);
+					void* MipAllocData = Mip.BulkData.Realloc(MipSize);
+					Ar.Serialize(MipAllocData, MipSize);
+					Mip.BulkData.Unlock();
+					Mip.DerivedDataKey.Empty();
+				}
+			}
+			else
+			{
+				UE_LOG(LogTexture, Warning, TEXT("Missing output for mip %d for build of '%s' by %s."),
+					MipIndex, *WriteToString<128>(Output.GetName()), *WriteToString<32>(Output.GetFunction()));
+				return;
+			}
+		}
+	}
+
+	static UE::DerivedData::EPriority ConvertPriority(EQueuedWorkPriority SourcePriority)
+	{
+		using namespace UE::DerivedData;
+		switch (SourcePriority)
+		{
+		case EQueuedWorkPriority::Lowest:  return EPriority::Lowest;
+		case EQueuedWorkPriority::Low:     return EPriority::Low;
+		case EQueuedWorkPriority::Normal:  return EPriority::Normal;
+		case EQueuedWorkPriority::High:    return EPriority::High;
+		case EQueuedWorkPriority::Highest: return EPriority::Highest;
+		default:                           return EPriority::Normal;
+		}
+	}
+
+	static EQueuedWorkPriority ConvertPriority(UE::DerivedData::EPriority SourcePriority)
+	{
+		using namespace UE::DerivedData;
+		switch (SourcePriority)
+		{
+		case EPriority::Lowest:   return EQueuedWorkPriority::Lowest;
+		case EPriority::Low:      return EQueuedWorkPriority::Low;
+		case EPriority::Normal:   return EQueuedWorkPriority::Normal;
+		case EPriority::High:     return EQueuedWorkPriority::High;
+		case EPriority::Highest:  return EQueuedWorkPriority::Highest;
+		case EPriority::Blocking: return EQueuedWorkPriority::Highest;
+		default:                  return EQueuedWorkPriority::Normal;
+		}
+	}
+
+	static bool LoadModules()
+	{
+		FModuleManager::LoadModuleChecked<IImageWrapperModule>(FName("ImageWrapper"));
+		FModuleManager::LoadModuleChecked<ITextureCompressorModule>(TEXTURE_COMPRESSOR_MODULENAME);
+		return true;
+	}
+
+	std::atomic<UE::DerivedData::FRequest*> ActiveRequest;
+	FTexturePlatformData& DerivedData;
+	UE::DerivedData::FRequest Request;
+	UE::DerivedData::FRequest ShippingRequest;
+	UE::DerivedData::FOptionalBuildSession BuildSession;
+	UE::DerivedData::EPriority Priority;
+	bool bCacheHit;
+	bool bInlineMips;
+	int32 FirstMipToLoad;
+	uint64 BuildOutputSize = 0;
+	UE::TextureDerivedData::FTextureBuildInputResolver InputResolver;
+	FRWLock Lock;
+};
+
+FTextureAsyncCacheDerivedDataTask* CreateTextureBuildTask(
+	UTexture& Texture,
+	FTexturePlatformData& DerivedData,
+	const FTextureBuildSettings& Settings,
+	EQueuedWorkPriority Priority,
+	ETextureCacheFlags Flags)
+{
+	TStringBuilder<64> FunctionName;
+	if (TryFindTextureBuildFunction(FunctionName, Settings))
+	{
+		return new FTextureBuildTask(Texture, FunctionName, DerivedData, Settings, Priority, Flags);
+	}
+	return nullptr;
+}
+
 #endif // WITH_EDITOR
diff --git a/Engine/Source/Runtime/Engine/Private/TextureDerivedDataTask.h b/Engine/Source/Runtime/Engine/Private/TextureDerivedDataTask.h
--- a/Engine/Source/Runtime/Engine/Private/TextureDerivedDataTask.h
+++ b/Engine/Source/Runtime/Engine/Private/TextureDerivedDataTask.h
@@ -187,7 +187,7 @@ struct FTextureAsyncCacheDerivedDataTask
 	virtual ~FTextureAsyncCacheDerivedDataTask() = default;
 	virtual void Finalize(bool& bOutFoundInCache, uint64& OutProcessedByteCount) = 0;
 	virtual EQueuedWorkPriority GetPriority() const = 0;
-	virtual bool SetPriority(EQueuedWorkPriority InQueuedWorkPriority) = 0;
+	virtual bool SetPriority(EQueuedWorkPriority QueuedWorkPriority) = 0;
 	virtual bool Cancel() = 0;
 	virtual void Wait() = 0;
 	virtual bool WaitWithTimeout(float TimeLimitSeconds) = 0;
@@ -228,9 +228,9 @@ public:
 		return FAsyncTask<FTextureCacheDerivedDataWorker>::GetPriority();
 	}
 
-	bool SetPriority(EQueuedWorkPriority InQueuedWorkPriority) final
+	bool SetPriority(EQueuedWorkPriority QueuedWorkPriority) final
 	{
-		return FAsyncTask<FTextureCacheDerivedDataWorker>::Reschedule(QueuedPool, InQueuedWorkPriority);
+		return FAsyncTask<FTextureCacheDerivedDataWorker>::Reschedule(QueuedPool, QueuedWorkPriority);
 	}
 
 	bool Cancel() final
@@ -257,4 +257,11 @@ private:
 	FQueuedThreadPool* QueuedPool;
 };
 
+FTextureAsyncCacheDerivedDataTask* CreateTextureBuildTask(
+	UTexture& Texture,
+	FTexturePlatformData& DerivedData,
+	const FTextureBuildSettings& Settings,
+	EQueuedWorkPriority Priority,
+	ETextureCacheFlags Flags);
+
 #endif // WITH_EDITOR
