Material autogen derivative code matches coercion logic of legacy code for min/max operations

diff --git a/Engine/Source/Runtime/Engine/Private/Materials/HLSLMaterialDerivativeAutogen.cpp b/Engine/Source/Runtime/Engine/Private/Materials/HLSLMaterialDerivativeAutogen.cpp
--- a/Engine/Source/Runtime/Engine/Private/Materials/HLSLMaterialDerivativeAutogen.cpp
+++ b/Engine/Source/Runtime/Engine/Private/Materials/HLSLMaterialDerivativeAutogen.cpp
@@ -222,41 +222,7 @@ FString FMaterialDerivativeAutogen::CoerceValueRaw(FHLSLMaterialTranslator& Tran
 		Ret = Ret + TEXT(".Value");
 	}
 
-	if (SrcType == DstType)
-	{
-		// no op
-	}
-	else
-	{
-		if (SrcType == 0)
-		{
-			if (DstType == 0)
-			{
-				Ret = FString::Printf( TEXT("MaterialFloat(%s)"), *Ret);
-			}
-			else if (DstType == 1)
-			{
-				Ret = FString::Printf( TEXT("MaterialFloat2(%s,%s)"), *Ret, *Ret);
-			}
-			else if (DstType == 2)
-			{
-				Ret = FString::Printf( TEXT("MaterialFloat3(%s,%s,%s)"), *Ret, *Ret, *Ret);
-			}
-			else if (DstType == 3)
-			{
-				Ret = FString::Printf( TEXT("MaterialFloat4(%s,%s,%s,%s)"), *Ret, *Ret, *Ret, *Ret);
-			}
-			else
-			{
-				check(0);
-			}
-		}
-		else
-		{
-			Translator.Errorf(TEXT("Coercion failed: %s: %s -> %s"), *Token, GetFloatVectorName(SrcType), GetFloatVectorName(DstType));
-			return TEXT("");
-		}
-	}
+	Ret = CoerceFloat(*Ret, DstType, SrcType);
 	return Ret;
 }
 
@@ -372,8 +338,6 @@ int32 FMaterialDerivativeAutogen::GetFunc2ReturnNumComponents(int32 LhsTypeIndex
 	case EFunc2::Mul:
 	case EFunc2::Div:
 	case EFunc2::Fmod:
-	case EFunc2::Max:
-	case EFunc2::Min:
 	case EFunc2::Pow:
 	case EFunc2::PowPositiveClamped:
 	case EFunc2::Atan2:
@@ -385,6 +349,13 @@ int32 FMaterialDerivativeAutogen::GetFunc2ReturnNumComponents(int32 LhsTypeIndex
 			DstTypeIndex = FMath::Max<int32>(LhsTypeIndex, RhsTypeIndex);
 		}
 		break;
+	case EFunc2::Max:
+	case EFunc2::Min:
+		// Min/Max should probably use the same logic as other binary operations
+		// However, legacy HLSLMaterialTranslator::Min/Max simply coerce the Rhs value to the type of the Lhs value
+		// Need to preserve this behavior here (may want to wrap this behind a legacy flag/CVAR at some point)
+		DstTypeIndex = LhsTypeIndex;
+		break;
 	case EFunc2::Dot:
 		DstTypeIndex = 0;
 		break;
@@ -609,11 +580,22 @@ int32 FMaterialDerivativeAutogen::GenerateExpressionFunc2(FHLSLMaterialTranslato
 	const FDerivInfo LhsDerivInfo = Translator.GetDerivInfo(LhsCode);
 	const FDerivInfo RhsDerivInfo = Translator.GetDerivInfo(RhsCode);
 
-	int32 IntermediaryTypeIndex = FMath::Max<int32>(LhsDerivInfo.TypeIndex, RhsDerivInfo.TypeIndex);
-
-	if (Op == EFunc2::Cross)
+	int32 IntermediaryTypeIndex;
+	switch (Op)
 	{
+	case EFunc2::Cross:
 		IntermediaryTypeIndex = 2;
+		break;
+	case EFunc2::Min:
+	case EFunc2::Max:
+		// Min/Max should probably use the same logic as other binary operations
+		// However, legacy HLSLMaterialTranslator::Min/Max simply coerce the Rhs value to the type of the Lhs value
+		// Need to preserve this behavior here (may want to wrap this behind a legacy flag/CVAR at some point)
+		IntermediaryTypeIndex = LhsDerivInfo.TypeIndex;
+		break;
+	default:
+		IntermediaryTypeIndex = FMath::Max<int32>(LhsDerivInfo.TypeIndex, RhsDerivInfo.TypeIndex);
+		break;
 	}
 
 	if (Op == EFunc2::Fmod)
