Improve the comments for Data Assets to better explain the distinction between asset instances and blueprint classes

diff --git a/Engine/Source/Editor/UnrealEd/Private/Factories/EditorFactories.cpp b/Engine/Source/Editor/UnrealEd/Private/Factories/EditorFactories.cpp
--- a/Engine/Source/Editor/UnrealEd/Private/Factories/EditorFactories.cpp
+++ b/Engine/Source/Editor/UnrealEd/Private/Factories/EditorFactories.cpp
@@ -7225,7 +7225,7 @@ bool UDataAssetFactory::ConfigureProperties()
 	Filter->DisallowedClassFlags = CLASS_Abstract | CLASS_Deprecated | CLASS_NewerVersionExists | CLASS_HideDropDown;
 	Filter->AllowedChildrenOfClasses.Add(UDataAsset::StaticClass());
 
-	const FText TitleText = LOCTEXT("CreateDataAssetOptions", "Pick Data Asset Class");
+	const FText TitleText = LOCTEXT("CreateDataAssetOptions", "Pick Class For Data Asset Instance");
 	UClass* ChosenClass = nullptr;
 	const bool bPressedOk = SClassPickerDialog::PickClass(TitleText, Options, ChosenClass, UDataAsset::StaticClass());
 
diff --git a/Engine/Source/Runtime/Engine/Classes/Engine/DataAsset.h b/Engine/Source/Runtime/Engine/Classes/Engine/DataAsset.h
--- a/Engine/Source/Runtime/Engine/Classes/Engine/DataAsset.h
+++ b/Engine/Source/Runtime/Engine/Classes/Engine/DataAsset.h
@@ -10,7 +10,9 @@
 #include "DataAsset.generated.h"
 
 /**
- * Base class for a simple asset containing data. The editor will list this in the content browser if you inherit from this class
+ * Create a simple asset that stores data related to a particular system in an instance of this class.
+ * Assets can be made in the Content Browser using any native class that inherits from this.
+ * If you want data inheritance or a complicated hierarchy, Data Only Blueprint Classes should be created instead.
  */
 UCLASS(abstract, MinimalAPI)
 class UDataAsset : public UObject
@@ -28,12 +30,12 @@ private:
 };
 
 /**
- * A DataAsset that implements GetPrimaryAssetId and has asset bundle support, which makes it something that can be manually loaded/unloaded from the AssetManager
- * Making blueprint subclasses of this is useful because you can make blueprint-only primary asset types
- * Blueprint subclasses will end up with a PrimaryAssetType equal to the name of the first native class found going up the hierarchy, or the top level blueprint class that directly subclasses this
- * IE, if you have PrimaryDataAsset -> ParentNativeClass -> ChildNativeClass -> BlueprintAsset the type will be ChildNativeClass
- * Whereas if you have PrimaryDataAsset -> ParentBlueprintClass -> ChildBlueprintClass the type will be ParentBlueprintClass
- * To override this behavior, override GetPrimaryAssetId in your native class
+ * A DataAsset that implements GetPrimaryAssetId and has asset bundle support, which allows it to be manually loaded/unloaded from the AssetManager.
+ * Instances of native subclasses can be created directly as Data Assets and will use the name of the native class as the PrimaryAssetType.
+ * Or, blueprint subclasses can be created with a PrimaryAssetType equal to the name of the first native class going up the hierarchy, or the highest level blueprint class.
+ * IE, if you have UPrimaryDataAsset -> UParentNativeClass -> UChildNativeClass -> DataOnlyBlueprintClass the type will be ChildNativeClass.
+ * Whereas if you have UPrimaryDataAsset -> ParentBlueprintClass -> DataOnlyBlueprintClass the type will be ParentBlueprintClass.
+ * To change this behavior, override GetPrimaryAssetId in your native class or copy those functions into a different native base class.
  */
 UCLASS(abstract, MinimalAPI, Blueprintable)
 class UPrimaryDataAsset : public UDataAsset
