Manually merging 16897423  - adding new csv stat for remapped objects

diff --git a/Engine/Source/Runtime/Engine/Private/NetDriver.cpp b/Engine/Source/Runtime/Engine/Private/NetDriver.cpp
--- a/Engine/Source/Runtime/Engine/Private/NetDriver.cpp
+++ b/Engine/Source/Runtime/Engine/Private/NetDriver.cpp
@@ -152,6 +152,8 @@ DECLARE_CYCLE_STAT(TEXT("NetDriver TickFlush GatherStatsPerfCounters"), STAT_Net
 
 DEFINE_LOG_CATEGORY_STATIC(LogNetSyncLoads, Log, All);
 
+CSV_DEFINE_CATEGORY(Networking, true);
+
 int32 GNumSaturatedConnections; // Counter for how many connections are skipped/early out due to bandwidth saturation
 int32 GNumSharedSerializationHit;
 int32 GNumSharedSerializationMiss;
@@ -1120,152 +1122,153 @@ void UNetDriver::TickFlush(float DeltaSeconds)
 		DrawNetDriverDebug();
 	}
 
-	if ( CVarOptimizedRemapping.GetValueOnAnyThread() && GuidCache.IsValid() )
 	{
-		SCOPE_CYCLE_COUNTER( STAT_NetUpdateUnmappedObjectsTime );
-
-		// Go over recently imported network guids, and see if there are any replicators that need to map them
-		TSet< FNetworkGUID >& ImportedNetGuids = GuidCache->ImportedNetGuids;
-		TMap< FNetworkGUID, TSet< FNetworkGUID > >& PendingOuterNetGuids = GuidCache->PendingOuterNetGuids;
-
-		SET_DWORD_STAT( STAT_ImportedNetGuids, ImportedNetGuids.Num() );
-		SET_DWORD_STAT( STAT_PendingOuterNetGuids, PendingOuterNetGuids.Num() );
-		SET_DWORD_STAT( STAT_UnmappedReplicators, UnmappedReplicators.Num() );
-
-		TSet< FObjectReplicator* > ForceUpdateReplicators;
+		CSV_SCOPED_TIMING_STAT(Networking, UpdateUnmappedObjects);
+		SCOPE_CYCLE_COUNTER(STAT_NetUpdateUnmappedObjectsTime);
 
-		for (FObjectReplicator* Replicator : UnmappedReplicators)
+		if (CVarOptimizedRemapping.GetValueOnAnyThread() && GuidCache.IsValid())
 		{
-			if (Replicator->bForceUpdateUnmapped)
-			{
-				Replicator->bForceUpdateUnmapped = false;
-				ForceUpdateReplicators.Add(Replicator);
-			}
-		}
+			// Go over recently imported network guids, and see if there are any replicators that need to map them
+			TSet< FNetworkGUID >& ImportedNetGuids = GuidCache->ImportedNetGuids;
+			TMap< FNetworkGUID, TSet< FNetworkGUID > >& PendingOuterNetGuids = GuidCache->PendingOuterNetGuids;
 
-		if ( ImportedNetGuids.Num() || ForceUpdateReplicators.Num() )
-		{
-			int32 NumRemapTests = 0;
-			int32 MaxRemaps = IsServer() ? 0 : CVarMaxClientGuidRemaps.GetValueOnAnyThread();
-			bool bFilterGuidRemapping = (CVarFilterGuidRemapping.GetValueOnAnyThread() > 0);
+			SET_DWORD_STAT(STAT_ImportedNetGuids, ImportedNetGuids.Num());
+			SET_DWORD_STAT(STAT_PendingOuterNetGuids, PendingOuterNetGuids.Num());
+			SET_DWORD_STAT(STAT_UnmappedReplicators, UnmappedReplicators.Num());
 
-			TArray< FNetworkGUID > UnmappedGuids;
-			TArray< FNetworkGUID > NewlyMappedGuids;
+			TSet< FObjectReplicator* > ForceUpdateReplicators;
 
-			for ( auto It = ImportedNetGuids.CreateIterator(); It; ++It )
+			for (FObjectReplicator* Replicator : UnmappedReplicators)
 			{
-				const FNetworkGUID NetworkGuid = *It;
-				bool bMappedOrBroken = false;
-
-				if ( GuidCache->GetObjectFromNetGUID( NetworkGuid, false ) != nullptr )
+				if (Replicator->bForceUpdateUnmapped)
 				{
-					NewlyMappedGuids.Add( NetworkGuid );
-					bMappedOrBroken = true;
+					Replicator->bForceUpdateUnmapped = false;
+					ForceUpdateReplicators.Add(Replicator);
 				}
+			}
 
-				if ( GuidCache->IsGUIDBroken( NetworkGuid, false ) )
-				{
-					bMappedOrBroken = true;
-				}
+			if (ImportedNetGuids.Num() || ForceUpdateReplicators.Num())
+			{
+				int32 NumRemapTests = 0;
+				int32 MaxRemaps = IsServer() ? 0 : CVarMaxClientGuidRemaps.GetValueOnAnyThread();
+				bool bFilterGuidRemapping = (CVarFilterGuidRemapping.GetValueOnAnyThread() > 0);
 
-				It.RemoveCurrent();
+				TArray< FNetworkGUID > UnmappedGuids;
+				TArray< FNetworkGUID > NewlyMappedGuids;
 
-				if (!bMappedOrBroken)
+				for (auto It = ImportedNetGuids.CreateIterator(); It; ++It)
 				{
-					FNetworkGUID OuterGUID = GuidCache->GetOuterNetGUID(NetworkGuid);
+					const FNetworkGUID NetworkGuid = *It;
+					bool bMappedOrBroken = false;
 
-					// we're missing the outer, stop checking until we map it
-					if (bFilterGuidRemapping && OuterGUID.IsValid() && !OuterGUID.IsDefault() && !GuidCache->IsGUIDLoaded(OuterGUID) && !GuidCache->IsGUIDPending(OuterGUID))
+					if (GuidCache->GetObjectFromNetGUID(NetworkGuid, false) != nullptr)
 					{
-						UE_LOG(LogNetPackageMap, Log, TEXT("Missing outer (%s) for unmapped guid (%s), marking pending"), *GuidCache->FullNetGUIDPath(OuterGUID), *GuidCache->FullNetGUIDPath(NetworkGuid));
+						NewlyMappedGuids.Add(NetworkGuid);
+						bMappedOrBroken = true;
+					}
 
-						TSet<FNetworkGUID>& PendingGuids = PendingOuterNetGuids.FindOrAdd(OuterGUID);
-						PendingGuids.Add(NetworkGuid);
+					if (GuidCache->IsGUIDBroken(NetworkGuid, false))
+					{
+						bMappedOrBroken = true;
 					}
-					else
+
+					It.RemoveCurrent();
+
+					if (!bMappedOrBroken)
 					{
-						if (ensure(NetworkGuid.IsValid()))
+						FNetworkGUID OuterGUID = GuidCache->GetOuterNetGUID(NetworkGuid);
+
+						// we're missing the outer, stop checking until we map it
+						if (bFilterGuidRemapping && OuterGUID.IsValid() && !OuterGUID.IsDefault() && !GuidCache->IsGUIDLoaded(OuterGUID) && !GuidCache->IsGUIDPending(OuterGUID))
 						{
-							UnmappedGuids.Add(NetworkGuid);
+							UE_LOG(LogNetPackageMap, Log, TEXT("Missing outer (%s) for unmapped guid (%s), marking pending"), *GuidCache->FullNetGUIDPath(OuterGUID), *GuidCache->FullNetGUIDPath(NetworkGuid));
+
+							TSet<FNetworkGUID>& PendingGuids = PendingOuterNetGuids.FindOrAdd(OuterGUID);
+							PendingGuids.Add(NetworkGuid);
+						}
+						else
+						{
+							if (ensure(NetworkGuid.IsValid()))
+							{
+								UnmappedGuids.Add(NetworkGuid);
+							}
 						}
 					}
-				}
 
-				++NumRemapTests;
+					++NumRemapTests;
 
- 				if ((MaxRemaps > 0) && (NumRemapTests >= MaxRemaps))
-				{
-					break;
+					if ((MaxRemaps > 0) && (NumRemapTests >= MaxRemaps))
+					{
+						break;
+					}
 				}
-			}
 
-			// attempt to resolve dependent guids next tick (outer is now mapped)
-			for ( const FNetworkGUID& NetGuid : NewlyMappedGuids )
-			{
-				TSet< FNetworkGUID >* DependentGuids = PendingOuterNetGuids.Find( NetGuid );
-				if ( DependentGuids )
+				// attempt to resolve dependent guids next tick (outer is now mapped)
+				for (const FNetworkGUID& NetGuid : NewlyMappedGuids)
 				{
-					UE_LOG(LogNetPackageMap, Log, TEXT("Newly mapped outer (%s) removing from pending"), *GuidCache->FullNetGUIDPath(NetGuid));
+					TSet< FNetworkGUID >* DependentGuids = PendingOuterNetGuids.Find(NetGuid);
+					if (DependentGuids)
+					{
+						UE_LOG(LogNetPackageMap, Log, TEXT("Newly mapped outer (%s) removing from pending"), *GuidCache->FullNetGUIDPath(NetGuid));
 
-					ImportedNetGuids.Append( *DependentGuids );
-					PendingOuterNetGuids.Remove( NetGuid );
+						ImportedNetGuids.Append(*DependentGuids);
+						PendingOuterNetGuids.Remove(NetGuid);
+					}
 				}
-			}
-
-			// any tested guids that could not yet be mapped are added back to the list
-			if (UnmappedGuids.Num())
-			{
-				ImportedNetGuids.CompactStable();
-				ImportedNetGuids.Append(UnmappedGuids);
-			}
 
-			if ( NewlyMappedGuids.Num() || ForceUpdateReplicators.Num() )
-			{
-				TSet< FObjectReplicator* > AllReplicators = ForceUpdateReplicators;
+				// any tested guids that could not yet be mapped are added back to the list
+				if (UnmappedGuids.Num())
+				{
+					ImportedNetGuids.CompactStable();
+					ImportedNetGuids.Append(UnmappedGuids);
+				}
 
-				for ( const FNetworkGUID& NetGuid : NewlyMappedGuids )
+				if (NewlyMappedGuids.Num() || ForceUpdateReplicators.Num())
 				{
-					TSet< FObjectReplicator* >* Replicators = GuidToReplicatorMap.Find( NetGuid );
+					TSet< FObjectReplicator* > AllReplicators = ForceUpdateReplicators;
 
-					if ( Replicators )
+					for (const FNetworkGUID& NetGuid : NewlyMappedGuids)
 					{
-						AllReplicators.Append( *Replicators );
+						TSet< FObjectReplicator* >* Replicators = GuidToReplicatorMap.Find(NetGuid);
+
+						if (Replicators)
+						{
+							AllReplicators.Append(*Replicators);
+						}
 					}
-				}
 
-				for ( FObjectReplicator* Replicator : AllReplicators )
-				{
-					if ( UnmappedReplicators.Contains( Replicator ) )
+					for (FObjectReplicator* Replicator : AllReplicators)
 					{
-						bool bHasMoreUnmapped = false;
-						Replicator->UpdateUnmappedObjects( bHasMoreUnmapped );
-
-						if ( !bHasMoreUnmapped )
+						if (UnmappedReplicators.Contains(Replicator))
 						{
-							UnmappedReplicators.Remove( Replicator );
+							bool bHasMoreUnmapped = false;
+							Replicator->UpdateUnmappedObjects(bHasMoreUnmapped);
+
+							if (!bHasMoreUnmapped)
+							{
+								UnmappedReplicators.Remove(Replicator);
+							}
 						}
 					}
 				}
 			}
 		}
-	}
-	else
-	{
-		SCOPE_CYCLE_COUNTER( STAT_NetUpdateUnmappedObjectsTime );
-
-		// Update properties that are unmapped, try to hook up the object pointers if they exist now
-		for ( auto It = UnmappedReplicators.CreateIterator(); It; ++It )
+		else
 		{
-			FObjectReplicator* Replicator = *It;
+			// Update properties that are unmapped, try to hook up the object pointers if they exist now
+			for (auto It = UnmappedReplicators.CreateIterator(); It; ++It)
+			{
+				FObjectReplicator* Replicator = *It;
 
-			bool bHasMoreUnmapped = false;
+				bool bHasMoreUnmapped = false;
 
-			Replicator->UpdateUnmappedObjects( bHasMoreUnmapped );
+				Replicator->UpdateUnmappedObjects(bHasMoreUnmapped);
 
-			if ( !bHasMoreUnmapped )
-			{
-				// If there are no more unmapped objects, we can also stop checking
-				It.RemoveCurrent();
+				if (!bHasMoreUnmapped)
+				{
+					// If there are no more unmapped objects, we can also stop checking
+					It.RemoveCurrent();
+				}
 			}
 		}
 	}
