Fixed crash in ray tracing with invalid static mesh proxies

diff --git a/Engine/Source/Runtime/Engine/Private/StaticMeshRender.cpp b/Engine/Source/Runtime/Engine/Private/StaticMeshRender.cpp
--- a/Engine/Source/Runtime/Engine/Private/StaticMeshRender.cpp
+++ b/Engine/Source/Runtime/Engine/Private/StaticMeshRender.cpp
@@ -1699,11 +1699,6 @@ void FStaticMeshSceneProxy::GetDynamicRayTracingInstances(FRayTracingMaterialGat
 	const uint32 LODIndex = FMath::Max(GetLOD(Context.ReferenceView), (int32)GetCurrentFirstLODIdx_RenderThread());
 	const FStaticMeshLODResources& LODModel = RenderData->LODResources[LODIndex];
 
-	if (LODModel.GetNumVertices() <= 0)
-	{
-		return;
-	}
-
 	bool bEvaluateWPO = CVarRayTracingStaticMeshesWPO.GetValueOnRenderThread() == 1;
 
 	if (bEvaluateWPO && CVarRayTracingStaticMeshesWPOCulling.GetValueOnRenderThread() > 0)
@@ -1720,6 +1715,11 @@ void FStaticMeshSceneProxy::GetDynamicRayTracingInstances(FRayTracingMaterialGat
 	}
 
 	FRayTracingGeometry& Geometry = bEvaluateWPO? DynamicRayTracingGeometries[LODIndex] : RenderData->LODResources[LODIndex].RayTracingGeometry;
+	
+	if (LODModel.GetNumVertices() <= 0 || Geometry.Initializer.TotalPrimitiveCount <= 0)
+	{
+		return;
+	}
 
 	// Early out for now if no valid RHI RT geometry yet (still pending build request)
 	// TODO: select different LOD if available
