Added command line arguments for 2 AudioPlatformSettings parameters
AudioCallbackBufferFrameSize and AudioNumBufferToEnqueue

diff --git a/Engine/Source/Runtime/Engine/Private/AudioDevice.cpp b/Engine/Source/Runtime/Engine/Private/AudioDevice.cpp
--- a/Engine/Source/Runtime/Engine/Private/AudioDevice.cpp
+++ b/Engine/Source/Runtime/Engine/Private/AudioDevice.cpp
@@ -360,6 +360,10 @@ bool FAudioDevice::Init(Audio::FDeviceId InDeviceID, int32 InMaxSources)
 	// Get a copy of the platform-specific settings (overridden by platforms)
 	PlatformSettings = GetPlatformSettings();
 
+	// Override platform settings with the command line
+	FParse::Value(FCommandLine::Get(), TEXT("-AudioCallbackBufferFrameSize="), PlatformSettings.CallbackBufferFrameSize);
+	FParse::Value(FCommandLine::Get(), TEXT("-AudioNumBuffersToEnqueue="), PlatformSettings.NumBuffers);
+
 	// MaxSources is the max value supplied to Init call (quality settings), unless overwritten by the platform settings.
 	// This does not have to be the minimum value in this case (nor is it desired, so platforms can potentially scale up)
 	// as the Sources array has yet to be initialized. If the cvar is largest, take that value to allow for testing
