Only prefetch the lowest mips during texture build/postload instead of prefetching all mips
  - This will still avoid blocking calls for high prio request on the render-thread while reducing Jupiter downloads significantly
  - 51.2GB -> 8.3GB downloaded from jupiter when loading P_Construct_WP with an empty local DDC

#rnx

diff --git a/Engine/Source/Runtime/Engine/Private/TextureDerivedData.cpp b/Engine/Source/Runtime/Engine/Private/TextureDerivedData.cpp
--- a/Engine/Source/Runtime/Engine/Private/TextureDerivedData.cpp
+++ b/Engine/Source/Runtime/Engine/Private/TextureDerivedData.cpp
@@ -1294,20 +1294,25 @@ bool FTexturePlatformData::AreDerivedMipsAvailable(FStringView Context) const
 			MipKeys.Add(Mip.DerivedDataKey);
 		}
 	}
-	if (IsInGameThread())
+
+	if (MipKeys.IsEmpty())
 	{
-		return GetDerivedDataCacheRef().AllCachedDataProbablyExists(MipKeys);
+		return true;
 	}
-	else
+
+	const bool bAreDerivedMipsAvailable = GetDerivedDataCacheRef().AllCachedDataProbablyExists(MipKeys);
+
+	// When using a shared DDC and performing async loading, 
+	// prefetch the lowest mip to avoid high prio request stalls from the render thread
+	if (bAreDerivedMipsAvailable && !IsInGameThread())
 	{
-		// When using a shared DDC and performing async loading, 
-		// prefetch the data so we don't pay the latency on the game-thread
-		// to fetch the asset from a remote location once
-		// the mips needs to be accessed
-		TRACE_CPUPROFILER_EVENT_SCOPE(PrefetchMips);
-		return GetDerivedDataCacheRef().TryToPrefetch(MipKeys, Context);
+		TRACE_CPUPROFILER_EVENT_SCOPE(PrefetchSmallestMip);
+		GetDerivedDataCacheRef().TryToPrefetch({ MipKeys.Last() }, Context);
 	}
+
+	return bAreDerivedMipsAvailable;
 }
+
 bool FTexturePlatformData::AreDerivedVTChunksAvailable(FStringView Context) const
 {
 	check(VTData);
@@ -1319,19 +1324,23 @@ bool FTexturePlatformData::AreDerivedVTChunksAvailable(FStringView Context) cons
 			ChunkKeys.Add(Chunk.DerivedDataKey);
 		}
 	}
-	if (IsInGameThread())
+
+	if (ChunkKeys.IsEmpty())
 	{
-		return GetDerivedDataCacheRef().AllCachedDataProbablyExists(ChunkKeys);
+		return true;
 	}
-	else
+
+	const bool bAreDerivedChunksAvailable = GetDerivedDataCacheRef().AllCachedDataProbablyExists(ChunkKeys);
+
+	// When using a shared DDC and performing async loading, 
+	// prefetch the lowest mip to avoid high prio request stalls from the render thread
+	if (bAreDerivedChunksAvailable && !IsInGameThread())
 	{
-		// When using a shared DDC and performing async loading, 
-		// prefetch the data so we don't pay the latency on the game-thread
-		// to fetch the asset from a remote location once
-		// the VT chunks needs to be accessed
-		TRACE_CPUPROFILER_EVENT_SCOPE(PrefetchDerivedVTChunks);
-		return GetDerivedDataCacheRef().TryToPrefetch(ChunkKeys, Context);
+		TRACE_CPUPROFILER_EVENT_SCOPE(PrefetchSmallestDerivedVTChunk);
+		GetDerivedDataCacheRef().TryToPrefetch({ ChunkKeys.Last() }, Context);
 	}
+
+	return bAreDerivedChunksAvailable;
 }
 
 bool FTexturePlatformData::AreDerivedMipsAvailable() const
