Prepare the lowest mip VT chunk during texture preparation phase to avoid high prio stalls on the render-thread
  - Saves 14s sequential on the render-thread after loading P_Construct_WP with an empty VT folder

#rnx

diff --git a/Engine/Source/Runtime/Engine/Private/TextureDerivedDataTask.cpp b/Engine/Source/Runtime/Engine/Private/TextureDerivedDataTask.cpp
--- a/Engine/Source/Runtime/Engine/Private/TextureDerivedDataTask.cpp
+++ b/Engine/Source/Runtime/Engine/Private/TextureDerivedDataTask.cpp
@@ -49,6 +49,7 @@
 #include "ProfilingDebugging/CookStats.h"
 #include "Serialization/BulkDataRegistry.h"
 #include "TextureDerivedDataBuildUtils.h"
+#include "VT/VirtualTextureChunkDDCCache.h"
 #include "VT/VirtualTextureDataBuilder.h"
 #include <atomic>
 
@@ -954,6 +955,12 @@ void FTextureCacheDerivedDataWorker::DoWork()
 	{
 		TextureData.ReleaseMemory();
 		CompositeTextureData.ReleaseMemory();
+
+		// Populate the VT DDC Cache now if we're asynchronously loading to avoid too many high prio/synchronous request on the render thread
+		if (!IsInGameThread() && DerivedData->VTData)
+		{
+			GetVirtualTextureChunkDDCCache()->MakeChunkAvailable_Concurrent(&DerivedData->VTData->Chunks.Last());
+		}
 	}
 }
 
diff --git a/Engine/Source/Runtime/Engine/Private/VT/VirtualTextureChunkDDCCache.cpp b/Engine/Source/Runtime/Engine/Private/VT/VirtualTextureChunkDDCCache.cpp
--- a/Engine/Source/Runtime/Engine/Private/VT/VirtualTextureChunkDDCCache.cpp
+++ b/Engine/Source/Runtime/Engine/Private/VT/VirtualTextureChunkDDCCache.cpp
@@ -290,5 +290,29 @@ bool FVirtualTextureChunkDDCCache::MakeChunkAvailable(struct FVirtualTextureData
 	return false;
 }
 
+void FVirtualTextureChunkDDCCache::MakeChunkAvailable_Concurrent(struct FVirtualTextureDataChunk* Chunk)
+{
+	TRACE_CPUPROFILER_EVENT_SCOPE(FVirtualTextureChunkDDCCache::MakeChunkAvailable_Concurrent);
+
+	const FString CachedFilePath = AbsoluteCachePath / Chunk->ShortDerivedDataKey;
+	const FString TempFilePath = AbsoluteCachePath / FGuid::NewGuid().ToString() + ".tmp";
+
+	if (Chunk->bCorruptDataLoadedFromDDC)
+	{
+		// We determined data loaded from DDC was corrupt...this means any file saved to VT DDC cache is also corrupt and can no longer be used
+		FPlatformFileManager::Get().GetPlatformFile().DeleteFile(*CachedFilePath);
+		Chunk->bCorruptDataLoadedFromDDC = false;
+		Chunk->bFileAvailableInVTDDCDache = false;
+	}
+
+	// File already available? 
+	if (Chunk->bFileAvailableInVTDDCDache)
+	{
+		return;
+	}
+
+	FAsyncFillCacheWorker SyncWorker(TempFilePath, CachedFilePath, Chunk);
+	SyncWorker.DoWork();
+}
 
 #endif
diff --git a/Engine/Source/Runtime/Engine/Private/VT/VirtualTextureChunkDDCCache.h b/Engine/Source/Runtime/Engine/Private/VT/VirtualTextureChunkDDCCache.h
--- a/Engine/Source/Runtime/Engine/Private/VT/VirtualTextureChunkDDCCache.h
+++ b/Engine/Source/Runtime/Engine/Private/VT/VirtualTextureChunkDDCCache.h
@@ -15,6 +15,7 @@ public:
 
 	bool MakeChunkAvailable(struct FVirtualTextureDataChunk* Chunk, bool bAsync, FString& OutChunkFileName, int64& OutOffsetInFile);
 
+	void MakeChunkAvailable_Concurrent(struct FVirtualTextureDataChunk* Chunk);
 private:
 	TArray<struct FVirtualTextureDataChunk*> ActiveChunks;
 	FString AbsoluteCachePath;
