Prefetch audio chunks asynchronously to avoid paying the high-latency remote check on every load
  - 2m35s -> 1m53s when loading P_Construct_WP

#rnx

diff --git a/Engine/Source/Runtime/Core/Public/Async/Async.h b/Engine/Source/Runtime/Core/Public/Async/Async.h
--- a/Engine/Source/Runtime/Core/Public/Async/Async.h
+++ b/Engine/Source/Runtime/Core/Public/Async/Async.h
@@ -385,14 +385,14 @@ auto Async(EAsyncExecution Execution, CallableType&& Callable, TUniqueFunction<v
  * @result A TFuture object that will receive the return value from the function.
  */
 template<typename CallableType>
-auto AsyncPool(FQueuedThreadPool& ThreadPool, CallableType&& Callable, TUniqueFunction<void()> CompletionCallback = nullptr) -> TFuture<decltype(Forward<CallableType>(Callable)())>
+auto AsyncPool(FQueuedThreadPool& ThreadPool, CallableType&& Callable, TUniqueFunction<void()> CompletionCallback = nullptr, EQueuedWorkPriority InQueuedWorkPriority = EQueuedWorkPriority::Normal) -> TFuture<decltype(Forward<CallableType>(Callable)())>
 {
 	using ResultType = decltype(Forward<CallableType>(Callable)());
 	TUniqueFunction<ResultType()> Function(Forward<CallableType>(Callable));
 	TPromise<ResultType> Promise(MoveTemp(CompletionCallback));
 	TFuture<ResultType> Future = Promise.GetFuture();
 
-	ThreadPool.AddQueuedWork(new TAsyncQueuedWork<ResultType>(MoveTemp(Function), MoveTemp(Promise)));
+	ThreadPool.AddQueuedWork(new TAsyncQueuedWork<ResultType>(MoveTemp(Function), MoveTemp(Promise)), InQueuedWorkPriority);
 
 	return MoveTemp(Future);
 }
diff --git a/Engine/Source/Runtime/Engine/Classes/Sound/SoundWave.h b/Engine/Source/Runtime/Engine/Classes/Sound/SoundWave.h
--- a/Engine/Source/Runtime/Engine/Classes/Sound/SoundWave.h
+++ b/Engine/Source/Runtime/Engine/Classes/Sound/SoundWave.h
@@ -140,7 +140,11 @@ struct ENGINE_API FStreamedAudioPlatformData
 	void FinishCache();
 	bool IsFinishedCache() const;
 	bool TryInlineChunkData();
+
+	UE_DEPRECATED(5.00, "Use AreDerivedChunksAvailable with the context instead.")
 	bool AreDerivedChunksAvailable() const;
+	
+	bool AreDerivedChunksAvailable(FStringView Context) const;
 #endif // WITH_EDITORONLY_DATA
 
 private:
diff --git a/Engine/Source/Runtime/Engine/Private/AudioDerivedData.cpp b/Engine/Source/Runtime/Engine/Private/AudioDerivedData.cpp
--- a/Engine/Source/Runtime/Engine/Private/AudioDerivedData.cpp
+++ b/Engine/Source/Runtime/Engine/Private/AudioDerivedData.cpp
@@ -11,6 +11,7 @@
 #include "Interfaces/ITargetPlatform.h"
 #include "Interfaces/ITargetPlatformManagerModule.h"
 #include "Sound/SoundWave.h"
+#include "Async/Async.h"
 #include "Sound/SoundEffectBase.h"
 #include "DerivedDataCacheInterface.h"
 #include "ProfilingDebugging/CookStats.h"
@@ -560,7 +561,7 @@ public:
 			}
 			else
 			{
-				bSucceeded = DerivedData->AreDerivedChunksAvailable();
+				bSucceeded = DerivedData->AreDerivedChunksAvailable(SoundWave.GetPathName());
 			}
 			bLoadedFromDDC = true;
 		}
@@ -882,7 +883,7 @@ int32 FStreamedAudioPlatformData::GetChunkFromDDC(int32 ChunkIndex, uint8** OutC
 }
 
 #if WITH_EDITORONLY_DATA
-bool FStreamedAudioPlatformData::AreDerivedChunksAvailable() const
+bool FStreamedAudioPlatformData::AreDerivedChunksAvailable(FStringView InContext) const
 {
 	TArray<FString> ChunkKeys;
 	for (const FStreamedAudioChunk& Chunk : Chunks)
@@ -892,8 +893,33 @@ bool FStreamedAudioPlatformData::AreDerivedChunksAvailable() const
 			ChunkKeys.Add(Chunk.DerivedDataKey);
 		}
 	}
-	return GetDerivedDataCacheRef().AllCachedDataProbablyExists(ChunkKeys);
+	
+	const bool bAllCachedDataProbablyExists = GetDerivedDataCacheRef().AllCachedDataProbablyExists(ChunkKeys);
+
+	// If this is called from the game thread, try to prefetch chunks locally on background thread
+	// to avoid doing high latency remote calls every time we reload this data.
+	if (bAllCachedDataProbablyExists && GDDCIOThreadPool && IsInGameThread())
+	{
+		FString Context{ InContext };
+		AsyncPool(
+			*GDDCIOThreadPool,
+			[ChunkKeys, Context]()
+			{
+				GetDerivedDataCacheRef().TryToPrefetch(ChunkKeys, Context);
+			},
+			nullptr,
+			EQueuedWorkPriority::Low
+		);
+	}
+
+	return bAllCachedDataProbablyExists;
+}
+
+bool FStreamedAudioPlatformData::AreDerivedChunksAvailable() const
+{
+	return AreDerivedChunksAvailable(TEXT("DerivedAudioChunks"_SV));
 }
+
 #endif // #if WITH_EDITORONLY_DATA
 
 void FStreamedAudioPlatformData::Serialize(FArchive& Ar, USoundWave* Owner)
