Add a metadata to serialize property in ComponentInstanceDataCache that are marked as "SaveToComponentInstanceDataCache". Use when the property is not editable or interp.

diff --git a/Engine/Source/Runtime/Engine/Classes/Engine/EngineTypes.h b/Engine/Source/Runtime/Engine/Classes/Engine/EngineTypes.h
--- a/Engine/Source/Runtime/Engine/Classes/Engine/EngineTypes.h
+++ b/Engine/Source/Runtime/Engine/Classes/Engine/EngineTypes.h
@@ -3861,7 +3861,7 @@ struct ENGINE_API FComponentReference
 	FName ComponentProperty;
 
 	/** Path to the component from its owner actor */
-	UPROPERTY()
+	UPROPERTY(meta = (SaveToComponentInstanceDataCache))
 	FString PathToComponent;
 
 	/** Allows direct setting of first component to constraint. */
@@ -3874,8 +3874,28 @@ struct ENGINE_API FComponentReference
 	{
 		return OtherActor == Other.OtherActor && ComponentProperty == Other.ComponentProperty && PathToComponent == Other.PathToComponent && OverrideComponent == Other.OverrideComponent;
 	}
+
+	///** Serialize natively, is much faster due to volume of these records. Must stay in sync with fields */
+	//bool Serialize(FArchive& Ar)
+	//{
+	//	Ar << OtherActor;
+	//	Ar << ComponentProperty;
+	//	Ar << PathToComponent;
+	//	Ar << OverrideComponent;
+
+	//	return true;
+	//}
 };
 
+//template<>
+//struct TStructOpsTypeTraits<FComponentReference> : public TStructOpsTypeTraitsBase2<FComponentReference>
+//{
+//	enum
+//	{
+//		WithSerializer = true
+//	};
+//};
+
 /** Types of valid physical material mask colors which may be associated with a physical material */
 UENUM(BlueprintType)
 namespace EPhysicalMaterialMaskColor
diff --git a/Engine/Source/Runtime/Engine/Private/ComponentInstanceDataCache.cpp b/Engine/Source/Runtime/Engine/Private/ComponentInstanceDataCache.cpp
--- a/Engine/Source/Runtime/Engine/Private/ComponentInstanceDataCache.cpp
+++ b/Engine/Source/Runtime/Engine/Private/ComponentInstanceDataCache.cpp
@@ -58,13 +58,19 @@ public:
 
 	virtual bool ShouldSkipProperty(const FProperty* InProperty) const override
 	{
+		static FName NAME_SaveToComponentInstanceDataCache = "SaveToComponentInstanceDataCache";
+
 		// Immutable structs expect to serialize all properties so don't skip regardless of other conditions
 		UScriptStruct* ScriptStruct = InProperty->GetOwner<UScriptStruct>();
 		const bool bPropertyInImmutableStruct = ScriptStruct && ((ScriptStruct->StructFlags & STRUCT_Immutable) != 0);
 
 		return (!bPropertyInImmutableStruct
 			&& (InProperty->HasAnyPropertyFlags(CPF_Transient)
+#if WITH_EDITORONLY_DATA
+				|| (!InProperty->HasAnyPropertyFlags(CPF_Edit | CPF_Interp) && !InProperty->HasMetaData(NAME_SaveToComponentInstanceDataCache))
+#else
 				|| !InProperty->HasAnyPropertyFlags(CPF_Edit | CPF_Interp)
+#endif
 				|| InProperty->IsA<FMulticastDelegateProperty>()
 				|| PropertiesToSkip.Contains(InProperty)
 				)
