Add a metadata to serialize property in ComponentInstanceDataCache that are marked as "SaveToComponentInstanceDataCache". Use when the property is not editable or interp.

diff --git a/Engine/Source/Runtime/Engine/Classes/Engine/EngineTypes.h b/Engine/Source/Runtime/Engine/Classes/Engine/EngineTypes.h
--- a/Engine/Source/Runtime/Engine/Classes/Engine/EngineTypes.h
+++ b/Engine/Source/Runtime/Engine/Classes/Engine/EngineTypes.h
@@ -3874,28 +3874,8 @@ struct ENGINE_API FComponentReference
 	{
 		return OtherActor == Other.OtherActor && ComponentProperty == Other.ComponentProperty && PathToComponent == Other.PathToComponent && OverrideComponent == Other.OverrideComponent;
 	}
-
-	///** Serialize natively, is much faster due to volume of these records. Must stay in sync with fields */
-	//bool Serialize(FArchive& Ar)
-	//{
-	//	Ar << OtherActor;
-	//	Ar << ComponentProperty;
-	//	Ar << PathToComponent;
-	//	Ar << OverrideComponent;
-
-	//	return true;
-	//}
 };
 
-//template<>
-//struct TStructOpsTypeTraits<FComponentReference> : public TStructOpsTypeTraitsBase2<FComponentReference>
-//{
-//	enum
-//	{
-//		WithSerializer = true
-//	};
-//};
-
 /** Types of valid physical material mask colors which may be associated with a physical material */
 UENUM(BlueprintType)
 namespace EPhysicalMaterialMaskColor
diff --git a/Engine/Source/Runtime/Engine/Private/ComponentInstanceDataCache.cpp b/Engine/Source/Runtime/Engine/Private/ComponentInstanceDataCache.cpp
--- a/Engine/Source/Runtime/Engine/Private/ComponentInstanceDataCache.cpp
+++ b/Engine/Source/Runtime/Engine/Private/ComponentInstanceDataCache.cpp
@@ -58,7 +58,9 @@ public:
 
 	virtual bool ShouldSkipProperty(const FProperty* InProperty) const override
 	{
+#if WITH_EDITORONLY_DATA
 		static FName NAME_SaveToComponentInstanceDataCache = "SaveToComponentInstanceDataCache";
+#endif
 
 		// Immutable structs expect to serialize all properties so don't skip regardless of other conditions
 		UScriptStruct* ScriptStruct = InProperty->GetOwner<UScriptStruct>();
