Don't use scene compute data provider to drive invocation counts.

diff --git a/Engine/Source/Runtime/Engine/Private/ComputeFramework/ComputeGraphScheduler.cpp b/Engine/Source/Runtime/Engine/Private/ComputeFramework/ComputeGraphScheduler.cpp
--- a/Engine/Source/Runtime/Engine/Private/ComputeFramework/ComputeGraphScheduler.cpp
+++ b/Engine/Source/Runtime/Engine/Private/ComputeFramework/ComputeGraphScheduler.cpp
@@ -43,7 +43,7 @@ void FComputeGraphScheduler::EnqueueForExecution(const FComputeGraphProxy* Compu
 	int32 FirstProvider = -1;
 	for (int32 ProviderIndex = 0; ProviderIndex < ComputeDataProviders.Num(); ++ProviderIndex)
 	{
-		if (ComputeDataProviders[ProviderIndex] != nullptr)
+		if (ComputeDataProviders[ProviderIndex] != nullptr && ComputeDataProviders[ProviderIndex]->GetInvocationCount() > 0)
 		{
 			FirstProvider = ProviderIndex;
 			break;
diff --git a/Engine/Source/Runtime/Engine/Public/ComputeFramework/ComputeDataProvider.h b/Engine/Source/Runtime/Engine/Public/ComputeFramework/ComputeDataProvider.h
--- a/Engine/Source/Runtime/Engine/Public/ComputeFramework/ComputeDataProvider.h
+++ b/Engine/Source/Runtime/Engine/Public/ComputeFramework/ComputeDataProvider.h
@@ -38,7 +38,7 @@ public:
 	virtual ~FComputeDataProviderRenderProxy() {}
 
 	/** Called on render thread to determine how many dispatches are required to do all work on the associated data provider. */
-	virtual int32 GetInvocationCount() const { return 1; }
+	virtual int32 GetInvocationCount() const { return 0; }
 
 	/** Called on render thread to determine dispatch dimension required to do all work on the associated data provider. */
 	virtual FIntVector GetDispatchDim(int32 InvocationIndex, FIntVector GroupDim) const { return FIntVector(1, 1, 1); }
