Use master pose component on GPU cloth simulation (based on GitHub 6964 pull request with modification)

diff --git a/Engine/Source/Runtime/Engine/Private/SkeletalRenderGPUSkin.cpp b/Engine/Source/Runtime/Engine/Private/SkeletalRenderGPUSkin.cpp
--- a/Engine/Source/Runtime/Engine/Private/SkeletalRenderGPUSkin.cpp
+++ b/Engine/Source/Runtime/Engine/Private/SkeletalRenderGPUSkin.cpp
@@ -2156,6 +2156,23 @@ bool FDynamicSkelMeshObjectDataGPUSkin::UpdateClothSimulationData(USkinnedMeshCo
 {
 	USkeletalMeshComponent* SimMeshComponent = Cast<USkeletalMeshComponent>(InMeshComponent);
 
+	if (InMeshComponent->MasterPoseComponent.IsValid() && (SimMeshComponent && SimMeshComponent->IsClothBoundToMasterComponent()))
+	{
+		USkeletalMeshComponent* SrcComponent = SimMeshComponent;
+		// if I have master, override sim component
+		SimMeshComponent = Cast<USkeletalMeshComponent>(InMeshComponent->MasterPoseComponent.Get());
+		// IF we don't have sim component that is skeletalmeshcomponent, just ignore
+		if (!SimMeshComponent)
+		{
+			return false;
+		}
+
+		ClothObjectLocalToWorld = SrcComponent->GetComponentToWorld().ToMatrixWithScale();
+		ClothBlendWeight = SimMeshComponent->ClothBlendWeight;
+		ClothingSimData = SimMeshComponent->GetCurrentClothingData_AnyThread();
+		return true;
+	}
+
 	if (SimMeshComponent)
 	{
 		ClothObjectLocalToWorld = SimMeshComponent->GetComponentToWorld().ToMatrixWithScale();
