[VisualLogger] added missing defines for UE_VLOG_CIRCLE, UE_CVLOG_CIRCLE, UE_VLOG_CIRCLE_THICK and UE_CVLOG_CIRCLE_THICK when ENABLE_VISUAL_LOG is 0

diff --git a/Engine/Source/Runtime/Engine/Public/VisualLogger/VisualLogger.h b/Engine/Source/Runtime/Engine/Public/VisualLogger/VisualLogger.h
--- a/Engine/Source/Runtime/Engine/Public/VisualLogger/VisualLogger.h
+++ b/Engine/Source/Runtime/Engine/Public/VisualLogger/VisualLogger.h
@@ -118,6 +118,10 @@
 #define UE_CVLOG_CONVEXPOLY(Condition, LogOwner, CategoryName, Verbosity, Points, Color, Format, ...)
 #define UE_VLOG_ARROW(LogOwner, CategoryName, Verbosity, SegmentStart, SegmentEnd, Color, Format, ...) 
 #define UE_CVLOG_ARROW(Condition, LogOwner, CategoryName, Verbosity, SegmentStart, SegmentEnd, Color, Format, ...) 
+#define UE_VLOG_CIRCLE(LogOwner, CategoryName, Verbosity, Center, UpAxis, Radius, Color, Format, ...)
+#define UE_CVLOG_CIRCLE(Condition, LogOwner, CategoryName, Verbosity, Center, UpAxis, Radius, Color, Format, ...)
+#define UE_VLOG_CIRCLE_THICK(LogOwner, CategoryName, Verbosity, Center, UpAxis, Radius, Color, Thickness, Format, ...)
+#define UE_CVLOG_CIRCLE_THICK(Condition, LogOwner, CategoryName, Verbosity, Center, UpAxis, Radius, Color, Thickness, Format, ...)
 
 #define DECLARE_VLOG_EVENT(EventName)
 #define DEFINE_VLOG_EVENT(EventName, Verbosity, UserFriendlyDesc)
