When material function changes, update all dependent function instances, and all material instances that use any of those function (instances)

diff --git a/Engine/Source/Editor/MaterialEditor/Private/MaterialEditingLibrary.cpp b/Engine/Source/Editor/MaterialEditor/Private/MaterialEditingLibrary.cpp
--- a/Engine/Source/Editor/MaterialEditor/Private/MaterialEditingLibrary.cpp
+++ b/Engine/Source/Editor/MaterialEditor/Private/MaterialEditingLibrary.cpp
@@ -881,6 +881,8 @@ void UMaterialEditingLibrary::UpdateMaterialFunction(UMaterialFunctionInterface*
 				if (Functions.Contains(MaterialFunction))
 				{
 					FunctionInstance->UpdateParameterSet();
+					FunctionInstance->PreEditChange(nullptr);
+					FunctionInstance->PostEditChange();
 					FunctionInstance->MarkPackageDirty();
 				}
 			}
diff --git a/Engine/Source/Runtime/Engine/Private/Materials/MaterialExpressions.cpp b/Engine/Source/Runtime/Engine/Private/Materials/MaterialExpressions.cpp
--- a/Engine/Source/Runtime/Engine/Private/Materials/MaterialExpressions.cpp
+++ b/Engine/Source/Runtime/Engine/Private/Materials/MaterialExpressions.cpp
@@ -12248,25 +12248,31 @@ void UMaterialFunctionInterface::GetAssetRegistryTags(TArray<FAssetRegistryTag>&
 #if WITH_EDITOR
 void UMaterialFunctionInterface::PostEditChangeProperty(FPropertyChangedEvent& PropertyChangedEvent)
 {
+	//@todo - recreate guid only when needed, not when a comment changes
+	StateId = FGuid::NewGuid();
+
 	// Go through all materials in memory and recompile them if they use this function
 	FMaterialUpdateContext UpdateContext;
-	for (TObjectIterator<UMaterial> It; It; ++It)
+	for (TObjectIterator<UMaterialInterface> It; It; ++It)
 	{
-		UMaterial* CurrentMaterial = *It;
+		UMaterialInterface* CurrentMaterial = *It;
 
 		bool bRecompile = false;
-		for (const FMaterialFunctionInfo& FunctionInfo : CurrentMaterial->GetCachedExpressionData().FunctionInfos)
-		{
-			if (FunctionInfo.Function == this)
+		UMaterialFunctionInterface* Self = this;
+		CurrentMaterial->IterateDependentFunctions(
+			[Self, &bRecompile](UMaterialFunctionInterface* InFunction)
 			{
-				bRecompile = true;
-				break;
-			}
-		}
+				if (InFunction == Self)
+				{
+					bRecompile = true;
+					return false;
+				}
+				return true;
+			});
 
 		if (bRecompile)
 		{
-			UpdateContext.AddMaterial(CurrentMaterial);
+			UpdateContext.AddMaterialInterface(CurrentMaterial);
 
 			// Propagate the change to this material
 			CurrentMaterial->PreEditChange(nullptr);
@@ -12396,9 +12402,6 @@ void UMaterialFunction::PostEditChangeProperty(FPropertyChangedEvent& PropertyCh
 	UpdateDependentFunctionCandidates();
 #endif
 
-	//@todo - recreate guid only when needed, not when a comment changes
-	StateId = FGuid::NewGuid();
-
 	Super::PostEditChangeProperty(PropertyChangedEvent);
 }
 #endif // WITH_EDITOR
