Fix for audio component notifications for sound done in BP in editor mode.

diff --git a/Engine/Source/Runtime/Engine/Private/AudioDevice.cpp b/Engine/Source/Runtime/Engine/Private/AudioDevice.cpp
--- a/Engine/Source/Runtime/Engine/Private/AudioDevice.cpp
+++ b/Engine/Source/Runtime/Engine/Private/AudioDevice.cpp
@@ -5009,47 +5009,55 @@ void FAudioDevice::ProcessingPendingActiveSoundStops(bool bForceDelete)
 		}
 	}
 
-	// Stop any pending active sounds that need to be stopped
-	for (FActiveSound* ActiveSound : PendingSoundsToStop)
+	while (PendingSoundsToStop.Num() > 0)
 	{
-		check(ActiveSound);
+		// Need to make a copy since stopping sounds can add to this list of sounds to stop when the audio thread isn't running (e.g. editor).
+		TSet<FActiveSound*> PendingSoundsToStopCopy = PendingSoundsToStop;
 
-		// If the request was to stop an ActiveSound that
-		// is set to re-trigger but is not playing, remove
-		// and continue
-		if (RemoveVirtualLoop(*ActiveSound))
+		// Stop any pending active sounds that need to be stopped
+		for (FActiveSound* ActiveSound : PendingSoundsToStopCopy)
 		{
-			delete ActiveSound;
-		}
-		else
-		{
-			ActiveSound->MarkPendingDestroy(bForceDelete);
-
-			USoundBase* Sound = ActiveSound->GetSound();
-
-			// If the active sound is a one shot, decrement the one shot counter
-			if (Sound && !Sound->IsLooping())
-			{
-				OneShotCount--;
-			}
-
-			const bool bIsStopping = ActiveSound->IsStopping();
+			check(ActiveSound);
 
-			// If we can delete the active sound now, then delete it
-			if (bForceDelete || (ActiveSound->CanDelete() && !bIsStopping))
+			// If the request was to stop an ActiveSound that
+			// is set to re-trigger but is not playing, remove
+			// and continue
+			if (RemoveVirtualLoop(*ActiveSound))
 			{
-				ActiveSound->bAsyncOcclusionPending = false;
-
 				delete ActiveSound;
 			}
 			else
 			{
-				// There was an async operation pending or we are stopping (not stopped) so we need to defer deleting this sound
-				PendingSoundsToDelete.AddUnique(ActiveSound);
+				ActiveSound->MarkPendingDestroy(bForceDelete);
+
+				USoundBase* Sound = ActiveSound->GetSound();
+
+				// If the active sound is a one shot, decrement the one shot counter
+				if (Sound && !Sound->IsLooping())
+				{
+					OneShotCount--;
+				}
+
+				const bool bIsStopping = ActiveSound->IsStopping();
+
+				// If we can delete the active sound now, then delete it
+				if (bForceDelete || (ActiveSound->CanDelete() && !bIsStopping))
+				{
+					ActiveSound->bAsyncOcclusionPending = false;
+
+					delete ActiveSound;
+				}
+				else
+				{
+					// There was an async operation pending or we are stopping (not stopped) so we need to defer deleting this sound
+					PendingSoundsToDelete.AddUnique(ActiveSound);
+				}
 			}
+
+			// Remove from the list of pending sounds to stop
+			PendingSoundsToStop.Remove(ActiveSound);
 		}
 	}
-	PendingSoundsToStop.Reset();
 }
 
 void FAudioDevice::AddSoundToStop(FActiveSound* SoundToStop)
