Fix for platform downsampling feature.

diff --git a/Engine/Source/Editor/AudioEditor/Private/Factories/SoundFactory.cpp b/Engine/Source/Editor/AudioEditor/Private/Factories/SoundFactory.cpp
--- a/Engine/Source/Editor/AudioEditor/Private/Factories/SoundFactory.cpp
+++ b/Engine/Source/Editor/AudioEditor/Private/Factories/SoundFactory.cpp
@@ -517,6 +517,7 @@ UObject* USoundFactory::CreateObject
 		}
 
 		Sound->Duration = (float)NumFrames / *WaveInfo.pSamplesPerSec;
+		Sound->SetImportedSampleRate(*WaveInfo.pSamplesPerSec);
 		Sound->SetSampleRate(*WaveInfo.pSamplesPerSec);
 		Sound->NumChannels = ChannelCount;
 		Sound->TotalSamples = *WaveInfo.pSamplesPerSec * Sound->Duration;
diff --git a/Engine/Source/Runtime/Engine/Classes/Sound/SoundWave.h b/Engine/Source/Runtime/Engine/Classes/Sound/SoundWave.h
--- a/Engine/Source/Runtime/Engine/Classes/Sound/SoundWave.h
+++ b/Engine/Source/Runtime/Engine/Classes/Sound/SoundWave.h
@@ -706,10 +706,16 @@ public:
 
 protected:
 
-	/** Cached sample rate for displaying in the tools */
+	/** Cooked sample rate of the asset. Can be modified by sample rate override. */
 	UPROPERTY(Category = Info, AssetRegistrySearchable, VisibleAnywhere)
 	int32 SampleRate;
 
+#if WITH_EDITORONLY_DATA
+	/** Sample rate of the imported sound wave. */
+	UPROPERTY(Category = Info, AssetRegistrySearchable, VisibleAnywhere)
+	int32 ImportedSampleRate;
+#endif
+
 public:
 
 	/** Resource index to cross reference with buffers */
@@ -887,10 +893,16 @@ public:
 		return bIsGeneratingAudio;
 	}
 
+	void SetImportedSampleRate(uint32 InImportedSampleRate)
+	{
+#if WITH_EDITORONLY_DATA
+		ImportedSampleRate = InImportedSampleRate;
+#endif
+	}
+
 	/**
 	* Overwrite sample rate. Used for procedural soundwaves, as well as sound waves that are resampled on compress/decompress.
 	*/
-
 	void SetSampleRate(uint32 InSampleRate)
 	{
 		SampleRate = InSampleRate;
diff --git a/Engine/Source/Runtime/Engine/Private/AudioDerivedData.cpp b/Engine/Source/Runtime/Engine/Private/AudioDerivedData.cpp
--- a/Engine/Source/Runtime/Engine/Private/AudioDerivedData.cpp
+++ b/Engine/Source/Runtime/Engine/Private/AudioDerivedData.cpp
@@ -1071,6 +1071,7 @@ static void CookSimpleWave(USoundWave* SoundWave, FName FormatName, const IAudio
 	bool bWasLocked = false;
 
 	int32 WaveSampleRate = 0;
+	int32 ImportSampleRate = 0;
 
 	// Don't need to cook source buses
 	if (USoundSourceBus* SourceBus = Cast<USoundSourceBus>(SoundWave))
@@ -1122,6 +1123,7 @@ static void CookSimpleWave(USoundWave* SoundWave, FName FormatName, const IAudio
 	else
 	{
 		WaveSampleRate = *WaveInfo.pSamplesPerSec;
+		ImportSampleRate = WaveSampleRate;
 
 		float SampleRateOverride = -1.0f;
 
@@ -1170,6 +1172,7 @@ static void CookSimpleWave(USoundWave* SoundWave, FName FormatName, const IAudio
 		if(Format.Cook(FormatName, Input, QualityInfo, Output))
 		{
 			SoundWave->SetSampleRate(QualityInfo.SampleRate);
+			SoundWave->SetImportedSampleRate(ImportSampleRate);
 
 			if (SoundWave->NumChannels != *WaveInfo.pChannels)
 			{
diff --git a/Engine/Source/Runtime/Engine/Private/SoundWave.cpp b/Engine/Source/Runtime/Engine/Private/SoundWave.cpp
--- a/Engine/Source/Runtime/Engine/Private/SoundWave.cpp
+++ b/Engine/Source/Runtime/Engine/Private/SoundWave.cpp
@@ -1393,6 +1393,14 @@ void USoundWave::InvalidateSoundWaveIfNeccessary()
 		return;
 	}
 
+	// This will cause a recook of the sample rate overrides have changed since it last played
+	int32 SampleRateOverride = (int32)GetSampleRateForCurrentPlatform();
+	if (SampleRate != SampleRateOverride)
+	{
+		InvalidateCompressedData(true);
+		SampleRate = SampleRateOverride;
+	}
+	
 	// if stream caching was enabled since the last time we invalidated the compressed audio, force a re-cook.
 	const bool bIsStreamCachingEnabled = FPlatformCompressionUtilities::IsCurrentPlatformUsingStreamCaching();
 	if (bWasStreamCachingEnabledOnLastCook != bIsStreamCachingEnabled)
@@ -2578,12 +2586,38 @@ float USoundWave::GetSampleRateForCurrentPlatform()
 		return SampleRate;
 	}
 
+#if WITH_EDITORONLY_DATA
+	// GIsEditor is false in standalone but WITH_EDITORONLY_DATA is true
 	if (GIsEditor)
 	{
 		float SampleRateOverride = FPlatformCompressionUtilities::GetTargetSampleRateForPlatform(SampleRateQuality);
-		return (SampleRateOverride > 0) ? FMath::Min(SampleRateOverride, (float)SampleRate) : SampleRate;
+		if (SampleRateOverride > 0)
+		{
+			// If we don't have an imported sample rate defined (i.e. old asset), just use the override. This will actually update the imported sample rate when it recooks.
+			if (!ImportedSampleRate)
+			{
+				return SampleRateOverride;
+			}
+			else
+			{
+				// Otherwise, use the min of the imported SR and the override
+				return FMath::Min(SampleRateOverride, (float)ImportedSampleRate);
+			}
+		}
+		else
+		{
+			if (!ImportedSampleRate)
+			{
+				return SampleRate;
+			}
+			else
+			{
+				return ImportedSampleRate;
+			}
+		}
 	}
 	else
+#endif
 	{
 		if (bCachedSampleRateFromPlatformSettings)
 		{
@@ -2615,7 +2649,14 @@ float USoundWave::GetSampleRateForCompressionOverrides(const FPlatformAudioCookO
 	const float* SampleRatePtr = CompressionOverrides->PlatformSampleRates.Find(SampleRateQuality);
 	if (SampleRatePtr && *SampleRatePtr > 0.0f)
 	{
-		return FMath::Min(*SampleRatePtr, static_cast<float>(SampleRate));
+		if (GIsEditor)
+		{
+			return GetSampleRateForCurrentPlatform();
+		}
+		else
+		{
+			return FMath::Min(*SampleRatePtr, static_cast<float>(SampleRate));
+		}
 	}
 	else
 	{
