Allow non replicated child actors to be destroyed by the parent

diff --git a/Engine/Source/Runtime/Engine/Classes/Components/ChildActorComponent.h b/Engine/Source/Runtime/Engine/Classes/Components/ChildActorComponent.h
--- a/Engine/Source/Runtime/Engine/Classes/Components/ChildActorComponent.h
+++ b/Engine/Source/Runtime/Engine/Classes/Components/ChildActorComponent.h
@@ -192,6 +192,9 @@ public:
 
 	void SetEditorTreeViewVisualizationMode(EChildActorComponentTreeViewVisualizationMode InMode);
 #endif
+
+private:
+	bool IsChildActorReplicated() const;
 };
 
 struct FActorParentComponentSetter
diff --git a/Engine/Source/Runtime/Engine/Private/Components/ChildActorComponent.cpp b/Engine/Source/Runtime/Engine/Private/Components/ChildActorComponent.cpp
--- a/Engine/Source/Runtime/Engine/Private/Components/ChildActorComponent.cpp
+++ b/Engine/Source/Runtime/Engine/Private/Components/ChildActorComponent.cpp
@@ -584,16 +584,22 @@ EDataValidationResult UChildActorComponent::IsDataValid(TArray<FText>& Validatio
 }
 #endif
 
+bool UChildActorComponent::IsChildActorReplicated() const
+{
+	AActor* ChildClassCDO = (ChildActorClass ? ChildActorClass->GetDefaultObject<AActor>() : nullptr);
+	const bool bChildActorClassReplicated = ChildClassCDO && ChildClassCDO->GetIsReplicated();
+	const bool bChildActorTemplateReplicated = ChildActorTemplate && (ChildActorTemplate->GetClass() == ChildActorClass) && ChildActorTemplate->GetIsReplicated();
+
+	return (bChildActorClassReplicated || bChildActorTemplateReplicated);
+}
+
 void UChildActorComponent::CreateChildActor()
 {
 	AActor* MyOwner = GetOwner();
 
 	if (MyOwner && !MyOwner->HasAuthority())
 	{
-		AActor* ChildClassCDO = (ChildActorClass ? ChildActorClass->GetDefaultObject<AActor>() : nullptr);
-		const bool bChildActorTemplateReplicated = ChildActorTemplate && ChildActorTemplate->GetClass() == ChildActorClass && ChildActorTemplate->GetIsReplicated();
-
-		if ((ChildClassCDO && ChildClassCDO->GetIsReplicated()) || bChildActorTemplateReplicated)
+		if (IsChildActorReplicated())
 		{
 			// If we belong to an actor that is not authoritative and the child class is replicated then we expect that Actor will be replicated across so don't spawn one
 			return;
@@ -722,7 +728,7 @@ void UChildActorComponent::DestroyChildActor()
 		return false;
 	};
 
-	if (ChildActor && ChildActor->HasAuthority() && !IsLevelBeingRemoved())
+	if (ChildActor && (ChildActor->HasAuthority() || !IsChildActorReplicated()) && !IsLevelBeingRemoved())
 	{
 		if (!GExitPurge)
 		{
