Fix world partition debug display legend counts.

diff --git a/Engine/Source/Runtime/Engine/Private/WorldPartition/WorldPartitionLevelStreamingPolicy.cpp b/Engine/Source/Runtime/Engine/Private/WorldPartition/WorldPartitionLevelStreamingPolicy.cpp
--- a/Engine/Source/Runtime/Engine/Private/WorldPartition/WorldPartitionLevelStreamingPolicy.cpp
+++ b/Engine/Source/Runtime/Engine/Private/WorldPartition/WorldPartitionLevelStreamingPolicy.cpp
@@ -251,6 +251,11 @@ void UWorldPartitionLevelStreamingPolicy::DrawStreamingStatusLegend(UCanvas* Can
 		StatusCount[(int32)Cell->GetStreamingStatus()]++;
 	});
 
+	// @todo_ow: This is not exactly the good value, as we could have pending unload level from Level Instances, etc.
+	//           We could modify GetNumLevelsPendingPurge to return the number of pending purge levels from the grid, 
+	//           bu that will do for now.
+	StatusCount[LEVEL_UnloadedButStillAround] = FLevelStreamingGCHelper::GetNumLevelsPendingPurge();
+
 	// Draw legend
 	FVector2D Pos = Offset;
 	float MaxTextWidth = 0.f;
@@ -260,7 +265,12 @@ void UWorldPartitionLevelStreamingPolicy::DrawStreamingStatusLegend(UCanvas* Can
 	{
 		EStreamingStatus Status = (EStreamingStatus)i;
 		const FColor& StatusColor = ULevelStreaming::GetLevelStreamingStatusColor(Status);
-		FWorldPartitionDebugHelper::DrawLegendItem(Canvas, *FString::Printf(TEXT("%d) %s (%d)"), i, ULevelStreaming::GetLevelStreamingStatusDisplayName(Status), StatusCount[(int32)Status]), GEngine->GetTinyFont(), StatusColor, Pos, &MaxTextWidth);
+		FString DebugString = *FString::Printf(TEXT("%d) %s"), i, ULevelStreaming::GetLevelStreamingStatusDisplayName(Status));
+		if (Status != LEVEL_Unloaded)
+		{
+			DebugString += *FString::Printf(TEXT(" (%d)"), StatusCount[(int32)Status]);
+		}
+		FWorldPartitionDebugHelper::DrawLegendItem(Canvas, *DebugString, GEngine->GetTinyFont(), StatusColor, Pos, &MaxTextWidth);
 	}
 
 	Offset.X += MaxTextWidth + 10;
