Change how LevelStreamingContinuouslyIncrementalGCWhileLevelsPendingPurge works:
- Instead of forcing a GC when there's at least one pending levels to unload, triggers when we go over a certain amount.
- Default behavior is not modified by this change, but under certain conditions this will be enabled (for partitioned worlds, for instance).

diff --git a/Engine/Source/Runtime/Engine/Private/LevelStreaming.cpp b/Engine/Source/Runtime/Engine/Private/LevelStreaming.cpp
--- a/Engine/Source/Runtime/Engine/Private/LevelStreaming.cpp
+++ b/Engine/Source/Runtime/Engine/Private/LevelStreaming.cpp
@@ -541,7 +541,7 @@ void ULevelStreaming::UpdateStreamingState(bool& bOutUpdateAgain, bool& bOutRede
 		if (GLevelStreamingContinuouslyIncrementalGCWhileLevelsPendingPurge)
 		{
 			// Figure out whether there are any levels we haven't collected garbage yet.
-			const bool bAreLevelsPendingPurge = FLevelStreamingGCHelper::GetNumLevelsPendingPurge() > 0;
+			const bool bAreLevelsPendingPurge = FLevelStreamingGCHelper::GetNumLevelsPendingPurge() >= GLevelStreamingContinuouslyIncrementalGCWhileLevelsPendingPurge;
 
 			// Request a 'soft' GC if there are levels pending purge and there are levels to be loaded. In the case of a blocking
 			// load this is going to guarantee GC firing first thing afterwards and otherwise it is going to sneak in right before
