Add a profile tag in BuildLevelTextureStreamingComponentDataFromActors.

diff --git a/Engine/Source/Runtime/Engine/Private/Streaming/TextureStreamingBuild.cpp b/Engine/Source/Runtime/Engine/Private/Streaming/TextureStreamingBuild.cpp
--- a/Engine/Source/Runtime/Engine/Private/Streaming/TextureStreamingBuild.cpp
+++ b/Engine/Source/Runtime/Engine/Private/Streaming/TextureStreamingBuild.cpp
@@ -71,6 +71,8 @@ ENGINE_API bool BuildLevelTextureStreamingComponentDataFromActors(ULevel* InLeve
 		return true;
 	}
 
+	TRACE_CPUPROFILER_EVENT_SCOPE(BuildLevelTextureStreamingComponentDataFromActors);
+
 	// Find actors with texture streaming built data (built with BuildActorTextureStreamingData) and remap this data to the level
 	bool bIsLevelInitialized = false;
 	for (AActor* Actor : InLevel->Actors)
