WorldPartition: Allow disabling of PlayerController streaming source

diff --git a/Engine/Source/Runtime/Engine/Classes/GameFramework/PlayerController.h b/Engine/Source/Runtime/Engine/Classes/GameFramework/PlayerController.h
--- a/Engine/Source/Runtime/Engine/Classes/GameFramework/PlayerController.h
+++ b/Engine/Source/Runtime/Engine/Classes/GameFramework/PlayerController.h
@@ -472,6 +472,10 @@ public:
 	UPROPERTY(EditAnywhere, BlueprintReadWrite, Category="Game|Feedback")
 	uint32 bForceFeedbackEnabled:1;
 
+	/** Whether the PlayerController should be used as a World Partiton streaming source.  */
+	UPROPERTY(EditAnywhere, BlueprintReadWrite, Category = WorldPartition)
+	uint32 bEnableStreamingSource:1;
+
 	/** Scale applied to force feedback values */
 	UPROPERTY(config)
 	float ForceFeedbackScale;
@@ -664,6 +668,14 @@ public:
 	 */
 	virtual void UpdateRotation(float DeltaTime);
 
+	/**
+	* Whether the PlayerController should be used as a World Partiton streaming source. 
+	* Default implementation returns bEnableStreamingSource but can be overriden in child classes.
+	* @return true if it should.
+	*/
+	UFUNCTION(BlueprintCallable, Category = WorldPartition)
+	virtual bool IsStreamingSourceEnabled() const { return bEnableStreamingSource; }
+
 protected:
 	/** Pawn has been possessed, so changing state to NAME_Playing. Start it walking and begin playing with it. */
 	virtual void BeginPlayingState();
diff --git a/Engine/Source/Runtime/Engine/Private/PlayerController.cpp b/Engine/Source/Runtime/Engine/Private/PlayerController.cpp
--- a/Engine/Source/Runtime/Engine/Private/PlayerController.cpp
+++ b/Engine/Source/Runtime/Engine/Private/PlayerController.cpp
@@ -136,6 +136,8 @@ APlayerController::APlayerController(const FObjectInitializer& ObjectInitializer
 	bIsLocalPlayerController = false;
 	bDisableHaptics = false;
 
+	bEnableStreamingSource = true;
+
 	ClickEventKeys.Add(EKeys::LeftMouseButton);
 
 	if (RootComponent)
diff --git a/Engine/Source/Runtime/Engine/Private/WorldPartition/WorldPartitionStreamingPolicy.cpp b/Engine/Source/Runtime/Engine/Private/WorldPartition/WorldPartitionStreamingPolicy.cpp
--- a/Engine/Source/Runtime/Engine/Private/WorldPartition/WorldPartitionStreamingPolicy.cpp
+++ b/Engine/Source/Runtime/Engine/Private/WorldPartition/WorldPartitionStreamingPolicy.cpp
@@ -107,7 +107,7 @@ void UWorldPartitionStreamingPolicy::UpdateStreamingSources()
 			for (int32 PlayerIndex = 0; PlayerIndex < NumPlayers; ++PlayerIndex)
 			{
 				ULocalPlayer* Player = GEngine->GetGamePlayer(World, PlayerIndex);
-				if (Player && Player->PlayerController)
+				if (Player && Player->PlayerController && Player->PlayerController->IsStreamingSourceEnabled())
 				{
 					FVector ViewLocation;
 					FRotator ViewRotation;
