World Partition GC tweaks:
- Force a GC when there's more than 64 pending purge levels (by moving or when teleporting), using changes from 16897200 (can be tweaked with wp.runtime.GLevelStreamingContinuouslyIncrementalGCWhileLevelsPendingPurgeForWP cvar).
- Remove hardcoded GC every 2 minutes, let the project decide in config files instead (which will be every minutes by default).

diff --git a/Engine/Source/Runtime/Engine/Private/WorldPartition/WorldPartitionSubsystem.cpp b/Engine/Source/Runtime/Engine/Private/WorldPartition/WorldPartitionSubsystem.cpp
--- a/Engine/Source/Runtime/Engine/Private/WorldPartition/WorldPartitionSubsystem.cpp
+++ b/Engine/Source/Runtime/Engine/Private/WorldPartition/WorldPartitionSubsystem.cpp
@@ -8,6 +8,7 @@
 #include "HAL/IConsoleManager.h"
 #include "Engine/Canvas.h"
 #include "Engine/Console.h"
+#include "Engine/CoreSettings.h"
 #include "ConsoleSettings.h"
 #include "Debug/DebugDrawService.h"
 
@@ -49,6 +50,14 @@ static FAutoConsoleCommand CVarDrawDataLayers(
 	TEXT("Toggles debug display of active data layers."),
 	FConsoleCommandDelegate::CreateLambda([] { GDrawDataLayers = !GDrawDataLayers; }));
 
+int32 GLevelStreamingContinuouslyIncrementalGCWhileLevelsPendingPurgeForWP = 64;
+static FAutoConsoleVariableRef CVarGLevelStreamingContinuouslyIncrementalGCWhileLevelsPendingPurgeForWP(
+	TEXT("wp.Runtime.LevelStreamingContinuouslyIncrementalGCWhileLevelsPendingPurgeForWP"),
+	GLevelStreamingContinuouslyIncrementalGCWhileLevelsPendingPurgeForWP,
+	TEXT("Force a GC update when there's more than the number of specified pending purge levels."),
+	ECVF_Default
+);
+
 UWorldPartitionSubsystem::UWorldPartitionSubsystem()
 {}
 
@@ -100,16 +109,14 @@ void UWorldPartitionSubsystem::PostInitialize()
 			DrawHandle = UDebugDrawService::Register(TEXT("Game"), FDebugDrawDelegate::CreateUObject(this, &UWorldPartitionSubsystem::Draw));
 		}
 
-		// For now, enforce some GC settings when using World Partition
+		// Enforce some GC settings when using World Partition
 		if (GetWorld()->IsGameWorld())
 		{
-			PreviousCVarValues.ReadFromCVars();
+			LevelStreamingContinuouslyIncrementalGCWhileLevelsPendingPurge = GLevelStreamingContinuouslyIncrementalGCWhileLevelsPendingPurge;
+			LevelStreamingForceGCAfterLevelStreamedOut = GLevelStreamingForceGCAfterLevelStreamedOut;
 
-			FWorldPartitionCVars OverrideCVars;
-			OverrideCVars.ContinuouslyIncremental = 0;
-			OverrideCVars.ForceGCAfterLevelStreamedOut = 0;
-			OverrideCVars.TimeBetweenPurgingPendingKillObjects = 120.f;
-			OverrideCVars.WriteToCVars();
+			GLevelStreamingContinuouslyIncrementalGCWhileLevelsPendingPurge = GLevelStreamingContinuouslyIncrementalGCWhileLevelsPendingPurgeForWP;
+			GLevelStreamingForceGCAfterLevelStreamedOut = 0;
 		}
 	}
 }
@@ -119,7 +126,8 @@ void UWorldPartitionSubsystem::Deinitialize()
 	UWorldPartition* MainPartition = GetMainWorldPartition();
 	if (MainPartition && GetWorld()->IsGameWorld())
 	{
-		PreviousCVarValues.WriteToCVars();
+		GLevelStreamingContinuouslyIncrementalGCWhileLevelsPendingPurge = LevelStreamingContinuouslyIncrementalGCWhileLevelsPendingPurge;
+		GLevelStreamingForceGCAfterLevelStreamedOut = LevelStreamingForceGCAfterLevelStreamedOut;
 	}
 
 	if (DrawHandle.IsValid())
@@ -264,6 +272,13 @@ void UWorldPartitionSubsystem::Draw(UCanvas* Canvas, class APlayerController* PC
 			}
 			CurrentOffset.X = CanvasBottomRightPadding.X;
 		}
+
+		FString StatusText;
+		if (IsIncrementalPurgePending()) { StatusText += TEXT("(Purging) "); }
+		if (IsAsyncLoading()) { StatusText += TEXT("(AsyncLoading) "); }
+		if (StatusText.IsEmpty()) { StatusText = TEXT("(Idle)"); }
+		const FString Text = FString::Printf(TEXT("Streaming Status: %s"), *StatusText);
+		FWorldPartitionDebugHelper::DrawText(Canvas, Text, GEngine->GetSmallFont(), FColor::White, CurrentOffset);
 	}
 
 	if (GDrawStreamingSources || GDrawRuntimeHash2D)
@@ -314,57 +329,4 @@ void UWorldPartitionSubsystem::Draw(UCanvas* Canvas, class APlayerController* PC
 			Partition->DrawRuntimeCellsDetails(Canvas, CurrentOffset);
 		}
 	}
-}
-
-//
-// UWorldPartitionSubsystem::FWorldPartitionCVars
-//
-
-const TCHAR* UWorldPartitionSubsystem::FWorldPartitionCVars::ContinuouslyIncrementalText = TEXT("s.ContinuouslyIncrementalGCWhileLevelsPendingPurge");
-const TCHAR* UWorldPartitionSubsystem::FWorldPartitionCVars::ForceGCAfterLevelStreamedOutText = TEXT("s.ForceGCAfterLevelStreamedOut");
-const TCHAR* UWorldPartitionSubsystem::FWorldPartitionCVars::TimeBetweenPurgingPendingKillObjectsText = TEXT("gc.TimeBetweenPurgingPendingKillObjects");
-
-void UWorldPartitionSubsystem::FWorldPartitionCVars::ReadFromCVars()
-{
-	if (IConsoleVariable* ContinuouslyIncrementalGCCVar = IConsoleManager::Get().FindConsoleVariable(ContinuouslyIncrementalText))
-	{
-		ContinuouslyIncremental = ContinuouslyIncrementalGCCVar->GetInt();
-	}
-
-	if (IConsoleVariable* ForceGCAfterLevelStreamedOutCVar = IConsoleManager::Get().FindConsoleVariable(ForceGCAfterLevelStreamedOutText))
-	{
-		ForceGCAfterLevelStreamedOut = ForceGCAfterLevelStreamedOutCVar->GetInt();
-	}
-
-	if (IConsoleVariable* TimeBetweenPurgingPendingKillObjectsCVar = IConsoleManager::Get().FindConsoleVariable(TimeBetweenPurgingPendingKillObjectsText))
-	{
-		TimeBetweenPurgingPendingKillObjects = TimeBetweenPurgingPendingKillObjectsCVar->GetFloat();
-	}
-}
-
-void UWorldPartitionSubsystem::FWorldPartitionCVars::WriteToCVars() const
-{
-	if (ContinuouslyIncremental.IsSet())
-	{
-		if (IConsoleVariable* ContinuouslyIncrementalGCCVar = IConsoleManager::Get().FindConsoleVariable(ContinuouslyIncrementalText))
-		{
-			ContinuouslyIncrementalGCCVar->Set(ContinuouslyIncremental.GetValue(), ECVF_SetByCode);
-		}
-	}
-
-	if (ForceGCAfterLevelStreamedOut.IsSet())
-	{
-		if (IConsoleVariable* ForceGCAfterLevelStreamedOutCVar = IConsoleManager::Get().FindConsoleVariable(ForceGCAfterLevelStreamedOutText))
-		{
-			ForceGCAfterLevelStreamedOutCVar->Set(ForceGCAfterLevelStreamedOut.GetValue(), ECVF_SetByCode);
-		}
-	}
-
-	if (TimeBetweenPurgingPendingKillObjects.IsSet())
-	{
-		if (IConsoleVariable* TimeBetweenPurgingPendingKillObjectsCVar = IConsoleManager::Get().FindConsoleVariable(TimeBetweenPurgingPendingKillObjectsText))
-		{
-			TimeBetweenPurgingPendingKillObjectsCVar->Set(TimeBetweenPurgingPendingKillObjects.GetValue(), ECVF_SetByCode);
-		}
-	}
-}
+}
\ No newline at end of file
diff --git a/Engine/Source/Runtime/Engine/Public/WorldPartition/WorldPartitionSubsystem.h b/Engine/Source/Runtime/Engine/Public/WorldPartition/WorldPartitionSubsystem.h
--- a/Engine/Source/Runtime/Engine/Public/WorldPartition/WorldPartitionSubsystem.h
+++ b/Engine/Source/Runtime/Engine/Public/WorldPartition/WorldPartitionSubsystem.h
@@ -65,21 +65,7 @@ private:
 
 	FDelegateHandle	DrawHandle;
 
-	struct FWorldPartitionCVars
-	{
-		FWorldPartitionCVars() {}
-
-		void ReadFromCVars();
-		void WriteToCVars() const;
-
-		static const TCHAR* ContinuouslyIncrementalText;
-		static const TCHAR* ForceGCAfterLevelStreamedOutText;
-		static const TCHAR* TimeBetweenPurgingPendingKillObjectsText;
-
-		TOptional<int32> ContinuouslyIncremental;
-		TOptional<int32> ForceGCAfterLevelStreamedOut;
-		TOptional<float> TimeBetweenPurgingPendingKillObjects;
-	};
-
-	FWorldPartitionCVars PreviousCVarValues;
+	// GC backup values
+	int32 LevelStreamingContinuouslyIncrementalGCWhileLevelsPendingPurge;
+	int32 LevelStreamingForceGCAfterLevelStreamedOut;
 };
