Properly clamp a few properties of FMeshApproximationSettings
* ClampVoxelDimension -> Max of [REDACTED_PHONE]^3 is the last integer < 2^31, using a bigger number results in failures in TArray code & probably elsewhere)
* RenderCaptureResolution -> Min of 0 (If Value is zero, will use MaterialSettings resolution, otherwise override the render capture resolution)

diff --git a/Engine/Source/Runtime/Engine/Classes/Engine/MeshMerging.h b/Engine/Source/Runtime/Engine/Classes/Engine/MeshMerging.h
--- a/Engine/Source/Runtime/Engine/Classes/Engine/MeshMerging.h
+++ b/Engine/Source/Runtime/Engine/Classes/Engine/MeshMerging.h
@@ -720,8 +720,8 @@ struct FMeshApproximationSettings
 	UPROPERTY(BlueprintReadWrite, EditAnywhere, Category = ShapeSettings, meta = (DisplayName = "Approximation Accuracy (meters)", ClampMin = "0.001"))
 	float ApproximationAccuracy = 1.0;
 
-	/** Maximum allowable voxel count along main directions. This is a limit on ApproximationAccuracy */
-	UPROPERTY(BlueprintReadWrite, EditAnywhere, AdvancedDisplay, Category = ShapeSettings, meta = (ClampMin = "64"))
+	/** Maximum allowable voxel count along main directions. This is a limit on ApproximationAccuracy. Max of 1290 (1290^3 is the last integer < 2^31, using a bigger number results in failures in TArray code & probably elsewhere) */
+	UPROPERTY(BlueprintReadWrite, EditAnywhere, AdvancedDisplay, Category = ShapeSettings, meta = (ClampMin = "64", ClampMax = "1290"))
 	int32 ClampVoxelDimension = 1024;
 
 	/** if enabled, we will attempt to auto-thicken thin parts or flat sheets */
@@ -802,7 +802,7 @@ struct FMeshApproximationSettings
 	int32 MultiSamplingAA = 0;
 
 	/** If Value is zero, use MaterialSettings resolution, otherwise override the render capture resolution */
-	UPROPERTY(BlueprintReadWrite, EditAnywhere, Category = MaterialSettings, meta = (ClampMin = "128"))
+	UPROPERTY(BlueprintReadWrite, EditAnywhere, Category = MaterialSettings, meta = (ClampMin = "0"))
 	int32 RenderCaptureResolution = 2048;
 
 	/** Material generation settings */
