Missing file for 16897632

diff --git a/Engine/Source/Runtime/Engine/Private/MeshCardRepresentation.cpp b/Engine/Source/Runtime/Engine/Private/MeshCardRepresentation.cpp
--- a/Engine/Source/Runtime/Engine/Private/MeshCardRepresentation.cpp
+++ b/Engine/Source/Runtime/Engine/Private/MeshCardRepresentation.cpp
@@ -52,10 +52,34 @@ static TAutoConsoleVariable<int32> CVarCardRepresentation(
 	TEXT(""),
 	ECVF_ReadOnly);
 
-static TAutoConsoleVariable<float> CVarCardRepresentationMinSurface(
-	TEXT("r.MeshCardRepresentation.MinSurface"),
-	0.2f,
-	TEXT("Min percentage of surface treshold to spawn a new card, [0;1] range."),
+static TAutoConsoleVariable<float> CVarCardRepresentationMinDensity(
+	TEXT("r.MeshCardRepresentation.MinDensity"),
+	0.1f,
+	TEXT("How much of filled area needs to be there to spawn a card, [0;1] range."),
+	ECVF_ReadOnly);
+
+static TAutoConsoleVariable<float> CVarCardRepresentationNormalTreshold(
+	TEXT("r.MeshCardRepresentation.NormalTreshold"),
+	0.25f,
+	TEXT("Normal treshold when surface elements should be clustered together."),
+	ECVF_ReadOnly);
+
+static TAutoConsoleVariable<float> CVarCardRepresentationDistanceTreshold(
+	TEXT("r.MeshCardRepresentation.DistanceTreshold"),
+	4.0f,
+	TEXT("Distance treshold (in surfels) when surface elements should be clustered together."),
+	ECVF_ReadOnly);
+
+static TAutoConsoleVariable<int32> CVarCardRepresentationSeedIterations(
+	TEXT("r.MeshCardRepresentation.SeedIterations"),
+	3,
+	TEXT("Max number of clustering iterations."),
+	ECVF_ReadOnly);
+
+static TAutoConsoleVariable<int32> CVarCardRepresentationGrowIterations(
+	TEXT("r.MeshCardRepresentation.GrowIterations"),
+	3,
+	TEXT("Max number of grow iterations."),
 	ECVF_ReadOnly);
 
 FCardRepresentationAsyncQueue* GCardRepresentationAsyncQueue = NULL;
@@ -63,15 +87,19 @@ FCardRepresentationAsyncQueue* GCardRepresentationAsyncQueue = NULL;
 #if WITH_EDITOR
 
 // DDC key for card representation data, must be changed when modifying the generation code or data format
-#define CARDREPRESENTATION_DERIVEDDATA_VER TEXT("0F32177F-E7CA-442E-BEEC-1AEBE4B462BC")
+#define CARDREPRESENTATION_DERIVEDDATA_VER TEXT("A1C7CF355B17B3E195BE69EAF6A10292")
 
 FString BuildCardRepresentationDerivedDataKey(const FString& InMeshKey)
 {
-	const float MinSurfaceThreshold = CVarCardRepresentationMinSurface.GetValueOnAnyThread();
+	const float MinDensity = CVarCardRepresentationMinDensity.GetValueOnAnyThread();
+	const float NormalTreshold = CVarCardRepresentationNormalTreshold.GetValueOnAnyThread();
+	const float DistanceTreshold = CVarCardRepresentationDistanceTreshold.GetValueOnAnyThread();
+	const int32 SeedIterations = CVarCardRepresentationSeedIterations.GetValueOnAnyThread();
+	const int32 GrowIterations = CVarCardRepresentationGrowIterations.GetValueOnAnyThread();
 
 	return FDerivedDataCacheInterface::BuildCacheKey(
 		TEXT("CARD"),
-		*FString::Printf(TEXT("%s_%s_%.3f"), *InMeshKey, CARDREPRESENTATION_DERIVEDDATA_VER, MinSurfaceThreshold),
+		*FString::Printf(TEXT("%s_%s_%.3f_%.3f_%.3f_%d_%d"), *InMeshKey, CARDREPRESENTATION_DERIVEDDATA_VER, MinDensity, NormalTreshold, DistanceTreshold, SeedIterations, GrowIterations),
 		TEXT(""));
 }
 
