Run a fake "game tick" during GenerateHLODsForGrid(), to allow graphic resources to be released
Moved DoCollectGarbage() from UWorldPartitionBuilder to FWorldPartitionHelpers

diff --git a/Engine/Source/Editor/UnrealEd/Private/WorldPartition/WorldPartitionBuilder.cpp b/Engine/Source/Editor/UnrealEd/Private/WorldPartition/WorldPartitionBuilder.cpp
--- a/Engine/Source/Editor/UnrealEd/Private/WorldPartition/WorldPartitionBuilder.cpp
+++ b/Engine/Source/Editor/UnrealEd/Private/WorldPartition/WorldPartitionBuilder.cpp
@@ -242,7 +242,8 @@ bool UWorldPartitionBuilder::Run(UWorld* World, FPackageSourceControlHelper& Pac
 						WorldPartition->UnloadEditorCells(LoadedBounds);
 						// Reset Loaded Bounds
 						LoadedBounds.Init();
-						DoCollectGarbage();
+
+						FWorldPartitionHelpers::DoCollectGarbage();
 					}
 				}
 			}
@@ -262,15 +263,3 @@ bool UWorldPartitionBuilder::Run(UWorld* World, FPackageSourceControlHelper& Pac
 		return RunInternal(World, BoundsToLoad, PackageHelper);
 	}
 }
-
-void UWorldPartitionBuilder::DoCollectGarbage() const
-{
-	const FPlatformMemoryStats MemStatsBefore = FPlatformMemory::GetStats();
-	CollectGarbage(RF_NoFlags, true);
-	const FPlatformMemoryStats MemStatsAfter = FPlatformMemory::GetStats();
-
-	UE_LOG(LogWorldPartitionBuilder, Display, TEXT("GC Performed - Freed Physical: %.2fGB, Freed Virtual: %.2fGB"),
-		((int64)MemStatsAfter.AvailablePhysical - (int64)MemStatsBefore.AvailablePhysical) / (1024.0 * 1024.0 * 1024.0),
-		((int64)MemStatsAfter.AvailableVirtual - (int64)MemStatsBefore.AvailableVirtual) / (1024.0 * 1024.0 * 1024.0)
-	);
-};
diff --git a/Engine/Source/Editor/UnrealEd/Private/WorldPartition/WorldPartitionHLODsBuilder.cpp b/Engine/Source/Editor/UnrealEd/Private/WorldPartition/WorldPartitionHLODsBuilder.cpp
--- a/Engine/Source/Editor/UnrealEd/Private/WorldPartition/WorldPartitionHLODsBuilder.cpp
+++ b/Engine/Source/Editor/UnrealEd/Private/WorldPartition/WorldPartitionHLODsBuilder.cpp
@@ -372,7 +372,7 @@ bool UWorldPartitionHLODsBuilder::SetupHLODActors(bool bCreateOnly)
 				// When running distributed builds, we wanna leave the machine clean, so added files are deleted, check'd out files are reverted
 				// and deleted files are restored.
 				WorldPartition->Uninitialize();
-				DoCollectGarbage();
+				FWorldPartitionHelpers::DoCollectGarbage();
 
 				ModifiedFiles.Append(SourceControlHelper->GetModifiedFiles());
 
@@ -474,7 +474,7 @@ bool UWorldPartitionHLODsBuilder::BuildHLODActors()
 
 		if (FWorldPartitionHelpers::HasExceededMaxMemory())
 		{
-			DoCollectGarbage();
+			FWorldPartitionHelpers::DoCollectGarbage();
 		}
 	}
 
@@ -487,7 +487,7 @@ bool UWorldPartitionHLODsBuilder::BuildHLODActors()
 		// When running distributed builds, we wanna leave the machine clean, so added files are deleted, check'd out files are reverted
 		// and deleted files are restored.
 		WorldPartition->Uninitialize();
-		DoCollectGarbage();
+		FWorldPartitionHelpers::DoCollectGarbage();
 
 		ModifiedFiles.Append(SourceControlHelper->GetModifiedFiles());
 
@@ -529,7 +529,7 @@ bool UWorldPartitionHLODsBuilder::DeleteHLODActors()
 	// When running distributed builds, we wanna leave the machine clean, so added files are deleted, checked out files are reverted
 	// and deleted files are restored.
 	WorldPartition->Uninitialize();
-	DoCollectGarbage();
+	FWorldPartitionHelpers::DoCollectGarbage();
 
 	int32 NumDeleted = 0;
 	for (const FString& PackageName : PackagesToDelete)
diff --git a/Engine/Source/Editor/UnrealEd/Public/WorldPartition/WorldPartitionBuilder.h b/Engine/Source/Editor/UnrealEd/Public/WorldPartition/WorldPartitionBuilder.h
--- a/Engine/Source/Editor/UnrealEd/Public/WorldPartition/WorldPartitionBuilder.h
+++ b/Engine/Source/Editor/UnrealEd/Public/WorldPartition/WorldPartitionBuilder.h
@@ -31,8 +31,6 @@ public:
 protected:
 	virtual bool RunInternal(UWorld* World, const FBox& Bounds, FPackageSourceControlHelper& PackageHelper) PURE_VIRTUAL(UWorldPartition::RunInternal, return false;);
 
-	void DoCollectGarbage() const;
-
 	int32 IterativeCellSize = 102400;
 	int32 IterativeCellOverlapSize = 0;
 	TSet<FName> DataLayerLabels;
diff --git a/Engine/Source/Runtime/Engine/Private/WorldPartition/RuntimeSpatialHash/RuntimeSpatialHashHLOD.cpp b/Engine/Source/Runtime/Engine/Private/WorldPartition/RuntimeSpatialHash/RuntimeSpatialHashHLOD.cpp
--- a/Engine/Source/Runtime/Engine/Private/WorldPartition/RuntimeSpatialHash/RuntimeSpatialHashHLOD.cpp
+++ b/Engine/Source/Runtime/Engine/Private/WorldPartition/RuntimeSpatialHash/RuntimeSpatialHashHLOD.cpp
@@ -91,20 +91,45 @@ static void DeletePackage(UWorldPartition* WorldPartition, FWorldPartitionActorD
 	}
 }
 
-static TArray<FGuid> GenerateHLODsForGrid(UWorldPartition* WorldPartition, const FActorContainerInstance& MainContainerInstance, const FSpatialHashRuntimeGrid& RuntimeGrid, uint32 HLODLevel, FHLODCreationContext& Context, ISourceControlHelper* SourceControlHelper, bool bCreateActorsOnly, const TArray<const FActorClusterInstance*>& ClusterInstances)
+static void GameTick(UWorld* InWorld)
 {
-	auto DoCollectGarbage = []()
+	static int32 TickRendering = 0;
+	static const int32 FlushRenderingFrequency = 256;
+
+	// Perform a GC when memory usage exceeds a given threshold
+	if (FWorldPartitionHelpers::HasExceededMaxMemory())
 	{
-		const FPlatformMemoryStats MemStatsBefore = FPlatformMemory::GetStats();
-		CollectGarbage(RF_NoFlags, true);
-		const FPlatformMemoryStats MemStatsAfter = FPlatformMemory::GetStats();
+		FWorldPartitionHelpers::DoCollectGarbage();
+	}
 
-		UE_LOG(LogWorldPartitionRuntimeSpatialHashHLOD, Display, TEXT("GC Performed - Freed Physical: %.2fGB, Freed Virtual: %.2fGB"),
-			((int64)MemStatsAfter.AvailablePhysical - (int64)MemStatsBefore.AvailablePhysical) / (1024.0 * 1024.0 * 1024.0),
-			((int64)MemStatsAfter.AvailableVirtual - (int64)MemStatsBefore.AvailableVirtual) / (1024.0 * 1024.0 * 1024.0)
-		);
-	};
+	// When running with -AllowCommandletRendering we want to flush
+	if (((++TickRendering % FlushRenderingFrequency) == 0) && IsAllowCommandletRendering())
+	{
+		if (FSceneInterface* Scene = InWorld->Scene)
+		{
+			// BeingFrame/EndFrame (taken from FEngineLoop)
+			uint64 CurrentFrameCounter = GFrameCounter;
+
+			ENQUEUE_RENDER_COMMAND(BeginFrame)([CurrentFrameCounter](FRHICommandListImmediate& RHICmdList)
+			{ 
+				GFrameNumberRenderThread++;
+				RHICmdList.BeginFrame();
+				FCoreDelegates::OnBeginFrameRT.Broadcast(); 
+			});
+
+			ENQUEUE_RENDER_COMMAND(EndFrame)([CurrentFrameCounter](FRHICommandListImmediate& RHICmdList)
+			{
+				FCoreDelegates::OnEndFrameRT.Broadcast();
+				RHICmdList.EndFrame(); 
+			});
 
+			FlushRenderingCommands();
+		}
+	}
+}
+
+static TArray<FGuid> GenerateHLODsForGrid(UWorldPartition* WorldPartition, const FActorContainerInstance& MainContainerInstance, const FSpatialHashRuntimeGrid& RuntimeGrid, uint32 HLODLevel, FHLODCreationContext& Context, ISourceControlHelper* SourceControlHelper, bool bCreateActorsOnly, const TArray<const FActorClusterInstance*>& ClusterInstances)
+{
 	const FBox WorldBounds = MainContainerInstance.Bounds;
 
 	const FSquare2DGridHelper PartitionedActors = GetPartitionedActors(WorldPartition, WorldBounds, RuntimeGrid, ClusterInstances);
@@ -152,6 +177,9 @@ static TArray<FGuid> GenerateHLODsForGrid(UWorldPartition* WorldPartition, const
 
 		for (const FSquare2DGridHelper::FGridLevel::FGridCellDataChunk& GridCellDataChunk : GridCell.GetDataChunks())
 		{
+			// Fake tick
+			GameTick(WorldPartition->GetWorld());
+
 			const bool bShouldGenerateHLODs = ShouldGenerateHLODs(GridCell, GridCellDataChunk);
 			if (bShouldGenerateHLODs)
 			{
@@ -197,13 +225,6 @@ static TArray<FGuid> GenerateHLODsForGrid(UWorldPartition* WorldPartition, const
 					{
 						if (!bCreateActorsOnly)
 						{
-							FStaticMeshCompilingManager::Get().FinishAllCompilation();
-
-							if (FWorldPartitionHelpers::HasExceededMaxMemory())
-							{
-								DoCollectGarbage();
-							}
-
 							CellHLODActor->BuildHLOD();
 						}
 
@@ -243,7 +264,7 @@ static TArray<FGuid> GenerateHLODsForGrid(UWorldPartition* WorldPartition, const
 
 	// Need to collect garbage here since some HLOD actors have been marked pending kill when destroying them
 	// and they may be loaded when generating the next HLOD layer.
-	DoCollectGarbage();
+	FWorldPartitionHelpers::DoCollectGarbage();
 
 	return GridHLODActors;
 }
diff --git a/Engine/Source/Runtime/Engine/Private/WorldPartition/WorldPartitionHelpers.cpp b/Engine/Source/Runtime/Engine/Private/WorldPartition/WorldPartitionHelpers.cpp
--- a/Engine/Source/Runtime/Engine/Private/WorldPartition/WorldPartitionHelpers.cpp
+++ b/Engine/Source/Runtime/Engine/Private/WorldPartition/WorldPartitionHelpers.cpp
@@ -5,6 +5,9 @@
 #include "WorldPartition/WorldPartition.h"
 #include "WorldPartition/WorldPartitionEditorHash.h"
 
+DEFINE_LOG_CATEGORY_STATIC(LogWorldPartitionHelpers, Log, All);
+
+
 #if WITH_EDITOR
 
 void FWorldPartitionHelpers::ForEachIntersectingActorDesc(UWorldPartition* WorldPartition, const FBox& Box, TSubclassOf<AActor> ActorClass, TFunctionRef<bool(const FWorldPartitionActorDesc*)> Predicate)
@@ -106,4 +109,16 @@ bool FWorldPartitionHelpers::HasExceededMaxMemory()
 	return bHasExceededMaxMemory || (FPlatformTime::Seconds() - GetLastGCTime()) > 30;
 };
 
+void FWorldPartitionHelpers::DoCollectGarbage()
+{
+	const FPlatformMemoryStats MemStatsBefore = FPlatformMemory::GetStats();
+	CollectGarbage(RF_NoFlags, true);
+	const FPlatformMemoryStats MemStatsAfter = FPlatformMemory::GetStats();
+
+	UE_LOG(LogWorldPartition, Log, TEXT("GC Performed - Available Physical: %.2fGB, Available Virtual: %.2fGB"),
+		(int64)MemStatsAfter.AvailablePhysical / (1024.0 * 1024.0 * 1024.0),
+		(int64)MemStatsAfter.AvailableVirtual / (1024.0 * 1024.0 * 1024.0)
+	);
+};
+
 #endif // #if WITH_EDITOR
diff --git a/Engine/Source/Runtime/Engine/Public/WorldPartition/WorldPartitionHelpers.h b/Engine/Source/Runtime/Engine/Public/WorldPartition/WorldPartitionHelpers.h
--- a/Engine/Source/Runtime/Engine/Public/WorldPartition/WorldPartitionHelpers.h
+++ b/Engine/Source/Runtime/Engine/Public/WorldPartition/WorldPartitionHelpers.h
@@ -35,6 +35,7 @@ public:
 	static void ForEachActorWithLoading(UWorldPartition* WorldPartition, TSubclassOf<AActor> ActorClass, TFunctionRef<bool(const FWorldPartitionActorDesc*)> Func);
 
 	static bool HasExceededMaxMemory();
+	static void DoCollectGarbage();
 };
 
 #endif // WITH_EDITOR
