Fix for crash in chase

diff --git a/Engine/Source/Runtime/Engine/Private/ActiveSound.cpp b/Engine/Source/Runtime/Engine/Private/ActiveSound.cpp
--- a/Engine/Source/Runtime/Engine/Private/ActiveSound.cpp
+++ b/Engine/Source/Runtime/Engine/Private/ActiveSound.cpp
@@ -1680,6 +1680,11 @@ void FActiveSound::UpdateAttenuation(float DeltaTime, FSoundParseParameters& Par
 {
 	// Get the attenuation settings to use for this application to the active sound
 	const FSoundAttenuationSettings* Settings = SettingsAttenuationNode ? SettingsAttenuationNode : &AttenuationSettings;
+	if (!Settings)
+	{
+		UE_LOG(LogAudio, Warning, TEXT("No attenuation settings found for active sound."));
+		return;
+	}
 
 	// Reset Focus data and recompute if necessary
 	FAttenuationFocusData FocusDataToApply;
@@ -1768,8 +1773,10 @@ void FActiveSound::UpdateAttenuation(float DeltaTime, FSoundParseParameters& Par
 		if (ParseParams.SoundClass)
 		{
 			FSoundClassDynamicProperties* DynamicSoundClassProperties = AudioDevice->GetSoundClassDynamicProperties(ParseParams.SoundClass);
-
-			FocusDataToApply.DistanceScale *= FMath::Max(DynamicSoundClassProperties->AttenuationScaleParam.GetValue(), 0.0f);
+			if (DynamicSoundClassProperties)
+			{
+				FocusDataToApply.DistanceScale *= FMath::Max(DynamicSoundClassProperties->AttenuationScaleParam.GetValue(), 0.0f);
+			}
 		}
 
 		if (Settings->AttenuationShape == EAttenuationShape::Sphere)
