Add support for deleting light weight instances.
Make light weight instance managers be per data layer instead of per level
Add stable unique identifiers to handles for light weight instances so that they can be managed outside of the system with accurate bookkeeping.

diff --git a/Engine/Plugins/LightWeightInstancesEditor/Source/LightWeightInstancesEditor/LightWeightInstancesEditor.Build.cs b/Engine/Plugins/LightWeightInstancesEditor/Source/LightWeightInstancesEditor/LightWeightInstancesEditor.Build.cs
--- a/Engine/Plugins/LightWeightInstancesEditor/Source/LightWeightInstancesEditor/LightWeightInstancesEditor.Build.cs
+++ b/Engine/Plugins/LightWeightInstancesEditor/Source/LightWeightInstancesEditor/LightWeightInstancesEditor.Build.cs
@@ -41,8 +41,14 @@ public class LightWeightInstancesEditor : ModuleRules
 				// ... add private dependencies that you statically link with here ...	
 			}
 			);
-		
-		
+
+		if (Target.Type == TargetType.Editor)
+		{
+			PrivateDependencyModuleNames.Add("DataLayerEditor");
+		}
+
+
+
 		DynamicallyLoadedModuleNames.AddRange(
 			new string[]
 			{
diff --git a/Engine/Plugins/LightWeightInstancesEditor/Source/LightWeightInstancesEditor/Private/LightWeightInstancesEditor.cpp b/Engine/Plugins/LightWeightInstancesEditor/Source/LightWeightInstancesEditor/Private/LightWeightInstancesEditor.cpp
--- a/Engine/Plugins/LightWeightInstancesEditor/Source/LightWeightInstancesEditor/Private/LightWeightInstancesEditor.cpp
+++ b/Engine/Plugins/LightWeightInstancesEditor/Source/LightWeightInstancesEditor/Private/LightWeightInstancesEditor.cpp
@@ -7,6 +7,7 @@
 #include "Framework/MultiBox/MultiBoxExtender.h"
 #include "Framework/Commands/UICommandList.h"
 #include "GameFramework/Actor.h"
+#include "DataLayer/DataLayerEditorSubsystem.h"
 #endif
 #include "GameFramework/LightWeightInstanceManager.h"
 #include "GameFramework/LightWeightInstanceSubsystem.h"
@@ -106,10 +107,17 @@ void FLightWeightInstancesEditorModule::ConvertActorsToLWIsUIAction(const TArray
 		UE_LOG(LogLWIEditor, Log, TEXT("Unable to convert unspecified actors to light weight instances"));
 		return;
 	}
-	ALightWeightInstanceManager* Manager = FLightWeightInstanceSubsystem::Get().FindOrAddLightWeightInstanceManager(InActors[0]->GetClass(), InActors[0]->GetLevel());
-	check(Manager);
+	
 	for (AActor* Actor : InActors)
 	{
+		// use the first layer the actor is in if it's in multiple layers
+		TArray<const UDataLayer*> DataLayers = Actor->GetDataLayerObjects();
+		const UDataLayer* Layer = DataLayers.Num() > 0 ? DataLayers[0] : nullptr;
+
+		ALightWeightInstanceManager* Manager = FLightWeightInstanceSubsystem::Get().FindOrAddLightWeightInstanceManager(Actor->GetClass(), Layer, Actor->GetWorld());
+		check(Manager);
+		UDataLayerEditorSubsystem::Get()->OnActorDataLayersChanged().Broadcast(Manager);
+
 		Manager->ConvertActorToLightWeightInstance(Actor);
 	}
 }
diff --git a/Engine/Source/Runtime/Engine/Classes/Engine/EngineTypes.h b/Engine/Source/Runtime/Engine/Classes/Engine/EngineTypes.h
--- a/Engine/Source/Runtime/Engine/Classes/Engine/EngineTypes.h
+++ b/Engine/Source/Runtime/Engine/Classes/Engine/EngineTypes.h
@@ -2081,9 +2081,6 @@ struct ENGINE_API FActorInstanceHandle
 	FRotator GetRotation() const;
 	FTransform GetTransform() const;
 
-	ULevel* GetLevel() const;
-	bool IsInLevel(ULevel* InLevel) const;
-
 	FName GetFName() const;
 	FString GetName() const;
 
@@ -2118,6 +2115,8 @@ struct ENGINE_API FActorInstanceHandle
 
 	friend ENGINE_API FArchive& operator<<(FArchive& Ar, FActorInstanceHandle& Handle);
 
+	uint32 GetInstanceUID() const { return InstanceUID; }
+
 private:
 	/**
 	 * helper functions that let us treat the actor pointer as a UObject in templated functions
@@ -2138,6 +2137,9 @@ private:
 
 	/** Identifies the instance within the manager */
 	int32 InstanceIndex;
+
+	/** Unique identifier for instances represented by the handle */
+	uint32 InstanceUID;
 };
 
 template <typename T>
diff --git a/Engine/Source/Runtime/Engine/Classes/Engine/World.h b/Engine/Source/Runtime/Engine/Classes/Engine/World.h
--- a/Engine/Source/Runtime/Engine/Classes/Engine/World.h
+++ b/Engine/Source/Runtime/Engine/Classes/Engine/World.h
@@ -1422,6 +1422,10 @@ public:
 	UPROPERTY(Transient)
 	TObjectPtr<class UPhysicsFieldComponent> PhysicsField;
 
+	/** Tracks the last assigned unique id for light weight instances in this world. */
+	UPROPERTY()
+	uint32 LWILastAssignedUID;
+
 private:
 
 	/** Array of components that need to wait on tasks before end of frame updates */
diff --git a/Engine/Source/Runtime/Engine/Classes/GameFramework/LightWeightInstanceBlueprintFunctionLibrary.h b/Engine/Source/Runtime/Engine/Classes/GameFramework/LightWeightInstanceBlueprintFunctionLibrary.h
--- a/Engine/Source/Runtime/Engine/Classes/GameFramework/LightWeightInstanceBlueprintFunctionLibrary.h
+++ b/Engine/Source/Runtime/Engine/Classes/GameFramework/LightWeightInstanceBlueprintFunctionLibrary.h
@@ -16,7 +16,7 @@ public:
 
 	// Returns a handle to a new light weight instance that represents an object of type ActorClass
 	UFUNCTION(BlueprintCallable, Category = "Light Weight Instance")
-	static FActorInstanceHandle CreateNewLightWeightInstance(UClass* ActorClass, FTransform Transform, ULevel* Level);
+	static FActorInstanceHandle CreateNewLightWeightInstance(UClass* ActorClass, FTransform Transform, UDataLayer* Layer, UWorld* World);
 
 	// Returns a handle to the light weight representation and destroys Actor if successful; Returns a handle to Actor otherwise
 	UFUNCTION(BlueprintCallable, Category = "Light Weight Instance")
diff --git a/Engine/Source/Runtime/Engine/Classes/GameFramework/LightWeightInstanceManager.h b/Engine/Source/Runtime/Engine/Classes/GameFramework/LightWeightInstanceManager.h
--- a/Engine/Source/Runtime/Engine/Classes/GameFramework/LightWeightInstanceManager.h
+++ b/Engine/Source/Runtime/Engine/Classes/GameFramework/LightWeightInstanceManager.h
@@ -172,6 +172,10 @@ protected:
 	// Gets the class that implements interfaces for light weight instances owned by this manager
 	virtual UClass* GetInterfaceClass() const;
 
+	// Helper functions for converting between our internal storage indices and the indices used by external bookkeeping
+	virtual int32 ConvertInternalIndexToHandleIndex(int32 InInternalIndex) const { return InInternalIndex; }
+	virtual int32 ConvertHandleIndexToInternalIndex(int32 InHandleIndex) const { return InHandleIndex; }
+
 	//
 	// Data and replication functions
 	//
diff --git a/Engine/Source/Runtime/Engine/Classes/GameFramework/LightWeightInstanceStaticMeshManager.h b/Engine/Source/Runtime/Engine/Classes/GameFramework/LightWeightInstanceStaticMeshManager.h
--- a/Engine/Source/Runtime/Engine/Classes/GameFramework/LightWeightInstanceStaticMeshManager.h
+++ b/Engine/Source/Runtime/Engine/Classes/GameFramework/LightWeightInstanceStaticMeshManager.h
@@ -35,7 +35,7 @@ protected:
 
 	void RemoveInstanceFromRendering(int32 DataIndex);
 
-	void PostRemoveInstanceFromRendering(int32 RenderingIndex);
+	void PostRemoveInstanceFromRendering();
 
 	// sets the parameters on the instanced static mesh component
 	virtual void SetInstancedStaticMeshParams();
@@ -68,6 +68,9 @@ protected:
 	virtual void DuplicateLWIInstances(TArrayView<const int32> DataIndices, TArray<int32>& OutNewDataIndices);
 	void GetLWIDataIndices(TArrayView<const FSMInstanceId> InstanceIds, TArray<int32>& OutDataIndices) const;
 
+	virtual int32 ConvertInternalIndexToHandleIndex(int32 InternalIndex) const override;
+	virtual int32 ConvertHandleIndexToInternalIndex(int32 HandleIndex) const override;
+
 protected:
 
 	UPROPERTY(EditInstanceOnly, BlueprintReadOnly, Category = Debug, ReplicatedUsing = OnRep_StaticMesh)
@@ -87,4 +90,7 @@ protected:
 	TArray<int32> RenderingIndicesToDataIndices;
 	UPROPERTY(Replicated)
 	TArray<int32> DataIndicesToRenderingIndices;
+
+	// Data indices that we are going to delete later in the frame, we impose a small delay to group deletions to avoid cases where we might try to delete the same index multiple times in a frame.
+	TArray<int32> DataIndicesToBeDeleted;
 };
diff --git a/Engine/Source/Runtime/Engine/Classes/GameFramework/LightWeightInstanceSubsystem.h b/Engine/Source/Runtime/Engine/Classes/GameFramework/LightWeightInstanceSubsystem.h
--- a/Engine/Source/Runtime/Engine/Classes/GameFramework/LightWeightInstanceSubsystem.h
+++ b/Engine/Source/Runtime/Engine/Classes/GameFramework/LightWeightInstanceSubsystem.h
@@ -38,11 +38,11 @@ struct ENGINE_API FLightWeightInstanceSubsystem
 	ALightWeightInstanceManager* FindLightWeightInstanceManager(const FActorInstanceHandle& Handle) const;
 
 	// Returns the instance manager that handles actors of type ActorClass in level Level
-	ALightWeightInstanceManager* FindLightWeightInstanceManager(UClass* ActorClass, ULevel* Level) const;
+	ALightWeightInstanceManager* FindLightWeightInstanceManager(UClass* ActorClass, const UDataLayer* Layer) const;
 
 	// Returns the instance manager that handles instances of type Class that live in Level
 	UFUNCTION(Server, Unreliable)
-	ALightWeightInstanceManager* FindOrAddLightWeightInstanceManager(UClass* ActorClass, ULevel* Level);
+	ALightWeightInstanceManager* FindOrAddLightWeightInstanceManager(UClass* ActorClass, const UDataLayer* Layer, UWorld* World);
 
 	// Returns the actor specified by Handle. This may require loading and creating the actor object.
 	AActor* FetchActor(const FActorInstanceHandle& Handle);
@@ -63,7 +63,10 @@ struct ENGINE_API FLightWeightInstanceSubsystem
 	bool IsInLevel(const FActorInstanceHandle& Handle, const ULevel* InLevel);
 
 	// Returns a handle to a new light weight instance that represents an object of type ActorClass
-	FActorInstanceHandle CreateNewLightWeightInstance(UClass* ActorClass, FLWIData* InitData, ULevel* Level);
+	FActorInstanceHandle CreateNewLightWeightInstance(UClass* ActorClass, FLWIData* InitData, UDataLayer* Layer, UWorld* World);
+
+	// deletes the instance identified by Handle
+	void DeleteInstance(const FActorInstanceHandle& Handle);
 
 	// Returns true if the handle can return an object that implements the interface U
 	template<typename U>
@@ -90,7 +93,7 @@ struct ENGINE_API FLightWeightInstanceSubsystem
 	}
 
 protected:
-	// Returns the class of the insance manager best suited to support instances of type ActorClass
+	// Returns the class of the instance manager best suited to support instances of type ActorClass
 	UClass* FindBestInstanceManagerClass(const UClass* ActorClass);
 
 	// Returns the index associated with Manager
diff --git a/Engine/Source/Runtime/Engine/Private/EngineTypes.cpp b/Engine/Source/Runtime/Engine/Private/EngineTypes.cpp
--- a/Engine/Source/Runtime/Engine/Private/EngineTypes.cpp
+++ b/Engine/Source/Runtime/Engine/Private/EngineTypes.cpp
@@ -241,6 +241,7 @@ FActorInstanceHandle::FActorInstanceHandle()
 	: Actor(nullptr)
 	, ManagerIndex(INDEX_NONE)
 	, InstanceIndex(INDEX_NONE)
+	, InstanceUID(0)
 {
 	// do nothing
 }
@@ -252,7 +253,14 @@ FActorInstanceHandle::FActorInstanceHandle(AActor* InActor)
 {
 	if (InActor)
 	{
-		if (ALightWeightInstanceManager* LWIManager = FLightWeightInstanceSubsystem::Get().FindLightWeightInstanceManager(InActor->StaticClass(), InActor->GetLevel()))
+#if WITH_EDITOR
+		// use the first layer the actor is in if it's in multiple layers
+		TArray<const UDataLayer*> DataLayers = InActor->GetDataLayerObjects();
+		const UDataLayer* Layer = DataLayers.Num() > 0 ? DataLayers[0] : nullptr;
+#else
+		const UDataLayer* Layer = nullptr;
+#endif // WITH_EDITOR
+		if (ALightWeightInstanceManager* LWIManager = FLightWeightInstanceSubsystem::Get().FindLightWeightInstanceManager(InActor->StaticClass(), Layer))
 		{
 			InstanceIndex = LWIManager->FindIndexForActor(InActor);
 			if (InstanceIndex != INDEX_NONE)
@@ -277,11 +285,16 @@ FActorInstanceHandle::FActorInstanceHandle(int32 InManagerIndex, int32 InInstanc
 		{
 			Actor = *FoundActor;
 		}
+
+		UWorld* World = Manager->GetWorld();
+		ensure(World);
+		InstanceUID = World->LWILastAssignedUID++;
 	}
 	else
 	{
 		ManagerIndex = INDEX_NONE;
 		InstanceIndex = INDEX_NONE;
+		InstanceUID = 0;
 	}
 }
 
@@ -289,6 +302,7 @@ FActorInstanceHandle::FActorInstanceHandle(ALightWeightInstanceManager* Manager,
 	: Actor(nullptr)
 	, ManagerIndex(INDEX_NONE)
 	, InstanceIndex(InInstanceIndex)
+	, InstanceUID(0)
 {
 	if (ensure(Manager))
 	{
@@ -300,6 +314,10 @@ FActorInstanceHandle::FActorInstanceHandle(ALightWeightInstanceManager* Manager,
 
 		ManagerIndex = FLightWeightInstanceSubsystem::Get().GetManagerIndex(Manager);
 		ensure(ManagerIndex != INDEX_NONE);
+
+		UWorld* World = Manager->GetWorld();
+		ensure(World);
+		InstanceUID = World->LWILastAssignedUID++;
 	}
 }
 
@@ -308,6 +326,7 @@ FActorInstanceHandle::FActorInstanceHandle(const FActorInstanceHandle& Other)
 	Actor = Other.Actor;
 	ManagerIndex = Other.ManagerIndex;
 	InstanceIndex = Other.InstanceIndex;
+	InstanceUID = Other.InstanceUID;
 }
 
 bool FActorInstanceHandle::IsValid() const
@@ -400,21 +419,6 @@ FTransform FActorInstanceHandle::GetTransform() const
 	return FTransform();
 }
 
-ULevel* FActorInstanceHandle::GetLevel() const
-{
-	if (IsActorValid())
-	{
-		return Actor->GetLevel();
-	}
-
-	return nullptr;
-}
-
-bool FActorInstanceHandle::IsInLevel(ULevel* Level) const
-{
-	return Level && Level == GetLevel();
-}
-
 FName FActorInstanceHandle::GetFName() const
 {
 	if (IsActorValid())
@@ -536,7 +540,16 @@ bool FActorInstanceHandle::operator==(const AActor* OtherActor) const
 	}
 
 	// we don't have an actor so see if we can look up an instance associated with OtherActor and see if we refer to the same instance
-	if (ALightWeightInstanceManager* Manager = FLightWeightInstanceSubsystem::Get().FindLightWeightInstanceManager(OtherActor->StaticClass(), OtherActor->GetLevel()))
+
+#if WITH_EDITOR
+	// use the first layer the actor is in if it's in multiple layers
+	TArray<const UDataLayer*> DataLayers = OtherActor->GetDataLayerObjects();
+	const UDataLayer* Layer = DataLayers.Num() > 0 ? DataLayers[0] : nullptr;
+#else
+	const UDataLayer* Layer = nullptr;
+#endif // WITH_EDITOR
+
+	if (ALightWeightInstanceManager* Manager = FLightWeightInstanceSubsystem::Get().FindLightWeightInstanceManager(OtherActor->StaticClass(), Layer))
 	{
 		if (FLightWeightInstanceSubsystem::Get().GetManagerIndex(Manager) != ManagerIndex)
 		{
diff --git a/Engine/Source/Runtime/Engine/Private/LightWeightInstanceBlueprintFunctionLibrary.cpp b/Engine/Source/Runtime/Engine/Private/LightWeightInstanceBlueprintFunctionLibrary.cpp
--- a/Engine/Source/Runtime/Engine/Private/LightWeightInstanceBlueprintFunctionLibrary.cpp
+++ b/Engine/Source/Runtime/Engine/Private/LightWeightInstanceBlueprintFunctionLibrary.cpp
@@ -4,13 +4,13 @@
 #include "GameFramework/LightWeightInstanceSubsystem.h"
 #include "GameFramework/LightWeightInstanceManager.h"
 
-FActorInstanceHandle ULightWeightInstanceBlueprintFunctionLibrary::CreateNewLightWeightInstance(UClass* InActorClass, FTransform InTransform, ULevel* InLevel)
+FActorInstanceHandle ULightWeightInstanceBlueprintFunctionLibrary::CreateNewLightWeightInstance(UClass* InActorClass, FTransform InTransform, UDataLayer* InLayer, UWorld* World)
 {
 	// set up initialization data
 	FLWIData PerInstanceData;
 	PerInstanceData.Transform = InTransform;
 
-	return FLightWeightInstanceSubsystem::Get().CreateNewLightWeightInstance(InActorClass, &PerInstanceData, InLevel);
+	return FLightWeightInstanceSubsystem::Get().CreateNewLightWeightInstance(InActorClass, &PerInstanceData, InLayer, World);
 }
 
 FActorInstanceHandle ULightWeightInstanceBlueprintFunctionLibrary::ConvertActorToLightWeightInstance(AActor* InActor)
@@ -20,8 +20,15 @@ FActorInstanceHandle ULightWeightInstanceBlueprintFunctionLibrary::ConvertActorT
 		return FActorInstanceHandle();
 	}
 
-	// Get or create a light weight instance for this class and level
-	if (ALightWeightInstanceManager* LWIManager = FLightWeightInstanceSubsystem::Get().FindOrAddLightWeightInstanceManager(InActor->GetClass(), InActor->GetLevel()))
+	// Get or create a light weight instance for this class and layer
+	// use the first layer the actor is in if it's in multiple layers
+#if WITH_EDITOR
+	TArray<const UDataLayer*> DataLayers = InActor->GetDataLayerObjects();
+	const UDataLayer* Layer = DataLayers.Num() > 0 ? DataLayers[0] : nullptr;
+#else
+	const UDataLayer* Layer = nullptr;
+#endif // WITH_EDITOR
+	if (ALightWeightInstanceManager* LWIManager = FLightWeightInstanceSubsystem::Get().FindOrAddLightWeightInstanceManager(InActor->GetClass(), Layer, InActor->GetWorld()))
 	{
 		return LWIManager->ConvertActorToLightWeightInstance(InActor);
 	}
diff --git a/Engine/Source/Runtime/Engine/Private/LightWeightInstanceManager.cpp b/Engine/Source/Runtime/Engine/Private/LightWeightInstanceManager.cpp
--- a/Engine/Source/Runtime/Engine/Private/LightWeightInstanceManager.cpp
+++ b/Engine/Source/Runtime/Engine/Private/LightWeightInstanceManager.cpp
@@ -69,6 +69,11 @@ AActor* ALightWeightInstanceManager::ConvertInstanceToActor(const FActorInstance
 	{
 		return Actors[Handle.GetInstanceIndex()];
 	}
+	// if we're pointing at an invalid index we can't return an actor
+	if (ValidIndices.Num() >= Handle.GetInstanceIndex() || ValidIndices[Handle.GetInstanceIndex()] == false)
+	{
+		return nullptr;
+	}
 
 	FActorSpawnParameters SpawnParams;
 	SetSpawnParameters(SpawnParams);
@@ -122,6 +127,7 @@ FActorInstanceHandle ALightWeightInstanceManager::ConvertActorToLightWeightInsta
 		if (Idx == INDEX_NONE)
 		{
 			Idx = AddNewInstance(Data);
+			Idx = ConvertInternalIndexToHandleIndex(Idx);
 		}
 		else
 		{
@@ -131,6 +137,7 @@ FActorInstanceHandle ALightWeightInstanceManager::ConvertActorToLightWeightInsta
 		// Update our handle
 		ReturnHandle.ManagerIndex = FLightWeightInstanceSubsystem::Get().GetManagerIndex(this);
 		ReturnHandle.InstanceIndex = Idx;
+		ReturnHandle.InstanceUID = InActor->GetWorld()->LWILastAssignedUID++;
 
 		// cleanup
 		InActor->Destroy();
diff --git a/Engine/Source/Runtime/Engine/Private/LightWeightInstanceStaticMeshManager.cpp b/Engine/Source/Runtime/Engine/Private/LightWeightInstanceStaticMeshManager.cpp
--- a/Engine/Source/Runtime/Engine/Private/LightWeightInstanceStaticMeshManager.cpp
+++ b/Engine/Source/Runtime/Engine/Private/LightWeightInstanceStaticMeshManager.cpp
@@ -4,6 +4,9 @@
 #include "Components/HierarchicalInstancedStaticMeshComponent.h"
 #include "Net/UnrealNetwork.h"
 #include "TimerManager.h"
+#if WITH_EDITOR
+#include "Editor.h"
+#endif // WITH_EDITOR
 
 ALightWeightInstanceStaticMeshManager::ALightWeightInstanceStaticMeshManager(const FObjectInitializer& ObjectInitializer)
 {
@@ -90,34 +93,59 @@ void ALightWeightInstanceStaticMeshManager::RemoveInstanceFromRendering(int32 Da
 {
 	if (IsIndexValid(DataIndex))
 	{
-		const int32 RenderingIndex = DataIndicesToRenderingIndices[DataIndex];
-
-		if (RenderingIndex != INDEX_NONE)
+		if (DataIndicesToBeDeleted.IsEmpty())
 		{
-			if (InstancedStaticMeshComponent)
+			// We can't remove the indices right away because we could have a situation where we receive multiple requests to convert an instance to an actor in the same frame	
+#if WITH_EDITOR
+			if (GEditor)
 			{
-				InstancedStaticMeshComponent->RemoveInstance(RenderingIndex);
+				GEditor->GetTimerManager()->SetTimerForNextTick([this]() { PostRemoveInstanceFromRendering(); });
+			}
+			else
+#endif // WITH_EDITOR
+			{
+				GetWorld()->GetTimerManager().SetTimerForNextTick([this]() { PostRemoveInstanceFromRendering(); });
 			}
-
-			// We can't remove the indices right away because we could have a situation where we receive multiple requests to convert an instance to an actor in the same frame			
-			GetWorld()->GetTimerManager().SetTimerForNextTick([this, RenderingIndex]() { PostRemoveInstanceFromRendering(RenderingIndex); });
 		}
+
+		DataIndicesToBeDeleted.AddUnique(DataIndex);
 	}
 }
 
-void ALightWeightInstanceStaticMeshManager::PostRemoveInstanceFromRendering(int32 RenderingIndex)
+void ALightWeightInstanceStaticMeshManager::PostRemoveInstanceFromRendering()
 {
-	RenderingIndicesToDataIndices.RemoveAtSwap(RenderingIndex);
+	TArray<int32> RenderingIndicesToBeDeleted;
+	for (int32 DataIndex : DataIndicesToBeDeleted)
+	{
+		RenderingIndicesToBeDeleted.Add(DataIndicesToRenderingIndices[DataIndex]);
+	}
 
-	// fix up the other side of the map to match the change we just made
-	// if we removed the last element than nothing was moved so we're done
-	if (RenderingIndex < RenderingIndicesToDataIndices.Num())
+	RenderingIndicesToBeDeleted.Sort();
+	int32 Count = RenderingIndicesToBeDeleted.Num();
+	for (int32 Idx = Count-1; Idx >= 0; Idx--)
 	{
-		// find the data index that corresponds with the changed rendering index
-		const int32 ShiftedDataIndex = RenderingIndicesToDataIndices[RenderingIndex];
+		int32 RenderingIndex = RenderingIndicesToBeDeleted[Idx];
+		if (RenderingIndex != INDEX_NONE)
+		{
+			if (InstancedStaticMeshComponent)
+			{
+				InstancedStaticMeshComponent->RemoveInstance(RenderingIndex);
+			}
+			RenderingIndicesToDataIndices.RemoveAtSwap(RenderingIndex);
+
+			// fix up the other side of the map to match the change we just made
+			// if we removed the last element than nothing was moved so we're done
+			if (RenderingIndex < RenderingIndicesToDataIndices.Num())
+			{
+				// find the data index that corresponds with the changed rendering index
+				const int32 ShiftedDataIndex = RenderingIndicesToDataIndices[RenderingIndex];
 
-		DataIndicesToRenderingIndices[ShiftedDataIndex] = RenderingIndex;
+				DataIndicesToRenderingIndices[ShiftedDataIndex] = RenderingIndex;
+			}
+		}
 	}
+
+	DataIndicesToBeDeleted.Reset();
 }
 
 void ALightWeightInstanceStaticMeshManager::GetLifetimeReplicatedProps(TArray<FLifetimeProperty>& OutLifetimeProps) const
@@ -349,3 +377,13 @@ void ALightWeightInstanceStaticMeshManager::GetLWIDataIndices(TArrayView<const F
 		}
 	}
 }
+
+int32 ALightWeightInstanceStaticMeshManager::ConvertInternalIndexToHandleIndex(int32 InInternalIndex) const
+{
+	return DataIndicesToRenderingIndices[InInternalIndex];
+}
+
+int32 ALightWeightInstanceStaticMeshManager::ConvertHandleIndexToInternalIndex(int32 InHandleIndex) const
+{
+	return RenderingIndicesToDataIndices[InHandleIndex];
+}
diff --git a/Engine/Source/Runtime/Engine/Private/LightWeightInstanceSubsystem.cpp b/Engine/Source/Runtime/Engine/Private/LightWeightInstanceSubsystem.cpp
--- a/Engine/Source/Runtime/Engine/Private/LightWeightInstanceSubsystem.cpp
+++ b/Engine/Source/Runtime/Engine/Private/LightWeightInstanceSubsystem.cpp
@@ -3,6 +3,7 @@
 #include "GameFramework/LightWeightInstanceSubsystem.h"
 #include "Engine/World.h"
 #include "UObject/UObjectIterator.h"
+#include "WorldPartition/DataLayer/DataLayer.h"
 
 DEFINE_LOG_CATEGORY(LogLightWeightInstance);
 
@@ -40,9 +41,21 @@ ALightWeightInstanceManager* FLightWeightInstanceSubsystem::FindLightWeightInsta
 	{
 		for (ALightWeightInstanceManager* LWInstance : LWInstanceManagers)
 		{
-			if (Handle.Actor->GetClass() == LWInstance->GetRepresentedClass() && Handle.Actor->GetLevel() == LWInstance->GetLevel())
+			if (Handle.Actor->GetClass() == LWInstance->GetRepresentedClass())
 			{
-				return LWInstance;
+#if WITH_EDITOR
+				// make sure the data layers match
+				TArray<const UDataLayer*> ManagerLayers = LWInstance->GetDataLayerObjects();
+				const UDataLayer* ManagerLayer = ManagerLayers.Num() > 0 ? ManagerLayers[0] : nullptr;
+
+				TArray<const UDataLayer*> ActorLayers = LWInstance->GetDataLayerObjects();
+				const UDataLayer* ActorLayer = ActorLayers.Num() > 0 ? ActorLayers[0] : nullptr;
+
+				if (ManagerLayer == ActorLayer)
+#endif // WITH_EDITOR
+				{
+					return LWInstance;
+				}
 			}
 		}
 	}
@@ -50,9 +63,9 @@ ALightWeightInstanceManager* FLightWeightInstanceSubsystem::FindLightWeightInsta
 	return nullptr;
 }
 
-ALightWeightInstanceManager* FLightWeightInstanceSubsystem::FindLightWeightInstanceManager(UClass* ActorClass, ULevel* Level) const
+ALightWeightInstanceManager* FLightWeightInstanceSubsystem::FindLightWeightInstanceManager(UClass* ActorClass, const UDataLayer* Layer) const
 {
-	if (ActorClass == nullptr || Level == nullptr)
+	if (ActorClass == nullptr)
 	{
 		return nullptr;
 	}
@@ -60,28 +73,38 @@ ALightWeightInstanceManager* FLightWeightInstanceSubsystem::FindLightWeightInsta
 	// see if we already have a match
 	for (ALightWeightInstanceManager* Instance : LWInstanceManagers)
 	{
-		if (Instance->GetRepresentedClass() == ActorClass && Instance->GetLevel() == Level)
+		if (Instance->GetRepresentedClass() == ActorClass)
 		{
-			return Instance;
+#if WITH_EDITOR
+			if (!Layer || (Instance->SupportsDataLayer() && Instance->ContainsDataLayer(Layer)))
+#endif // WITH_EDITOR
+			{
+				return Instance;
+			}
 		}
 	}
 
 	return nullptr;
 }
 
-ALightWeightInstanceManager* FLightWeightInstanceSubsystem::FindOrAddLightWeightInstanceManager(UClass* ActorClass, ULevel* Level)
+ALightWeightInstanceManager* FLightWeightInstanceSubsystem::FindOrAddLightWeightInstanceManager(UClass* ActorClass, const UDataLayer* Layer, UWorld* World)
 {
-	if (ActorClass == nullptr || Level == nullptr)
+	if (ActorClass == nullptr || World == nullptr)
 	{
 		return nullptr;
 	}
 	
 	// see if we already have a match
-	for (ALightWeightInstanceManager* Instance : LWInstanceManagers)
+	for (ALightWeightInstanceManager* InstanceManager : LWInstanceManagers)
 	{
-		if (Instance->GetRepresentedClass() == ActorClass && Instance->GetLevel() == Level)
+		if (InstanceManager->GetRepresentedClass() == ActorClass)
 		{
-			return Instance;
+#if WITH_EDITOR
+			if (!Layer || (InstanceManager->SupportsDataLayer() && InstanceManager->ContainsDataLayer(Layer)))
+#endif // WITH_EDITOR
+			{
+				return InstanceManager;
+			}
 		}
 	}
 
@@ -95,17 +118,25 @@ ALightWeightInstanceManager* FLightWeightInstanceSubsystem::FindOrAddLightWeight
 	}
 
 	FActorSpawnParameters SpawnParams;
-	SpawnParams.OverrideLevel = Level;
 	SpawnParams.SpawnCollisionHandlingOverride = ESpawnActorCollisionHandlingMethod::AlwaysSpawn;
 	SpawnParams.ObjectFlags = RF_Transactional;
 
-	ALightWeightInstanceManager* NewInstance = Level->GetWorld()->SpawnActor<ALightWeightInstanceManager>(BestMatchingClass, FTransform::Identity, SpawnParams);
-	ULevel* InstanceLvl = NewInstance->GetLevel();
-	NewInstance->SetRepresentedClass(ActorClass);
+	ALightWeightInstanceManager* NewInstanceManager = World->SpawnActor<ALightWeightInstanceManager>(BestMatchingClass, FTransform::Identity, SpawnParams);
+	NewInstanceManager->SetRepresentedClass(ActorClass);
+
+#if WITH_EDITOR
+	// Add the new manager to the layer
+	if (Layer)
+	{
+		ensure(NewInstanceManager->IsValidForDataLayer());
+
+		NewInstanceManager->AddDataLayer(Layer);
+	}
+#endif // WITH_EDITOR
 
-	check(LWInstanceManagers.Find(NewInstance) != INDEX_NONE);
+	check(LWInstanceManagers.Find(NewInstanceManager) != INDEX_NONE);
 
-	return NewInstance;
+	return NewInstanceManager;
 }
 
 UClass* FLightWeightInstanceSubsystem::FindBestInstanceManagerClass(const UClass* InActorClass)
@@ -266,16 +297,24 @@ bool FLightWeightInstanceSubsystem::IsInLevel(const FActorInstanceHandle& Handle
 	return false;
 }
 
-FActorInstanceHandle FLightWeightInstanceSubsystem::CreateNewLightWeightInstance(UClass* InActorClass, FLWIData* InitData, ULevel* InLevel)
+FActorInstanceHandle FLightWeightInstanceSubsystem::CreateNewLightWeightInstance(UClass* InActorClass, FLWIData* InitData, UDataLayer* InLayer, UWorld* World)
 {
-	// Get or create a light weight instance for this class and level
-	if (ALightWeightInstanceManager* LWIManager = FindOrAddLightWeightInstanceManager(InActorClass, InLevel))
+	// Get or create a light weight instance for this class and data layer
+	if (ALightWeightInstanceManager* LWIManager = FindOrAddLightWeightInstanceManager(InActorClass, InLayer, World))
 	{
 		// create an instance with the given data
 		int32 InstanceIdx = LWIManager->AddNewInstance(InitData);
-
+		InstanceIdx = LWIManager->ConvertInternalIndexToHandleIndex(InstanceIdx);
 		return FActorInstanceHandle(LWIManager, InstanceIdx);
 	}
 
 	return FActorInstanceHandle();
 }
+
+void FLightWeightInstanceSubsystem::DeleteInstance(const FActorInstanceHandle& Handle)
+{
+	if (ALightWeightInstanceManager* LWIManager = FindLightWeightInstanceManager(Handle))
+	{
+		LWIManager->RemoveInstance(Handle.GetInstanceIndex());
+	}
+}
diff --git a/Engine/Source/Runtime/Engine/Private/World.cpp b/Engine/Source/Runtime/Engine/Private/World.cpp
--- a/Engine/Source/Runtime/Engine/Private/World.cpp
+++ b/Engine/Source/Runtime/Engine/Private/World.cpp
@@ -373,6 +373,7 @@ UWorld::UWorld( const FObjectInitializer& ObjectInitializer )
 #if WITH_EDITOR
 , HierarchicalLODBuilder(new FHierarchicalLODBuilder(this))
 #endif
+, LWILastAssignedUID(0)
 , URL(FURL(NULL))
 ,	FXSystem(NULL)
 ,	TickTaskLevel(FTickTaskManagerInterface::Get().AllocateTickTaskLevel())
