Fix skeletal mesh ray tracing BLAS being recreated every frame

diff --git a/Engine/Source/Runtime/Engine/Private/SkeletalRenderGPUSkin.cpp b/Engine/Source/Runtime/Engine/Private/SkeletalRenderGPUSkin.cpp
--- a/Engine/Source/Runtime/Engine/Private/SkeletalRenderGPUSkin.cpp
+++ b/Engine/Source/Runtime/Engine/Private/SkeletalRenderGPUSkin.cpp
@@ -397,7 +397,6 @@ void FSkeletalMeshObjectGPUSkin::UpdateDynamicData_RenderThread(FGPUSkinCache* G
 	// So make sure to set bRequireRecreatingRayTracingGeometry to true to account for this, so RT geometry is updated correctly.
 	bRequireRecreatingRayTracingGeometry = (DynamicData == nullptr || RayTracingGeometry.Initializer.Segments.Num() == 0 || // Newly created
 		(DynamicData != nullptr && DynamicData->RayTracingLODIndex != InDynamicData->RayTracingLODIndex) || // LOD level changed
-		SkinCacheEntry == nullptr ||
 		bHiddenMaterialVisibilityDirtyForRayTracing ||
 		SkinCacheEntryForRayTracing == nullptr);
 
