Removed GetDummyPrimitiveInstance() helper (inlined in calling function) to fix static initialization order problems

diff --git a/Engine/Source/Runtime/Engine/Private/InstanceUniformShaderParameters.cpp b/Engine/Source/Runtime/Engine/Private/InstanceUniformShaderParameters.cpp
--- a/Engine/Source/Runtime/Engine/Private/InstanceUniformShaderParameters.cpp
+++ b/Engine/Source/Runtime/Engine/Private/InstanceUniformShaderParameters.cpp
@@ -99,22 +99,14 @@ void FInstanceSceneShaderData::Setup(
 	Data[9]    = *(const  FVector4*)&LightMapShadowMapUVBias;
 }
 
-static FPrimitiveInstance DummyInstance =
-	ConstructPrimitiveInstance(
-		FRenderBounds(FVector3f::ZeroVector, FVector3f::ZeroVector),
-		NANITE_INVALID_HIERARCHY_OFFSET,
-		0u /* Instance Flags */
-	);
-
-ENGINE_API const FPrimitiveInstance& GetDummyPrimitiveInstance()
-{
-	return DummyInstance;
-}
-
 ENGINE_API const FInstanceSceneShaderData& GetDummyInstanceSceneShaderData()
 {
 	static FInstanceSceneShaderData DummyShaderData = FInstanceSceneShaderData(
-		GetDummyPrimitiveInstance(),
+		ConstructPrimitiveInstance(
+			FRenderBounds(FVector3f::ZeroVector, FVector3f::ZeroVector),
+			NANITE_INVALID_HIERARCHY_OFFSET,
+			0u /* Instance Flags */
+		),
 		0xFFFFFFFFu, /* Primitive Id */
 		FRenderTransform::Identity, /* Primitive LocalToWorld */
 		FRenderTransform::Identity,  /* Primitive PrevLocalToWorld */
diff --git a/Engine/Source/Runtime/Engine/Public/InstanceUniformShaderParameters.h b/Engine/Source/Runtime/Engine/Public/InstanceUniformShaderParameters.h
--- a/Engine/Source/Runtime/Engine/Public/InstanceUniformShaderParameters.h
+++ b/Engine/Source/Runtime/Engine/Public/InstanceUniformShaderParameters.h
@@ -145,5 +145,4 @@ struct FInstanceSceneShaderData
 	);
 };
 
-ENGINE_API const FPrimitiveInstance& GetDummyPrimitiveInstance();
 ENGINE_API const FInstanceSceneShaderData& GetDummyInstanceSceneShaderData();
