PackedLevelInstance: remove uneeded cooking of source level because of non editor soft object path

diff --git a/Engine/Source/Runtime/Engine/Private/LevelInstance/LevelInstanceActor.cpp b/Engine/Source/Runtime/Engine/Private/LevelInstance/LevelInstanceActor.cpp
--- a/Engine/Source/Runtime/Engine/Private/LevelInstance/LevelInstanceActor.cpp
+++ b/Engine/Source/Runtime/Engine/Private/LevelInstance/LevelInstanceActor.cpp
@@ -151,9 +151,18 @@ void ALevelInstance::UnloadLevelInstance()
 	}
 }
 
+const TSoftObjectPtr<UWorld>& ALevelInstance::GetWorldAsset() const
+{
+#if WITH_EDITORONLY_DATA
+	return WorldAsset;
+#else
+	return CookedWorldAsset;
+#endif
+}
+
 bool ALevelInstance::IsLevelInstancePathValid() const
 {
-	return WorldAsset.GetUniqueID().IsValid();
+	return GetWorldAsset().GetUniqueID().IsValid();
 }
 
 bool ALevelInstance::HasValidLevelInstanceID() const
@@ -232,6 +241,13 @@ void ALevelInstance::PostLoad()
 	}
 #endif
 		
+#if WITH_EDITORONLY_DATA
+	if (IsRunningCookCommandlet() && SupportsLoading())
+	{
+		CookedWorldAsset = WorldAsset;
+	}
+#endif
+
 	OnLevelInstanceActorPostLoad.Broadcast(this);
 }
 
@@ -309,7 +325,7 @@ void ALevelInstance::PostEditUndoInternal()
 
 FString ALevelInstance::GetWorldAssetPackage() const
 {
-	return WorldAsset.GetUniqueID().GetLongPackageName();
+	return GetWorldAsset().GetUniqueID().GetLongPackageName();
 }
 
 void ALevelInstance::PreEditChange(FProperty* PropertyThatWillChange)
diff --git a/Engine/Source/Runtime/Engine/Public/LevelInstance/LevelInstanceActor.h b/Engine/Source/Runtime/Engine/Public/LevelInstance/LevelInstanceActor.h
--- a/Engine/Source/Runtime/Engine/Public/LevelInstance/LevelInstanceActor.h
+++ b/Engine/Source/Runtime/Engine/Public/LevelInstance/LevelInstanceActor.h
@@ -25,9 +25,14 @@ class ENGINE_API ALevelInstance : public AActor
 	GENERATED_UCLASS_BODY()
 
 protected:
+#if WITH_EDITORONLY_DATA
 	/** Level LevelInstance */
 	UPROPERTY(EditAnywhere, Category = LevelInstance, meta = (NoCreate))
 	TSoftObjectPtr<UWorld> WorldAsset;
+#endif
+
+	UPROPERTY()
+	TSoftObjectPtr<UWorld> CookedWorldAsset;
 
 public:
 
@@ -40,7 +45,7 @@ public:
 	void UnloadLevelInstance();
 	virtual bool SupportsLoading() const;
 	bool IsLevelInstancePathValid() const;
-	const TSoftObjectPtr<UWorld>& GetWorldAsset() const { return WorldAsset; }
+	const TSoftObjectPtr<UWorld>& GetWorldAsset() const;
 	class ULevelInstanceSubsystem* GetLevelInstanceSubsystem() const;
 	
 	const FLevelInstanceID& GetLevelInstanceID() const;
