Added some cpu trace profile tags to World Partition init phase.

diff --git a/Engine/Source/Runtime/Engine/Private/WorldPartition/WorldPartition.cpp b/Engine/Source/Runtime/Engine/Private/WorldPartition/WorldPartition.cpp
--- a/Engine/Source/Runtime/Engine/Private/WorldPartition/WorldPartition.cpp
+++ b/Engine/Source/Runtime/Engine/Private/WorldPartition/WorldPartition.cpp
@@ -390,7 +390,10 @@ void UWorldPartition::Initialize(UWorld* InWorld, const FTransform& InTransform)
 	{
 		UPackage* LevelPackage = OuterWorld->PersistentLevel->GetOutermost();
 		const FName PackageName = LevelPackage->GetLoadedPath().GetPackageFName();
-		UActorDescContainer::Initialize(World, PackageName);
+		{
+			TRACE_CPUPROFILER_EVENT_SCOPE(UActorDescContainer::Initialize);
+			UActorDescContainer::Initialize(World, PackageName);
+		}
 		check(bContainerInitialized);
 
 		bool bIsInstanced = (bEditorOnly && !IsRunningCommandlet()) ? OuterWorld->PersistentLevel->IsInstancedLevel() : false;
@@ -409,24 +412,27 @@ void UWorldPartition::Initialize(UWorld* InWorld, const FTransform& InTransform)
 			ReplaceTo = DestWorldName + TEXT(".") + DestWorldName;
 		}
 
-		for (FActorDescList::TIterator<> ActorDescIterator(this); ActorDescIterator; ++ActorDescIterator)
 		{
-			if (bIsInstanced)
+			TRACE_CPUPROFILER_EVENT_SCOPE(HashActorDescs);
+			for (FActorDescList::TIterator<> ActorDescIterator(this); ActorDescIterator; ++ActorDescIterator)
 			{
-				const FString LongActorPackageName = ActorDescIterator->ActorPackage.ToString();
-				const FString ActorPackageName = FPaths::GetBaseFilename(LongActorPackageName);
-				const FString InstancedName = FString::Printf(TEXT("%s_InstanceOf_%s"), *LevelPackage->GetName(), *ActorPackageName);
+				if (bIsInstanced)
+				{
+					const FString LongActorPackageName = ActorDescIterator->ActorPackage.ToString();
+					const FString ActorPackageName = FPaths::GetBaseFilename(LongActorPackageName);
+					const FString InstancedName = FString::Printf(TEXT("%s_InstanceOf_%s"), *LevelPackage->GetName(), *ActorPackageName);
 
-				InstancingContext.AddMapping(*LongActorPackageName, *InstancedName);
+					InstancingContext.AddMapping(*LongActorPackageName, *InstancedName);
 
-				ActorDescIterator->TransformInstance(ReplaceFrom, ReplaceTo, InstanceTransform);
-			}
+					ActorDescIterator->TransformInstance(ReplaceFrom, ReplaceTo, InstanceTransform);
+				}
 
-			ActorDescIterator->OnRegister();
+				ActorDescIterator->OnRegister();
 
-			if (bEditorOnly)
-			{
-				HashActorDesc(*ActorDescIterator);
+				if (bEditorOnly)
+				{
+					HashActorDesc(*ActorDescIterator);
+				}
 			}
 		}
 	}
