WorldPartition: wp.Runtime.EnableSimulationStreamingSource (enabled by default) allows disabling simulation streaming source and default to the game streaming sources when ejecting the camera

diff --git a/Engine/Source/Runtime/Engine/Private/WorldPartition/WorldPartitionStreamingPolicy.cpp b/Engine/Source/Runtime/Engine/Private/WorldPartition/WorldPartitionStreamingPolicy.cpp
--- a/Engine/Source/Runtime/Engine/Private/WorldPartition/WorldPartitionStreamingPolicy.cpp
+++ b/Engine/Source/Runtime/Engine/Private/WorldPartition/WorldPartitionStreamingPolicy.cpp
@@ -35,6 +35,12 @@ static FAutoConsoleVariableRef CVarUpdateStreamingSources(
 	GUpdateStreamingSources,
 	TEXT("Set to 0 to stop updating (freeze) world partition streaming sources."));
 
+static int32 GEnableSimulationStreamingSource = 1;
+static FAutoConsoleVariableRef CVarEnableSimulationStreamingSource(
+	TEXT("wp.Runtime.EnableSimulationStreamingSource"),
+	GEnableSimulationStreamingSource,
+	TEXT("Set to 0 to if you want to disable the simulation/ejected camera streaming source."));
+
 static int32 GMaxLoadingStreamingCells = 4;
 static FAutoConsoleVariableRef CMaxLoadingStreamingCells(
 	TEXT("wp.Runtime.MaxLoadingStreamingCells"),
@@ -86,7 +92,7 @@ void UWorldPartitionStreamingPolicy::UpdateStreamingSources()
 
 #if WITH_EDITOR
 	// We are in the SIE
-	if (UWorldPartition::IsSimulating())
+	if (GEnableSimulationStreamingSource && UWorldPartition::IsSimulating())
 	{
 		// Transform to Local
 		const FVector ViewLocation = GCurrentLevelEditingViewportClient->GetViewLocation();
