Fix a crash bug that happens when you delete a light  weight instance manager in the editor and then try to add another manager of the same type.

diff --git a/Engine/Source/Runtime/Engine/Classes/GameFramework/LightWeightInstanceSubsystem.h b/Engine/Source/Runtime/Engine/Classes/GameFramework/LightWeightInstanceSubsystem.h
--- a/Engine/Source/Runtime/Engine/Classes/GameFramework/LightWeightInstanceSubsystem.h
+++ b/Engine/Source/Runtime/Engine/Classes/GameFramework/LightWeightInstanceSubsystem.h
@@ -21,7 +21,6 @@ ENGINE_API DECLARE_LOG_CATEGORY_EXTERN(LogLightWeightInstance, Log, Warning);
 
 struct ENGINE_API FLightWeightInstanceSubsystem
 {
-
 	friend struct FActorInstanceHandle;
 	friend class ALightWeightInstanceManager;
 
@@ -34,6 +33,9 @@ struct ENGINE_API FLightWeightInstanceSubsystem
 		return *LWISubsystem;
 	}
 
+	FLightWeightInstanceSubsystem();
+	~FLightWeightInstanceSubsystem();
+
 	// Returns the instance manager that handles the given handle
 	ALightWeightInstanceManager* FindLightWeightInstanceManager(const FActorInstanceHandle& Handle) const;
 
@@ -109,4 +111,11 @@ private:
 	// TODO: preallocate the size of this based on a config variable
 	UPROPERTY()
 	TArray<ALightWeightInstanceManager*> LWInstanceManagers;
+
+
+#ifdef WITH_EDITOR
+private:
+	FDelegateHandle OnLevelActorAddedHandle;
+	FDelegateHandle OnLevelActorDeletedHandle;
+#endif // WITH_EDITOR
 };
diff --git a/Engine/Source/Runtime/Engine/Private/LightWeightInstanceSubsystem.cpp b/Engine/Source/Runtime/Engine/Private/LightWeightInstanceSubsystem.cpp
--- a/Engine/Source/Runtime/Engine/Private/LightWeightInstanceSubsystem.cpp
+++ b/Engine/Source/Runtime/Engine/Private/LightWeightInstanceSubsystem.cpp
@@ -9,6 +9,37 @@ DEFINE_LOG_CATEGORY(LogLightWeightInstance);
 
 TSharedPtr<FLightWeightInstanceSubsystem> FLightWeightInstanceSubsystem::LWISubsystem;
 
+FLightWeightInstanceSubsystem::FLightWeightInstanceSubsystem()
+{
+#if WITH_EDITOR
+	if (GEngine)
+	{
+		OnLevelActorAddedHandle = GEngine->OnLevelActorAdded().AddLambda([this](AActor* Actor)
+			{
+				if (ALightWeightInstanceManager* LWIManager = Cast<ALightWeightInstanceManager>(Actor))
+				{
+					FLightWeightInstanceSubsystem::Get().LWInstanceManagers.AddUnique(LWIManager);
+				}
+			});
+		OnLevelActorDeletedHandle = GEngine->OnLevelActorDeleted().AddLambda([this](AActor* Actor)
+			{
+				if (ALightWeightInstanceManager* LWIManager = Cast<ALightWeightInstanceManager>(Actor))
+				{
+					FLightWeightInstanceSubsystem::Get().LWInstanceManagers.Remove(LWIManager);
+				}
+			});
+	}
+#endif // WITH_EDITOR
+}
+
+FLightWeightInstanceSubsystem::~FLightWeightInstanceSubsystem()
+{
+#if WITH_EDITOR
+	GEngine->OnLevelActorAdded().Remove(OnLevelActorAddedHandle);
+	GEngine->OnLevelActorDeleted().Remove(OnLevelActorDeletedHandle);
+#endif // WITH_EDITOR
+}
+
 int32 FLightWeightInstanceSubsystem::GetManagerIndex(const ALightWeightInstanceManager* Manager) const
 {
 	for (int32 Idx = 0; Idx < LWInstanceManagers.Num(); ++Idx)
