No Animation coming through Sequence Evaluators
#fix allow for UID comparison in case a SmartName has been initialized without a FName

diff --git a/Engine/Source/Runtime/Engine/Private/Animation/AnimDataModel.cpp b/Engine/Source/Runtime/Engine/Private/Animation/AnimDataModel.cpp
--- a/Engine/Source/Runtime/Engine/Private/Animation/AnimDataModel.cpp
+++ b/Engine/Source/Runtime/Engine/Private/Animation/AnimDataModel.cpp
@@ -147,7 +147,7 @@ const FFloatCurve* UAnimDataModel::FindFloatCurve(const FAnimationCurveIdentifie
 	ensure(CurveIdentifier.CurveType == ERawCurveTrackTypes::RCT_Float);
 	for (const FFloatCurve& FloatCurve : CurveData.FloatCurves)
 	{
-		if (FloatCurve.Name == CurveIdentifier.InternalName)
+		if (FloatCurve.Name == CurveIdentifier.InternalName || FloatCurve.Name.UID == CurveIdentifier.InternalName.UID)
 		{
 			return &FloatCurve;
 		}
@@ -161,7 +161,7 @@ const FTransformCurve* UAnimDataModel::FindTransformCurve(const FAnimationCurveI
 	ensure(CurveIdentifier.CurveType == ERawCurveTrackTypes::RCT_Transform);
 	for (const FTransformCurve& TransformCurve : CurveData.TransformCurves)
 	{
-		if (TransformCurve.Name == CurveIdentifier.InternalName)
+		if (TransformCurve.Name == CurveIdentifier.InternalName || TransformCurve.Name.UID == CurveIdentifier.InternalName.UID)
 		{
 			return &TransformCurve;
 		}
