Fixed cooking of distance field and card representation data for non PC platforms when it's not available in the DDC.

Bug flow:
1. GetPlatformStaticMeshRenderData creates separate FStaticMeshRenderData for target platform and calls CacheDerivedData
2. CacheDerivedData runs a FAsyncCardRepresentationTask / FAsyncDistanceFieldTask
3. On task completion, ProcessAsyncTasks is updating StaticMesh->RenderData, but didn't update the entire linked list of per platform RenderData objects
4. Unitialized SDF/cards data is serialized for the target platform

Repro:
* Bump SDF/cards DDC key
* Run cook or COTF on any non PC platform
* Observe broken SDF/cards data

diff --git a/Engine/Source/Runtime/Engine/Private/DistanceFieldAtlas.cpp b/Engine/Source/Runtime/Engine/Private/DistanceFieldAtlas.cpp
--- a/Engine/Source/Runtime/Engine/Private/DistanceFieldAtlas.cpp
+++ b/Engine/Source/Runtime/Engine/Private/DistanceFieldAtlas.cpp
@@ -741,13 +741,15 @@ void FDistanceFieldAsyncQueue::ProcessAsyncTasks(bool bLimitExecutionTime)
 		if (Task->StaticMesh)
 		{
 			Task->GeneratedVolumeData->bAsyncBuilding = false;
-			FDistanceFieldVolumeData* OldVolumeData = Task->StaticMesh->GetRenderData()->LODResources[0].DistanceFieldData;
+
+			FStaticMeshRenderData* RenderData = Task->StaticMesh->GetRenderData();
+			FDistanceFieldVolumeData* OldVolumeData = RenderData->LODResources[0].DistanceFieldData;
 
 			// Assign the new volume data, this is safe because the render thread makes a copy of the pointer at scene proxy creation time.
-			Task->StaticMesh->GetRenderData()->LODResources[0].DistanceFieldData = Task->GeneratedVolumeData;
+			RenderData->LODResources[0].DistanceFieldData = Task->GeneratedVolumeData;
 
 			// Renderstates are not initialized between UStaticMesh::PreEditChange() and UStaticMesh::PostEditChange()
-			if (Task->StaticMesh->GetRenderData()->IsInitialized())
+			if (RenderData->IsInitialized())
 			{
 				for (UStaticMeshComponent* Component : ObjectCacheScope.GetContext().GetStaticMeshComponents(Task->StaticMesh))
 				{
@@ -758,13 +760,23 @@ void FDistanceFieldAsyncQueue::ProcessAsyncTasks(bool bLimitExecutionTime)
 				}
 			}
 
-
 			if (OldVolumeData)
 			{
 				// Rendering thread may still be referencing the old one, use the deferred cleanup interface to delete it next frame when it is safe
 				BeginCleanup(OldVolumeData);
 			}
 
+			// Need also to update platform render data if it's being cached
+			FStaticMeshRenderData* PlatformRenderData = RenderData->NextCachedRenderData.Get();
+			while (PlatformRenderData)
+			{
+				if (PlatformRenderData->LODResources[0].DistanceFieldData)
+				{
+					*PlatformRenderData->LODResources[0].DistanceFieldData = *Task->GeneratedVolumeData;
+				}
+				PlatformRenderData = PlatformRenderData->NextCachedRenderData.Get();
+			}
+
 			{
 				TArray<uint8> DerivedData;
 				// Save built distance field volume to DDC
diff --git a/Engine/Source/Runtime/Engine/Private/MeshCardRepresentation.cpp b/Engine/Source/Runtime/Engine/Private/MeshCardRepresentation.cpp
--- a/Engine/Source/Runtime/Engine/Private/MeshCardRepresentation.cpp
+++ b/Engine/Source/Runtime/Engine/Private/MeshCardRepresentation.cpp
@@ -621,13 +621,14 @@ void FCardRepresentationAsyncQueue::ProcessAsyncTasks(bool bLimitExecutionTime)
 		// Editor 'force delete' can null any UObject pointers which are seen by reference collecting (eg UProperty or serialized)
 		if (Task->StaticMesh && Task->bSuccess)
 		{
-			FCardRepresentationData* OldCardData = Task->StaticMesh->GetRenderData()->LODResources[0].CardRepresentationData;
+			FStaticMeshRenderData* RenderData = Task->StaticMesh->GetRenderData();
+			FCardRepresentationData* OldCardData = RenderData->LODResources[0].CardRepresentationData;
 
 			// Assign the new data, this is safe because the render threads makes a copy of the pointer at scene proxy creation time.
-			Task->StaticMesh->GetRenderData()->LODResources[0].CardRepresentationData = Task->GeneratedCardRepresentation;
+			RenderData->LODResources[0].CardRepresentationData = Task->GeneratedCardRepresentation;
 
 			// Any already created render state needs to be dirtied
-			if (Task->StaticMesh->GetRenderData()->IsInitialized())
+			if (RenderData->IsInitialized())
 			{
 				for (UStaticMeshComponent* Component : ObjectCacheScope.GetContext().GetStaticMeshComponents(Task->StaticMesh))
 				{
@@ -641,6 +642,17 @@ void FCardRepresentationAsyncQueue::ProcessAsyncTasks(bool bLimitExecutionTime)
 			// Rendering thread may still be referencing the old one, use the deferred cleanup interface to delete it next frame when it is safe
 			BeginCleanup(OldCardData);
 
+			// Need also to update platform render data if it's being cached
+			FStaticMeshRenderData* PlatformRenderData = RenderData->NextCachedRenderData.Get();
+			while (PlatformRenderData)
+			{
+				if (PlatformRenderData->LODResources[0].CardRepresentationData)
+				{
+					*PlatformRenderData->LODResources[0].CardRepresentationData = *Task->GeneratedCardRepresentation;
+				}
+				PlatformRenderData = PlatformRenderData->NextCachedRenderData.Get();
+			}
+
 			{
 				TArray<uint8> DerivedData;
 				// Save built data to DDC
