Fix type coercion when using Clamp material node

diff --git a/Engine/Source/Runtime/Engine/Private/Materials/HLSLMaterialTranslator.cpp b/Engine/Source/Runtime/Engine/Private/Materials/HLSLMaterialTranslator.cpp
--- a/Engine/Source/Runtime/Engine/Private/Materials/HLSLMaterialTranslator.cpp
+++ b/Engine/Source/Runtime/Engine/Private/Materials/HLSLMaterialTranslator.cpp
@@ -7522,8 +7522,9 @@ int32 FHLSLMaterialTranslator::Clamp(int32 X,int32 A,int32 B)
 	{
 		if (IsAnalyticDerivEnabled())
 		{
-			int32 MaxAX = DerivativeAutogen.GenerateExpressionFunc2(*this,FMaterialDerivativeAutogen::EFunc2::Max,A,X);
-			int32 Result = DerivativeAutogen.GenerateExpressionFunc2(*this,FMaterialDerivativeAutogen::EFunc2::Min,MaxAX,B);
+			// Make sure 'X' is given as first parameter, to ensure proper type coercion
+			int32 MaxXA = DerivativeAutogen.GenerateExpressionFunc2(*this,FMaterialDerivativeAutogen::EFunc2::Max,X,A);
+			int32 Result = DerivativeAutogen.GenerateExpressionFunc2(*this,FMaterialDerivativeAutogen::EFunc2::Min, MaxXA,B);
 			return Result;
 		}
 		else
