First Iteration: Replay of Streaming Sources (eventually this should be driven by replication of UWorldPartitionStreamingSourceComponent)

diff --git a/Engine/Source/Runtime/Engine/Private/WorldPartition/WorldPartition.cpp b/Engine/Source/Runtime/Engine/Private/WorldPartition/WorldPartition.cpp
--- a/Engine/Source/Runtime/Engine/Private/WorldPartition/WorldPartition.cpp
+++ b/Engine/Source/Runtime/Engine/Private/WorldPartition/WorldPartition.cpp
@@ -11,6 +11,7 @@
 #include "WorldPartition/WorldPartitionActorDesc.h"
 #include "WorldPartition/WorldPartitionSubsystem.h"
 #include "WorldPartition/WorldPartitionStreamingSource.h"
+#include "WorldPartition/WorldPartitionReplay.h"
 #include "Algo/Accumulate.h"
 #include "Algo/Transform.h"
 #include "Engine/World.h"
@@ -238,6 +239,9 @@ UWorldPartition::UWorldPartition(const FObjectInitializer& ObjectInitializer)
 	, InitState(EWorldPartitionInitState::Uninitialized)
 	, InstanceTransform(FTransform::Identity)
 	, StreamingPolicy(nullptr)
+#if !UE_BUILD_SHIPPING
+	, Replay(nullptr)
+#endif
 {
 	static bool bRegisteredDelegate = false;
 	if (!bRegisteredDelegate)
@@ -363,6 +367,8 @@ void UWorldPartition::Initialize(UWorld* InWorld, const FTransform& InTransform)
 
 	RegisterDelegates();
 
+	AWorldPartitionReplay::Initialize(World);
+	
 #if WITH_EDITOR
 	bool bEditorOnly = !World->IsGameWorld();
 	if (bEditorOnly)
diff --git a/Engine/Source/Runtime/Engine/Private/WorldPartition/WorldPartitionReplay.cpp b/Engine/Source/Runtime/Engine/Private/WorldPartition/WorldPartitionReplay.cpp
new file mode 100644
--- /dev/null
+++ b/Engine/Source/Runtime/Engine/Private/WorldPartition/WorldPartitionReplay.cpp
@@ -0,0 +1,221 @@
+// Copyright Epic Games, Inc. All Rights Reserved.
+
+#include "WorldPartition/WorldPartitionReplay.h"
+#include "WorldPartition/WorldPartition.h"
+#include "WorldPartition/WorldPartitionStreamingSource.h"
+#include "Engine/World.h"
+#include "EngineUtils.h"
+#include "Net/UnrealNetwork.h"
+#include "UObject/CoreNet.h"
+#include "Serialization/BitWriter.h"
+#include "Engine/DemoNetDriver.h"
+
+FArchive& operator<<(FArchive& Ar, FWorldPartitionReplayStreamingSource& StreamingSource)
+{
+	Ar << StreamingSource.Location;
+	Ar << StreamingSource.Rotation;
+	Ar << StreamingSource.TargetState;
+	Ar << StreamingSource.bBlockOnSlowLoading;
+	Ar << StreamingSource.Priority;
+	Ar << StreamingSource.Velocity;
+	return Ar;
+}
+
+FWorldPartitionReplaySample::FWorldPartitionReplaySample(AWorldPartitionReplay* InReplay)
+	: Replay(InReplay)
+{
+	for (const FWorldPartitionStreamingSource& StreamingSource : Replay->GetRecordingStreamingSources())
+	{
+		StreamingSources.Add(FWorldPartitionReplayStreamingSource(StreamingSource));
+		int32 Index = Replay->StreamingSourceNames.Find(StreamingSource.Name);
+		if (Index != INDEX_NONE)
+		{
+			StreamingSourceNameIndices.Add(Index);
+		}
+		else
+		{
+			Replay->StreamingSourceNames.Add(StreamingSource.Name);
+			StreamingSourceNameIndices.Add(Replay->StreamingSourceNames.Num() - 1);
+		}
+	}
+}
+
+FArchive& operator<<(FArchive& Ar, FWorldPartitionReplaySample& ReplaySample)
+{
+	Ar << ReplaySample.StreamingSourceNameIndices;
+	Ar << ReplaySample.StreamingSources;
+
+	if (Ar.IsLoading())
+	{
+		check(ReplaySample.StreamingSourceNameIndices.Num() == ReplaySample.StreamingSources.Num());
+		for (int32 i = 0; i < ReplaySample.StreamingSources.Num(); ++i)
+		{
+			int32 NameIndex = ReplaySample.StreamingSourceNameIndices[i];
+			if (ensure(ReplaySample.Replay->StreamingSourceNames.IsValidIndex(NameIndex)))
+			{
+				ReplaySample.StreamingSources[i].Name = ReplaySample.Replay->StreamingSourceNames[NameIndex];
+			}
+		}
+	}
+	return Ar;
+}
+
+AWorldPartitionReplay::AWorldPartitionReplay(const FObjectInitializer& ObjectInitializer)
+	: Super(ObjectInitializer)
+{
+	bAlwaysRelevant = true;
+	bReplicates = true;
+	NetDriverName = NAME_DemoNetDriver;
+#if WITH_EDITORONLY_DATA
+	bListedInSceneOutliner = false;
+#endif
+}
+
+void AWorldPartitionReplay::Initialize(UWorld* World)
+{
+#if !UE_BUILD_SHIPPING
+	if (World->IsGameWorld())
+	{
+		check(World->GetWorldPartition());
+		UWorldPartition* WorldPartition = World->GetWorldPartition();
+
+		check(!WorldPartition->Replay);
+		FActorSpawnParameters SpawnParams;
+		SpawnParams.Name = AWorldPartitionReplay::StaticClass()->GetFName();
+		WorldPartition->Replay = World->SpawnActor<AWorldPartitionReplay>(SpawnParams);
+	}
+#endif
+}
+
+bool AWorldPartitionReplay::IsEnabled(UWorld* World)
+{
+#if !UE_BUILD_SHIPPING
+	check(World->IsGameWorld());
+	check(World->GetWorldPartition());
+	UWorldPartition* WorldPartition = World->GetWorldPartition();
+	return WorldPartition->Replay && WorldPartition->Replay->IsEnabled();
+#else
+	return false;
+#endif
+}
+
+void AWorldPartitionReplay::BeginPlay()
+{
+	UWorld* World = GetWorld();
+	check(World->IsGameWorld());
+	check(World->GetWorldPartition());
+	ENetMode NetMode = World->GetNetMode();
+	
+	// Recording will happen only in NM_Standalone/NM_Client. Replay can happen on any net mode.
+	bEnabled = (NetMode == NM_Standalone || NetMode == NM_Client || World->IsPlayingReplay());
+}
+
+void AWorldPartitionReplay::GetLifetimeReplicatedProps(TArray<FLifetimeProperty>& OutLifetimeProps) const
+{
+	Super::GetLifetimeReplicatedProps(OutLifetimeProps);
+
+	DOREPLIFETIME(AWorldPartitionReplay, StreamingSourceNames);
+}
+
+void AWorldPartitionReplay::PreReplication(IRepChangedPropertyTracker& ChangedPropertyTracker)
+{
+	check(GetWorld()->IsRecordingReplay());
+	Super::PreReplication(ChangedPropertyTracker);
+	
+	UWorldPartition* WorldPartition = GetWorld()->GetWorldPartition();
+	check(WorldPartition);
+
+	FWorldPartitionReplaySample Replay(this);
+	if (Replay.StreamingSources.Num())
+	{
+		FBitWriter Writer(/*InMaxBits=*/1024, /*InAllowResize=*/true);
+		Writer << Replay;
+		if (Writer.GetNumBits())
+		{
+			ChangedPropertyTracker.SetExternalData(Writer.GetData(), Writer.GetNumBits());
+		}
+	}
+}
+
+bool AWorldPartitionReplay::GetReplayStreamingSources(TArray<FWorldPartitionStreamingSource>& OutStreamingSources)
+{
+	UWorld* World = GetWorld();
+	UDemoNetDriver* DemoNetDriver = World->GetDemoNetDriver();
+	if (DemoNetDriver && DemoNetDriver->IsPlaying())
+	{
+		const float CurrentTime = DemoNetDriver->GetDemoCurrentTime();
+			
+		// Erase samples older than 1 second
+		while (ReplaySamples.Num() > 0)
+		{
+			if (ReplaySamples[0].TimeSeconds > CurrentTime - 1.0f)
+			{
+				break;
+			}
+
+			ReplaySamples.RemoveAt(0);
+		}
+
+		FReplayExternalDataArray* ExternalReplayData = DemoNetDriver->GetExternalDataArrayForObject(this);
+		
+		// Read new samples
+		if (ExternalReplayData && ExternalReplayData->Num() > 0)
+		{
+			for (int i = 0; i < ExternalReplayData->Num(); i++)
+			{
+				FWorldPartitionReplaySample ReplaySample(this);
+				(*ExternalReplayData)[i].Reader << ReplaySample;
+				ReplaySample.TimeSeconds = (*ExternalReplayData)[i].TimeSeconds;
+				ReplaySamples.Add(ReplaySample);
+			}
+
+			ExternalReplayData->Empty();
+		}
+				
+		// Find current time samples to interpolate streaming sources
+		for (int32 i = 0; i < ReplaySamples.Num() - 1; ++i)
+		{
+			if (CurrentTime >= ReplaySamples[i].TimeSeconds && CurrentTime <= ReplaySamples[i + 1].TimeSeconds)
+			{
+				const FWorldPartitionReplaySample& ReplaySample1 = ReplaySamples[i];
+				const FWorldPartitionReplaySample& ReplaySample2 = ReplaySamples[i + 1];
+
+				for (int32 SourceIndex = 0; SourceIndex < ReplaySample1.StreamingSources.Num(); ++SourceIndex)
+				{
+					int32 SourceIndex2 = ReplaySample2.StreamingSources.IndexOfByPredicate([&](const FWorldPartitionStreamingSource& StreamingSource) { return StreamingSource.Name == ReplaySample1.StreamingSources[SourceIndex].Name; });
+					if (SourceIndex2 != INDEX_NONE)
+					{
+						const FWorldPartitionStreamingSource& Source1 = ReplaySample1.StreamingSources[SourceIndex];
+						const FWorldPartitionStreamingSource& Source2 = ReplaySample2.StreamingSources[SourceIndex2];
+						
+						const float EPSILON = SMALL_NUMBER;
+						const float Delta = ReplaySample2.TimeSeconds - ReplaySample1.TimeSeconds;
+						const float LerpPercent = Delta > EPSILON ? FMath::Clamp<float>((float)(CurrentTime - ReplaySample1.TimeSeconds) / Delta, 0.0f, 1.0f) : 1.0f;
+
+						const FVector Location = FMath::Lerp(Source1.Location, Source2.Location, LerpPercent);
+						const FQuat Rotation = FQuat::FastLerp(FQuat(Source1.Rotation), FQuat(Source2.Rotation), LerpPercent).GetNormalized();
+						const float Velocity = FMath::Lerp(Source1.Velocity, Source2.Velocity, LerpPercent);
+						OutStreamingSources.Add(FWorldPartitionStreamingSource(Source1.Name, Location, Rotation.Rotator(), Source1.TargetState, Source1.bBlockOnSlowLoading, Source1.Priority, Velocity));
+					}
+					else
+					{
+						OutStreamingSources.Add(ReplaySample1.StreamingSources[SourceIndex]);
+					}
+				}
+								
+				break;
+			}
+		}
+
+		return true;
+	}
+
+	return false;
+}
+
+TArray<FWorldPartitionStreamingSource> AWorldPartitionReplay::GetRecordingStreamingSources() const
+{
+	check(GetWorld());
+	check(GetWorld()->GetWorldPartition());
+	return GetWorld()->GetWorldPartition()->GetStreamingSources();
+}
diff --git a/Engine/Source/Runtime/Engine/Private/WorldPartition/WorldPartitionStreamingPolicy.cpp b/Engine/Source/Runtime/Engine/Private/WorldPartition/WorldPartitionStreamingPolicy.cpp
--- a/Engine/Source/Runtime/Engine/Private/WorldPartition/WorldPartitionStreamingPolicy.cpp
+++ b/Engine/Source/Runtime/Engine/Private/WorldPartition/WorldPartitionStreamingPolicy.cpp
@@ -8,6 +8,7 @@
 #include "WorldPartition/WorldPartitionRuntimeCell.h"
 #include "WorldPartition/WorldPartitionStreamingSource.h"
 #include "WorldPartition/WorldPartition.h"
+#include "WorldPartition/WorldPartitionReplay.h"
 #include "WorldPartition/DataLayer/DataLayerSubsystem.h"
 #include "WorldPartition/DataLayer/WorldDataLayers.h"
 #include "WorldPartition/HLOD/HLODActor.h"
@@ -90,62 +91,73 @@ void UWorldPartitionStreamingPolicy::UpdateStreamingSources()
 
 	const FTransform WorldToLocal = WorldPartition->GetInstanceTransform().Inverse();
 
-#if WITH_EDITOR
-	// We are in the SIE
-	if (GEnableSimulationStreamingSource && UWorldPartition::IsSimulating())
+#if !UE_BUILD_SHIPPING
+	bool bUseReplaySources = false;
+	if (AWorldPartitionReplay* Replay = WorldPartition->Replay)
 	{
-		// Transform to Local
-		const FVector ViewLocation = GCurrentLevelEditingViewportClient->GetViewLocation();
-		const FRotator ViewRotation = GCurrentLevelEditingViewportClient->GetViewRotation();
-		const FVector ViewLocationLocal = WorldToLocal.TransformPosition(ViewLocation);
-		const FRotator ViewRotationLocal = WorldToLocal.TransformRotation(ViewRotation.Quaternion()).Rotator();
-		static const FName NAME_SIE(TEXT("SIE"));
-		StreamingSources.Add(FWorldPartitionStreamingSource(NAME_SIE, ViewLocationLocal, ViewRotationLocal, EStreamingSourceTargetState::Activated, /*bBlockOnSlowLoading=*/false));
+		bUseReplaySources = Replay->GetReplayStreamingSources(StreamingSources);
 	}
-	else
+
+	if (!bUseReplaySources)
 #endif
 	{
-		UWorld* World = WorldPartition->GetWorld();
-		const ENetMode NetMode = World->GetNetMode();
-		if(NetMode == NM_Standalone || NetMode == NM_Client)
+#if WITH_EDITOR
+		// We are in the SIE
+		if (GEnableSimulationStreamingSource && UWorldPartition::IsSimulating())
 		{
-			const int32 NumPlayers = GEngine->GetNumGamePlayers(World);
-			for (int32 PlayerIndex = 0; PlayerIndex < NumPlayers; ++PlayerIndex)
+			// Transform to Local
+			const FVector ViewLocation = GCurrentLevelEditingViewportClient->GetViewLocation();
+			const FRotator ViewRotation = GCurrentLevelEditingViewportClient->GetViewRotation();
+			const FVector ViewLocationLocal = WorldToLocal.TransformPosition(ViewLocation);
+			const FRotator ViewRotationLocal = WorldToLocal.TransformRotation(ViewRotation.Quaternion()).Rotator();
+			static const FName NAME_SIE(TEXT("SIE"));
+			StreamingSources.Add(FWorldPartitionStreamingSource(NAME_SIE, ViewLocationLocal, ViewRotationLocal, EStreamingSourceTargetState::Activated, /*bBlockOnSlowLoading=*/false));
+		}
+		else
+#endif
+		{
+			UWorld* World = WorldPartition->GetWorld();
+			const ENetMode NetMode = World->GetNetMode();
+			if (NetMode == NM_Standalone || NetMode == NM_Client || AWorldPartitionReplay::IsEnabled(World))
 			{
-				ULocalPlayer* Player = GEngine->GetGamePlayer(World, PlayerIndex);
-				if (Player && Player->PlayerController && Player->PlayerController->IsStreamingSourceEnabled())
+				const int32 NumPlayers = GEngine->GetNumGamePlayers(World);
+				for (int32 PlayerIndex = 0; PlayerIndex < NumPlayers; ++PlayerIndex)
 				{
-					FVector ViewLocation;
-					FRotator ViewRotation;
-					Player->PlayerController->GetPlayerViewPoint(ViewLocation, ViewRotation);
-					
-					// Transform to Local
-					ViewLocation = WorldToLocal.TransformPosition(ViewLocation);
-					ViewRotation = WorldToLocal.TransformRotation(ViewRotation.Quaternion()).Rotator();
-					StreamingSources.Add(FWorldPartitionStreamingSource(Player->GetFName(), ViewLocation, ViewRotation, EStreamingSourceTargetState::Activated, /*bBlockOnSlowLoading=*/true));
+					ULocalPlayer* Player = GEngine->GetGamePlayer(World, PlayerIndex);
+					if (Player && Player->PlayerController && Player->PlayerController->IsStreamingSourceEnabled())
+					{
+						FVector ViewLocation;
+						FRotator ViewRotation;
+						Player->PlayerController->GetPlayerViewPoint(ViewLocation, ViewRotation);
+
+						// Transform to Local
+						ViewLocation = WorldToLocal.TransformPosition(ViewLocation);
+						ViewRotation = WorldToLocal.TransformRotation(ViewRotation.Quaternion()).Rotator();
+						StreamingSources.Add(FWorldPartitionStreamingSource(Player->GetFName(), ViewLocation, ViewRotation, EStreamingSourceTargetState::Activated, /*bBlockOnSlowLoading=*/true));
+					}
 				}
 			}
 		}
-	}
-
-	for (IWorldPartitionStreamingSourceProvider* StreamingSourceProvider : WorldPartition->StreamingSourceProviders)
-	{
-		FWorldPartitionStreamingSource StreamingSource;
-		if (StreamingSourceProvider->GetStreamingSource(StreamingSource))
+	
+		for (IWorldPartitionStreamingSourceProvider* StreamingSourceProvider : WorldPartition->StreamingSourceProviders)
 		{
-			// Transform to Local
-			StreamingSource.Location = WorldToLocal.TransformPosition(StreamingSource.Location);
-			StreamingSource.Rotation = WorldToLocal.TransformRotation(StreamingSource.Rotation.Quaternion()).Rotator();
-			// If none is provided, default Streaming Source provider's priority to be less than those based on player controllers
-			if (StreamingSource.Priority == EStreamingSourcePriority::Default)
+			FWorldPartitionStreamingSource StreamingSource;
+			if (StreamingSourceProvider->GetStreamingSource(StreamingSource))
 			{
-				StreamingSource.Priority = EStreamingSourcePriority::Low;
-			}
+				// Transform to Local
+				StreamingSource.Location = WorldToLocal.TransformPosition(StreamingSource.Location);
+				StreamingSource.Rotation = WorldToLocal.TransformRotation(StreamingSource.Rotation.Quaternion()).Rotator();
+				// If none is provided, default Streaming Source provider's priority to be less than those based on player controllers
+				if (StreamingSource.Priority == EStreamingSourcePriority::Default)
+				{
+					StreamingSource.Priority = EStreamingSourcePriority::Low;
+				}
 
-			StreamingSources.Add(StreamingSource);
+				StreamingSources.Add(StreamingSource);
+			}
 		}
 	}
-
+		
 	// Update streaming sources velocity
 	const float CurrentTime = WorldPartition->GetWorld()->GetTimeSeconds();
 	TSet<FName, DefaultKeyFuncs<FName>, TInlineSetAllocator<8>> ValidStreamingSources;
@@ -204,7 +216,7 @@ void UWorldPartitionStreamingPolicy::UpdateStreamingState()
 	check(LoadStreamingCells.IsEmpty());
 
 	const ENetMode NetMode = World->GetNetMode();
-	if(NetMode == NM_Standalone || NetMode == NM_Client)
+	if (NetMode == NM_Standalone || NetMode == NM_Client || AWorldPartitionReplay::IsEnabled(World))
 	{
 		// Early out if nothing loaded and no streaming source
 		if ((StreamingSources.Num() == 0) && (ActivatedCells.Num() == 0) && (LoadedCells.Num() == 0))
diff --git a/Engine/Source/Runtime/Engine/Public/WorldPartition/WorldPartition.h b/Engine/Source/Runtime/Engine/Public/WorldPartition/WorldPartition.h
--- a/Engine/Source/Runtime/Engine/Public/WorldPartition/WorldPartition.h
+++ b/Engine/Source/Runtime/Engine/Public/WorldPartition/WorldPartition.h
@@ -247,7 +247,9 @@ private:
 
 #if !UE_BUILD_SHIPPING
 	void GetOnScreenMessages(FCoreDelegates::FSeverityMessageMap& OutMessages);
-#endif
 
+	class AWorldPartitionReplay* Replay;
+#endif
+	friend class AWorldPartitionReplay;
 	friend class UWorldPartitionStreamingPolicy;
 };
\ No newline at end of file
diff --git a/Engine/Source/Runtime/Engine/Public/WorldPartition/WorldPartitionReplay.h b/Engine/Source/Runtime/Engine/Public/WorldPartition/WorldPartitionReplay.h
new file mode 100644
--- /dev/null
+++ b/Engine/Source/Runtime/Engine/Public/WorldPartition/WorldPartitionReplay.h
@@ -0,0 +1,67 @@
+// Copyright Epic Games, Inc. All Rights Reserved.
+
+#pragma once
+
+#include "CoreMinimal.h"
+#include "UObject/ObjectMacros.h"
+#include "GameFramework/Actor.h"
+#include "WorldPartition/WorldPartitionStreamingSource.h"
+#include "WorldPartitionReplay.generated.h"
+
+struct FWorldPartitionReplayStreamingSource : public FWorldPartitionStreamingSource
+{
+	FWorldPartitionReplayStreamingSource()
+	{
+
+	}
+
+	FWorldPartitionReplayStreamingSource(const FWorldPartitionStreamingSource& InStreamingSource)
+		: FWorldPartitionStreamingSource(InStreamingSource.Name, InStreamingSource.Location, InStreamingSource.Rotation, InStreamingSource.TargetState, InStreamingSource.bBlockOnSlowLoading, InStreamingSource.Priority, InStreamingSource.Velocity)
+	{
+	}
+
+	friend FArchive& operator<<(FArchive& Ar, FWorldPartitionReplayStreamingSource& StreamingSource);
+};
+
+struct FWorldPartitionReplaySample
+{
+	FWorldPartitionReplaySample(AWorldPartitionReplay* InReplay);
+			
+	friend FArchive& operator<<(FArchive& Ar, FWorldPartitionReplaySample& StreamingSource);
+	
+	TArray<int32> StreamingSourceNameIndices;
+	TArray<FWorldPartitionReplayStreamingSource> StreamingSources;
+	class AWorldPartitionReplay* Replay = nullptr;
+	float TimeSeconds = 0.f;
+};
+
+/**
+ * Actor used to record world partition replay data (streaming sources for now)
+ */
+UCLASS(notplaceable)
+class ENGINE_API AWorldPartitionReplay : public AActor
+{
+	GENERATED_UCLASS_BODY()
+
+public:
+	static void Initialize(UWorld* World);
+	static bool IsEnabled(UWorld* World);
+
+	virtual void BeginPlay() override;
+	virtual void PreReplication(IRepChangedPropertyTracker& ChangedPropertyTracker) override;
+		
+	bool IsEnabled() const { return bEnabled; }
+	bool GetReplayStreamingSources(TArray<FWorldPartitionStreamingSource>& OutStreamingSources);
+		
+private:
+	friend FArchive& operator<<(FArchive& Ar, FWorldPartitionReplaySample& StreamingSource);
+	friend struct FWorldPartitionReplaySample;
+	TArray<FWorldPartitionStreamingSource> GetRecordingStreamingSources() const;
+
+	UPROPERTY(Transient, Replicated)
+	TArray<FName> StreamingSourceNames;
+
+	TArray<FWorldPartitionReplaySample> ReplaySamples;
+	bool bEnabled = false;
+};
+
