WorldPartition: Add StreamingSource extra info in the debug display

diff --git a/Engine/Source/Runtime/Engine/Private/WorldPartition/WorldPartitionSubsystem.cpp b/Engine/Source/Runtime/Engine/Private/WorldPartition/WorldPartitionSubsystem.cpp
--- a/Engine/Source/Runtime/Engine/Private/WorldPartition/WorldPartitionSubsystem.cpp
+++ b/Engine/Source/Runtime/Engine/Private/WorldPartition/WorldPartitionSubsystem.cpp
@@ -294,7 +294,14 @@ void UWorldPartitionSubsystem::Draw(UCanvas* Canvas, class APlayerController* PC
 			float MaxTextWidth = 0;
 			for (const FWorldPartitionStreamingSource& StreamingSource : *StreamingSources)
 			{
-				FWorldPartitionDebugHelper::DrawText(Canvas, StreamingSource.Name.ToString(), GEngine->GetSmallFont(), StreamingSource.GetDebugColor(), Pos, &MaxTextWidth);
+				FString StreamingSourceDisplay = StreamingSource.Name.ToString();
+#if !NO_LOGGING
+				if (!StreamingSource.ExtraInfo.IsEmpty())
+				{
+					StreamingSourceDisplay += FString::Printf(TEXT(" (%s)"), *StreamingSource.ExtraInfo );
+				}
+#endif
+				FWorldPartitionDebugHelper::DrawText(Canvas, StreamingSourceDisplay, GEngine->GetSmallFont(), StreamingSource.GetDebugColor(), Pos, &MaxTextWidth);
 			}
 			Pos = CurrentOffset + FVector2D(MaxTextWidth + 10, 0.f);
 			for (const FWorldPartitionStreamingSource& StreamingSource : *StreamingSources)
