Add missing transient flag on AWorldPartitionReplay

diff --git a/Engine/Source/Runtime/Engine/Public/WorldPartition/WorldPartitionReplay.h b/Engine/Source/Runtime/Engine/Public/WorldPartition/WorldPartitionReplay.h
--- a/Engine/Source/Runtime/Engine/Public/WorldPartition/WorldPartitionReplay.h
+++ b/Engine/Source/Runtime/Engine/Public/WorldPartition/WorldPartitionReplay.h
@@ -38,7 +38,7 @@ struct FWorldPartitionReplaySample
 /**
  * Actor used to record world partition replay data (streaming sources for now)
  */
-UCLASS(notplaceable)
+UCLASS(notplaceable, transient)
 class ENGINE_API AWorldPartitionReplay : public AActor
 {
 	GENERATED_UCLASS_BODY()
