[Navigation] Fix in RVO avoidance with some trivial name corrections

diff --git a/Engine/Source/Runtime/Engine/Private/AI/Navigation/AvoidanceManager.cpp b/Engine/Source/Runtime/Engine/Private/AI/Navigation/AvoidanceManager.cpp
--- a/Engine/Source/Runtime/Engine/Private/AI/Navigation/AvoidanceManager.cpp
+++ b/Engine/Source/Runtime/Engine/Private/AI/Navigation/AvoidanceManager.cpp
@@ -328,23 +328,23 @@ static bool AvoidsNavEdges(const FVector& OrgLocation, const FVector& TestVeloci
 
 	for (int32 Idx = 0; Idx < NavEdges.Num(); Idx++)
 	{
-		const FVector2D Seg1ToSeg0(NavEdges[Idx].P1 - NavEdges[Idx].P0);
-		const FVector2D NewPosToOrg(TestVelocity);
+		const FVector2D Seg0ToSeg1(NavEdges[Idx].P1 - NavEdges[Idx].P0);
+		const FVector2D OrgToNewPos(TestVelocity);
 		const FVector2D OrgToSeg0(NavEdges[Idx].P0 - OrgLocation);
-		const float CrossD = FVector2D::CrossProduct(Seg1ToSeg0, NewPosToOrg);
+		const float CrossD = FVector2D::CrossProduct(Seg0ToSeg1, OrgToNewPos);
 		if (FMath::Abs(CrossD) < KINDA_SMALL_NUMBER)
 		{
 			continue;
 		}
 
-		const float CrossS = FVector2D::CrossProduct(NewPosToOrg, OrgToSeg0) / CrossD;
-		const float CrossT = FVector2D::CrossProduct(Seg1ToSeg0, OrgToSeg0) / CrossD;
+		const float CrossS = FVector2D::CrossProduct(OrgToNewPos, OrgToSeg0) / CrossD;
+		const float CrossT = FVector2D::CrossProduct(Seg0ToSeg1, OrgToSeg0) / CrossD;
 		if (CrossS < 0.0f || CrossS > 1.0f || CrossT < 0.0f || CrossT > 1.0f)
 		{
 			continue;
 		}
 
-		const FVector CrossPt = FMath::Lerp(NavEdges[Idx].P0, NavEdges[Idx].P1, CrossT);
+		const FVector CrossPt = FMath::Lerp(NavEdges[Idx].P0, NavEdges[Idx].P1, CrossS);
 		const float ZDiff = FMath::Abs(OrgLocation.Z - CrossPt.Z);
 		if (ZDiff > MaxZDiff)
 		{
