Force disable r.Streaming.UseTextureStreamingBuiltData in code.

diff --git a/Engine/Source/Runtime/Engine/Private/Streaming/TextureStreamingBuild.cpp b/Engine/Source/Runtime/Engine/Private/Streaming/TextureStreamingBuild.cpp
--- a/Engine/Source/Runtime/Engine/Private/Streaming/TextureStreamingBuild.cpp
+++ b/Engine/Source/Runtime/Engine/Private/Streaming/TextureStreamingBuild.cpp
@@ -358,7 +358,7 @@ void FStreamingTextureBuildInfo::PackFrom(ITextureStreamingContainer* TextureStr
 }
 #endif
 
-bool FStreamingTextureLevelContext::UseTextureStreamingBuiltData = true;
+bool FStreamingTextureLevelContext::UseTextureStreamingBuiltData = false;
 FAutoConsoleCommand FStreamingTextureLevelContext::UseTextureStreamingBuiltDataCommand(
 	TEXT("r.Streaming.UseTextureStreamingBuiltData"),
 	TEXT("Turn on/off usage of texture streaming built data (0 to turn off)."),
