LevelInstance: Fix Attachment lost on Break operation

diff --git a/Engine/Source/Runtime/Engine/Private/LevelInstance/LevelInstanceSubsystem.cpp b/Engine/Source/Runtime/Engine/Private/LevelInstance/LevelInstanceSubsystem.cpp
--- a/Engine/Source/Runtime/Engine/Private/LevelInstance/LevelInstanceSubsystem.cpp
+++ b/Engine/Source/Runtime/Engine/Private/LevelInstance/LevelInstanceSubsystem.cpp
@@ -937,22 +937,43 @@ void ULevelInstanceSubsystem::BreakLevelInstance_Impl(ALevelInstance* LevelInsta
 			return;
 		}
 
-		TArray<AActor*> ActorsToMove;
-		ForEachActorInLevelInstance(LevelInstanceActor, [this, &ActorsToMove](AActor* Actor)
+		TSet<AActor*> ActorsToMove;
+		TFunction<bool(AActor*)> AddActorToMove = [this, &ActorsToMove, &AddActorToMove](AActor* Actor)
+		{
+			if (ActorsToMove.Contains(Actor))
+			{
+				return true;
+			}
+
+			// Skip some actor types
+			if (!Actor->IsA<ALevelBounds>() && !Actor->IsA<ABrush>() && !Actor->IsA<AWorldSettings>() && !Actor->IsA<ALevelInstanceEditorInstanceActor>())
 			{
-				// Skip some actor types
-				if (!Actor->IsA<ALevelBounds>() && !Actor->IsA<ABrush>() && !Actor->IsA<AWorldSettings>() && !Actor->IsA<ALevelInstanceEditorInstanceActor>())
+				if (CanMoveActorToLevel(Actor))
 				{
-					if (CanMoveActorToLevel(Actor))
+					FSetActorHiddenInSceneOutliner Show(Actor, false);
+
+					// Detach if Parent Actor can't be moved
+					if (AActor* ParentActor = Actor->GetAttachParentActor())
 					{
-						FSetActorHiddenInSceneOutliner Show(Actor, false);
-						Actor->DetachFromActor(FDetachmentTransformRules::KeepWorldTransform);
-						ActorsToMove.Add(Actor);
+						if (!AddActorToMove(ParentActor))
+						{
+							Actor->DetachFromActor(FDetachmentTransformRules::KeepWorldTransform);
+						}
 					}
+
+					ActorsToMove.Add(Actor);
+					return true;
 				}
+			}
 
-				return true;
-			});
+			return false;
+		};
+
+		ForEachActorInLevelInstance(LevelInstanceActor, [this, &ActorsToMove, &AddActorToMove](AActor* Actor)
+		{
+			AddActorToMove(Actor);
+			return true;
+		});
 
 		ULevel* DestinationLevel = GetWorld()->GetCurrentLevel();
 		check(DestinationLevel);
@@ -960,7 +981,7 @@ void ULevelInstanceSubsystem::BreakLevelInstance_Impl(ALevelInstance* LevelInsta
 		const bool bWarnAboutReferences = true;
 		const bool bWarnAboutRenaming = true;
 		const bool bMoveAllOrFail = true;
-		if (!EditorLevelUtils::CopyActorsToLevel(ActorsToMove, DestinationLevel, bWarnAboutReferences, bWarnAboutRenaming, bMoveAllOrFail))
+		if (!EditorLevelUtils::CopyActorsToLevel(ActorsToMove.Array(), DestinationLevel, bWarnAboutReferences, bWarnAboutRenaming, bMoveAllOrFail))
 		{
 			UE_LOG(LogLevelInstance, Warning, TEXT("Failed to break Level Instance because not all actors could be moved"));
 			return;
