Fix Build.

diff --git a/Engine/Source/Runtime/Engine/Private/WorldPartition/WorldPartitionSubsystem.cpp b/Engine/Source/Runtime/Engine/Private/WorldPartition/WorldPartitionSubsystem.cpp
--- a/Engine/Source/Runtime/Engine/Private/WorldPartition/WorldPartitionSubsystem.cpp
+++ b/Engine/Source/Runtime/Engine/Private/WorldPartition/WorldPartitionSubsystem.cpp
@@ -295,12 +295,6 @@ void UWorldPartitionSubsystem::Draw(UCanvas* Canvas, class APlayerController* PC
 			for (const FWorldPartitionStreamingSource& StreamingSource : *StreamingSources)
 			{
 				FString StreamingSourceDisplay = StreamingSource.Name.ToString();
-#if !NO_LOGGING
-				if (!StreamingSource.ExtraInfo.IsEmpty())
-				{
-					StreamingSourceDisplay += FString::Printf(TEXT(" (%s)"), *StreamingSource.ExtraInfo );
-				}
-#endif
 				FWorldPartitionDebugHelper::DrawText(Canvas, StreamingSourceDisplay, GEngine->GetSmallFont(), StreamingSource.GetDebugColor(), Pos, &MaxTextWidth);
 			}
 			Pos = CurrentOffset + FVector2D(MaxTextWidth + 10, 0.f);
