Revert SaveToComponentInstanceDataCache on ComponentInstanceData and use EditCondition=false instead. It give the same result and it's less intrusive.

diff --git a/Engine/Source/Runtime/Engine/Classes/Engine/EngineTypes.h b/Engine/Source/Runtime/Engine/Classes/Engine/EngineTypes.h
--- a/Engine/Source/Runtime/Engine/Classes/Engine/EngineTypes.h
+++ b/Engine/Source/Runtime/Engine/Classes/Engine/EngineTypes.h
@@ -3864,7 +3864,7 @@ struct ENGINE_API FComponentReference
 	FName ComponentProperty;
 
 	/** Path to the component from its owner actor */
-	UPROPERTY(meta = (SaveToComponentInstanceDataCache))
+	UPROPERTY(EditDefaultsOnly, Category=Component, meta =(EditCondition="false", EditConditionHides))
 	FString PathToComponent;
 
 	/** Allows direct setting of first component to constraint. */
diff --git a/Engine/Source/Runtime/Engine/Private/ComponentInstanceDataCache.cpp b/Engine/Source/Runtime/Engine/Private/ComponentInstanceDataCache.cpp
--- a/Engine/Source/Runtime/Engine/Private/ComponentInstanceDataCache.cpp
+++ b/Engine/Source/Runtime/Engine/Private/ComponentInstanceDataCache.cpp
@@ -58,21 +58,13 @@ public:
 
 	virtual bool ShouldSkipProperty(const FProperty* InProperty) const override
 	{
-#if WITH_EDITORONLY_DATA
-		static FName NAME_SaveToComponentInstanceDataCache = "SaveToComponentInstanceDataCache";
-#endif
-
 		// Immutable structs expect to serialize all properties so don't skip regardless of other conditions
 		UScriptStruct* ScriptStruct = InProperty->GetOwner<UScriptStruct>();
 		const bool bPropertyInImmutableStruct = ScriptStruct && ((ScriptStruct->StructFlags & STRUCT_Immutable) != 0);
 
 		return (!bPropertyInImmutableStruct
 			&& (InProperty->HasAnyPropertyFlags(CPF_Transient)
-#if WITH_EDITORONLY_DATA
-				|| (!InProperty->HasAnyPropertyFlags(CPF_Edit | CPF_Interp) && !InProperty->HasMetaData(NAME_SaveToComponentInstanceDataCache))
-#else
 				|| !InProperty->HasAnyPropertyFlags(CPF_Edit | CPF_Interp)
-#endif
 				|| InProperty->IsA<FMulticastDelegateProperty>()
 				|| PropertiesToSkip.Contains(InProperty)
 				)
