Fixing static analysis warnings

diff --git a/Engine/Source/Runtime/Engine/Private/EngineTypes.cpp b/Engine/Source/Runtime/Engine/Private/EngineTypes.cpp
--- a/Engine/Source/Runtime/Engine/Private/EngineTypes.cpp
+++ b/Engine/Source/Runtime/Engine/Private/EngineTypes.cpp
@@ -287,8 +287,10 @@ FActorInstanceHandle::FActorInstanceHandle(int32 InManagerIndex, int32 InInstanc
 		}
 
 		UWorld* World = Manager->GetWorld();
-		ensure(World);
-		InstanceUID = World->LWILastAssignedUID++;
+		if (ensure(World))
+		{
+			InstanceUID = World->LWILastAssignedUID++;
+		}
 	}
 	else
 	{
@@ -316,8 +318,10 @@ FActorInstanceHandle::FActorInstanceHandle(ALightWeightInstanceManager* Manager,
 		ensure(ManagerIndex != INDEX_NONE);
 
 		UWorld* World = Manager->GetWorld();
-		ensure(World);
-		InstanceUID = World->LWILastAssignedUID++;
+		if (ensure(World))
+		{
+			InstanceUID = World->LWILastAssignedUID++;
+		}
 	}
 }
 
