Add ray tracing markers to tell which skin cache entries are which category

diff --git a/Engine/Source/Runtime/Engine/Private/GPUSkinCache.cpp b/Engine/Source/Runtime/Engine/Private/GPUSkinCache.cpp
--- a/Engine/Source/Runtime/Engine/Private/GPUSkinCache.cpp
+++ b/Engine/Source/Runtime/Engine/Private/GPUSkinCache.cpp
@@ -1145,6 +1145,7 @@ void FGPUSkinCache::DispatchUpdateSkinTangents(FRHICommandListImmediate& RHICmdL
 	FSkeletalMeshRenderData& SkelMeshRenderData = Entry->GPUSkin->GetSkeletalMeshRenderData();
 	const int32 LODIndex = Entry->LOD;
 	FSkeletalMeshLODRenderData& LodData = SkelMeshRenderData.LODRenderData[LODIndex];
+	const FString RayTracingTag = (Entry->Mode == EGPUSkinCacheEntryMode::RayTracing ? TEXT("[RT]") : TEXT(""));
 
 	if (bTrianglePass)
 	{
@@ -1205,8 +1206,8 @@ void FGPUSkinCache::DispatchUpdateSkinTangents(FRHICommandListImmediate& RHICmdL
 			uint32 NumTriangles = DispatchData.NumTriangles;
 			uint32 ThreadGroupCountValue = FMath::DivideAndRoundUp(NumTriangles, FBaseRecomputeTangentsPerTriangleShader::ThreadGroupSizeX);
 
-			SCOPED_DRAW_EVENTF(RHICmdList, SkinTangents_PerTrianglePass, TEXT("TangentsTri  Mesh=%s, LOD=%d, Chunk=%d, IndexStart=%d Tri=%d BoneInfluenceType=%d UVPrecision=%d"),
-				*GetSkeletalMeshObjectName(Entry->GPUSkin), LODIndex, SectionIndex, DispatchData.IndexBufferOffsetValue, DispatchData.NumTriangles, Entry->BoneInfluenceType, bFullPrecisionUV);
+			SCOPED_DRAW_EVENTF(RHICmdList, SkinTangents_PerTrianglePass, TEXT("%sTangentsTri  Mesh=%s, LOD=%d, Chunk=%d, IndexStart=%d Tri=%d BoneInfluenceType=%d UVPrecision=%d"),
+				*RayTracingTag , *GetSkeletalMeshObjectName(Entry->GPUSkin), LODIndex, SectionIndex, DispatchData.IndexBufferOffsetValue, DispatchData.NumTriangles, Entry->BoneInfluenceType, bFullPrecisionUV);
 
 			FRHIComputeShader* ShaderRHI = Shader.GetComputeShader();
 			RHICmdList.SetComputeShader(ShaderRHI);
@@ -1238,8 +1239,8 @@ void FGPUSkinCache::DispatchUpdateSkinTangents(FRHICommandListImmediate& RHICmdL
 	}
 	else
 	{
-		SCOPED_DRAW_EVENTF(RHICmdList, SkinTangents_PerVertexPass, TEXT("TangentsVertex Mesh=%s, LOD=%d, Chunk=%d, InputStreamStart=%d, OutputStreamStart=%d, Vert=%d"),
-			*GetSkeletalMeshObjectName(Entry->GPUSkin), LODIndex, SectionIndex, DispatchData.InputStreamStart, DispatchData.OutputStreamStart, DispatchData.NumVertices);
+		SCOPED_DRAW_EVENTF(RHICmdList, SkinTangents_PerVertexPass, TEXT("%sTangentsVertex Mesh=%s, LOD=%d, Chunk=%d, InputStreamStart=%d, OutputStreamStart=%d, Vert=%d"),
+			*RayTracingTag, *GetSkeletalMeshObjectName(Entry->GPUSkin), LODIndex, SectionIndex, DispatchData.InputStreamStart, DispatchData.OutputStreamStart, DispatchData.NumVertices);
 		//#todo-gpuskin Feature level?
 		auto* GlobalShaderMap = GetGlobalShaderMap(GetFeatureLevel());
 		TShaderMapRef<FRecomputeTangentsPerVertexPassCS<0>> ComputeShader0(GlobalShaderMap);
@@ -2009,10 +2010,11 @@ void FGPUSkinCache::DispatchUpdateSkinning(FRHICommandListImmediate& RHICmdList,
 {
 	FGPUSkinCacheEntry::FSectionDispatchData& DispatchData = Entry->DispatchData[Section];
 	FGPUBaseSkinVertexFactory::FShaderDataType& ShaderData = DispatchData.SourceVertexFactory->GetShaderData();
+	const FString RayTracingTag = (Entry->Mode == EGPUSkinCacheEntryMode::RayTracing ? TEXT("[RT]") : TEXT(""));
 
 	SCOPED_DRAW_EVENTF(RHICmdList, SkinCacheDispatch,
-		TEXT("Skinning%d%d%d Mesh=%s LOD=%d Chunk=%d InStreamStart=%d OutStart=%d Vert=%d Morph=%d/%d"),
-		(int32)Entry->bUse16BitBoneIndex, (int32)Entry->BoneInfluenceType, DispatchData.SkinType, *GetSkeletalMeshObjectName(Entry->GPUSkin), Entry->LOD,
+		TEXT("%sSkinning%d%d%d Mesh=%s LOD=%d Chunk=%d InStreamStart=%d OutStart=%d Vert=%d Morph=%d/%d"),
+		*RayTracingTag, (int32)Entry->bUse16BitBoneIndex, (int32)Entry->BoneInfluenceType, DispatchData.SkinType, *GetSkeletalMeshObjectName(Entry->GPUSkin), Entry->LOD,
 		DispatchData.SectionIndex, DispatchData.InputStreamStart, DispatchData.OutputStreamStart, DispatchData.NumVertices, Entry->MorphBuffer != 0, DispatchData.MorphBufferOffset);
 	auto* GlobalShaderMap = GetGlobalShaderMap(GetFeatureLevel());
 	TShaderMapRef<TGPUSkinCacheCS<0>> SkinCacheCS000(GlobalShaderMap);		// 16bit_0, BoneInfluenceType_0, SkinType_0
@@ -2407,8 +2409,9 @@ void FGPUSkinCache::PrintMemorySummary() const
 		FGPUSkinCacheEntry* Entry = Entries[i];
 		if (Entry)
 		{
+			const FString RayTracingTag = (Entry->Mode == EGPUSkinCacheEntryMode::RayTracing ? TEXT("[RT]") : TEXT(""));
 			uint64 MemInBytes = Entry->PositionAllocation ? Entry->PositionAllocation->GetNumBytes() : 0;
-			UE_LOG(LogSkinCache, Display, TEXT("   SkinCacheEntry_%d: Mesh=%s, LOD=%d, Mem=%.3fKB"), i, *GetSkeletalMeshObjectName(Entry->GPUSkin), Entry->LOD, MemInBytes / 1024.f);
+			UE_LOG(LogSkinCache, Display, TEXT("   SkinCacheEntry_%d: %sMesh=%s, LOD=%d, Mem=%.3fKB"), i, *RayTracingTag, *GetSkeletalMeshObjectName(Entry->GPUSkin), Entry->LOD, MemInBytes / 1024.f);
 			TotalMemInBytes += MemInBytes;
 		}
 	}
