Replace manager index with a weak pointer to the manager inside actor instance handles. This helps make them more robust to managers being added and removed.

diff --git a/Engine/Plugins/LightWeightInstancesEditor/Source/LightWeightInstancesEditor/Private/LightWeightInstancesEditor.cpp b/Engine/Plugins/LightWeightInstancesEditor/Source/LightWeightInstancesEditor/Private/LightWeightInstancesEditor.cpp
--- a/Engine/Plugins/LightWeightInstancesEditor/Source/LightWeightInstancesEditor/Private/LightWeightInstancesEditor.cpp
+++ b/Engine/Plugins/LightWeightInstancesEditor/Source/LightWeightInstancesEditor/Private/LightWeightInstancesEditor.cpp
@@ -122,7 +122,7 @@ void FLightWeightInstancesEditorModule::ConvertActorsToLWIsUIAction(const TArray
 	}
 }
 
-#endif
+#endif // WITH_EDITOR
 
 #undef LOCTEXT_NAMESPACE
 	
diff --git a/Engine/Source/Runtime/Engine/Classes/Engine/EngineTypes.h b/Engine/Source/Runtime/Engine/Classes/Engine/EngineTypes.h
--- a/Engine/Source/Runtime/Engine/Classes/Engine/EngineTypes.h
+++ b/Engine/Source/Runtime/Engine/Classes/Engine/EngineTypes.h
@@ -2066,7 +2066,6 @@ struct ENGINE_API FActorInstanceHandle
 	FActorInstanceHandle();
 
 	explicit FActorInstanceHandle(AActor* InActor);
-	explicit FActorInstanceHandle(int32 InManagerIndex, int32 InInstanceIndex);
 	explicit FActorInstanceHandle(class ALightWeightInstanceManager* Manager, int32 InInstanceIndex);
 
 	FActorInstanceHandle(const FActorInstanceHandle& Other);
@@ -2133,7 +2132,7 @@ private:
 	mutable TWeakObjectPtr<AActor> Actor;
 
 	/** Identifies the light weight instance manager to use */
-	int32 ManagerIndex;
+	TWeakObjectPtr<ALightWeightInstanceManager> Manager;
 
 	/** Identifies the instance within the manager */
 	int32 InstanceIndex;
diff --git a/Engine/Source/Runtime/Engine/Classes/GameFramework/LightWeightInstanceManager.h b/Engine/Source/Runtime/Engine/Classes/GameFramework/LightWeightInstanceManager.h
--- a/Engine/Source/Runtime/Engine/Classes/GameFramework/LightWeightInstanceManager.h
+++ b/Engine/Source/Runtime/Engine/Classes/GameFramework/LightWeightInstanceManager.h
@@ -241,5 +241,5 @@ protected:
 	//
 protected:
 	virtual void PostEditChangeProperty(struct FPropertyChangedEvent& PropertyChangedEvent) override;	
-#endif
+#endif // WITH_EDITOR
 };
diff --git a/Engine/Source/Runtime/Engine/Private/EngineTypes.cpp b/Engine/Source/Runtime/Engine/Private/EngineTypes.cpp
--- a/Engine/Source/Runtime/Engine/Private/EngineTypes.cpp
+++ b/Engine/Source/Runtime/Engine/Private/EngineTypes.cpp
@@ -239,7 +239,6 @@ UActorComponent* FComponentReference::GetComponent(AActor* OwningActor) const
 
 FActorInstanceHandle::FActorInstanceHandle()
 	: Actor(nullptr)
-	, ManagerIndex(INDEX_NONE)
 	, InstanceIndex(INDEX_NONE)
 	, InstanceUID(0)
 {
@@ -248,7 +247,6 @@ FActorInstanceHandle::FActorInstanceHandle()
 
 FActorInstanceHandle::FActorInstanceHandle(AActor* InActor)
 	: Actor(InActor)
-	, ManagerIndex(INDEX_NONE)
 	, InstanceIndex(INDEX_NONE)
 {
 	if (InActor)
@@ -260,53 +258,17 @@ FActorInstanceHandle::FActorInstanceHandle(AActor* InActor)
 #else
 		const UDataLayer* Layer = nullptr;
 #endif // WITH_EDITOR
-		if (ALightWeightInstanceManager* LWIManager = FLightWeightInstanceSubsystem::Get().FindLightWeightInstanceManager(InActor->StaticClass(), Layer))
-		{
-			InstanceIndex = LWIManager->FindIndexForActor(InActor);
-			if (InstanceIndex != INDEX_NONE)
-			{
-				ManagerIndex = FLightWeightInstanceSubsystem::Get().GetManagerIndex(LWIManager);
-			}
-		}
+		Manager = FLightWeightInstanceSubsystem::Get().FindLightWeightInstanceManager(InActor->StaticClass(), Layer);
 	}
 }
 
-FActorInstanceHandle::FActorInstanceHandle(int32 InManagerIndex, int32 InInstanceIndex)
+FActorInstanceHandle::FActorInstanceHandle(ALightWeightInstanceManager* InManager, int32 InInstanceIndex)
 	: Actor(nullptr)
-	, ManagerIndex(InManagerIndex)
-{
-	const ALightWeightInstanceManager* Manager = FLightWeightInstanceSubsystem::Get().GetManagerAt(InManagerIndex);
-
-	if (ensure(Manager))
-	{
-		InstanceIndex = Manager->ConvertCollisionIndexToLightWeightIndex(InInstanceIndex);
-
-		if (AActor* const* FoundActor = Manager->Actors.Find(InInstanceIndex))
-		{
-			Actor = *FoundActor;
-		}
-
-		UWorld* World = Manager->GetWorld();
-		if (ensure(World))
-		{
-			InstanceUID = World->LWILastAssignedUID++;
-		}
-	}
-	else
-	{
-		ManagerIndex = INDEX_NONE;
-		InstanceIndex = INDEX_NONE;
-		InstanceUID = 0;
-	}
-}
-
-FActorInstanceHandle::FActorInstanceHandle(ALightWeightInstanceManager* Manager, int32 InInstanceIndex)
-	: Actor(nullptr)
-	, ManagerIndex(INDEX_NONE)
 	, InstanceIndex(InInstanceIndex)
 	, InstanceUID(0)
 {
-	if (ensure(Manager))
+	Manager = InManager;
+	if (Manager.IsValid())
 	{
 		InstanceIndex = Manager->ConvertCollisionIndexToLightWeightIndex(InInstanceIndex);
 		if (AActor* const* FoundActor = Manager->Actors.Find(InstanceIndex))
@@ -314,9 +276,6 @@ FActorInstanceHandle::FActorInstanceHandle(ALightWeightInstanceManager* Manager,
 			Actor = *FoundActor;
 		}
 
-		ManagerIndex = FLightWeightInstanceSubsystem::Get().GetManagerIndex(Manager);
-		ensure(ManagerIndex != INDEX_NONE);
-
 		UWorld* World = Manager->GetWorld();
 		if (ensure(World))
 		{
@@ -328,14 +287,15 @@ FActorInstanceHandle::FActorInstanceHandle(ALightWeightInstanceManager* Manager,
 FActorInstanceHandle::FActorInstanceHandle(const FActorInstanceHandle& Other)
 {
 	Actor = Other.Actor;
-	ManagerIndex = Other.ManagerIndex;
+
+	Manager = Other.Manager;
 	InstanceIndex = Other.InstanceIndex;
 	InstanceUID = Other.InstanceUID;
 }
 
 bool FActorInstanceHandle::IsValid() const
 {
-	return (ManagerIndex != INDEX_NONE && InstanceIndex != INDEX_NONE) || IsActorValid();
+	return (Manager.IsValid() && InstanceIndex != INDEX_NONE) || IsActorValid();
 }
 
 bool FActorInstanceHandle::DoesRepresentClass(const UClass* OtherClass) const
@@ -350,7 +310,7 @@ bool FActorInstanceHandle::DoesRepresentClass(const UClass* OtherClass) const
 		return Actor->IsA(OtherClass);
 	}
 
-	if (const ALightWeightInstanceManager* Manager = FLightWeightInstanceSubsystem::Get().GetManagerAt(ManagerIndex))
+	if (Manager.IsValid())
 	{
 		return Manager->DoesRepresentClass(OtherClass);
 	}
@@ -370,7 +330,7 @@ UClass* FActorInstanceHandle::GetRepresentedClass() const
 		return Actor->GetClass();
 	}
 
-	if (const ALightWeightInstanceManager* Manager = FLightWeightInstanceSubsystem::Get().GetManagerAt(ManagerIndex))
+	if (Manager.IsValid())
 	{
 		return Manager->GetRepresentedClass();
 	}
@@ -385,7 +345,7 @@ FVector FActorInstanceHandle::GetLocation() const
 		return Actor->GetActorLocation();
 	}
 
-	if (const ALightWeightInstanceManager* Manager = FLightWeightInstanceSubsystem::Get().GetManagerAt(ManagerIndex))
+	if (Manager.IsValid())
 	{
 		Manager->GetLocation(*this);
 	}
@@ -400,7 +360,7 @@ FRotator FActorInstanceHandle::GetRotation() const
 		return Actor->GetActorRotation();
 	}
 
-	if (const ALightWeightInstanceManager* Manager = FLightWeightInstanceSubsystem::Get().GetManagerAt(ManagerIndex))
+	if (Manager.IsValid())
 	{
 		Manager->GetRotation(*this);
 	}
@@ -415,7 +375,7 @@ FTransform FActorInstanceHandle::GetTransform() const
 		return Actor->GetActorTransform();
 	}
 
-	if (const ALightWeightInstanceManager* Manager = FLightWeightInstanceSubsystem::Get().GetManagerAt(ManagerIndex))
+	if (Manager.IsValid())
 	{
 		Manager->GetTransform(*this);
 	}
@@ -440,7 +400,7 @@ FString FActorInstanceHandle::GetName() const
 		return Actor->GetName();
 	}
 
-	if (const ALightWeightInstanceManager* Manager = FLightWeightInstanceSubsystem::Get().GetManagerAt(ManagerIndex))
+	if (Manager.IsValid())
 	{
 		Manager->GetName(*this);
 	}
@@ -455,7 +415,7 @@ AActor* FActorInstanceHandle::GetManagingActor() const
 		return Actor.Get();
 	}
 
-	return const_cast<ALightWeightInstanceManager*>(FLightWeightInstanceSubsystem::Get().GetManagerAt(ManagerIndex));
+	return Manager.Get();
 }
 
 USceneComponent* FActorInstanceHandle::GetRootComponent() const
@@ -465,8 +425,7 @@ USceneComponent* FActorInstanceHandle::GetRootComponent() const
 		return Actor->GetRootComponent();
 	}
 
-	const AActor* Manager = FLightWeightInstanceSubsystem::Get().GetManagerAt(ManagerIndex);
-	return Manager ? Manager->GetRootComponent() : nullptr;
+	return Manager.IsValid() ? Manager->GetRootComponent() : nullptr;
 }
 
 AActor* FActorInstanceHandle::FetchActor() const
@@ -481,8 +440,7 @@ AActor* FActorInstanceHandle::FetchActor() const
 
 int32 FActorInstanceHandle::GetRenderingInstanceIndex() const
 {
-	const ALightWeightInstanceManager* Manager = FLightWeightInstanceSubsystem::Get().GetManagerAt(ManagerIndex);
-	return Manager ? Manager->ConvertLightWeightIndexToCollisionIndex(InstanceIndex) : INDEX_NONE;
+	return Manager.IsValid() ? Manager->ConvertLightWeightIndexToCollisionIndex(InstanceIndex) : INDEX_NONE;
 }
 
 UObject* FActorInstanceHandle::GetActorAsUObject()
@@ -513,7 +471,7 @@ bool FActorInstanceHandle::IsActorValid() const
 FActorInstanceHandle& FActorInstanceHandle::operator=(AActor* OtherActor)
 {
 	Actor = OtherActor;
-	ManagerIndex = INDEX_NONE;
+	Manager.Reset();
 	InstanceIndex = INDEX_NONE;
 
 	return *this;
@@ -521,7 +479,7 @@ FActorInstanceHandle& FActorInstanceHandle::operator=(AActor* OtherActor)
 
 bool FActorInstanceHandle::operator==(const FActorInstanceHandle& Other) const
 {
-	return ManagerIndex == Other.ManagerIndex && InstanceIndex == Other.InstanceIndex;
+	return Manager == Other.Manager && InstanceIndex == Other.InstanceIndex;
 }
 
 bool FActorInstanceHandle::operator!=(const FActorInstanceHandle& Other) const
@@ -540,7 +498,7 @@ bool FActorInstanceHandle::operator==(const AActor* OtherActor) const
 	// if OtherActor is null then we're only equal if this doesn't refer to a valid instance
 	if (OtherActor == nullptr)
 	{
-		return ManagerIndex == INDEX_NONE && InstanceIndex == INDEX_NONE;
+		return !Manager.IsValid() && InstanceIndex == INDEX_NONE;
 	}
 
 	// we don't have an actor so see if we can look up an instance associated with OtherActor and see if we refer to the same instance
@@ -553,9 +511,9 @@ bool FActorInstanceHandle::operator==(const AActor* OtherActor) const
 	const UDataLayer* Layer = nullptr;
 #endif // WITH_EDITOR
 
-	if (ALightWeightInstanceManager* Manager = FLightWeightInstanceSubsystem::Get().FindLightWeightInstanceManager(OtherActor->StaticClass(), Layer))
+	if (ALightWeightInstanceManager* LWIManager = FLightWeightInstanceSubsystem::Get().FindLightWeightInstanceManager(OtherActor->StaticClass(), Layer))
 	{
-		if (FLightWeightInstanceSubsystem::Get().GetManagerIndex(Manager) != ManagerIndex)
+		if (Manager.Get() != LWIManager)
 		{
 			return false;
 		}
@@ -578,7 +536,10 @@ uint32 GetTypeHash(const FActorInstanceHandle& Handle)
 	{
 		FCrc::StrCrc32(*(Handle.Actor->GetPathName()), Hash);
 	}
-	Hash = HashCombine(Hash, Handle.ManagerIndex);
+	if (Handle.Manager.IsValid())
+	{
+		Hash = HashCombine(Hash, GetTypeHash(Handle.Manager.Get()));
+	}
 	Hash = HashCombine(Hash, Handle.InstanceIndex);
 
 	return Hash;
@@ -587,7 +548,7 @@ uint32 GetTypeHash(const FActorInstanceHandle& Handle)
 FArchive& operator<<(FArchive& Ar, FActorInstanceHandle& Handle)
 {
 	Ar << Handle.Actor;
-	Ar << Handle.ManagerIndex;
+	Ar << Handle.Manager;
 	Ar << Handle.InstanceIndex;
 
 	return Ar;
diff --git a/Engine/Source/Runtime/Engine/Private/LightWeightInstanceManager.cpp b/Engine/Source/Runtime/Engine/Private/LightWeightInstanceManager.cpp
--- a/Engine/Source/Runtime/Engine/Private/LightWeightInstanceManager.cpp
+++ b/Engine/Source/Runtime/Engine/Private/LightWeightInstanceManager.cpp
@@ -70,7 +70,7 @@ AActor* ALightWeightInstanceManager::ConvertInstanceToActor(const FActorInstance
 		return Actors[Handle.GetInstanceIndex()];
 	}
 	// if we're pointing at an invalid index we can't return an actor
-	if (ValidIndices.Num() >= Handle.GetInstanceIndex() || ValidIndices[Handle.GetInstanceIndex()] == false)
+	if (ValidIndices.Num() <= Handle.GetInstanceIndex() || ValidIndices[Handle.GetInstanceIndex()] == false)
 	{
 		return nullptr;
 	}
@@ -135,7 +135,7 @@ FActorInstanceHandle ALightWeightInstanceManager::ConvertActorToLightWeightInsta
 		}
 
 		// Update our handle
-		ReturnHandle.ManagerIndex = FLightWeightInstanceSubsystem::Get().GetManagerIndex(this);
+		ReturnHandle.Manager = this;
 		ReturnHandle.InstanceIndex = Idx;
 		ReturnHandle.InstanceUID = InActor->GetWorld()->LWILastAssignedUID++;
 
diff --git a/Engine/Source/Runtime/Engine/Private/LightWeightInstanceSubsystem.cpp b/Engine/Source/Runtime/Engine/Private/LightWeightInstanceSubsystem.cpp
--- a/Engine/Source/Runtime/Engine/Private/LightWeightInstanceSubsystem.cpp
+++ b/Engine/Source/Runtime/Engine/Private/LightWeightInstanceSubsystem.cpp
@@ -60,12 +60,9 @@ const ALightWeightInstanceManager* FLightWeightInstanceSubsystem::GetManagerAt(i
 
 ALightWeightInstanceManager* FLightWeightInstanceSubsystem::FindLightWeightInstanceManager(const FActorInstanceHandle& Handle) const
 {
-	if (Handle.ManagerIndex != INDEX_NONE)
+	if (Handle.Manager.IsValid())
 	{
-		if (ensure(Handle.ManagerIndex < LWInstanceManagers.Num()))
-		{
-			return LWInstanceManagers[Handle.ManagerIndex];
-		}
+		return Handle.Manager.Get();
 	}
 
 	if (Handle.Actor.IsValid())
