Add ray tracing tag to skin cache fail to allocate memory log

diff --git a/Engine/Source/Runtime/Engine/Private/GPUSkinCache.cpp b/Engine/Source/Runtime/Engine/Private/GPUSkinCache.cpp
--- a/Engine/Source/Runtime/Engine/Private/GPUSkinCache.cpp
+++ b/Engine/Source/Runtime/Engine/Private/GPUSkinCache.cpp
@@ -1577,9 +1577,10 @@ bool FGPUSkinCache::ProcessEntry(
 		{
 			if (GSkinCachePrintMemorySummary > 0)
 			{
+				const FString RayTracingTag = (Mode == EGPUSkinCacheEntryMode::RayTracing ? TEXT("[RT]") : TEXT(""));
 				uint64 RequiredMemInBytes = FRWBuffersAllocation::CalculateRequiredMemory(TotalNumVertices, WithTangents, bEntryUseIntermediateTangents, TotalNumTriangles);
-				UE_LOG(LogSkinCache, Warning, TEXT("FGPUSkinCache::ProcessEntry failed to allocate %.3fMB for mesh %s LOD%d, extra required memory increased to %.3fMB"),
-					RequiredMemInBytes / MBSize, *GetSkeletalMeshObjectName(Skin), LODIndex, ExtraRequiredMemory / MBSize);
+				UE_LOG(LogSkinCache, Warning, TEXT("FGPUSkinCache::ProcessEntry%s failed to allocate %.3fMB for mesh %s LOD%d, extra required memory increased to %.3fMB"),
+					*RayTracingTag, RequiredMemInBytes / MBSize, *GetSkeletalMeshObjectName(Skin), LODIndex, ExtraRequiredMemory / MBSize);
 			}
 
 			// Couldn't fit; caller will notify OOM
