Add -RESOURCESIZESORT parameter to obj list command.  This will sort the UObjects by their total resource size memory.

- For static and skeletal meshes this is a rough approxmiation of how much GPU memory the asset will take up.

diff --git a/Engine/Source/Runtime/Engine/Private/UnrealEngine.cpp b/Engine/Source/Runtime/Engine/Private/UnrealEngine.cpp
--- a/Engine/Source/Runtime/Engine/Private/UnrealEngine.cpp
+++ b/Engine/Source/Runtime/Engine/Private/UnrealEngine.cpp
@@ -8111,6 +8111,7 @@ bool UEngine::HandleObjCommand( const TCHAR* Cmd, FOutputDevice& Ar )
 
 		const bool bAlphaSort = FParse::Param( Cmd, TEXT("ALPHASORT") );
 		const bool bCountSort = FParse::Param( Cmd, TEXT("COUNTSORT") );
+		const bool bResourceSizeSort = FParse::Param(Cmd, TEXT("RESOURCESIZESORT"));
 		const bool bCSV = FParse::Param(Cmd, TEXT("CSV"));
 		const bool bShowFullClassName = FParse::Param(Cmd, TEXT("FULLCLASSNAME"));
 
@@ -8118,22 +8119,27 @@ bool UEngine::HandleObjCommand( const TCHAR* Cmd, FOutputDevice& Ar )
 		{
 			struct FCompareFSubItem
 			{
-				bool bAlphaSort;
-				FCompareFSubItem( bool InAlphaSort )
-					: bAlphaSort( InAlphaSort )
+				bool bAlphaSort, bResourceSizeSort;
+				FCompareFSubItem( bool InAlphaSort, bool InResourceSizeSort)
+					: bAlphaSort( InAlphaSort ),
+					  bResourceSizeSort( InResourceSizeSort )
 				{}
 
 				FORCEINLINE bool operator()( const FSubItem& A, const FSubItem& B ) const
 				{
-					if (bAlphaSort || A.Max == B.Max)
+					if (bAlphaSort)
 					{
 						return A.Object->GetPathName() < B.Object->GetPathName();
 					}
+					else if (bResourceSizeSort)
+					{
+						return (B.TrueResourceSize.GetTotalMemoryBytes() < A.TrueResourceSize.GetTotalMemoryBytes());
+					}
 
 					return B.Max < A.Max;
 				}
 			};
-			Objects.Sort( FCompareFSubItem( bAlphaSort ) );
+			Objects.Sort( FCompareFSubItem( bAlphaSort, bResourceSizeSort) );
 
 			if (bCSV)
 			{
@@ -8196,17 +8202,17 @@ bool UEngine::HandleObjCommand( const TCHAR* Cmd, FOutputDevice& Ar )
 		{
 			struct FCompareFItem
 			{
-				bool bAlphaSort, bCountSort;
-				FCompareFItem( bool InAlphaSort, bool InCountSort )
+				bool bAlphaSort, bCountSort, bResourceSizeSort;
+				FCompareFItem( bool InAlphaSort, bool InCountSort, bool InResourceSizeSort )
 					: bAlphaSort( InAlphaSort )
 					, bCountSort( InCountSort )
 				{}
 				FORCEINLINE bool operator()( const FItem& A, const FItem& B ) const
 				{
-					return bAlphaSort ? (A.Class->GetName() < B.Class->GetName()) : bCountSort ? (B.Count < A.Count) : (B.Max < A.Max); 
+					return bAlphaSort ? (A.Class->GetName() < B.Class->GetName()) : bCountSort ? (B.Count < A.Count) : bResourceSizeSort ? (B.TrueResourceSize.GetTotalMemoryBytes() < A.TrueResourceSize.GetTotalMemoryBytes()) : (B.Max < A.Max);
 				}
 			};
-			List.Sort( FCompareFItem( bAlphaSort, bCountSort ) );
+			List.Sort( FCompareFItem( bAlphaSort, bCountSort, bResourceSizeSort ) );
 			
 			if (bCSV)
 			{
