Removing IStreamedCompressedInfo::SampleRate as it was almost never set, was a bad bit depth, and never read globally.
#pf 60f9f43a1f926d0001f7d730

diff --git a/Engine/Source/Runtime/BinkAudioDecoder/Module/Public/BinkAudioInfo.h b/Engine/Source/Runtime/BinkAudioDecoder/Module/Public/BinkAudioInfo.h
--- a/Engine/Source/Runtime/BinkAudioDecoder/Module/Public/BinkAudioInfo.h
+++ b/Engine/Source/Runtime/BinkAudioDecoder/Module/Public/BinkAudioInfo.h
@@ -27,7 +27,8 @@ public:
 protected:
 
 	// copied from header during ParseHeader
-	uint32 MaxCompSpaceNeeded; 
+	uint32 MaxCompSpaceNeeded;
+	uint32 SampleRate;
 
 	struct BinkAudioDecoder* Decoder = 0;
 	uint8* RawMemory = 0;
diff --git a/Engine/Source/Runtime/Engine/Private/AudioDecompress.cpp b/Engine/Source/Runtime/Engine/Private/AudioDecompress.cpp
--- a/Engine/Source/Runtime/Engine/Private/AudioDecompress.cpp
+++ b/Engine/Source/Runtime/Engine/Private/AudioDecompress.cpp
@@ -15,7 +15,6 @@ IStreamedCompressedInfo::IStreamedCompressedInfo()
 	, SrcBufferOffset(0)
 	, AudioDataOffset(0)
 	, AudioDataChunkIndex(0)
-	, SampleRate(0)
 	, TrueSampleCount(0)
 	, CurrentSampleCount(0)
 	, NumChannels(0)
diff --git a/Engine/Source/Runtime/Engine/Public/AudioDecompress.h b/Engine/Source/Runtime/Engine/Public/AudioDecompress.h
--- a/Engine/Source/Runtime/Engine/Public/AudioDecompress.h
+++ b/Engine/Source/Runtime/Engine/Public/AudioDecompress.h
@@ -290,8 +290,6 @@ protected:
 	uint32 AudioDataOffset;
 	/** The chunk index where the actual audio data starts. */
 	uint32 AudioDataChunkIndex;
-	/** Sample rate of the source file */
-	uint16 SampleRate;
 	/** The total sample count of the source file. */
 	uint32 TrueSampleCount;
 	/** How many samples we've currently read in the source file. */
diff --git a/Engine/Source/Runtime/Engine/Public/OpusAudioInfo.h b/Engine/Source/Runtime/Engine/Public/OpusAudioInfo.h
--- a/Engine/Source/Runtime/Engine/Public/OpusAudioInfo.h
+++ b/Engine/Source/Runtime/Engine/Public/OpusAudioInfo.h
@@ -34,4 +34,7 @@ public:
 protected:
 	/** Wrapper around Opus-specific decoding state and APIs */
 	FOpusDecoderWrapper*	OpusDecoderWrapper;
+
+	/** Copied from the header. Opus only need 16 bits */
+	uint16 SampleRate;
 };
