Another attachment fix to avoid dirtying Parent actor package

diff --git a/Engine/Source/Runtime/Engine/Private/LevelActor.cpp b/Engine/Source/Runtime/Engine/Private/LevelActor.cpp
--- a/Engine/Source/Runtime/Engine/Private/LevelActor.cpp
+++ b/Engine/Source/Runtime/Engine/Private/LevelActor.cpp
@@ -900,7 +900,8 @@ bool UWorld::DestroyActor( AActor* ThisActor, bool bNetForce, bool bShouldModify
 		AActor* OldParentActor = RootComp->GetAttachParent()->GetOwner();
 		if (OldParentActor)
 		{
-			OldParentActor->Modify();
+			// Attachment is persisted on the child so modify both actors for Undo/Redo but do not mark the Parent package dirty
+			OldParentActor->Modify(/*bAlwaysMarkDirty=*/false);
 		}
 
 		ThisActor->DetachFromActor(FDetachmentTransformRules::KeepWorldTransform);
