Add nullptr check, should fix static analysis warning

diff --git a/Engine/Source/Runtime/Engine/Private/HLSLTree/HLSLTree.cpp b/Engine/Source/Runtime/Engine/Private/HLSLTree/HLSLTree.cpp
--- a/Engine/Source/Runtime/Engine/Private/HLSLTree/HLSLTree.cpp
+++ b/Engine/Source/Runtime/Engine/Private/HLSLTree/HLSLTree.cpp
@@ -307,7 +307,7 @@ const UE::HLSLTree::FEmitValue* UE::HLSLTree::FEmitContext::AcquireValue(FExpres
 		}
 	}
 
-	return &Entry->Value;
+	return Entry ? &Entry->Value : nullptr;
 }
 
 const UE::HLSLTree::FEmitValue* UE::HLSLTree::FEmitContext::AcquireValue(FFunctionCall* FunctionCall, int32 OutputIndex)
