Optimus: Set better default group size

diff --git a/Engine/Source/Runtime/Engine/Private/ComputeFramework/ComputeGraphScheduler.cpp b/Engine/Source/Runtime/Engine/Private/ComputeFramework/ComputeGraphScheduler.cpp
--- a/Engine/Source/Runtime/Engine/Private/ComputeFramework/ComputeGraphScheduler.cpp
+++ b/Engine/Source/Runtime/Engine/Private/ComputeFramework/ComputeGraphScheduler.cpp
@@ -30,7 +30,7 @@ void FComputeGraphProxy::Initialize(UComputeGraph* ComputeGraph)
 			FKernelInvocation KernelInvocation = {
 				Kernel->GetFName(),
 				FName("InvocationName"),
-				FIntVector(32, 1, 1), // todo[CF]: read group size from kernel (or possibly apply it through defines)
+				FIntVector(64, 1, 1), // todo[CF]: read group size from kernel (or possibly apply it through defines)
 				ShaderMetadata,
 				KernelBindings,
 				KernelResource
