Remove duplicates from the shader jobs count being reported by the progress bar

#rnx

diff --git a/Engine/Source/Runtime/Engine/Private/ShaderCompiler/ShaderCompiler.cpp b/Engine/Source/Runtime/Engine/Private/ShaderCompiler/ShaderCompiler.cpp
--- a/Engine/Source/Runtime/Engine/Private/ShaderCompiler/ShaderCompiler.cpp
+++ b/Engine/Source/Runtime/Engine/Private/ShaderCompiler/ShaderCompiler.cpp
@@ -3136,7 +3136,11 @@ TArrayView<FName> FShaderCompilingManager::GetDependentTypeNames() const
 
 int32 FShaderCompilingManager::GetNumRemainingAssets() const
 {
-	return GetNumOutstandingJobs() + GetNumPendingJobs();
+	// Currently, jobs are difficult to track but the purpose of the GetNumRemainingAssets function is to never return 0
+	// if there are still shaders that have not had their primitives updated on the render thread.
+	// So we track jobs first and when everything is finished compiling but are still lying around in other structures
+	// waiting to be further processed, we show those numbers and ultimately we always return 1 unless IsCompiling() is false.
+	return FMath::Max3(GetNumRemainingJobs(), ShaderMapJobs.Num() + PendingFinalizeShaderMaps.Num(), IsCompiling() ? 1 : 0);
 }
 
 void FShaderCompilingManager::ProcessAsyncTasks(bool bLimitExecutionTime)
