Fix ddc fill not filling audio chunks due to a bug introduced in 7095880
This can also fix audio to be rebuilt even when the data is available from the DDC

#rnx

diff --git a/Engine/Source/Runtime/Engine/Private/AudioDerivedData.cpp b/Engine/Source/Runtime/Engine/Private/AudioDerivedData.cpp
--- a/Engine/Source/Runtime/Engine/Private/AudioDerivedData.cpp
+++ b/Engine/Source/Runtime/Engine/Private/AudioDerivedData.cpp
@@ -858,16 +858,15 @@ int32 FStreamedAudioPlatformData::GetChunkFromDDC(int32 ChunkIndex, uint8** OutC
 	}
 
 	// Wait for async DDC to complete
-	// TODO: Not necessary since bMakeSureChunkIsLoaded was introduced. 
-	// Also, if you do this- why not just call DDC.GetSynchronous?
-// 	if (Chunk.DerivedDataKey.IsEmpty() == false)
-// 	{
-// 		DDC.WaitAsynchronousCompletion(AsyncHandle);
-// 		if (DDC.GetAsynchronousResults(AsyncHandle, TempData))
-// 		{
-// 			ChunkDataSize = DeserializeChunkFromDDC(TempData, Chunk, ChunkIndex, OutChunkData);
-// 		}
-// 	}
+	// Necessary otherwise we will return a ChunkDataSize of 0 which is considered a failure by most callers and will trigger rebuild. 
+	if (Chunk.DerivedDataKey.IsEmpty() == false && AsyncHandle)
+	{
+		DDC.WaitAsynchronousCompletion(AsyncHandle);
+		if (DDC.GetAsynchronousResults(AsyncHandle, TempData))
+		{
+			ChunkDataSize = DeserializeChunkFromDDC(TempData, Chunk, ChunkIndex, OutChunkData);
+		}
+	}
 #else // #if WITH_EDITORONLY_DATA
 	// Load chunk from bulk data if available. If the chunk is not loaded, GetCopy will load it synchronously.
 	if (Chunk.BulkData.IsBulkDataLoaded() || bMakeSureChunkIsLoaded)
