Break LevelInstance: Ignore Default Brush instead of any Brush (fixes fact that some volume types were excluded on Break)

diff --git a/Engine/Source/Runtime/Engine/Private/LevelInstance/LevelInstanceSubsystem.cpp b/Engine/Source/Runtime/Engine/Private/LevelInstance/LevelInstanceSubsystem.cpp
--- a/Engine/Source/Runtime/Engine/Private/LevelInstance/LevelInstanceSubsystem.cpp
+++ b/Engine/Source/Runtime/Engine/Private/LevelInstance/LevelInstanceSubsystem.cpp
@@ -946,7 +946,7 @@ void ULevelInstanceSubsystem::BreakLevelInstance_Impl(ALevelInstance* LevelInsta
 			}
 
 			// Skip some actor types
-			if (!Actor->IsA<ALevelBounds>() && !Actor->IsA<ABrush>() && !Actor->IsA<AWorldSettings>() && !Actor->IsA<ALevelInstanceEditorInstanceActor>())
+			if (!Actor->IsA<ALevelBounds>() && (Actor != Actor->GetLevel()->GetDefaultBrush()) && !Actor->IsA<AWorldSettings>() && !Actor->IsA<ALevelInstanceEditorInstanceActor>())
 			{
 				if (CanMoveActorToLevel(Actor))
 				{
