Fix static analysis warning about unreleased allocations.

diff --git a/Engine/Source/Runtime/Engine/Private/ComputeFramework/ComputeGraphComponent.cpp b/Engine/Source/Runtime/Engine/Private/ComputeFramework/ComputeGraphComponent.cpp
--- a/Engine/Source/Runtime/Engine/Private/ComputeFramework/ComputeGraphComponent.cpp
+++ b/Engine/Source/Runtime/Engine/Private/ComputeFramework/ComputeGraphComponent.cpp
@@ -59,13 +59,11 @@ void UComputeGraphComponent::SendRenderDynamicData_Concurrent()
 		return;
 	}
 
-	FComputeGraphProxy* ComputeGraphProxy = new FComputeGraphProxy();
-	ComputeGraphProxy->Initialize(ComputeGraph);
-
 	TArray<FComputeDataProviderRenderProxy*> ComputeDataProviderProxies;
 	for (UComputeDataProvider* DataProvider : DataProviders)
 	{
-		// Add null provider slots because we want to maintain consistent array indices.
+		// Be sure to add null provider slots because we want to maintain consistent array indices.
+		// Note that we expect GetRenderProxy() to return a pointer that we can own and call delete on.
 		FComputeDataProviderRenderProxy* ProviderProxy = DataProvider != nullptr ? DataProvider->GetRenderProxy() : nullptr;
 		ComputeDataProviderProxies.Add(ProviderProxy);
 	}
@@ -73,10 +71,17 @@ void UComputeGraphComponent::SendRenderDynamicData_Concurrent()
 	// Don't submit work if we don't have all of the expected bindings.
 	if (!ensure(ComputeGraph->ValidateBindings(ComputeDataProviderProxies)))
 	{
+		for (FComputeDataProviderRenderProxy* ProviderProxy : ComputeDataProviderProxies)
+		{
+			delete ProviderProxy;
+		}
 		// todo[CF]: We should have a default fallback for all cases where we can't submit work.
 		return;
 	}
 
+	FComputeGraphProxy* ComputeGraphProxy = new FComputeGraphProxy();
+	ComputeGraphProxy->Initialize(ComputeGraph);
+
 	ENQUEUE_RENDER_COMMAND(ComputeFrameworkEnqueueExecutionCommand)(
 		[ComputeGraphScheduler, ComputeGraphProxy, DataProviderProxies = MoveTemp(ComputeDataProviderProxies)](FRHICommandListImmediate& RHICmdList)
 		{
