Replace previous frame bone matrices with current frame data on skeletal mesh when world is paused, so velocity is cancelled and sk mesh isn't blurred

diff --git a/Engine/Source/Runtime/Engine/Private/GPUSkinCache.cpp b/Engine/Source/Runtime/Engine/Private/GPUSkinCache.cpp
--- a/Engine/Source/Runtime/Engine/Private/GPUSkinCache.cpp
+++ b/Engine/Source/Runtime/Engine/Private/GPUSkinCache.cpp
@@ -762,14 +762,16 @@ IMPLEMENT_SHADER_TYPE(template<>, TGPUSkinCacheCS<42>, TEXT("/Engine/Private/Gpu
 IMPLEMENT_SHADER_TYPE(template<>, TGPUSkinCacheCS<50>, TEXT("/Engine/Private/GpuSkinCacheComputeShader.usf"), TEXT("SkinCacheUpdateBatchCS"), SF_Compute);  // 16bit_1, BoneInfluenceType_0, SkinType_2, MultipleClothInfluences_1 
 IMPLEMENT_SHADER_TYPE(template<>, TGPUSkinCacheCS<54>, TEXT("/Engine/Private/GpuSkinCacheComputeShader.usf"), TEXT("SkinCacheUpdateBatchCS"), SF_Compute);  // 16bit_1, BoneInfluenceType_1, SkinType_2, MultipleClothInfluences_1 
 
-FGPUSkinCache::FGPUSkinCache(ERHIFeatureLevel::Type InFeatureLevel, bool bInRequiresMemoryLimit)
+FGPUSkinCache::FGPUSkinCache(ERHIFeatureLevel::Type InFeatureLevel, bool bInRequiresMemoryLimit, UWorld* InWorld)
 	: UsedMemoryInBytes(0)
 	, ExtraRequiredMemory(0)
 	, FlushCounter(0)
 	, bRequiresMemoryLimit(bInRequiresMemoryLimit)
 	, CurrentStagingBufferIndex(0)
 	, FeatureLevel(InFeatureLevel)
+	, World(InWorld)
 {
+	check(World);
 }
 
 FGPUSkinCache::~FGPUSkinCache()
@@ -1886,29 +1888,6 @@ void FGPUSkinCache::Release(FGPUSkinCacheEntry*& SkinCacheEntry)
 	}
 }
 
-void FGPUSkinCache::SetVertexStreams(FGPUSkinCacheEntry* Entry, int32 Section, FRHICommandList& RHICmdList,
-	FRHIVertexShader* ShaderRHI, const FGPUSkinPassthroughVertexFactory* VertexFactory,
-	uint32 BaseVertexIndex, FShaderResourceParameter GPUSkinCachePreviousPositionBuffer)
-{
-	INC_DWORD_STAT(STAT_GPUSkinCache_NumSetVertexStreams);
-	check(Entry);
-	check(Entry->IsSectionValid(Section));
-
-	FGPUSkinCacheEntry::FSectionDispatchData& DispatchData = Entry->DispatchData[Section];
-
-	//UE_LOG(LogSkinCache, Warning, TEXT("*** SetVertexStreams E %p All %p Sec %d(%p) LOD %d"), Entry, Entry->DispatchData[Section].Allocation, Section, Entry->DispatchData[Section].Section, Entry->LOD);
-	RHICmdList.SetStreamSource(VertexFactory->GetPositionStreamIndex(), DispatchData.GetPositionRWBuffer()->Buffer.Buffer, 0);
-	if (VertexFactory->GetTangentStreamIndex() > -1 && DispatchData.GetTangentRWBuffer())
-	{
-		RHICmdList.SetStreamSource(VertexFactory->GetTangentStreamIndex(), DispatchData.GetTangentRWBuffer()->Buffer.Buffer, 0);
-	}
-
-	if (ShaderRHI && GPUSkinCachePreviousPositionBuffer.IsBound())
-	{
-		RHICmdList.SetShaderResourceViewParameter(ShaderRHI, GPUSkinCachePreviousPositionBuffer.GetBaseIndex(), DispatchData.GetPreviousPositionRWBuffer()->Buffer.SRV);
-	}
-}
-
 void FGPUSkinCache::GetShaderBindings(
 	FGPUSkinCacheEntry* Entry, 
 	int32 Section,
@@ -1923,6 +1902,7 @@ void FGPUSkinCache::GetShaderBindings(
 	INC_DWORD_STAT(STAT_GPUSkinCache_NumSetVertexStreams);
 	check(Entry);
 	check(Entry->IsSectionValid(Section));
+	check(Entry->SkinCache);
 
 	FGPUSkinCacheEntry::FSectionDispatchData& DispatchData = Entry->DispatchData[Section];
 
@@ -1936,7 +1916,10 @@ void FGPUSkinCache::GetShaderBindings(
 	}
 
 	ShaderBindings.Add(GPUSkinCachePositionBuffer, DispatchData.GetPositionRWBuffer()->Buffer.SRV);
-	ShaderBindings.Add(GPUSkinCachePreviousPositionBuffer, DispatchData.GetPreviousPositionRWBuffer()->Buffer.SRV);
+
+	// If world is paused, use current frame bone matrices, so velocity is canceled and skeletal mesh isn't blurred from motion.
+	const bool bWorldPaused = Entry->SkinCache->IsWorldPaused();
+	ShaderBindings.Add(GPUSkinCachePreviousPositionBuffer, bWorldPaused ? DispatchData.GetPositionRWBuffer()->Buffer.SRV : DispatchData.GetPreviousPositionRWBuffer()->Buffer.SRV);
 }
 
 void FGPUSkinCache::PrepareUpdateSkinning(FGPUSkinCacheEntry* Entry, int32 Section, uint32 RevisionNumber, TArray<FSkinCacheRWBuffer*>* OverlappedUAVs)
diff --git a/Engine/Source/Runtime/Engine/Private/GPUSkinVertexFactory.cpp b/Engine/Source/Runtime/Engine/Private/GPUSkinVertexFactory.cpp
--- a/Engine/Source/Runtime/Engine/Private/GPUSkinVertexFactory.cpp
+++ b/Engine/Source/Runtime/Engine/Private/GPUSkinVertexFactory.cpp
@@ -698,7 +698,9 @@ public:
 				// todo: Maybe a check for PreviousData!=CurrentData would save some performance (when objects don't have velocty yet) but removing the bool also might save performance
 				bLocalPerBoneMotionBlur = true;
 
-				FRHIShaderResourceView* PreviousData = ShaderData.GetBoneBufferForReading(true).VertexBufferSRV;
+				// If world is paused, use current frame bone matrices, so velocity is canceled and skeletal mesh isn't blurred from motion.
+				const bool bWorldPaused = Scene->GetWorld() ? Scene->GetWorld()->IsPaused() : false;
+				FRHIShaderResourceView* PreviousData = ShaderData.GetBoneBufferForReading(bWorldPaused ? false : true).VertexBufferSRV;
 				ShaderBindings.Add(PreviousBoneMatrices, PreviousData);
 			}
 		}
diff --git a/Engine/Source/Runtime/Engine/Public/GPUSkinCache.h b/Engine/Source/Runtime/Engine/Public/GPUSkinCache.h
--- a/Engine/Source/Runtime/Engine/Public/GPUSkinCache.h
+++ b/Engine/Source/Runtime/Engine/Public/GPUSkinCache.h
@@ -153,7 +153,7 @@ public:
 	};
 
 	FGPUSkinCache() = delete;
-	ENGINE_API FGPUSkinCache(ERHIFeatureLevel::Type InFeatureLevel, bool bInRequiresMemoryLimit);
+	ENGINE_API FGPUSkinCache(ERHIFeatureLevel::Type InFeatureLevel, bool bInRequiresMemoryLimit, UWorld* InWorld);
 	ENGINE_API ~FGPUSkinCache();
 
 	ENGINE_API FCachedGeometry GetCachedGeometry(uint32 ComponentId) const;
@@ -179,10 +179,6 @@ public:
 		FGPUSkinCacheEntry*& InOutEntry
 		);
 
-	static void SetVertexStreams(FGPUSkinCacheEntry* Entry, int32 Section, FRHICommandList& RHICmdList,
-		class FRHIVertexShader* ShaderRHI, const FGPUSkinPassthroughVertexFactory* VertexFactory,
-		uint32 BaseVertexIndex, FShaderResourceParameter PreviousStreamBuffer);
-
 	static void GetShaderBindings(
 		FGPUSkinCacheEntry* Entry,
 		int32 Section,
@@ -453,6 +449,7 @@ public:
 	bool IsBatchingDispatch() const { return bShouldBatchDispatches; }
 
 	inline ERHIFeatureLevel::Type GetFeatureLevel() const { return FeatureLevel; }
+	inline bool IsWorldPaused() const { return World ? World->IsPaused() : false; }
 
 protected:
 	void MakeBufferTransitions(FRHICommandListImmediate& RHICmdList, TArray<FSkinCacheRWBuffer*>& Buffers, ERHIAccess ToState);
@@ -502,6 +499,7 @@ protected:
 	int32 CurrentStagingBufferIndex;
 
 	ERHIFeatureLevel::Type FeatureLevel;
+	UWorld* World;
 
 	static void CVarSinkFunction();
 	static FAutoConsoleVariableSink CVarSink;
diff --git a/Engine/Source/Runtime/Renderer/Private/RendererScene.cpp b/Engine/Source/Runtime/Renderer/Private/RendererScene.cpp
--- a/Engine/Source/Runtime/Renderer/Private/RendererScene.cpp
+++ b/Engine/Source/Runtime/Renderer/Private/RendererScene.cpp
@@ -1077,7 +1077,7 @@ FScene::FScene(UWorld* InWorld, bool bInRequiresHitProxies, bool bInIsEditorScen
 	if (IsGPUSkinCacheAvailable(GetFeatureLevelShaderPlatform(InFeatureLevel)))
 	{
 		const bool bRequiresMemoryLimit = !bInIsEditorScene;
-		GPUSkinCache = new FGPUSkinCache(InFeatureLevel, bRequiresMemoryLimit);
+		GPUSkinCache = new FGPUSkinCache(InFeatureLevel, bRequiresMemoryLimit, World);
 	}
 
 	if (ComputeFramework::IsEnabled(InFeatureLevel, GetFeatureLevelShaderPlatform(InFeatureLevel)))
