WorldPartition: replay fix scrubbing support

diff --git a/Engine/Source/Runtime/Engine/Private/WorldPartition/WorldPartitionReplay.cpp b/Engine/Source/Runtime/Engine/Private/WorldPartition/WorldPartitionReplay.cpp
--- a/Engine/Source/Runtime/Engine/Private/WorldPartition/WorldPartitionReplay.cpp
+++ b/Engine/Source/Runtime/Engine/Private/WorldPartition/WorldPartitionReplay.cpp
@@ -66,11 +66,19 @@ AWorldPartitionReplay::AWorldPartitionReplay(const FObjectInitializer& ObjectIni
 	bAlwaysRelevant = true;
 	bReplicates = true;
 	NetDriverName = NAME_DemoNetDriver;
+	bReplayRewindable = true;
 #if WITH_EDITORONLY_DATA
 	bListedInSceneOutliner = false;
 #endif
 }
 
+void AWorldPartitionReplay::RewindForReplay()
+{
+	Super::RewindForReplay();
+
+	ReplaySamples.Empty();
+}
+
 void AWorldPartitionReplay::Initialize(UWorld* World)
 {
 #if !UE_BUILD_SHIPPING
diff --git a/Engine/Source/Runtime/Engine/Public/WorldPartition/WorldPartitionReplay.h b/Engine/Source/Runtime/Engine/Public/WorldPartition/WorldPartitionReplay.h
--- a/Engine/Source/Runtime/Engine/Public/WorldPartition/WorldPartitionReplay.h
+++ b/Engine/Source/Runtime/Engine/Public/WorldPartition/WorldPartitionReplay.h
@@ -48,6 +48,7 @@ public:
 	static bool IsEnabled(UWorld* World);
 
 	virtual void BeginPlay() override;
+	virtual void RewindForReplay() override;
 	virtual void PreReplication(IRepChangedPropertyTracker& ChangedPropertyTracker) override;
 		
 	bool IsEnabled() const { return bEnabled; }
