Add "unused" command to listtextures which will only list textures that are not referenced by any primitive/material.

diff --git a/Engine/Source/Runtime/Engine/Private/UnrealEngine.cpp b/Engine/Source/Runtime/Engine/Private/UnrealEngine.cpp
--- a/Engine/Source/Runtime/Engine/Private/UnrealEngine.cpp
+++ b/Engine/Source/Runtime/Engine/Private/UnrealEngine.cpp
@@ -5403,13 +5403,15 @@ bool UEngine::HandleListTexturesCommand( const TCHAR* Cmd, FOutputDevice& Ar )
 	const bool bShouldOnlyListForced = FParse::Command(&Cmd, TEXT("FORCED")) && !bShouldOnlyListStreaming && !bShouldOnlyListNonStreaming;
 	const bool bShouldOnlyListUncompressed = FParse::Command(&Cmd, TEXT("UNCOMPRESSED")) && !bShouldOnlyListStreaming && !bShouldOnlyListNonStreaming && !bShouldOnlyListForced;
 	const bool bShouldOnlyListNonVT = FParse::Command(&Cmd, TEXT("NONVT")) && !bShouldOnlyListStreaming && !bShouldOnlyListNonStreaming && !bShouldOnlyListForced && !bShouldOnlyListUncompressed;
+	const bool bListUnused = FParse::Command(&Cmd, TEXT("UNUSED"));
 	const bool bAlphaSort = FParse::Param( Cmd, TEXT("ALPHASORT") );
 	const bool bCSV = FParse::Param( Cmd, TEXT("CSV") );
 	const bool bUnknownRefOnly = FParse::Param(Cmd, TEXT("unknownrefonly"));
 
-	Ar.Logf(TEXT("Listing %s%s textures."),
+	Ar.Logf(TEXT("Listing %s%s%s textures."),
 		bShouldOnlyListForced ? TEXT("forced") : bShouldOnlyListNonStreaming ? TEXT("non streaming") : bShouldOnlyListStreaming ? TEXT("streaming") : bShouldOnlyListUncompressed ? TEXT("uncompressed") : bShouldOnlyListNonVT ? TEXT("NONVT") : TEXT("all"),
-		bUnknownRefOnly ? TEXT(" unknown-ref") : TEXT(""));
+		bUnknownRefOnly ? TEXT(" unknown-ref") : TEXT(""),
+		bListUnused ? TEXT(" unused") : TEXT(""));
 
 	// Find out how many times a texture is referenced by primitive components.
 	TMap<UTexture2D*,int32> TextureToUsageMap;
@@ -5512,6 +5514,10 @@ bool UEngine::HandleListTexturesCommand( const TCHAR* Cmd, FOutputDevice& Ar )
 			Format				= TextureCube->GetPixelFormat();
 		}
 
+		if (bListUnused && UsageCount != 0)
+		{
+			continue;
+		}
 
 		if( (bShouldOnlyListStreaming && bIsStreamingTexture) ||	
 			(bShouldOnlyListNonStreaming && !bIsStreamingTexture) ||
