Fix for non-vector friendly buffer processing in ModulationDestinations.
#rnx

diff --git a/Engine/Source/Runtime/Engine/Private/SoundModulationDestination.cpp b/Engine/Source/Runtime/Engine/Private/SoundModulationDestination.cpp
--- a/Engine/Source/Runtime/Engine/Private/SoundModulationDestination.cpp
+++ b/Engine/Source/Runtime/Engine/Private/SoundModulationDestination.cpp
@@ -202,9 +202,9 @@ namespace Audio
 		// Fade from last target to new if output buffer is active
 		if (OutputBuffer.Num() > 0)
 		{
-			if (OutputBuffer.Num() % 4 == 0)
+			if (OutputBuffer.Num() % AUDIO_NUM_FLOATS_PER_VECTOR_REGISTER == 0)
 			{
-				if (LastTarget == ValueTarget)
+				if (FMath::IsNearlyEqual(LastTarget, ValueTarget))
 				{
 					BufferSetToConstantInplace(OutputBuffer, ValueTarget);
 				}
@@ -216,12 +216,19 @@ namespace Audio
 			}
 			else
 			{
-				float Gain = LastTarget;
-				const float DeltaValue = (ValueTarget - LastTarget) / OutputBuffer.Num();
-				for (int32 i = 0; i < OutputBuffer.Num(); ++i)
+				if (FMath::IsNearlyEqual(LastTarget, ValueTarget))
 				{
-					OutputBuffer[i] *= Gain;
-					Gain += DeltaValue;
+					OutputBuffer.Init(ValueTarget, InNumSamples);
+				}
+				else
+				{
+					float SampleValue = LastTarget;
+					const float DeltaValue = (ValueTarget - LastTarget) / OutputBuffer.Num();
+					for (int32 i = 0; i < OutputBuffer.Num(); ++i)
+					{
+						OutputBuffer[i] = SampleValue;
+						SampleValue += DeltaValue;
+					}
 				}
 			}
 		}
