Fix a crash that can occur on shutdown when cooking

diff --git a/Engine/Source/Runtime/Engine/Private/LightWeightInstanceSubsystem.cpp b/Engine/Source/Runtime/Engine/Private/LightWeightInstanceSubsystem.cpp
--- a/Engine/Source/Runtime/Engine/Private/LightWeightInstanceSubsystem.cpp
+++ b/Engine/Source/Runtime/Engine/Private/LightWeightInstanceSubsystem.cpp
@@ -35,8 +35,11 @@ FLightWeightInstanceSubsystem::FLightWeightInstanceSubsystem()
 FLightWeightInstanceSubsystem::~FLightWeightInstanceSubsystem()
 {
 #if WITH_EDITOR
-	GEngine->OnLevelActorAdded().Remove(OnLevelActorAddedHandle);
-	GEngine->OnLevelActorDeleted().Remove(OnLevelActorDeletedHandle);
+	if (GEngine)
+	{
+		GEngine->OnLevelActorAdded().Remove(OnLevelActorAddedHandle);
+		GEngine->OnLevelActorDeleted().Remove(OnLevelActorDeletedHandle);
+	}
 #endif // WITH_EDITOR
 }
 
