When in game, ignore skeletal mesh updates sent in the middle of a frame (Tick()) and only process those from SendAllEndOfFrameUpdates

diff --git a/Engine/Source/Runtime/Engine/Private/SkeletalRenderGPUSkin.cpp b/Engine/Source/Runtime/Engine/Private/SkeletalRenderGPUSkin.cpp
--- a/Engine/Source/Runtime/Engine/Private/SkeletalRenderGPUSkin.cpp
+++ b/Engine/Source/Runtime/Engine/Private/SkeletalRenderGPUSkin.cpp
@@ -322,10 +322,10 @@ void FSkeletalMeshObjectGPUSkin::Update(int32 LODIndex,USkinnedMeshComponent* In
 	// queue a call to update this data
 	FSkeletalMeshObjectGPUSkin* MeshObject = this;
 	ENQUEUE_RENDER_COMMAND(SkelMeshObjectUpdateDataCommand)(
-		[MeshObject, FrameNumberToPrepare, RevisionNumber, NewDynamicData, GPUSkinCache](FRHICommandListImmediate& RHICmdList)
+		[MeshObject, FrameNumberToPrepare, RevisionNumber, NewDynamicData, GPUSkinCache, bIsGameWorld = (InMeshComponent->GetWorld() != nullptr && InMeshComponent->GetWorld()->IsGameWorld())](FRHICommandListImmediate& RHICmdList)
 		{
 			FScopeCycleCounter Context(MeshObject->GetStatId());
-			MeshObject->UpdateDynamicData_RenderThread(GPUSkinCache, RHICmdList, NewDynamicData, nullptr, FrameNumberToPrepare, RevisionNumber);
+			MeshObject->UpdateDynamicData_RenderThread(GPUSkinCache, RHICmdList, NewDynamicData, nullptr, bIsGameWorld, FrameNumberToPrepare, RevisionNumber);
 		}
 	);
 }
@@ -379,7 +379,7 @@ static TAutoConsoleVariable<int32> CVarDeferSkeletalDynamicDataUpdateUntilGDME(
 	0,
 	TEXT("If > 0, then do skeletal mesh dynamic data updates will be deferred until GDME. Experimental option."));
 
-void FSkeletalMeshObjectGPUSkin::UpdateDynamicData_RenderThread(FGPUSkinCache* GPUSkinCache, FRHICommandListImmediate& RHICmdList, FDynamicSkelMeshObjectDataGPUSkin* InDynamicData, FSceneInterface* Scene, uint32 FrameNumberToPrepare, uint32 RevisionNumber)
+void FSkeletalMeshObjectGPUSkin::UpdateDynamicData_RenderThread(FGPUSkinCache* GPUSkinCache, FRHICommandListImmediate& RHICmdList, FDynamicSkelMeshObjectDataGPUSkin* InDynamicData, FSceneInterface* Scene, bool bIsGameWorld, uint32 FrameNumberToPrepare, uint32 RevisionNumber)
 {
 	SCOPE_CYCLE_COUNTER(STAT_GPUSkinUpdateRTTime);
 	check(InDynamicData != nullptr);
@@ -436,13 +436,18 @@ void FSkeletalMeshObjectGPUSkin::UpdateDynamicData_RenderThread(FGPUSkinCache* G
 	#if RHI_RAYTRACING
 		if (FGPUSkinCache::IsGPUSkinCacheRayTracingSupported() && GPUSkinCache && SkeletalMeshRenderData->bSupportRayTracing)
 		{
-			ProcessUpdatedDynamicData(EGPUSkinCacheEntryMode::RayTracing, GPUSkinCache, RHICmdList, FrameNumberToPrepare, RevisionNumber, bMorphNeedsUpdate, DynamicData->RayTracingLODIndex);
+			// When in game, updates not sent in UWorld::SendAllEndOfFrameUpdates() (implied by IsBatchingDispatch()) are ignored.
+			if (!bIsGameWorld || GPUSkinCache->IsBatchingDispatch())
+			{
+				ProcessUpdatedDynamicData(EGPUSkinCacheEntryMode::RayTracing, GPUSkinCache, RHICmdList, FrameNumberToPrepare, RevisionNumber, bMorphNeedsUpdate, DynamicData->RayTracingLODIndex);
+			}
 		}
 	#endif
 	}
 
 #if RHI_RAYTRACING
-	if (FGPUSkinCache::IsGPUSkinCacheRayTracingSupported() && GPUSkinCache && SkeletalMeshRenderData->bSupportRayTracing && !GPUSkinCache->IsBatchingDispatch())
+	// Immediate update path. Only used in the editor as the world can be not ticking. When in game, updates not sent in UWorld::SendAllEndOfFrameUpdates() are ignored.
+	if (!bIsGameWorld && FGPUSkinCache::IsGPUSkinCacheRayTracingSupported() && GPUSkinCache && SkeletalMeshRenderData->bSupportRayTracing && !GPUSkinCache->IsBatchingDispatch())
 	{
 		if (SkinCacheEntryForRayTracing)
 		{
diff --git a/Engine/Source/Runtime/Engine/Private/SkeletalRenderGPUSkin.h b/Engine/Source/Runtime/Engine/Private/SkeletalRenderGPUSkin.h
--- a/Engine/Source/Runtime/Engine/Private/SkeletalRenderGPUSkin.h
+++ b/Engine/Source/Runtime/Engine/Private/SkeletalRenderGPUSkin.h
@@ -374,7 +374,7 @@ public:
 	virtual void InitResources(USkinnedMeshComponent* InMeshComponent) override;
 	virtual void ReleaseResources() override;
 	virtual void Update(int32 LODIndex,USkinnedMeshComponent* InMeshComponent,const TArray<FActiveMorphTarget>& ActiveMorphTargets, const TArray<float>& MorphTargetWeights, EPreviousBoneTransformUpdateMode PreviousBoneTransformUpdateMode) override;
-	void UpdateDynamicData_RenderThread(FGPUSkinCache* GPUSkinCache, FRHICommandListImmediate& RHICmdList, FDynamicSkelMeshObjectDataGPUSkin* InDynamicData, FSceneInterface* Scene, uint32 FrameNumberToPrepare, uint32 RevisionNumber);
+	void UpdateDynamicData_RenderThread(FGPUSkinCache* GPUSkinCache, FRHICommandListImmediate& RHICmdList, FDynamicSkelMeshObjectDataGPUSkin* InDynamicData, FSceneInterface* Scene, bool bIsGameWorld, uint32 FrameNumberToPrepare, uint32 RevisionNumber);
 	virtual void PreGDMECallback(FGPUSkinCache* GPUSkinCache, uint32 FrameNumber) override;
 	virtual const FVertexFactory* GetSkinVertexFactory(const FSceneView* View, int32 LODIndex,int32 ChunkIdx, ESkinVertexFactoryMode VFMode = ESkinVertexFactoryMode::Default) const override;
 	virtual void CacheVertices(int32 LODIndex, bool bForce) const override {}
