- Support case where texture was modified and saved after the build of per-actor texture streaming.
- Put back r.Streaming.UseTextureStreamingBuiltData default value to true.
- Moved per-actor texture streaming build in AHLODActor::BuildHLOD instead of PreSave to allow updating per-actor texture streaming data even if HLOD didn't change.
- Added logic to detect if newly built per-actor texture streaming built data has changed and dirty actor's package if it's the case.

diff --git a/Engine/Source/Runtime/Engine/Classes/Components/PrimitiveComponent.h b/Engine/Source/Runtime/Engine/Classes/Components/PrimitiveComponent.h
--- a/Engine/Source/Runtime/Engine/Classes/Components/PrimitiveComponent.h
+++ b/Engine/Source/Runtime/Engine/Classes/Components/PrimitiveComponent.h
@@ -1767,6 +1767,9 @@ public:
 	 *	
 	 */
 	virtual bool RemapActorTextureStreamingBuiltDataToLevel(const class UActorTextureStreamingBuildDataComponent* InActorTextureBuildData) { return false; }
+
+	/** Computes a hash of component's texture streaming built data. */
+	virtual uint32 ComputeHashTextureStreamingBuiltData() const { return 0; }
 #endif // WITH_EDITOR
 
 	/**
diff --git a/Engine/Source/Runtime/Engine/Classes/Components/StaticMeshComponent.h b/Engine/Source/Runtime/Engine/Classes/Components/StaticMeshComponent.h
--- a/Engine/Source/Runtime/Engine/Classes/Components/StaticMeshComponent.h
+++ b/Engine/Source/Runtime/Engine/Classes/Components/StaticMeshComponent.h
@@ -525,6 +525,7 @@ public:
 	virtual void GetStreamingRenderAssetInfo(FStreamingTextureLevelContext& LevelContext, TArray<FStreamingRenderAssetPrimitiveInfo>& OutStreamingRenderAssets) const override;
 #if WITH_EDITOR
 	virtual bool RemapActorTextureStreamingBuiltDataToLevel(const class UActorTextureStreamingBuildDataComponent* InActorTextureBuildData) override;
+	virtual uint32 ComputeHashTextureStreamingBuiltData() const override;
 #endif 
 
 	virtual class UBodySetup* GetBodySetup() override;
diff --git a/Engine/Source/Runtime/Engine/Classes/Engine/TextureStreamingTypes.h b/Engine/Source/Runtime/Engine/Classes/Engine/TextureStreamingTypes.h
--- a/Engine/Source/Runtime/Engine/Classes/Engine/TextureStreamingTypes.h
+++ b/Engine/Source/Runtime/Engine/Classes/Engine/TextureStreamingTypes.h
@@ -159,6 +159,9 @@ struct FStreamingTextureBuildInfo
 	 *	@param	Info						[in]		The unpacked params.
 	 */
 	ENGINE_API void PackFrom(ITextureStreamingContainer* TextureStreamingContainer, const FBoxSphereBounds& RefBounds, const FStreamingRenderAssetPrimitiveInfo& Info);
+
+	/** Returns hash of content. */
+	ENGINE_API uint32 ComputeHash() const;
 #endif
 };
 
diff --git a/Engine/Source/Runtime/Engine/Private/Components/StaticMeshComponent.cpp b/Engine/Source/Runtime/Engine/Private/Components/StaticMeshComponent.cpp
--- a/Engine/Source/Runtime/Engine/Private/Components/StaticMeshComponent.cpp
+++ b/Engine/Source/Runtime/Engine/Private/Components/StaticMeshComponent.cpp
@@ -29,6 +29,7 @@
 #include "HierarchicalLODUtilitiesModule.h"
 #include "Rendering/StaticLightingSystemInterface.h"
 #include "Streaming/ActorTextureStreamingBuildDataComponent.h"
+#include "Misc/Crc.h"
 #endif
 #include "LightMap.h"
 #include "ShadowMap.h"
@@ -827,6 +828,16 @@ bool UStaticMeshComponent::RemapActorTextureStreamingBuiltDataToLevel(const UAct
 	}
 	return true;
 }
+
+uint32 UStaticMeshComponent::ComputeHashTextureStreamingBuiltData() const
+{
+	uint32 Hash = 0;
+	for (const FStreamingTextureBuildInfo& Data : StreamingTextureData)
+	{
+		Hash = FCrc::TypeCrc32(Data.ComputeHash(), Hash);
+	}
+	return Hash;
+}
 #endif
 
 bool UStaticMeshComponent::BuildTextureStreamingDataImpl(ETextureStreamingBuildType BuildType, EMaterialQualityLevel::Type QualityLevel, ERHIFeatureLevel::Type FeatureLevel, TSet<FGuid>& DependentResources, bool& bOutSupportsBuildTextureStreamingData)
diff --git a/Engine/Source/Runtime/Engine/Private/Streaming/ActorTextureStreamingBuildDataComponent.cpp b/Engine/Source/Runtime/Engine/Private/Streaming/ActorTextureStreamingBuildDataComponent.cpp
--- a/Engine/Source/Runtime/Engine/Private/Streaming/ActorTextureStreamingBuildDataComponent.cpp
+++ b/Engine/Source/Runtime/Engine/Private/Streaming/ActorTextureStreamingBuildDataComponent.cpp
@@ -10,6 +10,7 @@
 
 #include "Engine/TextureStreamingTypes.h"
 #include "Engine/Texture.h"
+#include "Misc/Crc.h"
 
 void UActorTextureStreamingBuildDataComponent::InitializeTextureStreamingContainer(uint32 InPackedTextureStreamingQualityLevelFeatureLevel)
 {
@@ -48,4 +49,19 @@ bool UActorTextureStreamingBuildDataComponent::GetStreamableTexture(uint16 InTex
 	return false;
 }
 
+uint32 UActorTextureStreamingBuildDataComponent::ComputeHash() const
+{
+	uint32 Hash = PackedTextureStreamingQualityLevelFeatureLevel;
+	for (const FStreamableTexture& StreamableTexture : StreamableTextures)
+	{
+		Hash = FCrc::TypeCrc32(StreamableTexture.ComputeHash(), Hash);
+	}
+	return Hash;
+}
+
+uint32 FStreamableTexture::ComputeHash() const
+{
+	return FCrc::TypeCrc32(GetTypeHash(Name), GetTypeHash(Guid));
+}
+
 #endif
\ No newline at end of file
diff --git a/Engine/Source/Runtime/Engine/Private/Streaming/TextureStreamingBuild.cpp b/Engine/Source/Runtime/Engine/Private/Streaming/TextureStreamingBuild.cpp
--- a/Engine/Source/Runtime/Engine/Private/Streaming/TextureStreamingBuild.cpp
+++ b/Engine/Source/Runtime/Engine/Private/Streaming/TextureStreamingBuild.cpp
@@ -23,6 +23,10 @@ TextureStreamingBuild.cpp : Contains definitions to build texture streaming data
 #include "Logging/MessageLog.h"
 #include "Misc/UObjectToken.h"
 #include "UnrealEngine.h"
+#if WITH_EDITOR
+#include "Misc/Crc.h"
+#endif
+
 
 DEFINE_LOG_CATEGORY(TextureStreamingBuild);
 #define LOCTEXT_NAMESPACE "TextureStreamingBuild"
@@ -36,13 +40,41 @@ ENGINE_API uint32 GetPackedTextureStreamingQualityLevelFeatureLevel(EMaterialQua
 	return ((uint32)InQualityLevel << 16) | ((uint32)InFeatureLevel & 0xFFFF);
 }
 
+static uint32 ComputeHashTextureStreamingDataForActor(AActor* InActor)
+{
+	uint32 Hash = 0;
+	if (UActorTextureStreamingBuildDataComponent* BuiltDataComponent = InActor->FindComponentByClass<UActorTextureStreamingBuildDataComponent>())
+	{
+		Hash = FCrc::TypeCrc32(BuiltDataComponent->ComputeHash(), Hash);
+	}
+	
+	TInlineComponentArray<UPrimitiveComponent*> Primitives;
+	InActor->GetComponents<UPrimitiveComponent>(Primitives);
+
+	TArray<uint32> PrimitiveHashes;
+	for (UPrimitiveComponent* Primitive : Primitives)
+	{
+		if (Primitive)
+		{
+			PrimitiveHashes.Add(FCrc::TypeCrc32(Primitive->ComputeHashTextureStreamingBuiltData()));
+		}
+	}
+	PrimitiveHashes.Sort();
+	for (uint32 PrimitiveHash : PrimitiveHashes)
+	{
+		Hash = FCrc::TypeCrc32(PrimitiveHash, Hash);
+	}
+	return Hash;
+}
+
 ENGINE_API void BuildActorTextureStreamingData(AActor* InActor, EMaterialQualityLevel::Type InQualityLevel, ERHIFeatureLevel::Type InFeatureLevel)
 {
 	if (!InActor || InActor->IsTemplate())
 	{
 		return;
 	}
-
+	uint32 OldHash = ComputeHashTextureStreamingDataForActor(InActor);
+	
 	TSet<FGuid> DummyResourceGuids;
 	UActorTextureStreamingBuildDataComponent* BuiltDataComponent = InActor->FindComponentByClass<UActorTextureStreamingBuildDataComponent>();
 	if (!BuiltDataComponent)
@@ -51,7 +83,7 @@ ENGINE_API void BuildActorTextureStreamingData(AActor* InActor, EMaterialQuality
 		InActor->AddInstanceComponent(BuiltDataComponent);
 	}
 	BuiltDataComponent->InitializeTextureStreamingContainer(GetPackedTextureStreamingQualityLevelFeatureLevel(InQualityLevel, InFeatureLevel));
-	
+
 	TInlineComponentArray<UPrimitiveComponent*> Primitives;
 	InActor->GetComponents<UPrimitiveComponent>(Primitives);
 	for (UPrimitiveComponent* Primitive : Primitives)
@@ -61,6 +93,13 @@ ENGINE_API void BuildActorTextureStreamingData(AActor* InActor, EMaterialQuality
 			Primitive->BuildTextureStreamingData(TSB_ActorBuild, InQualityLevel, InFeatureLevel, DummyResourceGuids);
 		}
 	}
+
+	// If actor's texture streaming built data has changed, dirty its package
+	uint32 NewHash = ComputeHashTextureStreamingDataForActor(InActor);
+	if (NewHash != OldHash)
+	{
+		InActor->GetPackage()->MarkPackageDirty();
+	}
 }
 
 ENGINE_API bool BuildLevelTextureStreamingComponentDataFromActors(ULevel* InLevel)
@@ -356,9 +395,14 @@ void FStreamingTextureBuildInfo::PackFrom(ITextureStreamingContainer* TextureStr
 	TextureLevelIndex = TextureStreamingContainer->RegisterStreamableTexture(Texture);
 	TexelFactor = Info.TexelFactor;
 }
+
+uint32 FStreamingTextureBuildInfo::ComputeHash() const
+{
+	return FCrc::TypeCrc32(PackedRelativeBox, FCrc::TypeCrc32(TextureLevelIndex, FCrc::TypeCrc32(TexelFactor, 0)));
+}
 #endif
 
-bool FStreamingTextureLevelContext::UseTextureStreamingBuiltData = false;
+bool FStreamingTextureLevelContext::UseTextureStreamingBuiltData = true;
 FAutoConsoleCommand FStreamingTextureLevelContext::UseTextureStreamingBuiltDataCommand(
 	TEXT("r.Streaming.UseTextureStreamingBuiltData"),
 	TEXT("Turn on/off usage of texture streaming built data (0 to turn off)."),
@@ -429,7 +473,6 @@ void FStreamingTextureLevelContext::UpdateContext(EMaterialQualityLevel::Type In
 					{
 						LevelStreamingTextures[TextureLevelIndex] = Texture;
 						GetBuildDataIndexRef(Texture, /*ForceUpdate*/ true);
-						check(Texture->LevelIndex == TextureLevelIndex);
 					}
 				}
 			}
@@ -576,7 +619,6 @@ void FStreamingTextureLevelContext::ProcessMaterial(const FBoxSphereBounds& Comp
 				{
 					if (UTexture* Texture = LevelStreamingTextures[BuildInfo.TextureLevelIndex])
 					{
-						check(Texture->IsStreamable());
 						check(Texture->LevelIndex != INDEX_NONE);
 						Textures.Add(Texture);
 					}
diff --git a/Engine/Source/Runtime/Engine/Private/WorldPartition/HLOD/HLODActor.cpp b/Engine/Source/Runtime/Engine/Private/WorldPartition/HLOD/HLODActor.cpp
--- a/Engine/Source/Runtime/Engine/Private/WorldPartition/HLOD/HLODActor.cpp
+++ b/Engine/Source/Runtime/Engine/Private/WorldPartition/HLOD/HLODActor.cpp
@@ -71,22 +71,6 @@ void AWorldPartitionHLOD::Serialize(FArchive& Ar)
 void AWorldPartitionHLOD::RerunConstructionScripts()
 {}
 
-void AWorldPartitionHLOD::PreSave(FObjectPreSaveContext ObjectSaveContext)
-{
-	Super::PreSave(ObjectSaveContext);
-
-#if WITH_EDITOR
-	// When generating WorldPartition HLODs, we have the renderer initialized.
-	// Take advantage of this and generate texture streaming built data (local to the actor).
-	// This built data will be used by the cooking (it will convert it to level texture streaming built data).
-	if (!ObjectSaveContext.IsCooking())
-	{
-		// Use same quality level and feature level as FEditorBuildUtils::EditorBuildTextureStreaming
-		BuildActorTextureStreamingData(this, EMaterialQualityLevel::High, GMaxRHIFeatureLevel);
-	}
-#endif
-}
-
 #if WITH_EDITOR
 
 bool AWorldPartitionHLOD::IsHiddenEd() const
@@ -196,10 +180,12 @@ void AWorldPartitionHLOD::BuildHLOD(bool bForceBuild)
 
 		HLODHash = WPHLODUtilities->BuildHLOD(this);
 	}
-	else
-	{
 
-	}
+	// When generating WorldPartition HLODs, we have the renderer initialized.
+	// Take advantage of this and generate texture streaming built data (local to the actor).
+	// This built data will be used by the cooking (it will convert it to level texture streaming built data).
+	// Use same quality level and feature level as FEditorBuildUtils::EditorBuildTextureStreaming
+	BuildActorTextureStreamingData(this, EMaterialQualityLevel::High, GMaxRHIFeatureLevel);
 }
 
 #endif // WITH_EDITOR
diff --git a/Engine/Source/Runtime/Engine/Public/Streaming/ActorTextureStreamingBuildDataComponent.h b/Engine/Source/Runtime/Engine/Public/Streaming/ActorTextureStreamingBuildDataComponent.h
--- a/Engine/Source/Runtime/Engine/Public/Streaming/ActorTextureStreamingBuildDataComponent.h
+++ b/Engine/Source/Runtime/Engine/Public/Streaming/ActorTextureStreamingBuildDataComponent.h
@@ -32,6 +32,10 @@ struct ENGINE_API FStreamableTexture
 	UPROPERTY()
 	FGuid Guid;
 #endif
+
+#if WITH_EDITOR
+	uint32 ComputeHash() const;
+#endif
 };
 
 UCLASS()
@@ -53,6 +57,7 @@ public:
 	//~End ITextureStreamingContainer Interface.
 
 	uint32 GetPackedTextureStreamingQualityLevelFeatureLevel() const { return PackedTextureStreamingQualityLevelFeatureLevel; }
+	uint32 ComputeHash() const;
 #endif
 
 private:
diff --git a/Engine/Source/Runtime/Engine/Public/WorldPartition/HLOD/HLODActor.h b/Engine/Source/Runtime/Engine/Public/WorldPartition/HLOD/HLODActor.h
--- a/Engine/Source/Runtime/Engine/Public/WorldPartition/HLOD/HLODActor.h
+++ b/Engine/Source/Runtime/Engine/Public/WorldPartition/HLOD/HLODActor.h
@@ -64,7 +64,6 @@ protected:
 	//~ Begin UObject Interface.
 	virtual void Serialize(FArchive& Ar) override;
 	virtual void RerunConstructionScripts() override;
-	virtual void PreSave(FObjectPreSaveContext ObjectSaveContext) override;
 	//~ End UObject Interface.
 
 	//~ Begin AActor Interface.
