WorldPartition: support PIE WorldTravel to WP map destination (used by demoplay, open Map command)

diff --git a/Engine/Source/Runtime/Engine/Private/WorldPartition/WorldPartition.cpp b/Engine/Source/Runtime/Engine/Private/WorldPartition/WorldPartition.cpp
--- a/Engine/Source/Runtime/Engine/Private/WorldPartition/WorldPartition.cpp
+++ b/Engine/Source/Runtime/Engine/Private/WorldPartition/WorldPartition.cpp
@@ -392,7 +392,10 @@ void UWorldPartition::Initialize(UWorld* InWorld, const FTransform& InTransform)
 		RuntimeHash->SetDefaultValues();
 	}
 
-	if (bEditorOnly || IsRunningGame())
+	// Did we travel into a WP map in PIE (null StreamingPolicy means GenerateStreaming wasn't called)
+	const bool bPIEWorldTravel = World->WorldType == EWorldType::PIE && !StreamingPolicy;
+
+	if (bEditorOnly || IsRunningGame() || bPIEWorldTravel)
 	{
 		UPackage* LevelPackage = OuterWorld->PersistentLevel->GetOutermost();
 		const FName PackageName = LevelPackage->GetLoadedPath().GetPackageFName();
@@ -402,7 +405,7 @@ void UWorldPartition::Initialize(UWorld* InWorld, const FTransform& InTransform)
 		}
 		check(bContainerInitialized);
 
-		bool bIsInstanced = (bEditorOnly && !IsRunningCommandlet()) ? OuterWorld->PersistentLevel->IsInstancedLevel() : false;
+		bool bIsInstanced = (bEditorOnly && !IsRunningCommandlet()) ? OuterWorld->PersistentLevel->IsInstancedLevel() : bPIEWorldTravel;
 
 		FString ReplaceFrom;
 		FString ReplaceTo;
@@ -415,7 +418,9 @@ void UWorldPartition::Initialize(UWorld* InWorld, const FTransform& InTransform)
 			const FString DestWorldName = FPaths::GetBaseFilename(LevelPackage->GetFName().ToString());
 
 			ReplaceFrom = SourceWorldName + TEXT(".") + SourceWorldName;
-			ReplaceTo = DestWorldName + TEXT(".") + DestWorldName;
+			
+			// In PIE UWorld::PostLoad will not rename the world to match its package name so use SourceWorldName instead
+			ReplaceTo = DestWorldName + TEXT(".") + (bPIEWorldTravel ? SourceWorldName : DestWorldName);
 		}
 
 		{
@@ -505,7 +510,7 @@ void UWorldPartition::Initialize(UWorld* InWorld, const FTransform& InTransform)
 #if WITH_EDITOR
 	if (!bEditorOnly)
 	{
-		if (IsRunningGame())
+		if (IsRunningGame() || bPIEWorldTravel)
 		{
 			OnBeginPlay(EWorldPartitionStreamingMode::EditorStandalone);
 		}
diff --git a/Engine/Source/Runtime/Engine/Private/WorldPartition/WorldPartitionRuntimeHash.cpp b/Engine/Source/Runtime/Engine/Private/WorldPartition/WorldPartitionRuntimeHash.cpp
--- a/Engine/Source/Runtime/Engine/Private/WorldPartition/WorldPartitionRuntimeHash.cpp
+++ b/Engine/Source/Runtime/Engine/Private/WorldPartition/WorldPartitionRuntimeHash.cpp
@@ -64,7 +64,8 @@ void UWorldPartitionRuntimeHash::OnEndPlay()
 
 void UWorldPartitionRuntimeHash::ForceExternalActorLevelReference(bool bForceExternalActorLevelReferenceForPIE)
 {
-	if (!IsRunningGame())
+	// Do this only on non game worlds prior to PIE so that always loaded actors get duplicated with the world
+	if (!GetWorld()->IsGameWorld())
 	{
 		for (const FAlwaysLoadedActorForPIE& AlwaysLoadedActor : AlwaysLoadedActorsForPIE)
 		{
