Added blueprint callable function on physics asset to interact with constraints

diff --git a/Engine/Source/Runtime/Engine/Classes/PhysicsEngine/PhysicsAsset.h b/Engine/Source/Runtime/Engine/Classes/PhysicsEngine/PhysicsAsset.h
--- a/Engine/Source/Runtime/Engine/Classes/PhysicsEngine/PhysicsAsset.h
+++ b/Engine/Source/Runtime/Engine/Classes/PhysicsEngine/PhysicsAsset.h
@@ -198,6 +198,7 @@ public:
 	ENGINE_API int32		FindControllingBodyIndex(class USkeletalMesh* skelMesh, int32 BoneIndex);
 	ENGINE_API int32		FindParentBodyIndex(class USkeletalMesh * skelMesh, int32 StartBoneIndex) const;
 	ENGINE_API int32		FindConstraintIndex(FName ConstraintName);
+	ENGINE_API int32		FindConstraintIndex(FName Bone1Name, FName Bone2Name);
 	FName					FindConstraintBoneName(int32 ConstraintIndex);
 	ENGINE_API int32		FindMirroredBone(class USkeletalMesh* skelMesh, int32 BoneIndex);
 
@@ -280,6 +281,16 @@ public:
 	virtual void SetPreviewMesh(USkeletalMesh* PreviewMesh, bool bMarkAsDirty = true);
 	virtual USkeletalMesh* GetPreviewMesh() const;
 
+#if WITH_EDITOR
+	UFUNCTION(BlueprintCallable, Category = "Physics|Constraints")
+	FConstraintInstanceAccessor GetConstraintByName(FName ConstraintName);
+
+	UFUNCTION(BlueprintCallable, Category = "Physics|Constraints")
+	FConstraintInstanceAccessor GetConstraintByBoneNames(FName Bone1Name, FName Bone2Name);
+
+	FConstraintInstanceAccessor GetConstraintInstanceAccessorByIndex(int32 Index);
+	FConstraintInstance* GetConstraintInstanceByIndex(uint32 Index);
+#endif
 
 private:
 
diff --git a/Engine/Source/Runtime/Engine/Private/PhysicsEngine/ConstraintInstance.cpp b/Engine/Source/Runtime/Engine/Private/PhysicsEngine/ConstraintInstance.cpp
--- a/Engine/Source/Runtime/Engine/Private/PhysicsEngine/ConstraintInstance.cpp
+++ b/Engine/Source/Runtime/Engine/Private/PhysicsEngine/ConstraintInstance.cpp
@@ -8,6 +8,7 @@
 #include "Components/PrimitiveComponent.h"
 #include "Components/SkeletalMeshComponent.h"
 #include "PhysicsEngine/PhysicsConstraintComponent.h"
+#include "PhysicsEngine/PhysicsAsset.h"
 #include "PhysicsPublic.h"
 #include "Physics/PhysicsInterfaceTypes.h"
 
@@ -1202,6 +1203,12 @@ FConstraintInstance* FConstraintInstanceAccessor::Get() const
 		{
 			return &(PhysicsConstraintComponent->ConstraintInstance);
 		}
+#if WITH_EDITOR
+		if (UPhysicsAsset* PhysicsAsset = Cast<UPhysicsAsset>(Owner.Get()))
+		{
+			return PhysicsAsset->GetConstraintInstanceByIndex(Index);
+		}
+#endif
 	}
 	return nullptr;
 }
diff --git a/Engine/Source/Runtime/Engine/Private/PhysicsEngine/PhysicsAsset.cpp b/Engine/Source/Runtime/Engine/Private/PhysicsEngine/PhysicsAsset.cpp
--- a/Engine/Source/Runtime/Engine/Private/PhysicsEngine/PhysicsAsset.cpp
+++ b/Engine/Source/Runtime/Engine/Private/PhysicsEngine/PhysicsAsset.cpp
@@ -504,6 +504,20 @@ int32 UPhysicsAsset::FindConstraintIndex(FName ConstraintName)
 	return INDEX_NONE;
 }
 
+int32 UPhysicsAsset::FindConstraintIndex(FName Bone1Name, FName Bone2Name)
+{
+	for (int32 i = 0; i < ConstraintSetup.Num(); i++)
+	{
+		if (ConstraintSetup[i]->DefaultInstance.ConstraintBone1 == Bone1Name &&
+			ConstraintSetup[i]->DefaultInstance.ConstraintBone2 == Bone2Name)
+		{
+			return i;
+		}
+	}
+
+	return INDEX_NONE;
+}
+
 FName UPhysicsAsset::FindConstraintBoneName(int32 ConstraintIndex)
 {
 	if ( (ConstraintIndex < 0) || (ConstraintIndex >= ConstraintSetup.Num()) )
@@ -843,6 +857,46 @@ USkeletalMesh* UPhysicsAsset::GetPreviewMesh() const
 #endif
 }
 
+#if WITH_EDITOR
+FConstraintInstanceAccessor UPhysicsAsset::GetConstraintInstanceAccessorByIndex(int32 ConstraintIndex)
+{
+	if (ConstraintIndex == INDEX_NONE || ConstraintIndex >= ConstraintSetup.Num())
+	{
+		return FConstraintInstanceAccessor();
+	}
+
+	if (ConstraintSetup[ConstraintIndex])
+	{
+		return FConstraintInstanceAccessor(this, ConstraintIndex);
+	}
+
+	return FConstraintInstanceAccessor();
+}
+
+FConstraintInstanceAccessor UPhysicsAsset::GetConstraintByName(FName ConstraintName)
+{
+	return GetConstraintInstanceAccessorByIndex(FindConstraintIndex(ConstraintName));
+}
+
+FConstraintInstanceAccessor UPhysicsAsset::GetConstraintByBoneNames(FName Bone1Name, FName Bone2Name)
+{
+	return GetConstraintInstanceAccessorByIndex(FindConstraintIndex(Bone1Name, Bone2Name));
+}
+
+FConstraintInstance* UPhysicsAsset::GetConstraintInstanceByIndex(uint32 Index)
+{
+	if (Index < (uint32)ConstraintSetup.Num())
+	{
+		if (ConstraintSetup[Index])
+		{
+			return &ConstraintSetup[Index]->DefaultInstance;
+		}
+	}
+
+	return nullptr;
+}
+#endif
+
 void UPhysicsAsset::SetPreviewMesh(USkeletalMesh* PreviewMesh, bool bMarkAsDirty/*=true*/)
 {
 #if WITH_EDITORONLY_DATA
