LevelInstance: Quick cleanup of unused param

diff --git a/Engine/Source/Runtime/Engine/Private/LevelInstance/LevelInstanceLevelStreaming.cpp b/Engine/Source/Runtime/Engine/Private/LevelInstance/LevelInstanceLevelStreaming.cpp
--- a/Engine/Source/Runtime/Engine/Private/LevelInstance/LevelInstanceLevelStreaming.cpp
+++ b/Engine/Source/Runtime/Engine/Private/LevelInstance/LevelInstanceLevelStreaming.cpp
@@ -147,8 +147,6 @@ void ULevelStreamingLevelInstance::UnloadInstance(ULevelStreamingLevelInstance*
 			return true;
 		}, true);
 
-		// No need to clear the whole editor selection since actor of this level will be removed from the selection by: UEditorEngine::OnLevelRemovedFromWorld
-		const bool bClearSelection = false;
 		LevelStreaming->GetWorld()->GetSubsystem<ULevelInstanceSubsystem>()->RemoveLevelFromWorld(LevelStreaming->GetLoadedLevel(), bResetTrans);
 	}
 #endif 
diff --git a/Engine/Source/Runtime/Engine/Private/LevelInstance/LevelInstanceSubsystem.cpp b/Engine/Source/Runtime/Engine/Private/LevelInstance/LevelInstanceSubsystem.cpp
--- a/Engine/Source/Runtime/Engine/Private/LevelInstance/LevelInstanceSubsystem.cpp
+++ b/Engine/Source/Runtime/Engine/Private/LevelInstance/LevelInstanceSubsystem.cpp
@@ -1081,7 +1081,9 @@ void ULevelInstanceSubsystem::RemoveLevelFromWorld(ULevel* Level, bool bResetTra
 	}
 	else
 	{
-		EditorLevelUtils::RemoveLevelFromWorld(Level, false, bResetTrans);
+		// No need to clear the whole editor selection since actor of this level will be removed from the selection by: UEditorEngine::OnLevelRemovedFromWorld
+		const bool bClearSelection = false;
+		EditorLevelUtils::RemoveLevelFromWorld(Level, bClearSelection, bResetTrans);
 	}
 }
 
