1. Porting the shading model "Subsurface" and "PreIntegrated Skin" to mobile.
2. Support "Subsurface Profile" shading mobile by preintegral burley diffusion on ring, just like the PreIntegrated Skin, but could be generated at runtime based on the Burley inputs.

#p4v-preflight-copy 16327658

diff --git a/Engine/Shaders/Private/BurleyNormalizedSSSCommon.ush b/Engine/Shaders/Private/BurleyNormalizedSSSCommon.ush
--- a/Engine/Shaders/Private/BurleyNormalizedSSSCommon.ush
+++ b/Engine/Shaders/Private/BurleyNormalizedSSSCommon.ush
@@ -10,6 +10,11 @@
 #define BURLEY_MM_2_CM		0.1f
 #define BURLEY_CM_2_MM      10.0f
 
+// exactly one of these should be true
+#define LIGHT_PERPENDICULAR 0
+#define LIGHT_DIFFUSESURFACE 0
+#define LIGHT_PERPENDICULAR_DIFFUSE_SURFACE 1
+
 //------------------------------------------------------------------------------------------
 // Consistent in BurleyNormalizedSSSCommon.ush and SubsurfaceProfile.cpp
 
@@ -146,6 +151,30 @@ float3 GetSearchLightDiffuseScalingFactor3D(float3 SurfaceAlbedo)
 	return 3.5 + 100 * Value * Value * Value * Value;
 }
 
+float GetScalingFactor(float A)
+{
+#if LIGHT_PERPENDICULAR
+	float S = GetPerpendicularScalingFactor(A);
+#elif (LIGHT_DIFFUSESURFACE)
+	float S = GetDiffuseSurfaceScalingFactor(A);
+#elif (LIGHT_PERPENDICULAR_DIFFUSE_SURFACE)
+	float S = GetSearchLightDiffuseScalingFactor(A);
+#endif
+	return S;
+}
+
+float3 GetScalingFactor3D(float3 SurfaceAlbedo)
+{
+#if LIGHT_PERPENDICULAR
+	float3 S3D = GetPerpendicularScalingFactor3D(SurfaceAlbedo);
+#elif (LIGHT_DIFFUSESURFACE)
+	float3 S3D = GetDiffuseSurfaceScalingFactor3D(SurfaceAlbedo);
+#elif (LIGHT_PERPENDICULAR_DIFFUSE_SURFACE)
+	float3 S3D = GetSearchLightDiffuseScalingFactor3D(SurfaceAlbedo);
+#endif
+	return S3D;
+}
+
 //@TODO Revisit here after offline comparison.
 //call this function in TransmissionCommon.ush for Debug. No transmission tint is applied.
 
diff --git a/Engine/Shaders/Private/MobileBasePassPixelShader.usf b/Engine/Shaders/Private/MobileBasePassPixelShader.usf
--- a/Engine/Shaders/Private/MobileBasePassPixelShader.usf
+++ b/Engine/Shaders/Private/MobileBasePassPixelShader.usf
@@ -10,6 +10,13 @@
 #define MobileSceneTextures MobileBasePass.SceneTextures
 #define EyeAdaptationStruct MobileBasePass
 
+#define ActualSSProfilesTexture MobileBasePass.SSProfilesTexture
+#define SSProfilesSampler MobileBasePass.SSProfilesSampler
+
+#define SSProfilesPreIntegratedTexture MobileBasePass.SSProfilesPreIntegratedTexture
+#define SSProfilesPreIntegratedSampler MobileBasePass.SSProfilesPreIntegratedSampler
+#define SSProfilesTextureSizeAndInvSize MobileBasePass.SSProfilesTextureSizeAndInvSize
+
 #ifndef MOBILE_QL_FORCE_DISABLE_PREINTEGRATEDGF
 #define MOBILE_QL_FORCE_DISABLE_PREINTEGRATEDGF 0
 #endif
@@ -653,6 +660,14 @@ void Main(
 	// 0..1, SubsurfaceProfileId = int(x * 255)
 	float SubsurfaceProfile = 0;
 
+#if (MATERIAL_SHADINGMODEL_SUBSURFACE || MATERIAL_SHADINGMODEL_PREINTEGRATED_SKIN || MATERIAL_SHADINGMODEL_SUBSURFACE_PROFILE) && MOBILE_SHADINGMODEL_SUPPORT
+	half4 SubsurfaceData = GetMaterialSubsurfaceData(PixelMaterialInputs);
+
+	SubsurfaceColor = SubsurfaceData.rgb * ResolvedView.DiffuseOverrideParameter.w + ResolvedView.DiffuseOverrideParameter.xyz;
+
+	SubsurfaceProfile = SubsurfaceData.a;
+#endif
+
 #if APPLY_AO
 	half4 GatheredAmbientOcclusion = Texture2DSample(AmbientOcclusionTexture, AmbientOcclusionSampler, SvPositionToBufferUV(SvPosition));
 
@@ -709,6 +724,18 @@ void Main(
 	GBuffer.BaseColor *= TopMaterialCoverage;
 #endif
 
+#if MATERIAL_SHADINGMODEL_SUBSURFACE_PROFILE && MOBILE_SHADINGMODEL_SUPPORT
+	if (GBuffer.ShadingModelID == SHADINGMODELID_SUBSURFACE_PROFILE)
+	{
+		#if MATERIAL_SUBSURFACE_PROFILE_USE_CURVATURE
+			GBuffer.CustomData.g = GetMaterialCustomData0(MaterialParameters);
+		#else
+			GBuffer.CustomData.g = CalculateCurvature(GBuffer.WorldNormal, MaterialParameters.AbsoluteWorldPosition);
+		#endif
+
+		GBuffer.CustomData.g = max(saturate(GBuffer.CustomData.g), 0.001f);
+	}
+#endif
 	InitShadingModelContext(ShadingModelContext, GBuffer, MaterialParameters.CameraVector);
 	float3 DiffuseDir = MaterialParameters.WorldNormal;
 
@@ -746,6 +773,13 @@ void Main(
 	half3 IndirectColor;
 	half3 DiffuseColorForIndirect = ShadingModelContext.DiffuseColor;
 
+#if (MATERIAL_SHADINGMODEL_SUBSURFACE || MATERIAL_SHADINGMODEL_PREINTEGRATED_SKIN) && MOBILE_SHADINGMODEL_SUPPORT
+	if (GBuffer.ShadingModelID == SHADINGMODELID_SUBSURFACE || GBuffer.ShadingModelID == SHADINGMODELID_PREINTEGRATED_SKIN)
+	{
+		DiffuseColorForIndirect += ShadingModelContext.SubsurfaceColor;
+	}
+#endif
+
 	// Apply MaterialAO since we don't have the DiffuseIndirectComposite pass on mobile deferred.
 	ComputeIndirect(LightmapVTPageTableResult, Interpolants, DiffuseDir, DiffuseColorForIndirect, MaterialAO, IndirectIrradiance, MobileLightAccumulator);
 
@@ -951,6 +985,9 @@ void Main(
 		ShadingModelContext.WaterDiffuseIndirectLuminance += SkyDiffuseLighting;
 	#endif
 		half3 DiffuseColorForSky = ShadingModelContext.DiffuseColor;
+	#if MATERIAL_SHADINGMODEL_SUBSURFACE || MATERIAL_SHADINGMODEL_PREINTEGRATED_SKIN
+		DiffuseColorForSky += ShadingModelContext.SubsurfaceColor;
+	#endif
 		MobileLightAccumulator_Add(MobileLightAccumulator, SkyDiffuseLighting, 0.0f, half3(ResolvedView.SkyLightColor.rgb) * DiffuseColorForSky * MaterialAO);
 	#if MATERIAL_TWOSIDED && LQ_TEXTURE_LIGHTMAP
 	}
diff --git a/Engine/Shaders/Private/MobileDeferredShading.usf b/Engine/Shaders/Private/MobileDeferredShading.usf
--- a/Engine/Shaders/Private/MobileDeferredShading.usf
+++ b/Engine/Shaders/Private/MobileDeferredShading.usf
@@ -332,6 +332,16 @@ void MobileDirectLightPS(
 #endif
 	half3 SkyDiffuseLighting = GetSkySHDiffuseSimple(GBuffer.WorldNormal);
 	half3 DiffuseColorForSky = ShadingModelContext.DiffuseColor;
+
+#if MOBILE_SHADINGMODEL_SUPPORT
+	if (GBuffer.ShadingModelID == SHADINGMODELID_SUBSURFACE || GBuffer.ShadingModelID == SHADINGMODELID_PREINTEGRATED_SKIN)
+	{
+		half3 SubsurfaceColor = ExtractSubsurfaceColor(GBuffer);
+		// Add subsurface energy to diffuse
+		DiffuseColorForSky += SubsurfaceColor;
+	}
+#endif
+
 	MobileLightAccumulator_Add(MobileLightAccumulator, SkyDiffuseLighting, 0.0f, half3(View.SkyLightColor.rgb) * DiffuseColorForSky * GBuffer.GBufferAO);
 	
 	half3 Color = MobileLightAccumulator_GetResult(MobileLightAccumulator);
diff --git a/Engine/Shaders/Private/MobileShadingModels.ush b/Engine/Shaders/Private/MobileShadingModels.ush
--- a/Engine/Shaders/Private/MobileShadingModels.ush
+++ b/Engine/Shaders/Private/MobileShadingModels.ush
@@ -8,6 +8,7 @@
 #include "ShadowFilteringCommon.ush"
 #include "DeferredShadingCommon.ush"
 #include "HairBsdf.ush"
+#include "BurleyNormalizedSSSCommon.ush"
 
 #ifndef FEATURE_GATHER4
 #define FEATURE_GATHER4 1
@@ -87,6 +88,16 @@ struct FMobileShadingModelContext
 #if MOBILE_SHADINGMODEL_SUPPORT
 	half ClearCoat;
 	half ClearCoatRoughness;
+
+	// If we don't use this shading model the color should be black (don't generate shader code for unused data, don't do indirectlighting cache lighting with this color).
+	half3 SubsurfaceColor;
+
+	uint SubsurfaceProfileInt;
+	half  Curvature;
+	
+	half Lobe0Roughness;
+	half Lobe1Roughness;
+	half LobeMix;
 #endif
 };
 
@@ -104,6 +115,38 @@ half3 GetEnvBRDF(half3 SpecularColor, half Roughness, half NoV)
 #endif
 }
 
+#if MOBILE_SHADINGMODEL_SUPPORT
+half CalculateCurvature(half3 WorldNormal, float3 WorldPosition)
+{
+#if 0
+	half DeltaNormal = length(abs(DDX(WorldNormal)) + abs(DDY(WorldNormal)));
+	half DeltaPosition = length(abs(DDX(WorldPosition)) + abs(DDY(WorldPosition))) * BURLEY_CM_2_MM;
+	half CurvatureApprox = DeltaNormal / DeltaPosition;
+#else
+	half3 dNdx = ddx(WorldNormal);
+	half3 dNdy = ddy(WorldNormal);
+	half x = dot(dNdx, dNdx);
+	half y = dot(dNdy, dNdy);
+	half CurvatureApprox = pow(max(x, y), ResolvedView.NormalCurvatureToRoughnessScaleBias.z);
+#endif
+	return CurvatureApprox;
+}
+
+void GetProfileDualSpecular(uint SubsurfaceProfileInt, out half2 AverageToRoughness, out half LobeMix)
+{
+	half4 Data = GetSubsurfaceProfileTexture(SSSS_DUAL_SPECULAR_OFFSET, SubsurfaceProfileInt);
+	AverageToRoughness = Data.xy * SSSS_MAX_DUAL_SPECULAR_ROUGHNESS;
+	LobeMix = Data.z;
+}
+
+half4 GetSSProfilePreIntegratedValue(uint SubsurfaceProfileInt, half NoL, half Curvature)
+{
+	float3 UV = float3((NoL * .5 + .5), Curvature, SubsurfaceProfileInt);
+
+	return Texture2DArraySampleLevel(SSProfilesPreIntegratedTexture, SSProfilesPreIntegratedSampler, UV, 0);
+}
+#endif
+
 void InitShadingModelContext(inout FMobileShadingModelContext ShadingModelContext, inout FGBufferData GBuffer, half3 CameraVector)
 {
 	half NoV = max(dot(GBuffer.WorldNormal, CameraVector), 0.0f);
@@ -153,6 +196,24 @@ void InitShadingModelContext(inout FMobileShadingModelContext ShadingModelContex
 		half DielectricSpecular = 0.08 * Specular;
 		ShadingModelContext.SpecularColor = (DielectricSpecular - DielectricSpecular * GBuffer.Metallic) + GBuffer.BaseColor * GBuffer.Metallic;	// 2 mad
 	}
+	else if (GBuffer.ShadingModelID == SHADINGMODELID_SUBSURFACE_PROFILE)
+	{
+		ShadingModelContext.SubsurfaceProfileInt = ExtractSubsurfaceProfileInt(GBuffer);
+				
+		ShadingModelContext.Curvature = GBuffer.CustomData.g;
+
+		half2 AverageToRoughness;
+		GetProfileDualSpecular(ShadingModelContext.SubsurfaceProfileInt, AverageToRoughness, ShadingModelContext.LobeMix);
+
+		half AverageRoughness = GBuffer.Roughness;
+		ShadingModelContext.Lobe0Roughness = max(saturate(AverageRoughness * AverageToRoughness.x), 0.02f);
+		ShadingModelContext.Lobe1Roughness = saturate(AverageRoughness * AverageToRoughness.y);
+
+		// Smoothly lerp to default single GGX lobe as Opacity approaches 0
+		half RoughnessLerpFactor = saturate(ShadingModelContext.Opacity * 10.0f);
+		ShadingModelContext.Lobe0Roughness = lerp(1.0f, ShadingModelContext.Lobe0Roughness, RoughnessLerpFactor);
+		ShadingModelContext.Lobe1Roughness = lerp(1.0f, ShadingModelContext.Lobe1Roughness, RoughnessLerpFactor);
+	}
 #endif
 
 #if FULLY_ROUGH
@@ -278,6 +339,59 @@ FMobileDirectLighting MobileHairBxDF(FMobileShadingModelContext ShadingModelCont
 
 	return Lighting;
 }
+
+FMobileDirectLighting MobileSubsurfaceBxDF(FMobileShadingModelContext ShadingModelContext, FGBufferData GBuffer, half NoL, half NoH, half VoH, half3 V, half3 L, half Shadow)
+{
+	FMobileDirectLighting Lighting;
+
+	Lighting = MobileDefaultLitBxDF(ShadingModelContext, GBuffer, NoL, NoH, VoH, V, L, Shadow);
+	
+	half InScatter = pow(saturate(dot(L, -V)), 12) * lerp(3, .1f, ShadingModelContext.Opacity);
+
+	half NormalContribution = saturate(NoH * ShadingModelContext.Opacity + 1 - ShadingModelContext.Opacity);
+	half BackScatter = GBuffer.GBufferAO * NormalContribution / (PI * 2);
+
+	// lerp to never exceed 1 (energy conserving)
+	half3 Transmission = lerp(BackScatter, 1, InScatter) * ShadingModelContext.SubsurfaceColor * PI; // multiply PI since the LightColor is divided by PI in C++
+
+	Lighting.Diffuse += Transmission * Shadow;
+
+	return Lighting;
+}
+
+FMobileDirectLighting MobilePreintegratedSkinBxDF(FMobileShadingModelContext ShadingModelContext, FGBufferData GBuffer, half NoL, half NoH, half VoH, half3 V, half3 L, half Shadow)
+{
+	FMobileDirectLighting Lighting;
+
+	Lighting = MobileDefaultLitBxDF(ShadingModelContext, GBuffer, NoL, NoH, VoH, V, L, Shadow);
+
+	half3 PreintegratedBRDF = Texture2DSampleLevel(View.PreIntegratedBRDF, View.PreIntegratedBRDFSampler, float2(saturate(NoL * .5 + .5), 1 - ShadingModelContext.Opacity), 0).rgb;
+	half3 Transmission = PreintegratedBRDF * ShadingModelContext.SubsurfaceColor * PI; // multiply PI since the LightColor is divided by PI in C++
+
+	Lighting.Diffuse += Transmission * Shadow;
+
+	return Lighting;
+}
+
+FMobileDirectLighting MobileSubsurfaceProfileBxDF(FMobileShadingModelContext ShadingModelContext, FGBufferData GBuffer, half NoL, half NoH, half VoH, half3 V, half3 L, half Shadow)
+{
+	FMobileDirectLighting Lighting;
+
+	half UnClampedNoL = dot(GBuffer.WorldNormal, L);
+
+	half ShadowFactor = 1 - sqrt(Shadow);
+
+	// Rotate the world normal based on the shadow value, it's just a experimental value 
+	half UnClampedRotatedNoL = max(UnClampedNoL - max(2.0f * UnClampedNoL, 0.4f) * ShadowFactor, -1.0f);
+
+	half4 BurleyDiffuse = GetSSProfilePreIntegratedValue(ShadingModelContext.SubsurfaceProfileInt, UnClampedRotatedNoL, ShadingModelContext.Curvature);
+
+	Lighting.Diffuse = BurleyDiffuse.rgb * ShadingModelContext.DiffuseColor;
+
+	Lighting.Specular = Shadow * NoL * MobileDualSpecularGGX(GBuffer.Roughness, ShadingModelContext.Lobe0Roughness, ShadingModelContext.Lobe1Roughness, ShadingModelContext.LobeMix, ShadingModelContext.EnvBrdf, ShadingModelContext.SpecularColor, ShadingModelContext.NoV, NoH, NoL, VoH);
+
+	return Lighting;
+}
 #endif
 
 FMobileDirectLighting MobileIntegrateBxDF(FMobileShadingModelContext ShadingModelContext, FGBufferData GBuffer, half NoL, half NoH, half VoH, half3 V, half3 L, half Shadow = 1.0f)
@@ -291,6 +405,18 @@ FMobileDirectLighting MobileIntegrateBxDF(FMobileShadingModelContext ShadingMode
 	{
 		return MobileHairBxDF(ShadingModelContext, GBuffer, NoL, NoH, VoH, V, L, Shadow);
 	}
+	else if (GBuffer.ShadingModelID == SHADINGMODELID_SUBSURFACE)
+	{
+		return MobileSubsurfaceBxDF(ShadingModelContext, GBuffer, NoL, NoH, VoH, V, L, Shadow);
+	}
+	else if (GBuffer.ShadingModelID == SHADINGMODELID_PREINTEGRATED_SKIN)
+	{
+		return MobilePreintegratedSkinBxDF(ShadingModelContext, GBuffer, NoL, NoH, VoH, V, L, Shadow);
+	}
+	else if (GBuffer.ShadingModelID == SHADINGMODELID_SUBSURFACE_PROFILE)
+	{
+		return MobileSubsurfaceProfileBxDF(ShadingModelContext, GBuffer, NoL, NoH, VoH, V, L, Shadow);
+	}
 #endif
 
 	return MobileDefaultLitBxDF(ShadingModelContext, GBuffer, NoL, NoH, VoH, V, L, Shadow);
diff --git a/Engine/Shaders/Private/SSProfilePreIntegratedMobile.usf b/Engine/Shaders/Private/SSProfilePreIntegratedMobile.usf
new file mode 100644
--- /dev/null
+++ b/Engine/Shaders/Private/SSProfilePreIntegratedMobile.usf
@@ -0,0 +1,151 @@
+// Copyright Epic Games, Inc. All Rights Reserved.
+
+/*=============================================================================
+	SSProfilePreIntegratedMobile.usf
+=============================================================================*/
+
+#include "Common.ush"
+#include "BurleyNormalizedSSSCommon.ush"
+
+float4 SSProfilesPreIntegratedTextureSizeAndInvSize;
+uint SubsurfaceProfileInt;
+
+Texture2DArray InTextureArray;
+SamplerState InTextureSampler;
+
+void MainPSCopy(
+	FScreenVertexOutput Input,
+	out float4 OutColor : SV_Target0
+)
+{
+	OutColor = Texture2DArraySample(InTextureArray, InTextureSampler, float3(Input.UV, SubsurfaceProfileInt));
+}
+
+bool IsSubsurfaceProfileUseBurley()
+{
+	// 0..255, which SubSurface profile to pick
+	float Type = GetSubsurfaceProfileTexture(SSSS_BOUNDARY_COLOR_BLEED_OFFSET, SubsurfaceProfileInt).a;
+
+	return abs(Type - SSS_TYPE_BURLEY) < 0.01f;
+}
+
+float3 BurleyUnit(float Angle, float3 Sharpness)
+{
+	return exp(Sharpness * Angle);
+}
+
+// The integral of cos(x) * e^(a * x)
+float3 IntegralCosineWithExponent(float Angle, float SinAngle, float CosAngle, float3 Sharpness)
+{
+	return (SinAngle + Sharpness * CosAngle) * BurleyUnit(Angle, Sharpness) / (1 + Sharpness * Sharpness);
+}
+
+// The integral of sin(x) * e^(a * x)
+float3 IntegralSineWithExponent(float Angle, float SinAngle, float CosAngle, float3 Sharpness)
+{
+	return (-CosAngle + Sharpness * SinAngle) * BurleyUnit(Angle, Sharpness) / (1 + Sharpness * Sharpness);
+}
+
+// Theta is the angle between N and L, BurleyA is -((s * r) / m) * 2 / PI, BurleyB is BurleyA / 3
+float3 ApproximationRingIntegralOfBurleyDiffusion(float Angle, float SinAngle, float CosAngle, float Theta, float CosTheta, float SinTheta, float3 BurleyA, float3 BurleyB)
+{
+	return CosTheta * (IntegralCosineWithExponent(Angle, SinAngle, CosAngle, BurleyA) + IntegralCosineWithExponent(Angle, SinAngle, CosAngle, BurleyB))
+			-SinTheta * (IntegralSineWithExponent(Angle, SinAngle, CosAngle, BurleyA) + IntegralSineWithExponent(Angle, SinAngle, CosAngle, BurleyB));
+}
+
+float3 IntegralBurleyUnit(float Angle, float3 Sharpness)
+{
+	return (1 / Sharpness) * BurleyUnit(Angle, Sharpness);
+}
+
+float3 IntegralBurleyDiffusion(float Angle, float3 BurleyA, float3 BurleyB)
+{
+	return IntegralBurleyUnit(Angle, BurleyA) + IntegralBurleyUnit(Angle, BurleyB);
+}
+
+float4 GetIrradianceFromBurleyDiffusionOnRing(float4 SurfaceAlbedo, float4 DiffuseMeanFreePath, float Curvature, float Theta, float CosTheta, float SinTheta)
+{
+	const float HalfPI = PI * 0.5f;
+
+	float3 S3D = GetScalingFactor3D(SurfaceAlbedo.xyz);
+
+	// d = DiffuseMeanFreePath / S;
+	// r = 1.0f / Curvature;
+	// d / r = DiffuseMeanFreePath.xyz * Curvature / S3D;
+	// Multiply 7.0f / 3.0f / PI to be close to the range of 2 * sin(x/2);
+	float3 BurleyA = (-1.0f * S3D / (DiffuseMeanFreePath.xyz * Curvature)) * 2.0f / PI;
+
+	float3 BurleyB = BurleyA / 3.0f;
+
+	// Integral on the ring (from -PI to PI) to normalize the final integral.
+	float3 IntegralOnRing = 2 * (IntegralBurleyDiffusion(PI, BurleyA, BurleyB) - IntegralBurleyDiffusion(0.0f, BurleyA, BurleyB));
+
+	float3 IntegralOnNegativeAngle, IntegralOnPositiveAngle, BurleyDiffuse;
+	if (Theta <= HalfPI)
+	{
+		// Integral from -(PI/2 + Theta) to 0.
+		float AngleStart = -(HalfPI + Theta);
+		IntegralOnNegativeAngle = ApproximationRingIntegralOfBurleyDiffusion(0.0f, 0.0f, 1.0f, Theta, CosTheta, SinTheta, -BurleyA, -BurleyB) - ApproximationRingIntegralOfBurleyDiffusion(AngleStart, sin(AngleStart), cos(AngleStart), Theta, CosTheta, SinTheta, -BurleyA, -BurleyB);
+
+		// Integral from 0 to PI/2 - Theta
+		float AngleEnd = HalfPI - Theta;
+		IntegralOnPositiveAngle = ApproximationRingIntegralOfBurleyDiffusion(AngleEnd, sin(AngleEnd), cos(AngleEnd), Theta, CosTheta, SinTheta, BurleyA, BurleyB) - ApproximationRingIntegralOfBurleyDiffusion(0.0f, 0.0f, 1.0f, Theta, CosTheta, SinTheta, BurleyA, BurleyB);
+
+	}
+	else
+	{
+		// Integral from -PI to PI/2 - Theta.
+		float AngleEnd = HalfPI - Theta; 
+		IntegralOnNegativeAngle = ApproximationRingIntegralOfBurleyDiffusion(AngleEnd, sin(AngleEnd), cos(AngleEnd), Theta, CosTheta, SinTheta, -BurleyA, -BurleyB) - ApproximationRingIntegralOfBurleyDiffusion(-PI, 0.0f, -1.0f, Theta, CosTheta, SinTheta, -BurleyA, -BurleyB);
+
+		// Integral from 3*PI/2 - Theta to PI
+		float AngleStart = HalfPI * 3 - Theta;
+		IntegralOnPositiveAngle = ApproximationRingIntegralOfBurleyDiffusion(PI, 0.0f, -1.0f, Theta, CosTheta, SinTheta, BurleyA, BurleyB) - ApproximationRingIntegralOfBurleyDiffusion(AngleStart, sin(AngleStart), cos(AngleStart), Theta, CosTheta, SinTheta, BurleyA, BurleyB);
+
+	}
+
+	BurleyDiffuse = (IntegralOnNegativeAngle + IntegralOnPositiveAngle) / IntegralOnRing;
+
+//	float3 IntegralOnPI = ApproximationRingIntegralOfBurleyDiffusion(PI, 0.0f, -1.0f, 0.0f, 1.0f, 0.0f, BurleyA, BurleyB);
+
+//	float3 ScatteringFactor3D = 2.0f * (ApproximationRingIntegralOfBurleyDiffusion(PI, 0.0f, -1.0f, 0.0f, 1.0f, 0.0f, BurleyA, BurleyB) - IntegralOnPI) / IntegralOnRing;
+
+//	float ScatteringFactor = max3(ScatteringFactor3D.x, ScatteringFactor3D.y, ScatteringFactor3D.z);
+
+	return float4(BurleyDiffuse, 1.0f);
+}
+
+void SSProfilePreIntegratedPS(
+	in float4 SvPosition : SV_POSITION,
+	out HALF4_TYPE OutColor : SV_Target0)
+{
+	OutColor = HALF4_TYPE(1.0f, 1.0f, 1.0f, 0.0f);
+
+	float2 UVPos = SvPosition.xy * SSProfilesPreIntegratedTextureSizeAndInvSize.zw;
+
+	float UnClampedNoL = UVPos.x * 2.0f - 1.0f;
+	float Curvature = UVPos.y;
+
+	const bool bUseBurley = IsSubsurfaceProfileUseBurley();
+
+	uint SurfaceAlbedoSampleIndex = BSSS_SURFACEALBEDO_OFFSET;
+	uint DiffuseMeanFreePathSampleIndex = BSSS_DMFP_OFFSET;
+
+	if (!bUseBurley)
+	{
+		SurfaceAlbedoSampleIndex = SSSS_APPROXIMATE_SURFACEALBEDO_OFFSET;
+		DiffuseMeanFreePathSampleIndex = SSSS_APPROXIMATE_DMFP_OFFSET;
+	}
+
+	float4 SurfaceAlbedo = GetSubsurfaceProfileTexture(SurfaceAlbedoSampleIndex, SubsurfaceProfileInt);
+
+	float WorldUnitScale = DecodeWorldUnitScale(GetSubsurfaceProfileTexture(SSSS_SUBSURFACE_COLOR_OFFSET, SubsurfaceProfileInt).a) * BURLEY_CM_2_MM;
+
+	float4 DiffuseMeanFreePath = max(DecodeDiffuseMeanFreePath(GetSubsurfaceProfileTexture(DiffuseMeanFreePathSampleIndex, SubsurfaceProfileInt)) * WorldUnitScale, 0.01f);;
+
+	float CosTheta = UnClampedNoL;
+	float Theta = acos(CosTheta);
+	float SinTheta = sin(Theta);
+
+	OutColor = GetIrradianceFromBurleyDiffusionOnRing(SurfaceAlbedo, DiffuseMeanFreePath, Curvature, Theta, CosTheta, SinTheta);
+}
\ No newline at end of file
diff --git a/Engine/Shaders/Private/ShadingModelsMaterial.ush b/Engine/Shaders/Private/ShadingModelsMaterial.ush
--- a/Engine/Shaders/Private/ShadingModelsMaterial.ush
+++ b/Engine/Shaders/Private/ShadingModelsMaterial.ush
@@ -65,7 +65,7 @@ void SetGBufferForShadingModel(
 
 			// Average roughness for dual specular.
 			uint SubsurfaceProfileUInt = uint(SubsurfaceProfile * 255.0f + 0.5f);
-			float MaterialRoughnessToAverage = ActualSSProfilesTexture.Load(int3(SSSS_DUAL_SPECULAR_OFFSET, SubsurfaceProfileUInt, 0)).w * SSSS_MAX_DUAL_SPECULAR_ROUGHNESS;
+			float MaterialRoughnessToAverage = GetSubsurfaceProfileTexture(SSSS_DUAL_SPECULAR_OFFSET, SubsurfaceProfileUInt).w * SSSS_MAX_DUAL_SPECULAR_ROUGHNESS;
 
 			// Smooth blend out dual specular when opacity is low, we have the extra SSSS_OPACITY_THRESHOLD_EPS so that we fade out by the time we
 			// get to 0.01, as opposed to 0.0.
diff --git a/Engine/Shaders/Private/SubsurfaceBurleyNormalized.ush b/Engine/Shaders/Private/SubsurfaceBurleyNormalized.ush
--- a/Engine/Shaders/Private/SubsurfaceBurleyNormalized.ush
+++ b/Engine/Shaders/Private/SubsurfaceBurleyNormalized.ush
@@ -64,11 +64,6 @@
 #define MORTON_USE_LUT 0
 #define REWEIGHT_CENTER_SAMPLE 1
 
-// exactly one of these should be true
-#define LIGHT_PERPENDICULAR 0
-#define LIGHT_DIFFUSESURFACE 0
-#define LIGHT_PERPENDICULAR_DIFFUSE_SURFACE 1
-
 #define VARIANCE_LEVEL 0.0001
 #define HIGH_LUMA_SAMPLE_COUNT 8
 #define LOW_LUMA_SAMPLE_COUNT 16
@@ -131,18 +126,6 @@
 #define MORTON_USE_LUT 0
 #endif
 
-#ifndef LIGHT_PERPENDICULAR
-#define LIGHT_PERPENDICULAR 0
-#endif
-
-#ifndef LIGHT_DIFFUSESURFACE
-#define LIGHT_DIFFUSESURFACE 0
-#endif
-
-#ifndef LIGHT_PERPENDICULAR_DIFFUSE_SURFACE
-#define LIGHT_PERPENDICULAR_DIFFUSE_SURFACE 0
-#endif
-
 #ifndef ENABLE_PROFILE_ID_CACHE
 #define ENABLE_PROFILE_ID_CACHE 0
 #endif
@@ -795,31 +778,6 @@ BurleySampleDiffuseNormal SampleSSSColorConsideringLocalShared(float2 CenterUV,
 #endif
 }
 
-float GetScalingFactor(float A)
-{
-#if LIGHT_PERPENDICULAR
-	float S = GetPerpendicularScalingFactor(A);
-#elif (LIGHT_DIFFUSESURFACE)
-	float S = GetDiffuseSurfaceScalingFactor(A);
-#elif (LIGHT_PERPENDICULAR_DIFFUSE_SURFACE)
-	float S = GetSearchLightDiffuseScalingFactor(A);
-#endif
-	return S;
-}
-
-float3 GetScalingFactor3D(float3 SurfaceAlbedo)
-{
-#if LIGHT_PERPENDICULAR
-	float3 S3D = GetPerpendicularScalingFactor3D(SurfaceAlbedo);
-#elif (LIGHT_DIFFUSESURFACE)
-	float3 S3D = GetDiffuseSurfaceScalingFactor3D(SurfaceAlbedo);
-#elif (LIGHT_PERPENDICULAR_DIFFUSE_SURFACE)
-	float3 S3D = GetSearchLightDiffuseScalingFactor3D(SurfaceAlbedo);
-#endif
-	return S3D;
-}
-
-
 float2 CalculateBurleyScale(float WorldUnitScale, float DepthAtCenter)
 {
 	float2 BurleyScale = WorldUnitScale;
diff --git a/Engine/Shaders/Private/SubsurfaceProfileCommon.ush b/Engine/Shaders/Private/SubsurfaceProfileCommon.ush
--- a/Engine/Shaders/Private/SubsurfaceProfileCommon.ush
+++ b/Engine/Shaders/Private/SubsurfaceProfileCommon.ush
@@ -13,7 +13,9 @@
 #define BSSS_DMFP_OFFSET                        (BSSS_SURFACEALBEDO_OFFSET+1)
 #define SSSS_TRANSMISSION_OFFSET				(BSSS_DMFP_OFFSET+1)
 #define SSSS_BOUNDARY_COLOR_BLEED_OFFSET		(SSSS_TRANSMISSION_OFFSET+1)
-#define SSSS_DUAL_SPECULAR_OFFSET				(SSSS_BOUNDARY_COLOR_BLEED_OFFSET+1)
+#define SSSS_APPROXIMATE_SURFACEALBEDO_OFFSET	(SSSS_BOUNDARY_COLOR_BLEED_OFFSET+1)
+#define SSSS_APPROXIMATE_DMFP_OFFSET			(SSSS_APPROXIMATE_SURFACEALBEDO_OFFSET+1)
+#define SSSS_DUAL_SPECULAR_OFFSET				(SSSS_APPROXIMATE_DMFP_OFFSET+1)
 #define SSSS_KERNEL0_OFFSET						(SSSS_DUAL_SPECULAR_OFFSET+1)
 #define SSSS_KERNEL0_SIZE						13
 #define SSSS_KERNEL1_OFFSET						(SSSS_KERNEL0_OFFSET + SSSS_KERNEL0_SIZE)
@@ -39,5 +41,22 @@ Texture2D SSProfilesTexture;
 // Rename to allow the base pass to intercept with a value from its pass uniform buffer
 #ifndef ActualSSProfilesTexture
 #define ActualSSProfilesTexture SSProfilesTexture
+
+	Texture2DArray SSProfilesPreIntegratedTexture;
+	SamplerState SSProfilesPreIntegratedSampler;
+	#if COMPILER_GLSL_ES3_1
+	SamplerState SSProfilesSampler;
+	float4 SSProfilesTextureSizeAndInvSize;
+	#endif
+
 #endif
 
+half4 GetSubsurfaceProfileTexture(uint SampleIndex, uint SubsurfaceProfileInt)
+{
+#if COMPILER_GLSL_ES3_1 // Force to use a point sampler for Texture2D.Load on OpenGLES platform
+	return ActualSSProfilesTexture.SampleLevel(SSProfilesSampler, (uint2(SampleIndex, SubsurfaceProfileInt) + float2(0.5f, 0.5f)) * SSProfilesTextureSizeAndInvSize.zw, 0);
+#else
+	return ActualSSProfilesTexture.Load(int3(SampleIndex, SubsurfaceProfileInt, 0));
+#endif 
+}
+
diff --git a/Engine/Source/Runtime/Engine/Classes/Engine/SubsurfaceProfile.h b/Engine/Source/Runtime/Engine/Classes/Engine/SubsurfaceProfile.h
--- a/Engine/Source/Runtime/Engine/Classes/Engine/SubsurfaceProfile.h
+++ b/Engine/Source/Runtime/Engine/Classes/Engine/SubsurfaceProfile.h
@@ -23,19 +23,19 @@ struct FSubsurfaceProfileStruct
 	* Controls how far light goes into the subsurface in the Red, Green and Blue channel. It is scaled by Mean Free path distance.
 	*/
 	UPROPERTY(Category = "Burley Normalized", EditAnywhere, BlueprintReadOnly, meta = (ClampMin = "0.001", UIMax = "1.0", ClampMax = "1.0", HideAlphaChannel, editcondition = "bEnableBurley"))
-		FLinearColor MeanFreePathColor; //MeanFreePathLength;
+	FLinearColor MeanFreePathColor; //MeanFreePathLength;
 	
 	/**
 	* Subsurface mean free path distance in world/unreal units (cm)
 	*/
 	UPROPERTY(Category = "Burley Normalized", EditAnywhere, BluePrintReadOnly, meta = (ClampMin = "0.1", UIMax = "50.0", ClampMax = "50.0", editcondition = "bEnableBurley"))
-		float MeanFreePathDistance;
+	float MeanFreePathDistance;
 
 	/**
 	* Control the scale of world/unreal units (cm)
 	*/
 	UPROPERTY(Category = "Burley Normalized", EditAnywhere, BlueprintReadOnly, meta = (ClampMin = "0.1", UIMax = "50.0", ClampMax = "50.0", editcondition = "bEnableBurley"))
-		float WorldUnitScale;
+	float WorldUnitScale;
 
 	/**
 	* Effective only when Burley subsurface scattering is enabled in cmd.
@@ -91,7 +91,7 @@ struct FSubsurfaceProfileStruct
 	* Transmission tint control. It is multiplied on the transmission results. Works only when Burley is enabled.
 	*/
 	UPROPERTY(Category = "Transmission", EditAnywhere, BlueprintReadOnly, meta = (ClampMin = "0.001", UIMax = "1.0", ClampMax = "1.0", HideAlphaChannel))
-		FLinearColor TransmissionTintColor;
+	FLinearColor TransmissionTintColor;
 
 	// constructor
 	FSubsurfaceProfileStruct()
@@ -192,6 +192,8 @@ public:
 	struct IPooledRenderTarget* GetTexture(FRHICommandListImmediate& RHICmdList);
 	struct IPooledRenderTarget* GetTexture();
 
+	struct IPooledRenderTarget* GetSSProfilesPreIntegratedTexture(FRDGBuilder& GraphBuilder, EShaderPlatform ShaderPlatform);
+
 	//~ Begin FRenderResource Interface.
 	/**
 	* Release textures when device is lost/destroyed.
@@ -245,10 +247,13 @@ static const int32 SUBSURFACE_KERNEL_SIZE = 3;
 extern ENGINE_API TGlobalResource<FSubsurfaceProfileTexture> GSubsurfaceProfileTextureObject;
 
 // Initializes or updates the contents of the subsurface profile texture.
-ENGINE_API void UpdateSubsurfaceProfileTexture(FRHICommandListImmediate& RHICmdList);
+ENGINE_API void UpdateSubsurfaceProfileTexture(FRDGBuilder& GraphBuilder, EShaderPlatform ShaderPlatform);
 
 // Returns the subsurface profile texture if it exists, or null.
 ENGINE_API FRHITexture* GetSubsurfaceProfileTexture();
 
 // Returns the subsurface profile texture if it exists, or black.
-ENGINE_API FRHITexture* GetSubsurfaceProfileTextureWithFallback();
\ No newline at end of file
+ENGINE_API FRHITexture* GetSubsurfaceProfileTextureWithFallback();
+
+// Returns the Preintegrated texture if it exists, or black.
+ENGINE_API FRHITexture* GetSSProfilesPreIntegratedTextureWithFallback();
\ No newline at end of file
diff --git a/Engine/Source/Runtime/Engine/Private/Materials/HLSLMaterialTranslator.cpp b/Engine/Source/Runtime/Engine/Private/Materials/HLSLMaterialTranslator.cpp
--- a/Engine/Source/Runtime/Engine/Private/Materials/HLSLMaterialTranslator.cpp
+++ b/Engine/Source/Runtime/Engine/Private/Materials/HLSLMaterialTranslator.cpp
@@ -153,6 +153,7 @@ FHLSLMaterialTranslator::FHLSLMaterialTranslator(FMaterial* InMaterial,
 ,	bAllowCodeChunkGeneration(true)
 ,	bUsesAnisotropy(false)
 ,	bMaterialIsStrata(false)
+,	bUsesCurvature(false)
 ,	AllocatedUserTexCoords()
 ,	AllocatedUserVertexTexCoords()
 ,	DynamicParticleParameterMask(0)
@@ -910,6 +911,8 @@ bool FHLSLMaterialTranslator::Translate()
 			AddEstimatedTextureSample(1);
 		}
 
+		bUsesCurvature = FeatureLevel == ERHIFeatureLevel::ES3_1 && MaterialShadingModels.HasShadingModel(MSM_SubsurfaceProfile) && IsMaterialPropertyUsed(MP_CustomData0, Chunk[MP_CustomData0], FLinearColor(1, 0, 0, 0), 1);
+
 		if (BlendMode == BLEND_Modulate && MaterialShadingModels.IsLit() && !Material->IsDeferredDecal())
 		{
 			Errorf(TEXT("Dynamically lit translucency is not supported for BLEND_Modulate materials."));
@@ -1528,6 +1531,11 @@ void FHLSLMaterialTranslator::GetMaterialEnvironment(EShaderPlatform InPlatform,
 		{
 			OutEnvironment.SetDefine(TEXT("MATERIAL_SHADINGMODEL_SUBSURFACE_PROFILE"), TEXT("1"));
 			NumSetMaterials++;
+
+			if (bUsesCurvature)
+			{
+				OutEnvironment.SetDefine(TEXT("MATERIAL_SUBSURFACE_PROFILE_USE_CURVATURE"), TEXT("1"));
+			}
 		}
 		if (ShadingModels.HasShadingModel(MSM_ClearCoat))
 		{
diff --git a/Engine/Source/Runtime/Engine/Private/Materials/HLSLMaterialTranslator.h b/Engine/Source/Runtime/Engine/Private/Materials/HLSLMaterialTranslator.h
--- a/Engine/Source/Runtime/Engine/Private/Materials/HLSLMaterialTranslator.h
+++ b/Engine/Source/Runtime/Engine/Private/Materials/HLSLMaterialTranslator.h
@@ -379,6 +379,8 @@ protected:
 	
 	uint32 bEnableExecutionFlow : 1;
 
+	uint32 bUsesCurvature : 1;
+
 	/** Tracks the texture coordinates used by this material. */
 	TBitArray<> AllocatedUserTexCoords;
 	/** Tracks the texture coordinates used by the vertex shader in this material. */
diff --git a/Engine/Source/Runtime/Engine/Private/Materials/Material.cpp b/Engine/Source/Runtime/Engine/Private/Materials/Material.cpp
--- a/Engine/Source/Runtime/Engine/Private/Materials/Material.cpp
+++ b/Engine/Source/Runtime/Engine/Private/Materials/Material.cpp
@@ -6770,7 +6770,7 @@ static bool IsPropertyActive_Internal(EMaterialProperty InProperty,
 		Active = ShadingModels.HasAnyShadingModel({ MSM_Subsurface, MSM_PreintegratedSkin, MSM_TwoSidedFoliage, MSM_Cloth });
 		break;
 	case MP_CustomData0:
-		Active = ShadingModels.HasAnyShadingModel({ MSM_ClearCoat, MSM_Hair, MSM_Cloth, MSM_Eye });
+		Active = ShadingModels.HasAnyShadingModel({ MSM_ClearCoat, MSM_Hair, MSM_Cloth, MSM_Eye, MSM_SubsurfaceProfile });
 		break;
 	case MP_CustomData1:
 		Active = ShadingModels.HasAnyShadingModel({ MSM_ClearCoat, MSM_Eye });
diff --git a/Engine/Source/Runtime/Engine/Private/Materials/MaterialShared.cpp b/Engine/Source/Runtime/Engine/Private/Materials/MaterialShared.cpp
--- a/Engine/Source/Runtime/Engine/Private/Materials/MaterialShared.cpp
+++ b/Engine/Source/Runtime/Engine/Private/Materials/MaterialShared.cpp
@@ -4555,6 +4555,7 @@ FText FMaterialAttributeDefinitionMap::GetAttributeOverrideForMaterial(const FGu
 		CustomPinNames.Add({MSM_Hair, "Backlit"});
 		CustomPinNames.Add({MSM_Cloth, "Cloth"});
 		CustomPinNames.Add({MSM_Eye, "Iris Mask"});
+		CustomPinNames.Add({MSM_SubsurfaceProfile, "Curvature" });
 		return FText::FromString(GetPinNameFromShadingModelField(Material->GetShadingModels(), CustomPinNames, "Custom Data 0"));
 	case MP_CustomData1:
 		CustomPinNames.Add({ MSM_ClearCoat, "Clear Coat Roughness" });
diff --git a/Engine/Source/Runtime/Engine/Private/Rendering/SubsurfaceProfile.cpp b/Engine/Source/Runtime/Engine/Private/Rendering/SubsurfaceProfile.cpp
--- a/Engine/Source/Runtime/Engine/Private/Rendering/SubsurfaceProfile.cpp
+++ b/Engine/Source/Runtime/Engine/Private/Rendering/SubsurfaceProfile.cpp
@@ -7,15 +7,35 @@
 #include "Rendering/BurleyNormalizedSSS.h"
 #include "EngineModule.h"
 #include "RenderTargetPool.h"
+#include "PixelShaderUtils.h"
 
 DEFINE_LOG_CATEGORY_STATIC(LogSubsurfaceProfile, Log, All);
 
+static TAutoConsoleVariable<int32> CVarSSProfilesPreIntegratedTextureResolution(
+	TEXT("r.SSProfilesPreIntegratedTextureResolution"),
+	64,
+	TEXT("The resolution of the subsurface profile preintegrated texture.\n"),
+	ECVF_RenderThreadSafe
+);
+
+#if !(UE_BUILD_SHIPPING || UE_BUILD_TEST)
+static TAutoConsoleVariable<int32> CVarSSProfilesPreIntegratedTextureForceUpdate(
+	TEXT("r.SSProfilesPreIntegratedTextureForceUpdate"),
+	0,
+	TEXT("0: Only update the preintegrated texture as needed.\n")
+	TEXT("1: Force to update the preintegrated texture for debugging.\n"),
+	ECVF_Cheat | ECVF_RenderThreadSafe);
+#endif
+
 // lives on the render thread
 ENGINE_API TGlobalResource<FSubsurfaceProfileTexture> GSubsurfaceProfileTextureObject;
 
 // Texture with one or more SubSurfaceProfiles or 0 if there is no user
 static TRefCountPtr<IPooledRenderTarget> GSSProfiles;
 
+// Texture with one or more pre-integrated textures or 0 if there is no user
+static TRefCountPtr<IPooledRenderTarget> GSSProfilesPreIntegratedTexture;
+
 FSubsurfaceProfileTexture::FSubsurfaceProfileTexture()
 {
 	check(IsInGameThread());
@@ -96,6 +116,7 @@ void FSubsurfaceProfileTexture::UpdateProfile(int32 AllocationId, const FSubsurf
 	SubsurfaceProfileEntries[AllocationId].Settings = Settings;
 
 	GSSProfiles.SafeRelease();
+	GSSProfilesPreIntegratedTexture.SafeRelease();
 }
 
 IPooledRenderTarget* FSubsurfaceProfileTexture::GetTexture()
@@ -113,9 +134,92 @@ IPooledRenderTarget* FSubsurfaceProfileTexture::GetTexture(FRHICommandListImmedi
 	return GSSProfiles;
 }
 
+class FSSProfilePreIntegratedPS : public FGlobalShader
+{
+public:
+	DECLARE_GLOBAL_SHADER(FSSProfilePreIntegratedPS);
+	SHADER_USE_PARAMETER_STRUCT(FSSProfilePreIntegratedPS, FGlobalShader);
+
+	BEGIN_SHADER_PARAMETER_STRUCT(FParameters, )
+		SHADER_PARAMETER_TEXTURE(Texture2D, SSProfilesTexture)
+		SHADER_PARAMETER_SAMPLER(SamplerState, SSProfilesSampler)
+		SHADER_PARAMETER(FVector4, SSProfilesTextureSizeAndInvSize)
+		SHADER_PARAMETER(FVector4, SSProfilesPreIntegratedTextureSizeAndInvSize)
+		SHADER_PARAMETER(int32, SubsurfaceProfileInt)
+		RENDER_TARGET_BINDING_SLOTS()
+	END_SHADER_PARAMETER_STRUCT()
+
+	static bool ShouldCompilePermutation(const FGlobalShaderPermutationParameters& Parameters)
+	{
+		return IsMobilePlatform(Parameters.Platform);
+	}
+};
+
+IMPLEMENT_GLOBAL_SHADER(FSSProfilePreIntegratedPS, "/Engine/Private/SSProfilePreIntegratedMobile.usf", "SSProfilePreIntegratedPS", SF_Pixel);
+
+IPooledRenderTarget* FSubsurfaceProfileTexture::GetSSProfilesPreIntegratedTexture(FRDGBuilder& GraphBuilder, EShaderPlatform ShaderPlatform)
+{
+	// For now the pre-integrated texture only used on mobile platform
+	if (!IsMobilePlatform(ShaderPlatform))
+	{
+		return nullptr;
+	}
+
+	// PreIntegrated SSS look up texture
+	int32 SSProfilesPreIntegratedTextureResolution = FMath::RoundUpToPowerOfTwo(FMath::Max(CVarSSProfilesPreIntegratedTextureResolution.GetValueOnAnyThread(), 32));
+
+	int32 SubsurfaceProfileEntriesNum = SubsurfaceProfileEntries.Num();
+
+	bool bSSProfilesPreIntegratedTextureNeedsUpdate = false;
+
+	if (!GSSProfilesPreIntegratedTexture || GSSProfilesPreIntegratedTexture->GetDesc().Extent != SSProfilesPreIntegratedTextureResolution)
+	{
+		bSSProfilesPreIntegratedTextureNeedsUpdate = true;
+
+		// Use RGB10A2 since it could be compressed by mobile hardware according to ARM.
+		FPooledRenderTargetDesc Desc(FPooledRenderTargetDesc::Create2DArrayDesc(SSProfilesPreIntegratedTextureResolution, PF_A2B10G10R10, FClearValueBinding::Black, TexCreate_TargetArraySlicesIndependently, TexCreate_ShaderResource | TexCreate_RenderTargetable, false, SubsurfaceProfileEntriesNum));
+
+		GRenderTargetPool.FindFreeElement(GraphBuilder.RHICmdList, Desc, GSSProfilesPreIntegratedTexture, TEXT("SSProfilePreIntegratedTexture"));
+	}
+	// Generate the new preintegrated texture if needed.
+	if (bSSProfilesPreIntegratedTextureNeedsUpdate
+#if !(UE_BUILD_SHIPPING || UE_BUILD_TEST)
+		|| CVarSSProfilesPreIntegratedTextureForceUpdate.GetValueOnAnyThread() == 1
+#endif
+		)
+	{
+		for (int32 i = 0; i < SubsurfaceProfileEntriesNum; ++i)
+		{
+			FSSProfilePreIntegratedPS::FParameters* PassParameters = GraphBuilder.AllocParameters<FSSProfilePreIntegratedPS::FParameters>();
+			PassParameters->RenderTargets[0] = FRenderTargetBinding(GraphBuilder.RegisterExternalTexture(GSSProfilesPreIntegratedTexture), ERenderTargetLoadAction::EClear, 0, i);
+
+			PassParameters->SSProfilesTexture = GetSubsurfaceProfileTextureWithFallback();
+			PassParameters->SSProfilesSampler = TStaticSamplerState<SF_Point, AM_Clamp, AM_Clamp, AM_Clamp>::GetRHI();
+
+			FIntVector SSProfilesTextureSize = PassParameters->SSProfilesTexture->GetSizeXYZ();
+
+			PassParameters->SSProfilesTextureSizeAndInvSize = FVector4(SSProfilesTextureSize.X, SSProfilesTextureSize.Y, 1.0f / SSProfilesTextureSize.X, 1.0f / SSProfilesTextureSize.Y);
+
+			PassParameters->SSProfilesPreIntegratedTextureSizeAndInvSize = FVector4(SSProfilesPreIntegratedTextureResolution, SSProfilesPreIntegratedTextureResolution, 1.0f / SSProfilesPreIntegratedTextureResolution, 1.0f / SSProfilesPreIntegratedTextureResolution);
+
+			PassParameters->SubsurfaceProfileInt = i;
+
+			FIntRect ViewRect = FIntRect(FIntPoint(0, 0), SSProfilesPreIntegratedTextureResolution);
+
+			const auto GlobalShaderMap = GetGlobalShaderMap(ERHIFeatureLevel::ES3_1);
+			TShaderMapRef<FSSProfilePreIntegratedPS> PixelShader(GlobalShaderMap);
+
+			FPixelShaderUtils::AddFullscreenPass(GraphBuilder, GlobalShaderMap, RDG_EVENT_NAME("SSProfilePreIntegrated"), PixelShader, PassParameters, ViewRect);
+		}
+	}
+
+	return GSSProfilesPreIntegratedTexture;
+}
+
 void FSubsurfaceProfileTexture::ReleaseDynamicRHI()
 {
 	GSSProfiles.SafeRelease();
+	GSSProfilesPreIntegratedTexture.SafeRelease();
 }
 
 static float GetNextSmallerPositiveFloat(float x)
@@ -136,7 +240,9 @@ static float GetNextSmallerPositiveFloat(float x)
 #define BSSS_DMFP_OFFSET                        (BSSS_SURFACEALBEDO_OFFSET+1)
 #define SSSS_TRANSMISSION_OFFSET				(BSSS_DMFP_OFFSET+1)
 #define SSSS_BOUNDARY_COLOR_BLEED_OFFSET		(SSSS_TRANSMISSION_OFFSET+1)
-#define SSSS_DUAL_SPECULAR_OFFSET				(SSSS_BOUNDARY_COLOR_BLEED_OFFSET+1)
+#define SSSS_APPROXIMATE_SURFACEALBEDO_OFFSET	(SSSS_BOUNDARY_COLOR_BLEED_OFFSET+1)
+#define SSSS_APPROXIMATE_DMFP_OFFSET			(SSSS_APPROXIMATE_SURFACEALBEDO_OFFSET+1)
+#define SSSS_DUAL_SPECULAR_OFFSET				(SSSS_APPROXIMATE_DMFP_OFFSET+1)
 #define SSSS_KERNEL0_OFFSET						(SSSS_DUAL_SPECULAR_OFFSET+1)
 #define SSSS_KERNEL0_SIZE						13
 #define SSSS_KERNEL1_OFFSET						(SSSS_KERNEL0_OFFSET + SSSS_KERNEL0_SIZE)
@@ -211,6 +317,21 @@ float Pow4(float X)
 	return X * X * X * X;
 }
 
+void ConvertSubsurfaceParametersFromSeparableToBurley(const FSubsurfaceProfileStruct& Settings, FLinearColor& SurfaceAlbedo, FLinearColor& MeanFreePathColor, float& MeanFreePathDistance)
+{
+	MapFallOffColor2SurfaceAlbedoAndDiffuseMeanFreePath(Settings.FalloffColor.R, SurfaceAlbedo.R, MeanFreePathColor.R);
+	MapFallOffColor2SurfaceAlbedoAndDiffuseMeanFreePath(Settings.FalloffColor.G, SurfaceAlbedo.G, MeanFreePathColor.G);
+	MapFallOffColor2SurfaceAlbedoAndDiffuseMeanFreePath(Settings.FalloffColor.B, SurfaceAlbedo.B, MeanFreePathColor.B);
+
+	//Normalize mean free path color and set the corresponding dfmp
+	float MaxMeanFreePathColor = FMath::Max3(MeanFreePathColor.R, MeanFreePathColor.G, MeanFreePathColor.B);
+	if (MaxMeanFreePathColor > 1)
+	{
+		MeanFreePathColor /= MaxMeanFreePathColor;
+		MeanFreePathDistance = FMath::Clamp(Settings.ScatterRadius * MaxMeanFreePathColor, 0.1f, 50.0f);	// 50.0f is the ClampMax of MeanFreePathDistance.
+	}
+}
+
 void FSubsurfaceProfileTexture::CreateTexture(FRHICommandListImmediate& RHICmdList)
 {
 	uint32 Height = SubsurfaceProfileEntries.Num();
@@ -239,8 +360,12 @@ void FSubsurfaceProfileTexture::CreateTexture(FRHICommandListImmediate& RHICmdLi
 	FLinearColor TextureRow[Width];
 	FMemory::Memzero(TextureRow);
 
-	// Android devices have different precision to desktop GPU's, scaling to one less works on Adreno and Mali
-	const float FloatScaleInitial = GMaxRHIShaderPlatform == SP_VULKAN_SM5_ANDROID ? 0xffff : 0x10000;
+	// bias to avoid div by 0 and a jump to a different value
+	// this basically means we don't want subsurface scattering
+	// 0.0001f turned out to be too small to fix the issue (for a small KernelSize)
+	const float Bias = 0.009f;
+	const float CmToMm = 10.f;
+	const float FloatScaleInitial = 0x10000;
 	const float FloatScale = GetNextSmallerPositiveFloat(FloatScaleInitial);
 	check((int32)GetNextSmallerPositiveFloat(FloatScaleInitial) == (FloatScaleInitial - 1));
 
@@ -248,10 +373,22 @@ void FSubsurfaceProfileTexture::CreateTexture(FRHICommandListImmediate& RHICmdLi
 	{
 		FSubsurfaceProfileStruct Data = SubsurfaceProfileEntries[y].Settings;
 
-		// bias to avoid div by 0 and a jump to a different value
-		// this basically means we don't want subsurface scattering
-		// 0.0001f turned out to be too small to fix the issue (for a small KernelSize)
-		const float Bias = 0.009f;
+		// Convert subsurface parameters from separable to burley for mobile platform, and because the current shader platform is unknown at here, the converted parameters need to be saved in different slots(SSSS_APPROXIMATE_SURFACEALBEDO_OFFSET and SSSS_APPROXIMATE_DMFP_OFFSET).
+		if (!Data.bEnableBurley)
+		{
+			FLinearColor ApproximateSurfaceAlbedo, ApproximateMeanFreePathColor;
+			float ApproximateMeanFreePathDistance;
+			ConvertSubsurfaceParametersFromSeparableToBurley(Data, ApproximateSurfaceAlbedo, ApproximateMeanFreePathColor, ApproximateMeanFreePathDistance);
+			
+			ApproximateMeanFreePathColor = ApproximateMeanFreePathColor.GetClamped(Bias);
+			ApproximateSurfaceAlbedo = ApproximateSurfaceAlbedo.GetClamped(Bias);
+			
+			FLinearColor ApproximateDiffuseMeanFreePathInMm = (ApproximateMeanFreePathColor * ApproximateMeanFreePathDistance) * CmToMm;
+			SetupSurfaceAlbedoAndDiffuseMeanFreePath(ApproximateSurfaceAlbedo, ApproximateDiffuseMeanFreePathInMm);
+
+			TextureRow[SSSS_APPROXIMATE_SURFACEALBEDO_OFFSET] = ApproximateSurfaceAlbedo;
+			TextureRow[SSSS_APPROXIMATE_DMFP_OFFSET] = EncodeDiffuseMeanFreePath(ApproximateDiffuseMeanFreePathInMm);
+		}
 
 		Data.SubsurfaceColor = Data.SubsurfaceColor.GetClamped();
 		Data.FalloffColor = Data.FalloffColor.GetClamped(Bias);
@@ -264,7 +401,6 @@ void FSubsurfaceProfileTexture::CreateTexture(FRHICommandListImmediate& RHICmdLi
 		TextureRow[SSSS_SUBSURFACE_COLOR_OFFSET].A = EncodeWorldUnitScale(Data.WorldUnitScale);
 		
 		const float UnitToCm = Data.WorldUnitScale;
-		const float CmToMm = 10.f;
 
 		FLinearColor DifffuseMeanFreePathInMm = (Data.MeanFreePathColor*Data.MeanFreePathDistance) * CmToMm; // convert cm to mm.
 		SetupSurfaceAlbedoAndDiffuseMeanFreePath(Data.SurfaceAlbedo, DifffuseMeanFreePathInMm);
@@ -376,12 +512,29 @@ void FSubsurfaceProfileTexture::CreateTexture(FRHICommandListImmediate& RHICmdLi
 				// scale from 0..1 to 0..0xffff
 				// scale with 0x10000 and round down to evenly distribute, avoid 0x10000
 
-				uint16* Dest = (uint16*)(DestBuffer + DestStride * y);
-
-				Dest[Pos * 4 + 0] = (uint16)(C.X * FloatScale);
-				Dest[Pos * 4 + 1] = (uint16)(C.Y * FloatScale);
-				Dest[Pos * 4 + 2] = (uint16)(C.Z * FloatScale);
-				Dest[Pos * 4 + 3] = (uint16)(C.W * FloatScale);
+				// RGBA16 UNorm is not supported on android mobile, it falls back to RGBA16F.
+				const bool bUseRGBA16F = GPixelFormats[PF_A16B16G16R16].PlatformFormat == GPixelFormats[PF_FloatRGBA].PlatformFormat;
+
+				if (bUseRGBA16F)
+				{ 
+					float PlatformFloatScale = 1.0f / (FloatScaleInitial - 1.0f);
+
+					uint16* Dest = (uint16*)(DestBuffer + DestStride * y);
+
+					Dest[Pos * 4 + 0] = FFloat16(((uint16)(C.X * FloatScale)) * PlatformFloatScale).Encoded;
+					Dest[Pos * 4 + 1] = FFloat16(((uint16)(C.Y * FloatScale)) * PlatformFloatScale).Encoded;
+					Dest[Pos * 4 + 2] = FFloat16(((uint16)(C.Z * FloatScale)) * PlatformFloatScale).Encoded;
+					Dest[Pos * 4 + 3] = FFloat16(((uint16)(C.W * FloatScale)) * PlatformFloatScale).Encoded;
+				}
+				else
+				{
+					uint16* Dest = (uint16*)(DestBuffer + DestStride * y);
+
+					Dest[Pos * 4 + 0] = (uint16)(C.X * FloatScale);
+					Dest[Pos * 4 + 1] = (uint16)(C.Y * FloatScale);
+					Dest[Pos * 4 + 2] = (uint16)(C.Z * FloatScale);
+					Dest[Pos * 4 + 3] = (uint16)(C.W * FloatScale);
+				}
 			}
 			else
 			{
@@ -501,9 +654,15 @@ FRHITexture* GetSubsurfaceProfileTextureWithFallback()
 	return GSSProfiles ? GSSProfiles->GetShaderResourceRHI() : static_cast<FRHITexture*>(GBlackTexture->TextureRHI);
 }
 
-void UpdateSubsurfaceProfileTexture(FRHICommandListImmediate& RHICmdList)
+FRHITexture* GetSSProfilesPreIntegratedTextureWithFallback()
 {
-	GSubsurfaceProfileTextureObject.GetTexture(RHICmdList);
+	return GSSProfilesPreIntegratedTexture ? GSSProfilesPreIntegratedTexture->GetShaderResourceRHI() : static_cast<FRHITexture*>(GBlackTexture->TextureRHI);
+}
+
+void UpdateSubsurfaceProfileTexture(FRDGBuilder& GraphBuilder, EShaderPlatform ShaderPlatform)
+{
+	GSubsurfaceProfileTextureObject.GetTexture(GraphBuilder.RHICmdList);
+	GSubsurfaceProfileTextureObject.GetSSProfilesPreIntegratedTexture(GraphBuilder, ShaderPlatform);
 }
 
 // ------------------------------------------------------
@@ -549,17 +708,7 @@ void USubsurfaceProfile::PostLoad()
 
 		if (bUpdateBurleyParametersFromSeparable)
 		{
-			MapFallOffColor2SurfaceAlbedoAndDiffuseMeanFreePath(Settings.FalloffColor.R, Settings.SurfaceAlbedo.R, Settings.MeanFreePathColor.R);
-			MapFallOffColor2SurfaceAlbedoAndDiffuseMeanFreePath(Settings.FalloffColor.G, Settings.SurfaceAlbedo.G, Settings.MeanFreePathColor.G);
-			MapFallOffColor2SurfaceAlbedoAndDiffuseMeanFreePath(Settings.FalloffColor.B, Settings.SurfaceAlbedo.B, Settings.MeanFreePathColor.B);
-
-			//Normalize mean free path color and set the corresponding dfmp
-			float MaxMeanFreePathColor = FMath::Max3(Settings.MeanFreePathColor.R, Settings.MeanFreePathColor.G, Settings.MeanFreePathColor.B);
-			if (MaxMeanFreePathColor > 1)
-			{
-				Settings.MeanFreePathColor /= MaxMeanFreePathColor;
-				Settings.MeanFreePathDistance = FMath::Clamp(Settings.ScatterRadius*MaxMeanFreePathColor, 0.1f, 50.0f);	// 50.0f is the ClampMax of MeanFreePathDistance.
-			}
+			ConvertSubsurfaceParametersFromSeparableToBurley(Settings, Settings.SurfaceAlbedo, Settings.MeanFreePathColor, Settings.MeanFreePathDistance);
 		}
 	}
 }
diff --git a/Engine/Source/Runtime/RenderCore/Public/RendererInterface.h b/Engine/Source/Runtime/RenderCore/Public/RendererInterface.h
--- a/Engine/Source/Runtime/RenderCore/Public/RendererInterface.h
+++ b/Engine/Source/Runtime/RenderCore/Public/RendererInterface.h
@@ -106,6 +106,48 @@ public:
 		return NewDesc;
 	}
 
+	/**
+ * Factory function to create 2D array texture description
+ * @param InFlags bit mask combined from elements of ETextureCreateFlags e.g. TexCreate_UAV
+ */
+	static FPooledRenderTargetDesc Create2DArrayDesc(
+		FIntPoint InExtent,
+		EPixelFormat InFormat,
+		const FClearValueBinding& InClearValue,
+		ETextureCreateFlags InFlags,
+		ETextureCreateFlags InTargetableFlags,
+		bool bInForceSeparateTargetAndShaderResource,
+		uint32 InArraySize,
+		uint16 InNumMips = 1,
+		bool InAutowritable = true,
+		bool InCreateRTWriteMask = false,
+		bool InCreateFmask = false)
+	{
+		check(InExtent.X);
+		check(InExtent.Y);
+
+		FPooledRenderTargetDesc NewDesc;
+		NewDesc.ClearValue = InClearValue;
+		NewDesc.Extent = InExtent;
+		NewDesc.Depth = 0;
+		NewDesc.ArraySize = InArraySize;
+		NewDesc.bIsArray = true;
+		NewDesc.bIsCubemap = false;
+		NewDesc.NumMips = InNumMips;
+		NewDesc.NumSamples = 1;
+		NewDesc.Format = InFormat;
+		NewDesc.Flags = InFlags;
+		NewDesc.TargetableFlags = InTargetableFlags;
+		NewDesc.bForceSeparateTargetAndShaderResource = bInForceSeparateTargetAndShaderResource;
+		NewDesc.DebugName = TEXT("UnknownTexture2DArray");
+		NewDesc.AutoWritable = InAutowritable;
+		NewDesc.bCreateRenderTargetWriteMask = InCreateRTWriteMask;
+		NewDesc.bCreateRenderTargetFmask = InCreateFmask;
+		check(NewDesc.Is2DTexture());
+
+		return NewDesc;
+	}
+
 	/**
 	 * Factory function to create 3D texture description
 	 * @param InFlags bit mask combined from elements of ETextureCreateFlags e.g. TexCreate_UAV
diff --git a/Engine/Source/Runtime/Renderer/Private/DeferredShadingRenderer.cpp b/Engine/Source/Runtime/Renderer/Private/DeferredShadingRenderer.cpp
--- a/Engine/Source/Runtime/Renderer/Private/DeferredShadingRenderer.cpp
+++ b/Engine/Source/Runtime/Renderer/Private/DeferredShadingRenderer.cpp
@@ -1881,7 +1881,7 @@ void FDeferredShadingSceneRenderer::Render(FRDGBuilder& GraphBuilder)
 		GSystemTextures.InitializeTextures(GraphBuilder.RHICmdList, FeatureLevel);
 
 		// Force the subsurface profile texture to be updated.
-		UpdateSubsurfaceProfileTexture(GraphBuilder.RHICmdList);
+		UpdateSubsurfaceProfileTexture(GraphBuilder, ShaderPlatform);
 	}
 
 	const FSceneTexturesConfig SceneTexturesConfig = FSceneTexturesConfig::Create(ViewFamily);
diff --git a/Engine/Source/Runtime/Renderer/Private/MobileBasePassRendering.cpp b/Engine/Source/Runtime/Renderer/Private/MobileBasePassRendering.cpp
--- a/Engine/Source/Runtime/Renderer/Private/MobileBasePassRendering.cpp
+++ b/Engine/Source/Runtime/Renderer/Private/MobileBasePassRendering.cpp
@@ -17,6 +17,7 @@
 
 #include "FramePro/FrameProProfiler.h"
 #include "PostProcess/PostProcessPixelProjectedReflectionMobile.h"
+#include "Engine/SubsurfaceProfile.h"
 
 // Changing this causes a full shader recompile
 static TAutoConsoleVariable<int32> CVarMobileDisableVertexFog(
@@ -224,6 +225,16 @@ void SetupMobileBasePassUniformParameters(
 		BasePassParameters.ScreenSpaceShadowMaskTexture = GSystemTextures.GetWhiteDummy(GraphBuilder);
 		BasePassParameters.ScreenSpaceShadowMaskSampler = TStaticSamplerState<SF_Point, AM_Clamp, AM_Clamp, AM_Clamp>::GetRHI();
 	}
+
+	BasePassParameters.SSProfilesTexture = GetSubsurfaceProfileTextureWithFallback();
+	BasePassParameters.SSProfilesSampler = TStaticSamplerState<SF_Point, AM_Clamp, AM_Clamp, AM_Clamp>::GetRHI();
+
+	FIntVector SSProfilesTextureSize = BasePassParameters.SSProfilesTexture->GetSizeXYZ();
+
+	BasePassParameters.SSProfilesTextureSizeAndInvSize = FVector4(SSProfilesTextureSize.X, SSProfilesTextureSize.Y, 1.0f / SSProfilesTextureSize.X, 1.0f / SSProfilesTextureSize.Y);
+
+	BasePassParameters.SSProfilesPreIntegratedTexture = GetSSProfilesPreIntegratedTextureWithFallback();
+	BasePassParameters.SSProfilesPreIntegratedSampler = TStaticSamplerState<SF_Bilinear, AM_Clamp, AM_Clamp, AM_Clamp>::GetRHI();
 }
 
 TRDGUniformBufferRef<FMobileBasePassUniformParameters> CreateMobileBasePassUniformBuffer(
diff --git a/Engine/Source/Runtime/Renderer/Private/MobileBasePassRendering.h b/Engine/Source/Runtime/Renderer/Private/MobileBasePassRendering.h
--- a/Engine/Source/Runtime/Renderer/Private/MobileBasePassRendering.h
+++ b/Engine/Source/Runtime/Renderer/Private/MobileBasePassRendering.h
@@ -27,6 +27,7 @@
 
 BEGIN_GLOBAL_SHADER_PARAMETER_STRUCT(FMobileBasePassUniformParameters, )
 	SHADER_PARAMETER(float, AmbientOcclusionStaticFraction)
+	SHADER_PARAMETER(FVector4, SSProfilesTextureSizeAndInvSize)
 	SHADER_PARAMETER_STRUCT(FFogUniformParameters, Fog)
 	SHADER_PARAMETER_STRUCT(FPlanarReflectionUniformParameters, PlanarReflection) // Single global planar reflection for the forward pass.
 	SHADER_PARAMETER_STRUCT(FMobileSceneTextureUniformParameters, SceneTextures)
@@ -37,6 +38,10 @@ BEGIN_GLOBAL_SHADER_PARAMETER_STRUCT(FMobileBasePassUniformParameters, )
 	SHADER_PARAMETER_SAMPLER(SamplerState, AmbientOcclusionSampler)
 	SHADER_PARAMETER_RDG_TEXTURE(Texture2D, ScreenSpaceShadowMaskTexture)
 	SHADER_PARAMETER_SAMPLER(SamplerState, ScreenSpaceShadowMaskSampler)
+	SHADER_PARAMETER_TEXTURE(Texture2D, SSProfilesTexture)
+	SHADER_PARAMETER_SAMPLER(SamplerState, SSProfilesSampler)
+	SHADER_PARAMETER_TEXTURE(Texture2DArray, SSProfilesPreIntegratedTexture)
+	SHADER_PARAMETER_SAMPLER(SamplerState, SSProfilesPreIntegratedSampler)
 END_GLOBAL_SHADER_PARAMETER_STRUCT()
 
 enum class EMobileBasePass
diff --git a/Engine/Source/Runtime/Renderer/Private/MobileShadingRenderer.cpp b/Engine/Source/Runtime/Renderer/Private/MobileShadingRenderer.cpp
--- a/Engine/Source/Runtime/Renderer/Private/MobileShadingRenderer.cpp
+++ b/Engine/Source/Runtime/Renderer/Private/MobileShadingRenderer.cpp
@@ -692,6 +692,10 @@ void FMobileSceneRenderer::Render(FRDGBuilder& GraphBuilder)
 
 	// Initialize global system textures (pass-through if already initialized).
 	GSystemTextures.InitializeTextures(GraphBuilder.RHICmdList, FeatureLevel);
+
+	// Force the subsurface profile texture to be updated.
+	UpdateSubsurfaceProfileTexture(GraphBuilder, ShaderPlatform);
+
 	FRDGSystemTextures::Create(GraphBuilder);
 
 	FInstanceCullingManager& InstanceCullingManager = *GraphBuilder.AllocObject<FInstanceCullingManager>(Scene->GPUScene.IsEnabled());
