WorldPartition: Add missing ShowCategories for DataLayers category on AInfo(s) that support Data Layers

diff --git a/Engine/Source/Runtime/Engine/Classes/Atmosphere/AtmosphericFog.h b/Engine/Source/Runtime/Engine/Classes/Atmosphere/AtmosphericFog.h
--- a/Engine/Source/Runtime/Engine/Classes/Atmosphere/AtmosphericFog.h
+++ b/Engine/Source/Runtime/Engine/Classes/Atmosphere/AtmosphericFog.h
@@ -12,7 +12,7 @@
  *	A placeable fog actor that simulates atmospheric light scattering  
  *	@see https://docs.unrealengine.com/latest/INT/Engine/Actors/FogEffects/AtmosphericFog/index.html
  */
-UCLASS(showcategories=(Movement, Rendering, Transformation, "Input|MouseInput", "Input|TouchInput"), ClassGroup=Fog, hidecategories=(Info,Object,Input), MinimalAPI, notplaceable)
+UCLASS(showcategories=(Movement, Rendering, Transformation, DataLayers, "Input|MouseInput", "Input|TouchInput"), ClassGroup=Fog, hidecategories=(Info,Object,Input), MinimalAPI, notplaceable)
 class UE_DEPRECATED(4.26, "Please use the SkyAtmosphere actor instead.") AAtmosphericFog : public AInfo
 {
 	GENERATED_UCLASS_BODY()
diff --git a/Engine/Source/Runtime/Engine/Classes/Components/SkyAtmosphereComponent.h b/Engine/Source/Runtime/Engine/Classes/Components/SkyAtmosphereComponent.h
--- a/Engine/Source/Runtime/Engine/Classes/Components/SkyAtmosphereComponent.h
+++ b/Engine/Source/Runtime/Engine/Classes/Components/SkyAtmosphereComponent.h
@@ -267,7 +267,7 @@ protected:
  * A placeable actor that represents a planet atmosphere material and simulates sky and light scattering within it.
  * @see https://docs.unrealengine.com/en-US/Engine/Actors/FogEffects/SkyAtmosphere/index.html
  */
-UCLASS(showcategories = (Movement, Rendering, Transformation, "Input|MouseInput", "Input|TouchInput"), ClassGroup = Fog, hidecategories = (Info, Object, Input), MinimalAPI)
+UCLASS(showcategories = (Movement, Rendering, Transformation, DataLayers, "Input|MouseInput", "Input|TouchInput"), ClassGroup = Fog, hidecategories = (Info, Object, Input), MinimalAPI)
 class ASkyAtmosphere : public AInfo
 {
 	GENERATED_UCLASS_BODY()
diff --git a/Engine/Source/Runtime/Engine/Classes/Components/VolumetricCloudComponent.h b/Engine/Source/Runtime/Engine/Classes/Components/VolumetricCloudComponent.h
--- a/Engine/Source/Runtime/Engine/Classes/Components/VolumetricCloudComponent.h
+++ b/Engine/Source/Runtime/Engine/Classes/Components/VolumetricCloudComponent.h
@@ -170,7 +170,7 @@ private:
  * A placeable actor that represents a participating media material around a planet, e.g. clouds.
  * @see TODO address to the documentation.
  */
-UCLASS(showcategories = (Movement, Rendering, Transformation, "Input|MouseInput", "Input|TouchInput"), ClassGroup = Fog, hidecategories = (Info, Object, Input), MinimalAPI)
+UCLASS(showcategories = (Movement, Rendering, Transformation, DataLayers, "Input|MouseInput", "Input|TouchInput"), ClassGroup = Fog, hidecategories = (Info, Object, Input), MinimalAPI)
 class AVolumetricCloud : public AInfo
 {
 	GENERATED_UCLASS_BODY()
diff --git a/Engine/Source/Runtime/Engine/Classes/Engine/ExponentialHeightFog.h b/Engine/Source/Runtime/Engine/Classes/Engine/ExponentialHeightFog.h
--- a/Engine/Source/Runtime/Engine/Classes/Engine/ExponentialHeightFog.h
+++ b/Engine/Source/Runtime/Engine/Classes/Engine/ExponentialHeightFog.h
@@ -10,7 +10,7 @@
 /**
  * Implements an Actor for exponential height fog.
  */
-UCLASS(showcategories=(Movement, Rendering, Transformation), ClassGroup=Fog, MinimalAPI)
+UCLASS(showcategories=(Movement, Rendering, Transformation, DataLayers), ClassGroup=Fog, MinimalAPI)
 class AExponentialHeightFog
 	: public AInfo
 {
diff --git a/Engine/Source/Runtime/Engine/Classes/Engine/SkyLight.h b/Engine/Source/Runtime/Engine/Classes/Engine/SkyLight.h
--- a/Engine/Source/Runtime/Engine/Classes/Engine/SkyLight.h
+++ b/Engine/Source/Runtime/Engine/Classes/Engine/SkyLight.h
@@ -7,7 +7,7 @@
 #include "GameFramework/Info.h"
 #include "SkyLight.generated.h"
 
-UCLASS(ClassGroup=Lights, hidecategories=(Input,Collision,Replication,Info), showcategories=("Rendering", "Input|MouseInput", "Input|TouchInput"), ComponentWrapperClass, ConversionRoot, Blueprintable)
+UCLASS(ClassGroup=Lights, hidecategories=(Input,Collision,Replication,Info), showcategories=(Rendering, DataLayers, "Input|MouseInput", "Input|TouchInput"), ComponentWrapperClass, ConversionRoot, Blueprintable)
 class ENGINE_API ASkyLight : public AInfo
 {
 	GENERATED_UCLASS_BODY()
diff --git a/Engine/Source/Runtime/Engine/Classes/Engine/WindDirectionalSource.h b/Engine/Source/Runtime/Engine/Classes/Engine/WindDirectionalSource.h
--- a/Engine/Source/Runtime/Engine/Classes/Engine/WindDirectionalSource.h
+++ b/Engine/Source/Runtime/Engine/Classes/Engine/WindDirectionalSource.h
@@ -7,7 +7,7 @@
 #include "WindDirectionalSource.generated.h"
 
 /** Actor that provides a directional wind source. Only affects SpeedTree assets. */
-UCLASS(ClassGroup=Wind, showcategories=(Rendering, Transformation))
+UCLASS(ClassGroup=Wind, showcategories=(Rendering, Transformation, DataLayers))
 class ENGINE_API AWindDirectionalSource : public AInfo
 {
 	GENERATED_UCLASS_BODY()
