Fix Embedded LevelInstance Actors Cell Location

diff --git a/Engine/Source/Runtime/Engine/Private/WorldPartition/RuntimeSpatialHash/RuntimeSpatialHashGridHelper.cpp b/Engine/Source/Runtime/Engine/Private/WorldPartition/RuntimeSpatialHash/RuntimeSpatialHashGridHelper.cpp
--- a/Engine/Source/Runtime/Engine/Private/WorldPartition/RuntimeSpatialHash/RuntimeSpatialHashGridHelper.cpp
+++ b/Engine/Source/Runtime/Engine/Private/WorldPartition/RuntimeSpatialHash/RuntimeSpatialHashGridHelper.cpp
@@ -169,7 +169,7 @@ FSquare2DGridHelper GetPartitionedActors(const UWorldPartition* WorldPartition,
 			check(ActorCluster->Actors.Num() == 1);
 			const FGuid& ActorGuid = *ActorCluster->Actors.CreateConstIterator();
 			FActorInstance ActorInstance(ActorGuid, ClusterInstance->ContainerInstance);
-			if (ensure(PartitionedActors.GetLowestLevel().GetCellCoords(FVector2D(ActorInstance.GetActorDescView().GetOrigin()), CellCoords)))
+			if (ensure(PartitionedActors.GetLowestLevel().GetCellCoords(FVector2D(ActorInstance.GetOrigin()), CellCoords)))
 			{
 				PartitionedActors.GetLowestLevel().GetCell(CellCoords).AddActor(MoveTemp(ActorInstance), ClusterInstance->DataLayers);
 			}
diff --git a/Engine/Source/Runtime/Engine/Private/WorldPartition/WorldPartitionActorCluster.cpp b/Engine/Source/Runtime/Engine/Private/WorldPartition/WorldPartitionActorCluster.cpp
--- a/Engine/Source/Runtime/Engine/Private/WorldPartition/WorldPartitionActorCluster.cpp
+++ b/Engine/Source/Runtime/Engine/Private/WorldPartition/WorldPartitionActorCluster.cpp
@@ -234,7 +234,8 @@ bool FActorInstance::ShouldStripFromStreaming() const
 
 FVector FActorInstance::GetOrigin() const
 {
-	return ContainerInstance->Transform.TransformPosition(GetActorDescView().GetOrigin());
+	// Embedded mode means all actors of the container share the same location so that they get placed in the same runtime cell
+	return ContainerInstance->ClusterMode == EContainerClusterMode::Embedded ? ContainerInstance->Transform.GetLocation() : ContainerInstance->Transform.TransformPosition(GetActorDescView().GetOrigin());
 }
 
 const FWorldPartitionActorDescView& FActorInstance::GetActorDescView() const
