Invalidate all viewport hit-proxies when an async HISM tree build finishes

This is a workaround for a race condition where the hit-proxies may be updated while an async build is in progress, leading to incomplete or incorrect hit-proxy data.

diff --git a/Engine/Source/Editor/UnrealEd/Classes/Editor/UnrealEdEngine.h b/Engine/Source/Editor/UnrealEd/Classes/Editor/UnrealEdEngine.h
--- a/Engine/Source/Editor/UnrealEd/Classes/Editor/UnrealEdEngine.h
+++ b/Engine/Source/Editor/UnrealEd/Classes/Editor/UnrealEdEngine.h
@@ -25,6 +25,7 @@ class FViewport;
 class IEngineLoop;
 class ITargetPlatform;
 class UPrimitiveComponent;
+class UHierarchicalInstancedStaticMeshComponent;
 class UTexture2D;
 class UUnrealEdOptions;
 class USelection;
@@ -926,6 +927,9 @@ protected:
 	/** Called when the element selection set associated with the global editor selection changes */
 	void OnEditorElementSelectionChanged(const UTypedElementSelectionSet* SelectionSet);
 
+	/** Called when a HISM tree has finished building */
+	void OnHISMTreeBuilt(UHierarchicalInstancedStaticMeshComponent* Component, bool bWasAsyncBuild);
+
 	/** The package auto-saver instance used by the editor */
 	TUniquePtr<IPackageAutoSaver> PackageAutoSaver;
 
diff --git a/Engine/Source/Editor/UnrealEd/Private/UnrealEdEngine.cpp b/Engine/Source/Editor/UnrealEd/Private/UnrealEdEngine.cpp
--- a/Engine/Source/Editor/UnrealEd/Private/UnrealEdEngine.cpp
+++ b/Engine/Source/Editor/UnrealEd/Private/UnrealEdEngine.cpp
@@ -43,6 +43,7 @@
 #include "EdMode.h"
 #include "PropertyEditorModule.h"
 #include "LevelEditor.h"
+#include "InstancedStaticMeshDelegates.h"
 #include "Interfaces/IMainFrameModule.h"
 #include "Settings/EditorLoadingSavingSettingsCustomization.h"
 #include "Settings/GameMapsSettingsCustomization.h"
@@ -121,6 +122,8 @@ void UUnrealEdEngine::Init(IEngineLoop* InEngineLoop)
 	FEditorSupportDelegates::PreWindowsMessage.AddUObject(this, &UUnrealEdEngine::OnPreWindowsMessage);
 	FEditorSupportDelegates::PostWindowsMessage.AddUObject(this, &UUnrealEdEngine::OnPostWindowsMessage);
 
+	FHierarchicalInstancedStaticMeshDelegates::OnTreeBuilt.AddUObject(this, &UUnrealEdEngine::OnHISMTreeBuilt);
+
 	USelection::SelectionChangedEvent.AddUObject(this, &UUnrealEdEngine::OnEditorSelectionChanged);
 	USelection::SelectionElementSelectionPtrChanged.AddUObject(this, &UUnrealEdEngine::OnEditorElementSelectionPtrChanged);
 
@@ -892,6 +895,22 @@ bool UUnrealEdEngine::PreferToStreamLevelsInPIE() const
 	return GetDefault<ULevelEditorPlaySettings>()->bPreferToStreamLevelsInPIE;
 }
 
+void UUnrealEdEngine::OnHISMTreeBuilt(UHierarchicalInstancedStaticMeshComponent* Component, bool bWasAsyncBuild)
+{
+	if (bWasAsyncBuild)
+	{
+		// Async tree builds require that hit-proxies be updated when the build has finished (as it generates new hit-proxies).
+		// Sadly we don't know which viewport needs updating, so we just have to update all of them.
+		for (FEditorViewportClient* ViewportClient : GetAllViewportClients())
+		{
+			if (ViewportClient && ViewportClient->Viewport)
+			{
+				ViewportClient->RequestInvalidateHitProxy(ViewportClient->Viewport);
+			}
+		}
+	}
+}
+
 void UUnrealEdEngine::RedrawLevelEditingViewports(bool bInvalidateHitProxies)
 {
 	// Redraw Slate based viewports
diff --git a/Engine/Source/Runtime/Engine/Classes/Components/HierarchicalInstancedStaticMeshComponent.h b/Engine/Source/Runtime/Engine/Classes/Components/HierarchicalInstancedStaticMeshComponent.h
--- a/Engine/Source/Runtime/Engine/Classes/Components/HierarchicalInstancedStaticMeshComponent.h
+++ b/Engine/Source/Runtime/Engine/Classes/Components/HierarchicalInstancedStaticMeshComponent.h
@@ -269,7 +269,7 @@ public:
 protected:
 	void BuildTree();
 	void BuildTreeAsync();
-	void ApplyBuildTree(FClusterBuilder& Builder);
+	void ApplyBuildTree(FClusterBuilder& Builder, const bool bWasAsyncBuild);
 	void ApplyEmpty();
 	void SetPerInstanceLightMapAndEditorData(FStaticMeshInstanceData& PerInstanceData, const TArray<TRefCountPtr<HHitProxy>>& HitProxies);
 
diff --git a/Engine/Source/Runtime/Engine/Private/HierarchicalInstancedStaticMesh.cpp b/Engine/Source/Runtime/Engine/Private/HierarchicalInstancedStaticMesh.cpp
--- a/Engine/Source/Runtime/Engine/Private/HierarchicalInstancedStaticMesh.cpp
+++ b/Engine/Source/Runtime/Engine/Private/HierarchicalInstancedStaticMesh.cpp
@@ -846,6 +846,8 @@ struct FFoliageElementParams;
 struct FFoliageRenderInstanceParams;
 struct FFoliageCullInstanceParams;
 
+FHierarchicalInstancedStaticMeshDelegates::FOnTreeBuilt FHierarchicalInstancedStaticMeshDelegates::OnTreeBuilt;
+
 class FHierarchicalStaticMeshSceneProxy final : public FInstancedStaticMeshSceneProxy
 {
 	TSharedRef<TArray<FClusterNode>, ESPMode::ThreadSafe> ClusterTreePtr;
@@ -2719,7 +2721,7 @@ void UHierarchicalInstancedStaticMeshComponent::BuildTree()
 		FClusterBuilder Builder(InstanceTransforms, PerInstanceSMCustomData, NumCustomDataFloats, GetStaticMesh()->GetBounds().GetBox(), DesiredInstancesPerLeaf(), CurrentDensityScaling, InstancingRandomSeed, PerInstanceSMData.Num() > 0);
 		Builder.BuildTreeAndBuffer();
 
-		ApplyBuildTree(Builder);
+		ApplyBuildTree(Builder, /*bWasAsyncBuild*/false);
 	}
 	else
 	{
@@ -2821,7 +2823,7 @@ void UHierarchicalInstancedStaticMeshComponent::ApplyBuildTreeAsync(ENamedThread
 	}
 
 	// Completed the build
-	ApplyBuildTree(Builder.Get());
+	ApplyBuildTree(Builder.Get(), /*bWasAsyncBuild*/true);
 }
 
 void UHierarchicalInstancedStaticMeshComponent::ApplyEmpty()
@@ -2843,10 +2845,11 @@ void UHierarchicalInstancedStaticMeshComponent::ApplyEmpty()
 		PerInstanceRenderData->UpdateFromPreallocatedData(*BuiltInstanceData);
 		PerInstanceRenderData->HitProxies.Empty();
 		MarkRenderStateDirty();
+		FHierarchicalInstancedStaticMeshDelegates::OnTreeBuilt.Broadcast(this, /*bWasAsyncBuild*/false);
 	}
 }
 
-void UHierarchicalInstancedStaticMeshComponent::ApplyBuildTree(FClusterBuilder& Builder)
+void UHierarchicalInstancedStaticMeshComponent::ApplyBuildTree(FClusterBuilder& Builder, const bool bWasAsyncBuild)
 {
 	bIsOutOfDate = false;
 	InstanceUpdateCmdBuffer.Reset();
@@ -2896,6 +2899,8 @@ void UHierarchicalInstancedStaticMeshComponent::ApplyBuildTree(FClusterBuilder&
 	FlushAccumulatedNavigationUpdates();
 	PostBuildStats();
 	MarkRenderStateDirty();
+
+	FHierarchicalInstancedStaticMeshDelegates::OnTreeBuilt.Broadcast(this, bWasAsyncBuild);
 }
 
 bool UHierarchicalInstancedStaticMeshComponent::BuildTreeIfOutdated(bool Async, bool ForceUpdate)
diff --git a/Engine/Source/Runtime/Engine/Public/InstancedStaticMeshDelegates.h b/Engine/Source/Runtime/Engine/Public/InstancedStaticMeshDelegates.h
--- a/Engine/Source/Runtime/Engine/Public/InstancedStaticMeshDelegates.h
+++ b/Engine/Source/Runtime/Engine/Public/InstancedStaticMeshDelegates.h
@@ -6,6 +6,7 @@
 #include "Containers/ArrayView.h"
 
 class UInstancedStaticMeshComponent;
+class UHierarchicalInstancedStaticMeshComponent;
 
 struct ENGINE_API FInstancedStaticMeshDelegates
 {
@@ -49,3 +50,12 @@ struct ENGINE_API FInstancedStaticMeshDelegates
 	DECLARE_MULTICAST_DELEGATE_TwoParams(FOnInstanceIndexUpdated, UInstancedStaticMeshComponent* /*Component*/, TArrayView<const FInstanceIndexUpdateData> /*IndexUpdates*/);
 	static FOnInstanceIndexUpdated OnInstanceIndexUpdated;
 };
+
+struct ENGINE_API FHierarchicalInstancedStaticMeshDelegates
+{
+	/**
+	 * Delegate called when the tree of a HISM has been (re)built.
+	 */
+	DECLARE_MULTICAST_DELEGATE_TwoParams(FOnTreeBuilt, UHierarchicalInstancedStaticMeshComponent* /*Component*/, bool /*bWasAsyncBuild*/);
+	static FOnTreeBuilt OnTreeBuilt;
+};
