Fix a collect garbage issue when loading skeletalmesh DDC. The repro step are simple but not repro easily, The issue happen when we load a map which load all asset and then unload the current map which trig a GC. At this point its possible all async SK cache not done.

diff --git a/Engine/Source/Developer/MeshBuilder/Private/SkeletalMeshBuilder.cpp b/Engine/Source/Developer/MeshBuilder/Private/SkeletalMeshBuilder.cpp
--- a/Engine/Source/Developer/MeshBuilder/Private/SkeletalMeshBuilder.cpp
+++ b/Engine/Source/Developer/MeshBuilder/Private/SkeletalMeshBuilder.cpp
@@ -68,7 +68,7 @@ bool FSkeletalMeshBuilder::Build(const FSkeletalMeshBuildParameters& SkeletalMes
 	check(SkeletalMesh->GetImportedModel()->LODModels.IsValidIndex(LODIndex));
 	check(SkeletalMesh->GetLODInfo(LODIndex) != nullptr);
 	
-	FSkeletalMeshLODInfo* LODInfo = SkeletalMesh->GetLODInfo(LODIndex);
+	const FSkeletalMeshLODInfo* LODInfo = SkeletalMesh->GetLODInfo(LODIndex);
 	//We want to backup in case the LODModel is regenerated, this data is use to validate in the UI if the ddc must be rebuild
 	const FString BackupBuildStringID = SkeletalMesh->GetImportedModel()->LODModels[LODIndex].BuildStringID;
 
diff --git a/Engine/Source/Runtime/Engine/Classes/Engine/SkeletalMesh.h b/Engine/Source/Runtime/Engine/Classes/Engine/SkeletalMesh.h
--- a/Engine/Source/Runtime/Engine/Classes/Engine/SkeletalMesh.h
+++ b/Engine/Source/Runtime/Engine/Classes/Engine/SkeletalMesh.h
@@ -29,6 +29,7 @@
 #include "SkeletalMeshLODSettings.h"
 #include "Animation/NodeMappingProviderInterface.h"
 #include "Animation/SkinWeightProfile.h"
+#include "Animation/MorphTarget.h"
 #include "Engine/StreamableRenderAsset.h"
 
 #include "SkeletalMesh.generated.h"
@@ -36,7 +37,6 @@
 class UAnimInstance;
 class UAssetUserData;
 class UBodySetup;
-class UMorphTarget;
 class USkeletalMeshSocket;
 class USkeleton;
 class UClothingAssetBase;
@@ -121,14 +121,40 @@ enum class ESkeletalMeshAsyncProperties : uint64
 
 ENUM_CLASS_FLAGS(ESkeletalMeshAsyncProperties);
 
-class FSkeletalMeshPostLoadContext
+struct FFinishBuildInternalData
+{
+public:
+	/** Editor MorphTargets data to set. */
+	bool bApplyMorphTargetsData = false;
+	TMap<FName, TArray<FMorphTargetLODModel> > MorphLODModelsPerTargetName;
+
+	void ApplyEditorData(USkeletalMesh* SkeletalMesh) const;
+
+private:
+	void ApplyMorphTargetsEditorData(USkeletalMesh* SkeletalMesh) const;
+};
+
+class FSkeletalMeshCompilationContext
+{
+public:
+	FSkeletalMeshCompilationContext() = default;
+	// Non-copyable
+	FSkeletalMeshCompilationContext(const FSkeletalMeshCompilationContext&) = delete;
+	FSkeletalMeshCompilationContext& operator=(const FSkeletalMeshCompilationContext&) = delete;
+	// Movable
+	FSkeletalMeshCompilationContext(FSkeletalMeshCompilationContext&&) = default;
+	FSkeletalMeshCompilationContext& operator=(FSkeletalMeshCompilationContext&&) = default;
+	FFinishBuildInternalData FinishBuildInternalData;
+};
+
+class FSkeletalMeshPostLoadContext : public FSkeletalMeshCompilationContext
 {
 public:
 	bool bIsPreSkeletalMeshBuildRefactor = false;
 	bool bHasCachedDerivedData = false;
 };
 
-class FSkeletalMeshBuildContext
+class FSkeletalMeshBuildContext : public FSkeletalMeshCompilationContext
 {
 public:
 	TUniquePtr<FSkinnedMeshComponentRecreateRenderStateContext> RecreateRenderStateContext;
@@ -2784,7 +2810,7 @@ private:
 
 #if WITH_EDITOR
 	/** Generate SkeletalMeshRenderData from ImportedModel */
-	void CacheDerivedData();
+	void CacheDerivedData(FSkeletalMeshCompilationContext* ContextPtr);
 
 	/**
 	 * Initial step for the building process - Can't be done in parallel.
@@ -2800,6 +2826,12 @@ private:
 	 * Complete the building process - Can't be done in parallel.
 	 */
 	void FinishBuildInternal(FSkeletalMeshBuildContext& Context);
+
+	/**
+	 * Copy build/load context result data to the skeletalmesh member on the game thread - Can't be done in parallel.
+	 */
+	void ApplyFinishBuildInternalData(FSkeletalMeshCompilationContext* ContextPtr);
+	
 #endif
 
 	/** Utility function to help with building the combined socket list */
diff --git a/Engine/Source/Runtime/Engine/Private/SkeletalMesh.cpp b/Engine/Source/Runtime/Engine/Private/SkeletalMesh.cpp
--- a/Engine/Source/Runtime/Engine/Private/SkeletalMesh.cpp
+++ b/Engine/Source/Runtime/Engine/Private/SkeletalMesh.cpp
@@ -104,6 +104,7 @@
 #include "SkeletalDebugRendering.h"
 #include "Misc/RuntimeErrors.h"
 #include "PlatformInfo.h"
+#include "Async/Async.h"
 
 #if RHI_RAYTRACING
 #include "RayTracingInstance.h"
@@ -1568,11 +1569,37 @@ bool USkeletalMesh::IsReadyForFinishDestroy()
 }
 
 #if WITH_EDITOR
+void CacheRunningPlatform(USkeletalMesh* Mesh)
+{
+	//Cache the running platform, dcc should be valid so it will be fast
+	ITargetPlatform* RunningPlatform = GetTargetPlatformManagerRef().GetRunningTargetPlatform();
+	check(RunningPlatform);
+	FSkeletalMeshBuildContext Context;
+	FSkeletalMeshRenderData RunningPlatformRenderData;
+	RunningPlatformRenderData.Cache(RunningPlatform, Mesh, &Context);
+
+	auto ApplyContext = [Mesh, Context = MoveTemp(Context)]()
+	{
+		check(IsInGameThread());
+		Context.FinishBuildInternalData.ApplyEditorData(Mesh);
+	};
+
+	if (IsInGameThread())
+	{
+		ApplyContext();
+	}
+	else
+	{
+		//Always apply context data on the game thread
+		Async(EAsyncExecution::TaskGraphMainThread, MoveTemp(ApplyContext));
+	}
+}
+
 FString BuildSkeletalMeshDerivedDataKey(const ITargetPlatform* TargetPlatform, USkeletalMesh* SkelMesh);
 
 static FSkeletalMeshRenderData& GetPlatformSkeletalMeshRenderData(USkeletalMesh* Mesh, const ITargetPlatform* TargetPlatform)
 {
-	FString PlatformDerivedDataKey = BuildSkeletalMeshDerivedDataKey(TargetPlatform, Mesh);
+	const FString PlatformDerivedDataKey = BuildSkeletalMeshDerivedDataKey(TargetPlatform, Mesh);
 	FSkeletalMeshRenderData* PlatformRenderData = Mesh->GetResourceForRendering();
 	if (Mesh->GetOutermost()->bIsCookedForEditor)
 	{
@@ -1589,7 +1616,10 @@ static FSkeletalMeshRenderData& GetPlatformSkeletalMeshRenderData(USkeletalMesh*
 	{
 		// Cache render data for this platform and insert it in to the linked list.
 		PlatformRenderData = new FSkeletalMeshRenderData();
-		PlatformRenderData->Cache(TargetPlatform, Mesh);
+		FSkeletalMeshBuildContext Context;
+		PlatformRenderData->Cache(TargetPlatform, Mesh, &Context);
+		//No need to apply the context editor data to the mesh when the target is not the running target
+		
 		check(PlatformRenderData->DerivedDataKey == PlatformDerivedDataKey);
 		Swap(PlatformRenderData->NextCachedRenderData, Mesh->GetResourceForRendering()->NextCachedRenderData);
 		Mesh->GetResourceForRendering()->NextCachedRenderData = TUniquePtr<FSkeletalMeshRenderData>(PlatformRenderData);
@@ -1600,15 +1630,14 @@ static FSkeletalMeshRenderData& GetPlatformSkeletalMeshRenderData(USkeletalMesh*
 			//Normally it should just take back the ddc for the running platform, since the ddc was cache when we have load the asset to cook it.
 			ITargetPlatform* RunningPlatform = GetTargetPlatformManagerRef().GetRunningTargetPlatform();
 			check(RunningPlatform);
-			if (RunningPlatform != TargetPlatform)
+			//The running platform was loaded before getting here, so we should not have an null PlatformRenderData if someone call
+			//GetPlatformSkeletalMeshRenderData with a platform using the same ddc key
+			const FString RunningPlatformDerivedDataKey = BuildSkeletalMeshDerivedDataKey(RunningPlatform, Mesh);
+			if (RunningPlatformDerivedDataKey != PlatformDerivedDataKey)
 			{
-				FString RunningPlatformDerivedDataKey = BuildSkeletalMeshDerivedDataKey(RunningPlatform, Mesh);
-				if (RunningPlatformDerivedDataKey != PlatformDerivedDataKey)
-				{
-					FSkeletalMeshRenderData RunningPlatformRenderData;
-					RunningPlatformRenderData.Cache(RunningPlatform, Mesh);
-					check(RunningPlatformRenderData.DerivedDataKey == RunningPlatformDerivedDataKey);
-				}
+				//Force a cache of the running platform to put back the running platform editor data (FSkeletalMeshLODModel and MorphTargets)
+				//Editor data do not have "DDC key linked list". Maybe we should move this data in the FSkeletalMeshRenderData.
+				CacheRunningPlatform(Mesh);
 			}
 		}
 	}
@@ -1913,6 +1942,72 @@ uint32 USkeletalMesh::GetVertexBufferFlags() const
 
 #if WITH_EDITOR
 
+void FFinishBuildInternalData::ApplyMorphTargetsEditorData(USkeletalMesh* SkeletalMesh) const
+{
+	check(IsInGameThread());
+	//Return if we do not need to apply data
+	if (!bApplyMorphTargetsData)
+	{
+		return;
+	}
+
+	FSkeletalMeshModel* SkelMeshModel = SkeletalMesh->GetImportedModel();
+	check(SkelMeshModel);
+
+	TMap<FName, UMorphTarget*> ExistingMorphTargets;
+	for (UMorphTarget* MorphTarget : SkeletalMesh->GetMorphTargets())
+	{
+		ExistingMorphTargets.Add(MorphTarget->GetFName(), MorphTarget);
+	}
+
+	int32 MorphTargetNumber = MorphLODModelsPerTargetName.Num();
+	TArray<UMorphTarget*> ToDeleteMorphTargets;
+	ToDeleteMorphTargets.Append(SkeletalMesh->GetMorphTargets());
+	SkeletalMesh->GetMorphTargets().Empty();
+	//Rebuild the MorphTarget object
+	for(TPair<FName, TArray<FMorphTargetLODModel>> TargetNameAndMorphLODModels : MorphLODModelsPerTargetName)
+	{
+		FName MorphTargetName = TargetNameAndMorphLODModels.Key;
+		UMorphTarget* MorphTarget = ExistingMorphTargets.FindRef(MorphTargetName);
+		if (!MorphTarget)
+		{
+			MorphTarget = NewObject<UMorphTarget>(SkeletalMesh, MorphTargetName);
+			check(MorphTarget);
+		}
+		else
+		{
+			ToDeleteMorphTargets.Remove(MorphTarget);
+		}
+		MorphTarget->MorphLODModels.Empty();
+		SkeletalMesh->GetMorphTargets().Add(MorphTarget);
+		const TArray<FMorphTargetLODModel>& MorphTargetLODModels = TargetNameAndMorphLODModels.Value;
+		int32 MorphLODModelNumber = MorphTargetLODModels.Num();
+		MorphTarget->MorphLODModels.AddDefaulted(MorphLODModelNumber);
+		for (int32 MorphDataIndex = 0; MorphDataIndex < MorphLODModelNumber; ++MorphDataIndex)
+		{
+			MorphTarget->MorphLODModels[MorphDataIndex] = MorphTargetLODModels[MorphDataIndex];
+		}
+	}
+	//Rebuild the mapping and rehook the curve data
+	SkeletalMesh->InitMorphTargets();
+
+	for (UMorphTarget* ToDeleteMorphTarget : ToDeleteMorphTargets)
+	{
+		ToDeleteMorphTarget->BaseSkelMesh = nullptr;
+		ToDeleteMorphTarget->MorphLODModels.Empty();
+
+		//Move the unused asset in the transient package and mark it pending kill
+		ToDeleteMorphTarget->Rename(nullptr, GetTransientPackage(), REN_ForceNoResetLoaders | REN_DoNotDirty | REN_DontCreateRedirectors | REN_NonTransactional);
+		ToDeleteMorphTarget->MarkPendingKill();
+	}
+}
+
+void FFinishBuildInternalData::ApplyEditorData(USkeletalMesh* SkeletalMesh) const
+{
+	check(IsInGameThread());
+	ApplyMorphTargetsEditorData(SkeletalMesh);
+}
+
 thread_local const USkeletalMesh* FSkeletalMeshAsyncBuildScope::SkeletalMeshBeingAsyncCompiled = nullptr;
 
 FSkeletalMeshAsyncBuildScope::FSkeletalMeshAsyncBuildScope(const USkeletalMesh* SkeletalMesh)
@@ -2087,7 +2182,7 @@ void USkeletalMesh::ExecuteBuildInternal(FSkeletalMeshBuildContext& Context)
 	FSkeletalMeshAsyncBuildScope AsyncBuildScope(this);
 
 	// rebuild render data from imported model
-	CacheDerivedData();
+	CacheDerivedData(&Context);
 
 	// Do not need to fix up 16-bit UVs here, as we assume all editor platforms support them.
 	ensure(GVertexElementTypeSupport.IsSupported(VET_Half2));
@@ -2099,12 +2194,22 @@ void USkeletalMesh::ExecuteBuildInternal(FSkeletalMeshBuildContext& Context)
 	UpdateGenerateUpToData();
 }
 
+void USkeletalMesh::ApplyFinishBuildInternalData(FSkeletalMeshCompilationContext* ContextPtr)
+{
+	//We cannot execute this code outside of the game thread
+	checkf(IsInGameThread(), TEXT("Cannot execute function USkeletalMesh::ApplyFinishBuildInternalData asynchronously. Asset: %s"), *this->GetFullName());
+	check(ContextPtr);
+	ContextPtr->FinishBuildInternalData.ApplyEditorData(this);
+}
+
 void USkeletalMesh::FinishBuildInternal(FSkeletalMeshBuildContext& Context)
 {
 	TRACE_CPUPROFILER_EVENT_SCOPE(USkeletalMesh::FinishBuildInternal);
 
 	ReleaseAsyncProperty();
 
+	ApplyFinishBuildInternalData(&Context);
+	
 	// Note: meshes can be built during automated importing.  We should not create resources in that case
 	// as they will never be released when this object is deleted
 	if (FApp::CanEverRender())
@@ -3239,7 +3344,7 @@ void USkeletalMesh::ExecutePostLoadInternal(FSkeletalMeshPostLoadContext& Contex
 
 		if (GetResourceForRendering() == nullptr)
 		{
-			CacheDerivedData();
+			CacheDerivedData(&Context);
 			Context.bHasCachedDerivedData = true;
 		}
 	}
@@ -3301,9 +3406,9 @@ void USkeletalMesh::FinishPostLoadInternal(FSkeletalMeshPostLoadContext& Context
 	}
 #endif // WITH_EDITOR
 
-	// init morph targets. 
-	// should do this before InitResource, so that we clear invalid morphtargets
-	InitMorphTargets();
+#if WITH_EDITOR
+	ApplyFinishBuildInternalData(&Context);
+#endif
 
 	// initialize rendering resources
 	if (FApp::CanEverRender())
@@ -4297,9 +4402,10 @@ namespace InternalSkeletalMeshHelper
 	}
 } //namespace InternalSkeletalMeshHelper
 
-void USkeletalMesh::CacheDerivedData()
+void USkeletalMesh::CacheDerivedData(FSkeletalMeshCompilationContext* ContextPtr)
 {
 	TRACE_CPUPROFILER_EVENT_SCOPE(USkeletalMesh::CacheDerivedData);
+	check(ContextPtr);
 
 	// Cache derived data for the running platform.
 	ITargetPlatform* RunningPlatform = GetTargetPlatformManagerRef().GetRunningTargetPlatform();
@@ -4316,7 +4422,7 @@ void USkeletalMesh::CacheDerivedData()
 	TMap<int32, TArray<int16>> BackupSectionsPerLOD;
 	InternalSkeletalMeshHelper::CreateLodMaterialMapBackup(this, BackupSectionsPerLOD);
 
-	GetSkeletalMeshRenderData()->Cache(RunningPlatform, this);
+	GetSkeletalMeshRenderData()->Cache(RunningPlatform, this, ContextPtr);
 
 	InternalSkeletalMeshHelper::RestoreLodMaterialMapBackup(this, BackupSectionsPerLOD);
 }
diff --git a/Engine/Source/Runtime/Engine/Private/SkeletalMeshRenderData.cpp b/Engine/Source/Runtime/Engine/Private/SkeletalMeshRenderData.cpp
--- a/Engine/Source/Runtime/Engine/Private/SkeletalMeshRenderData.cpp
+++ b/Engine/Source/Runtime/Engine/Private/SkeletalMeshRenderData.cpp
@@ -322,10 +322,10 @@ FString FSkeletalMeshRenderData::GetDerivedDataKey(const ITargetPlatform* Target
 	return BuildSkeletalMeshDerivedDataKey(TargetPlatform, Owner);
 }
 
-void FSkeletalMeshRenderData::Cache(const ITargetPlatform* TargetPlatform, USkeletalMesh* Owner)
+void FSkeletalMeshRenderData::Cache(const ITargetPlatform* TargetPlatform, USkeletalMesh* Owner, FSkeletalMeshCompilationContext* ContextPtr)
 {
 	check(Owner);
-
+	check(ContextPtr);
 
 	check(LODRenderData.Num() == 0); // Should only be called on new, empty RenderData
 	check(TargetPlatform);
@@ -363,7 +363,7 @@ void FSkeletalMeshRenderData::Cache(const ITargetPlatform* TargetPlatform, USkel
 			
 			FLargeMemoryReader Ar(DerivedData.GetData(), DerivedData.Num(), ELargeMemoryReaderFlags::Persistent);
 
-			//With skeletal mesh build refactor we serialize the LODModel sections into the DDC
+			//With skeletal mesh build refactor we serialize the LODModel data into the DDC
 			//We need to store those so we do not have to rerun the reduction to make them up to date
 			//with the serialize renderdata. This allow to use DDC when changing the reduction settings.
 			//The old workflow has to reduce the LODModel before getting the render data DDC.
@@ -372,80 +372,35 @@ void FSkeletalMeshRenderData::Cache(const ITargetPlatform* TargetPlatform, USkel
 				FSkeletalMeshModel* SkelMeshModel = Owner->GetImportedModel();
 				check(SkelMeshModel);
 
-				TMap<FName, UMorphTarget*> ExistingMorphTargets;
-				for (UMorphTarget* MorphTarget : Owner->GetMorphTargets())
-				{
-					ExistingMorphTargets.Add(MorphTarget->GetFName(), MorphTarget);
-				}
-
-				int32 MorphTargetNumber = 0;
-				Ar << MorphTargetNumber;
-				TArray<UMorphTarget*> ToDeleteMorphTargets;
-				ToDeleteMorphTargets.Append(Owner->GetMorphTargets());
-				Owner->GetMorphTargets().Empty();
-				//Rebuild the MorphTarget object
-				//We cannot serialize directly the UMorphTarget with a FMemoryArchive. This is not supported.
-				for (int32 MorphTargetIndex = 0; MorphTargetIndex < MorphTargetNumber; ++MorphTargetIndex)
+				//Get the morph target data, we put it in the compilation context to apply them in the game thread before the InitResources
 				{
-					FName MorphTargetName = NAME_None;
-					Ar << MorphTargetName;
-					UMorphTarget* MorphTarget = ExistingMorphTargets.FindRef(MorphTargetName);
-					if (!MorphTarget)
-					{
-						MorphTarget = NewObject<UMorphTarget>(Owner, MorphTargetName);
-						check(MorphTarget);
-					}
-					else
-					{
-						ToDeleteMorphTargets.Remove(MorphTarget);
-					}
-					MorphTarget->MorphLODModels.Empty();
-					Owner->GetMorphTargets().Add(MorphTarget);
-					check(MorphTargetIndex == Owner->GetMorphTargets().Num() - 1);
-					int32 MorphLODModelNumber = 0;
-					Ar << MorphLODModelNumber;
-					MorphTarget->MorphLODModels.AddDefaulted(MorphLODModelNumber);
-					for (int32 MorphDataIndex = 0; MorphDataIndex < MorphLODModelNumber; ++MorphDataIndex)
+					if (!ensureMsgf(!ContextPtr->FinishBuildInternalData.bApplyMorphTargetsData, TEXT("Error in FSkeletalMeshRenderData::Cache. The compilation context morph targets data was already set.")))
 					{
-						Ar << MorphTarget->MorphLODModels[MorphDataIndex];
+						ContextPtr->FinishBuildInternalData.MorphLODModelsPerTargetName.Empty();
+						ContextPtr->FinishBuildInternalData.bApplyMorphTargetsData = false;
 					}
-				}
-				//Rebuild the mapping and rehook the curve data
-				Owner->InitMorphTargets();
-
-				for (UMorphTarget* ToDeleteMorphTarget : ToDeleteMorphTargets)
-				{
-					ToDeleteMorphTarget->BaseSkelMesh = nullptr;
-					ToDeleteMorphTarget->MorphLODModels.Empty();
 
-					auto DeleteObject = 
-						[ObjectToDelete = TWeakObjectPtr<UMorphTarget>(ToDeleteMorphTarget)]()
-						{
-							//Move the unused asset in the transient package and mark it pending kill
-							if (ObjectToDelete.IsValid())
-							{
-								ObjectToDelete->Rename(nullptr, GetTransientPackage(), REN_ForceNoResetLoaders | REN_DoNotDirty | REN_DontCreateRedirectors | REN_NonTransactional);
-								ObjectToDelete->MarkPendingKill();
-							}
-						};
-
-					if (IsInGameThread())
+					int32 MorphTargetNumber = 0;
+					Ar << MorphTargetNumber;
+					ContextPtr->FinishBuildInternalData.MorphLODModelsPerTargetName.Reserve(MorphTargetNumber);
+					ContextPtr->FinishBuildInternalData.bApplyMorphTargetsData = true;
+					//We cannot serialize directly the UMorphTarget with a FMemoryArchive. This is not supported.
+					for (int32 MorphTargetIndex = 0; MorphTargetIndex < MorphTargetNumber; ++MorphTargetIndex)
 					{
-						DeleteObject();
-					}
-					else
-					{
-						Async(EAsyncExecution::TaskGraphMainThread, MoveTemp(DeleteObject));
+						FName MorphTargetName = NAME_None;
+						Ar << MorphTargetName;
+						TArray<FMorphTargetLODModel>& MorphTargetLODModels = ContextPtr->FinishBuildInternalData.MorphLODModelsPerTargetName.FindOrAdd(MorphTargetName);
+						int32 MorphLODModelNumber = 0;
+						Ar << MorphLODModelNumber;
+						MorphTargetLODModels.Empty(MorphLODModelNumber);
+						MorphTargetLODModels.AddDefaulted(MorphLODModelNumber);
+						for (int32 MorphDataIndex = 0; MorphDataIndex < MorphLODModelNumber; ++MorphDataIndex)
+						{
+							Ar << MorphTargetLODModels[MorphDataIndex];
+						}
 					}
 				}
 
-				// In case we're built async and objects were created, we need to remove the async flag now that they are referenced
-				// and reachable by the GC.
-				for (UMorphTarget* MorphTarget : Owner->GetMorphTargets())
-				{
-					MorphTarget->ClearInternalFlags(EInternalObjectFlags::Async);
-				}
-
 				//Serialize the LODModel sections since they are dependent on the reduction
 				for (int32 LODIndex = 0; LODIndex < SkelMeshModel->LODModels.Num(); LODIndex++)
 				{
@@ -485,7 +440,7 @@ void FSkeletalMeshRenderData::Cache(const ITargetPlatform* TargetPlatform, USkel
 			for (int32 LODIndex = 0; LODIndex < SkelMeshModel->LODModels.Num(); LODIndex++)
 			{
 				FSkeletalMeshLODModel* LODModel = &(SkelMeshModel->LODModels[LODIndex]);
-				FSkeletalMeshLODInfo* LODInfo = Owner->GetLODInfo(LODIndex);
+				const FSkeletalMeshLODInfo* LODInfo = Owner->GetLODInfo(LODIndex);
 				check(LODInfo);
 				bool bRawDataEmpty = Owner->IsLODImportedDataEmpty(LODIndex);
 				bool bRawBuildDataAvailable = Owner->IsLODImportedDataBuildAvailable(LODIndex);
diff --git a/Engine/Source/Runtime/Engine/Public/Rendering/SkeletalMeshRenderData.h b/Engine/Source/Runtime/Engine/Public/Rendering/SkeletalMeshRenderData.h
--- a/Engine/Source/Runtime/Engine/Public/Rendering/SkeletalMeshRenderData.h
+++ b/Engine/Source/Runtime/Engine/Public/Rendering/SkeletalMeshRenderData.h
@@ -54,7 +54,7 @@ public:
 	~FSkeletalMeshRenderData();
 
 #if WITH_EDITOR
-	void Cache(const ITargetPlatform* TargetPlatform, USkeletalMesh* Owner);
+	void Cache(const ITargetPlatform* TargetPlatform, USkeletalMesh* Owner, class FSkeletalMeshCompilationContext* ContextPtr);
 	FString GetDerivedDataKey(const ITargetPlatform* TargetPlatform, USkeletalMesh* Owner);
 
 	void SyncUVChannelData(const TArray<FSkeletalMaterial>& ObjectData);
