Fix for server build compiling with DerivedDataCache causing a crash during static init
#rnx

diff --git a/Engine/Source/Runtime/Engine/Engine.Build.cs b/Engine/Source/Runtime/Engine/Engine.Build.cs
--- a/Engine/Source/Runtime/Engine/Engine.Build.cs
+++ b/Engine/Source/Runtime/Engine/Engine.Build.cs
@@ -59,9 +59,12 @@ public class Engine : ModuleRules
 				new string[]
 				{
 					"EditorAnalyticsSession",
-					"DerivedDataCache",
 				}
 			);
+			if (Target.Type != TargetType.Server)
+			{
+				PrivateDependencyModuleNames.Add("DerivedDataCache");
+			}
 		}
 
 		PublicDependencyModuleNames.AddRange(
