Fix non-unity compilation error.

diff --git a/Engine/Source/Runtime/Engine/Classes/Engine/SubsurfaceProfile.h b/Engine/Source/Runtime/Engine/Classes/Engine/SubsurfaceProfile.h
--- a/Engine/Source/Runtime/Engine/Classes/Engine/SubsurfaceProfile.h
+++ b/Engine/Source/Runtime/Engine/Classes/Engine/SubsurfaceProfile.h
@@ -192,7 +192,7 @@ public:
 	struct IPooledRenderTarget* GetTexture(FRHICommandListImmediate& RHICmdList);
 	struct IPooledRenderTarget* GetTexture();
 
-	struct IPooledRenderTarget* GetSSProfilesPreIntegratedTexture(FRDGBuilder& GraphBuilder, EShaderPlatform ShaderPlatform);
+	struct IPooledRenderTarget* GetSSProfilesPreIntegratedTexture(class FRDGBuilder& GraphBuilder, EShaderPlatform ShaderPlatform);
 
 	//~ Begin FRenderResource Interface.
 	/**
@@ -247,7 +247,7 @@ static const int32 SUBSURFACE_KERNEL_SIZE = 3;
 extern ENGINE_API TGlobalResource<FSubsurfaceProfileTexture> GSubsurfaceProfileTextureObject;
 
 // Initializes or updates the contents of the subsurface profile texture.
-ENGINE_API void UpdateSubsurfaceProfileTexture(FRDGBuilder& GraphBuilder, EShaderPlatform ShaderPlatform);
+ENGINE_API void UpdateSubsurfaceProfileTexture(class FRDGBuilder& GraphBuilder, EShaderPlatform ShaderPlatform);
 
 // Returns the subsurface profile texture if it exists, or null.
 ENGINE_API FRHITexture* GetSubsurfaceProfileTexture();
