Converted rewind debugger's trace event attachments to proper typed fields.

#rnx

diff --git a/Engine/Plugins/Animation/GameplayInsights/Source/GameplayInsights/Private/AnimationAnalyzer.cpp b/Engine/Plugins/Animation/GameplayInsights/Source/GameplayInsights/Private/AnimationAnalyzer.cpp
--- a/Engine/Plugins/Animation/GameplayInsights/Source/GameplayInsights/Private/AnimationAnalyzer.cpp
+++ b/Engine/Plugins/Animation/GameplayInsights/Source/GameplayInsights/Private/AnimationAnalyzer.cpp
@@ -45,7 +45,9 @@ void FAnimationAnalyzer::OnAnalysisBegin(const FOnAnalysisContext& Context)
 
 bool FAnimationAnalyzer::OnEvent(uint16 RouteId, EStyle Style, const FOnEventContext& Context)
 {
-	TraceServices::FAnalysisSessionEditScope _(Session);
+	using namespace TraceServices;
+
+	FAnalysisSessionEditScope _(Session);
 
 	const auto& EventData = Context.EventData;
 	switch (RouteId)
@@ -122,7 +124,7 @@ bool FAnimationAnalyzer::OnEvent(uint16 RouteId, EStyle Style, const FOnEventCon
 	case RouteId_Name:
 	{
 		uint32 Id = EventData.GetValue<uint32>("Id");
-		FString Name = TraceServices::FTraceAnalyzerUtils::LegacyAttachmentString<TCHAR>("Name", Context);
+		FString Name = FTraceAnalyzerUtils::LegacyAttachmentString<TCHAR>("Name", Context);
 		AnimationProvider.AppendName(Id, *Name);
 		break;
 	}
@@ -155,8 +157,8 @@ bool FAnimationAnalyzer::OnEvent(uint16 RouteId, EStyle Style, const FOnEventCon
 		float RootMotionWeight = EventData.GetValue<float>("RootMotionWeight");
 		uint16 FrameCounter = EventData.GetValue<uint16>("FrameCounter");
 		uint8 Phase = EventData.GetValue<uint8>("Phase");
-		const TCHAR* TargetNodeName = reinterpret_cast<const TCHAR*>(EventData.GetAttachment());
-		AnimationProvider.AppendAnimNodeStart(AnimInstanceId, Context.EventTime.AsSeconds(StartCycle), FrameCounter, NodeId, PreviousNodeId, Weight, RootMotionWeight, TargetNodeName, Phase);
+		FString TargetNodeName = FTraceAnalyzerUtils::LegacyAttachmentString<TCHAR>("DisplayName", Context);
+		AnimationProvider.AppendAnimNodeStart(AnimInstanceId, Context.EventTime.AsSeconds(StartCycle), FrameCounter, NodeId, PreviousNodeId, Weight, RootMotionWeight, *TargetNodeName, Phase);
 		break;
 	}
 	case RouteId_AnimNodeAttribute:
@@ -173,6 +175,7 @@ bool FAnimationAnalyzer::OnEvent(uint16 RouteId, EStyle Style, const FOnEventCon
 	case RouteId_AnimNodeValueBool:
 	case RouteId_AnimNodeValueInt:
 	case RouteId_AnimNodeValueFloat:
+	case RouteId_AnimNodeValueVector2D:
 	case RouteId_AnimNodeValueVector:
 	case RouteId_AnimNodeValueString:
 	case RouteId_AnimNodeValueObject:
@@ -181,58 +184,62 @@ bool FAnimationAnalyzer::OnEvent(uint16 RouteId, EStyle Style, const FOnEventCon
 		uint64 Cycle = EventData.GetValue<uint64>("Cycle");
 		uint64 AnimInstanceId = EventData.GetValue<uint64>("AnimInstanceId");
 		int32 NodeId = EventData.GetValue<int32>("NodeId");
-		int32 KeyLength = EventData.GetValue<uint32>("KeyLength");
 		uint16 FrameCounter = EventData.GetValue<uint16>("FrameCounter");
-		const TCHAR* Key = reinterpret_cast<const TCHAR*>(EventData.GetAttachment());
+		FString Key = FTraceAnalyzerUtils::LegacyAttachmentString<TCHAR>("Key", Context);
 
 		switch (RouteId)
 		{
 		case RouteId_AnimNodeValueBool:
 		{
 			bool Value = EventData.GetValue<bool>("Value");
-			AnimationProvider.AppendAnimNodeValue(AnimInstanceId, Context.EventTime.AsSeconds(Cycle), FrameCounter, NodeId, Key, Value);
+			AnimationProvider.AppendAnimNodeValue(AnimInstanceId, Context.EventTime.AsSeconds(Cycle), FrameCounter, NodeId, *Key, Value);
 			break;
 		}
 		case RouteId_AnimNodeValueInt:
 		{
 			int32 Value = EventData.GetValue<int32>("Value");
-			AnimationProvider.AppendAnimNodeValue(AnimInstanceId, Context.EventTime.AsSeconds(Cycle), FrameCounter, NodeId, Key, Value);
+			AnimationProvider.AppendAnimNodeValue(AnimInstanceId, Context.EventTime.AsSeconds(Cycle), FrameCounter, NodeId, *Key, Value);
 			break;
 		}
 		case RouteId_AnimNodeValueFloat:
 		{
 			float Value = EventData.GetValue<float>("Value");
-			AnimationProvider.AppendAnimNodeValue(AnimInstanceId, Context.EventTime.AsSeconds(Cycle), FrameCounter, NodeId, Key, Value);
+			AnimationProvider.AppendAnimNodeValue(AnimInstanceId, Context.EventTime.AsSeconds(Cycle), FrameCounter, NodeId, *Key, Value);
 			break;
 		}
 		case RouteId_AnimNodeValueVector2D:
 		{
 			FVector2D Value(EventData.GetValue<float>("ValueX"), EventData.GetValue<float>("ValueY"));
-			AnimationProvider.AppendAnimNodeValue(AnimInstanceId, Context.EventTime.AsSeconds(Cycle), FrameCounter, NodeId, Key, Value);
+			AnimationProvider.AppendAnimNodeValue(AnimInstanceId, Context.EventTime.AsSeconds(Cycle), FrameCounter, NodeId, *Key, Value);
 			break;
 		}
 		case RouteId_AnimNodeValueVector:
 		{
 			FVector Value(EventData.GetValue<float>("ValueX"), EventData.GetValue<float>("ValueY"), EventData.GetValue<float>("ValueZ"));
-			AnimationProvider.AppendAnimNodeValue(AnimInstanceId, Context.EventTime.AsSeconds(Cycle), FrameCounter, NodeId, Key, Value);
+			AnimationProvider.AppendAnimNodeValue(AnimInstanceId, Context.EventTime.AsSeconds(Cycle), FrameCounter, NodeId, *Key, Value);
 			break;
 		}
 		case RouteId_AnimNodeValueString:
 		{
-			const TCHAR* Value = reinterpret_cast<const TCHAR*>(EventData.GetAttachment()) + KeyLength;
-			AnimationProvider.AppendAnimNodeValue(AnimInstanceId, Context.EventTime.AsSeconds(Cycle), FrameCounter, NodeId, Key, Value);
+			FString Value;
+			if (!EventData.GetString("Value", Value))
+			{
+				int32 KeyLength = EventData.GetValue<uint32>("KeyLength");
+				Value = reinterpret_cast<const TCHAR*>(EventData.GetAttachment()) + KeyLength;
+			}
+			AnimationProvider.AppendAnimNodeValue(AnimInstanceId, Context.EventTime.AsSeconds(Cycle), FrameCounter, NodeId, *Key, *Value);
 			break;
 		}
 		case RouteId_AnimNodeValueObject:
 		{
 			uint64 Value = EventData.GetValue<uint64>("Value");
-			AnimationProvider.AppendAnimNodeValueObject(AnimInstanceId, Context.EventTime.AsSeconds(Cycle), FrameCounter, NodeId, Key, Value);
+			AnimationProvider.AppendAnimNodeValueObject(AnimInstanceId, Context.EventTime.AsSeconds(Cycle), FrameCounter, NodeId, *Key, Value);
 			break;
 		}
 		case RouteId_AnimNodeValueClass:
 		{
 			uint64 Value = EventData.GetValue<uint64>("Value");
-			AnimationProvider.AppendAnimNodeValueClass(AnimInstanceId, Context.EventTime.AsSeconds(Cycle), FrameCounter, NodeId, Key, Value);
+			AnimationProvider.AppendAnimNodeValueClass(AnimInstanceId, Context.EventTime.AsSeconds(Cycle), FrameCounter, NodeId, *Key, Value);
 			break;
 		}
 		}
diff --git a/Engine/Plugins/Animation/GameplayInsights/Source/GameplayInsights/Private/GameplayAnalyzer.cpp b/Engine/Plugins/Animation/GameplayInsights/Source/GameplayInsights/Private/GameplayAnalyzer.cpp
--- a/Engine/Plugins/Animation/GameplayInsights/Source/GameplayInsights/Private/GameplayAnalyzer.cpp
+++ b/Engine/Plugins/Animation/GameplayInsights/Source/GameplayInsights/Private/GameplayAnalyzer.cpp
@@ -1,6 +1,7 @@
 // Copyright Epic Games, Inc. All Rights Reserved.
 
 #include "GameplayAnalyzer.h"
+#include "TraceServices/Utils.h"
 #include "TraceServices/Model/AnalysisSession.h"
 #include "GameplayProvider.h"
 
@@ -56,33 +57,50 @@ bool FGameplayAnalyzer::OnEvent(uint16 RouteId, EStyle Style, const FOnEventCont
 	}
 	case RouteId_Class:
 	{
-		const TCHAR* ClassNameAndPathName = reinterpret_cast<const TCHAR*>(EventData.GetAttachment());
-		const TCHAR* ClassName = ClassNameAndPathName;
-		int32 ClassNameStringLength = EventData.GetValue<int32>("ClassNameStringLength");
-		const TCHAR* ClassPathName = ClassNameAndPathName + ClassNameStringLength;
+		FString ClassName, ClassPathName;
+		if (EventData.GetString("Name", ClassName))
+		{
+			EventData.GetString("Path", ClassPathName);
+		}
+		else
+		{
+			const TCHAR* ClassNameAndPathName = reinterpret_cast<const TCHAR*>(EventData.GetAttachment());
+			ClassName = ClassNameAndPathName;
+			int32 ClassNameStringLength = EventData.GetValue<int32>("ClassNameStringLength");
+			ClassPathName = ClassNameAndPathName + ClassNameStringLength;
+		}
+
 		uint64 Id = EventData.GetValue<uint64>("Id");
 		uint64 SuperId = EventData.GetValue<uint64>("SuperId");
-		GameplayProvider.AppendClass(Id, SuperId, ClassName, ClassPathName);
+		GameplayProvider.AppendClass(Id, SuperId, *ClassName, *ClassPathName);
 		break;
 	}
 	case RouteId_Object:
 	{
-		const TCHAR* ObjectNameAndPathName = reinterpret_cast<const TCHAR*>(EventData.GetAttachment());
-		const TCHAR* ObjectName = ObjectNameAndPathName;
-		int32 NameStringLength = EventData.GetValue<int32>("ObjectNameStringLength");
-		const TCHAR* ObjectPathName = ObjectNameAndPathName + NameStringLength;
+		FString ObjectName, ObjectPathName;
+		if (EventData.GetString("Name", ObjectName))
+		{
+			EventData.GetString("Path", ObjectPathName);
+		}
+		else
+		{
+			const TCHAR* ObjectNameAndPathName = reinterpret_cast<const TCHAR*>(EventData.GetAttachment());
+			ObjectName = ObjectNameAndPathName;
+			int32 NameStringLength = EventData.GetValue<int32>("ObjectNameStringLength");
+			ObjectPathName = ObjectNameAndPathName + NameStringLength;
+		}
 		uint64 Id = EventData.GetValue<uint64>("Id");
 		uint64 OuterId = EventData.GetValue<uint64>("OuterId");
 		uint64 ClassId = EventData.GetValue<uint64>("ClassId");
-		GameplayProvider.AppendObject(Id, OuterId, ClassId, ObjectName, ObjectPathName);
+		GameplayProvider.AppendObject(Id, OuterId, ClassId, *ObjectName, *ObjectPathName);
 		break;
 	}
 	case RouteId_ObjectEvent:
 	{
 		uint64 Cycle = EventData.GetValue<uint64>("Cycle");
 		uint64 Id = EventData.GetValue<uint64>("Id");
-		const TCHAR* Event = reinterpret_cast<const TCHAR*>(EventData.GetAttachment());
-		GameplayProvider.AppendObjectEvent(Id, Context.EventTime.AsSeconds(Cycle), Event);
+		FString Event = TraceServices::FTraceAnalyzerUtils::LegacyAttachmentString<TCHAR>("Event", Context);
+		GameplayProvider.AppendObjectEvent(Id, Context.EventTime.AsSeconds(Cycle), *Event);
 		break;
 	}
 	case RouteId_View:
diff --git a/Engine/Plugins/Developer/TraceSourceFiltering/Source/SourceFilteringTrace/Private/SourceFilterTrace.cpp b/Engine/Plugins/Developer/TraceSourceFiltering/Source/SourceFilteringTrace/Private/SourceFilterTrace.cpp
--- a/Engine/Plugins/Developer/TraceSourceFiltering/Source/SourceFilteringTrace/Private/SourceFilterTrace.cpp
+++ b/Engine/Plugins/Developer/TraceSourceFiltering/Source/SourceFilteringTrace/Private/SourceFilterTrace.cpp
@@ -30,6 +30,7 @@ UE_TRACE_EVENT_BEGIN(SourceFilters, FilterInstance)
 	UE_TRACE_EVENT_FIELD(uint64, InstanceId)
 	UE_TRACE_EVENT_FIELD(uint64, ClassId)
 	UE_TRACE_EVENT_FIELD(uint64, SetId)
+	UE_TRACE_EVENT_FIELD(UE::Trace::WideString, DisplayString)
 UE_TRACE_EVENT_END()
 
 UE_TRACE_EVENT_BEGIN(SourceFilters, FilterSet)
@@ -45,12 +46,14 @@ UE_TRACE_EVENT_END()
 
 UE_TRACE_EVENT_BEGIN(SourceFilters, SetFilterSettingValue)
 	UE_TRACE_EVENT_FIELD(uint8, Value)
+	UE_TRACE_EVENT_FIELD(UE::Trace::AnsiString, PropertyName)
 UE_TRACE_EVENT_END()
 
 UE_TRACE_EVENT_BEGIN(WorldSourceFilters, WorldInstance)
 	UE_TRACE_EVENT_FIELD(uint64, InstanceId)
 	UE_TRACE_EVENT_FIELD(uint8, Type)
 	UE_TRACE_EVENT_FIELD(uint8, State)
+	UE_TRACE_EVENT_FIELD(UE::Trace::WideString, SendName)
 UE_TRACE_EVENT_END()
 
 UE_TRACE_EVENT_BEGIN(WorldSourceFilters, WorldOperation)
@@ -111,22 +114,13 @@ void FSourceFilterTrace::OutputInstance(const UDataSourceFilter* InFilter)
 
 		FText DisplayText;
 		InFilter->Execute_GetDisplayText(InFilter, DisplayText);
-
 		FString DisplayString = DisplayText.ToString();
 		
-		const uint32 AttachmentSize = (DisplayString.Len() + 1) * sizeof(TCHAR);
-		auto WriteAttachment = [DisplayString, AttachmentSize](uint8* Out)
-		{
-			TCHAR* StringPtr = (TCHAR*)Out;
-			FMemory::Memcpy(StringPtr, *DisplayString, AttachmentSize);
-			StringPtr[AttachmentSize / sizeof(TCHAR) - 1] = '\0';
-		};
-		
-		UE_TRACE_LOG(SourceFilters, FilterInstance, TraceSourceFiltersChannel, AttachmentSize)
+		UE_TRACE_LOG(SourceFilters, FilterInstance, TraceSourceFiltersChannel)
 			<< FilterInstance.ClassId(ClassId)
 			<< FilterInstance.InstanceId(InstanceId)
 			<< FilterInstance.SetId(SetId)
-			<< FilterInstance.Attachment(WriteAttachment);
+			<< FilterInstance.DisplayString(*DisplayString, DisplayString.Len());
 	}
 }
 
@@ -155,17 +149,9 @@ void FSourceFilterTrace::OutputFilterOperation(const UDataSourceFilter* InFilter
 
 void FSourceFilterTrace::OutputFilterSettingsValue(const FString& InPropertyName, const uint8 InValue)
 {
-	const uint32 AttachmentSize = (InPropertyName.Len() + 1) * sizeof(TCHAR);
-	auto WriteAttachment = [InPropertyName, AttachmentSize](uint8* Out)
-	{
-		TCHAR* StringPtr = (TCHAR*)Out;
-		FMemory::Memcpy(StringPtr, *InPropertyName, AttachmentSize);
-		StringPtr[AttachmentSize / sizeof(TCHAR) - 1] = '\0';
-	};
-
-	UE_TRACE_LOG(SourceFilters, SetFilterSettingValue, TraceSourceFiltersChannel, AttachmentSize)
+	UE_TRACE_LOG(SourceFilters, SetFilterSettingValue, TraceSourceFiltersChannel)
 		<< SetFilterSettingValue.Value(InValue)
-		<< SetFilterSettingValue.Attachment(WriteAttachment);
+		<< SetFilterSettingValue.PropertyName(*InPropertyName, InPropertyName.Len());
 }
 
 void FSourceFilterTrace::OutputWorld(const UWorld* InWorld)
@@ -187,11 +173,11 @@ void FSourceFilterTrace::OutputWorld(const UWorld* InWorld)
 			StringPtr[AttachmentSize / sizeof(TCHAR) - 1] = '\0';
 		};
 
-		UE_TRACE_LOG(WorldSourceFilters, WorldInstance, TraceSourceFiltersChannel, AttachmentSize)
+		UE_TRACE_LOG(WorldSourceFilters, WorldInstance, TraceSourceFiltersChannel)
 			<< WorldInstance.InstanceId(InstanceId)
 			<< WorldInstance.Type((uint8)InWorld->WorldType)
 			<< WorldInstance.State(CAN_TRACE_OBJECT(InWorld) ? 1 : 0)
-			<< WorldInstance.Attachment(WriteAttachment);
+			<< WorldInstance.SendName(*SendName, SendName.Len());
 	}
 }
 
diff --git a/Engine/Source/Runtime/Engine/Private/Animation/AnimTrace.cpp b/Engine/Source/Runtime/Engine/Private/Animation/AnimTrace.cpp
--- a/Engine/Source/Runtime/Engine/Private/Animation/AnimTrace.cpp
+++ b/Engine/Source/Runtime/Engine/Private/Animation/AnimTrace.cpp
@@ -83,6 +83,7 @@ UE_TRACE_EVENT_BEGIN(Animation, AnimNodeStart)
 	UE_TRACE_EVENT_FIELD(float, Weight)
 	UE_TRACE_EVENT_FIELD(uint16, FrameCounter)
 	UE_TRACE_EVENT_FIELD(uint8, Phase)
+	UE_TRACE_EVENT_FIELD(UE::Trace::WideString, DisplayName)
 UE_TRACE_EVENT_END()
 
 UE_TRACE_EVENT_BEGIN(Animation, AnimNodeEnd)
@@ -103,56 +104,57 @@ UE_TRACE_EVENT_BEGIN(Animation, AnimNodeValueBool)
 	UE_TRACE_EVENT_FIELD(uint64, Cycle)
 	UE_TRACE_EVENT_FIELD(uint64, AnimInstanceId)
 	UE_TRACE_EVENT_FIELD(int32, NodeId)
-	UE_TRACE_EVENT_FIELD(int32, KeyLength)
 	UE_TRACE_EVENT_FIELD(uint16, FrameCounter)
 	UE_TRACE_EVENT_FIELD(bool, Value)
+	UE_TRACE_EVENT_FIELD(UE::Trace::WideString, Key)
 UE_TRACE_EVENT_END()
 
 UE_TRACE_EVENT_BEGIN(Animation, AnimNodeValueInt)
 	UE_TRACE_EVENT_FIELD(uint64, Cycle)
 	UE_TRACE_EVENT_FIELD(uint64, AnimInstanceId)
 	UE_TRACE_EVENT_FIELD(int32, NodeId)
-	UE_TRACE_EVENT_FIELD(int32, KeyLength)
 	UE_TRACE_EVENT_FIELD(int32, Value)
 	UE_TRACE_EVENT_FIELD(uint16, FrameCounter)
+	UE_TRACE_EVENT_FIELD(UE::Trace::WideString, Key)
 UE_TRACE_EVENT_END()
 
 UE_TRACE_EVENT_BEGIN(Animation, AnimNodeValueFloat)
 	UE_TRACE_EVENT_FIELD(uint64, Cycle)
 	UE_TRACE_EVENT_FIELD(uint64, AnimInstanceId)
 	UE_TRACE_EVENT_FIELD(int32, NodeId)
-	UE_TRACE_EVENT_FIELD(int32, KeyLength)
 	UE_TRACE_EVENT_FIELD(float, Value)
 	UE_TRACE_EVENT_FIELD(uint16, FrameCounter)
+	UE_TRACE_EVENT_FIELD(UE::Trace::WideString, Key)
 UE_TRACE_EVENT_END()
 
 UE_TRACE_EVENT_BEGIN(Animation, AnimNodeValueVector2D)
 	UE_TRACE_EVENT_FIELD(uint64, Cycle)
 	UE_TRACE_EVENT_FIELD(uint64, AnimInstanceId)
 	UE_TRACE_EVENT_FIELD(int32, NodeId)
-	UE_TRACE_EVENT_FIELD(int32, KeyLength)
 	UE_TRACE_EVENT_FIELD(float, ValueX)
 	UE_TRACE_EVENT_FIELD(float, ValueY)
 	UE_TRACE_EVENT_FIELD(uint16, FrameCounter)
+	UE_TRACE_EVENT_FIELD(UE::Trace::WideString, Key)
 UE_TRACE_EVENT_END()
 
 UE_TRACE_EVENT_BEGIN(Animation, AnimNodeValueVector)
 	UE_TRACE_EVENT_FIELD(uint64, Cycle)
 	UE_TRACE_EVENT_FIELD(uint64, AnimInstanceId)
 	UE_TRACE_EVENT_FIELD(int32, NodeId)
-	UE_TRACE_EVENT_FIELD(int32, KeyLength)
 	UE_TRACE_EVENT_FIELD(float, ValueX)
 	UE_TRACE_EVENT_FIELD(float, ValueY)
 	UE_TRACE_EVENT_FIELD(float, ValueZ)
 	UE_TRACE_EVENT_FIELD(uint16, FrameCounter)
+	UE_TRACE_EVENT_FIELD(UE::Trace::WideString, Key)
 UE_TRACE_EVENT_END()
 
 UE_TRACE_EVENT_BEGIN(Animation, AnimNodeValueString)
 	UE_TRACE_EVENT_FIELD(uint64, Cycle)
 	UE_TRACE_EVENT_FIELD(uint64, AnimInstanceId)
 	UE_TRACE_EVENT_FIELD(int32, NodeId)
-	UE_TRACE_EVENT_FIELD(int32, KeyLength)
+	UE_TRACE_EVENT_FIELD(UE::Trace::WideString, Value)
 	UE_TRACE_EVENT_FIELD(uint16, FrameCounter)
+	UE_TRACE_EVENT_FIELD(UE::Trace::WideString, Key)
 UE_TRACE_EVENT_END()
 
 UE_TRACE_EVENT_BEGIN(Animation, AnimNodeValueObject)
@@ -160,8 +162,8 @@ UE_TRACE_EVENT_BEGIN(Animation, AnimNodeValueObject)
 	UE_TRACE_EVENT_FIELD(uint64, AnimInstanceId)
 	UE_TRACE_EVENT_FIELD(uint64, Value)
 	UE_TRACE_EVENT_FIELD(int32, NodeId)
-	UE_TRACE_EVENT_FIELD(int32, KeyLength)
 	UE_TRACE_EVENT_FIELD(uint16, FrameCounter)
+	UE_TRACE_EVENT_FIELD(UE::Trace::WideString, Key)
 UE_TRACE_EVENT_END()
 
 UE_TRACE_EVENT_BEGIN(Animation, AnimNodeValueClass)
@@ -169,8 +171,8 @@ UE_TRACE_EVENT_BEGIN(Animation, AnimNodeValueClass)
 	UE_TRACE_EVENT_FIELD(uint64, AnimInstanceId)
 	UE_TRACE_EVENT_FIELD(uint64, Value)
 	UE_TRACE_EVENT_FIELD(int32, NodeId)
-	UE_TRACE_EVENT_FIELD(int32, KeyLength)
 	UE_TRACE_EVENT_FIELD(uint16, FrameCounter)
+	UE_TRACE_EVENT_FIELD(UE::Trace::WideString, Key)
 UE_TRACE_EVENT_END()
 
 UE_TRACE_EVENT_BEGIN(Animation, AnimSequencePlayer)
@@ -653,7 +655,7 @@ void FAnimTrace::OutputAnimNodeStart(const FAnimationBaseContext& InContext, uin
 
 	check(InPreviousNodeId != InNodeId);
 
-	UE_TRACE_LOG(Animation, AnimNodeStart, AnimationChannel, (DisplayNameString.Len() + 1) * sizeof(TCHAR))
+	UE_TRACE_LOG(Animation, AnimNodeStart, AnimationChannel)
 		<< AnimNodeStart.StartCycle(InStartCycle)
 		<< AnimNodeStart.AnimInstanceId(FObjectTrace::GetObjectId(AnimInstance))
 		<< AnimNodeStart.PreviousNodeId(InPreviousNodeId)
@@ -661,7 +663,7 @@ void FAnimTrace::OutputAnimNodeStart(const FAnimationBaseContext& InContext, uin
 		<< AnimNodeStart.Weight(InBlendWeight)
 		<< AnimNodeStart.FrameCounter(FObjectTrace::GetObjectWorldTickCounter(AnimInstance))
 		<< AnimNodeStart.Phase(InPhase)
-		<< AnimNodeStart.Attachment(*DisplayNameString, (DisplayNameString.Len() + 1) * sizeof(TCHAR));
+		<< AnimNodeStart.DisplayName(*DisplayNameString, DisplayNameString.Len());
 }
 
 void FAnimTrace::OutputAnimNodeEnd(const FAnimationBaseContext& InContext, uint64 InEndCycle)
@@ -734,16 +736,13 @@ void FAnimTrace::OutputAnimNodeValue(const FAnimationBaseContext& InContext, con
 	UObject* AnimInstance = InContext.AnimInstanceProxy->GetAnimInstanceObject();
 	TRACE_OBJECT(AnimInstance);
 
-	int32 KeyLength = FCString::Strlen(InKey) + 1;
-
-	UE_TRACE_LOG(Animation, AnimNodeValueBool, AnimationChannel, KeyLength * sizeof(TCHAR))
+	UE_TRACE_LOG(Animation, AnimNodeValueBool, AnimationChannel)
 		<< AnimNodeValueBool.Cycle(FPlatformTime::Cycles64())
 		<< AnimNodeValueBool.AnimInstanceId(FObjectTrace::GetObjectId(AnimInstance))
 		<< AnimNodeValueBool.NodeId(InContext.GetCurrentNodeId())
-		<< AnimNodeValueBool.KeyLength(KeyLength)
 		<< AnimNodeValueBool.Value(InValue)
 		<< AnimNodeValueBool.FrameCounter(FObjectTrace::GetObjectWorldTickCounter(AnimInstance))
-		<< AnimNodeValueBool.Attachment(InKey, KeyLength * sizeof(TCHAR));
+		<< AnimNodeValueBool.Key(InKey);
 }
 
 void FAnimTrace::OutputAnimNodeValue(const FAnimationBaseContext& InContext, const TCHAR* InKey, int32 InValue)
@@ -764,16 +763,13 @@ void FAnimTrace::OutputAnimNodeValue(const FAnimationBaseContext& InContext, con
 	UObject* AnimInstance = InContext.AnimInstanceProxy->GetAnimInstanceObject();
 	TRACE_OBJECT(AnimInstance);
 
-	int32 KeyLength = FCString::Strlen(InKey) + 1;
-
-	UE_TRACE_LOG(Animation, AnimNodeValueInt, AnimationChannel, KeyLength * sizeof(TCHAR))
+	UE_TRACE_LOG(Animation, AnimNodeValueInt, AnimationChannel)
 		<< AnimNodeValueInt.Cycle(FPlatformTime::Cycles64())
 		<< AnimNodeValueInt.AnimInstanceId(FObjectTrace::GetObjectId(AnimInstance))
 		<< AnimNodeValueInt.NodeId(InContext.GetCurrentNodeId())
-		<< AnimNodeValueInt.KeyLength(KeyLength)
 		<< AnimNodeValueInt.Value(InValue)
 		<< AnimNodeValueInt.FrameCounter(FObjectTrace::GetObjectWorldTickCounter(AnimInstance))
-		<< AnimNodeValueInt.Attachment(InKey, KeyLength * sizeof(TCHAR));
+		<< AnimNodeValueInt.Key(InKey);
 }
 
 void FAnimTrace::OutputAnimNodeValue(const FAnimationBaseContext& InContext, const TCHAR* InKey, float InValue)
@@ -794,16 +790,13 @@ void FAnimTrace::OutputAnimNodeValue(const FAnimationBaseContext& InContext, con
 	UObject* AnimInstance = InContext.AnimInstanceProxy->GetAnimInstanceObject();
 	TRACE_OBJECT(AnimInstance);
 
-	int32 KeyLength = FCString::Strlen(InKey) + 1;
-
-	UE_TRACE_LOG(Animation, AnimNodeValueFloat, AnimationChannel, KeyLength * sizeof(TCHAR))
+	UE_TRACE_LOG(Animation, AnimNodeValueFloat, AnimationChannel)
 		<< AnimNodeValueFloat.Cycle(FPlatformTime::Cycles64())
 		<< AnimNodeValueFloat.AnimInstanceId(FObjectTrace::GetObjectId(AnimInstance))
 		<< AnimNodeValueFloat.NodeId(InContext.GetCurrentNodeId())
-		<< AnimNodeValueFloat.KeyLength(KeyLength)
 		<< AnimNodeValueFloat.Value(InValue)
 		<< AnimNodeValueFloat.FrameCounter(FObjectTrace::GetObjectWorldTickCounter(AnimInstance))
-		<< AnimNodeValueFloat.Attachment(InKey, KeyLength * sizeof(TCHAR));
+		<< AnimNodeValueFloat.Key(InKey);
 }
 
 void FAnimTrace::OutputAnimNodeValue(const FAnimationBaseContext& InContext, const TCHAR* InKey, const FVector2D& InValue)
@@ -819,17 +812,14 @@ void FAnimTrace::OutputAnimNodeValue(const FAnimationBaseContext& InContext, con
 	UObject* AnimInstance = InContext.AnimInstanceProxy->GetAnimInstanceObject();
 	TRACE_OBJECT(AnimInstance);
 
-	int32 KeyLength = FCString::Strlen(InKey) + 1;
-
-	UE_TRACE_LOG(Animation, AnimNodeValueVector2D, AnimationChannel, KeyLength * sizeof(TCHAR))
+	UE_TRACE_LOG(Animation, AnimNodeValueVector2D, AnimationChannel)
 		<< AnimNodeValueVector2D.Cycle(FPlatformTime::Cycles64())
 		<< AnimNodeValueVector2D.AnimInstanceId(FObjectTrace::GetObjectId(AnimInstance))
 		<< AnimNodeValueVector2D.NodeId(InContext.GetCurrentNodeId())
-		<< AnimNodeValueVector2D.KeyLength(KeyLength)
 		<< AnimNodeValueVector2D.ValueX(InValue.X)
 		<< AnimNodeValueVector2D.ValueY(InValue.Y)
 		<< AnimNodeValueVector2D.FrameCounter(FObjectTrace::GetObjectWorldTickCounter(AnimInstance))
-		<< AnimNodeValueVector2D.Attachment(InKey, KeyLength * sizeof(TCHAR));
+		<< AnimNodeValueVector2D.Key(InKey);
 }
 
 
@@ -857,18 +847,15 @@ void FAnimTrace::OutputAnimNodeValue(const FAnimationBaseContext& InContext, con
 	UObject* AnimInstance = InContext.AnimInstanceProxy->GetAnimInstanceObject();
 	TRACE_OBJECT(AnimInstance);
 
-	int32 KeyLength = FCString::Strlen(InKey) + 1;
-
-	UE_TRACE_LOG(Animation, AnimNodeValueVector, AnimationChannel, KeyLength * sizeof(TCHAR))
+	UE_TRACE_LOG(Animation, AnimNodeValueVector, AnimationChannel)
 		<< AnimNodeValueVector.Cycle(FPlatformTime::Cycles64())
 		<< AnimNodeValueVector.AnimInstanceId(FObjectTrace::GetObjectId(InContext.AnimInstanceProxy->GetAnimInstanceObject()))
 		<< AnimNodeValueVector.NodeId(InContext.GetCurrentNodeId())
-		<< AnimNodeValueVector.KeyLength(KeyLength)
 		<< AnimNodeValueVector.ValueX(InValue.X)
 		<< AnimNodeValueVector.ValueY(InValue.Y)
 		<< AnimNodeValueVector.ValueZ(InValue.Z)
 		<< AnimNodeValueVector.FrameCounter(FObjectTrace::GetObjectWorldTickCounter(AnimInstance))
-		<< AnimNodeValueVector.Attachment(InKey, KeyLength * sizeof(TCHAR));
+		<< AnimNodeValueVector.Key(InKey);
 }
 
 void FAnimTrace::OutputAnimNodeValue(const FAnimationBaseContext& InContext, const TCHAR* InKey, const FName& InValue)
@@ -889,22 +876,16 @@ void FAnimTrace::OutputAnimNodeValue(const FAnimationBaseContext& InContext, con
 	UObject* AnimInstance = InContext.AnimInstanceProxy->GetAnimInstanceObject();
 	TRACE_OBJECT(AnimInstance);
 
-	int32 KeyLength = FCString::Strlen(InKey) + 1;
-	int32 ValueLength = InValue.GetStringLength() + 1;
-
-	auto StringCopyFunc = [KeyLength, ValueLength, InKey, &InValue](uint8* Out)
-	{
-		FCString::Strncpy(reinterpret_cast<TCHAR*>(Out), InKey, KeyLength);
-		InValue.ToString(reinterpret_cast<TCHAR*>(Out) + KeyLength, ValueLength);
-	};
+	TCHAR Value[FName::StringBufferSize];
+	uint32 ValueLength = InValue.ToString(Value);
 
-	UE_TRACE_LOG(Animation, AnimNodeValueString, AnimationChannel, (KeyLength + ValueLength) * sizeof(TCHAR))
+	UE_TRACE_LOG(Animation, AnimNodeValueString, AnimationChannel)
 		<< AnimNodeValueString.Cycle(FPlatformTime::Cycles64())
 		<< AnimNodeValueString.AnimInstanceId(FObjectTrace::GetObjectId(AnimInstance))
 		<< AnimNodeValueString.NodeId(InContext.GetCurrentNodeId())
-		<< AnimNodeValueString.KeyLength(KeyLength)
 		<< AnimNodeValueString.FrameCounter(FObjectTrace::GetObjectWorldTickCounter(AnimInstance))
-		<< AnimNodeValueString.Attachment(StringCopyFunc);
+		<< AnimNodeValueString.Key(InKey)
+		<< AnimNodeValueString.Value(Value, ValueLength);
 }
 
 void FAnimTrace::OutputAnimNodeValue(const FAnimationBaseContext& InContext, const TCHAR* InKey, const TCHAR* InValue)
@@ -925,22 +906,13 @@ void FAnimTrace::OutputAnimNodeValue(const FAnimationBaseContext& InContext, con
 	UObject* AnimInstance = InContext.AnimInstanceProxy->GetAnimInstanceObject();
 	TRACE_OBJECT(AnimInstance);
 
-	int32 KeyLength = FCString::Strlen(InKey) + 1;
-	int32 ValueLength = FCString::Strlen(InValue) + 1;
-
-	auto StringCopyFunc = [KeyLength, ValueLength, InKey, InValue](uint8* Out)
-	{
-		FCString::Strncpy(reinterpret_cast<TCHAR*>(Out), InKey, KeyLength);
-		FCString::Strncpy(reinterpret_cast<TCHAR*>(Out) + KeyLength, InValue, ValueLength);
-	};
-
-	UE_TRACE_LOG(Animation, AnimNodeValueString, AnimationChannel, (KeyLength + ValueLength) * sizeof(TCHAR))
+	UE_TRACE_LOG(Animation, AnimNodeValueString, AnimationChannel)
 		<< AnimNodeValueString.Cycle(FPlatformTime::Cycles64())
 		<< AnimNodeValueString.AnimInstanceId(FObjectTrace::GetObjectId(AnimInstance))
 		<< AnimNodeValueString.NodeId(InContext.GetCurrentNodeId())
-		<< AnimNodeValueString.KeyLength(KeyLength)
 		<< AnimNodeValueString.FrameCounter(FObjectTrace::GetObjectWorldTickCounter(AnimInstance))
-		<< AnimNodeValueString.Attachment(StringCopyFunc);
+		<< AnimNodeValueString.Key(InKey)
+		<< AnimNodeValueString.Value(InValue);
 }
 
 void FAnimTrace::OutputAnimNodeValue(const FAnimationBaseContext& InContext, const TCHAR* InKey, const UObject* InValue)
@@ -962,16 +934,13 @@ void FAnimTrace::OutputAnimNodeValue(const FAnimationBaseContext& InContext, con
 	TRACE_OBJECT(AnimInstance);
 	TRACE_OBJECT(InValue);
 
-	int32 KeyLength = FCString::Strlen(InKey) + 1;
-
-	UE_TRACE_LOG(Animation, AnimNodeValueObject, AnimationChannel, KeyLength * sizeof(TCHAR))
+	UE_TRACE_LOG(Animation, AnimNodeValueObject, AnimationChannel)
 		<< AnimNodeValueObject.Cycle(FPlatformTime::Cycles64())
 		<< AnimNodeValueObject.AnimInstanceId(FObjectTrace::GetObjectId(AnimInstance))
 		<< AnimNodeValueObject.NodeId(InContext.GetCurrentNodeId())
 		<< AnimNodeValueObject.Value(FObjectTrace::GetObjectId(InValue))
-		<< AnimNodeValueObject.KeyLength(KeyLength)
 		<< AnimNodeValueObject.FrameCounter(FObjectTrace::GetObjectWorldTickCounter(AnimInstance))
-		<< AnimNodeValueObject.Attachment(InKey, KeyLength * sizeof(TCHAR));
+		<< AnimNodeValueObject.Key(InKey);
 }
 
 void FAnimTrace::OutputAnimNodeValue(const FAnimationBaseContext& InContext, const TCHAR* InKey, const UClass* InValue)
@@ -993,16 +962,13 @@ void FAnimTrace::OutputAnimNodeValue(const FAnimationBaseContext& InContext, con
 	TRACE_OBJECT(AnimInstance);
 	TRACE_CLASS(InValue);
 
-	int32 KeyLength = FCString::Strlen(InKey) + 1;
-
-	UE_TRACE_LOG(Animation, AnimNodeValueClass, AnimationChannel, KeyLength * sizeof(TCHAR))
+	UE_TRACE_LOG(Animation, AnimNodeValueClass, AnimationChannel)
 		<< AnimNodeValueClass.Cycle(FPlatformTime::Cycles64())
 		<< AnimNodeValueClass.AnimInstanceId(FObjectTrace::GetObjectId(AnimInstance))
 		<< AnimNodeValueClass.NodeId(InContext.GetCurrentNodeId())
 		<< AnimNodeValueClass.Value(FObjectTrace::GetObjectId(InValue))
-		<< AnimNodeValueClass.KeyLength(KeyLength)
 		<< AnimNodeValueClass.FrameCounter(FObjectTrace::GetObjectWorldTickCounter(AnimInstance))
-		<< AnimNodeValueClass.Attachment(InKey, KeyLength * sizeof(TCHAR));
+		<< AnimNodeValueClass.Key(InKey);
 }
 
 void FAnimTrace::OutputAnimSequencePlayer(const FAnimationBaseContext& InContext, const FAnimNode_SequencePlayerBase& InNode)
diff --git a/Engine/Source/Runtime/Engine/Private/ObjectTrace.cpp b/Engine/Source/Runtime/Engine/Private/ObjectTrace.cpp
--- a/Engine/Source/Runtime/Engine/Private/ObjectTrace.cpp
+++ b/Engine/Source/Runtime/Engine/Private/ObjectTrace.cpp
@@ -26,20 +26,22 @@ UE_TRACE_CHANNEL(ObjectChannel)
 UE_TRACE_EVENT_BEGIN(Object, Class)
 	UE_TRACE_EVENT_FIELD(uint64, Id)
 	UE_TRACE_EVENT_FIELD(uint64, SuperId)
-	UE_TRACE_EVENT_FIELD(int32, ClassNameStringLength)
+	UE_TRACE_EVENT_FIELD(UE::Trace::WideString, Name)
+	UE_TRACE_EVENT_FIELD(UE::Trace::WideString, Path)
 UE_TRACE_EVENT_END()
 
 UE_TRACE_EVENT_BEGIN(Object, Object)
 	UE_TRACE_EVENT_FIELD(uint64, Id)
 	UE_TRACE_EVENT_FIELD(uint64, ClassId)
 	UE_TRACE_EVENT_FIELD(uint64, OuterId)
-	UE_TRACE_EVENT_FIELD(int32, ObjectNameStringLength)
+	UE_TRACE_EVENT_FIELD(UE::Trace::WideString, Name)
+	UE_TRACE_EVENT_FIELD(UE::Trace::WideString, Path)
 UE_TRACE_EVENT_END()
 
 UE_TRACE_EVENT_BEGIN(Object, ObjectEvent)
 	UE_TRACE_EVENT_FIELD(uint64, Cycle)
 	UE_TRACE_EVENT_FIELD(uint64, Id)
-	UE_TRACE_EVENT_FIELD(uint8, Event)
+	UE_TRACE_EVENT_FIELD(UE::Trace::WideString, Event)
 UE_TRACE_EVENT_END()
 
 UE_TRACE_EVENT_BEGIN(Object, World)
@@ -289,21 +291,15 @@ void FObjectTrace::OutputClass(const UClass* InClass)
 	Annotation.bTraced = true;
 	GObjectTraceAnnotations.AddAnnotation(InClass, MoveTemp(Annotation));
 
-	int32 ClassNameStringLength = InClass->GetFName().GetStringLength() + 1;
-	int32 ClassFullNameStringLength = InClass->GetPathName().Len() + 1;
+	FString ClassPathName = InClass->GetPathName();
+	TCHAR ClassName[FName::StringBufferSize];
+	uint32 ClassNameLength = InClass->GetFName().ToString(ClassName);
 
-	auto StringCopyFunc = [ClassNameStringLength, ClassFullNameStringLength, InClass](uint8* Out)
-	{
-		InClass->GetFName().ToString(reinterpret_cast<TCHAR*>(Out), ClassNameStringLength);
-		FPlatformMemory::Memcpy(reinterpret_cast<TCHAR*>(Out) + ClassNameStringLength, *InClass->GetPathName(), ClassFullNameStringLength * sizeof(TCHAR));
-	};
-
-	UE_TRACE_LOG(Object, Class, ObjectChannel, (ClassNameStringLength + ClassFullNameStringLength) * sizeof(TCHAR))
-		<< Class.ClassNameStringLength(ClassNameStringLength)
+	UE_TRACE_LOG(Object, Class, ObjectChannel)
 		<< Class.Id(GetObjectId(InClass))
 		<< Class.SuperId(GetObjectId(InClass->GetSuperClass()))
-		<< Class.Attachment(StringCopyFunc);
-
+		<< Class.Name(ClassName, ClassNameLength)
+		<< Class.Path(*ClassPathName, ClassPathName.Len());
 }
 
 void FObjectTrace::OutputView(const UObject* InPlayer, const FSceneView* InView)
@@ -368,21 +364,16 @@ void FObjectTrace::OutputObject(const UObject* InObject)
 	// Trace the object's class first
 	TRACE_CLASS(InObject->GetClass());
 
-	int32 ObjectNameStringLength = InObject->GetFName().GetStringLength() + 1;
-	int32 ObjectPathNameStringLength = InObject->GetPathName().Len() + 1;
+	TCHAR ObjectName[FName::StringBufferSize];
+	uint32 ObjectNameLength = InObject->GetFName().ToString(ObjectName);
+	FString ObjectPathName = InObject->GetPathName();
 
-	auto StringCopyFunc = [ObjectNameStringLength, ObjectPathNameStringLength, InObject](uint8* Out)
-	{
-		InObject->GetFName().ToString(reinterpret_cast<TCHAR*>(Out), ObjectNameStringLength);
-		FPlatformMemory::Memcpy(reinterpret_cast<TCHAR*>(Out) + ObjectNameStringLength, *InObject->GetPathName(), ObjectPathNameStringLength * sizeof(TCHAR));
-	};
-
-	UE_TRACE_LOG(Object, Object, ObjectChannel, (ObjectNameStringLength + ObjectPathNameStringLength) * sizeof(TCHAR))
-		<< Object.ObjectNameStringLength(ObjectNameStringLength)
+	UE_TRACE_LOG(Object, Object, ObjectChannel)
 		<< Object.Id(GetObjectId(InObject))
 		<< Object.ClassId(GetObjectId(InObject->GetClass()))
 		<< Object.OuterId(GetObjectId(InObject->GetOuter()))
-		<< Object.Attachment(StringCopyFunc);
+		<< Object.Name(ObjectName, ObjectNameLength)
+		<< Object.Path(*ObjectPathName, ObjectPathName.Len());
 }
 
 void FObjectTrace::OutputObjectEvent(const UObject* InObject, const TCHAR* InEvent)
@@ -405,12 +396,10 @@ void FObjectTrace::OutputObjectEvent(const UObject* InObject, const TCHAR* InEve
 
 	TRACE_OBJECT(InObject);
 
-	int32 StringBufferSize = (FCString::Strlen(InEvent) + 1) * sizeof(TCHAR);
-
-	UE_TRACE_LOG(Object, ObjectEvent, ObjectChannel, StringBufferSize)
+	UE_TRACE_LOG(Object, ObjectEvent, ObjectChannel)
 		<< ObjectEvent.Cycle(FPlatformTime::Cycles64())
 		<< ObjectEvent.Id(GetObjectId(InObject))
-		<< ObjectEvent.Attachment(InEvent, StringBufferSize);
+		<< ObjectEvent.Event(InEvent);
 }
 
 void FObjectTrace::OutputWorld(const UWorld* InWorld)
