Use world is paused flag from scene view instead of querying from UWorld directly, so it is requested for rendering from the right thread

diff --git a/Engine/Source/Runtime/Engine/Private/GPUSkinCache.cpp b/Engine/Source/Runtime/Engine/Private/GPUSkinCache.cpp
--- a/Engine/Source/Runtime/Engine/Private/GPUSkinCache.cpp
+++ b/Engine/Source/Runtime/Engine/Private/GPUSkinCache.cpp
@@ -1897,7 +1897,8 @@ void FGPUSkinCache::GetShaderBindings(
 	FShaderResourceParameter GPUSkinCachePositionBuffer,
 	FShaderResourceParameter GPUSkinCachePreviousPositionBuffer,
 	class FMeshDrawSingleShaderBindings& ShaderBindings,
-	FVertexInputStreamArray& VertexStreams)
+	FVertexInputStreamArray& VertexStreams,
+	bool bWorldIsPaused)
 {
 	INC_DWORD_STAT(STAT_GPUSkinCache_NumSetVertexStreams);
 	check(Entry);
@@ -1918,8 +1919,7 @@ void FGPUSkinCache::GetShaderBindings(
 	ShaderBindings.Add(GPUSkinCachePositionBuffer, DispatchData.GetPositionRWBuffer()->Buffer.SRV);
 
 	// If world is paused, use current frame bone matrices, so velocity is canceled and skeletal mesh isn't blurred from motion.
-	const bool bWorldPaused = Entry->SkinCache->IsWorldPaused();
-	ShaderBindings.Add(GPUSkinCachePreviousPositionBuffer, bWorldPaused ? DispatchData.GetPositionRWBuffer()->Buffer.SRV : DispatchData.GetPreviousPositionRWBuffer()->Buffer.SRV);
+	ShaderBindings.Add(GPUSkinCachePreviousPositionBuffer, bWorldIsPaused ? DispatchData.GetPositionRWBuffer()->Buffer.SRV : DispatchData.GetPreviousPositionRWBuffer()->Buffer.SRV);
 }
 
 void FGPUSkinCache::PrepareUpdateSkinning(FGPUSkinCacheEntry* Entry, int32 Section, uint32 RevisionNumber, TArray<FSkinCacheRWBuffer*>* OverlappedUAVs)
diff --git a/Engine/Source/Runtime/Engine/Private/GPUSkinVertexFactory.cpp b/Engine/Source/Runtime/Engine/Private/GPUSkinVertexFactory.cpp
--- a/Engine/Source/Runtime/Engine/Private/GPUSkinVertexFactory.cpp
+++ b/Engine/Source/Runtime/Engine/Private/GPUSkinVertexFactory.cpp
@@ -699,8 +699,7 @@ public:
 				bLocalPerBoneMotionBlur = true;
 
 				// If world is paused, use current frame bone matrices, so velocity is canceled and skeletal mesh isn't blurred from motion.
-				const bool bWorldPaused = Scene->GetWorld() ? Scene->GetWorld()->IsPaused() : false;
-				FRHIShaderResourceView* PreviousData = ShaderData.GetBoneBufferForReading(bWorldPaused ? false : true).VertexBufferSRV;
+				FRHIShaderResourceView* PreviousData = ShaderData.GetBoneBufferForReading(!View->Family->bWorldIsPaused_IncludingSimulatingInEditor).VertexBufferSRV;
 				ShaderBindings.Add(PreviousBoneMatrices, PreviousData);
 			}
 		}
@@ -786,7 +785,7 @@ public:
 
 		// #dxr_todo do we need this call to the base?
 		FLocalVertexFactoryShaderParametersBase::GetElementShaderBindingsBase(Scene, View, Shader, InputStreamType, FeatureLevel, VertexFactory, BatchElement, VertexFactoryUniformBuffer, ShaderBindings, VertexStreams);
-		FGPUSkinCache::GetShaderBindings(BatchUserData->Entry, BatchUserData->Section, Shader, (const FGPUSkinPassthroughVertexFactory*)VertexFactory, BatchElement.MinVertexIndex, GPUSkinCachePositionBuffer, GPUSkinCachePreviousPositionBuffer, ShaderBindings, VertexStreams);
+		FGPUSkinCache::GetShaderBindings(BatchUserData->Entry, BatchUserData->Section, Shader, (const FGPUSkinPassthroughVertexFactory*)VertexFactory, BatchElement.MinVertexIndex, GPUSkinCachePositionBuffer, GPUSkinCachePreviousPositionBuffer, ShaderBindings, VertexStreams, View->Family->bWorldIsPaused_IncludingSimulatingInEditor);
 	}
 
 private:
diff --git a/Engine/Source/Runtime/Engine/Private/SceneView.cpp b/Engine/Source/Runtime/Engine/Private/SceneView.cpp
--- a/Engine/Source/Runtime/Engine/Private/SceneView.cpp
+++ b/Engine/Source/Runtime/Engine/Private/SceneView.cpp
@@ -2577,6 +2577,7 @@ FSceneViewFamily::FSceneViewFamily(const ConstructionValues& CVS)
 	SceneCaptureSource(SCS_FinalColorLDR),
 	SceneCaptureCompositeMode(SCCM_Overwrite),
 	bWorldIsPaused(false),
+	bWorldIsPaused_IncludingSimulatingInEditor(false),
 	bIsHDR(false),
 	bRequireMultiView(false),
 	GammaCorrection(CVS.GammaCorrection),
@@ -2627,6 +2628,7 @@ FSceneViewFamily::FSceneViewFamily(const ConstructionValues& CVS)
 			}
 
 			bWorldIsPaused = !World->IsCameraMoveable();
+			bWorldIsPaused_IncludingSimulatingInEditor = World->IsPaused();
 		}
 	}
 
diff --git a/Engine/Source/Runtime/Engine/Public/GPUSkinCache.h b/Engine/Source/Runtime/Engine/Public/GPUSkinCache.h
--- a/Engine/Source/Runtime/Engine/Public/GPUSkinCache.h
+++ b/Engine/Source/Runtime/Engine/Public/GPUSkinCache.h
@@ -188,7 +188,8 @@ public:
 		FShaderResourceParameter GPUSkinCachePositionBuffer,
 		FShaderResourceParameter GPUSkinCachePreviousPositionBuffer,
 		class FMeshDrawSingleShaderBindings& ShaderBindings,
-		FVertexInputStreamArray& VertexStreams);
+		FVertexInputStreamArray& VertexStreams,
+		bool bWorldIsPaused);
 
 	static void Release(FGPUSkinCacheEntry*& SkinCacheEntry);
 
@@ -449,7 +450,6 @@ public:
 	bool IsBatchingDispatch() const { return bShouldBatchDispatches; }
 
 	inline ERHIFeatureLevel::Type GetFeatureLevel() const { return FeatureLevel; }
-	inline bool IsWorldPaused() const { return World ? World->IsPaused() : false; }
 
 protected:
 	void MakeBufferTransitions(FRHICommandListImmediate& RHICmdList, TArray<FSkinCacheRWBuffer*>& Buffers, ERHIAccess ToState);
diff --git a/Engine/Source/Runtime/Engine/Public/SceneView.h b/Engine/Source/Runtime/Engine/Public/SceneView.h
--- a/Engine/Source/Runtime/Engine/Public/SceneView.h
+++ b/Engine/Source/Runtime/Engine/Public/SceneView.h
@@ -1647,6 +1647,9 @@ public:
 	 */
 	bool bWorldIsPaused;
 
+	/** GetWorld->IsPaused(). Simulate-in-Editor included. */
+	bool bWorldIsPaused_IncludingSimulatingInEditor;
+
 	/** When enabled, the post processing will output in HDR space */
 	bool bIsHDR;
 
