Add ability to define what shader frequency a material custom output operates in.

diff --git a/Engine/Source/Runtime/Engine/Classes/Materials/MaterialExpressionCustomOutput.h b/Engine/Source/Runtime/Engine/Classes/Materials/MaterialExpressionCustomOutput.h
--- a/Engine/Source/Runtime/Engine/Classes/Materials/MaterialExpressionCustomOutput.h
+++ b/Engine/Source/Runtime/Engine/Classes/Materials/MaterialExpressionCustomOutput.h
@@ -24,6 +24,7 @@ class UMaterialExpressionCustomOutput : public UMaterialExpression
 	virtual bool AllowMultipleCustomOutputs() { return false; }
 	virtual bool NeedsCustomOutputDefines() { return true; }
 	virtual bool ShouldCompileBeforeAttributes() { return false; }
+	virtual EShaderFrequency GetShaderFrequency() { return SF_Pixel; }
 #endif
 };
 
diff --git a/Engine/Source/Runtime/Engine/Private/Materials/HLSLMaterialTranslator.cpp b/Engine/Source/Runtime/Engine/Private/Materials/HLSLMaterialTranslator.cpp
--- a/Engine/Source/Runtime/Engine/Private/Materials/HLSLMaterialTranslator.cpp
+++ b/Engine/Source/Runtime/Engine/Private/Materials/HLSLMaterialTranslator.cpp
@@ -390,21 +390,22 @@ void FHLSLMaterialTranslator::CompileCustomOutputs(TArray<UMaterialExpressionCus
 
 			if (NumOutputs > 0)
 			{
+				EShaderFrequency CustomOutputShaderFrequency = CustomOutput->GetShaderFrequency();
 				for (int32 Index = 0; Index < NumOutputs; Index++)
 				{
 					{
-						ClearFunctionStack(SF_Pixel);
-						FunctionStacks[SF_Pixel].Add(new FMaterialFunctionCompileState(nullptr));
+						ClearFunctionStack(CustomOutputShaderFrequency);
+						FunctionStacks[CustomOutputShaderFrequency].Add(new FMaterialFunctionCompileState(nullptr));
 					}
 					MaterialProperty = MP_MAX; // Indicates we're not compiling any material property.
-					ShaderFrequency = SF_Pixel;
+					ShaderFrequency = CustomOutputShaderFrequency;
 					TArray<FShaderCodeChunk> CustomExpressionChunks;
 					AssignTempScope(CustomExpressionChunks);
 					CustomOutput->Compile(this, Index);
 				}
 
-				ClearFunctionStack(SF_Pixel);
-				FunctionStacks[SF_Pixel].Add(new FMaterialFunctionCompileState(nullptr));
+				ClearFunctionStack(CustomOutputShaderFrequency);
+				FunctionStacks[CustomOutputShaderFrequency].Add(new FMaterialFunctionCompileState(nullptr));
 			}
 		}
 	}
