Use GetPathName for texture worker so we know what failed in logs.

diff --git a/Engine/Source/Runtime/Engine/Private/TextureDerivedDataTask.cpp b/Engine/Source/Runtime/Engine/Private/TextureDerivedDataTask.cpp
--- a/Engine/Source/Runtime/Engine/Private/TextureDerivedDataTask.cpp
+++ b/Engine/Source/Runtime/Engine/Private/TextureDerivedDataTask.cpp
@@ -444,7 +444,7 @@ void FTextureCacheDerivedDataWorker::BuildTexture(bool bReplaceExistingDDC)
 	const bool bForVirtualTextureStreamingBuild = EnumHasAnyFlags(CacheFlags, ETextureCacheFlags::ForVirtualTextureStreamingBuild);
 
 	FFormatNamedArguments Args;
-	Args.Add(TEXT("TextureName"), FText::FromString(Texture.GetName()));
+	Args.Add(TEXT("TextureName"), FText::FromString(Texture.GetPathName()));
 	Args.Add(TEXT("TextureFormatName"), FText::FromString(BuildSettingsPerLayer[0].TextureFormatName.GetPlainNameString()));
 	Args.Add(TEXT("IsVT"), FText::FromString( FString( bForVirtualTextureStreamingBuild ? TEXT(" VT") : TEXT("") ) ) );
 	Args.Add(TEXT("TextureResolutionX"), FText::FromString(FString::FromInt(bHasValidMip0 ? TextureData.Blocks[0].MipsPerLayer[0][0].SizeX : 0)));
