Level Instance: Replace check by if when Embedding actors that failed to load (ex: Deleted Blueprint type)

diff --git a/Engine/Source/Runtime/Engine/Private/WorldPartition/WorldPartitionLevelHelper.cpp b/Engine/Source/Runtime/Engine/Private/WorldPartition/WorldPartitionLevelHelper.cpp
--- a/Engine/Source/Runtime/Engine/Private/WorldPartition/WorldPartitionLevelHelper.cpp
+++ b/Engine/Source/Runtime/Engine/Private/WorldPartition/WorldPartitionLevelHelper.cpp
@@ -56,6 +56,12 @@ void FWorldPartitionLevelHelper::MoveExternalActorsToLevel(const TArray<FWorldPa
 	// Move all actors to Cell level
 	for (const FWorldPartitionRuntimeCellObjectMapping& PackageObjectMapping : InChildPackages)
 	{
+		// We assume actor failed to duplicate if LoadedPath equals NAME_None (warning already logged we can skip this mapping)
+		if (PackageObjectMapping.LoadedPath == NAME_None && PackageObjectMapping.ContainerID != 0)
+		{
+			continue;
+		}
+
 		AActor* Actor = FindObject<AActor>(nullptr, *PackageObjectMapping.LoadedPath.ToString());
 		if (ensure(Actor))
 		{
@@ -325,27 +331,34 @@ bool FWorldPartitionLevelHelper::LoadActors(ULevel* InDestLevel, TArrayView<FWor
 				if (!DuplicatedPackage)
 				{
 					DuplicatedPackage = InPackageCache.DuplicateWorldPackage(LoadedPackage, DuplicatePackageName);
-					check(DuplicatedPackage)
-						DuplicatedWorld = UWorld::FindWorldInPackage(DuplicatedPackage);
+					check(DuplicatedPackage);
+					DuplicatedWorld = UWorld::FindWorldInPackage(DuplicatedPackage);
 				}
 
 				for (auto Mapping : Mappings)
 				{
 					FString ActorName = FPaths::GetExtension(Mapping->Path.ToString());
 					AActor* DuplicatedActor = FindObject<AActor>(DuplicatedWorld->PersistentLevel, *ActorName);
-					check(DuplicatedActor);
-
-					DuplicatedActor->Rename(*FString::Printf(TEXT("%s_%016llx"), *DuplicatedActor->GetName(), Mapping->ContainerID), InDestLevel, REN_NonTransactional | REN_ForceNoResetLoaders | REN_DoNotDirty | REN_DontCreateRedirectors);
-					USceneComponent* RootComponent = DuplicatedActor->GetRootComponent();
-					
-					FLevelUtils::FApplyLevelTransformParams TransformParams(nullptr, Mapping->ContainerTransform);
-					TransformParams.Actor = DuplicatedActor;
-					TransformParams.bDoPostEditMove = false;
-					FLevelUtils::ApplyLevelTransform(TransformParams);
-
-					// Path to use when searching for this actor in MoveExternalActorsToLevel
-					Mapping->LoadedPath = *DuplicatedActor->GetPathName();
-					UE_LOG(LogEngine, Verbose, TEXT(" ==> Duplicated %s (remaining: %d)"), *DuplicatedActor->GetFullName(), LoadProgress->NumPendingLoadRequests);
+					// Possible actor isn't found if Actor package failed to load because of missing dependencies (ex: deleted Blueprint)
+					if (DuplicatedActor)
+					{
+						DuplicatedActor->Rename(*FString::Printf(TEXT("%s_%016llx"), *DuplicatedActor->GetName(), Mapping->ContainerID), InDestLevel, REN_NonTransactional | REN_ForceNoResetLoaders | REN_DoNotDirty | REN_DontCreateRedirectors);
+						USceneComponent* RootComponent = DuplicatedActor->GetRootComponent();
+
+						FLevelUtils::FApplyLevelTransformParams TransformParams(nullptr, Mapping->ContainerTransform);
+						TransformParams.Actor = DuplicatedActor;
+						TransformParams.bDoPostEditMove = false;
+						FLevelUtils::ApplyLevelTransform(TransformParams);
+
+						// Path to use when searching for this actor in MoveExternalActorsToLevel
+						Mapping->LoadedPath = *DuplicatedActor->GetPathName();
+						UE_LOG(LogEngine, Verbose, TEXT(" ==> Duplicated %s (remaining: %d)"), *DuplicatedActor->GetFullName(), LoadProgress->NumPendingLoadRequests);
+					}
+					else
+					{
+						Mapping->LoadedPath = NAME_None;
+						UE_LOG(LogEngine, Warning, TEXT("Failed to find actor %s in package %s"), *ActorName, *LoadedPackageName.ToString());
+					}
 				}
 			}
 			else
