Fixed the actor partition subsystem returning a partition actor which may not have been of the queried class type in non-partitioned worlds.

diff --git a/Engine/Source/Runtime/Engine/Private/ActorPartition/ActorPartitionSubsystem.cpp b/Engine/Source/Runtime/Engine/Private/ActorPartition/ActorPartitionSubsystem.cpp
--- a/Engine/Source/Runtime/Engine/Private/ActorPartition/ActorPartitionSubsystem.cpp
+++ b/Engine/Source/Runtime/Engine/Private/ActorPartition/ActorPartitionSubsystem.cpp
@@ -109,7 +109,7 @@ public:
 		{
 			if (APartitionActor* PartitionActor = Cast<APartitionActor>(Actor))
 			{
-				if (PartitionActor->GetGridGuid() == InActorPartitionId.GetGridGuid())
+				if (PartitionActor->IsA(InActorPartitionId.GetClass()) && (PartitionActor->GetGridGuid() == InActorPartitionId.GetGridGuid()))
 				{
 					FoundActor = PartitionActor;
 					break;
