Add OodlePNGz0 option for Texture asset compression
fix handling of Uncompressed Quality in JPEG image wrapper

diff --git a/Engine/Source/Runtime/Engine/Private/Texture.cpp b/Engine/Source/Runtime/Engine/Private/Texture.cpp
--- a/Engine/Source/Runtime/Engine/Private/Texture.cpp
+++ b/Engine/Source/Runtime/Engine/Private/Texture.cpp
@@ -457,7 +457,7 @@ void UTexture::Serialize(FArchive& Ar)
 	Ar.UsingCustomVersion(FUE5MainStreamObjectVersion::GUID);
 
 	Super::Serialize(Ar);
-
+	
 	FStripDataFlags StripFlags(Ar);
 
 	/** Legacy serialization. */
@@ -1256,6 +1256,12 @@ void FTextureSource::Compress()
 	FWriteScopeLock BulkDataExclusiveScope(BulkDataLock.Get());
 #endif
 
+	// if bUseOodleOnPNGz0 , do PNG filters but then use Oodle instead of zlib back-end LZ
+	//	no point turning this on until we drill Oodle compressor options through BulkData
+	bool bUseOodleOnPNGz0 = false;
+
+	// may already have bPNGCompressed or "CompressionFormat" set
+
 	if (CanPNGCompress()) // Note that this will return false if the data is already a compressed PNG
 	{
 		FSharedBuffer Payload = BulkData.GetPayload().Get();
@@ -1264,13 +1270,19 @@ void FTextureSource::Compress()
 		TSharedPtr<IImageWrapper> ImageWrapper = ImageWrapperModule.CreateImageWrapper( EImageFormat::PNG );
 		// TODO: TSF_BGRA8 is stored as RGBA, so the R and B channels are swapped in the internal png. Should we fix this?
 		ERGBFormat RawFormat = (Format == TSF_G8 || Format == TSF_G16) ? ERGBFormat::Gray : ERGBFormat::RGBA;
-		if ( ImageWrapper.IsValid() && ImageWrapper->SetRaw(Payload.GetData(), Payload.GetSize(), SizeX, SizeY, RawFormat, (Format == TSF_G16 || Format == TSF_RGBA16) ? 16 : 8 ) )
+		int RawBitsPerChannel = (Format == TSF_G16 || Format == TSF_RGBA16) ? 16 : 8;
+		if ( ImageWrapper.IsValid() && ImageWrapper->SetRaw(Payload.GetData(), Payload.GetSize(), SizeX, SizeY, RawFormat, RawBitsPerChannel ) )
 		{
-			TArray64<uint8> CompressedData = ImageWrapper->GetCompressed();
+			int32 PngQuality = (int32)EImageCompressionQuality::Default; // 0 means default 
+			if ( bUseOodleOnPNGz0 )
+			{
+				PngQuality = (int32)EImageCompressionQuality::Uncompressed; // turn off zlib
+			}
+			TArray64<uint8> CompressedData = ImageWrapper->GetCompressed(PngQuality);
 			if ( CompressedData.Num() > 0 )
-			{			
+			{
 				BulkData.UpdatePayload(MakeSharedBufferFromArray(MoveTemp(CompressedData)), NAME_None);
-
+				
 				bPNGCompressed = true;
 				CompressionFormat = TSCF_PNG;
 			}
@@ -1285,7 +1297,20 @@ void FTextureSource::Compress()
 
 	// There was a brief period when textures could be saved with the wrong compression flag based on their storage format
 	// so we force set it here so that they will be fixed on resave.
-	BulkData.SetCompressionFormat(CompressionFormat == TSCF_None ? NAME_Default : NAME_None);
+
+	FName Compressor = NAME_None;
+	if ( CompressionFormat == TSCF_None )
+	{
+		//was NAME_Default which is Zlib
+		Compressor = NAME_Default;
+		// @todo use Oodle for all non-PNG and JPEG
+		//Compressor = NAME_Oodle;
+	}
+	else if ( CompressionFormat == TSCF_PNG && bUseOodleOnPNGz0 )
+	{
+		Compressor = NAME_Oodle;
+	}
+	BulkData.SetCompressionFormat(Compressor);
 }
 
 FSharedBuffer FTextureSource::Decompress(IImageWrapperModule* ImageWrapperModule) const
diff --git a/Engine/Source/Runtime/ImageWrapper/Private/Formats/JpegImageWrapper.cpp b/Engine/Source/Runtime/ImageWrapper/Private/Formats/JpegImageWrapper.cpp
--- a/Engine/Source/Runtime/ImageWrapper/Private/Formats/JpegImageWrapper.cpp
+++ b/Engine/Source/Runtime/ImageWrapper/Private/Formats/JpegImageWrapper.cpp
@@ -144,6 +144,25 @@ bool FJpegImageWrapper::SetRaw(const void* InRawData, int64 InRawSize, const int
 
 void FJpegImageWrapper::Compress(int32 Quality)
 {
+	if (Quality == 0)
+	{
+		//default 
+		Quality = 85;
+	}
+	else if (Quality == (int32)EImageCompressionQuality::Uncompressed)
+	{
+		// fix = 1 (Uncompressed) was previously treated as max-compress
+		Quality = 100;
+	}
+	else
+	{
+		ensure(Quality >= 1 && Quality <= 100);
+
+		#define JPEG_QUALITY_MIN 30
+		// JPEG should not be used below quality JPEG_QUALITY_MIN
+		Quality = FMath::Clamp(Quality, JPEG_QUALITY_MIN, 100);
+	}
+
 #if WITH_LIBJPEGTURBO
 	CompressTurbo(Quality);
 #else
@@ -151,10 +170,6 @@ void FJpegImageWrapper::Compress(int32 Quality)
 	{
 		FScopeLock JPEGLock(&GJPEGSection);
 		
-		if (Quality == 0) {Quality = 85;}
-		ensure(Quality >= 1 && Quality <= 100);
-		Quality = FMath::Clamp(Quality, 1, 100);
-
 		check(RawData.Num());
 		check(Width > 0);
 		check(Height > 0);
@@ -270,12 +285,10 @@ void FJpegImageWrapper::CompressTurbo(int32 Quality)
 	{
 		FScopeLock JPEGLock(&GJPEGSection);
 
-		check(Compressor);
-
-		if (Quality == 0) { Quality = 85; }
-		ensure(Quality >= 1 && Quality <= 100);
-		Quality = FMath::Clamp(Quality, 1, 100);
+		// Quality mapping should have already been done
+		check( Quality >= JPEG_QUALITY_MIN && Quality <= 100 );
 
+		check(Compressor);
 		check(RawData.Num());
 		check(Width > 0);
 		check(Height > 0);
diff --git a/Engine/Source/Runtime/ImageWrapper/Private/Formats/PngImageWrapper.cpp b/Engine/Source/Runtime/ImageWrapper/Private/Formats/PngImageWrapper.cpp
--- a/Engine/Source/Runtime/ImageWrapper/Private/Formats/PngImageWrapper.cpp
+++ b/Engine/Source/Runtime/ImageWrapper/Private/Formats/PngImageWrapper.cpp
@@ -163,7 +163,31 @@ void FPngImageWrapper::Compress(int32 Quality)
 		// Anything allocated on the stack after this point will not be destructed correctly in the case of an error
 
 		{
-			png_set_compression_level(png_ptr, Z_BEST_SPEED);
+			int ZlibLevel = 3; // default
+			// Quality == 0 is the default argument, does not set a zlib level
+			if ( Quality != 0 )
+			{
+				if ( Quality == (int32)EImageCompressionQuality::Uncompressed )
+				{
+					ZlibLevel = 0; // compression off
+				}
+				else if ( -Quality >= 1 && -Quality <= 9 )
+				{
+					// negative quality for zlib level
+					ZlibLevel = -Quality;
+				}
+				else if ( Quality >= 20 && Quality <= 100 )
+				{
+					// JPEG quality, just ignore
+					// calls to GetCompressed(100) are common
+				}
+				else
+				{
+					UE_LOG(LogImageWrapper, Warning, TEXT("PNG Quality ZlibLevel out of range %d"), Quality );
+				}
+			}
+
+			png_set_compression_level(png_ptr, ZlibLevel);
 			png_set_IHDR(png_ptr, info_ptr, Width, Height, RawBitDepth, (RawFormat == ERGBFormat::Gray) ? PNG_COLOR_TYPE_GRAY : PNG_COLOR_TYPE_RGBA, PNG_INTERLACE_NONE, PNG_COMPRESSION_TYPE_DEFAULT, PNG_FILTER_TYPE_DEFAULT);
 			png_set_write_fn(png_ptr, this, FPngImageWrapper::user_write_compressed, FPngImageWrapper::user_flush_data);
 
diff --git a/Engine/Source/Runtime/ImageWrapper/Public/IImageWrapper.h b/Engine/Source/Runtime/ImageWrapper/Public/IImageWrapper.h
--- a/Engine/Source/Runtime/ImageWrapper/Public/IImageWrapper.h
+++ b/Engine/Source/Runtime/ImageWrapper/Public/IImageWrapper.h
@@ -73,11 +73,13 @@ enum class ERGBFormat : int8
 
 /**
  * Enumerates available image compression qualities.
+ * 
+ * JPEG interprets Quality as 1-100
  */
 enum class EImageCompressionQuality : uint8
 {
 	Default = 0,
-	Uncompressed =  1,
+	Uncompressed = 1,
 };
 
 
