Reduce definition statements in HLSL material translator if they don't contribute to the function's return value.

diff --git a/Engine/Source/Runtime/Engine/Private/Materials/HLSLMaterialTranslator.cpp b/Engine/Source/Runtime/Engine/Private/Materials/HLSLMaterialTranslator.cpp
--- a/Engine/Source/Runtime/Engine/Private/Materials/HLSLMaterialTranslator.cpp
+++ b/Engine/Source/Runtime/Engine/Private/Materials/HLSLMaterialTranslator.cpp
@@ -1159,14 +1159,19 @@ bool FHLSLMaterialTranslator::Translate()
 				StartChunk = NormalCodeChunkEnd;
 			}
 
+			// Reduce definition statements that don't contribute to the function's return value.
+			// @todo-lh: This should be expanded to a general reduction, but is currently only intended to fix an FXC internal compiler error reported in UE-117831
+			const bool bReduceAfterReturnValue = (PropertyId == MP_WorldPositionOffset || PropertyId == CompiledMP_PrevWorldPositionOffset);
+
 			GetFixedParameterCode(
 				StartChunk,
 				SharedPropertyCodeChunks[PropertyShaderFrequency].Num(),
 				Chunk[PropertyId],
 				SharedPropertyCodeChunks[PropertyShaderFrequency],
-						DerivativeVariations[Variation].TranslatedCodeChunkDefinitions[PropertyId],
-						DerivativeVariations[Variation].TranslatedCodeChunks[PropertyId],
-						Variation);
+				DerivativeVariations[Variation].TranslatedCodeChunkDefinitions[PropertyId],
+				DerivativeVariations[Variation].TranslatedCodeChunks[PropertyId],
+				Variation,
+				bReduceAfterReturnValue);
 
 		}
 
@@ -1934,15 +1939,9 @@ void FHLSLMaterialTranslator::GetSharedInputsMaterialCode(FString& PixelMembersD
 			else
 			{
 					if (DerivativeVariations[DerivativeVariation].TranslatedCodeChunkDefinitions[Property].Len() > 0)
-				{
-					if (LastProperty >= 0)
 					{
-						// Verify that all code chunks have the same contents
-							check(DerivativeVariations[DerivativeVariation].TranslatedCodeChunkDefinitions[Property].Len() == DerivativeVariations[DerivativeVariation].TranslatedCodeChunkDefinitions[LastProperty].Len());
+						LastProperty = Property;
 					}
-
-					LastProperty = Property;
-				}
 				}
 
 				PixelInputInitializerValues += FString::Printf(TEXT("\tPixelMaterialInputs.%s = %s;\n"), *PropertyName, *DerivativeVariations[DerivativeVariation].TranslatedCodeChunks[Property]);
@@ -2117,13 +2116,9 @@ FString FHLSLMaterialTranslator::GetMaterialShaderCode()
 		}
 
 			if (DerivativeVariations[BaseDerivativeVariation].TranslatedCodeChunkDefinitions[MP_CustomizedUVs0 + CustomUVIndex].Len() > 0)
-		{
-			if (LastProperty >= 0)
 			{
-					check(DerivativeVariations[BaseDerivativeVariation].TranslatedCodeChunkDefinitions[LastProperty].Len() == DerivativeVariations[BaseDerivativeVariation].TranslatedCodeChunkDefinitions[MP_CustomizedUVs0 + CustomUVIndex].Len());
+				LastProperty = MP_CustomizedUVs0 + CustomUVIndex;
 			}
-			LastProperty = MP_CustomizedUVs0 + CustomUVIndex;
-		}
 			CustomUVAssignments += FString::Printf(TEXT("\tOutTexCoords[%u] = %s;") LINE_TERMINATOR, CustomUVIndex, *DerivativeVariations[BaseDerivativeVariation].TranslatedCodeChunks[MP_CustomizedUVs0 + CustomUVIndex]);
 		}
 	}
@@ -2457,7 +2452,7 @@ void FHLSLMaterialTranslator::SetParameterMaterialAttributes(int32 Index, uint64
 
 
 /** Creates a string of all definitions needed for the given material input. */
-FString FHLSLMaterialTranslator::GetDefinitions(TArray<FShaderCodeChunk>& CodeChunks, int32 StartChunk, int32 EndChunk, ECompiledPartialDerivativeVariation Variation) const
+FString FHLSLMaterialTranslator::GetDefinitions(const TArray<FShaderCodeChunk>& CodeChunks, int32 StartChunk, int32 EndChunk, ECompiledPartialDerivativeVariation Variation, const TCHAR* ReturnValueSymbolName) const
 {
 	FString Definitions;
 	for (int32 ChunkIndex = StartChunk; ChunkIndex < EndChunk; ChunkIndex++)
@@ -2468,13 +2463,19 @@ FString FHLSLMaterialTranslator::GetDefinitions(TArray<FShaderCodeChunk>& CodeCh
 			(!CodeChunk.bInline || CodeChunk.Type == MCT_VoidStatement))
 		{
 			Definitions += CodeChunk.AtDefinition(Variation);
+
+			// If we found the definition of the return value, there is no need to add more definitions as they won't contribute to the outcome
+			if (ReturnValueSymbolName != nullptr && CodeChunk.SymbolName == ReturnValueSymbolName)
+			{
+				break;
+			}
 		}
 	}
 	return Definitions;
 }
 
 // GetFixedParameterCode
-void FHLSLMaterialTranslator::GetFixedParameterCode(int32 StartChunk, int32 EndChunk, int32 ResultIndex, TArray<FShaderCodeChunk>& CodeChunks, FString& OutDefinitions, FString& OutValue, ECompiledPartialDerivativeVariation OriginalVariation)
+void FHLSLMaterialTranslator::GetFixedParameterCode(int32 StartChunk, int32 EndChunk, int32 ResultIndex, TArray<FShaderCodeChunk>& CodeChunks, FString& OutDefinitions, FString& OutValue, ECompiledPartialDerivativeVariation OriginalVariation, bool bReduceAfterReturnValue)
 {
 	// Only allow the analytic variation for pixel shaders.
 	ECompiledPartialDerivativeVariation Variation = OriginalVariation;
@@ -2493,9 +2494,12 @@ void FHLSLMaterialTranslator::GetFixedParameterCode(int32 StartChunk, int32 EndC
 		else
 		{
 			const FShaderCodeChunk& ResultChunk = CodeChunks[ResultIndex];
-			// Combine the definition lines and the return statement
+			// Combine the definition lines and the return statement.
+			// Also specify the return symbol name to terminate the iteration over all definitions earlier,
+			// if there are unnecessary statements that don't contribute to the outcome.
 			check(ResultChunk.bInline || ResultChunk.SymbolName.Len() > 0);
-			OutDefinitions = GetDefinitions(CodeChunks, StartChunk, EndChunk, Variation);
+			const TCHAR* ReturnValueSymbolName = (bReduceAfterReturnValue && !ResultChunk.bInline ? *ResultChunk.SymbolName : nullptr);
+			OutDefinitions = GetDefinitions(CodeChunks, StartChunk, EndChunk, Variation, ReturnValueSymbolName);
 			OutValue = ResultChunk.bInline ? ResultChunk.AtDefinition(Variation) : ResultChunk.SymbolName;
 			if (Variation == CompiledPDV_Analytic && ResultChunk.DerivativeStatus == EDerivativeStatus::Valid)
 			{
diff --git a/Engine/Source/Runtime/Engine/Private/Materials/HLSLMaterialTranslator.h b/Engine/Source/Runtime/Engine/Private/Materials/HLSLMaterialTranslator.h
--- a/Engine/Source/Runtime/Engine/Private/Materials/HLSLMaterialTranslator.h
+++ b/Engine/Source/Runtime/Engine/Private/Materials/HLSLMaterialTranslator.h
@@ -486,10 +486,10 @@ protected:
 	void SetParameterMaterialAttributes(int32 Index, uint64 Mask);
 
 	/** Creates a string of all definitions needed for the given material input. */
-	FString GetDefinitions(TArray<FShaderCodeChunk>& CodeChunks, int32 StartChunk, int32 EndChunk, ECompiledPartialDerivativeVariation Variation) const;
+	FString GetDefinitions(const TArray<FShaderCodeChunk>& CodeChunks, int32 StartChunk, int32 EndChunk, ECompiledPartialDerivativeVariation Variation, const TCHAR* ReturnValueSymbolName = nullptr) const;
 
 	// GetFixedParameterCode
-	void GetFixedParameterCode(int32 StartChunk, int32 EndChunk, int32 ResultIndex, TArray<FShaderCodeChunk>& CodeChunks, FString& OutDefinitions, FString& OutValue, ECompiledPartialDerivativeVariation Variation);
+	void GetFixedParameterCode(int32 StartChunk, int32 EndChunk, int32 ResultIndex, TArray<FShaderCodeChunk>& CodeChunks, FString& OutDefinitions, FString& OutValue, ECompiledPartialDerivativeVariation Variation, bool bReduceAfterReturnValue = false);
 
 	void LinkParentScopes(TArray<FShaderCodeChunk>& CodeChunks);
 
