Optimization, don't empty and grow this array multiple times

diff --git a/Engine/Source/Runtime/Engine/Private/Streaming/StreamingManagerTexture.cpp b/Engine/Source/Runtime/Engine/Private/Streaming/StreamingManagerTexture.cpp
--- a/Engine/Source/Runtime/Engine/Private/Streaming/StreamingManagerTexture.cpp
+++ b/Engine/Source/Runtime/Engine/Private/Streaming/StreamingManagerTexture.cpp
@@ -34,6 +34,11 @@ CSV_DECLARE_CATEGORY_MODULE_EXTERN(CORE_API, Basic);
 
 CSV_DEFINE_CATEGORY(TextureStreaming, true);
 
+#ifndef UE_STREAMINGRENDERASSETS_ARRAY_DEFAULT_RESERVED_SIZE
+// The default size will reserve ~4MB, the element size is ~208 bytes.
+#define UE_STREAMINGRENDERASSETS_ARRAY_DEFAULT_RESERVED_SIZE 20000
+#endif
+
 static TAutoConsoleVariable<int32> CVarStreamingOverlapAssetAndLevelTicks(
 	TEXT("r.Streaming.OverlapAssetAndLevelTicks"),
 	!WITH_EDITOR && (PLATFORM_PS4),
@@ -169,6 +174,8 @@ FRenderAssetStreamingManager::FRenderAssetStreamingManager()
 	{
 		MarkMountedStateDirty(MakeIoFilenameHash(FileName));
 	});
+
+	StreamingRenderAssets.Reserve(UE_STREAMINGRENDERASSETS_ARRAY_DEFAULT_RESERVED_SIZE);
 }
 
 FRenderAssetStreamingManager::~FRenderAssetStreamingManager()
@@ -442,7 +449,7 @@ void FRenderAssetStreamingManager::ProcessRemovedRenderAssets()
 		// This handles the case where the last element was also removed.
 		while (StreamingRenderAssets.IsValidIndex(AssetIndex) && !StreamingRenderAssets[AssetIndex].RenderAsset)
 		{
-			StreamingRenderAssets.RemoveAtSwap(AssetIndex);
+			StreamingRenderAssets.RemoveAtSwap(AssetIndex, 1, false);
 		}
 
 		if (StreamingRenderAssets.IsValidIndex(AssetIndex))
@@ -477,7 +484,7 @@ void FRenderAssetStreamingManager::ProcessAddedRenderAssets()
 			new (StreamingRenderAssets) FStreamingRenderAsset(Asset, NumStreamedMips, NumLODGroups, Settings);
 		}
 	}
-	PendingStreamingRenderAssets.Empty();
+	PendingStreamingRenderAssets.Reset();
 }
 
 void FRenderAssetStreamingManager::ConditionalUpdateStaticData()
