Fix WorldBounds not taking into account Container transform

diff --git a/Engine/Source/Runtime/Engine/Private/WorldPartition/WorldPartitionActorCluster.cpp b/Engine/Source/Runtime/Engine/Private/WorldPartition/WorldPartitionActorCluster.cpp
--- a/Engine/Source/Runtime/Engine/Private/WorldPartition/WorldPartitionActorCluster.cpp
+++ b/Engine/Source/Runtime/Engine/Private/WorldPartition/WorldPartitionActorCluster.cpp
@@ -337,13 +337,13 @@ void FActorClusterContext::CreateContainerInstanceRecursive(uint64 ID, const FTr
 			{
 				case EActorGridPlacement::Location:
 				{
-					FVector Location = ActorDescView.GetOrigin();
+					FVector Location = Transform.TransformPosition(ActorDescView.GetOrigin());
 					Bounds += FBox(Location, Location);
 				}
 				break;
 				case EActorGridPlacement::Bounds:
 				{
-					Bounds += ActorDescView.GetBounds();
+					Bounds += ActorDescView.GetBounds().TransformBy(Transform);
 				}
 				break;
 			}
