Improved performance when locking a mip in a read only mode from the source data of a texture.
This drastically improves the performance when doing the test to determine if the newly imported texture is a normal map or not.
(Before 0.288150 seconds, after 0.000161 seconds) These are from RGBA 16 bits 8k tif texture.

The improvement is made by:
1- avoiding copying the source mip.
2- avoiding a compression and decompression of the payload of a VirtualBulkData when getting the payload, if the payload is already loaded.

diff --git a/Engine/Source/Runtime/CoreUObject/Private/Virtualization/VirtualizedBulkData.cpp b/Engine/Source/Runtime/CoreUObject/Private/Virtualization/VirtualizedBulkData.cpp
--- a/Engine/Source/Runtime/CoreUObject/Private/Virtualization/VirtualizedBulkData.cpp
+++ b/Engine/Source/Runtime/CoreUObject/Private/Virtualization/VirtualizedBulkData.cpp
@@ -975,10 +975,18 @@ TFuture<FSharedBuffer> FVirtualizedUntypedBulkData::GetPayload() const
 {
 	TPromise<FSharedBuffer> Promise;
 
-	FSharedBuffer UncompressedPayload = GetDataInternal().Decompress();
+	if (Payload)
+	{
+		// Avoid a unnecessary compression and decompression if we already have the uncompressed payload
+		Promise.SetValue(Payload);
+	}
+	else
+	{
+		FSharedBuffer UncompressedPayload = GetDataInternal().Decompress();
 
-	// TODO: Not actually async yet!
-	Promise.SetValue(MoveTemp(UncompressedPayload));
+		// TODO: Not actually async yet!
+		Promise.SetValue(MoveTemp(UncompressedPayload));
+	}
 
 	return Promise.GetFuture();
 }
diff --git a/Engine/Source/Runtime/Engine/Private/Texture.cpp b/Engine/Source/Runtime/Engine/Private/Texture.cpp
--- a/Engine/Source/Runtime/Engine/Private/Texture.cpp
+++ b/Engine/Source/Runtime/Engine/Private/Texture.cpp
@@ -1368,8 +1368,17 @@ uint8* FTextureSource::LockMipInternal(int32 BlockIndex, int32 LayerIndex, int32
 			LockedMipData = Decompress(nullptr);
 		}
 
-		MipData = LockedMipData.GetDataReadWrite() + CalcMipOffset(BlockIndex, LayerIndex, MipIndex);
+		if (RequestedLockState == ELockState::ReadOnly)
+		{
+			MipData = const_cast<uint8*>(static_cast<const uint8*>(LockedMipData.GetDataReadOnly().GetData()));
+		}
+		else
+		{
+			MipData = LockedMipData.GetDataReadWrite();
+		}
 		
+		MipData += CalcMipOffset(BlockIndex, LayerIndex, MipIndex);
+
 		if (NumLockedMips == 0)
 		{
 			LockState = RequestedLockState;
@@ -1409,13 +1418,12 @@ void FTextureSource::UnlockMip(int32 BlockIndex, int32 LayerIndex, int32 MipInde
 
 			bPNGCompressed = false;
 			CompressionFormat = TSCF_None;
+
+			ForceGenerateGuid();
 		}
 
 		LockState = ELockState::None;
 		LockedMipData.Reset();
-
-		// TODO: Only call this if the LockState was ELockState::ReadWrite (this would need testing)
-		ForceGenerateGuid();
 	}
 }
 
