Check if AttachParent is valid when receiving AttachSocketName before updating the attachment.

diff --git a/Engine/Source/Runtime/Engine/Private/Components/SceneComponent.cpp b/Engine/Source/Runtime/Engine/Private/Components/SceneComponent.cpp
--- a/Engine/Source/Runtime/Engine/Private/Components/SceneComponent.cpp
+++ b/Engine/Source/Runtime/Engine/Private/Components/SceneComponent.cpp
@@ -3185,7 +3185,10 @@ void USceneComponent::OnRep_AttachParent()
 
 void USceneComponent::OnRep_AttachSocketName()
 {
-	bNetUpdateAttachment = true;
+	if (IsValid(AttachParent))
+	{
+		bNetUpdateAttachment = true;
+	}
 }
 
 void USceneComponent::OnRep_AttachChildren()
