Free (PC only) VT GPU staging buffer memory when we need less for current workload.
Can save > 100MB on projects where we start with a very high water mark.

diff --git a/Engine/Source/Runtime/Engine/Private/VT/VirtualTextureUploadCache.cpp b/Engine/Source/Runtime/Engine/Private/VT/VirtualTextureUploadCache.cpp
--- a/Engine/Source/Runtime/Engine/Private/VT/VirtualTextureUploadCache.cpp
+++ b/Engine/Source/Runtime/Engine/Private/VT/VirtualTextureUploadCache.cpp
@@ -117,7 +117,7 @@ void FVirtualTextureUploadCache::Finalize(FRDGBuilder& GraphBuilder)
 		// On some platforms the staging texture create/lock behavior will depend on whether we are running with RHI threading
 		const bool bIsCpuWritable = !IsRunningRHIInSeparateThread();
 
-		if (BatchCount > StagingTexture.BatchCapacity || bIsCpuWritable != StagingTexture.bIsCPUWritable)
+		if (BatchCount > StagingTexture.BatchCapacity || BatchCount * 2 <= StagingTexture.BatchCapacity || bIsCpuWritable != StagingTexture.bIsCPUWritable)
 		{
 			// Staging texture is vertical stacked in widths of multiples of 4
 			// Smaller widths mean smaller stride which is more efficient for copying
