Fix static analysis warning

diff --git a/Engine/Source/Runtime/Engine/Private/SkeletalRenderGPUSkin.cpp b/Engine/Source/Runtime/Engine/Private/SkeletalRenderGPUSkin.cpp
--- a/Engine/Source/Runtime/Engine/Private/SkeletalRenderGPUSkin.cpp
+++ b/Engine/Source/Runtime/Engine/Private/SkeletalRenderGPUSkin.cpp
@@ -322,7 +322,7 @@ void FSkeletalMeshObjectGPUSkin::Update(int32 LODIndex,USkinnedMeshComponent* In
 	// queue a call to update this data
 	FSkeletalMeshObjectGPUSkin* MeshObject = this;
 	ENQUEUE_RENDER_COMMAND(SkelMeshObjectUpdateDataCommand)(
-		[MeshObject, FrameNumberToPrepare, RevisionNumber, NewDynamicData, GPUSkinCache, bIsGameWorld = (InMeshComponent->GetWorld() != nullptr && InMeshComponent->GetWorld()->IsGameWorld())](FRHICommandListImmediate& RHICmdList)
+		[MeshObject, FrameNumberToPrepare, RevisionNumber, NewDynamicData, GPUSkinCache, bIsGameWorld = (InMeshComponent && InMeshComponent->GetWorld() != nullptr && InMeshComponent->GetWorld()->IsGameWorld())](FRHICommandListImmediate& RHICmdList)
 		{
 			FScopeCycleCounter Context(MeshObject->GetStatId());
 			MeshObject->UpdateDynamicData_RenderThread(GPUSkinCache, RHICmdList, NewDynamicData, nullptr, bIsGameWorld, FrameNumberToPrepare, RevisionNumber);
