WorldPartition: Move OnRegister/OnUnregister up into ActorDescContainer. FLevelInstanceActorDesc only registers it's container in OnRegister instead of Init which can be called while saving on orphan ActorDesc (fixes crash related to findobject while saving)

diff --git a/Engine/Source/Runtime/Engine/Private/WorldPartition/ActorDescContainer.cpp b/Engine/Source/Runtime/Engine/Private/WorldPartition/ActorDescContainer.cpp
--- a/Engine/Source/Runtime/Engine/Private/WorldPartition/ActorDescContainer.cpp
+++ b/Engine/Source/Runtime/Engine/Private/WorldPartition/ActorDescContainer.cpp
@@ -19,7 +19,7 @@ UActorDescContainer::UActorDescContainer(const FObjectInitializer& ObjectInitial
 #endif
 {}
 
-void UActorDescContainer::Initialize(UWorld* InWorld, FName InPackageName)
+void UActorDescContainer::Initialize(UWorld* InWorld, FName InPackageName, TFunctionRef<void(FWorldPartitionActorDesc*)> PreRegister)
 {
 	check(!World || World == InWorld);
 	World = InWorld;
@@ -109,6 +109,12 @@ void UActorDescContainer::Initialize(UWorld* InWorld, FName InPackageName)
 		}
 	}
 
+	for (FActorDescList::TIterator<> ActorDescIterator(this); ActorDescIterator; ++ActorDescIterator)
+	{
+		PreRegister(*ActorDescIterator);
+		ActorDescIterator->OnRegister();
+	}
+
 	RegisterEditorDelegates();
 
 	bContainerInitialized = true;
@@ -154,25 +160,7 @@ void UActorDescContainer::BeginDestroy()
 #if WITH_EDITOR
 bool UActorDescContainer::ShouldHandleActorEvent(const AActor* Actor)
 {
-	if (Actor && Actor->IsPackageExternal() && (Actor->GetLevel() == World->PersistentLevel) && Actor->IsMainPackageActor())
-	{
-		// Only handle assets that belongs to our level
-		auto RemoveAfterFirstDot = [](const FString& InValue)
-		{
-			int32 DotIndex;
-			if (InValue.FindChar(TEXT('.'), DotIndex))
-			{
-				return InValue.LeftChop(InValue.Len() - DotIndex);
-			}
-			return InValue;
-		};
-
-		const FString ThisLevelPath = ContainerPackageName.ToString();
-		const FString AssetLevelPath = RemoveAfterFirstDot(Actor->GetPathName());
-		return (ThisLevelPath == AssetLevelPath);
-	}
-
-	return false;
+	return Actor && Actor->IsMainPackageActor() && (Actor->GetLevel() != nullptr) && (Actor->GetLevel()->GetPackage()->GetFName() == ContainerPackageName);
 }
 
 void UActorDescContainer::OnObjectPreSave(UObject* Object, FObjectPreSaveContext SaveContext)
@@ -206,7 +194,6 @@ void UActorDescContainer::OnObjectPreSave(UObject* Object, FObjectPreSaveContext
 				{
 					FWorldPartitionActorDesc* const AddedActorDesc = AddActor(Actor);
 					OnActorDescAdded(AddedActorDesc);
-					OnActorDescAddedEvent.Broadcast(AddedActorDesc);
 				}
 			}
 		}
@@ -245,8 +232,6 @@ void UActorDescContainer::RemoveActor(const FGuid& ActorGuid)
 {
 	if (TUniquePtr<FWorldPartitionActorDesc>* ExistingActorDesc = GetActorDescriptor(ActorGuid))
 	{
-		OnActorDescRemovedEvent.Broadcast(ExistingActorDesc->Get());
-
 		OnActorDescRemoved(ExistingActorDesc->Get());
 		RemoveActorDescriptor(ExistingActorDesc->Get());
 		ExistingActorDesc->Reset();
@@ -328,11 +313,15 @@ void UActorDescContainer::UnregisterEditorDelegates()
 
 void UActorDescContainer::OnActorDescAdded(FWorldPartitionActorDesc* NewActorDesc)
 {
+	NewActorDesc->OnRegister();
+
 	OnActorDescAddedEvent.Broadcast(NewActorDesc);
 }
 
 void UActorDescContainer::OnActorDescRemoved(FWorldPartitionActorDesc* ActorDesc)
 {
 	OnActorDescRemovedEvent.Broadcast(ActorDesc);
+
+	ActorDesc->OnUnregister();
 }
 #endif
diff --git a/Engine/Source/Runtime/Engine/Private/WorldPartition/LevelInstance/LevelInstanceActorDesc.cpp b/Engine/Source/Runtime/Engine/Private/WorldPartition/LevelInstance/LevelInstanceActorDesc.cpp
--- a/Engine/Source/Runtime/Engine/Private/WorldPartition/LevelInstance/LevelInstanceActorDesc.cpp
+++ b/Engine/Source/Runtime/Engine/Private/WorldPartition/LevelInstance/LevelInstanceActorDesc.cpp
@@ -26,8 +26,6 @@ void FLevelInstanceActorDesc::Init(const AActor* InActor)
 	LevelPackage = *LevelInstanceActor->GetWorldAssetPackage();
 	LevelInstanceTransform = LevelInstanceActor->GetActorTransform();
 	DesiredRuntimeBehavior = LevelInstanceActor->GetDesiredRuntimeBehavior();
-
-	RegisterContainer();
 }
 
 void FLevelInstanceActorDesc::Init(UActorDescContainer* InContainer, const FWorldPartitionActorDescInitData& DescData)
@@ -36,10 +34,20 @@ void FLevelInstanceActorDesc::Init(UActorDescContainer* InContainer, const FWorl
 	DesiredRuntimeBehavior = CDO->GetDefaultRuntimeBehavior();
 
 	FWorldPartitionActorDesc::Init(InContainer, DescData);
-		
+}
+
+void FLevelInstanceActorDesc::OnRegister()
+{
+	FWorldPartitionActorDesc::OnRegister();
 	RegisterContainer();
 }
 
+void FLevelInstanceActorDesc::OnUnregister()
+{
+	FWorldPartitionActorDesc::OnUnregister();
+	LevelInstanceContainer = nullptr;
+}
+
 void FLevelInstanceActorDesc::RegisterContainer()
 {
 	check(!LevelInstanceContainer);
diff --git a/Engine/Source/Runtime/Engine/Private/WorldPartition/WorldPartition.cpp b/Engine/Source/Runtime/Engine/Private/WorldPartition/WorldPartition.cpp
--- a/Engine/Source/Runtime/Engine/Private/WorldPartition/WorldPartition.cpp
+++ b/Engine/Source/Runtime/Engine/Private/WorldPartition/WorldPartition.cpp
@@ -399,11 +399,6 @@ void UWorldPartition::Initialize(UWorld* InWorld, const FTransform& InTransform)
 	{
 		UPackage* LevelPackage = OuterWorld->PersistentLevel->GetOutermost();
 		const FName PackageName = LevelPackage->GetLoadedPath().GetPackageFName();
-		{
-			TRACE_CPUPROFILER_EVENT_SCOPE(UActorDescContainer::Initialize);
-			UActorDescContainer::Initialize(World, PackageName);
-		}
-		check(bContainerInitialized);
 
 		bool bIsInstanced = (bEditorOnly && !IsRunningCommandlet()) ? OuterWorld->PersistentLevel->IsInstancedLevel() : bPIEWorldTravel;
 
@@ -418,33 +413,33 @@ void UWorldPartition::Initialize(UWorld* InWorld, const FTransform& InTransform)
 			const FString DestWorldName = FPaths::GetBaseFilename(LevelPackage->GetFName().ToString());
 
 			ReplaceFrom = SourceWorldName + TEXT(".") + SourceWorldName;
-			
+
 			// In PIE UWorld::PostLoad will not rename the world to match its package name so use SourceWorldName instead
 			ReplaceTo = DestWorldName + TEXT(".") + (bPIEWorldTravel ? SourceWorldName : DestWorldName);
 		}
 
 		{
-			TRACE_CPUPROFILER_EVENT_SCOPE(HashActorDescs);
-			for (FActorDescList::TIterator<> ActorDescIterator(this); ActorDescIterator; ++ActorDescIterator)
+			TRACE_CPUPROFILER_EVENT_SCOPE(UActorDescContainer::Initialize);
+			UActorDescContainer::Initialize(World, PackageName, [&](FWorldPartitionActorDesc* ActorDesc)
 			{
+				// This will get called by UActorDescContainer after having initialized all ActorDescs and before calling OnRegister on the ActorDesc
 				if (bIsInstanced)
 				{
-					const FString LongActorPackageName = ActorDescIterator->ActorPackage.ToString();
+					const FString LongActorPackageName = ActorDesc->ActorPackage.ToString();
 					const FString ActorPackageName = FPaths::GetBaseFilename(LongActorPackageName);
 					const FString InstancedName = FString::Printf(TEXT("%s_InstanceOf_%s"), *LevelPackage->GetName(), *ActorPackageName);
 
 					InstancingContext.AddMapping(*LongActorPackageName, *InstancedName);
 
-					ActorDescIterator->TransformInstance(ReplaceFrom, ReplaceTo, InstanceTransform);
+					ActorDesc->TransformInstance(ReplaceFrom, ReplaceTo, InstanceTransform);
 				}
 
-				ActorDescIterator->OnRegister();
-
 				if (bEditorOnly)
 				{
-					HashActorDesc(*ActorDescIterator);
+					HashActorDesc(ActorDesc);
 				}
-			}
+			});
+			check(bContainerInitialized);
 		}
 	}
 
@@ -1016,10 +1011,10 @@ void UWorldPartition::UpdateLoadingEditorCell(UWorldPartitionEditorCell* Cell, b
 
 void UWorldPartition::OnActorDescAdded(FWorldPartitionActorDesc* NewActorDesc)
 {
-	HashActorDesc(NewActorDesc);
-
-	NewActorDesc->OnRegister();
+	Super::OnActorDescAdded(NewActorDesc);
 
+	HashActorDesc(NewActorDesc);
+		
 	if (WorldPartitionEditor)
 	{
 		WorldPartitionEditor->Refresh();
@@ -1028,8 +1023,9 @@ void UWorldPartition::OnActorDescAdded(FWorldPartitionActorDesc* NewActorDesc)
 
 void UWorldPartition::OnActorDescRemoved(FWorldPartitionActorDesc* ActorDesc)
 {
+	Super::OnActorDescRemoved(ActorDesc);
+	
 	UnhashActorDesc(ActorDesc);
-	ActorDesc->OnUnregister();
 		
 	if (WorldPartitionEditor)
 	{
diff --git a/Engine/Source/Runtime/Engine/Public/WorldPartition/ActorDescContainer.h b/Engine/Source/Runtime/Engine/Public/WorldPartition/ActorDescContainer.h
--- a/Engine/Source/Runtime/Engine/Public/WorldPartition/ActorDescContainer.h
+++ b/Engine/Source/Runtime/Engine/Public/WorldPartition/ActorDescContainer.h
@@ -18,7 +18,7 @@ class ENGINE_API UActorDescContainer : public UObject, public FActorDescList
 	friend class FWorldPartitionActorDesc;
 
 public:
-	void Initialize(UWorld* World, FName InPackageName);
+	void Initialize(UWorld* World, FName InPackageName, TFunctionRef<void(FWorldPartitionActorDesc*)> PreRegister = [](FWorldPartitionActorDesc*) {});
 	virtual void Uninitialize();
 	
 	virtual UWorld* GetWorld() const override;
diff --git a/Engine/Source/Runtime/Engine/Public/WorldPartition/LevelInstance/LevelInstanceActorDesc.h b/Engine/Source/Runtime/Engine/Public/WorldPartition/LevelInstance/LevelInstanceActorDesc.h
--- a/Engine/Source/Runtime/Engine/Public/WorldPartition/LevelInstance/LevelInstanceActorDesc.h
+++ b/Engine/Source/Runtime/Engine/Public/WorldPartition/LevelInstance/LevelInstanceActorDesc.h
@@ -31,6 +31,8 @@ protected:
 	virtual void Serialize(FArchive& Ar) override;
 	virtual void AddReferencedObjects(FReferenceCollector& Collector);
 	virtual void RegisterContainer();
+	virtual void OnRegister();
+	virtual void OnUnregister();
 
 	FName LevelPackage;
 	FTransform LevelInstanceTransform;
