PR #8269: Mark commonly-used audio types for export for C++ access/usage. (Contributed by [REDACTED_PERSON])

diff --git a/Engine/Source/Runtime/Engine/Classes/Sound/AudioVolume.h b/Engine/Source/Runtime/Engine/Classes/Sound/AudioVolume.h
--- a/Engine/Source/Runtime/Engine/Classes/Sound/AudioVolume.h
+++ b/Engine/Source/Runtime/Engine/Classes/Sound/AudioVolume.h
@@ -35,7 +35,7 @@ enum class EAudioVolumeLocationState : uint8
 
 /** Struct to determine dynamic submix send data for use with audio volumes. */
 USTRUCT(BlueprintType)
-struct FAudioVolumeSubmixSendSettings
+struct ENGINE_API FAudioVolumeSubmixSendSettings
 {
 	GENERATED_USTRUCT_BODY()
 
@@ -52,7 +52,7 @@ struct FAudioVolumeSubmixSendSettings
 };
 
 USTRUCT(BlueprintType)
-struct FAudioVolumeSubmixOverrideSettings
+struct ENGINE_API FAudioVolumeSubmixOverrideSettings
 {
 	GENERATED_USTRUCT_BODY()
 
diff --git a/Engine/Source/Runtime/Engine/Classes/Sound/ReverbSettings.h b/Engine/Source/Runtime/Engine/Classes/Sound/ReverbSettings.h
--- a/Engine/Source/Runtime/Engine/Classes/Sound/ReverbSettings.h
+++ b/Engine/Source/Runtime/Engine/Classes/Sound/ReverbSettings.h
@@ -45,7 +45,7 @@ enum ReverbPreset
 
 /** Struct encapsulating settings for reverb effects. */
 USTRUCT(BlueprintType)
-struct FReverbSettings
+struct ENGINE_API FReverbSettings
 {
 	GENERATED_USTRUCT_BODY()
 
