Cleaned up constraint instance properties copy across the physics editor & blueprint libraries
Added profile update when using constraint instance accessor from a physics asset
Added some blueprint functions to clean up usage & harmonize physics asset with skeletal mesh component
Changing properties in the accessor now dirties the asset

diff --git a/Engine/Source/Editor/PhysicsAssetEditor/Private/PhysicsAssetEditorSharedData.cpp b/Engine/Source/Editor/PhysicsAssetEditor/Private/PhysicsAssetEditorSharedData.cpp
--- a/Engine/Source/Editor/PhysicsAssetEditor/Private/PhysicsAssetEditorSharedData.cpp
+++ b/Engine/Source/Editor/PhysicsAssetEditor/Private/PhysicsAssetEditorSharedData.cpp
@@ -240,29 +240,7 @@ void FPhysicsAssetEditorSharedData::CopyConstraintProperties(const UPhysicsConst
 {
 	ToConstraintSetup->Modify();
 	FConstraintInstance OldInstance = ToConstraintSetup->DefaultInstance;
-	ToConstraintSetup->DefaultInstance.CopyConstraintParamsFrom(&FromConstraintSetup->DefaultInstance);
-
-	// recover certain data that we'd like to keep - i.e. bone indices those still should stay.  
-	// frame position offsets taken from old, but, optionally, frame orientations are taken from new source
-	ToConstraintSetup->DefaultInstance.ConstraintIndex = OldInstance.ConstraintIndex;
-#if WITH_PHYSX
-	ToConstraintSetup->DefaultInstance.ConstraintHandle = OldInstance.ConstraintHandle;
-#endif	//WITH_PHYSX
-	ToConstraintSetup->DefaultInstance.JointName = OldInstance.JointName;
-	ToConstraintSetup->DefaultInstance.ConstraintBone1 = OldInstance.ConstraintBone1;
-	ToConstraintSetup->DefaultInstance.ConstraintBone2 = OldInstance.ConstraintBone2;
-	ToConstraintSetup->DefaultInstance.Pos1 = OldInstance.Pos1;
-	ToConstraintSetup->DefaultInstance.Pos2 = OldInstance.Pos2;
-
-	if (bKeepOldRotation)
-	{
-		ToConstraintSetup->DefaultInstance.PriAxis1 = OldInstance.PriAxis1;
-		ToConstraintSetup->DefaultInstance.SecAxis1 = OldInstance.SecAxis1;
-		ToConstraintSetup->DefaultInstance.PriAxis2 = OldInstance.PriAxis2;
-		ToConstraintSetup->DefaultInstance.SecAxis2 = OldInstance.SecAxis2;
-		ToConstraintSetup->DefaultInstance.AngularRotationOffset = OldInstance.AngularRotationOffset;
-	}
-
+	ToConstraintSetup->DefaultInstance.CopyConstraintPhysicalPropertiesFrom(&FromConstraintSetup->DefaultInstance, /*bKeepPosition=*/true, bKeepOldRotation);
 	ToConstraintSetup->UpdateProfileInstance();
 }
 
diff --git a/Engine/Source/Runtime/Engine/Classes/PhysicsEngine/ConstraintInstance.h b/Engine/Source/Runtime/Engine/Classes/PhysicsEngine/ConstraintInstance.h
--- a/Engine/Source/Runtime/Engine/Classes/PhysicsEngine/ConstraintInstance.h
+++ b/Engine/Source/Runtime/Engine/Classes/PhysicsEngine/ConstraintInstance.h
@@ -631,6 +631,9 @@ public:
 	// @todo document
 	void CopyConstraintParamsFrom(const FConstraintInstance* FromInstance);
 
+	/** Copies non-identifying properties from another constraint instance */
+	void CopyConstraintPhysicalPropertiesFrom(const FConstraintInstance* FromInstance, bool bKeepPosition, bool bKeepRotation);
+
 	// Retrieve the constraint force most recently applied to maintain this constraint. Returns 0 forces if the constraint is not initialized or broken.
 	void GetConstraintForce(FVector& OutLinearForce, FVector& OutAngularForce);
 
@@ -1104,18 +1107,38 @@ public:
 		, Index(0)
 	{}
 
-	FConstraintInstanceAccessor(const TWeakObjectPtr<UObject>& Owner, uint32 Index = 0)
+	FConstraintInstanceAccessor(const TWeakObjectPtr<UObject>& Owner, uint32 Index = 0, TFunction<void(void)> InOnRelease = TFunction<void(void)>())
 		: Owner(Owner)
 		, Index(Index)
+#if WITH_EDITOR
+		, OnRelease(InOnRelease)
+#endif
 	{}
 
+#if WITH_EDITOR
+	~FConstraintInstanceAccessor()
+	{
+		if (OnRelease)
+		{
+			OnRelease();
+		}
+	}
+#endif
+
 	FConstraintInstance* Get() const;
 
+	/** Calls modify on the owner object to make sure the constraint is dirtied */
+	void Modify();
+
 private:
 	UPROPERTY()
 	TWeakObjectPtr<UObject> Owner;
 
 	UPROPERTY()
 	uint32 Index;
-	
+
+#if WITH_EDITOR
+	/** Warning: it would be unwieldy to make the accessor move only, so the OnRelease callback should be safe to be called multiple times */
+	TFunction<void(void)> OnRelease;
+#endif
 };
\ No newline at end of file
diff --git a/Engine/Source/Runtime/Engine/Classes/PhysicsEngine/ConstraintInstanceBlueprintLibrary.h b/Engine/Source/Runtime/Engine/Classes/PhysicsEngine/ConstraintInstanceBlueprintLibrary.h
--- a/Engine/Source/Runtime/Engine/Classes/PhysicsEngine/ConstraintInstanceBlueprintLibrary.h
+++ b/Engine/Source/Runtime/Engine/Classes/PhysicsEngine/ConstraintInstanceBlueprintLibrary.h
@@ -688,11 +688,11 @@ public:
 	);
 
 	/** Gets the drive params for the angular drive.
-*	@param Accessor				Constraint accessor to query
-*	@param OutPositionStrength	Positional strength for the drive (stiffness)
-*	@param OutVelocityStrength 	Velocity strength of the drive (damping)
-*	@param OutForceLimit		Max force applied by the drive
-*/
+	*	@param Accessor				Constraint accessor to query
+	*	@param OutPositionStrength	Positional strength for the drive (stiffness)
+	*	@param OutVelocityStrength 	Velocity strength of the drive (damping)
+	*	@param OutForceLimit		Max force applied by the drive
+	*/
 	UFUNCTION(BlueprintCallable, Category = "Physics|Constraints", meta = (Keywords = "motor target orientation velocity strength max force"))
 	static void GetAngularDriveParams(
 		UPARAM(ref) FConstraintInstanceAccessor& Accessor,
@@ -700,4 +700,18 @@ public:
 		float OutVelocityStrength,
 		float OutForceLimit
 	);
+
+	/** Copies all properties from one constraint to another
+	* @param Accessor	Constraint accessor to write to
+	* @param SourceAccessor Constraint accessor to read from
+	* @param bKeepPosition	Whether to keep original constraint positions
+	* @param bKeepRotation	Whether to keep original constraint rotations
+	*/
+	UFUNCTION(BlueprintCallable, Category = "Physics|Constraints")
+	static void CopyParams(
+		UPARAM(ref) FConstraintInstanceAccessor& Accessor,
+		UPARAM(ref) FConstraintInstanceAccessor& SourceAccessor,
+		bool bKeepPosition = true,
+		bool bKeepRotation = true
+	);
 };
\ No newline at end of file
diff --git a/Engine/Source/Runtime/Engine/Classes/PhysicsEngine/PhysicsAsset.h b/Engine/Source/Runtime/Engine/Classes/PhysicsEngine/PhysicsAsset.h
--- a/Engine/Source/Runtime/Engine/Classes/PhysicsEngine/PhysicsAsset.h
+++ b/Engine/Source/Runtime/Engine/Classes/PhysicsEngine/PhysicsAsset.h
@@ -282,12 +282,28 @@ public:
 	virtual USkeletalMesh* GetPreviewMesh() const;
 
 #if WITH_EDITOR
+	/** Gets a constraint by its joint name
+	* @param ConstraintName name of the constraint
+	* @return ConstraintInstance accessor to the constraint data
+	*/
 	UFUNCTION(BlueprintCallable, Category = "Physics|Constraints")
 	FConstraintInstanceAccessor GetConstraintByName(FName ConstraintName);
 
+	/** Gets a constraint by its joint name
+	* @param Bone1Name name of the first bone in the joint
+	* @param Bone2Name name of the second bone in the joint
+	* @return ConstraintInstance accessor to the constraint data
+	*/
 	UFUNCTION(BlueprintCallable, Category = "Physics|Constraints")
 	FConstraintInstanceAccessor GetConstraintByBoneNames(FName Bone1Name, FName Bone2Name);
 
+	/** Gets all constraints
+	* @param IncludesTerminated whether or not to return terminated constraints
+	* @param OutConstraints returned list of constraints matching the parameters
+	*/
+	UFUNCTION(BlueprintCallable, Category = "Physics|Constraints")
+	void GetConstraints(bool bIncludesTerminated, TArray<FConstraintInstanceAccessor>& OutConstraints);
+
 	FConstraintInstanceAccessor GetConstraintInstanceAccessorByIndex(int32 Index);
 	FConstraintInstance* GetConstraintInstanceByIndex(uint32 Index);
 #endif
diff --git a/Engine/Source/Runtime/Engine/Private/PhysicsEngine/ConstraintInstance.cpp b/Engine/Source/Runtime/Engine/Private/PhysicsEngine/ConstraintInstance.cpp
--- a/Engine/Source/Runtime/Engine/Private/PhysicsEngine/ConstraintInstance.cpp
+++ b/Engine/Source/Runtime/Engine/Private/PhysicsEngine/ConstraintInstance.cpp
@@ -600,6 +600,38 @@ void FConstraintInstance::CopyConstraintParamsFrom(const FConstraintInstance* Fr
 	*this = *FromInstance;
 }
 
+void FConstraintInstance::CopyConstraintPhysicalPropertiesFrom(const FConstraintInstance* FromInstance, bool bKeepPosition, bool bKeepRotation)
+{
+	check(FromInstance);
+
+	FConstraintInstance OldInstance = *this;
+	CopyConstraintParamsFrom(FromInstance);
+
+	// Recover internal data we'd like to keep - i.e. bone indices, etc.
+	ConstraintIndex = OldInstance.ConstraintIndex;
+#if WITH_PHYSX
+	ConstraintHandle = OldInstance.ConstraintHandle;
+#endif	//WITH_PHYSX
+	JointName = OldInstance.JointName;
+	ConstraintBone1 = OldInstance.ConstraintBone1;
+	ConstraintBone2 = OldInstance.ConstraintBone2;
+
+	if (bKeepPosition)
+	{
+		Pos1 = OldInstance.Pos1;
+		Pos2 = OldInstance.Pos2;
+	}
+
+	if (bKeepRotation)
+	{
+		PriAxis1 = OldInstance.PriAxis1;
+		SecAxis1 = OldInstance.SecAxis1;
+		PriAxis2 = OldInstance.PriAxis2;
+		SecAxis2 = OldInstance.SecAxis2;
+		AngularRotationOffset = OldInstance.AngularRotationOffset;
+	}
+}
+
 FTransform FConstraintInstance::GetRefFrame(EConstraintFrame::Type Frame) const
 {
 	FTransform Result;
@@ -1213,4 +1245,12 @@ FConstraintInstance* FConstraintInstanceAccessor::Get() const
 	return nullptr;
 }
 
+void FConstraintInstanceAccessor::Modify()
+{
+	if (Owner.IsValid())
+	{
+		Owner->Modify();
+	}
+}
+
 #undef LOCTEXT_NAMESPACE
diff --git a/Engine/Source/Runtime/Engine/Private/PhysicsEngine/ConstraintInstanceBlueprintLibrary.cpp b/Engine/Source/Runtime/Engine/Private/PhysicsEngine/ConstraintInstanceBlueprintLibrary.cpp
--- a/Engine/Source/Runtime/Engine/Private/PhysicsEngine/ConstraintInstanceBlueprintLibrary.cpp
+++ b/Engine/Source/Runtime/Engine/Private/PhysicsEngine/ConstraintInstanceBlueprintLibrary.cpp
@@ -32,6 +32,7 @@ void UConstraintInstanceBlueprintLibrary::SetDisableCollision(
 {
 	if (FConstraintInstance* ConstraintInstance = Accessor.Get())
 	{
+		Accessor.Modify();
 		ConstraintInstance->SetDisableCollision(bDisableCollision);
 	}
 }
@@ -54,6 +55,7 @@ void UConstraintInstanceBlueprintLibrary::SetProjectionParams(
 {
 	if (FConstraintInstance* ConstraintInstance = Accessor.Get())
 	{
+		Accessor.Modify();
 		ConstraintInstance->SetProjectionParams(bEnableProjection, ProjectionLinearAlpha, ProjectionAngularAlpha);
 	}
 }
@@ -85,6 +87,8 @@ void UConstraintInstanceBlueprintLibrary::SetParentDominates(
 {
 	if (FConstraintInstance* ConstraintInstance = Accessor.Get())
 	{
+		Accessor.Modify();
+
 		if (bParentDominates)
 		{
 			ConstraintInstance->EnableParentDominates();
@@ -121,6 +125,7 @@ void UConstraintInstanceBlueprintLibrary::SetLinearLimits(
 {
 	if (FConstraintInstance* ConstraintInstance = Accessor.Get())
 	{
+		Accessor.Modify();
 		ConstraintInstance->SetLinearLimits(XMotion, YMotion, ZMotion, Limit);
 	}
 }
@@ -160,6 +165,7 @@ void UConstraintInstanceBlueprintLibrary::SetLinearSoftLimitParams(
 {
 	if (FConstraintInstance* ConstraintInstance = Accessor.Get())
 	{
+		Accessor.Modify();
 		ConstraintInstance->SetSoftLinearLimitParams(bSoftLinearLimit, LinearLimitStiffness, LinearLimitDamping, LinearLimitRestitution, LinearLimitContactDistance);
 	}
 }
@@ -200,6 +206,7 @@ void UConstraintInstanceBlueprintLibrary::SetLinearBreakable(
 {
 	if (FConstraintInstance* ConstraintInstance = Accessor.Get())
 	{
+		Accessor.Modify();
 		ConstraintInstance->SetLinearBreakable(bLinearBreakable, LinearBreakThreshold);
 	}
 }
@@ -231,6 +238,7 @@ void UConstraintInstanceBlueprintLibrary::SetLinearPlasticity(
 {
 	if (FConstraintInstance* ConstraintInstance = Accessor.Get())
 	{
+		Accessor.Modify();
 		ConstraintInstance->SetLinearPlasticity(bLinearPlasticity, LinearPlasticityThreshold,PlasticityType);
 	}
 }
@@ -274,6 +282,7 @@ void UConstraintInstanceBlueprintLibrary::SetAngularLimits(
 {
 	if (FConstraintInstance* ConstraintInstance = Accessor.Get())
 	{
+		Accessor.Modify();
 		ConstraintInstance->SetAngularSwing1Limit(Swing1MotionType, Swing1LimitAngle);
 		ConstraintInstance->SetAngularSwing2Limit(Swing2MotionType, Swing2LimitAngle);
 		ConstraintInstance->SetAngularTwistLimit(TwistMotionType, TwistLimitAngle);
@@ -321,6 +330,7 @@ void UConstraintInstanceBlueprintLibrary::SetAngularSoftSwingLimitParams(
 {
 	if (FConstraintInstance* ConstraintInstance = Accessor.Get())
 	{
+		Accessor.Modify();
 		ConstraintInstance->SetSoftSwingLimitParams(bSoftSwingLimit, SwingLimitStiffness, SwingLimitDamping, SwingLimitRestitution, SwingLimitContactDistance);
 	}
 }
@@ -364,6 +374,7 @@ void UConstraintInstanceBlueprintLibrary::SetAngularSoftTwistLimitParams(
 {
 	if (FConstraintInstance* ConstraintInstance = Accessor.Get())
 	{
+		Accessor.Modify();
 		ConstraintInstance->SetSoftTwistLimitParams(bSoftTwistLimit, TwistLimitStiffness, TwistLimitDamping, TwistLimitRestitution, TwistLimitContactDistance);
 	}
 }
@@ -404,6 +415,7 @@ void UConstraintInstanceBlueprintLibrary::SetAngularBreakable(
 {
 	if (FConstraintInstance* ConstraintInstance = Accessor.Get())
 	{
+		Accessor.Modify();
 		ConstraintInstance->SetAngularBreakable(bAngularBreakable, AngularBreakThreshold);
 	}
 }
@@ -435,6 +447,7 @@ void UConstraintInstanceBlueprintLibrary::SetAngularPlasticity(
 {
 	if (FConstraintInstance* ConstraintInstance = Accessor.Get())
 	{
+		Accessor.Modify();
 		ConstraintInstance->SetAngularPlasticity(bAngularPlasticity, AngularPlasticityThreshold);
 	}
 }
@@ -473,6 +486,7 @@ void UConstraintInstanceBlueprintLibrary::SetLinearPositionDrive(
 {
 	if (FConstraintInstance* ConstraintInstance = Accessor.Get())
 	{
+		Accessor.Modify();
 		ConstraintInstance->SetLinearPositionDrive(bEnableDriveX, bEnableDriveY, bEnableDriveZ);
 	}
 }
@@ -507,6 +521,7 @@ void UConstraintInstanceBlueprintLibrary::SetLinearVelocityDrive(
 {
 	if (FConstraintInstance* ConstraintInstance = Accessor.Get())
 	{
+		Accessor.Modify();
 		ConstraintInstance->SetLinearVelocityDrive(bEnableDriveX, bEnableDriveY, bEnableDriveZ);
 	}
 }
@@ -539,6 +554,7 @@ void UConstraintInstanceBlueprintLibrary::SetLinearPositionTarget(
 {
 	if (FConstraintInstance* ConstraintInstance = Accessor.Get())
 	{
+		Accessor.Modify();
 		ConstraintInstance->SetLinearPositionTarget(InPosTarget);
 	}
 }
@@ -566,6 +582,7 @@ void UConstraintInstanceBlueprintLibrary::SetLinearVelocityTarget(
 {
 	if (FConstraintInstance* ConstraintInstance = Accessor.Get())
 	{
+		Accessor.Modify();
 		ConstraintInstance->SetLinearVelocityTarget(InVelTarget);
 	}
 }
@@ -594,6 +611,7 @@ void UConstraintInstanceBlueprintLibrary::SetLinearDriveParams(
 {
 	if (FConstraintInstance* ConstraintInstance = Accessor.Get())
 	{
+		Accessor.Modify();
 		ConstraintInstance->SetLinearDriveParams(PositionStrength, VelocityStrength, InForceLimit);
 	}
 }
@@ -632,6 +650,7 @@ void UConstraintInstanceBlueprintLibrary::SetOrientationDriveTwistAndSwing(
 {
 	if (FConstraintInstance* ConstraintInstance = Accessor.Get())
 	{
+		Accessor.Modify();
 		ConstraintInstance->SetOrientationDriveTwistAndSwing(bEnableTwistDrive, bEnableSwingDrive);
 	}
 }
@@ -660,6 +679,7 @@ void UConstraintInstanceBlueprintLibrary::SetOrientationDriveSLERP(
 {
 	if (FConstraintInstance* ConstraintInstance = Accessor.Get())
 	{
+		Accessor.Modify();
 		ConstraintInstance->SetOrientationDriveSLERP(bEnableSLERP);
 	}
 }
@@ -687,6 +707,7 @@ void UConstraintInstanceBlueprintLibrary::SetAngularVelocityDriveTwistAndSwing(
 {
 	if (FConstraintInstance* ConstraintInstance = Accessor.Get())
 	{
+		Accessor.Modify();
 		ConstraintInstance->SetAngularVelocityDriveTwistAndSwing(bEnableTwistDrive, bEnableSwingDrive);
 	}
 }
@@ -715,6 +736,7 @@ void UConstraintInstanceBlueprintLibrary::SetAngularVelocityDriveSLERP(
 {
 	if (FConstraintInstance* ConstraintInstance = Accessor.Get())
 	{
+		Accessor.Modify();
 		ConstraintInstance->SetAngularVelocityDriveSLERP(bEnableSLERP);
 	}
 }
@@ -741,6 +763,7 @@ void UConstraintInstanceBlueprintLibrary::SetAngularDriveMode(
 {
 	if (FConstraintInstance* ConstraintInstance = Accessor.Get())
 	{
+		Accessor.Modify();
 		ConstraintInstance->SetAngularDriveMode(DriveMode);
 	}
 }
@@ -768,6 +791,7 @@ void UConstraintInstanceBlueprintLibrary::SetAngularOrientationTarget(
 {
 	if (FConstraintInstance* ConstraintInstance = Accessor.Get())
 	{
+		Accessor.Modify();
 		ConstraintInstance->SetAngularOrientationTarget(InPosTarget.Quaternion());
 	}
 }
@@ -794,6 +818,7 @@ void UConstraintInstanceBlueprintLibrary::SetAngularVelocityTarget(
 {
 	if (FConstraintInstance* ConstraintInstance = Accessor.Get())
 	{
+		Accessor.Modify();
 		ConstraintInstance->SetAngularVelocityTarget(InVelTarget);
 	}
 }
@@ -822,6 +847,7 @@ void UConstraintInstanceBlueprintLibrary::SetAngularDriveParams(
 {
 	if (FConstraintInstance* ConstraintInstance = Accessor.Get())
 	{
+		Accessor.Modify();
 		ConstraintInstance->SetAngularDriveParams(PositionStrength, VelocityStrength, InForceLimit);
 	}
 }
@@ -843,4 +869,21 @@ void UConstraintInstanceBlueprintLibrary::GetAngularDriveParams(
 		OutVelocityStrength = 0.0f;
 		OutForceLimit = 0.0f;
 	}
+}
+
+void UConstraintInstanceBlueprintLibrary::CopyParams(
+	FConstraintInstanceAccessor& Accessor,
+	FConstraintInstanceAccessor& SourceAccessor,
+	bool bKeepPosition,
+	bool bKeepRotation
+)
+{
+	FConstraintInstance* Target = Accessor.Get();
+	FConstraintInstance* Source = SourceAccessor.Get();
+
+	if(Target && Source)
+	{
+		Accessor.Modify();
+		Target->CopyConstraintPhysicalPropertiesFrom(Source, bKeepPosition, bKeepRotation);
+	}
 }
\ No newline at end of file
diff --git a/Engine/Source/Runtime/Engine/Private/PhysicsEngine/PhysicsAsset.cpp b/Engine/Source/Runtime/Engine/Private/PhysicsEngine/PhysicsAsset.cpp
--- a/Engine/Source/Runtime/Engine/Private/PhysicsEngine/PhysicsAsset.cpp
+++ b/Engine/Source/Runtime/Engine/Private/PhysicsEngine/PhysicsAsset.cpp
@@ -867,7 +867,9 @@ FConstraintInstanceAccessor UPhysicsAsset::GetConstraintInstanceAccessorByIndex(
 
 	if (ConstraintSetup[ConstraintIndex])
 	{
-		return FConstraintInstanceAccessor(this, ConstraintIndex);
+		// Implementation note: Any changes on the constraint must be propagated from the instance to its profile
+		return FConstraintInstanceAccessor(this, ConstraintIndex, [this, ConstraintIndex]() {
+			ConstraintSetup[ConstraintIndex]->UpdateProfileInstance(); });
 	}
 
 	return FConstraintInstanceAccessor();
@@ -883,6 +885,23 @@ FConstraintInstanceAccessor UPhysicsAsset::GetConstraintByBoneNames(FName Bone1N
 	return GetConstraintInstanceAccessorByIndex(FindConstraintIndex(Bone1Name, Bone2Name));
 }
 
+void UPhysicsAsset::GetConstraints(bool bIncludesTerminated, TArray<FConstraintInstanceAccessor>& OutConstraints)
+{
+	for (int32 i = 0; i < ConstraintSetup.Num(); ++i)
+	{
+		if (ConstraintSetup[i])
+		{
+			if (bIncludesTerminated || ConstraintSetup[i]->DefaultInstance.IsTerminated())
+			{
+				// Implementation note: Any changes on the constraint must be propagated from the instance to its profile
+				OutConstraints.Emplace(this, i, [this, i]() {
+					ConstraintSetup[i]->UpdateProfileInstance();
+					});
+			}
+		}
+	}
+}
+
 FConstraintInstance* UPhysicsAsset::GetConstraintInstanceByIndex(uint32 Index)
 {
 	if (Index < (uint32)ConstraintSetup.Num())
