Fix deprecation warning in FAudioDevice::Exec. Even if InCmd were *InCmd, the text in question has been removed from the string

diff --git a/Engine/Source/Runtime/Engine/Private/AudioDevice.cpp b/Engine/Source/Runtime/Engine/Private/AudioDevice.cpp
--- a/Engine/Source/Runtime/Engine/Private/AudioDevice.cpp
+++ b/Engine/Source/Runtime/Engine/Private/AudioDevice.cpp
@@ -1978,7 +1978,7 @@ bool FAudioDevice::Exec(UWorld* InWorld, const TCHAR* Cmd, FOutputDevice& Ar)
 	{
 		if (FParse::Command(InCmd, InMatch))
 		{
-			UE_LOG(LogAudio, Warning, TEXT("The Exec command '%s' is deprecated. Use 'au.Debug.%s' instead"), InCmd, InMatch);
+			UE_LOG(LogAudio, Warning, TEXT("The Exec command '%s' is deprecated. Use 'au.Debug.%s' instead"), InMatch, InMatch);
 			return true;
 		}
 
