Make actor labels available in all in development builds.

diff --git a/Engine/Source/Runtime/Core/Public/UObject/UE5PrivateFrostyStreamObjectVersion.h b/Engine/Source/Runtime/Core/Public/UObject/UE5PrivateFrostyStreamObjectVersion.h
--- a/Engine/Source/Runtime/Core/Public/UObject/UE5PrivateFrostyStreamObjectVersion.h
+++ b/Engine/Source/Runtime/Core/Public/UObject/UE5PrivateFrostyStreamObjectVersion.h
@@ -21,6 +21,9 @@ struct CORE_API FUE5PrivateFrostyStreamObjectVersion
 		// Add the long range attachment tethers to the cloth asset to avoid a large hitch during the cloth's initialization.
 		ChaosClothAddTethersToCachedData,
 
+		// Always serialize the actor label in cooked builds
+		SerializeActorLabelInCookedBuilds,
+
 		// -----<new versions can be added above this line>-------------------------------------------------
 		VersionPlusOne,
 		LatestVersion = VersionPlusOne - 1
diff --git a/Engine/Source/Runtime/Engine/Classes/GameFramework/Actor.h b/Engine/Source/Runtime/Engine/Classes/GameFramework/Actor.h
--- a/Engine/Source/Runtime/Engine/Classes/GameFramework/Actor.h
+++ b/Engine/Source/Runtime/Engine/Classes/GameFramework/Actor.h
@@ -40,6 +40,10 @@ struct FNetworkObjectInfo;
 class UDataLayer;
 class AWorldDataLayers;
 
+// By default, debug and development builds (even cooked) will keep actor labels. Manually define this if you want to make a local build
+// that keep actor labels for Test or Shipping builds.
+#define ACTOR_HAS_LABELS (UE_BUILD_DEBUG || UE_BUILD_DEVELOPMENT)
+
 /** Chooses a method for actors to update overlap state (objects it is touching) on initialization, currently only used during level streaming. */
 UENUM(BlueprintType)
 enum class EActorUpdateOverlapsMethod : uint8
@@ -2268,6 +2272,23 @@ public:
 	}
 #endif		// WITH_EDITOR
 
+#if !WITH_EDITOR && ACTOR_HAS_LABELS
+private:
+	FString ActorLabel;
+
+public:
+	const FString& GetActorLabel() const { return ActorLabel; }
+#endif
+
+	const FString GetActorNameOrLabel() const
+	{
+#if WITH_EDITORONLY_DATA || (!WITH_EDITOR && ACTOR_HAS_LABELS)
+		return ActorLabel.IsEmpty() ? GetName() : ActorLabel;
+#else
+		return GetName();
+#endif
+	}
+
 	/**
 	 * Function used to prioritize actors when deciding which to replicate
 	 * @param ViewPos		Position of the viewer
diff --git a/Engine/Source/Runtime/Engine/Private/Actor.cpp b/Engine/Source/Runtime/Engine/Private/Actor.cpp
--- a/Engine/Source/Runtime/Engine/Private/Actor.cpp
+++ b/Engine/Source/Runtime/Engine/Private/Actor.cpp
@@ -17,6 +17,7 @@
 #include "UObject/UObjectHash.h"
 #include "UObject/PropertyPortFlags.h"
 #include "UObject/UE5MainStreamObjectVersion.h"
+#include "UObject/UE5PrivateFrostyStreamObjectVersion.h"
 #include "Kismet/GameplayStatics.h"
 #include "Engine/BlueprintGeneratedClass.h"
 #include "Engine/LocalPlayer.h"
@@ -722,6 +723,7 @@ bool AActor::IsReadyForFinishDestroy()
 void AActor::Serialize(FArchive& Ar)
 {
 	Ar.UsingCustomVersion(FUE5MainStreamObjectVersion::GUID);
+	Ar.UsingCustomVersion(FUE5PrivateFrostyStreamObjectVersion::GUID);
 
 #if WITH_EDITOR
 	// Prior to load, map natively-constructed component instances for Blueprint-generated class types to any serialized properties that might reference them.
@@ -788,6 +790,23 @@ void AActor::Serialize(FArchive& Ar)
 
 	Super::Serialize(Ar);
 
+	// Always serialize the actor label in cooked builds
+	if(Ar.IsPersistent() && (Ar.CustomVer(FUE5PrivateFrostyStreamObjectVersion::GUID) >= FUE5PrivateFrostyStreamObjectVersion::SerializeActorLabelInCookedBuilds))
+	{
+		bool bIsCooked = Ar.IsCooking();
+		Ar << bIsCooked;
+
+		if (bIsCooked)
+		{
+#if !ACTOR_HAS_LABELS
+			// In non-development builds, just skip over the actor labels. We need to figure out a way to either strip that data from shipping builds,
+			// or skip over the string data without doing any memory allocations, probably with a custom FString::SerializeToNull implementation.
+			FString ActorLabel;
+#endif
+			Ar << ActorLabel;
+		}
+	}
+
 #if WITH_EDITOR
 	// Fixup actor guids
 	if (Ar.IsLoading())
