Use actor label when possible when displaying streaming sources debug display.

diff --git a/Engine/Source/Runtime/Engine/Private/Components/WorldPartitionStreamingSourceComponent.cpp b/Engine/Source/Runtime/Engine/Private/Components/WorldPartitionStreamingSourceComponent.cpp
--- a/Engine/Source/Runtime/Engine/Private/Components/WorldPartitionStreamingSourceComponent.cpp
+++ b/Engine/Source/Runtime/Engine/Private/Components/WorldPartitionStreamingSourceComponent.cpp
@@ -57,16 +57,7 @@ bool UWorldPartitionStreamingSourceComponent::GetStreamingSource(FWorldPartition
 	{
 		AActor* Actor = GetOwner();
 
-		FName SourceName = Actor->GetFName();
-
-#if WITH_EDITOR
-		const FString& ActorLabel = Actor->GetActorLabel(false);
-		if (!ActorLabel.IsEmpty())
-		{
-			SourceName = *ActorLabel;
-		}
-#endif
-		OutStreamingSource.Name = SourceName;
+		OutStreamingSource.Name = *Actor->GetActorNameOrLabel();
 		OutStreamingSource.Location = Actor->GetActorLocation();
 		OutStreamingSource.Rotation = Actor->GetActorRotation();
 		OutStreamingSource.TargetState = TargetState;
