Small VT reporting improvements.
Don't double report header size to avoid confusion about total cost.
Use MB for disk usage.
Report residency oversubscription even when r.VT.Residency.MaxMipMapBias is 0.

diff --git a/Engine/Source/Runtime/Engine/Private/VT/UploadingVirtualTexture.cpp b/Engine/Source/Runtime/Engine/Private/VT/UploadingVirtualTexture.cpp
--- a/Engine/Source/Runtime/Engine/Private/VT/UploadingVirtualTexture.cpp
+++ b/Engine/Source/Runtime/Engine/Private/VT/UploadingVirtualTexture.cpp
@@ -12,8 +12,7 @@
 
 DECLARE_MEMORY_STAT(TEXT("File Cache Size"), STAT_FileCacheSize, STATGROUP_VirtualTextureMemory);
 DECLARE_MEMORY_STAT(TEXT("Total Header Size"), STAT_TotalHeaderSize, STATGROUP_VirtualTextureMemory);
-DECLARE_MEMORY_STAT(TEXT("Tile Header Size"), STAT_TileHeaderSize, STATGROUP_VirtualTextureMemory);
-DECLARE_DWORD_ACCUMULATOR_STAT(TEXT("Total Disk Size (KB)"), STAT_TotalDiskSize, STATGROUP_VirtualTextureMemory);
+DECLARE_DWORD_ACCUMULATOR_STAT(TEXT("Total Disk Size (MB)"), STAT_TotalDiskSize, STATGROUP_VirtualTextureMemory);
 DECLARE_DWORD_ACCUMULATOR_STAT(TEXT("Num Tile Headers"), STAT_NumTileHeaders, STATGROUP_VirtualTextureMemory);
 DECLARE_DWORD_ACCUMULATOR_STAT(TEXT("Num Codecs"), STAT_NumCodecs, STATGROUP_VirtualTextureMemory);
 
@@ -71,8 +70,7 @@ FUploadingVirtualTexture::FUploadingVirtualTexture(const FName& InName, FVirtual
 	StreamingManager = &IStreamingManager::Get().GetVirtualTextureStreamingManager();
 
 	INC_MEMORY_STAT_BY(STAT_TotalHeaderSize, InData->GetMemoryFootprint());
-	INC_MEMORY_STAT_BY(STAT_TileHeaderSize, InData->GetTileMemoryFootprint());
-	INC_DWORD_STAT_BY(STAT_TotalDiskSize, InData->GetDiskMemoryFootprint() / 1024);
+	INC_DWORD_STAT_BY(STAT_TotalDiskSize, InData->GetDiskMemoryFootprint() / (1024 * 1024));
 	INC_DWORD_STAT_BY(STAT_NumTileHeaders, InData->GetNumTileHeaders());
 }
 
@@ -80,8 +78,7 @@ FUploadingVirtualTexture::~FUploadingVirtualTexture()
 {
 	check(Data);
 	DEC_MEMORY_STAT_BY(STAT_TotalHeaderSize, Data->GetMemoryFootprint());
-	DEC_MEMORY_STAT_BY(STAT_TileHeaderSize, Data->GetTileMemoryFootprint());
-	DEC_DWORD_STAT_BY(STAT_TotalDiskSize, Data->GetDiskMemoryFootprint() / 1024);
+	DEC_DWORD_STAT_BY(STAT_TotalDiskSize, Data->GetDiskMemoryFootprint() / (1024 * 1024));
 	DEC_DWORD_STAT_BY(STAT_NumTileHeaders, Data->GetNumTileHeaders());
 
 	// Complete all open transcode requests before deleting IFileCacheHandle objects
