Add CVar to reset package map's AckState on server after seamless travel to prevent problems with clients missing entries in the NetGUID cache after travel.

diff --git a/Engine/Source/Runtime/Engine/Classes/Engine/PackageMapClient.h b/Engine/Source/Runtime/Engine/Classes/Engine/PackageMapClient.h
--- a/Engine/Source/Runtime/Engine/Classes/Engine/PackageMapClient.h
+++ b/Engine/Source/Runtime/Engine/Classes/Engine/PackageMapClient.h
@@ -500,6 +500,9 @@ public:
 	void RestorePackageMapExportAckStatus( const FPackageMapAckState& InState );
 	void OverridePackageMapExportAckStatus( FPackageMapAckState* NewState );
 
+	/** Resets the AckState and empties the PendingAckGuids, not meant to reset the OverrideAckState. */
+	void ResetAckState();
+
 	/** Functions to help with exporting/importing net field info */
 	TSharedPtr< FNetFieldExportGroup >	GetNetFieldExportGroup( const FString& PathName );
 	void								AddNetFieldExportGroup( const FString& PathName, TSharedPtr< FNetFieldExportGroup > NewNetFieldExportGroup );
diff --git a/Engine/Source/Runtime/Engine/Private/NetDriver.cpp b/Engine/Source/Runtime/Engine/Private/NetDriver.cpp
--- a/Engine/Source/Runtime/Engine/Private/NetDriver.cpp
+++ b/Engine/Source/Runtime/Engine/Private/NetDriver.cpp
@@ -336,6 +336,13 @@ static FAutoConsoleVariableRef CVarNetControlChannelDestructionInfo(
 	TEXT("0: Old behavior, use an actor channel. 1: New behavior, use the control channel"),
 	ECVF_Default);
 
+static int32 GNetResetAckStatePostSeamlessTravel = 0;
+static FAutoConsoleVariableRef CVarNetResetAckStatePostSeamlessTravel(
+	TEXT("net.ResetAckStatePostSeamlessTravel"),
+	GNetResetAckStatePostSeamlessTravel,
+	TEXT("If 1, the server will reset the ack state of the package map after seamless travel. Increases bandwidth usage, but may resolve some issues with GUIDs not being available on clients after seamlessly traveling."),
+	ECVF_Default);
+
 
 /*-----------------------------------------------------------------------------
 	UNetDriver implementation.
@@ -5695,6 +5702,20 @@ void UNetDriver::CleanPackageMaps()
 	{ 
 		GuidCache->CleanReferences();
 	}
+
+	if (GNetResetAckStatePostSeamlessTravel)
+	{
+		for (UNetConnection* Connection : ClientConnections)
+		{
+			if (Connection)
+			{ 
+				if (UPackageMapClient* PackageMapClient = Cast<UPackageMapClient>(Connection->PackageMap))
+				{
+					PackageMapClient->ResetAckState();
+				}
+			}
+		}
+	}
 }
 
 void UNetDriver::RemoveClassRepLayoutReferences(UClass* Class)
diff --git a/Engine/Source/Runtime/Engine/Private/PackageMapClient.cpp b/Engine/Source/Runtime/Engine/Private/PackageMapClient.cpp
--- a/Engine/Source/Runtime/Engine/Private/PackageMapClient.cpp
+++ b/Engine/Source/Runtime/Engine/Private/PackageMapClient.cpp
@@ -1914,6 +1914,12 @@ void UPackageMapClient::OverridePackageMapExportAckStatus( FPackageMapAckState*
 	OverrideAckState = NewState ? NewState : &AckState;
 }
 
+void UPackageMapClient::ResetAckState()
+{
+	AckState.Reset();
+	PendingAckGUIDs.Empty();
+}
+
 //--------------------------------------------------------------------
 //
 //	Network - ACKing
