- Added support of custom shapes for world partition streaming sources.
- Shapes can have their location/orientation and can be spheres or spherical sectors.
- Added IsStreamingCompleted function to WorldPartitionStreamingSourceComponent.

diff --git a/Engine/Source/Editor/ComponentVisualizers/Private/ComponentVisualizers.cpp b/Engine/Source/Editor/ComponentVisualizers/Private/ComponentVisualizers.cpp
--- a/Engine/Source/Editor/ComponentVisualizers/Private/ComponentVisualizers.cpp
+++ b/Engine/Source/Editor/ComponentVisualizers/Private/ComponentVisualizers.cpp
@@ -10,6 +10,7 @@
 #include "Components/SpotLightComponent.h"
 #include "Components/RectLightComponent.h"
 #include "Components/DecalComponent.h"
+#include "Components/WorldPartitionStreamingSourceComponent.h"
 #include "PhysicsEngine/RadialForceComponent.h"
 #include "UnrealEdGlobals.h"
 
@@ -32,6 +33,7 @@
 #include "SensingComponentVisualizer.h"
 #include "SpringComponentVisualizer.h"
 #include "StereoLayerComponentVisualizer.h"
+#include "WorldPartitionStreamingSourceComponentVisualizer.h"
 #include "PhysicsEngine/PhysicsConstraintComponent.h"
 #include "PhysicsEngine/PhysicalAnimationComponent.h"
 #include "Components/StereoLayerComponent.h"
@@ -56,6 +58,7 @@ void FComponentVisualizersModule::StartupModule()
 	RegisterComponentVisualizer(UPhysicsSpringComponent::StaticClass()->GetFName(), MakeShareable(new FSpringComponentVisualizer));
 	RegisterComponentVisualizer(UDecalComponent::StaticClass()->GetFName(), MakeShareable(new FDecalComponentVisualizer));
 	RegisterComponentVisualizer(UStereoLayerComponent::StaticClass()->GetFName(), MakeShareable(new FStereoLayerComponentVisualizer));
+	RegisterComponentVisualizer(UWorldPartitionStreamingSourceComponent::StaticClass()->GetFName(), MakeShareable(new FWorldPartitionStreamingSourceComponentVisualizer));
 }
 
 void FComponentVisualizersModule::ShutdownModule()
diff --git a/Engine/Source/Editor/ComponentVisualizers/Private/WorldPartitionStreamingSourceComponentVisualizer.cpp b/Engine/Source/Editor/ComponentVisualizers/Private/WorldPartitionStreamingSourceComponentVisualizer.cpp
new file mode 100644
--- /dev/null
+++ b/Engine/Source/Editor/ComponentVisualizers/Private/WorldPartitionStreamingSourceComponentVisualizer.cpp
@@ -0,0 +1,13 @@
+// Copyright Epic Games, Inc. All Rights Reserved.
+
+#include "WorldPartitionStreamingSourceComponentVisualizer.h"
+#include "SceneManagement.h"
+#include "Components/WorldPartitionStreamingSourceComponent.h"
+
+void FWorldPartitionStreamingSourceComponentVisualizer::DrawVisualization(const UActorComponent* Component, const FSceneView* View, FPrimitiveDrawInterface* PDI)
+{
+	if (const UWorldPartitionStreamingSourceComponent* StreamingSourceComponent = Cast<const UWorldPartitionStreamingSourceComponent>(Component))
+	{
+		StreamingSourceComponent->DrawVisualization(View, PDI);
+	}
+}
diff --git a/Engine/Source/Editor/ComponentVisualizers/Public/WorldPartitionStreamingSourceComponentVisualizer.h b/Engine/Source/Editor/ComponentVisualizers/Public/WorldPartitionStreamingSourceComponentVisualizer.h
new file mode 100644
--- /dev/null
+++ b/Engine/Source/Editor/ComponentVisualizers/Public/WorldPartitionStreamingSourceComponentVisualizer.h
@@ -0,0 +1,17 @@
+// Copyright Epic Games, Inc. All Rights Reserved.
+
+#pragma once
+
+#include "CoreMinimal.h"
+#include "ComponentVisualizer.h"
+
+class FPrimitiveDrawInterface;
+class FSceneView;
+
+class COMPONENTVISUALIZERS_API FWorldPartitionStreamingSourceComponentVisualizer : public FComponentVisualizer
+{
+public:
+	//~ Begin FComponentVisualizer Interface
+	virtual void DrawVisualization(const UActorComponent* Component, const FSceneView* View, FPrimitiveDrawInterface* PDI) override;
+	//~ End FComponentVisualizer Interface
+};
diff --git a/Engine/Source/Runtime/Engine/Classes/Components/WorldPartitionStreamingSourceComponent.h b/Engine/Source/Runtime/Engine/Classes/Components/WorldPartitionStreamingSourceComponent.h
--- a/Engine/Source/Runtime/Engine/Classes/Components/WorldPartitionStreamingSourceComponent.h
+++ b/Engine/Source/Runtime/Engine/Classes/Components/WorldPartitionStreamingSourceComponent.h
@@ -8,6 +8,9 @@
 #include "WorldPartition/WorldPartitionStreamingSource.h"
 #include "WorldPartitionStreamingSourceComponent.generated.h"
 
+class FSceneView;
+class FPrimitiveDrawInterface;
+
 UCLASS(Meta = (BlueprintSpawnableComponent), HideCategories = (Tags, Sockets, ComponentTick, ComponentReplication, Activation, Cooking, Events, AssetUserData, Collision))
 class ENGINE_API UWorldPartitionStreamingSourceComponent : public UActorComponent, public IWorldPartitionStreamingSourceProvider
 {
@@ -31,6 +34,27 @@ class ENGINE_API UWorldPartitionStreamingSourceComponent : public UActorComponen
 	// IWorldPartitionStreamingSourceProvider interface
 	virtual bool GetStreamingSource(FWorldPartitionStreamingSource& OutStreamingSource) override;
 
+	/** Returns true if streaming is completed for this streaming source component. */
+	UFUNCTION(BlueprintCallable, Category = "Streaming")
+	bool IsStreamingCompleted() const;
+
+	/** Displays a debug visualizer of the streaming source. Useful when using Shapes. */
+	void DrawVisualization(const FSceneView* View, FPrimitiveDrawInterface* PDI) const;
+
+#if WITH_EDITORONLY_DATA
+	/** Value used by debug visualizer when grid loading range is chosen. */
+	UPROPERTY(EditAnywhere, Category = "Debug")
+	float DefaultVisualizerLoadingRange;
+#endif
+
+	/** Optional target grid affected by streaming source. When none chosen, affects all grids. */
+	UPROPERTY(BlueprintReadWrite, EditAnywhere, Category = "Streaming")
+	FName TargetGrid;
+
+	/** Optional aggregated shape list used to build a custom shape for the streaming source. When empty, fallbacks sphere shape with a radius equal to grid's loading range. */
+	UPROPERTY(BlueprintReadWrite, EditAnywhere, Category = "Streaming")
+	TArray<FStreamingSourceShape> Shapes;
+
 private:
 	/** Whether this component is enabled or not */
 	UPROPERTY(EditAnywhere, Category = "Streaming")
diff --git a/Engine/Source/Runtime/Engine/Private/Components/WorldPartitionStreamingSourceComponent.cpp b/Engine/Source/Runtime/Engine/Private/Components/WorldPartitionStreamingSourceComponent.cpp
--- a/Engine/Source/Runtime/Engine/Private/Components/WorldPartitionStreamingSourceComponent.cpp
+++ b/Engine/Source/Runtime/Engine/Private/Components/WorldPartitionStreamingSourceComponent.cpp
@@ -3,10 +3,19 @@
 #include "Components/WorldPartitionStreamingSourceComponent.h"
 #include "WorldPartition/WorldPartition.h"
 #include "WorldPartition/WorldPartitionRuntimeCell.h"
+#include "WorldPartition/WorldPartitionSubsystem.h"
+#include "WorldPartition/DataLayer/DataLayerSubsystem.h"
+#include "WorldPartition/WorldPartitionStreamingSource.h"
+#include "SceneManagement.h"
+#include "SceneView.h"
 #include "Engine/World.h"
 
 UWorldPartitionStreamingSourceComponent::UWorldPartitionStreamingSourceComponent(const FObjectInitializer& ObjectInitializer)
 	: Super(ObjectInitializer)
+#if WITH_EDITORONLY_DATA
+	, DefaultVisualizerLoadingRange(10000.f)
+#endif
+	, TargetGrid(NAME_None)
 	, bStreamingSourceEnabled(true)
 	, TargetState(EStreamingSourceTargetState::Activated)
 {
@@ -61,7 +70,65 @@ bool UWorldPartitionStreamingSourceComponent::GetStreamingSource(FWorldPartition
 		OutStreamingSource.Location = Actor->GetActorLocation();
 		OutStreamingSource.Rotation = Actor->GetActorRotation();
 		OutStreamingSource.TargetState = TargetState;
+		OutStreamingSource.TargetGrid = TargetGrid;
+		OutStreamingSource.Shapes = Shapes;
 		return true;
 	}
 	return false;
+}
+
+bool UWorldPartitionStreamingSourceComponent::IsStreamingCompleted() const
+{
+	UWorld* World = GetWorld();
+	if (!bStreamingSourceEnabled || !World->IsGameWorld())
+	{
+		return false;
+	}
+	
+	UWorldPartitionSubsystem* WorldPartitionSubsystem = World->GetSubsystem<UWorldPartitionSubsystem>();
+	UDataLayerSubsystem* DataLayerSubsystem = World->GetSubsystem<UDataLayerSubsystem>();
+	if (!WorldPartitionSubsystem || !DataLayerSubsystem)	
+	{
+		return false;
+	}
+
+	// Build a query source
+	AActor* Actor = GetOwner();
+	TArray<FWorldPartitionStreamingQuerySource> QuerySources;
+	FWorldPartitionStreamingQuerySource& QuerySource = QuerySources.Emplace_GetRef();
+	QuerySource.bSpatialQuery = true;
+	QuerySource.Location = Actor->GetActorLocation();
+	QuerySource.Rotation = Actor->GetActorRotation();
+	QuerySource.TargetGrid = TargetGrid;
+	QuerySource.Shapes = Shapes;
+	QuerySource.bUseGridLoadingRange = true;
+	QuerySource.Radius = 0.f;
+	QuerySource.bDataLayersOnly = false;
+	QuerySource.DataLayers = (TargetState == EStreamingSourceTargetState::Loaded) ? DataLayerSubsystem->GetLoadedDataLayerNames().Array() : DataLayerSubsystem->GetActiveDataLayerNames().Array();
+
+	// Execute query
+	const EWorldPartitionRuntimeCellState QueryState = (TargetState == EStreamingSourceTargetState::Loaded) ? EWorldPartitionRuntimeCellState::Loaded : EWorldPartitionRuntimeCellState::Activated;
+	return WorldPartitionSubsystem->IsStreamingCompleted(QueryState, QuerySources, /*bExactState*/ true);
+}
+
+void UWorldPartitionStreamingSourceComponent::DrawVisualization(const FSceneView* View, FPrimitiveDrawInterface* PDI) const
+{
+#if WITH_EDITOR
+	AActor* Actor = GetOwner();
+	FStreamingSourceShapeHelper::ForEachShape(DefaultVisualizerLoadingRange, DefaultVisualizerLoadingRange, /*bInProjectIn2D*/ false, Actor->GetActorLocation(), Actor->GetActorRotation(), Shapes, [&PDI](const FSphericalSector& Shape)
+	{
+		if (Shape.IsSphere())
+		{
+			DrawWireSphere(PDI, Shape.GetCenter(), FColor::White, Shape.GetRadius(), 32, SDPG_World, 1.0, 0, true);
+		}
+		else
+		{
+			TArray<TPair<FVector, FVector>> Lines = Shape.BuildDebugMesh();
+			for (const auto& Line : Lines)
+			{
+				PDI->DrawLine(Line.Key, Line.Value, FColor::White, SDPG_World, 1.0, 0, true);
+			};
+		}
+	});
+#endif
 }
\ No newline at end of file
diff --git a/Engine/Source/Runtime/Engine/Private/WorldPartition/RuntimeSpatialHash/RuntimeSpatialHashGridHelper.cpp b/Engine/Source/Runtime/Engine/Private/WorldPartition/RuntimeSpatialHash/RuntimeSpatialHashGridHelper.cpp
--- a/Engine/Source/Runtime/Engine/Private/WorldPartition/RuntimeSpatialHash/RuntimeSpatialHashGridHelper.cpp
+++ b/Engine/Source/Runtime/Engine/Private/WorldPartition/RuntimeSpatialHash/RuntimeSpatialHashGridHelper.cpp
@@ -95,6 +95,18 @@ int32 FSquare2DGridHelper::ForEachIntersectingCells(const FSphere& InSphere, TFu
 	return NumCells;
 }
 
+int32 FSquare2DGridHelper::ForEachIntersectingCells(const FSphericalSector& InShape, TFunctionRef<void(const FIntVector&)> InOperation) const
+{
+	int32 NumCells = 0;
+
+	for (int32 Level = 0; Level < Levels.Num(); Level++)
+	{
+		NumCells += Levels[Level].ForEachIntersectingCells(InShape, [Level, InOperation](const FIntVector2& Coord) { InOperation(FIntVector(Coord.X, Coord.Y, Level)); });
+	}
+
+	return NumCells;
+}
+
 #if WITH_EDITOR
 void FSquare2DGridHelper::ValidateSingleActorReferer()
 {
@@ -268,3 +280,92 @@ FSquare2DGridHelper GetPartitionedActors(const UWorldPartition* WorldPartition,
 }
 
 #endif // #if WITH_EDITOR
+
+FORCEINLINE static bool IsClockWise(const FVector2D& V1, const FVector2D& V2) 
+{
+	return (V1 ^ V2) < 0;
+}
+
+FORCEINLINE static bool IsVectorInsideVectorPair(const FVector2D& TestVector, const FVector2D& V1, const FVector2D& V2)
+{
+	return IsClockWise(V1, TestVector) && !IsClockWise(V2, TestVector);
+}
+
+FORCEINLINE static bool IsPointInsideSector(const FVector2D& TestPoint, const FVector2D& SectorCenter, float SectorRadiusSquared, const FVector2D& SectorStart, const FVector2D& SectorEnd, float SectorAngle)
+{
+	const FVector2D TestVector = TestPoint - SectorCenter;
+	if (TestVector.SizeSquared() > SectorRadiusSquared)
+	{
+		return false;
+	}
+	
+	if (SectorAngle <= 180.0f)
+	{
+		return IsVectorInsideVectorPair(TestVector, SectorStart, SectorEnd);
+	}
+	else
+	{
+		return !IsVectorInsideVectorPair(TestVector, SectorEnd, SectorStart);
+	}
+}
+
+bool FSquare2DGridHelper::FGrid2D::DoesCircleSectorIntersectsCell(const FIntVector2& Coords, const FVector2D& SectorCenter, float SectorRadiusSquared, const FVector2D& SectorStartVector, const FVector2D& SectorEndVector, float SectorAngle) const
+{
+	const int32 CellIndex = Coords.Y * GridSize + Coords.X;
+	FBox2D CellBounds;
+	GetCellBounds(CellIndex, CellBounds);
+
+	// Test whether any cell corners are inside sector
+	if (IsPointInsideSector(FVector2D(CellBounds.Min.X, CellBounds.Min.Y), SectorCenter, SectorRadiusSquared, SectorStartVector, SectorEndVector, SectorAngle) ||
+		IsPointInsideSector(FVector2D(CellBounds.Max.X, CellBounds.Min.Y), SectorCenter, SectorRadiusSquared, SectorStartVector, SectorEndVector, SectorAngle) ||
+		IsPointInsideSector(FVector2D(CellBounds.Max.X, CellBounds.Max.Y), SectorCenter, SectorRadiusSquared, SectorStartVector, SectorEndVector, SectorAngle) ||
+		IsPointInsideSector(FVector2D(CellBounds.Min.X, CellBounds.Max.Y), SectorCenter, SectorRadiusSquared, SectorStartVector, SectorEndVector, SectorAngle))
+	{
+		return true;
+	}
+
+	// Test whether any sector point lies inside the cell bounds
+	if (CellBounds.IsInside(SectorCenter) ||
+		CellBounds.IsInside(SectorCenter + SectorStartVector) ||
+		CellBounds.IsInside(SectorCenter + SectorEndVector))
+	{
+		return true;
+	}
+
+	// Test whether closest point on cell from center is inside sector
+	if (IsPointInsideSector(CellBounds.GetClosestPointTo(SectorCenter), SectorCenter, SectorRadiusSquared, SectorStartVector, SectorEndVector, SectorAngle))
+	{
+		return true;
+	}
+
+	return false;
+}
+
+int32 FSquare2DGridHelper::FGrid2D::ForEachIntersectingCells(const FSphericalSector& InShape, TFunctionRef<void(const FIntVector2&)> InOperation) const
+{
+	check(!InShape.IsNearlyZero());
+	const bool bIsSphere = InShape.IsSphere();
+	const FVector ScaledAxis = FVector(FVector2D(InShape.GetAxis()), 0).GetSafeNormal() * InShape.GetRadius();
+	const FVector SectorStart = FRotator(0, 0.5f * InShape.GetAngle(), 0).RotateVector(ScaledAxis);
+	const FVector SectorEnd = FRotator(0, -0.5f * InShape.GetAngle(), 0).RotateVector(ScaledAxis);
+
+	const FVector2D Center2D(InShape.GetCenter());
+	const FVector2D SectorStart2D = FVector2D(SectorStart);
+	const FVector2D SectorEnd2D = FVector2D(SectorEnd);
+	const float SectorRadiusSquared = InShape.GetRadius() * InShape.GetRadius();
+
+	int32 NumCells = 0;
+	const FSphere Sphere(InShape.GetCenter(), InShape.GetRadius());
+	ForEachIntersectingCells(Sphere, [this, bIsSphere, Center2D, SectorRadiusSquared, SectorStart2D, SectorEnd2D, SectorAngle = InShape.GetAngle(), InOperation, &NumCells](const FIntVector2& Coords)
+	{
+		// If sector, filter coords that intersect sector
+		if (bIsSphere || DoesCircleSectorIntersectsCell(Coords, Center2D, SectorRadiusSquared, SectorStart2D, SectorEnd2D, SectorAngle))
+		{
+			InOperation(Coords);
+			NumCells++;
+		}
+	});
+
+	return NumCells;
+}
+
diff --git a/Engine/Source/Runtime/Engine/Private/WorldPartition/RuntimeSpatialHash/RuntimeSpatialHashGridHelper.h b/Engine/Source/Runtime/Engine/Private/WorldPartition/RuntimeSpatialHash/RuntimeSpatialHashGridHelper.h
--- a/Engine/Source/Runtime/Engine/Private/WorldPartition/RuntimeSpatialHash/RuntimeSpatialHashGridHelper.h
+++ b/Engine/Source/Runtime/Engine/Private/WorldPartition/RuntimeSpatialHash/RuntimeSpatialHashGridHelper.h
@@ -224,22 +224,33 @@ struct FSquare2DGridHelper
 			const FBox Box(InSphere.Center - FVector(InSphere.W), InSphere.Center + FVector(InSphere.W));
 
 			ForEachIntersectingCells(Box, [this, &InSphere, &InOperation, &NumCells](const FIntVector2& Coords)
-				{
-					const int32 CellIndex = Coords.Y * GridSize + Coords.X;
+			{
+				const int32 CellIndex = Coords.Y * GridSize + Coords.X;
 
-					FBox2D CellBounds;
-					GetCellBounds(CellIndex, CellBounds);
+				FBox2D CellBounds;
+				GetCellBounds(CellIndex, CellBounds);
 
-					FVector2D Delta = FVector2D(InSphere.Center) - FVector2D::Max(CellBounds.GetCenter() - CellBounds.GetExtent(), FVector2D::Min(FVector2D(InSphere.Center), CellBounds.GetCenter() + CellBounds.GetExtent()));
-					if ((Delta.X * Delta.X + Delta.Y * Delta.Y) < (InSphere.W * InSphere.W))
-					{
-						InOperation(Coords);
-						NumCells++;
-					}
-				});
+				FVector2D Delta = FVector2D(InSphere.Center) - FVector2D::Max(CellBounds.GetCenter() - CellBounds.GetExtent(), FVector2D::Min(FVector2D(InSphere.Center), CellBounds.GetCenter() + CellBounds.GetExtent()));
+				if ((Delta.X * Delta.X + Delta.Y * Delta.Y) < (InSphere.W * InSphere.W))
+				{
+					InOperation(Coords);
+					NumCells++;
+				}
+			});
 
 			return NumCells;
 		}
+		
+		/**
+		 * Runs a function on all intersecting cells for the provided spherical sector
+		 *
+		 * @return the number of intersecting cells
+		 */
+		int32 ForEachIntersectingCells(const FSphericalSector& InShape, TFunctionRef<void(const FIntVector2&)> InOperation) const;
+
+	private:
+
+		bool DoesCircleSectorIntersectsCell(const FIntVector2& Coords, const FVector2D& SectorCenter, float SectorRadiusSquared, const FVector2D& SectorStartVector, const FVector2D& SectorEndVector, float SectorAngle) const;
 	};
 
 	struct FGridLevel : public FGrid2D
@@ -446,6 +457,13 @@ struct FSquare2DGridHelper
 	 */
 	int32 ForEachIntersectingCells(const FSphere& InSphere, TFunctionRef<void(const FIntVector&)> InOperation) const;
 
+	/**
+	 * Runs a function on all intersecting cells for the provided spherical sector
+	 *
+	 * @return the number of intersecting cells
+	 */
+	int32 ForEachIntersectingCells(const FSphericalSector& InShape, TFunctionRef<void(const FIntVector&)> InOperation) const;
+
 #if WITH_EDITOR
 	// Validates that actor is not referenced by multiple cells
 	void ValidateSingleActorReferer();
diff --git a/Engine/Source/Runtime/Engine/Private/WorldPartition/WorldPartitionRuntimeCell.cpp b/Engine/Source/Runtime/Engine/Private/WorldPartition/WorldPartitionRuntimeCell.cpp
--- a/Engine/Source/Runtime/Engine/Private/WorldPartition/WorldPartitionRuntimeCell.cpp
+++ b/Engine/Source/Runtime/Engine/Private/WorldPartition/WorldPartitionRuntimeCell.cpp
@@ -73,7 +73,7 @@ void UWorldPartitionRuntimeCell::UpdateDebugName()
 
 #endif
 
-bool UWorldPartitionRuntimeCell::CacheStreamingSourceInfo(const FWorldPartitionStreamingSource& Source) const
+bool UWorldPartitionRuntimeCell::CacheStreamingSourceInfo(const UWorldPartitionRuntimeCell::FStreamingSourceInfo& Info) const
 {
 	bool bWasCacheDirtied = false;
 	if (CachedSourceInfoEpoch != UWorldPartitionRuntimeCell::StreamingSourceCacheEpoch)
@@ -83,7 +83,7 @@ bool UWorldPartitionRuntimeCell::CacheStreamingSourceInfo(const FWorldPartitionS
 	}
 
 	// If cache was dirtied, use value, else use minimum with existing cached value
-	CachedSourcePriority = bWasCacheDirtied ? (int32)Source.Priority : FMath::Min((int32)Source.Priority, CachedSourcePriority);
+	CachedSourcePriority = bWasCacheDirtied ? (int32)Info.Source.Priority : FMath::Min((int32)Info.Source.Priority, CachedSourcePriority);
 	return bWasCacheDirtied;
 }
 
diff --git a/Engine/Source/Runtime/Engine/Private/WorldPartition/WorldPartitionRuntimeHash.cpp b/Engine/Source/Runtime/Engine/Private/WorldPartition/WorldPartitionRuntimeHash.cpp
--- a/Engine/Source/Runtime/Engine/Private/WorldPartition/WorldPartitionRuntimeHash.cpp
+++ b/Engine/Source/Runtime/Engine/Private/WorldPartition/WorldPartitionRuntimeHash.cpp
@@ -310,10 +310,10 @@ EWorldPartitionStreamingPerformance UWorldPartitionRuntimeHash::GetStreamingPerf
 	return StreamingPerformance;
 }
 
-void UWorldPartitionRuntimeHash::FStreamingSourceCells::AddCell(const UWorldPartitionRuntimeCell* InCell, const FWorldPartitionStreamingSource& InSource)
+void UWorldPartitionRuntimeHash::FStreamingSourceCells::AddCell(const UWorldPartitionRuntimeCell* Cell, const UWorldPartitionRuntimeCell::FStreamingSourceInfo& Info)
 {
-	InCell->CacheStreamingSourceInfo(InSource);
-	Cells.Add(InCell);
+	Cell->CacheStreamingSourceInfo(Info);
+	Cells.Add(Cell);
 }
 
 #undef LOCTEXT_NAMESPACE
\ No newline at end of file
diff --git a/Engine/Source/Runtime/Engine/Private/WorldPartition/WorldPartitionRuntimeSpatialHash.cpp b/Engine/Source/Runtime/Engine/Private/WorldPartition/WorldPartitionRuntimeSpatialHash.cpp
--- a/Engine/Source/Runtime/Engine/Private/WorldPartition/WorldPartitionRuntimeSpatialHash.cpp
+++ b/Engine/Source/Runtime/Engine/Private/WorldPartition/WorldPartitionRuntimeSpatialHash.cpp
@@ -33,6 +33,7 @@
 #include "Algo/Transform.h"
 #include "Algo/Copy.h"
 #include "UObject/ObjectSaveContext.h"
+#include "Components/LineBatchComponent.h"
 
 #if WITH_EDITOR
 #include "Editor/EditorEngine.h"
@@ -142,23 +143,25 @@ void FSpatialHashStreamingGrid::GetCells(const FWorldPartitionStreamingQuerySour
 	};
 
 	const FSquare2DGridHelper& Helper = GetGridHelper();
-		// Spatial Query
+	// Spatial Query
 	if (QuerySource.bSpatialQuery)
 	{
-		const FSphere GridSphere(QuerySource.Location, QuerySource.bUseLoadingRangeRadius ? GetLoadingRange() : QuerySource.Radius);
-		Helper.ForEachIntersectingCells(GridSphere, [&](const FIntVector& Coords)
+		QuerySource.ForEachShape(GetLoadingRange(), GridName, /*bProjectIn2D*/ true, [&](const FSphericalSector& Shape)
 		{
-			if (const int32* LayerCellIndexPtr = GridLevels[Coords.Z].LayerCellsMapping.Find(Coords.Y * Helper.Levels[Coords.Z].GridSize + Coords.X))
+			Helper.ForEachIntersectingCells(Shape, [&](const FIntVector& Coords)
 			{
-				const FSpatialHashStreamingGridLayerCell& LayerCell = GridLevels[Coords.Z].LayerCells[*LayerCellIndexPtr];
-				for (const UWorldPartitionRuntimeCell* Cell : LayerCell.GridCells)
+				if (const int32* LayerCellIndexPtr = GridLevels[Coords.Z].LayerCellsMapping.Find(Coords.Y * Helper.Levels[Coords.Z].GridSize + Coords.X))
 				{
-					if (ShouldAddCell(Cell, QuerySource))
+					const FSpatialHashStreamingGridLayerCell& LayerCell = GridLevels[Coords.Z].LayerCells[*LayerCellIndexPtr];
+					for (const UWorldPartitionRuntimeCell* Cell : LayerCell.GridCells)
 					{
-						OutCells.Add(Cell);
+						if (ShouldAddCell(Cell, QuerySource))
+						{
+							OutCells.Add(Cell);
+						}
 					}
 				}
-			}
+			});
 		});
 	}
 
@@ -179,36 +182,41 @@ void FSpatialHashStreamingGrid::GetCells(const FWorldPartitionStreamingQuerySour
 void FSpatialHashStreamingGrid::GetCells(const TArray<FWorldPartitionStreamingSource>& Sources, const UDataLayerSubsystem* DataLayerSubsystem, UWorldPartitionRuntimeHash::FStreamingSourceCells& OutActivateCells, UWorldPartitionRuntimeHash::FStreamingSourceCells& OutLoadCells) const
 {
 	TRACE_CPUPROFILER_EVENT_SCOPE(FSpatialHashStreamingGrid::GetCells);
-
+	
+	const float GridLoadingRange = GetLoadingRange();
 	const FSquare2DGridHelper& Helper = GetGridHelper();
 	for (const FWorldPartitionStreamingSource& Source : Sources)
 	{
-		const FSphere GridSphere(Source.Location, GetLoadingRange());
-		Helper.ForEachIntersectingCells(GridSphere, [&](const FIntVector& Coords)
+		Source.ForEachShape(GridLoadingRange, GridName, /*bProjectIn2D*/ true, [&](const FSphericalSector& Shape)
 		{
-			if (const int32* LayerCellIndexPtr = GridLevels[Coords.Z].LayerCellsMapping.Find(Coords.Y * Helper.Levels[Coords.Z].GridSize + Coords.X))
-			{
-				const FSpatialHashStreamingGridLayerCell& LayerCell = GridLevels[Coords.Z].LayerCells[*LayerCellIndexPtr];
-				for (const UWorldPartitionRuntimeCell* Cell : LayerCell.GridCells)
+			UWorldPartitionRuntimeCell::FStreamingSourceInfo Info(Source, Shape);
+
+			Helper.ForEachIntersectingCells(Shape, [&](const FIntVector& Coords)
+			{ 
+				if (const int32* LayerCellIndexPtr = GridLevels[Coords.Z].LayerCellsMapping.Find(Coords.Y * Helper.Levels[Coords.Z].GridSize + Coords.X))
 				{
-					if (!Cell->HasDataLayers() || (DataLayerSubsystem && DataLayerSubsystem->IsAnyDataLayerInState(Cell->GetDataLayers(), EDataLayerState::Activated)))
+					const FSpatialHashStreamingGridLayerCell& LayerCell = GridLevels[Coords.Z].LayerCells[*LayerCellIndexPtr];
+					for (const UWorldPartitionRuntimeCell* Cell : LayerCell.GridCells)
 					{
-						if (Source.TargetState == EStreamingSourceTargetState::Loaded)
+						if (!Cell->HasDataLayers() || (DataLayerSubsystem && DataLayerSubsystem->IsAnyDataLayerInState(Cell->GetDataLayers(), EDataLayerState::Activated)))
 						{
-							OutLoadCells.AddCell(Cell, Source);
+							if (Source.TargetState == EStreamingSourceTargetState::Loaded)
+							{
+								OutLoadCells.AddCell(Cell, Info);
+							}
+							else
+							{
+								check(Source.TargetState == EStreamingSourceTargetState::Activated);
+								OutActivateCells.AddCell(Cell, Info);
+							}
 						}
-						else
+						else if (DataLayerSubsystem && DataLayerSubsystem->IsAnyDataLayerInState(Cell->GetDataLayers(), EDataLayerState::Loaded))
 						{
-							check(Source.TargetState == EStreamingSourceTargetState::Activated);
-							OutActivateCells.AddCell(Cell, Source);
+							OutLoadCells.AddCell(Cell, Info);
 						}
 					}
-					else if (DataLayerSubsystem && DataLayerSubsystem->IsAnyDataLayerInState(Cell->GetDataLayers(), EDataLayerState::Loaded))
-					{
-						OutLoadCells.AddCell(Cell, Source);
-					}
 				}
-			}
+			});
 		});
 	}
 
@@ -279,9 +287,10 @@ void FSpatialHashStreamingGrid::Draw3D(UWorld* World, const TArray<FWorldPartiti
 	int32 MaxGridLevel = FMath::Clamp<int32>(MinGridLevel + GShowRuntimeSpatialHashGridLevelCount - 1, 0, GridLevels.Num() - 1);
 	const float GridViewMinimumSizeInCellCount = 5.f;
 	const float GridViewLoadingRangeExtentRatio = 1.5f;
-	const float Radius = GetLoadingRange();
-	const float GridSideDistance = FMath::Max((2.f * Radius * GridViewLoadingRangeExtentRatio), CellSize * GridViewMinimumSizeInCellCount);
+	const float GridLoadingRange = GetLoadingRange();
+	const FVector MinExtent(CellSize * GridViewMinimumSizeInCellCount);
 	TArray<const UWorldPartitionRuntimeCell*> FilteredCells;
+	TSet<const UWorldPartitionRuntimeCell*> DrawnCells;
 
 	for (const FWorldPartitionStreamingSource& Source : Sources)
 	{
@@ -294,8 +303,9 @@ void FSpatialHashStreamingGrid::Draw3D(UWorld* World, const TArray<FWorldPartiti
 			Z = Hit.ImpactPoint.Z;
 		}
 
-		FSphere Sphere(Source.Location, GridSideDistance * 0.5f);
-		const FBox Region(Sphere.Center - Sphere.W, Sphere.Center + Sphere.W);
+		FBox Region = Source.CalcBounds(GridLoadingRange, GridName);
+		Region += FBox(Region.GetCenter() - MinExtent, Region.GetCenter() + MinExtent);
+
 		for (int32 GridLevel = MinGridLevel; GridLevel <= MaxGridLevel; ++GridLevel)
 		{
 			Helper.Levels[GridLevel].ForEachIntersectingCells(Region, [&](const FIntVector2& Coords)
@@ -316,6 +326,13 @@ void FSpatialHashStreamingGrid::Draw3D(UWorld* World, const TArray<FWorldPartiti
 
 				for (const UWorldPartitionRuntimeCell* Cell : FilteredCells)
 				{
+					bool bIsAlreadyInSet = false;
+					DrawnCells.Add(Cell, &bIsAlreadyInSet);
+					if (bIsAlreadyInSet)
+					{
+						continue;
+					}
+
 					// Draw Cell using its debug color
 					FColor CellColor = Cell->GetDebugColor().ToFColor(false).WithAlpha(16);
 					DrawDebugSolidBox(World, CellBox, CellColor, Transform, false, -1.f, 255);
@@ -341,10 +358,38 @@ void FSpatialHashStreamingGrid::Draw3D(UWorld* World, const TArray<FWorldPartiti
 			});
 		}
 
-		// Draw Loading Ranges
-		FVector SphereLocation(FVector2D(Source.Location), Z);
-		SphereLocation = Transform.TransformPosition(SphereLocation);
-		DrawDebugSphere(World, SphereLocation, Radius, 32, FColor::White, false, -1.f, 0, 20.f);
+		// Draw Streaming Source
+		const FColor Color = Source.GetDebugColor();
+		Source.ForEachShape(GetLoadingRange(), GridName, /*bProjectIn2D*/ true, [&Color, &Z, &Transform, &World, this](const FSphericalSector& Shape)
+		{
+			FSphericalSector ZOffsettedShape = Shape;
+			ZOffsettedShape.SetCenter(FVector(FVector2D(ZOffsettedShape.GetCenter()), Z));
+			DrawStreamingSource3D(World, ZOffsettedShape, Transform, Color);
+		});
+	}
+}
+
+void FSpatialHashStreamingGrid::DrawStreamingSource3D(UWorld* World, const FSphericalSector& InShape, const FTransform& InTransform, const FColor& InColor) const
+{
+	if (InShape.IsSphere())
+	{
+		FVector Location = InTransform.TransformPosition(InShape.GetCenter());
+		DrawDebugSphere(World, Location, InShape.GetRadius(), 32, InColor, false, -1.f, 0, 20.f);
+	}
+	else if (ULineBatchComponent* const LineBatcher = World ? World->LineBatcher : nullptr)
+	{
+		FSphericalSector Shape = InShape;
+		Shape.SetAxis(InTransform.TransformVector(Shape.GetAxis()));
+		Shape.SetCenter(InTransform.TransformPosition(Shape.GetCenter()));
+
+		TArray<TPair<FVector, FVector>> Lines = Shape.BuildDebugMesh();
+		TArray<FBatchedLine> BatchedLines;
+		BatchedLines.Empty(Lines.Num());
+		for (const auto& Line : Lines)
+		{
+			BatchedLines.Emplace(FBatchedLine(Line.Key, Line.Value, InColor, LineBatcher->DefaultLifeTime, 20.f, SDPG_World));
+		};
+		LineBatcher->DrawLines(BatchedLines);
 	}
 }
 
@@ -469,38 +514,15 @@ void FSpatialHashStreamingGrid::Draw2D(UCanvas* Canvas, UWorld* World, const TAr
 			}
 		});
 
-		// Draw Loading Ranges
-		float Range = GetLoadingRange();
-
-		FCanvasLineItem LineItem;
-		LineItem.LineThickness = 2;
-
+		// Draw Streaming Sources
+		const float GridLoadingRange = GetLoadingRange();
 		for (const FWorldPartitionStreamingSource& Source : Sources)
 		{
-			LineItem.SetColor(Source.GetDebugColor());
-
-			TArray<FVector> LinePoints;
-			LinePoints.SetNum(2);
-
-			float Sin, Cos;
-			FMath::SinCos(&Sin, &Cos, (63.0f / 64.0f) * 2.0f * PI);
-			FVector2D LineStart(Sin * Range, Cos * Range);
-
-			for (int32 i = 0; i < 64; i++)
-			{
-				FMath::SinCos(&Sin, &Cos, (i / 64.0f) * 2.0f * PI);
-				FVector2D LineEnd(Sin * Range, Cos * Range);
-				LineItem.Draw(CanvasObject, WorldToScreen(FVector2D(Source.Location) + LineStart), WorldToScreen(FVector2D(Source.Location) + LineEnd));
-				LineStart = LineEnd;
-			}
-
-			FVector2D SourceDir = FVector2D(Source.Rotation.Vector());
-			if (SourceDir.Size())
+			const FColor Color = Source.GetDebugColor();
+			Source.ForEachShape(GridLoadingRange, GridName, /*bProjectIn2D*/ true, [&Color, &Canvas, &WorldToScreen, this](const FSphericalSector& Shape)
 			{
-				SourceDir.Normalize();
-				FVector2D ConeCenter(FVector2D(Source.Location));
-				LineItem.Draw(CanvasObject, WorldToScreen(ConeCenter), WorldToScreen(ConeCenter + SourceDir * Range));
-			}
+				DrawStreamingSource2D(Canvas, Shape, WorldToScreen, Color);
+			});
 		}
 
 		FCanvasBoxItem Box(GridScreenBounds.Min, GridScreenBounds.GetSize());
@@ -523,6 +545,43 @@ void FSpatialHashStreamingGrid::Draw2D(UCanvas* Canvas, UWorld* World, const TAr
 	Canvas->DrawItem(Box);
 }
 
+void FSpatialHashStreamingGrid::DrawStreamingSource2D(UCanvas* Canvas, const FSphericalSector& Shape, TFunctionRef<FVector2D(const FVector2D&)> WorldToScreen, const FColor& Color) const
+{
+	check(!Shape.IsNearlyZero())
+
+	FCanvasLineItem LineItem;
+	LineItem.LineThickness = 2;
+	LineItem.SetColor(Color);
+
+	// Spherical Sector
+	const FVector2D Center2D(FVector2D(Shape.GetCenter()));
+	const FSphericalSector::FReal Angle = Shape.GetAngle();
+	const int32 MaxSegments = FMath::Max(4, FMath::CeilToInt(64 * Angle / 360.f));
+	const float AngleIncrement = Angle / MaxSegments;
+	const FVector2D Axis = FVector2D(Shape.GetAxis());
+	const FVector Startup = FRotator(0, -0.5f * Angle, 0).RotateVector(Shape.GetScaledAxis());
+	FCanvas* CanvasObject = Canvas->Canvas;
+
+	FVector2D LineStart = FVector2D(Startup);
+	if (!Shape.IsSphere())
+	{
+		// Draw sector start axis
+		LineItem.Draw(CanvasObject, WorldToScreen(Center2D), WorldToScreen(Center2D + LineStart));
+	}
+	// Draw sector Arc
+	for (int32 i = 1; i <= MaxSegments; i++)
+	{
+		FVector2D LineEnd = FVector2D(FRotator(0, AngleIncrement * i, 0).RotateVector(Startup));
+		LineItem.Draw(CanvasObject, WorldToScreen(Center2D + LineStart), WorldToScreen(Center2D + LineEnd));
+		LineStart = LineEnd;
+	}
+	// If sphere, close circle, else draw sector end axis
+	LineItem.Draw(CanvasObject, WorldToScreen(Center2D + LineStart), WorldToScreen(Center2D + (Shape.IsSphere() ? FVector2D(Startup) : FVector2D::ZeroVector)));
+
+	// Draw direction vector
+	LineItem.Draw(CanvasObject, WorldToScreen(Center2D), WorldToScreen(Center2D + Axis * Shape.GetRadius()));
+}
+
 // ------------------------------------------------------------------------------------------------
 
 ASpatialHashRuntimeGridInfo::ASpatialHashRuntimeGridInfo(const FObjectInitializer& ObjectInitializer)
@@ -1201,7 +1260,6 @@ void UWorldPartitionRuntimeSpatialHash::Draw2D(UCanvas* Canvas, const TArray<FWo
 	const float GridMaxScreenSize = CanvasMaxScreenSize / FilteredStreamingGrids.Num();
 	const float GridEffectiveScreenRatio = 1.f;
 	const float GridEffectiveScreenSize = FMath::Min(GridMaxScreenSize, PartitionCanvasSize.Y) - 10.f;
-	const float GridViewLoadingRangeExtentRatio = 1.5f;
 	const float GridViewMinimumSizeInCellCount = 5.f;
 	const FVector2D GridScreenExtent = FVector2D(GridEffectiveScreenSize, GridEffectiveScreenSize);
 	const FVector2D GridScreenHalfExtent = 0.5f * GridScreenExtent;
@@ -1217,17 +1275,18 @@ void UWorldPartitionRuntimeSpatialHash::Draw2D(UCanvas* Canvas, const TArray<FWo
 		// Take into consideration GShowRuntimeSpatialHashGridLevel when using CellSize
 		int32 MinGridLevel = FMath::Clamp<int32>(GShowRuntimeSpatialHashGridLevel, 0, StreamingGrid->GridLevels.Num() - 1);
 		int32 CellSize = (1 << MinGridLevel) * StreamingGrid->CellSize;
-		const float GridSideDistance = FMath::Max((2.f * StreamingGrid->GetLoadingRange() * GridViewLoadingRangeExtentRatio), CellSize * GridViewMinimumSizeInCellCount);
-		FSphere AverageSphere(ForceInit);
+		const FVector MinExtent(CellSize * GridViewMinimumSizeInCellCount);
+		FBox Region(ForceInit);
 		for (const FWorldPartitionStreamingSource& Source : Sources)
 		{
-			AverageSphere += FSphere(FVector(Source.Location.X, Source.Location.Y, 0.f), 0.5f * GridSideDistance);
+			Region += Source.CalcBounds(StreamingGrid->GetLoadingRange(), StreamingGrid->GridName, /*bCalcIn2D*/ true);
 		}
-		const FVector2D GridReferenceWorldPos = FVector2D(AverageSphere.Center);
-		const FBox Region(AverageSphere.Center - AverageSphere.W, AverageSphere.Center + AverageSphere.W);
+		Region += FBox(Region.GetCenter() - MinExtent, Region.GetCenter() + MinExtent);
+		const FVector2D GridReferenceWorldPos = FVector2D(Region.GetCenter());
+		const float RegionExtent = FVector2D(Region.GetExtent()).Size();
 		const FVector2D GridScreenOffset = GridScreenInitialOffset + ((float)GridIndex * FVector2D(GridMaxScreenSize, 0.f)) + GridScreenHalfExtent;
 		const FBox2D GridScreenBounds(GridScreenOffset - GridScreenHalfExtent, GridScreenOffset + GridScreenHalfExtent);
-		const float WorldToScreenScale = (0.5f * GridEffectiveScreenSize) / AverageSphere.W;
+		const float WorldToScreenScale = (0.5f * GridEffectiveScreenSize) / RegionExtent;
 		auto WorldToScreen = [&](const FVector2D& WorldPos) { return (WorldToScreenScale * (WorldPos - GridReferenceWorldPos)) + GridScreenOffset; };
 
 		StreamingGrid->Draw2D(Canvas, World, Sources, Region, GridScreenBounds, WorldToScreen);
diff --git a/Engine/Source/Runtime/Engine/Private/WorldPartition/WorldPartitionRuntimeSpatialHashCell.cpp b/Engine/Source/Runtime/Engine/Private/WorldPartition/WorldPartitionRuntimeSpatialHashCell.cpp
--- a/Engine/Source/Runtime/Engine/Private/WorldPartition/WorldPartitionRuntimeSpatialHashCell.cpp
+++ b/Engine/Source/Runtime/Engine/Private/WorldPartition/WorldPartitionRuntimeSpatialHashCell.cpp
@@ -38,9 +38,9 @@ void UWorldPartitionRuntimeSpatialHashCell::PostDuplicate(bool bDuplicateForPIE)
 }
 #endif
 
-bool UWorldPartitionRuntimeSpatialHashCell::CacheStreamingSourceInfo(const FWorldPartitionStreamingSource& Source) const
+bool UWorldPartitionRuntimeSpatialHashCell::CacheStreamingSourceInfo(const UWorldPartitionRuntimeCell::FStreamingSourceInfo& Info) const
 {
-	const bool bWasCacheDirtied = Super::CacheStreamingSourceInfo(Source);
+	const bool bWasCacheDirtied = Super::CacheStreamingSourceInfo(Info);
 	if (bWasCacheDirtied)
 	{
 		CachedIsBlockingSource = false;
@@ -48,12 +48,12 @@ bool UWorldPartitionRuntimeSpatialHashCell::CacheStreamingSourceInfo(const FWorl
 	}
 
 	const float AngleContribution = FMath::Clamp(GRuntimeSpatialHashCellToSourceAngleContributionToCellImportance, 0.f, 1.f);
-	const float SqrDistance = FVector::DistSquared(Source.Location, Position);
+	const float SqrDistance = FVector::DistSquared(Info.SourceShape.GetCenter(), Position);
 	float AngleFactor = 1.f;
 	if (!FMath::IsNearlyZero(AngleContribution))
 	{
-		const FVector2D SourceForward(Source.Rotation.Quaternion().GetForwardVector());
-		const FVector2D SourceToCell(Position - Source.Location);
+		const FVector2D SourceForward(Info.SourceShape.GetAxis());
+		const FVector2D SourceToCell(Position - Info.SourceShape.GetCenter());
 		const float Dot = FVector2D::DotProduct(SourceForward.GetSafeNormal(), SourceToCell.GetSafeNormal());
 		const float NormalizedAngle = FMath::Clamp(FMath::Abs(FMath::Acos(Dot) / PI), 0.f, 1.f);
 		AngleFactor = FMath::Pow(NormalizedAngle, AngleContribution);
@@ -65,7 +65,7 @@ bool UWorldPartitionRuntimeSpatialHashCell::CacheStreamingSourceInfo(const FWorl
 	CachedSourceMinSortDistance = bWasCacheDirtied ? ModulatedDistance : FMath::Min(ModulatedDistance, CachedSourceMinSortDistance);
 	
 	// Only consider blocking sources
-	if (Source.bBlockOnSlowLoading)
+	if (Info.Source.bBlockOnSlowLoading)
 	{
 		CachedIsBlockingSource = true;
 		CachedBlockingMinDistanceSquare = FMath::Min(SqrDistance, CachedBlockingMinDistanceSquare);
diff --git a/Engine/Source/Runtime/Engine/Public/WorldPartition/WorldPartitionRuntimeCell.h b/Engine/Source/Runtime/Engine/Public/WorldPartition/WorldPartitionRuntimeCell.h
--- a/Engine/Source/Runtime/Engine/Public/WorldPartition/WorldPartitionRuntimeCell.h
+++ b/Engine/Source/Runtime/Engine/Public/WorldPartition/WorldPartitionRuntimeCell.h
@@ -5,6 +5,7 @@
 #include "Templates/SubclassOf.h"
 #include "WorldPartition/DataLayer/DataLayer.h"
 #include "WorldPartition/WorldPartitionActorDescView.h"
+#include "WorldPartition/WorldPartitionStreamingSource.h"
 #include "ProfilingDebugging/ProfilingHelpers.h"
 #include "Algo/AnyOf.h"
 #include "WorldPartitionRuntimeCell.generated.h"
@@ -112,6 +113,16 @@ class UWorldPartitionRuntimeCell : public UObject
 {
 	GENERATED_UCLASS_BODY()
 
+	struct FStreamingSourceInfo
+	{
+		FStreamingSourceInfo(const FWorldPartitionStreamingSource& InSource, const FSphericalSector& InSourceShape)
+			: Source(InSource)
+			, SourceShape(InSourceShape)
+		{}
+		const FWorldPartitionStreamingSource& Source;
+		const FSphericalSector& SourceShape;
+	};
+
 	virtual void Load() const PURE_VIRTUAL(UWorldPartitionRuntimeCell::Load,);
 	virtual void Unload() const PURE_VIRTUAL(UWorldPartitionRuntimeCell::Unload,);
 	virtual void Activate() const PURE_VIRTUAL(UWorldPartitionRuntimeCell::Activate,);
@@ -131,7 +142,7 @@ class UWorldPartitionRuntimeCell : public UObject
 	virtual int32 SortCompare(const UWorldPartitionRuntimeCell* Other) const;
 	virtual FName GetGridName() const { return GridName; }
 	/** Caches information on streaming source that will be used later on to sort cell. Returns true if cache was reset, else returns false. */
-	virtual bool CacheStreamingSourceInfo(const struct FWorldPartitionStreamingSource& Source) const;
+	virtual bool CacheStreamingSourceInfo(const UWorldPartitionRuntimeCell::FStreamingSourceInfo& Info) const;
 
 	static void DirtyStreamingSourceCacheEpoch() { ++UWorldPartitionRuntimeCell::StreamingSourceCacheEpoch; }
 
diff --git a/Engine/Source/Runtime/Engine/Public/WorldPartition/WorldPartitionRuntimeHash.h b/Engine/Source/Runtime/Engine/Public/WorldPartition/WorldPartitionRuntimeHash.h
--- a/Engine/Source/Runtime/Engine/Public/WorldPartition/WorldPartitionRuntimeHash.h
+++ b/Engine/Source/Runtime/Engine/Public/WorldPartition/WorldPartitionRuntimeHash.h
@@ -8,14 +8,13 @@
 #include "WorldPartitionActorDescView.h"
 #include "WorldPartition/ActorDescList.h"
 #include "WorldPartition/WorldPartitionHandle.h"
+#include "WorldPartition/WorldPartitionRuntimeCell.h"
 #include "WorldPartition/WorldPartitionActorDescViewProxy.h"
 #if WITH_EDITOR
 #include "CookPackageSplitter.h"
 #endif
 #include "WorldPartitionRuntimeHash.generated.h"
 
-class UWorldPartitionRuntimeCell;
-
 UENUM()
 enum class EWorldPartitionStreamingPerformance
 {
@@ -54,7 +53,7 @@ class ENGINE_API UWorldPartitionRuntimeHash : public UObject
 	class FStreamingSourceCells
 	{
 	public:
-		void AddCell(const UWorldPartitionRuntimeCell* InCell, const FWorldPartitionStreamingSource& InSource);
+		void AddCell(const UWorldPartitionRuntimeCell* Cell, const UWorldPartitionRuntimeCell::FStreamingSourceInfo& Info);
 		void Reset() { Cells.Reset(); }
 		int32 Num() const { return Cells.Num(); }
 		TSet<const UWorldPartitionRuntimeCell*>& GetCells() { return Cells; }
diff --git a/Engine/Source/Runtime/Engine/Public/WorldPartition/WorldPartitionRuntimeSpatialHash.h b/Engine/Source/Runtime/Engine/Public/WorldPartition/WorldPartitionRuntimeSpatialHash.h
--- a/Engine/Source/Runtime/Engine/Public/WorldPartition/WorldPartitionRuntimeSpatialHash.h
+++ b/Engine/Source/Runtime/Engine/Public/WorldPartition/WorldPartitionRuntimeSpatialHash.h
@@ -138,6 +138,8 @@ struct FSpatialHashStreamingGrid
 	void Draw3D(UWorld* World, const TArray<FWorldPartitionStreamingSource>& Sources, const FTransform& Transform) const;
 
 private:
+	void DrawStreamingSource2D(UCanvas* Canvas, const FSphericalSector& Shape, TFunctionRef<FVector2D(const FVector2D&)> WorldToScreen, const FColor& Color) const;
+	void DrawStreamingSource3D(UWorld* World, const FSphericalSector& Shape, const FTransform& Transform, const FColor& Color) const;
 	void GetFilteredCellsForDebugDraw(const FSpatialHashStreamingGridLayerCell* LayerCell, const UDataLayerSubsystem* DataLayerSubsystem, TArray<const UWorldPartitionRuntimeCell*>& FilteredCells) const;
 	const FSquare2DGridHelper& GetGridHelper() const;
 	mutable FSquare2DGridHelper* GridHelper;
diff --git a/Engine/Source/Runtime/Engine/Public/WorldPartition/WorldPartitionRuntimeSpatialHashCell.h b/Engine/Source/Runtime/Engine/Public/WorldPartition/WorldPartitionRuntimeSpatialHashCell.h
--- a/Engine/Source/Runtime/Engine/Public/WorldPartition/WorldPartitionRuntimeSpatialHashCell.h
+++ b/Engine/Source/Runtime/Engine/Public/WorldPartition/WorldPartitionRuntimeSpatialHashCell.h
@@ -15,7 +15,7 @@ class UWorldPartitionRuntimeSpatialHashCell : public UWorldPartitionRuntimeCell
 #if WITH_EDITOR
 	virtual void PostDuplicate(bool bDuplicateForPIE) override;
 #endif
-	virtual bool CacheStreamingSourceInfo(const FWorldPartitionStreamingSource& Source) const override;
+	virtual bool CacheStreamingSourceInfo(const UWorldPartitionRuntimeCell::FStreamingSourceInfo& Info) const override;
 	virtual int32 SortCompare(const UWorldPartitionRuntimeCell* InOther) const override;
 
 	UPROPERTY()
diff --git a/Engine/Source/Runtime/Engine/Public/WorldPartition/WorldPartitionStreamingSource.h b/Engine/Source/Runtime/Engine/Public/WorldPartition/WorldPartitionStreamingSource.h
--- a/Engine/Source/Runtime/Engine/Public/WorldPartition/WorldPartitionStreamingSource.h
+++ b/Engine/Source/Runtime/Engine/Public/WorldPartition/WorldPartitionStreamingSource.h
@@ -6,6 +6,185 @@
 #include "Math/RandomStream.h"
 #include "WorldPartitionStreamingSource.generated.h"
 
+/** See https://en.wikipedia.org/wiki/Spherical_sector. */
+class FSphericalSector
+{
+public:
+	using FReal = FVector::FReal;
+
+	/** Creates and initializes a new spherical sector. */
+	FSphericalSector(EForceInit)
+		: Center(ForceInit)
+		, Radius(0.0f)
+		, Axis(ForceInit)
+	{
+		SetAsSphere();
+	}
+
+	/** Creates and initializes a spherical sector using given parameters. */
+	FSphericalSector(FVector InCenter, FReal InRadius, FVector InAxis = FVector::ForwardVector, FReal InAngle = 0)
+		: Center(InCenter)
+		, Radius(InRadius)
+	{
+		SetAngle(InAngle);
+		SetAxis(InAxis);
+	}
+
+	void SetCenter(const FVector& InCenter) { Center = InCenter; }
+	const FVector& GetCenter() const { return Center; }
+
+	void SetRadius(FReal InRadius) { Radius = InRadius; }
+	FReal GetRadius() const { return Radius; }
+
+	void SetAngle(FReal InAngle) { Angle = (InAngle <= 0.0f || InAngle > 360.0f) ? 360.0f : InAngle; }
+	FReal GetAngle() const { return Angle; }
+
+	void SetAxis(const FVector& InAxis) { Axis = InAxis.GetSafeNormal(); }
+	FVector GetAxis() const { return Axis; }
+	FVector GetScaledAxis() const { return Axis * Radius; }
+
+	void SetAsSphere() { SetAngle(360.0f); }
+	bool IsSphere() const { return FMath::IsNearlyEqual(Angle, 360.0f); }
+
+	bool IsNearlyZero() const { return FMath::IsNearlyZero(Radius) || Axis.IsNearlyZero() || FMath::IsNearlyZero(Angle); }
+	bool IsValid() const { return !IsNearlyZero(); }
+
+	FBox CalcBounds() const
+	{
+		const FVector Offset(Radius);
+		return FBox(Center - Offset, Center + Offset);
+	}
+
+	/** Get result of Transforming spherical sector with transform. */
+	FSphericalSector TransformBy(const FTransform& M) const
+	{
+		FSphericalSector Result(M.TransformPosition(Center), M.GetMaximumAxisScale() * Radius, M.TransformVector(Axis), Angle);
+		return Result;
+	}
+
+	/** Helper method that builds a list of debug display segments */
+	TArray<TPair<FVector, FVector>> BuildDebugMesh() const
+	{
+		TArray<TPair<FVector, FVector>> Segments;
+		if (!IsValid())
+		{
+			return Segments;
+		}
+
+		const int32 SegmentCount = FMath::Max(4, FMath::CeilToInt(64 * Angle / 360.f));
+		const float AngleStep = Angle / float(SegmentCount);
+		const FRotator ShapeRotation = FRotationMatrix::MakeFromX(Axis).Rotator();
+		const FVector ScaledAxis = FVector::ForwardVector * Radius;
+
+		FVector LastRollEndVector;
+		for (int i = 0; i <= 64; ++i)
+		{
+			const float Roll = 360.0f * i / 64.0f;
+			const FTransform Transform(FRotator(0, 0, Roll) + ShapeRotation, Center);
+			FVector SegmentStart = Transform.TransformPosition(FRotator(0, -0.5f * Angle, 0).RotateVector(ScaledAxis));
+			Segments.Emplace(Center, SegmentStart);
+			for (int32 j = 1; j <= SegmentCount; j++)
+			{
+				FVector SegmentEnd = Transform.TransformPosition(FRotator(0, -0.5f * Angle + (AngleStep * j), 0).RotateVector(ScaledAxis));
+				Segments.Emplace(SegmentStart, SegmentEnd);
+				SegmentStart = SegmentEnd;
+			}
+			Segments.Emplace(Center, SegmentStart);
+			if (i > 0)
+			{
+				Segments.Emplace(SegmentStart, LastRollEndVector);
+			}
+			LastRollEndVector = SegmentStart;
+		}
+		return Segments;
+	}
+
+private:
+	/** Sphere center point. */
+	FVector Center;
+
+	/** Sphere radius. */
+	FReal Radius;
+
+	/** Sector axis (direction). */
+	FVector Axis;
+
+	/** Optional sector angle in degree (360 = regular sphere). */
+	FReal Angle;
+};
+
+USTRUCT(BlueprintType)
+struct FStreamingSourceShape
+{
+	GENERATED_BODY()
+
+	FStreamingSourceShape()
+	: bUseGridLoadingRange(true)
+	, Radius(10000.0f)
+	, bIsSector(false)
+	, SectorAngle(360.0f)
+	, Location(ForceInitToZero)
+	, Rotation(ForceInitToZero)
+	{}
+
+	/* If True, streaming source shape radius is bound to loading range radius. */
+	UPROPERTY(BlueprintReadWrite, EditAnywhere, Category = Streaming)
+	bool bUseGridLoadingRange;
+
+	/* Custom streaming source shape radius (not used if bUseGridLoadingRange is True). */
+	UPROPERTY(BlueprintReadWrite, EditAnywhere, Category = Streaming, meta = (EditCondition = "!bUseGridLoadingRange"))
+	float Radius;
+
+	/* Whether the source shape is a spherical sector instead of a regular sphere source. */
+	UPROPERTY(BlueprintReadWrite, EditAnywhere, Category = Streaming)
+	bool bIsSector;
+
+	/* Shape's spherical sector angle in degree (not used if bIsSector is False). */
+	UPROPERTY(BlueprintReadWrite, EditAnywhere, Category = Streaming, meta = (EditCondition = "bIsSector", ClampMin = 0, ClampMax = 360))
+	float SectorAngle;
+
+	/* Streaming source shape location (local to streaming source). */
+	UPROPERTY(BlueprintReadWrite, EditAnywhere, Category = Streaming)
+	FVector Location;
+
+	/* Streaming source shape rotation (local to streaming source). */
+	UPROPERTY(BlueprintReadWrite, EditAnywhere, Category = Streaming)
+	FRotator Rotation;
+};
+
+/** Helper class used to iterate over streaming source shapes. */
+class FStreamingSourceShapeHelper
+{
+public:
+
+	FORCEINLINE static void ForEachShape(float InGridLoadingRange, float InDefaultRadius, bool bInProjectIn2D, const FVector& InLocation, const FRotator& InRotation, const TArray<FStreamingSourceShape>& InShapes, TFunctionRef<void(const FSphericalSector&)> InOperation)
+	{
+		const FTransform Transform(bInProjectIn2D ? FRotator(0, InRotation.Yaw, 0) : InRotation, InLocation);
+		if (InShapes.IsEmpty())
+		{
+			FSphericalSector LocalShape(FVector::ZeroVector, InDefaultRadius);
+			if (LocalShape.IsValid())
+			{
+				InOperation(LocalShape.TransformBy(Transform));
+			}
+		}
+		else
+		{
+			for (const FStreamingSourceShape& Shape : InShapes)
+			{
+				const FVector::FReal ShapeRadius = Shape.bUseGridLoadingRange ? InGridLoadingRange : Shape.Radius;
+				const FVector::FReal ShapeAngle = Shape.bIsSector ? Shape.SectorAngle : 360.0f;
+				const FVector ShapeAxis = bInProjectIn2D ? FRotator(0, Shape.Rotation.Yaw, 0).Vector() : Shape.Rotation.Vector();
+				FSphericalSector LocalShape(bInProjectIn2D ? FVector(Shape.Location.X, Shape.Location.Y, 0) : Shape.Location, ShapeRadius, ShapeAxis, ShapeAngle);
+				if (LocalShape.IsValid())
+				{
+					InOperation(LocalShape.TransformBy(Transform));
+				}
+			}
+		}
+	}
+};
+
 /**
  * Streaming Source Target State
  */
@@ -27,17 +206,21 @@ struct FWorldPartitionStreamingQuerySource
 	FWorldPartitionStreamingQuerySource()
 		: Location(FVector::ZeroVector)
 		, Radius(0.f)
-		, bUseLoadingRangeRadius(true)
+		, bUseGridLoadingRange(true)
 		, bDataLayersOnly(false)
 		, bSpatialQuery(true)
+		, Rotation(ForceInitToZero)
+		, TargetGrid(NAME_None)
 	{}
 
 	FWorldPartitionStreamingQuerySource(const FVector& InLocation)
 		: Location(InLocation)
 		, Radius(0.f)
-		, bUseLoadingRangeRadius(true)
+		, bUseGridLoadingRange(true)
 		, bDataLayersOnly(false)
 		, bSpatialQuery(true)
+		, Rotation(ForceInitToZero)
+		, TargetGrid(NAME_None)
 	{}
 
 	/* Location to query. (not used if bSpatialQuery is false) */
@@ -50,7 +233,7 @@ struct FWorldPartitionStreamingQuerySource
 
 	/* If True, Instead of providing a query radius, query can be bound to loading range radius. */
 	UPROPERTY(BlueprintReadWrite, EditAnywhere, Category = "Query")
-	bool bUseLoadingRangeRadius;
+	bool bUseGridLoadingRange;
 
 	/* Optional list of data layers to specialize the query. If empty only non data layer cells will be returned by the query. */
 	UPROPERTY(BlueprintReadWrite, EditAnywhere, Category = "Query")
@@ -63,6 +246,20 @@ struct FWorldPartitionStreamingQuerySource
 	/* If False, Location/Radius will not be used to find the cells. Only AlwaysLoaded cells will be returned by the query. */
 	UPROPERTY(BlueprintReadWrite, EditAnywhere, Category = "Query")
 	bool bSpatialQuery;
+
+	/* Reserved settings used by UWorldPartitionStreamingSourceComponent::IsStreamingCompleted. */
+	FRotator Rotation;
+	FName TargetGrid;
+	TArray<FStreamingSourceShape> Shapes;
+
+	/** Helper method that iterates over all shapes. If none is provided, it will still pass a sphere shape using Radius or grid's loading range (see bUseGridLoadingRange). */
+	FORCEINLINE void ForEachShape(float InGridLoadingRange, FName InGridName, bool bInProjectIn2D, TFunctionRef<void(const FSphericalSector&)> InOperation) const
+	{
+		if (bSpatialQuery && (TargetGrid.IsNone() || TargetGrid == InGridName))
+		{
+			FStreamingSourceShapeHelper::ForEachShape(InGridLoadingRange, bUseGridLoadingRange ? InGridLoadingRange : Radius, bInProjectIn2D, Location, Rotation, Shapes, InOperation);
+		}
+	}
 };
 
 /**
@@ -87,6 +284,7 @@ struct ENGINE_API FWorldPartitionStreamingSource
 		: bBlockOnSlowLoading(false)
 		, Priority(EStreamingSourcePriority::Default)
 		, Velocity(0.f)
+		, TargetGrid(NAME_None)
 	{}
 
 	FWorldPartitionStreamingSource(FName InName, const FVector& InLocation, const FRotator& InRotation, EStreamingSourceTargetState InTargetState, bool bInBlockOnSlowLoading, EStreamingSourcePriority InPriority = EStreamingSourcePriority::Default, float InVelocity = 0.f)
@@ -97,18 +295,58 @@ struct ENGINE_API FWorldPartitionStreamingSource
 		, bBlockOnSlowLoading(bInBlockOnSlowLoading)
 		, Priority(InPriority)
 		, Velocity(InVelocity)
+		, TargetGrid(NAME_None)
 	{
 	}
 
 	FColor GetDebugColor() const { return FColor::MakeRedToGreenColorFromScalar(FRandomStream(Name).GetFraction()); }
 
+	/** Source unique name. */
 	FName Name;
+
+	/** Source location. */
 	FVector Location;
+	
+	/** Source orientation (can impact streaming cell prioritization). */
 	FRotator Rotation;
+
+	/** Target streaming state. */
 	EStreamingSourceTargetState TargetState;
+
+	/** Whether this source will be considered when world partition detects slow loading and waits for cell streaming to complete. */
 	bool bBlockOnSlowLoading;
+
+	/** Streaming source priority. */
 	EStreamingSourcePriority Priority;
-	float Velocity; // Computed automatically
+
+	/** Source velocity (computed automatically). */
+	float Velocity;
+
+	/** When set, will only affect streaming on the provided target runtime streaming grid. When none is provided, applies to all runtime streaming grid. */
+	FName TargetGrid;
+
+	/** Source internal shapes. When none are provided, a sphere is automatically used. It's radius is equal to grid's loading range and center equals source's location. */
+	TArray<FStreamingSourceShape> Shapes;
+
+	/** Returns a box encapsulating all shapes. */
+	FORCEINLINE FBox CalcBounds(float InGridLoadingRange, FName InGridName, bool bCalcIn2D = false) const
+	{
+		FBox OutBounds(ForceInit);
+		ForEachShape(InGridLoadingRange, InGridName, bCalcIn2D, [&OutBounds](const FSphericalSector& Sector)
+		{
+			OutBounds += Sector.CalcBounds();
+		});
+		return OutBounds;
+	}
+
+	/** Helper method that iterates over all shapes. If none is provided, it will still pass a sphere shape using grid's loading range. */
+	FORCEINLINE void ForEachShape(float InGridLoadingRange, FName InGridName, bool bInProjectIn2D, TFunctionRef<void(const FSphericalSector&)> InOperation) const
+	{
+		if (TargetGrid.IsNone() || TargetGrid == InGridName)
+		{
+			FStreamingSourceShapeHelper::ForEachShape(InGridLoadingRange, InGridLoadingRange, bInProjectIn2D, Location, Rotation, Shapes, InOperation);
+		}
+	}
 };
 
 /**
