Merge LazyObjectPtr & SoftObjectPath PIE fixup archives
Archive no longer uses IsSaving in order to use faster binary serialization & avoid delta serialization.

diff --git a/Engine/Source/Editor/UnrealEd/Classes/Editor/FixupLazyObjectPtrForPIEArchive.h b/Engine/Source/Editor/UnrealEd/Classes/Editor/FixupLazyObjectPtrForPIEArchive.h
deleted file mode 100644
diff --git a/Engine/Source/Runtime/CoreUObject/Public/UObject/LazyObjectPtr.h b/Engine/Source/Runtime/CoreUObject/Public/UObject/LazyObjectPtr.h
--- a/Engine/Source/Runtime/CoreUObject/Public/UObject/LazyObjectPtr.h
+++ b/Engine/Source/Runtime/CoreUObject/Public/UObject/LazyObjectPtr.h
@@ -149,6 +149,12 @@ public:
 	{
 		TPersistentObjectPtr<FUniqueObjectGuid>::operator=(InObjectID);
 	}
+
+	/** Fixes up this FLazyObjectPtr to target the right UID as set in PIEGuidMap */
+	FORCEINLINE void FixupForPIE(int32 PIEInstance)
+	{
+		*this = GetUniqueID().FixupForPIE(PIEInstance);
+	}
 	
 	/** Called by UObject::Serialize so that we can save / load the Guid possibly associated with an object */
 	COREUOBJECT_API static void PossiblySerializeObjectGuid(UObject* Object, FStructuredArchive::FRecord Record);
diff --git a/Engine/Source/Runtime/Engine/Private/Level.cpp b/Engine/Source/Runtime/Engine/Private/Level.cpp
--- a/Engine/Source/Runtime/Engine/Private/Level.cpp
+++ b/Engine/Source/Runtime/Engine/Private/Level.cpp
@@ -59,6 +59,7 @@ Level.cpp: Level-related functions
 #include "AssetRegistryModule.h"
 #include "IAssetRegistry.h"
 #include "AssetData.h"
+#include "PieFixupSerializer.h"
 #endif
 #include "WorldPartition/WorldPartition.h"
 #include "Engine/LevelStreaming.h"
@@ -2761,51 +2762,7 @@ bool ULevel::CanEditChange(const FProperty* PropertyThatWillChange) const
 
 void ULevel::FixupForPIE(int32 InPIEInstanceID, TFunctionRef<void(int32, FSoftObjectPath&)> InCustomFixupFunction)
 {
-	struct FSoftPathPIEFixupSerializer : public FArchiveUObject
-	{
-		FSoftPathPIEFixupSerializer(UObject* InRoot, int32 InPIEInstanceID, TFunctionRef<void(int32, FSoftObjectPath&)> InCustomFixupFunction)
-			: Root(InRoot)
-			, PIEInstanceID(InPIEInstanceID)
-			, CustomFixupFunction(InCustomFixupFunction)
-		{
-			this->SetIsSaving(true);
-			this->ArShouldSkipBulkData = true;
-		}
-
-		virtual bool ShouldSkipProperty(const FProperty* InProperty) const override
-		{
-			return InProperty->IsA<FMulticastDelegateProperty>() || FArchiveUObject::ShouldSkipProperty(InProperty);
-		}
-
-		virtual FArchive& operator<<(UObject*& Object) override
-		{
-			if (Object && Object->IsIn(Root) && !VisitedObjects.Contains(Object))
-			{
-				VisitedObjects.Add(Object);
-
-				// Skip instanced static mesh component as their impact on serialization is enormous and they don't contain lazy ptrs.
-				if (!Cast<UInstancedStaticMeshComponent>(Object))
-				{
-					Object->Serialize(*this);
-				}
-			}
-			return *this;
-		}
-
-		FArchive& operator<<(FSoftObjectPath& Value)
-		{
-			Value.FixupForPIE(PIEInstanceID, CustomFixupFunction);
-			return *this;
-		}
-
-		UObject* Root;
-		TSet<UObject*> VisitedObjects;
-
-		int32 PIEInstanceID;
-		TFunctionRef<void(int32, FSoftObjectPath&)> CustomFixupFunction;
-	};
-
-	FSoftPathPIEFixupSerializer FixupSerializer(this, InPIEInstanceID, InCustomFixupFunction);
+	FPIEFixupSerializer FixupSerializer(this, InPIEInstanceID, InCustomFixupFunction);
 	Serialize(FixupSerializer);
 }
 
diff --git a/Engine/Source/Runtime/Engine/Private/PieFixupSerializer.cpp b/Engine/Source/Runtime/Engine/Private/PieFixupSerializer.cpp
new file mode 100644
--- /dev/null
+++ b/Engine/Source/Runtime/Engine/Private/PieFixupSerializer.cpp
@@ -0,0 +1,60 @@
+#include "PIEFixupSerializer.h"
+
+#include "Components/InstancedStaticMeshComponent.h"
+
+class FMulticastDelegateProperty;
+
+FPIEFixupSerializer::FPIEFixupSerializer(UObject* InRoot, int32 InPIEInstanceID)
+	: SoftObjectPathFixupFunction([](int32, FSoftObjectPath&) {})
+	, Root(InRoot)
+	, PIEInstanceID(InPIEInstanceID)
+{
+	this->ArShouldSkipBulkData = true;
+}
+
+FPIEFixupSerializer::FPIEFixupSerializer(UObject* InRoot, int32 InPIEInstanceID, TFunctionRef<void(int32, FSoftObjectPath&)> InSoftObjectPathFixupFunction)
+	: SoftObjectPathFixupFunction(InSoftObjectPathFixupFunction)
+	, Root(InRoot)
+	, PIEInstanceID(InPIEInstanceID)
+{
+	this->ArShouldSkipBulkData = true;
+}
+
+bool FPIEFixupSerializer::ShouldSkipProperty(const FProperty* InProperty) const
+{
+	return InProperty->IsA<FMulticastDelegateProperty>() || FArchiveUObject::ShouldSkipProperty(InProperty);
+}
+
+FArchive& FPIEFixupSerializer::operator<<(UObject*& Object)
+{
+	if (Object && (Object == Root ||Object->IsIn(Root)) && !VisitedObjects.Contains(Object))
+	{
+		VisitedObjects.Add(Object);
+
+		// Skip instanced static mesh component as their impact on serialization is enormous and they don't contain lazy ptrs.
+		if (!Cast<UInstancedStaticMeshComponent>(Object))
+		{
+			Object->Serialize(*this);
+		}
+	}
+	return *this;
+}
+
+FArchive& FPIEFixupSerializer::operator<<(FSoftObjectPath& Value)
+{
+	Value.FixupForPIE(PIEInstanceID, SoftObjectPathFixupFunction);
+	return *this;
+}
+
+FArchive& FPIEFixupSerializer::operator<<(FLazyObjectPtr& Value)
+{
+	Value.FixupForPIE(PIEInstanceID);
+	return *this;
+}
+
+FArchive& FPIEFixupSerializer::operator<<(FSoftObjectPtr& Value)
+{
+	// Forward the serialization to the FSoftObjectPath overload so it can be fixed up
+	*this << Value.GetUniqueID();
+	return *this;
+}
\ No newline at end of file
diff --git a/Engine/Source/Runtime/Engine/Private/PieFixupSerializer.h b/Engine/Source/Runtime/Engine/Private/PieFixupSerializer.h
new file mode 100644
--- /dev/null
+++ b/Engine/Source/Runtime/Engine/Private/PieFixupSerializer.h
@@ -0,0 +1,23 @@
+#pragma once
+
+#include "Serialization/ArchiveUObject.h"
+
+class FPIEFixupSerializer : public FArchiveUObject
+{
+public:
+	FPIEFixupSerializer(UObject* InRoot, int32 InPIEInstanceID);
+	FPIEFixupSerializer(UObject* InRoot, int32 InPIEInstanceID, TFunctionRef<void(int32, FSoftObjectPath&)> InSoftObjectPathFixupFunction);
+
+	bool ShouldSkipProperty(const FProperty* InProperty) const override;
+
+	FArchive& operator<<(UObject*& Object) override;
+	FArchive& operator<<(FSoftObjectPath& Value) override;
+	FArchive& operator<<(FLazyObjectPtr& LazyObjectPtr) override;
+	FArchive& operator<<(FSoftObjectPtr& Value) override;
+
+private:
+	TFunctionRef<void(int32, FSoftObjectPath&)> SoftObjectPathFixupFunction;
+	TSet<UObject*> VisitedObjects;
+	UObject* Root;
+	int32 PIEInstanceID;
+};
\ No newline at end of file
diff --git a/Engine/Source/Runtime/Engine/Private/World.cpp b/Engine/Source/Runtime/Engine/Private/World.cpp
--- a/Engine/Source/Runtime/Engine/Private/World.cpp
+++ b/Engine/Source/Runtime/Engine/Private/World.cpp
@@ -106,9 +106,9 @@
 	#include "ObjectTools.h"
 	#include "Engine/LODActor.h"
 	#include "PIEPreviewDeviceProfileSelectorModule.h"
-	#include "Editor/FixupLazyObjectPtrForPIEArchive.h"
 	#include "AssetCompilingManager.h"
 	#include "WorldPartition/DataLayer/WorldDataLayers.h"
+	#include "PieFixupSerializer.h"
 #endif
 
 
@@ -3394,8 +3394,8 @@ UWorld* UWorld::DuplicateWorldForPIE(const FString& PackageName, UWorld* OwningW
 
 	{
 		// The owning world may contain lazy pointers to actors in the sub-level we just duplicated so make sure they are fixed up with the PIE GUIDs
-		FFixupLazyObjectPtrForPIEArchive FixupLazyPointersAr;
-		FixupLazyPointersAr << OwningWorld;
+		FPIEFixupSerializer FixupSerializer(OwningWorld, PIEInstanceID);
+		FixupSerializer << OwningWorld;
 	}
 
 
diff --git a/Engine/Source/Runtime/Engine/Private/WorldPartition/WorldPartitionLevelStreamingDynamic.cpp b/Engine/Source/Runtime/Engine/Private/WorldPartition/WorldPartitionLevelStreamingDynamic.cpp
--- a/Engine/Source/Runtime/Engine/Private/WorldPartition/WorldPartitionLevelStreamingDynamic.cpp
+++ b/Engine/Source/Runtime/Engine/Private/WorldPartition/WorldPartitionLevelStreamingDynamic.cpp
@@ -14,7 +14,6 @@
 #include "ContentStreaming.h"
 #include "GameFramework/WorldSettings.h"
 #include "UnrealEngine.h"
-#include "Editor/FixupLazyObjectPtrForPIEArchive.h"
 #include "WorldPartition/WorldPartitionSubsystem.h"
 #include "WorldPartition/WorldPartitionStreamingPolicy.h"
 #include "WorldPartition/WorldPartitionLevelStreamingPolicy.h"
@@ -329,16 +328,7 @@ void UWorldPartitionLevelStreamingDynamic::FinalizeRuntimeLevel()
 		check(PIEInstanceID != INDEX_NONE);
 
 		{
-			TRACE_CPUPROFILER_EVENT_SCOPE(FixupLazyObjectPtr);
-			// PIE Fixup LazyObjectPtrs
-			FTemporaryPlayInEditorIDOverride SetPlayInEditorID(PIEInstanceID);
-			FFixupLazyObjectPtrForPIEArchive FixupLazyPointersAr;
-			FixupLazyPointersAr << RuntimeLevel;
-		}
-
-		{
-			TRACE_CPUPROFILER_EVENT_SCOPE(FixupSoftObjectPath);
-			// PIE Fixup SoftObjectPaths
+			TRACE_CPUPROFILER_EVENT_SCOPE(FixupIDsForPIE);
 			RuntimeLevel->FixupForPIE(PIEInstanceID, [&](int32 InPIEInstanceID, FSoftObjectPath& ObjectPath)
 			{
 				// Remap Runtime Level's SoftObjectPath before each PIE Fixup to avoid doing 2 passes of serialization
