Fix missing Copyright boilerplate
Fix cases in includes for linux compilation

diff --git a/Engine/Source/Runtime/Engine/Private/PieFixupSerializer.cpp b/Engine/Source/Runtime/Engine/Private/PieFixupSerializer.cpp
--- a/Engine/Source/Runtime/Engine/Private/PieFixupSerializer.cpp
+++ b/Engine/Source/Runtime/Engine/Private/PieFixupSerializer.cpp
@@ -1,4 +1,6 @@
-#include "PIEFixupSerializer.h"
+// Copyright Epic Games, Inc. All Rights Reserved.
+
+#include "PieFixupSerializer.h"
 
 #include "Components/InstancedStaticMeshComponent.h"
 
diff --git a/Engine/Source/Runtime/Engine/Private/PieFixupSerializer.h b/Engine/Source/Runtime/Engine/Private/PieFixupSerializer.h
--- a/Engine/Source/Runtime/Engine/Private/PieFixupSerializer.h
+++ b/Engine/Source/Runtime/Engine/Private/PieFixupSerializer.h
@@ -1,3 +1,5 @@
+// Copyright Epic Games, Inc. All Rights Reserved.
+
 #pragma once
 
 #include "Serialization/ArchiveUObject.h"
