Fixed Linux build (ambigous call).

diff --git a/Engine/Source/Runtime/Engine/Public/WorldPartition/WorldPartitionStreamingSource.h b/Engine/Source/Runtime/Engine/Public/WorldPartition/WorldPartitionStreamingSource.h
--- a/Engine/Source/Runtime/Engine/Public/WorldPartition/WorldPartitionStreamingSource.h
+++ b/Engine/Source/Runtime/Engine/Public/WorldPartition/WorldPartitionStreamingSource.h
@@ -44,7 +44,7 @@ public:
 	FVector GetScaledAxis() const { return Axis * Radius; }
 
 	void SetAsSphere() { SetAngle(360.0f); }
-	bool IsSphere() const { return FMath::IsNearlyEqual(Angle, 360.0f); }
+	bool IsSphere() const { return FMath::IsNearlyEqual(Angle, (FReal)360.0); }
 
 	bool IsNearlyZero() const { return FMath::IsNearlyZero(Radius) || Axis.IsNearlyZero() || FMath::IsNearlyZero(Angle); }
 	bool IsValid() const { return !IsNearlyZero(); }
