Added the ability to set a minimum LOD for LOD Sync Components

diff --git a/Engine/Source/Runtime/Engine/Classes/Components/LODSyncComponent.h b/Engine/Source/Runtime/Engine/Classes/Components/LODSyncComponent.h
--- a/Engine/Source/Runtime/Engine/Classes/Components/LODSyncComponent.h
+++ b/Engine/Source/Runtime/Engine/Classes/Components/LODSyncComponent.h
@@ -73,6 +73,10 @@ class ENGINE_API ULODSyncComponent : public UActorComponent
 	UPROPERTY(EditAnywhere, BlueprintReadWrite, Category = LOD)
 	int32 ForcedLOD = -1;
 
+	// Minimum LOD to use when syncing components
+	UPROPERTY(EditAnywhere, BlueprintReadWrite, Category = LOD)
+	int32 MinLOD = 0;
+
 	/** 
 	 *	Array of components whose LOD may drive or be driven by this component.
 	 *  Components that are flagged as 'Drive' are treated as being in priority order, with the last component having highest priority. The highest priority
diff --git a/Engine/Source/Runtime/Engine/Private/Components/LODSyncComponent.cpp b/Engine/Source/Runtime/Engine/Private/Components/LODSyncComponent.cpp
--- a/Engine/Source/Runtime/Engine/Private/Components/LODSyncComponent.cpp
+++ b/Engine/Source/Runtime/Engine/Private/Components/LODSyncComponent.cpp
@@ -242,7 +242,7 @@ void ULODSyncComponent::UpdateLOD()
 		if (bHaveValidSetting)
 		{
 			// ensure current WorkingLOD is with in the range
-			HighestPriLOD = FMath::Clamp(HighestPriLOD, 0, CurrentNumLODs - 1);
+			HighestPriLOD = FMath::Clamp(HighestPriLOD, MinLOD, CurrentNumLODs - 1);
 			UE_LOG(LogLODSync, Verbose, TEXT("LOD Sync : Current LOD (%d)"), HighestPriLOD);
 			for (UPrimitiveComponent* Component : SubComponents)
 			{
