Automatically set support for ray tracing far fields on the HLOD components if the underlying HLOD Layer enables it and if it's not an instancing layer.

diff --git a/Engine/Source/Runtime/Engine/Private/WorldPartition/HLOD/HLODActor.cpp b/Engine/Source/Runtime/Engine/Private/WorldPartition/HLOD/HLODActor.cpp
--- a/Engine/Source/Runtime/Engine/Private/WorldPartition/HLOD/HLODActor.cpp
+++ b/Engine/Source/Runtime/Engine/Private/WorldPartition/HLOD/HLODActor.cpp
@@ -18,13 +18,6 @@
 #include "Engine/TextureStreamingTypes.h"
 #endif
 
-static int32 GRayTracingFarFieldPostLoad = 0;
-static FAutoConsoleVariableRef CVarRaytracingFarFieldPostLoad(
-	TEXT("wp.Render.RaytracingFarFieldPostLoad"),
-	GRayTracingFarFieldPostLoad,
-	TEXT("Enable/Disable HLOD1 ray tracing.\n"),
-	ECVF_ReadOnly);
-
 AWorldPartitionHLOD::AWorldPartitionHLOD(const FObjectInitializer& ObjectInitializer)
 	: Super(ObjectInitializer)
 {
@@ -78,28 +71,6 @@ void AWorldPartitionHLOD::Serialize(FArchive& Ar)
 void AWorldPartitionHLOD::RerunConstructionScripts()
 {}
 
-void AWorldPartitionHLOD::PostLoad()
-{
-	Super::PostLoad();
-
-#if WITH_EDITOR
-	// Enable / Disable ray tracing for hlod1
-	if (GRayTracingFarFieldPostLoad > 0 && GetSubActorsHLODLayer()->GetLayerType() != EHLODLayerType::Instancing)
-	{
-		UStaticMeshComponent* const StaticMeshComp = Cast<UStaticMeshComponent>(GetHLODComponent());
-		if (StaticMeshComp)
-		{
-			StaticMeshComp->ConditionalPostLoad();
-			if (StaticMeshComp->GetStaticMesh())
-			{
-				StaticMeshComp->bRayTracingFarField = true;
-				StaticMeshComp->GetStaticMesh()->bSupportRayTracing = true;
-			}
-		}
-	}
-#endif
-}
-
 #if WITH_EDITOR
 
 bool AWorldPartitionHLOD::IsHiddenEd() const
diff --git a/Engine/Source/Runtime/Engine/Private/WorldPartition/HLOD/HLODBuilder.cpp b/Engine/Source/Runtime/Engine/Private/WorldPartition/HLOD/HLODBuilder.cpp
--- a/Engine/Source/Runtime/Engine/Private/WorldPartition/HLOD/HLODBuilder.cpp
+++ b/Engine/Source/Runtime/Engine/Private/WorldPartition/HLOD/HLODBuilder.cpp
@@ -57,6 +57,14 @@ void UHLODBuilder::Build(AWorldPartitionHLOD* InHLODActor, const UHLODLayer* InH
 
 		HLODPrimitive->bComputeFastLocalBounds = true;
 		HLODPrimitive->bComputeBoundsOnceDuringCook = true;
+
+		if (InHLODLayer->GetLayerType() != EHLODLayerType::Instancing)
+		{
+			if (UStaticMeshComponent* StaticMeshComponent = Cast<UStaticMeshComponent>(HLODPrimitive))
+			{
+				StaticMeshComponent->bRayTracingFarField = (StaticMeshComponent->GetStaticMesh() && StaticMeshComponent->GetStaticMesh()->bSupportRayTracing);
+			}
+		}
 	}
 
 	if (!HLODPrimitives.IsEmpty())
diff --git a/Engine/Source/Runtime/Engine/Public/WorldPartition/HLOD/HLODActor.h b/Engine/Source/Runtime/Engine/Public/WorldPartition/HLOD/HLODActor.h
--- a/Engine/Source/Runtime/Engine/Public/WorldPartition/HLOD/HLODActor.h
+++ b/Engine/Source/Runtime/Engine/Public/WorldPartition/HLOD/HLODActor.h
@@ -64,7 +64,6 @@ protected:
 	//~ Begin UObject Interface.
 	virtual void Serialize(FArchive& Ar) override;
 	virtual void RerunConstructionScripts() override;
-	virtual void PostLoad() override;
 	//~ End UObject Interface.
 
 	//~ Begin AActor Interface.
