- Using material function instances will also trigger slow material layer parameter path
- Material layer path removes a special case where global parameters only checked overriden parameters for defaults; this is required for MLI parameters to work properly.
This should not change behavior overall, as for regular materials (without MFI), the default value stored on the parameter should match the value that's stored in CachedParameterData.

diff --git a/Engine/Source/Runtime/Engine/Private/Materials/MaterialCachedData.cpp b/Engine/Source/Runtime/Engine/Private/Materials/MaterialCachedData.cpp
--- a/Engine/Source/Runtime/Engine/Private/Materials/MaterialCachedData.cpp
+++ b/Engine/Source/Runtime/Engine/Private/Materials/MaterialCachedData.cpp
@@ -139,6 +139,8 @@ bool FMaterialCachedExpressionData::UpdateForFunction(const FMaterialCachedExpre
 	UMaterialFunctionInstance* FunctionInstance = Cast<UMaterialFunctionInstance>(Function);
 	if (FunctionInstance)
 	{
+		// If a material has function instances, it also needs to use material layer code paths
+		bHasMaterialLayers = true;
 		for (const FScalarParameterValue& Param : FunctionInstance->ScalarParameterValues)
 		{
 			const FMaterialParameterInfo ParameterInfo(Param.ParameterInfo.Name, Association, ParameterIndex);
@@ -277,6 +279,11 @@ bool FMaterialCachedExpressionData::UpdateForLayerFunctions(const FMaterialCache
 		}
 	}
 
+	if (LayerFunctions.Layers.Num() > 0 || LayerFunctions.Blends.Num() > 0)
+	{
+		bHasMaterialLayers = true;
+	}
+
 	return bResult;
 }
 
diff --git a/Engine/Source/Runtime/Engine/Private/Materials/MaterialShared.cpp b/Engine/Source/Runtime/Engine/Private/Materials/MaterialShared.cpp
--- a/Engine/Source/Runtime/Engine/Private/Materials/MaterialShared.cpp
+++ b/Engine/Source/Runtime/Engine/Private/Materials/MaterialShared.cpp
@@ -1718,7 +1718,7 @@ bool FMaterialResource::HasRuntimeVirtualTextureOutput() const
 
 bool FMaterialResource::HasMaterialLayers() const
 {
-	return Material->GetCachedExpressionData().DefaultLayers.Num() > 0;
+	return Material->GetCachedExpressionData().bHasMaterialLayers;
 }
 
 bool FMaterialResource::CastsRayTracedShadows() const
diff --git a/Engine/Source/Runtime/Engine/Private/Shader/Preshader.cpp b/Engine/Source/Runtime/Engine/Private/Shader/Preshader.cpp
--- a/Engine/Source/Runtime/Engine/Private/Shader/Preshader.cpp
+++ b/Engine/Source/Runtime/Engine/Private/Shader/Preshader.cpp
@@ -110,12 +110,10 @@ static void GetVectorParameter(const FUniformExpressionSet& UniformExpressionSet
 	bool bNeedsDefaultValue = false;
 	if (!Context.MaterialRenderProxy || !Context.MaterialRenderProxy->GetVectorValue(Parameter.ParameterInfo, &ParameterValue, Context))
 	{
-		const bool bOveriddenParameterOnly = Parameter.ParameterInfo.Association == EMaterialParameterAssociation::GlobalParameter;
-
 		if (Context.Material.HasMaterialLayers())
 		{
 			UMaterialInterface* Interface = Context.Material.GetMaterialInterface();
-			if (!Interface || !Interface->GetVectorParameterDefaultValue(Parameter.ParameterInfo, ParameterValue, bOveriddenParameterOnly))
+			if (!Interface || !Interface->GetVectorParameterDefaultValue(Parameter.ParameterInfo, ParameterValue))
 			{
 				bNeedsDefaultValue = true;
 			}
@@ -147,12 +145,10 @@ static void GetScalarParameter(const FUniformExpressionSet& UniformExpressionSet
 	bool bNeedsDefaultValue = false;
 	if (!Context.MaterialRenderProxy || !Context.MaterialRenderProxy->GetScalarValue(Parameter.ParameterInfo, &ParameterValue, Context))
 	{
-		const bool bOveriddenParameterOnly = Parameter.ParameterInfo.Association == EMaterialParameterAssociation::GlobalParameter;
-
 		if (Context.Material.HasMaterialLayers())
 		{
 			UMaterialInterface* Interface = Context.Material.GetMaterialInterface();
-			if (!Interface || !Interface->GetScalarParameterDefaultValue(Parameter.ParameterInfo, ParameterValue, bOveriddenParameterOnly))
+			if (!Interface || !Interface->GetScalarParameterDefaultValue(Parameter.ParameterInfo, ParameterValue))
 			{
 				bNeedsDefaultValue = true;
 			}
diff --git a/Engine/Source/Runtime/Engine/Public/MaterialCachedData.h b/Engine/Source/Runtime/Engine/Public/MaterialCachedData.h
--- a/Engine/Source/Runtime/Engine/Public/MaterialCachedData.h
+++ b/Engine/Source/Runtime/Engine/Public/MaterialCachedData.h
@@ -232,6 +232,7 @@ struct FMaterialCachedExpressionData
 		, bHasPerInstanceCustomData(false)
 		, bHasPerInstanceRandom(false)
 		, bHasVertexInterpolator(false)
+		, bHasMaterialLayers(false)
 	{}
 
 #if WITH_EDITOR
@@ -300,6 +301,9 @@ struct FMaterialCachedExpressionData
 	UPROPERTY()
 	uint32 bHasVertexInterpolator : 1;
 
+	UPROPERTY()
+	uint32 bHasMaterialLayers : 1;
+
 	/** Each bit corresponds to EMaterialProperty connection status. */
 	UPROPERTY()
 	uint32 MaterialAttributesPropertyConnectedBitmask = 0;
