- Added owner HLODLayer in world partition Runtime Grid (used by streaming source).
- Added Target HLODLayer option in WorldPartition Streaming Source.

diff --git a/Engine/Source/Runtime/Engine/Classes/Components/WorldPartitionStreamingSourceComponent.h b/Engine/Source/Runtime/Engine/Classes/Components/WorldPartitionStreamingSourceComponent.h
--- a/Engine/Source/Runtime/Engine/Classes/Components/WorldPartitionStreamingSourceComponent.h
+++ b/Engine/Source/Runtime/Engine/Classes/Components/WorldPartitionStreamingSourceComponent.h
@@ -6,6 +6,7 @@
 #include "UObject/ObjectMacros.h"
 #include "Components/ActorComponent.h"
 #include "WorldPartition/WorldPartitionStreamingSource.h"
+#include "WorldPartition/HLOD/HLODLayer.h"
 #include "WorldPartitionStreamingSourceComponent.generated.h"
 
 class FSceneView;
@@ -47,14 +48,22 @@ class ENGINE_API UWorldPartitionStreamingSourceComponent : public UActorComponen
 	float DefaultVisualizerLoadingRange;
 #endif
 
-	/** Optional target grid affected by streaming source. When none chosen, affects all grids. */
+	/** Optional target grid affected by streaming source. */
 	UPROPERTY(BlueprintReadWrite, EditAnywhere, Category = "Streaming")
 	FName TargetGrid;
 
+	/** Optional target HLODLayer affected by the streaming source. */
+	UPROPERTY(BlueprintReadWrite, EditAnywhere, Category = "Streaming")
+	TObjectPtr<const UHLODLayer> TargetHLODLayer;
+
 	/** Optional aggregated shape list used to build a custom shape for the streaming source. When empty, fallbacks sphere shape with a radius equal to grid's loading range. */
 	UPROPERTY(BlueprintReadWrite, EditAnywhere, Category = "Streaming")
 	TArray<FStreamingSourceShape> Shapes;
 
+#if WITH_EDITOR
+	virtual bool CanEditChange(const FProperty* InProperty) const override;
+#endif
+
 private:
 	/** Whether this component is enabled or not */
 	UPROPERTY(EditAnywhere, Category = "Streaming")
diff --git a/Engine/Source/Runtime/Engine/Private/Components/WorldPartitionStreamingSourceComponent.cpp b/Engine/Source/Runtime/Engine/Private/Components/WorldPartitionStreamingSourceComponent.cpp
--- a/Engine/Source/Runtime/Engine/Private/Components/WorldPartitionStreamingSourceComponent.cpp
+++ b/Engine/Source/Runtime/Engine/Private/Components/WorldPartitionStreamingSourceComponent.cpp
@@ -71,6 +71,7 @@ bool UWorldPartitionStreamingSourceComponent::GetStreamingSource(FWorldPartition
 		OutStreamingSource.Rotation = Actor->GetActorRotation();
 		OutStreamingSource.TargetState = TargetState;
 		OutStreamingSource.TargetGrid = TargetGrid;
+		OutStreamingSource.TargetHLODLayer = TargetHLODLayer;
 		OutStreamingSource.Shapes = Shapes;
 		return true;
 	}
@@ -131,4 +132,15 @@ void UWorldPartitionStreamingSourceComponent::DrawVisualization(const FSceneView
 		}
 	});
 #endif
-}
\ No newline at end of file
+}
+
+#if WITH_EDITOR
+bool UWorldPartitionStreamingSourceComponent::CanEditChange(const FProperty* InProperty) const
+{
+	if (InProperty && InProperty->GetName() == GET_MEMBER_NAME_STRING_CHECKED(UWorldPartitionStreamingSourceComponent, TargetGrid))
+	{
+		return TargetHLODLayer == nullptr;
+	}
+	return Super::CanEditChange(InProperty);
+}
+#endif
diff --git a/Engine/Source/Runtime/Engine/Private/WorldPartition/RuntimeSpatialHash/RuntimeSpatialHashHLOD.cpp b/Engine/Source/Runtime/Engine/Private/WorldPartition/RuntimeSpatialHash/RuntimeSpatialHashHLOD.cpp
--- a/Engine/Source/Runtime/Engine/Private/WorldPartition/RuntimeSpatialHash/RuntimeSpatialHashHLOD.cpp
+++ b/Engine/Source/Runtime/Engine/Private/WorldPartition/RuntimeSpatialHash/RuntimeSpatialHashHLOD.cpp
@@ -311,6 +311,7 @@ static void GatherHLODGrids(UWorldPartition* WorldPartition, TMap<FName, FSpatia
 		HLODGrid.DebugColor = FLinearColor::Red;
 		HLODGrid.GridName = HLODLayer->GetRuntimeGrid(InitialHLODLevel);
 		HLODGrid.bClientOnlyVisible = true;
+		HLODGrid.HLODLayer = HLODLayer;
 
 		OutHLODGrids.Emplace(HLODGrid.GridName, HLODGrid);
 
@@ -423,7 +424,7 @@ static void UpdateHLODGridsActors(UWorld* World, const TMap<FName, FSpatialHashR
 		if (GridActor->ActorHasTag(HLODGridTag))
 		{
 			const FSpatialHashRuntimeGrid* HLODGrid = HLODGrids.Find(GridActor->GridSettings.GridName);
-			if (HLODGrid && GridActor->GridSettings.Priority)
+			if (HLODGrid && GridActor->GridSettings.Priority && GridActor->GridSettings.HLODLayer)
 			{
 				ExistingGridActors.Emplace(GridActor->GridSettings.GridName, GridActor);
 			}
diff --git a/Engine/Source/Runtime/Engine/Private/WorldPartition/WorldPartitionRuntimeSpatialHash.cpp b/Engine/Source/Runtime/Engine/Private/WorldPartition/WorldPartitionRuntimeSpatialHash.cpp
--- a/Engine/Source/Runtime/Engine/Private/WorldPartition/WorldPartitionRuntimeSpatialHash.cpp
+++ b/Engine/Source/Runtime/Engine/Private/WorldPartition/WorldPartitionRuntimeSpatialHash.cpp
@@ -146,7 +146,7 @@ void FSpatialHashStreamingGrid::GetCells(const FWorldPartitionStreamingQuerySour
 	// Spatial Query
 	if (QuerySource.bSpatialQuery)
 	{
-		QuerySource.ForEachShape(GetLoadingRange(), GridName, /*bProjectIn2D*/ true, [&](const FSphericalSector& Shape)
+		QuerySource.ForEachShape(GetLoadingRange(), GridName, HLODLayer, /*bProjectIn2D*/ true, [&](const FSphericalSector& Shape)
 		{
 			Helper.ForEachIntersectingCells(Shape, [&](const FIntVector& Coords)
 			{
@@ -187,7 +187,7 @@ void FSpatialHashStreamingGrid::GetCells(const TArray<FWorldPartitionStreamingSo
 	const FSquare2DGridHelper& Helper = GetGridHelper();
 	for (const FWorldPartitionStreamingSource& Source : Sources)
 	{
-		Source.ForEachShape(GridLoadingRange, GridName, /*bProjectIn2D*/ true, [&](const FSphericalSector& Shape)
+		Source.ForEachShape(GridLoadingRange, GridName, HLODLayer, /*bProjectIn2D*/ true, [&](const FSphericalSector& Shape)
 		{
 			UWorldPartitionRuntimeCell::FStreamingSourceInfo Info(Source, Shape);
 
@@ -303,7 +303,7 @@ void FSpatialHashStreamingGrid::Draw3D(UWorld* World, const TArray<FWorldPartiti
 			Z = Hit.ImpactPoint.Z;
 		}
 
-		FBox Region = Source.CalcBounds(GridLoadingRange, GridName);
+		FBox Region = Source.CalcBounds(GridLoadingRange, GridName, HLODLayer);
 		Region += FBox(Region.GetCenter() - MinExtent, Region.GetCenter() + MinExtent);
 
 		for (int32 GridLevel = MinGridLevel; GridLevel <= MaxGridLevel; ++GridLevel)
@@ -360,7 +360,7 @@ void FSpatialHashStreamingGrid::Draw3D(UWorld* World, const TArray<FWorldPartiti
 
 		// Draw Streaming Source
 		const FColor Color = Source.GetDebugColor();
-		Source.ForEachShape(GetLoadingRange(), GridName, /*bProjectIn2D*/ true, [&Color, &Z, &Transform, &World, this](const FSphericalSector& Shape)
+		Source.ForEachShape(GetLoadingRange(), GridName, HLODLayer, /*bProjectIn2D*/ true, [&Color, &Z, &Transform, &World, this](const FSphericalSector& Shape)
 		{
 			FSphericalSector ZOffsettedShape = Shape;
 			ZOffsettedShape.SetCenter(FVector(FVector2D(ZOffsettedShape.GetCenter()), Z));
@@ -519,7 +519,7 @@ void FSpatialHashStreamingGrid::Draw2D(UCanvas* Canvas, UWorld* World, const TAr
 		for (const FWorldPartitionStreamingSource& Source : Sources)
 		{
 			const FColor Color = Source.GetDebugColor();
-			Source.ForEachShape(GridLoadingRange, GridName, /*bProjectIn2D*/ true, [&Color, &Canvas, &WorldToScreen, this](const FSphericalSector& Shape)
+			Source.ForEachShape(GridLoadingRange, GridName, HLODLayer, /*bProjectIn2D*/ true, [&Color, &Canvas, &WorldToScreen, this](const FSphericalSector& Shape)
 			{
 				DrawStreamingSource2D(Canvas, Shape, WorldToScreen, Color);
 			});
@@ -847,6 +847,7 @@ bool UWorldPartitionRuntimeSpatialHash::CreateStreamingGrid(const FSpatialHashRu
 	CurrentStreamingGrid.bBlockOnSlowStreaming = RuntimeGrid.bBlockOnSlowStreaming;
 	CurrentStreamingGrid.DebugColor = RuntimeGrid.DebugColor;
 	CurrentStreamingGrid.bClientOnlyVisible = RuntimeGrid.bClientOnlyVisible;
+	CurrentStreamingGrid.HLODLayer = RuntimeGrid.HLODLayer;
 
 	// Move actors into the final streaming grids
 	CurrentStreamingGrid.GridLevels.Reserve(PartionedActors.Levels.Num());
@@ -1279,7 +1280,7 @@ void UWorldPartitionRuntimeSpatialHash::Draw2D(UCanvas* Canvas, const TArray<FWo
 		FBox Region(ForceInit);
 		for (const FWorldPartitionStreamingSource& Source : Sources)
 		{
-			Region += Source.CalcBounds(StreamingGrid->GetLoadingRange(), StreamingGrid->GridName, /*bCalcIn2D*/ true);
+			Region += Source.CalcBounds(StreamingGrid->GetLoadingRange(), StreamingGrid->GridName, StreamingGrid->HLODLayer, /*bCalcIn2D*/ true);
 		}
 		Region += FBox(Region.GetCenter() - MinExtent, Region.GetCenter() + MinExtent);
 		const FVector2D GridReferenceWorldPos = FVector2D(Region.GetCenter());
diff --git a/Engine/Source/Runtime/Engine/Public/WorldPartition/WorldPartitionRuntimeSpatialHash.h b/Engine/Source/Runtime/Engine/Public/WorldPartition/WorldPartitionRuntimeSpatialHash.h
--- a/Engine/Source/Runtime/Engine/Public/WorldPartition/WorldPartitionRuntimeSpatialHash.h
+++ b/Engine/Source/Runtime/Engine/Public/WorldPartition/WorldPartitionRuntimeSpatialHash.h
@@ -9,6 +9,7 @@
 
 #include "WorldPartition.h"
 #include "WorldPartition/DataLayer/DataLayersID.h"
+#include "WorldPartition/HLOD/HLODLayer.h"
 #include "WorldPartitionRuntimeHash.h"
 #include "WorldPartitionRuntimeSpatialHashCell.h"
 #include "WorldPartitionRuntimeSpatialHashGridPreviewer.h"
@@ -118,6 +119,9 @@ struct FSpatialHashStreamingGrid
 	UPROPERTY()
 	bool bClientOnlyVisible;
 
+	UPROPERTY()
+	TObjectPtr<const UHLODLayer> HLODLayer;
+
 	float OverrideLoadingRange;
 
 	float GetLoadingRange() const
@@ -189,6 +193,9 @@ struct FSpatialHashRuntimeGrid
 
 	UPROPERTY()
 	bool bClientOnlyVisible;
+
+	UPROPERTY()
+	TObjectPtr<const UHLODLayer> HLODLayer;
 #endif
 };
 
diff --git a/Engine/Source/Runtime/Engine/Public/WorldPartition/WorldPartitionStreamingSource.h b/Engine/Source/Runtime/Engine/Public/WorldPartition/WorldPartitionStreamingSource.h
--- a/Engine/Source/Runtime/Engine/Public/WorldPartition/WorldPartitionStreamingSource.h
+++ b/Engine/Source/Runtime/Engine/Public/WorldPartition/WorldPartitionStreamingSource.h
@@ -157,6 +157,16 @@ class FStreamingSourceShapeHelper
 {
 public:
 
+	FORCEINLINE static bool IsSourceAffectingGrid(FName InSourceTargetGrid, const FSoftObjectPath& InSourceTargetHLODLayer, FName InGridName, const FSoftObjectPath& InGridHLODLayer)
+	{
+		if ((InSourceTargetGrid.IsNone() && InSourceTargetHLODLayer.IsNull()) ||
+			(InSourceTargetHLODLayer.IsNull() && (InSourceTargetGrid == InGridName)) ||
+			(!InSourceTargetHLODLayer.IsNull() && (InSourceTargetHLODLayer == InGridHLODLayer)))
+		{
+			return true;
+		}
+		return false;
+	}
 	FORCEINLINE static void ForEachShape(float InGridLoadingRange, float InDefaultRadius, bool bInProjectIn2D, const FVector& InLocation, const FRotator& InRotation, const TArray<FStreamingSourceShape>& InShapes, TFunctionRef<void(const FSphericalSector&)> InOperation)
 	{
 		const FTransform Transform(bInProjectIn2D ? FRotator(0, InRotation.Yaw, 0) : InRotation, InLocation);
@@ -250,12 +260,18 @@ struct FWorldPartitionStreamingQuerySource
 	/* Reserved settings used by UWorldPartitionStreamingSourceComponent::IsStreamingCompleted. */
 	FRotator Rotation;
 	FName TargetGrid;
+	FSoftObjectPath TargetHLODLayer;
 	TArray<FStreamingSourceShape> Shapes;
 
 	/** Helper method that iterates over all shapes. If none is provided, it will still pass a sphere shape using Radius or grid's loading range (see bUseGridLoadingRange). */
-	FORCEINLINE void ForEachShape(float InGridLoadingRange, FName InGridName, bool bInProjectIn2D, TFunctionRef<void(const FSphericalSector&)> InOperation) const
+	FORCEINLINE void ForEachShape(float InGridLoadingRange, FName InGridName, const FSoftObjectPath& InGridHLODLayer, bool bInProjectIn2D, TFunctionRef<void(const FSphericalSector&)> InOperation) const
 	{
-		if (bSpatialQuery && (TargetGrid.IsNone() || TargetGrid == InGridName))
+		if (!bSpatialQuery)
+		{
+			return;
+		}
+
+		if (FStreamingSourceShapeHelper::IsSourceAffectingGrid(TargetGrid, TargetHLODLayer, InGridName, InGridHLODLayer))
 		{
 			FStreamingSourceShapeHelper::ForEachShape(InGridLoadingRange, bUseGridLoadingRange ? InGridLoadingRange : Radius, bInProjectIn2D, Location, Rotation, Shapes, InOperation);
 		}
@@ -322,17 +338,20 @@ struct ENGINE_API FWorldPartitionStreamingSource
 	/** Source velocity (computed automatically). */
 	float Velocity;
 
-	/** When set, will only affect streaming on the provided target runtime streaming grid. When none is provided, applies to all runtime streaming grid. */
+	/** When set, will only affect streaming on the provided target runtime streaming grid. */
 	FName TargetGrid;
 
+	/** When set, will only affect streaming on HLODs associated to the provided target HLODLayer. */
+	FSoftObjectPath TargetHLODLayer;
+
 	/** Source internal shapes. When none are provided, a sphere is automatically used. It's radius is equal to grid's loading range and center equals source's location. */
 	TArray<FStreamingSourceShape> Shapes;
 
 	/** Returns a box encapsulating all shapes. */
-	FORCEINLINE FBox CalcBounds(float InGridLoadingRange, FName InGridName, bool bCalcIn2D = false) const
+	FORCEINLINE FBox CalcBounds(float InGridLoadingRange, FName InGridName, const FSoftObjectPath& InGridHLODLayer, bool bCalcIn2D = false) const
 	{
 		FBox OutBounds(ForceInit);
-		ForEachShape(InGridLoadingRange, InGridName, bCalcIn2D, [&OutBounds](const FSphericalSector& Sector)
+		ForEachShape(InGridLoadingRange, InGridName, InGridHLODLayer, bCalcIn2D, [&OutBounds](const FSphericalSector& Sector)
 		{
 			OutBounds += Sector.CalcBounds();
 		});
@@ -340,9 +359,9 @@ struct ENGINE_API FWorldPartitionStreamingSource
 	}
 
 	/** Helper method that iterates over all shapes. If none is provided, it will still pass a sphere shape using grid's loading range. */
-	FORCEINLINE void ForEachShape(float InGridLoadingRange, FName InGridName, bool bInProjectIn2D, TFunctionRef<void(const FSphericalSector&)> InOperation) const
+	FORCEINLINE void ForEachShape(float InGridLoadingRange, FName InGridName, const FSoftObjectPath& InGridHLODLayer, bool bInProjectIn2D, TFunctionRef<void(const FSphericalSector&)> InOperation) const
 	{
-		if (TargetGrid.IsNone() || TargetGrid == InGridName)
+		if (FStreamingSourceShapeHelper::IsSourceAffectingGrid(TargetGrid, TargetHLODLayer, InGridName, InGridHLODLayer))
 		{
 			FStreamingSourceShapeHelper::ForEachShape(InGridLoadingRange, InGridLoadingRange, bInProjectIn2D, Location, Rotation, Shapes, InOperation);
 		}
