DDC2 async texture build flow changes related to messaging and validation:

Migrate FTextureStatusMessageContext - functionality abandoned due to limitations of FSlowTask
Migrate FTextureSourceData::Init() validation
Migrate FTextureCacheDerivedDataWorker::BuildTexture() validation
Migrate FTextureCacheDerivedDataWorker::FTextureCacheDerivedDataWorker() validation
Only use the composite texture source if it passed validation

diff --git a/Engine/Source/Runtime/Engine/Private/TextureDerivedDataTask.cpp b/Engine/Source/Runtime/Engine/Private/TextureDerivedDataTask.cpp
--- a/Engine/Source/Runtime/Engine/Private/TextureDerivedDataTask.cpp
+++ b/Engine/Source/Runtime/Engine/Private/TextureDerivedDataTask.cpp
@@ -80,6 +80,34 @@ public:
 	}
 };
 
+static FText ComposeTextureBuildText(const UTexture& Texture, int32 SizeX, int32 SizeY, int32 NumBlocks, int32 NumLayers, const FTextureBuildSettings& BuildSettings, int64 RequiredMemoryEstimate, bool bIsVT)
+{
+	FFormatNamedArguments Args;
+	Args.Add(TEXT("TextureName"), FText::FromString(Texture.GetPathName()));
+	Args.Add(TEXT("TextureFormatName"), FText::FromString(BuildSettings.TextureFormatName.GetPlainNameString()));
+	Args.Add(TEXT("IsVT"), FText::FromString( FString( bIsVT ? TEXT(" VT") : TEXT("") ) ) );
+	Args.Add(TEXT("TextureResolutionX"), FText::FromString(FString::FromInt(SizeX)));
+	Args.Add(TEXT("TextureResolutionY"), FText::FromString(FString::FromInt(SizeY)));
+	Args.Add(TEXT("NumBlocks"), FText::FromString(FString::FromInt(NumBlocks)));
+	Args.Add(TEXT("NumLayers"), FText::FromString(FString::FromInt(NumLayers)));
+	Args.Add(TEXT("EstimatedMemory"), FText::FromString(FString::SanitizeFloat(double(RequiredMemoryEstimate) / (1024.0*1024.0), 3)));
+
+	return FText::Format(
+		NSLOCTEXT("Engine", "BuildTextureStatus", "Building textures: {TextureName} ({TextureFormatName}{IsVT}, {TextureResolutionX}X{TextureResolutionY} X{NumBlocks}X{NumLayers}) (Required Memory Estimate: {EstimatedMemory} MB)"), 
+		Args
+	);
+}
+
+static FText ComposeTextureBuildText(const UTexture& Texture, const FTextureSourceData& TextureData, const FTextureBuildSettings& BuildSettings, int64 RequiredMemoryEstimate, bool bIsVT)
+{
+	return ComposeTextureBuildText(Texture, TextureData.Blocks[0].MipsPerLayer[0][0].SizeX, TextureData.Blocks[0].MipsPerLayer[0][0].SizeY, TextureData.Blocks.Num(), TextureData.Layers.Num(), BuildSettings, RequiredMemoryEstimate, bIsVT);
+}
+
+static FText ComposeTextureBuildText(const UTexture& Texture, const FTextureBuildSettings& BuildSettings, int64 RequiredMemoryEstimate, bool bIsVT)
+{
+	return ComposeTextureBuildText(Texture, Texture.Source.GetSizeX(), Texture.Source.GetSizeY(), Texture.Source.GetNumBlocks(), Texture.Source.GetNumLayers(), BuildSettings, RequiredMemoryEstimate, bIsVT);
+}
+
 static bool ValidateTexture2DPlatformData(const FTexturePlatformData& TextureData, const UTexture2D& Texture, bool bFromDDC)
 {
 	// Temporarily disable as the size check reports false negatives on some platforms
@@ -441,22 +469,6 @@ void FTextureCacheDerivedDataWorker::BuildTexture(bool bReplaceExistingDDC)
 	const bool bHasValidMip0 = TextureData.Blocks.Num() && TextureData.Blocks[0].MipsPerLayer.Num() && TextureData.Blocks[0].MipsPerLayer[0].Num();
 	const bool bForVirtualTextureStreamingBuild = EnumHasAnyFlags(CacheFlags, ETextureCacheFlags::ForVirtualTextureStreamingBuild);
 
-	FFormatNamedArguments Args;
-	Args.Add(TEXT("TextureName"), FText::FromString(Texture.GetPathName()));
-	Args.Add(TEXT("TextureFormatName"), FText::FromString(BuildSettingsPerLayer[0].TextureFormatName.GetPlainNameString()));
-	Args.Add(TEXT("IsVT"), FText::FromString( FString( bForVirtualTextureStreamingBuild ? TEXT(" VT") : TEXT("") ) ) );
-	Args.Add(TEXT("TextureResolutionX"), FText::FromString(FString::FromInt(bHasValidMip0 ? TextureData.Blocks[0].MipsPerLayer[0][0].SizeX : 0)));
-	Args.Add(TEXT("TextureResolutionY"), FText::FromString(FString::FromInt(bHasValidMip0 ? TextureData.Blocks[0].MipsPerLayer[0][0].SizeY : 0)));
-	Args.Add(TEXT("NumBlocks"), FText::FromString(FString::FromInt(TextureData.Blocks.Num())));
-	Args.Add(TEXT("NumLayers"), FText::FromString(FString::FromInt(TextureData.Layers.Num())));
-	Args.Add(TEXT("EstimatedMemory"), FText::FromString(FString::SanitizeFloat(double(RequiredMemoryEstimate) / (1024.0*1024.0), 3)));
-
-	FTextureStatusMessageContext StatusMessage(
-		FText::Format(
-			NSLOCTEXT("Engine", "BuildTextureStatus", "Building textures: {TextureName} ({TextureFormatName}{IsVT}, {TextureResolutionX}X{TextureResolutionY} X{NumBlocks}X{NumLayers}) (Required Memory Estimate: {EstimatedMemory} MB)"), 
-			Args
-		));
-
 	if (!ensure(Compressor))
 	{
 		UE_LOG(LogTexture, Warning, TEXT("Missing Compressor required to build texture %s"), *Texture.GetPathName());
@@ -468,6 +480,10 @@ void FTextureCacheDerivedDataWorker::BuildTexture(bool bReplaceExistingDDC)
 		return;
 	}
 
+	FTextureStatusMessageContext StatusMessage(
+		ComposeTextureBuildText(Texture, TextureData, BuildSettingsPerLayer[0], RequiredMemoryEstimate, bForVirtualTextureStreamingBuild)
+		);
+
 	if (bForVirtualTextureStreamingBuild)
 	{
 		if (DerivedData->VTData == nullptr)
@@ -1019,15 +1035,27 @@ public:
 		IBuildInputResolver* GlobalResolver = GetGlobalBuildInputResolver();
 		BuildSession = Build.CreateSession(TexturePath, GlobalResolver ? GlobalResolver : &InputResolver);
 
-		Owner.Emplace(EnumHasAnyFlags(Flags, ETextureCacheFlags::Async) ? ConvertPriority(Priority) : UE::DerivedData::EPriority::Blocking);
+		EPriority OwnerPriority = EnumHasAnyFlags(Flags, ETextureCacheFlags::Async) ? ConvertPriority(Priority) : UE::DerivedData::EPriority::Blocking;
+		Owner.Emplace(OwnerPriority);
+
+		bool bUseCompositeTexture;
+		if (!IsTextureValidForBuilding(Texture, Flags, bUseCompositeTexture))
+		{
+			return;
+		}
+		
+		if (IsInGameThread() && OwnerPriority == EPriority::Blocking)
+		{
+			StatusMessage.Emplace(ComposeTextureBuildText(Texture, Settings, Texture.GetBuildRequiredMemory(), EnumHasAnyFlags(Flags, ETextureCacheFlags::ForVirtualTextureStreamingBuild)));
+		}
 
-		FBuildDefinition Definition = CreateDefinition(Build, Texture, TexturePath, FunctionName, KeySuffix, Settings);
+		FBuildDefinition Definition = CreateDefinition(Build, Texture, TexturePath, FunctionName, KeySuffix, Settings, bUseCompositeTexture);
 
 		if (!EnumHasAnyFlags(Flags, ETextureCacheFlags::ForceRebuild) && Settings.bHasEditorOnlyData)
 		{
 			FTextureBuildSettings ShippingSettings = Settings;
 			ShippingSettings.bHasEditorOnlyData = false;
-			FBuildDefinition ShippingDefinition = CreateDefinition(Build, Texture, TexturePath, FunctionName, KeySuffix, ShippingSettings);
+			FBuildDefinition ShippingDefinition = CreateDefinition(Build, Texture, TexturePath, FunctionName, KeySuffix, ShippingSettings, bUseCompositeTexture);
 			BuildSession.Get().Build(ShippingDefinition, EBuildPolicy::Cache, *Owner,
 				[this, Definition = MoveTemp(Definition), Flags](FBuildCompleteParams&& Params)
 				{
@@ -1053,13 +1081,14 @@ public:
 		FStringView TexturePath,
 		FStringView FunctionName,
 		const FString& KeySuffix,
-		const FTextureBuildSettings& Settings)
+		const FTextureBuildSettings& Settings,
+		const bool bUseCompositeTexture)
 	{
 		UE::DerivedData::FBuildDefinitionBuilder DefinitionBuilder = Build.CreateDefinition(TexturePath, FunctionName);
 		DefinitionBuilder.AddConstant(TEXT("Settings"_SV),
 			SaveTextureBuildSettings(KeySuffix, Texture, Settings, 0, NUM_INLINE_DERIVED_MIPS));
 		DefinitionBuilder.AddInputBulkData(TEXT("Source"_SV), Texture.Source.GetPersistentId());
-		if (Texture.CompositeTexture)
+		if (Texture.CompositeTexture && bUseCompositeTexture)
 		{
 			DefinitionBuilder.AddInputBulkData(TEXT("CompositeSource"_SV), Texture.CompositeTexture->Source.GetPersistentId());
 		}
@@ -1248,6 +1277,116 @@ private:
 		return true;
 	}
 
+	static bool IsTextureValidForBuilding(UTexture& Texture, ETextureCacheFlags Flags, bool& bOutUseCompositeTexture)
+	{
+		const int32 NumBlocks = Texture.Source.GetNumBlocks();
+		const int32 NumLayers = Texture.Source.GetNumLayers();
+		if (NumBlocks < 1 || NumLayers < 1)
+		{
+			UE_LOG(LogTexture, Error, TEXT("Texture has no source data: %s"), *Texture.GetPathName());
+			return false;
+		}
+
+		for (int LayerIndex = 0; LayerIndex < NumLayers; ++LayerIndex)
+		{
+			switch (Texture.Source.GetFormat(LayerIndex))
+			{
+			case TSF_G8:
+			case TSF_G16:
+			case TSF_BGRA8:
+			case TSF_BGRE8:
+			case TSF_RGBA16:
+			case TSF_RGBA16F:
+			break;
+			default:
+				UE_LOG(LogTexture, Fatal, TEXT("Texture %s has source art in an invalid format."), *Texture.GetPathName());
+				return false;
+			}
+		}
+
+		const bool bCompositeTextureViable = Texture.CompositeTexture && Texture.CompositeTextureMode != CTM_Disabled;
+		bool bMatchingBlocks = bCompositeTextureViable && (Texture.CompositeTexture->Source.GetNumBlocks() == Texture.Source.GetNumBlocks());
+		bool bMatchingAspectRatio = bCompositeTextureViable;
+		bool bOnlyPowerOfTwoSize = bCompositeTextureViable;
+
+		int32 BlockSizeX = 0;
+		int32 BlockSizeY = 0;
+		TArray<FIntPoint> BlockSizes;
+		BlockSizes.Reserve(NumBlocks);
+		for (int32 BlockIndex = 0; BlockIndex < NumBlocks; ++BlockIndex)
+		{
+			FTextureSourceBlock SourceBlock;
+			Texture.Source.GetBlock(BlockIndex, SourceBlock);
+			if (SourceBlock.NumMips > 0 && SourceBlock.NumSlices > 0)
+			{
+				BlockSizes.Emplace(SourceBlock.SizeX, SourceBlock.SizeY);
+				BlockSizeX = FMath::Max(BlockSizeX, SourceBlock.SizeX);
+				BlockSizeY = FMath::Max(BlockSizeY, SourceBlock.SizeY);
+			}
+
+			if (bCompositeTextureViable)
+			{
+				FTextureSourceBlock CompositeTextureBlock;
+				Texture.CompositeTexture->Source.GetBlock(BlockIndex, CompositeTextureBlock);
+
+				bMatchingBlocks = bMatchingBlocks && SourceBlock.BlockX == CompositeTextureBlock.BlockX && SourceBlock.BlockY == CompositeTextureBlock.BlockY;
+				bMatchingAspectRatio = bMatchingAspectRatio && SourceBlock.SizeX * CompositeTextureBlock.SizeY == SourceBlock.SizeY * CompositeTextureBlock.SizeX;
+				bOnlyPowerOfTwoSize = bOnlyPowerOfTwoSize && FMath::IsPowerOfTwo(SourceBlock.SizeX) && FMath::IsPowerOfTwo(SourceBlock.SizeY);
+			}
+		}
+
+		for (int32 BlockIndex = 0; BlockIndex < BlockSizes.Num(); ++BlockIndex)
+		{
+			const int32 MipBiasX = FMath::CeilLogTwo(BlockSizeX / BlockSizes[BlockIndex].X);
+			const int32 MipBiasY = FMath::CeilLogTwo(BlockSizeY / BlockSizes[BlockIndex].Y);
+			if (MipBiasX != MipBiasY)
+			{
+				UE_LOG(LogTexture, Error, TEXT("Texture %s has blocks with mismatched aspect ratios"), *Texture.GetPathName());
+				return false;
+			}
+		}
+
+		if (bCompositeTextureViable)
+		{
+			if (!bMatchingBlocks)
+			{
+				UE_LOG(LogTexture, Warning, TEXT("Issue while building %s : Composite texture resolution/UDIMs do not match. Composite texture will be ignored"), *Texture.GetPathName());
+			}
+			else if (!bOnlyPowerOfTwoSize)
+			{
+				UE_LOG(LogTexture, Warning, TEXT("Issue while building %s : Some blocks (UDIMs) have a non power of two size. Composite texture will be ignored"), *Texture.GetPathName());
+			}
+			else if (!bMatchingAspectRatio)
+			{
+				UE_LOG(LogTexture, Warning, TEXT("Issue while building %s : Some blocks (UDIMs) have mismatched aspect ratio. Composite texture will be ignored"), *Texture.GetPathName());
+			}
+		}
+
+		bOutUseCompositeTexture = bMatchingBlocks && bMatchingAspectRatio && bOnlyPowerOfTwoSize;
+
+		// TODO: Add validation equivalent to that found in FTextureCacheDerivedDataWorker::BuildTexture for virtual textures
+		//		 if virtual texture support is added for this code path.
+		if (!EnumHasAnyFlags(Flags, ETextureCacheFlags::ForVirtualTextureStreamingBuild))
+		{
+			// Only support single Block/Layer here (Blocks and Layers are intended for VT support)
+			if (NumBlocks > 1)
+			{
+				// This warning can happen if user attempts to import a UDIM without VT enabled
+				UE_LOG(LogTexture, Warning, TEXT("Texture %s was imported as UDIM with %d blocks but VirtualTexturing is not enabled, only the first block will be available"),
+					*Texture.GetPathName(), NumBlocks);
+			}
+
+			// No user-facing way to generated multi-layered textures currently, so this should not occur
+			if (NumLayers > 1)
+			{
+				UE_LOG(LogTexture, Warning, TEXT("Texture %s has %d layers but VirtualTexturing is not enabled, only the first layer will be available"),
+					*Texture.GetPathName(), NumLayers);
+			}
+		}
+
+		return true;
+	}
+
 	FTexturePlatformData& DerivedData;
 	TOptional<UE::DerivedData::FRequestOwner> Owner;
 	UE::DerivedData::FOptionalBuildSession BuildSession;
@@ -1256,6 +1395,7 @@ private:
 	bool bInlineMips;
 	int32 FirstMipToLoad;
 	uint64 BuildOutputSize = 0;
+	TOptional<FTextureStatusMessageContext> StatusMessage;
 	UE::TextureDerivedData::FTextureBuildInputResolver InputResolver;
 	FRWLock Lock;
 };
