Fixed build warnings

diff --git a/Engine/Source/Runtime/Engine/Private/WorldPartition/WorldPartitionRuntimeSpatialHash.cpp b/Engine/Source/Runtime/Engine/Private/WorldPartition/WorldPartitionRuntimeSpatialHash.cpp
--- a/Engine/Source/Runtime/Engine/Private/WorldPartition/WorldPartitionRuntimeSpatialHash.cpp
+++ b/Engine/Source/Runtime/Engine/Private/WorldPartition/WorldPartitionRuntimeSpatialHash.cpp
@@ -88,6 +88,7 @@ FSpatialHashStreamingGrid::FSpatialHashStreamingGrid()
 	, DebugColor(ForceInitToZero)
 	, WorldBounds(ForceInitToZero)
 	, bClientOnlyVisible(false)
+	, HLODLayer(nullptr)
 	, OverrideLoadingRange(-1.f)
 	, GridHelper(nullptr)
 {
@@ -376,20 +377,24 @@ void FSpatialHashStreamingGrid::DrawStreamingSource3D(UWorld* World, const FSphe
 		FVector Location = InTransform.TransformPosition(InShape.GetCenter());
 		DrawDebugSphere(World, Location, InShape.GetRadius(), 32, InColor, false, -1.f, 0, 20.f);
 	}
-	else if (ULineBatchComponent* const LineBatcher = World ? World->LineBatcher : nullptr)
+	else
 	{
-		FSphericalSector Shape = InShape;
-		Shape.SetAxis(InTransform.TransformVector(Shape.GetAxis()));
-		Shape.SetCenter(InTransform.TransformPosition(Shape.GetCenter()));
-
-		TArray<TPair<FVector, FVector>> Lines = Shape.BuildDebugMesh();
-		TArray<FBatchedLine> BatchedLines;
-		BatchedLines.Empty(Lines.Num());
-		for (const auto& Line : Lines)
+		ULineBatchComponent* const LineBatcher = World->LineBatcher;
+		if (LineBatcher)
 		{
-			BatchedLines.Emplace(FBatchedLine(Line.Key, Line.Value, InColor, LineBatcher->DefaultLifeTime, 20.f, SDPG_World));
-		};
-		LineBatcher->DrawLines(BatchedLines);
+			FSphericalSector Shape = InShape;
+			Shape.SetAxis(InTransform.TransformVector(Shape.GetAxis()));
+			Shape.SetCenter(InTransform.TransformPosition(Shape.GetCenter()));
+
+			TArray<TPair<FVector, FVector>> Lines = Shape.BuildDebugMesh();
+			TArray<FBatchedLine> BatchedLines;
+			BatchedLines.Empty(Lines.Num());
+			for (const auto& Line : Lines)
+			{
+				BatchedLines.Emplace(FBatchedLine(Line.Key, Line.Value, InColor, LineBatcher->DefaultLifeTime, 20.f, SDPG_World));
+			};
+			LineBatcher->DrawLines(BatchedLines);
+		}
 	}
 }
 
diff --git a/Engine/Source/Runtime/Engine/Public/WorldPartition/WorldPartitionRuntimeSpatialHash.h b/Engine/Source/Runtime/Engine/Public/WorldPartition/WorldPartitionRuntimeSpatialHash.h
--- a/Engine/Source/Runtime/Engine/Public/WorldPartition/WorldPartitionRuntimeSpatialHash.h
+++ b/Engine/Source/Runtime/Engine/Public/WorldPartition/WorldPartitionRuntimeSpatialHash.h
@@ -168,6 +168,7 @@ struct FSpatialHashRuntimeGrid
 		, Priority(0)
 		, DebugColor(FLinearColor::MakeRandomColor())
 		, bClientOnlyVisible(false)
+		, HLODLayer(nullptr)
 #endif
 	{}
 
