Fix a memory leak when translating materials found by ChrisW.

diff --git a/Engine/Source/Runtime/Engine/Private/Materials/HLSLMaterialTranslator.cpp b/Engine/Source/Runtime/Engine/Private/Materials/HLSLMaterialTranslator.cpp
--- a/Engine/Source/Runtime/Engine/Private/Materials/HLSLMaterialTranslator.cpp
+++ b/Engine/Source/Runtime/Engine/Private/Materials/HLSLMaterialTranslator.cpp
@@ -1324,8 +1324,8 @@ bool FHLSLMaterialTranslator::Translate()
 
 		MaterialCompilationOutput.UniformExpressionSet.SetParameterCollections(ParameterCollections);
 
-		// Create the material uniform buffer struct.
-		MaterialCompilationOutput.UniformExpressionSet.CreateBufferStruct();
+		// This will be created shortly after the Translate call in FMaterial::BeginCompileShaderMap()
+		//MaterialCompilationOutput.UniformExpressionSet.CreateBufferStruct();
 
 		// Store the number of unique VT samples
 		MaterialCompilationOutput.EstimatedNumVirtualTextureLookups = NumVtSamples;
