Add the ability to tweak the priorities of theads at runtime
Fix the affinity setting at runtime

diff --git a/Engine/Source/Runtime/Core/Private/GenericPlatform/GenericPlatformProcess.cpp b/Engine/Source/Runtime/Core/Private/GenericPlatform/GenericPlatformProcess.cpp
--- a/Engine/Source/Runtime/Core/Private/GenericPlatform/GenericPlatformProcess.cpp
+++ b/Engine/Source/Runtime/Core/Private/GenericPlatform/GenericPlatformProcess.cpp
@@ -67,6 +67,11 @@ void FGenericPlatformProcess::SetThreadAffinityMask( uint64 AffinityMask )
 	// Not implemented cross-platform. Each platform may or may not choose to implement this.
 }
 
+void FGenericPlatformProcess::SetThreadPriority(EThreadPriority NewPriority)
+{
+	// Not implemented cross-platform. Each platform may or may not choose to implement this.
+}
+
 uint32 FGenericPlatformProcess::GetStackSize()
 {
 	return 0;
diff --git a/Engine/Source/Runtime/Core/Public/GenericPlatform/GenericPlatformProcess.h b/Engine/Source/Runtime/Core/Public/GenericPlatform/GenericPlatformProcess.h
--- a/Engine/Source/Runtime/Core/Public/GenericPlatform/GenericPlatformProcess.h
+++ b/Engine/Source/Runtime/Core/Public/GenericPlatform/GenericPlatformProcess.h
@@ -6,6 +6,8 @@
 #include "Containers/UnrealString.h"
 #include "Templates/Function.h"
 #include "Misc/EnumClassFlags.h"
+#include "Misc/EnumClassFlags.h"
+#include "GenericPlatform/GenericPlatformAffinity.h"
 
 ////////////////////////////////////////////////////////////////////////////////
 #if PLATFORM_CPU_X86_FAMILY
@@ -223,6 +225,13 @@ struct CORE_API FGenericPlatformProcess
 	 */
 	static void SetThreadAffinityMask( uint64 AffinityMask );
 
+	/**
+	 * Change the thread processor priority
+	 *
+	 * @param NewPriority an EThreadPriority indicating what priority the thread is to run at.
+	 */
+	static void SetThreadPriority( EThreadPriority NewPriority );
+
 	/**
 	 * Helper function to set thread name of the current thread.
 	 * @param ThreadName   Name to set
diff --git a/Engine/Source/Runtime/Engine/Private/UnrealEngine.cpp b/Engine/Source/Runtime/Engine/Private/UnrealEngine.cpp
--- a/Engine/Source/Runtime/Engine/Private/UnrealEngine.cpp
+++ b/Engine/Source/Runtime/Engine/Private/UnrealEngine.cpp
@@ -16205,133 +16205,287 @@ static FAutoConsoleCommand UAssetLoadTestCmd(
 
 #endif
 
-static uint64 TaskThread = 0xFFFFFFFFFFFFFFFF;
-static uint64 GameThread = 0xFFFFFFFFFFFFFFFF;
-static uint64 RenderThread = 0xFFFFFFFFFFFFFFFF;
-static uint64 RHIThread = 0xFFFFFFFFFFFFFFFF;
-FThreadSafeCounter StallForTaskThread;
+static uint64 TaskThreadAffinity = 0xFFFFFFFFFFFFFFFF;
+static uint64 TaskThreadBPAffinity = 0xFFFFFFFFFFFFFFFF;
+static uint64 GameThreadAffinity = 0xFFFFFFFFFFFFFFFF;
+static uint64 RenderThreadAffinity = 0xFFFFFFFFFFFFFFFF;
+static uint64 RHIThreadAffinity = 0xFFFFFFFFFFFFFFFF;
 
-void SetAffinityOnThread()
+static EThreadPriority TaskThreadPriority = EThreadPriority::TPri_Normal;
+static EThreadPriority TaskThreadBPPriority = EThreadPriority::TPri_BelowNormal;
+static EThreadPriority GameThreadPriority = EThreadPriority::TPri_Highest;
+static EThreadPriority RenderThreadPriority = EThreadPriority::TPri_Highest;
+static EThreadPriority RHIThreadPriority = EThreadPriority::TPri_Highest;
+
+static FString CVar_ThreadConfigString = TEXT("");
+
+extern CORE_API int32 GUseNewTaskBackend;
+
+static bool AffinitiesLoadedDefaults = false;
+static bool PrioritiesLoadedDefaults = false;
+
+void GetDefaultAffinities(bool bForce = false)
+{
+	if (!AffinitiesLoadedDefaults || bForce)
+	{
+		AffinitiesLoadedDefaults = true;
+		TaskThreadAffinity = FPlatformAffinity::GetTaskGraphThreadMask();
+		TaskThreadBPAffinity = FPlatformAffinity::GetTaskGraphBackgroundTaskMask();
+		GameThreadAffinity = FPlatformAffinity::GetMainGameMask();
+		RenderThreadAffinity = FPlatformAffinity::GetRenderingThreadMask();
+		RHIThreadAffinity = FPlatformAffinity::GetRHIThreadMask();
+	}
+}
+
+void GetDefaultPriorities(bool bForce = false)
 {
-	if (IsInActualRenderingThread())
+	if (!PrioritiesLoadedDefaults || bForce)
 	{
-		FPlatformProcess::SetThreadAffinityMask(RenderThread);
-		UE_LOG(LogConsoleResponse, Display, TEXT("RT     %016llX"), RenderThread);
+		PrioritiesLoadedDefaults = true;
+		TaskThreadPriority = EThreadPriority::TPri_Normal;
+		TaskThreadBPPriority = EThreadPriority::TPri_BelowNormal;
+		GameThreadPriority = EThreadPriority::TPri_Highest;
+		RenderThreadPriority = FPlatformAffinity::GetRenderingThreadPriority();
+		RHIThreadPriority = FPlatformAffinity::GetRHIThreadPriority();
 	}
-	else if (IsInRHIThread())
+}
+
+EThreadPriority StringToThreadPriority(FString InString)
+{
+	if (InString == TEXT("TPri_Normal"))
 	{
-		FPlatformProcess::SetThreadAffinityMask(RHIThread);
-		UE_LOG(LogConsoleResponse, Display, TEXT("RHI    %016llX"), RHIThread);
+		return EThreadPriority::TPri_Normal;
 	}
-	else if (IsInGameThread())
+	if (InString == TEXT("TPri_AboveNormal"))
 	{
-		FPlatformProcess::SetThreadAffinityMask(GameThread);
-		UE_LOG(LogConsoleResponse, Display, TEXT("GT     %016llX"), GameThread);
+		return EThreadPriority::TPri_AboveNormal;
 	}
-	else // assume task thread
+	if (InString == TEXT("TPri_BelowNormal"))
 	{
-		int32 TaskThreadIndex = int32(FTaskGraphInterface::Get().GetCurrentThreadIfKnown()) - int32(ENamedThreads::ActualRenderingThread) - 1;
-		FPlatformProcess::SetThreadAffinityMask(TaskThread);
-		UE_LOG(LogConsoleResponse, Display, TEXT("Task%2d %016llX"), TaskThreadIndex, TaskThread);
-		StallForTaskThread.Decrement();
-		// we wait for the others to finish here so that we do all task threads
-		while (StallForTaskThread.GetValue())
-		{
-			FPlatformProcess::Sleep(.0001f);
-		}
+		return EThreadPriority::TPri_BelowNormal;
 	}
+	if (InString == TEXT("TPri_Highest"))
+	{
+		return EThreadPriority::TPri_Highest;
+	}
+	if (InString == TEXT("TPri_Lowest"))
+	{
+		return EThreadPriority::TPri_Lowest;
+	}
+	if (InString == TEXT("TPri_SlightlyBelowNormal"))
+	{
+		return EThreadPriority::TPri_SlightlyBelowNormal;
+	}
+	if (InString == TEXT("TPri_TimeCritical"))
+	{
+		return EThreadPriority::TPri_TimeCritical;
+	}
+	return EThreadPriority::TPri_Num;
 }
 
 
-static void SetupThreadAffinity(const TArray<FString>& Args)
+FString ThreadPriorityToString(EThreadPriority InPriority)
 {
-	static bool LoadedDefaults = false;
-	if (!LoadedDefaults || (Args.Num() && Args[0] == TEXT("default")))
+	switch (InPriority)
 	{
-		LoadedDefaults = true;
-		TaskThread = FPlatformAffinity::GetTaskGraphThreadMask();
-		GameThread = FPlatformAffinity::GetMainGameMask();
-		RenderThread = FPlatformAffinity::GetRenderingThreadMask();
-		RHIThread = FPlatformAffinity::GetRHIThreadMask();
-	}
-	for (int32 Index = 0; Index + 1 < Args.Num(); Index += 2)
+
+	case EThreadPriority::TPri_Normal:
+		return FString(TEXT("TPri_Normal"));
+	case EThreadPriority::TPri_AboveNormal:
+		return  FString(TEXT("TPri_AboveNormal"));
+	case EThreadPriority::TPri_BelowNormal:
+		return  FString(TEXT("TPri_BelowNormal"));
+	case EThreadPriority::TPri_Highest:
+		return  FString(TEXT("TPri_Highest"));
+	case EThreadPriority::TPri_Lowest:
+		return  FString(TEXT("TPri_Lowest"));
+	case EThreadPriority::TPri_SlightlyBelowNormal:
+		return  FString(TEXT("TPri_SlightlyBelowNormal"));
+	case EThreadPriority::TPri_TimeCritical:
+		return  FString(TEXT("TPri_TimeCritical"));
+	case EThreadPriority::TPri_Num:
+	default:
+		break;
+	};
+	return  FString(TEXT("TPri_Undefined"));
+}
+
+
+void SetPriorityAndAffinityOnGameThread()
+{
+	if (GameThreadPriority != EThreadPriority::TPri_Num)
 	{
-		uint64 Aff = FParse::HexNumber(*Args[Index + 1]); // this is only 32 bits
-		if (!Aff)
-		{
-			UE_LOG(LogConsoleResponse, Display, TEXT("Parsed 0 for affinity, using 0xFFFFFFFFFFFFFFFF instead"));
-			Aff = 0xFFFFFFFFFFFFFFFF;
-		}
-		if (Args[Index] == TEXT("GT"))
-		{
-			GameThread = Aff;
-		}
-		else if (Args[Index] == TEXT("RT"))
-		{
-			RenderThread = Aff;
-		}
-		else if (Args[Index] == TEXT("RHI"))
-		{
-			RHIThread = Aff;
-		}
-		else if (Args[Index] == TEXT("Task"))
-		{
-			TaskThread = Aff;
-		}
-		else
-		{
-			UE_LOG(LogConsoleResponse, Display, TEXT("Unrecognized thread name %s"), *Args[Index]);
-		}
+		FPlatformProcess::SetThreadPriority(GameThreadPriority);
+		UE_LOG(LogConsoleResponse, Display, TEXT("GT Priority %s"), *ThreadPriorityToString(GameThreadPriority));
 	}
+	FPlatformProcess::SetThreadAffinityMask(GameThreadAffinity);
+	UE_LOG(LogConsoleResponse, Display, TEXT("GT Affinity     %016llX"), GameThreadAffinity);
+}
 
-	StallForTaskThread.Reset();
-	StallForTaskThread.Add(FTaskGraphInterface::Get().GetNumWorkerThreads());
-
-	for (int32 Index = 0; Index < FTaskGraphInterface::Get().GetNumWorkerThreads(); Index++)
+void SetPriorityAndAffinityOnRenderThread()
+{
+	if (RenderThreadPriority != EThreadPriority::TPri_Num)
 	{
-		FSimpleDelegateGraphTask::CreateAndDispatchWhenReady(
-			FSimpleDelegateGraphTask::FDelegate::CreateStatic(&SetAffinityOnThread),
-			TStatId(), NULL, ENamedThreads::AnyNormalThreadHiPriTask);
+		FPlatformProcess::SetThreadPriority(RenderThreadPriority);
+		UE_LOG(LogConsoleResponse, Display, TEXT("RT Priority %s"), *ThreadPriorityToString(RenderThreadPriority));
 	}
-	if (ENamedThreads::bHasHighPriorityThreads)
+	FPlatformProcess::SetThreadAffinityMask(RenderThreadAffinity);
+	UE_LOG(LogConsoleResponse, Display, TEXT("RT Affinity    %016llX"), RenderThreadAffinity);
+}
+
+void SetPriorityAndAffinityOnRHIThread()
+{
+	if (RHIThreadPriority != EThreadPriority::TPri_Num)
 	{
-		for (int32 Index = 0; Index < FTaskGraphInterface::Get().GetNumWorkerThreads(); Index++)
-		{
-			FSimpleDelegateGraphTask::CreateAndDispatchWhenReady(
-				FSimpleDelegateGraphTask::FDelegate::CreateStatic(&SetAffinityOnThread),
-				TStatId(), NULL, ENamedThreads::AnyHiPriThreadHiPriTask);
-		}
+		FPlatformProcess::SetThreadPriority(RHIThreadPriority);
+		UE_LOG(LogConsoleResponse, Display, TEXT("RHI Priority %s"), *ThreadPriorityToString(RHIThreadPriority));
 	}
-	if (ENamedThreads::bHasBackgroundThreads)
+	FPlatformProcess::SetThreadAffinityMask(RHIThreadAffinity);
+	UE_LOG(LogConsoleResponse, Display, TEXT("RHI Affinity  %016llX"), RHIThreadAffinity);
+}
+
+
+
+static void SetupThreadConfig(const TArray<FString>& Args)
+{
+	if (!GUseNewTaskBackend)
+		return;
+
+	GetDefaultPriorities((Args.Num() && Args[0] == TEXT("default")));
+	GetDefaultAffinities((Args.Num() && Args[0] == TEXT("default")));
+
+	if ((Args.Num() && Args[0] != TEXT("default")))
 	{
-		for (int32 Index = 0; Index < FTaskGraphInterface::Get().GetNumWorkerThreads(); Index++)
+		for (int32 Index = 0; Index + 1 < Args.Num(); )
 		{
-			FSimpleDelegateGraphTask::CreateAndDispatchWhenReady(
-				FSimpleDelegateGraphTask::FDelegate::CreateStatic(&SetAffinityOnThread),
-				TStatId(), NULL, ENamedThreads::AnyBackgroundHiPriTask);
+			int32 OriginalIndex = Index;
+			uint64 Aff = 0xffffffffffffffff;
+			EThreadPriority NewPriority = EThreadPriority::TPri_Num;
+
+			int32 Index2 = Index + 1;
+			for (; Index2 < Args.Num(); Index2++)
+			{
+				bool ValidArgFound = false;
+				// if the arg is a hex number we are setting affinity
+				if (Args[Index2].StartsWith(TEXT("0x")))
+				{
+					Aff = FParse::HexNumber(*Args[Index2]); // this is only 32 bits
+					ValidArgFound = true;
+				}
+				// if the arg starts with TPri we are setting priority
+				if (Args[Index2].StartsWith(TEXT("TPri")))
+				{
+					NewPriority = StringToThreadPriority(Args[Index2]);
+					ValidArgFound = true;
+				}
+				if (!ValidArgFound)
+				{
+					// we didn't find an expected arg so we'll move on 
+					break;
+				}
+			}
+			Index = Index2;
+
+			if (Aff != 0xffffffffffffffff)
+			{
+				if (Args[OriginalIndex] == TEXT("GT"))
+				{
+					GameThreadAffinity = Aff;
+				}
+				else if (Args[OriginalIndex] == TEXT("RT"))
+				{
+					RenderThreadAffinity = Aff;
+				}
+				else if (Args[OriginalIndex] == TEXT("RHI"))
+				{
+					RHIThreadAffinity = Aff;
+				}
+				else if (Args[OriginalIndex] == TEXT("Task"))
+				{
+					TaskThreadAffinity = Aff;
+				}
+				else if (Args[OriginalIndex] == TEXT("TaskBP"))
+				{
+					TaskThreadBPAffinity = Aff;
+				}
+			}
+
+			if (NewPriority != EThreadPriority::TPri_Num)
+			{
+				if (Args[OriginalIndex] == TEXT("GT"))
+				{
+					GameThreadPriority = NewPriority;
+				}
+				else if (Args[OriginalIndex] == TEXT("RT"))
+				{
+					RenderThreadPriority = NewPriority;
+				}
+				else if (Args[OriginalIndex] == TEXT("RHI"))
+				{
+					RHIThreadPriority = NewPriority;
+				}
+				else if (Args[OriginalIndex] == TEXT("Task"))
+				{
+					TaskThreadPriority = NewPriority;
+				}
+				else if (Args[OriginalIndex] == TEXT("TaskBP"))
+				{
+					TaskThreadBPPriority = NewPriority;
+				}
+			}
 		}
 	}
+
 	FSimpleDelegateGraphTask::CreateAndDispatchWhenReady(
-		FSimpleDelegateGraphTask::FDelegate::CreateStatic(&SetAffinityOnThread),
+		FSimpleDelegateGraphTask::FDelegate::CreateStatic(&SetPriorityAndAffinityOnRenderThread),
 		TStatId(), NULL, ENamedThreads::GetRenderThread());
+
 	if (IsRHIThreadRunning())
 	{
 		FSimpleDelegateGraphTask::CreateAndDispatchWhenReady(
-			FSimpleDelegateGraphTask::FDelegate::CreateStatic(&SetAffinityOnThread),
+			FSimpleDelegateGraphTask::FDelegate::CreateStatic(&SetPriorityAndAffinityOnRHIThread),
 			TStatId(), NULL, ENamedThreads::RHIThread);
 	}
 	check(IsInGameThread());
-	SetAffinityOnThread();
+	SetPriorityAndAffinityOnGameThread();
 	FlushRenderingCommands();
+	int32 NumWorkerThreads = FPlatformMisc::NumberOfWorkerThreadsToSpawn();
+	int32 NumBackgroundWorkers = FMath::Max(1, NumWorkerThreads - FMath::Min<int>(2, NumWorkerThreads));
+	int32 NumForegroundWorkers = FMath::Max(1, NumWorkerThreads - NumBackgroundWorkers);
+	LowLevelTasks::FScheduler::Get().RestartWorkers(NumForegroundWorkers, NumBackgroundWorkers, FForkProcessHelper::IsForkedMultithreadInstance(), TaskThreadPriority, TaskThreadBPPriority, TaskThreadAffinity, TaskThreadBPAffinity);
+
+	UE_LOG(LogConsoleResponse, Display, TEXT("Task Priority %s"), *ThreadPriorityToString(TaskThreadPriority));
+	UE_LOG(LogConsoleResponse, Display, TEXT("Task Affinity    %016llX"), TaskThreadAffinity);
+
+	UE_LOG(LogConsoleResponse, Display, TEXT("TaskBP Priority %s"), *ThreadPriorityToString(TaskThreadBPPriority));
+	UE_LOG(LogConsoleResponse, Display, TEXT("TaskBP Affinity    %016llX"), TaskThreadBPAffinity);
+
 	GLog->FlushThreadedLogs();
 }
 
-static FAutoConsoleCommand SetupThreadAffinityCmd(
-	TEXT("SetThreadAffinity"),
-	TEXT("Sets the thread affinity. A single arg of default resets the thread affinity, otherwise pairs of args [GT|RT|RHI|Task] [Hex affinity] sets the affinity."),
-	FConsoleCommandWithArgsDelegate::CreateStatic(&SetupThreadAffinity)
+// example string to pass SetThreadConfig Task 0x3f TPri_Normal TaskBP 0x3f TPri_BelowNormal GT 0x1 TPri_Highest RT 0x10 TPri_Highest RHI TPri_Highest
+static FAutoConsoleCommand SetupThreadConfigCmd(
+	TEXT("SetThreadConfig"),
+	TEXT("Sets a thread Priority and/or Affinity. A single arg of default resets all the thread Priorities and Affinities, otherwise [GT|RT|RHI|Task|TaskBP] both/either [TPri_type] [0x] sets the Priority and/or Affinity."),
+	FConsoleCommandWithArgsDelegate::CreateStatic(&SetupThreadConfig)
 );
 
+void ThreadConfigStringChanged(IConsoleVariable* Variable)
+{
+	TArray<FString> Args;
+	CVar_ThreadConfigString.ParseIntoArray(Args, TEXT(" "), true);
+	SetupThreadConfig(Args);
+}
+
+static FAutoConsoleVariableRef CVarThreadConfigString(
+	TEXT("ThreadConfig"),
+	CVar_ThreadConfigString,
+	TEXT("Configures the Thread Priorites and Affinites of GT/RT/RHI/Task/TaskBP"),
+	FConsoleVariableDelegate::CreateStatic(&ThreadConfigStringChanged),
+	ECVF_Default
+);
+
+
 #if !UE_BUILD_SHIPPING
 
 static void PakFileTest(const TArray<FString>& Args)
