Small amelioration of Spherical sector debug display

diff --git a/Engine/Source/Runtime/Engine/Public/WorldPartition/WorldPartitionStreamingSource.h b/Engine/Source/Runtime/Engine/Public/WorldPartition/WorldPartitionStreamingSource.h
--- a/Engine/Source/Runtime/Engine/Public/WorldPartition/WorldPartitionStreamingSource.h
+++ b/Engine/Source/Runtime/Engine/Public/WorldPartition/WorldPartitionStreamingSource.h
@@ -75,14 +75,18 @@ public:
 		const float AngleStep = Angle / float(SegmentCount);
 		const FRotator ShapeRotation = FRotationMatrix::MakeFromX(Axis).Rotator();
 		const FVector ScaledAxis = FVector::ForwardVector * Radius;
+		const int32 RollCount = 16;
 
-		FVector LastRollEndVector;
-		for (int i = 0; i <= 64; ++i)
+		Segments.Reserve(2 * (RollCount + 1) * (SegmentCount + 2));
+		int32 LastArcStartIndex = -1;
+		for (int i = 0; i <= RollCount; ++i)
 		{
-			const float Roll = 360.0f * i / 64.0f;
+			const float Roll = 180.f * i / float(RollCount);
 			const FTransform Transform(FRotator(0, 0, Roll) + ShapeRotation, Center);
 			FVector SegmentStart = Transform.TransformPosition(FRotator(0, -0.5f * Angle, 0).RotateVector(ScaledAxis));
 			Segments.Emplace(Center, SegmentStart);
+			int32 CurrentArcStartIndex = Segments.Num();
+			// Build sector arc
 			for (int32 j = 1; j <= SegmentCount; j++)
 			{
 				FVector SegmentEnd = Transform.TransformPosition(FRotator(0, -0.5f * Angle + (AngleStep * j), 0).RotateVector(ScaledAxis));
@@ -92,9 +96,14 @@ public:
 			Segments.Emplace(Center, SegmentStart);
 			if (i > 0)
 			{
-				Segments.Emplace(SegmentStart, LastRollEndVector);
+				// Connect sector arc to previous arc
+				for (int32 j = 0; j < SegmentCount; j++)
+				{
+					Segments.Emplace(Segments[LastArcStartIndex + j].Key, Segments[CurrentArcStartIndex + j].Key);
+				}
+				Segments.Emplace(Segments[LastArcStartIndex + SegmentCount - 1].Value, Segments[CurrentArcStartIndex + SegmentCount - 1].Value);
 			}
-			LastRollEndVector = SegmentStart;
+			LastArcStartIndex = CurrentArcStartIndex;
 		}
 		return Segments;
 	}
