Lumen post process volume configuration inconsistency between camera settings.

The copy function of FPostProcessSettings does not copy the GI variable. Find out that the LumenFinalGatherQuality is also not copied. With optimization off, the result is consistent, but will not override. When the optimization flag is on, the struct memory can be arbitrary.

diff --git a/Engine/Source/Runtime/Engine/Private/Scene.cpp b/Engine/Source/Runtime/Engine/Private/Scene.cpp
--- a/Engine/Source/Runtime/Engine/Private/Scene.cpp
+++ b/Engine/Source/Runtime/Engine/Private/Scene.cpp
@@ -790,6 +790,8 @@ FPostProcessSettings::FPostProcessSettings(const FPostProcessSettings& Settings)
 	, bOverride_RayTracingTranslucencySamplesPerPixel(Settings.bOverride_RayTracingTranslucencySamplesPerPixel)
 	, bOverride_RayTracingTranslucencyShadows(Settings.bOverride_RayTracingTranslucencyShadows)
 	, bOverride_RayTracingTranslucencyRefraction(Settings.bOverride_RayTracingTranslucencyRefraction)
+	, bOverride_DynamicGlobalIlluminationMethod(Settings.bOverride_DynamicGlobalIlluminationMethod)
+    , bOverride_LumenFinalGatherQuality(Settings.bOverride_LumenFinalGatherQuality)
 	, bOverride_RayTracingGI(Settings.bOverride_RayTracingGI)
 	, bOverride_RayTracingGIMaxBounces(Settings.bOverride_RayTracingGIMaxBounces)
 	, bOverride_RayTracingGISamplesPerPixel(Settings.bOverride_RayTracingGISamplesPerPixel)
