Fixed bug in world partition cell streaming when no streaming source are added but data layers are loaded/activated.

diff --git a/Engine/Source/Runtime/Engine/Private/WorldPartition/WorldPartitionStreamingPolicy.cpp b/Engine/Source/Runtime/Engine/Private/WorldPartition/WorldPartitionStreamingPolicy.cpp
--- a/Engine/Source/Runtime/Engine/Private/WorldPartition/WorldPartitionStreamingPolicy.cpp
+++ b/Engine/Source/Runtime/Engine/Private/WorldPartition/WorldPartitionStreamingPolicy.cpp
@@ -218,12 +218,6 @@ void UWorldPartitionStreamingPolicy::UpdateStreamingState()
 	const ENetMode NetMode = World->GetNetMode();
 	if (NetMode == NM_Standalone || NetMode == NM_Client || AWorldPartitionReplay::IsEnabled(World))
 	{
-		// Early out if nothing loaded and no streaming source
-		if ((StreamingSources.Num() == 0) && (ActivatedCells.Num() == 0) && (LoadedCells.Num() == 0))
-		{
-			return;
-		}
-				
 		// When uninitializing, UpdateStreamingState is called, but we don't want any cells to be loaded
 		if (WorldPartition->IsInitialized())
 		{
