Make sure to keep HLOD layers referenced during HLOD generation to avoid them getting GC'd while still referenced by non-UObject structs.

diff --git a/Engine/Source/Runtime/Engine/Private/WorldPartition/RuntimeSpatialHash/RuntimeSpatialHashHLOD.cpp b/Engine/Source/Runtime/Engine/Private/WorldPartition/RuntimeSpatialHash/RuntimeSpatialHashHLOD.cpp
--- a/Engine/Source/Runtime/Engine/Private/WorldPartition/RuntimeSpatialHash/RuntimeSpatialHashHLOD.cpp
+++ b/Engine/Source/Runtime/Engine/Private/WorldPartition/RuntimeSpatialHash/RuntimeSpatialHashHLOD.cpp
@@ -483,7 +483,9 @@ bool UWorldPartitionRuntimeSpatialHash::GenerateHLOD(ISourceControlHelper* Sourc
 	UWorldPartition* WorldPartition = GetOuterUWorldPartition();
 
 	// Find all used HLOD grids & build a dependency graph
-	TMap<FName, FSpatialHashRuntimeGrid> HLODGrids;
+	check(HLODGrids.IsEmpty());
+	ON_SCOPE_EXIT { HLODGrids.Empty(); };
+
 	TMap<FName, TSet<FName>>			 HLODGridsGraph;
 	GatherHLODGrids(WorldPartition, HLODGrids, HLODGridsGraph);
 
diff --git a/Engine/Source/Runtime/Engine/Public/WorldPartition/WorldPartitionRuntimeSpatialHash.h b/Engine/Source/Runtime/Engine/Public/WorldPartition/WorldPartitionRuntimeSpatialHash.h
--- a/Engine/Source/Runtime/Engine/Public/WorldPartition/WorldPartitionRuntimeSpatialHash.h
+++ b/Engine/Source/Runtime/Engine/Public/WorldPartition/WorldPartitionRuntimeSpatialHash.h
@@ -274,6 +274,9 @@ private:
 	UPROPERTY(Transient)
 	mutable FWorldPartitionRuntimeSpatialHashGridPreviewer GridPreviewer;
 
+	UPROPERTY(Transient)
+	TMap<FName, FSpatialHashRuntimeGrid> HLODGrids;
+
 	TMap<FString, UWorldPartitionRuntimeCell*> PackagesToGenerateForCook;
 #endif
 
