Change RestoreRemappedChannel from recursive to iterative
Prevent loops from being generated in the remapped channels

diff --git a/Engine/Source/Runtime/Engine/Private/NetConnection.cpp b/Engine/Source/Runtime/Engine/Private/NetConnection.cpp
--- a/Engine/Source/Runtime/Engine/Private/NetConnection.cpp
+++ b/Engine/Source/Runtime/Engine/Private/NetConnection.cpp
@@ -2960,62 +2960,49 @@ void UNetConnection::ReceivedPacket( FBitReader& Reader, bool bIsReinjectedPacke
 
 void UNetConnection::RestoreRemappedChannel(const int32 ChIndex)
 {
+	TArray<int32, TInlineAllocator<16>> RemapChain;
+
 	if (ChannelIndexMap.Contains(ChIndex))
 	{
-		const int32 RemappedIndex = ChannelIndexMap[ChIndex];
+		RemapChain.Add(ChIndex);
 
-		if (UChannel* Channel = Channels[ChIndex])
+		while (const int32* ChainIndex = ChannelIndexMap.FindKey(RemapChain[RemapChain.Num() - 1]))
 		{
-			// It is possible the index we want to swap back to wasn't cleaned up because the channel was marked broken by backwards compatibility
-			if (Channel->Broken)
+			const int32 FoundIdx = RemapChain.Find(*ChainIndex);
+			
+			if (!ensureMsgf(FoundIdx == INDEX_NONE, TEXT("RestoreRemappedChannel: Loop detected in channel remaps!")))
 			{
-				if (UActorChannel* ActorChannel = Cast<UActorChannel>(Channel))
-				{
-					// look for a queued close bunch
-					for (FInBunch* InBunch : ActorChannel->QueuedBunches)
-					{
-						if (InBunch && InBunch->bClose)
-						{
-							UE_LOG(LogNet, Warning, TEXT("SetAllowExistingChannelIndex:  Cleaning up broken channel: %s"), *ActorChannel->Describe());
-							ActorChannel->ConditionalCleanUp(true, InBunch->CloseReason);
-							break;
-						}
-					}
-
-					// broken channel but the actor never spawned, should be safe to remove
-					if (ActorChannel->Actor == nullptr)
-					{
-						UE_LOG(LogNet, Warning, TEXT("SetAllowExistingChannelIndex:  Cleaning up broken channel: %s"), *ActorChannel->Describe());
-						ActorChannel->ConditionalCleanUp(true, EChannelCloseReason::Destroyed);
-					}
-				}
+				break;
 			}
 
-			if (Channels[ChIndex] != nullptr)
-			{
-				// this map be part of a chain of remapped channel indices
-				if (const int32* Key = ChannelIndexMap.FindKey(ChIndex))
-				{
-					RestoreRemappedChannel(*Key);
-				}
-			}
+			RemapChain.Add(*ChainIndex);
 		}
+	}
+
+	for (int32 i = RemapChain.Num() - 1; i >= 0; --i)
+	{
+		const int32 SourceIndex = RemapChain[i];
+
+		if (ChannelIndexMap.Contains(SourceIndex))
+		{
+			const int32 RemappedIndex = ChannelIndexMap[SourceIndex];
 
-		UChannel* SrcChannel = Channels[RemappedIndex];
-		UChannel* DstChannel = Channels[ChIndex];
+			UChannel* SrcChannel = Channels[RemappedIndex];
+			UChannel* DstChannel = Channels[SourceIndex];
 
 			// this channel should still exist, but the location we want to swap it back to should be empty
-		if (ensureMsgf(SrcChannel && !DstChannel, TEXT("Source should exist: [%s] Destination should be null: [%s]"), SrcChannel ? *SrcChannel->Describe() : TEXT("null"), DstChannel ? *DstChannel->Describe() : TEXT("null")))
+			if (ensureMsgf(SrcChannel && !DstChannel, TEXT("Source should exist: [%s] Destination should be null: [%s]"), SrcChannel ? *SrcChannel->Describe() : TEXT("null"), DstChannel ? *DstChannel->Describe() : TEXT("null")))
 			{
-			SrcChannel->ChIndex = ChIndex;
+				SrcChannel->ChIndex = SourceIndex;
 
-			Swap(Channels[ChIndex], Channels[RemappedIndex]);
-			Swap(InReliable[ChIndex], InReliable[RemappedIndex]);
+				Swap(Channels[SourceIndex], Channels[RemappedIndex]);
+				Swap(InReliable[SourceIndex], InReliable[RemappedIndex]);
 
-			ChannelIndexMap.Remove(ChIndex);
+				ChannelIndexMap.Remove(SourceIndex);
 			}
 		}
 	}
+}
 
 void UNetConnection::SetAllowExistingChannelIndex(bool bAllow)
 {
@@ -3031,6 +3018,41 @@ void UNetConnection::SetAllowExistingChannelIndex(bool bAllow)
 		TArray<int32> RemappedKeys;
 		ChannelIndexMap.GenerateKeyArray(RemappedKeys);
 
+		// clean up any broken channels first
+		for (const int32 ChIndex : RemappedKeys)
+		{
+			if (UChannel* Channel = Channels[ChIndex])
+			{
+				// It is possible the index we want to swap back to wasn't cleaned up because the channel was marked broken by backwards compatibility
+				if (Channel->Broken)
+				{
+					if (UActorChannel* ActorChannel = Cast<UActorChannel>(Channel))
+					{
+						// look for a queued close bunch
+						for (FInBunch* InBunch : ActorChannel->QueuedBunches)
+						{
+							if (InBunch && InBunch->bClose)
+							{
+								UE_LOG(LogNet, Warning, TEXT("SetAllowExistingChannelIndex:  Cleaning up broken channel: %s"), *ActorChannel->Describe());
+								ActorChannel->ConditionalCleanUp(true, InBunch->CloseReason);
+								break;
+							}
+						}
+
+						// broken channel but the actor never spawned, should be safe to remove
+						if (ActorChannel->Actor == nullptr)
+						{
+							UE_LOG(LogNet, Warning, TEXT("SetAllowExistingChannelIndex:  Cleaning up broken channel: %s"), *ActorChannel->Describe());
+							ActorChannel->ConditionalCleanUp(true, EChannelCloseReason::Destroyed);
+						}
+					}
+				}
+			}
+		}
+
+		ChannelIndexMap.GenerateKeyArray(RemappedKeys);
+
+		// now try to remap the channels
 		for (const int32 Key : RemappedKeys)
 		{
 			RestoreRemappedChannel(Key);
@@ -3346,8 +3368,9 @@ int32 UNetConnection::GetFreeChannelIndex(const FName& ChName) const
 	for (ChIndex = FirstChannel; ChIndex < Channels.Num(); ChIndex++)
 	{
 		const bool bIgnoreReserved = bIgnoreReservedChannels && ReservedChannels.Contains(ChIndex);
+		const bool bIgnoreRemapped = bAllowExistingChannelIndex && ChannelIndexMap.Contains(ChIndex);
 
-		if (!Channels[ChIndex] && !bIgnoreReserved)
+		if (!Channels[ChIndex] && !bIgnoreReserved && !bIgnoreRemapped)
 		{
 			break;
 		}
