Show the WorldDataLayers actor in the world outliner, so users can execute actions on it such as SCC file history, etc.
Also make sure it display no properties the user can edit, and that it's impossible to delete through the world outliner.

diff --git a/Engine/Source/Runtime/Engine/Private/WorldPartition/DataLayer/WorldDataLayers.cpp b/Engine/Source/Runtime/Engine/Private/WorldPartition/DataLayer/WorldDataLayers.cpp
--- a/Engine/Source/Runtime/Engine/Private/WorldPartition/DataLayer/WorldDataLayers.cpp
+++ b/Engine/Source/Runtime/Engine/Private/WorldPartition/DataLayer/WorldDataLayers.cpp
@@ -238,7 +238,6 @@ AWorldDataLayers* AWorldDataLayers::Create(UWorld* World)
 	{
 		FActorSpawnParameters SpawnParams;
 		SpawnParams.OverrideLevel = World->PersistentLevel;
-		SpawnParams.bHideFromSceneOutliner = true;
 		SpawnParams.Name = WorldDataLayersName;
 		SpawnParams.NameMode = FActorSpawnParameters::ESpawnActorNameMode::Required_Fatal;
 		WorldDataLayers = World->SpawnActor<AWorldDataLayers>(AWorldDataLayers::StaticClass(), SpawnParams);
@@ -436,6 +435,8 @@ void AWorldDataLayers::PostLoad()
 			DataLayer->SetIsDynamicallyLoadedInEditor(true);
 }
 	}
+
+	bListedInSceneOutliner = true;
 #else
 	// Build acceleration tables
 	for (const UDataLayer* DataLayer : WorldDataLayers)
diff --git a/Engine/Source/Runtime/Engine/Public/WorldPartition/DataLayer/WorldDataLayers.h b/Engine/Source/Runtime/Engine/Public/WorldPartition/DataLayer/WorldDataLayers.h
--- a/Engine/Source/Runtime/Engine/Public/WorldPartition/DataLayer/WorldDataLayers.h
+++ b/Engine/Source/Runtime/Engine/Public/WorldPartition/DataLayer/WorldDataLayers.h
@@ -13,7 +13,7 @@ class UDataLayer;
 /**
  * Actor containing all data layers for a world
  */
-UCLASS(hidecategories = (Actor, Advanced, Display, Events, Object, Attachment, Info, Input, Blueprint, Layers, Tags, Replication), notplaceable)
+UCLASS(hidecategories = (Actor, HLOD, Cooking, Transform, Advanced, Display, Events, Object, Attachment, Info, Input, Blueprint, Layers, Tags, Replication), notplaceable)
 class ENGINE_API AWorldDataLayers : public AInfo
 {
 	GENERATED_UCLASS_BODY()
@@ -22,7 +22,10 @@ public:
 	virtual void PostLoad() override;
 	virtual void RewindForReplay() override;
 	virtual void BeginPlay() override;
+
 #if WITH_EDITOR
+	virtual bool IsUserManaged() const override { return false; }
+
 	static AWorldDataLayers* Create(UWorld* World);
 	UDataLayer* CreateDataLayer(FName InName = TEXT("DataLayer"), EObjectFlags InObjectFlags = RF_NoFlags);
 	bool RemoveDataLayer(UDataLayer* InDataLayer);
