Animation Sequence data model fixes
    - Guard against Needs_Load flag to ensure we only create a UAnimationDataModel instance for new UAnimSequenceBase's
    - Fix incorrect (inverted) boolean value
    - Extracted out DataModel event binding to function (rule of three)
    - Using NewObject rather than CreateDefaultSubObject, as we are not interested in any of the DefaultSubObjects behaviour - and it only caused issues with -Game Editor builds and the fixed upgrade path (creating model outside of Constructor call path)

diff --git a/Engine/Source/Runtime/Engine/Classes/Animation/AnimSequenceBase.h b/Engine/Source/Runtime/Engine/Classes/Animation/AnimSequenceBase.h
--- a/Engine/Source/Runtime/Engine/Classes/Animation/AnimSequenceBase.h
+++ b/Engine/Source/Runtime/Engine/Classes/Animation/AnimSequenceBase.h
@@ -295,6 +295,9 @@ protected:
 
 	/** Validates that DataModel contains a valid UAnimDataModel object */
 	void ValidateModel() const;
+	
+	/** Binds to DataModel its modification delegate */
+	void BindToModelModificationEvent();
 
 	/** Replaces the current DataModel, if any, with the provided one */
 	void CopyDataModel(const UAnimDataModel* ModelToDuplicate);
diff --git a/Engine/Source/Runtime/Engine/Private/Animation/AnimSequence.cpp b/Engine/Source/Runtime/Engine/Private/Animation/AnimSequence.cpp
--- a/Engine/Source/Runtime/Engine/Private/Animation/AnimSequence.cpp
+++ b/Engine/Source/Runtime/Engine/Private/Animation/AnimSequence.cpp
@@ -439,7 +439,7 @@ UAnimSequence::UAnimSequence(const FObjectInitializer& ObjectInitializer)
 	bAllowFrameStripping = true;
 	CompressionErrorThresholdScale = 1.f;
 
-	if (!HasAllFlags(EObjectFlags::RF_ClassDefaultObject))
+	if (!HasAnyFlags(EObjectFlags::RF_ClassDefaultObject| EObjectFlags::RF_NeedLoad))
 	{
 		check(DataModel);
 	}
diff --git a/Engine/Source/Runtime/Engine/Private/Animation/AnimSequenceBase.cpp b/Engine/Source/Runtime/Engine/Private/Animation/AnimSequenceBase.cpp
--- a/Engine/Source/Runtime/Engine/Private/Animation/AnimSequenceBase.cpp
+++ b/Engine/Source/Runtime/Engine/Private/Animation/AnimSequenceBase.cpp
@@ -39,7 +39,7 @@ UAnimSequenceBase::UAnimSequenceBase(const FObjectInitializer& ObjectInitializer
 #endif // WITH_EDITORONLY_DATA
 {
 #if WITH_EDITOR
-	if (!HasAllFlags(EObjectFlags::RF_ClassDefaultObject))
+	if (!HasAnyFlags(EObjectFlags::RF_ClassDefaultObject| EObjectFlags::RF_NeedLoad))
 	{
 		CreateModel();
 		GetController();
@@ -70,19 +70,19 @@ void UAnimSequenceBase::PostLoad()
 #if WITH_EDITORONLY_DATA
 	if (!HasAnyFlags(EObjectFlags::RF_ClassDefaultObject))
 	{
-		const bool bRequiresModelCreation = DataModel != nullptr;
+		const bool bRequiresModelCreation = DataModel == nullptr;
 		const bool bRequiresModelPopulation = GetLinkerCustomVersion(FUE5MainStreamObjectVersion::GUID) < FUE5MainStreamObjectVersion::IntroducingAnimationDataModel;
 		checkf(bRequiresModelPopulation || DataModel != nullptr, TEXT("Invalid Animation Sequence base state, no data model found past upgrade object version"));
 
 		// Construct a new UAnimDataModel instance
-		if (bRequiresModelCreation)
+		if(bRequiresModelCreation)
 		{
 			CreateModel();
 		}
-		else
-		{
-			DataModel->GetModifiedEvent().AddUObject(this, &UAnimSequenceBase::OnModelModified);
-		}
+
+		ValidateModel();
+		GetController();
+		BindToModelModificationEvent();
 
 		if (USkeleton* MySkeleton = GetSkeleton())
 		{
@@ -209,8 +209,7 @@ void UAnimSequenceBase::PostDuplicate(EDuplicateMode::Type DuplicateMode)
 #if WITH_EDITOR
 	if (ensure(DataModel))
 	{
-		DataModel->GetModifiedEvent().RemoveAll(this);
-		DataModel->GetModifiedEvent().AddUObject(this, &UAnimSequenceBase::OnModelModified);
+		BindToModelModificationEvent();
 	}
 #endif // WITH_EDITOR
 }
@@ -1222,7 +1221,7 @@ IAnimationDataController& UAnimSequenceBase::GetController()
 		Controller->SetModel(DataModel);
 	}
 
-	ensure(Controller->GetModel() == DataModel);
+	ensureAlways(Controller->GetModel() == DataModel);
 
 	return *Controller;
 }
@@ -1251,9 +1250,16 @@ void UAnimSequenceBase::ValidateModel() const
 	checkf(DataModel != nullptr, TEXT("Invalidate data model"));
 }
 
+void UAnimSequenceBase::BindToModelModificationEvent()
+{
+	ValidateModel();
+	DataModel->GetModifiedEvent().RemoveAll(this);
+	DataModel->GetModifiedEvent().AddUObject(this, &UAnimSequenceBase::OnModelModified);
+}
+
 void UAnimSequenceBase::CopyDataModel(const UAnimDataModel* ModelToDuplicate)
 {
-	checkf(ModelToDuplicate != nullptr, TEXT("Invalidate data model"));
+	checkf(ModelToDuplicate != nullptr, TEXT("Invalidate data model %s"), *GetFullName());
 	if (DataModel)
 	{
 		DataModel->GetModifiedEvent().RemoveAll(this);
@@ -1261,18 +1267,16 @@ void UAnimSequenceBase::CopyDataModel(const UAnimDataModel* ModelToDuplicate)
 
 	DataModel = DuplicateObject(ModelToDuplicate, this);
 	Controller->SetModel(DataModel);
-	DataModel->GetModifiedEvent().AddUObject(this, &UAnimSequenceBase::OnModelModified);
+
+	BindToModelModificationEvent();
 }
 
 void UAnimSequenceBase::CreateModel()
-{	
-	if (DataModel == nullptr)
-	{
-		DataModel = CreateDefaultSubobject<UAnimDataModel>(FName(TEXT("AnimationDataModel")));
-	}
-	
-	DataModel->GetModifiedEvent().RemoveAll(this);
-	DataModel->GetModifiedEvent().AddUObject(this, &UAnimSequenceBase::OnModelModified);
+{
+	checkf(DataModel == nullptr, TEXT("Invalid attempt to override the existing data model %s"), *GetFullName());
+	DataModel = NewObject<UAnimDataModel>(this, FName(TEXT("AnimationDataModel")));
+		
+	BindToModelModificationEvent();
 }
 #endif // WITH_EDITOR
 
