Asset streaming :
- New computation for the load order priority when LODs are streamed out.
- Cleaned up MemoryUsed and TempMemoryUsed being passed along different functions (not required anymore)
- Changed the streaming logic so that we can issue both StreamIn and StreamOut using the same temporary memory.
This assumes that stream out operations are faster and that they will generally complete before IO operations.
This prevents a usecase where we would perform StreamIns before any StreamOuts and become overbuget
(because the load order was higher for StreamIn)

diff --git a/Engine/Source/Runtime/Engine/Private/Streaming/AsyncTextureStreaming.cpp b/Engine/Source/Runtime/Engine/Private/Streaming/AsyncTextureStreaming.cpp
--- a/Engine/Source/Runtime/Engine/Private/Streaming/AsyncTextureStreaming.cpp
+++ b/Engine/Source/Runtime/Engine/Private/Streaming/AsyncTextureStreaming.cpp
@@ -507,7 +507,7 @@ void FRenderAssetStreamingMipCalcTask::TryKeepMips(TArray<int32>& PrioritizedRen
 	}
 }
 
-void FRenderAssetStreamingMipCalcTask::UpdateBudgetedMips_Async(int64& MemoryUsed, int64& TempMemoryUsed)
+void FRenderAssetStreamingMipCalcTask::UpdateBudgetedMips_Async()
 {
 	//*************************************
 	// Update Budget
@@ -525,8 +525,7 @@ void FRenderAssetStreamingMipCalcTask::UpdateBudgetedMips_Async(int64& MemoryUse
 	int64 MemoryBudgeted = 0;
 	int64 MeshMemoryBudgeted = 0;
 	int64 MemoryUsedByNonTextures = 0;
-	MemoryUsed = 0;
-	TempMemoryUsed = 0;
+	int64 MemoryUsed = 0;
 
 	for (FStreamingRenderAsset& StreamingRenderAsset : StreamingRenderAssets)
 	{
@@ -547,11 +546,6 @@ void FRenderAssetStreamingMipCalcTask::UpdateBudgetedMips_Async(int64& MemoryUse
 			MemoryUsedByNonTextures += AssetMemUsed;
 			++NumMeshes;
 		}
-
-		if (StreamingRenderAsset.ResidentMips != StreamingRenderAsset.RequestedMips)
-		{
-			TempMemoryUsed += StreamingRenderAsset.GetSize(StreamingRenderAsset.RequestedMips);
-		}
 	}
 
 	//*************************************
@@ -781,72 +775,88 @@ void FRenderAssetStreamingMipCalcTask::UpdateBudgetedMips_Async(int64& MemoryUse
 #endif
 }
 
-void FRenderAssetStreamingMipCalcTask::UpdateLoadAndCancelationRequests_Async(int64 MemoryUsed, int64 TempMemoryUsed)
+void FRenderAssetStreamingMipCalcTask::UpdateLoadAndCancelationRequests_Async()
 {
 	TArray<FStreamingRenderAsset>& StreamingRenderAssets = StreamingManager.StreamingRenderAssets;
 	const FRenderAssetStreamingSettings& Settings = StreamingManager.Settings;
 
+	LoadRequests.Empty();
+	CancelationRequests.Empty();
+
+	int64 StreamOutMemoryBudget = TempMemoryBudget;
+	int64 StreamInMemoryBudget = TempMemoryBudget;
+
 	TArray<int32> PrioritizedRenderAssets;
 	PrioritizedRenderAssets.Empty(StreamingRenderAssets.Num());
 	for (int32 AssetIndex = 0; AssetIndex < StreamingRenderAssets.Num() && !IsAborted(); ++AssetIndex)
 	{
-		FStreamingRenderAsset& StreamingTexture = StreamingRenderAssets[AssetIndex];
-		if (StreamingTexture.UpdateLoadOrderPriority_Async(Settings.MinMipForSplitRequest))
+		FStreamingRenderAsset& StreamingRenderAsset = StreamingRenderAssets[AssetIndex];
+
+		// If we need to change the number of resident mips.
+		if (StreamingRenderAsset.UpdateLoadOrderPriority_Async(Settings.MinMipForSplitRequest))
+		{
+			// If there is no pending update, kick one if the budget allows it.
+			if (StreamingRenderAsset.RequestedMips == StreamingRenderAsset.ResidentMips)
+			{
+				PrioritizedRenderAssets.Add(AssetIndex);
+			}
+			// Otherwise, if the update is trying to load to many mips, or if it is trying to unload required mips, (try to) cancel it.
+			else if (StreamingRenderAsset.RequestedMips > FMath::Max<int32>(StreamingRenderAsset.ResidentMips, StreamingRenderAsset.WantedMips + 1) ||
+					StreamingRenderAsset.RequestedMips < FMath::Min<int32>(StreamingRenderAsset.ResidentMips, StreamingRenderAsset.WantedMips))
+			{
+				CancelationRequests.Add(AssetIndex);
+			}
+		}
+
+		// Reduce the stream in/out budgets from pending updates.
+		const int64 TempMemoryUsed = StreamingRenderAsset.GetSize(StreamingRenderAsset.RequestedMips);
+		if (StreamingRenderAsset.RequestedMips < StreamingRenderAsset.ResidentMips)
+		{
+			// Here we assume that the stream out complete before new stream in requests start, so it doesn't affect stream in budget.
+			StreamOutMemoryBudget -= TempMemoryUsed;
+		}
+		else if (StreamingRenderAsset.RequestedMips > StreamingRenderAsset.ResidentMips)
 		{
-			PrioritizedRenderAssets.Add(AssetIndex);
+			// If there is a pending stream in, remove the temporary memory from both stream in and stream out budget.
+			// When the request was made, there were possibly stream out issued at the same time to free memory in case of budget limit.
+			StreamInMemoryBudget -= TempMemoryUsed;
+			StreamOutMemoryBudget -= TempMemoryUsed;
 		}
 	}
+
 	PrioritizedRenderAssets.Sort(FCompareRenderAssetByLoadOrderPriority(StreamingRenderAssets));
 
-	LoadRequests.Empty();
-	CancelationRequests.Empty();
+	// If possible, free as much memory with stream out operations, as will be required with new stream in requests.
+	// This prevents becoming overbudget momentarily, when we are already at the budget limit.
+	TArray<int32> StreamOutRequests;
+	TArray<int32> StreamInRequests;
 
 	// Now fill in the LoadRequest and CancelationRequests
 	for (int32 PriorityIndex = 0; PriorityIndex < PrioritizedRenderAssets.Num() && !IsAborted(); ++PriorityIndex)
 	{
 		int32 AssetIndex = PrioritizedRenderAssets[PriorityIndex];
 		FStreamingRenderAsset& StreamingRenderAsset = StreamingRenderAssets[AssetIndex];
+		// This assumes that the assets are streamed through a copy of all LODs.
+		// Even though this is only the case for non partially resident textures, 
+		// we still use this metric to limit the number of pending streaming requests.
+		const int64 TempMemoryRequired = StreamingRenderAsset.GetSize(StreamingRenderAsset.WantedMips);
 
-		// If there is a pending update with no cancelation request
-		if (StreamingRenderAsset.RequestedMips != StreamingRenderAsset.ResidentMips)
-		{
-			// If there is a pending load that attempts to load unrequired data (by at least 2 mips), 
-			// or if there is a pending unload that attempts to unload required data, try to cancel it.
-			if (StreamingRenderAsset.RequestedMips > FMath::Max<int32>(StreamingRenderAsset.ResidentMips, StreamingRenderAsset.WantedMips + 1 ) ||
-				StreamingRenderAsset.RequestedMips < FMath::Min<int32>(StreamingRenderAsset.ResidentMips, StreamingRenderAsset.WantedMips ))
-			{
-				CancelationRequests.Add(AssetIndex);
-			}
-		}
-		else if (StreamingRenderAsset.WantedMips < StreamingRenderAsset.ResidentMips && TempMemoryUsed < TempMemoryBudget)
+		// Check whether the budget allows the update, with the exception of always allowing a single update of any size (otherwise completion might never happen).
+		if (StreamingRenderAsset.WantedMips < StreamingRenderAsset.ResidentMips && (TempMemoryRequired <= StreamOutMemoryBudget || !StreamOutRequests.Num()))
 		{
-			const int64 TempMemoryRequired = StreamingRenderAsset.GetSize(StreamingRenderAsset.WantedMips);
-			const int64 UsedMemoryRequired = StreamingRenderAsset.GetSize(StreamingRenderAsset.WantedMips) - StreamingRenderAsset.GetSize(StreamingRenderAsset.ResidentMips);
-
-			// Respect the temporary budget unless this is the first unload request. This allows a single mip update of any size.
-			if (TempMemoryUsed + TempMemoryRequired <= TempMemoryBudget || !LoadRequests.Num())
-			{
-				LoadRequests.Add(AssetIndex);
-	
-				MemoryUsed -= UsedMemoryRequired;
-				TempMemoryUsed += TempMemoryRequired;
-			}
+			StreamOutRequests.Add(AssetIndex);
+			StreamOutMemoryBudget -= TempMemoryRequired;
 		}
-		else if (StreamingRenderAsset.WantedMips > StreamingRenderAsset.ResidentMips && TempMemoryUsed < TempMemoryBudget)
+		else if (StreamingRenderAsset.WantedMips > StreamingRenderAsset.ResidentMips && (TempMemoryRequired <= StreamInMemoryBudget || !StreamInRequests.Num()))
 		{
-			const int64 UsedMemoryRequired = StreamingRenderAsset.GetSize(StreamingRenderAsset.WantedMips) - StreamingRenderAsset.GetSize(StreamingRenderAsset.ResidentMips);
-			const int64 TempMemoryRequired = StreamingRenderAsset.GetSize(StreamingRenderAsset.WantedMips);
-
-			// Respect the temporary budget unless this is the first load request. This allows a single mip update of any size.
-			if (TempMemoryUsed + TempMemoryRequired <= TempMemoryBudget || !LoadRequests.Num())
-			{
-				LoadRequests.Add(AssetIndex);
-	
-				MemoryUsed += UsedMemoryRequired;
-				TempMemoryUsed += TempMemoryRequired;
-			}
+			StreamInRequests.Add(AssetIndex);
+			StreamInMemoryBudget -= TempMemoryRequired;
 		}
 	}
+
+	// Process stream out requests first since they execute faster, freeing the memory for the stream in requests.
+	LoadRequests.Append(StreamOutRequests);
+	LoadRequests.Append(StreamInRequests);
 }
 
 void FRenderAssetStreamingMipCalcTask::UpdatePendingStreamingStatus_Async()
@@ -895,12 +905,11 @@ void FRenderAssetStreamingMipCalcTask::DoWork()
 		StreamingRenderAsset.DynamicBoostFactor = 1.f; // Reset after every computation.
 	}
 
-	int64 MemoryUsed, TempMemoryUsed;
 	// According to budget, make relevant sacrifices and keep possible unwanted mips
-	UpdateBudgetedMips_Async(MemoryUsed, TempMemoryUsed);
+	UpdateBudgetedMips_Async();
 
 	// Update load requests.
-	UpdateLoadAndCancelationRequests_Async(MemoryUsed, TempMemoryUsed);
+	UpdateLoadAndCancelationRequests_Async();
 
 	// Update bHasStreamingUpdatePending
 	UpdatePendingStreamingStatus_Async();
diff --git a/Engine/Source/Runtime/Engine/Private/Streaming/AsyncTextureStreaming.h b/Engine/Source/Runtime/Engine/Private/Streaming/AsyncTextureStreaming.h
--- a/Engine/Source/Runtime/Engine/Private/Streaming/AsyncTextureStreaming.h
+++ b/Engine/Source/Runtime/Engine/Private/Streaming/AsyncTextureStreaming.h
@@ -193,9 +193,9 @@ private:
 
 	void TryKeepMips(TArray<int32>& PrioritizedRenderAssets, int64& MemoryBudgeted, const int64 InMemoryBudget);
 
-	void UpdateBudgetedMips_Async(int64& OutMemoryUsed, int64& OutTempMemoryUsed);
+	void UpdateBudgetedMips_Async();
 
-	void UpdateLoadAndCancelationRequests_Async(int64 MemoryUsed, int64 TempMemoryUsed);
+	void UpdateLoadAndCancelationRequests_Async();
 
 	void UpdatePendingStreamingStatus_Async();
 
diff --git a/Engine/Source/Runtime/Engine/Private/Streaming/StreamingTexture.cpp b/Engine/Source/Runtime/Engine/Private/Streaming/StreamingTexture.cpp
--- a/Engine/Source/Runtime/Engine/Private/Streaming/StreamingTexture.cpp
+++ b/Engine/Source/Runtime/Engine/Private/Streaming/StreamingTexture.cpp
@@ -489,12 +489,23 @@ bool FStreamingRenderAsset::UpdateLoadOrderPriority_Async(int32 MinMipForSplitRe
 	{
 		const bool bIsVisible			= ResidentMips < VisibleWantedMips; // Otherwise it means we are loading mips that are only useful for non visible primitives.
 		const bool bMustLoadFirst		= bForceFullyLoadHeuristic || bIsTerrainTexture || bLoadWithHigherPriority;
-		const bool bMipIsImportant		= WantedMips - ResidentMips > (bLooksLowRes ? 1 : 2);
 
-		if (bIsVisible)				LoadOrderPriority += 1024;
-		if (bMustLoadFirst)			LoadOrderPriority += 512; 
-		if (bMipIsImportant)		LoadOrderPriority += 256;
-		if (!bIsVisible)			LoadOrderPriority += FMath::Clamp<int32>(255 - (int32)LastRenderTime, 1, 255);
+		if (WantedMips > RequestedMips)
+		{
+			const bool bMipIsImportant = WantedMips - ResidentMips > (bLooksLowRes ? 1 : 2);
+
+			if (bIsVisible)				LoadOrderPriority += 1024;
+			if (bMustLoadFirst)			LoadOrderPriority += 512;
+			if (bMipIsImportant)		LoadOrderPriority += 256;
+			if (!bIsVisible)			LoadOrderPriority += FMath::Clamp<int32>(255 - (int32)LastRenderTime, 1, 255);
+		}
+		else // WantedMips < RequestedMips
+		{
+			// For stream out operations, unload mips that we will be less important to load quickly.
+			// Visibility is less relevant here because stream in requests might come from a visibily change.
+			if (!bMustLoadFirst)		LoadOrderPriority += 1024;
+			if (!bIsVisible)			LoadOrderPriority += 512;
+		}
 
 		return true;
 	}
