Some PLATFORM_COMPILER_HAS_IF_CONSTEXPR usage removed.

diff --git a/Engine/Source/Runtime/Core/Public/Containers/Array.h b/Engine/Source/Runtime/Core/Public/Containers/Array.h
--- a/Engine/Source/Runtime/Core/Public/Containers/Array.h
+++ b/Engine/Source/Runtime/Core/Public/Containers/Array.h
@@ -450,28 +450,6 @@ public:
 	TArray& operator=(const TArrayView<OtherElementType, OtherSizeType>& Other);
 
 private:
-#if !PLATFORM_COMPILER_HAS_IF_CONSTEXPR
-	template <
-		typename FromArrayType,
-		typename ToArrayType,
-		std::enable_if_t<TCanMoveBetweenAllocators<typename FromArrayType::AllocatorType, typename ToArrayType::AllocatorType>::Value>* = nullptr
-	>
-	static FORCEINLINE void MoveAllocatorToEmpty(FromArrayType& FromArray, ToArrayType& ToArray)
-	{
-		ToArray.AllocatorInstance.template MoveToEmptyFromOtherAllocator<typename FromArrayType::AllocatorType>(FromArray.AllocatorInstance);
-	}
-
-	template <
-		typename FromArrayType,
-		typename ToArrayType,
-		std::enable_if_t<!TCanMoveBetweenAllocators<typename FromArrayType::AllocatorType, typename ToArrayType::AllocatorType>::Value>* = nullptr
-	>
-	static FORCEINLINE void MoveAllocatorToEmpty(FromArrayType& FromArray, ToArrayType& ToArray)
-	{
-		ToArray.AllocatorInstance.MoveToEmpty(FromArray.AllocatorInstance);
-	}
-#endif
-
 	/**
 	 * Moves or copies array. Depends on the array type traits.
 	 *
@@ -486,7 +464,6 @@ private:
 		using FromAllocatorType = typename FromArrayType::AllocatorType;
 		using ToAllocatorType   = typename ToArrayType::AllocatorType;
 
-#if PLATFORM_COMPILER_HAS_IF_CONSTEXPR
 		if constexpr (TCanMoveBetweenAllocators<FromAllocatorType, ToAllocatorType>::Value)
 		{
 			ToArray.AllocatorInstance.template MoveToEmptyFromOtherAllocator<FromAllocatorType>(FromArray.AllocatorInstance);
@@ -495,9 +472,6 @@ private:
 		{
 			ToArray.AllocatorInstance.MoveToEmpty(FromArray.AllocatorInstance);
 		}
-#else
-		MoveAllocatorToEmpty(FromArray, ToArray);
-#endif
 
 		ToArray  .ArrayNum = FromArray.ArrayNum;
 		ToArray  .ArrayMax = FromArray.ArrayMax;
diff --git a/Engine/Source/Runtime/Engine/Public/Animation/AttributeTraits.h b/Engine/Source/Runtime/Engine/Public/Animation/AttributeTraits.h
--- a/Engine/Source/Runtime/Engine/Public/Animation/AttributeTraits.h
+++ b/Engine/Source/Runtime/Engine/Public/Animation/AttributeTraits.h
@@ -27,53 +27,5 @@ namespace UE
 		struct TAttributeTypeTraits : public TAttributeTypeTraitsBase<AttributeType>
 		{
 		};
-		
-		/** Implemented functionality for retrieving a types trait value for platforms without constexpr support */
-#if !PLATFORM_COMPILER_HAS_IF_CONSTEXPR
-		/**
-		 * Selection of IsBlendableType call.
-		 */
-		template<class AttributeType>
-		FORCEINLINE typename TEnableIf<!TAttributeTypeTraits<AttributeType>::IsBlendable, bool>::Type IsBlendableType()
-		{
-			return false;
-		}
-
-		template<class AttributeType>
-		FORCEINLINE typename TEnableIf<TAttributeTypeTraits<AttributeType>::IsBlendable, bool>::Type IsBlendableType()
-		{
-			return true;
-		}
-
-		/**
-		 * Selection of IsStepInterpolatedType call.
-		 */
-		template<class AttributeType>
-		FORCEINLINE typename TEnableIf<!TAttributeTypeTraits<AttributeType>::StepInterpolate, bool>::Type IsStepInterpolatedType()
-		{
-			return false;
-		}
-
-		template<class AttributeType>
-		FORCEINLINE typename TEnableIf<TAttributeTypeTraits<AttributeType>::StepInterpolate, bool>::Type IsStepInterpolatedType()
-		{
-			return true;
-		}
-
-		/**
-		 * Selection of CustomBlendedType call.
-		 */
-		template<class AttributeType>
-		FORCEINLINE typename TEnableIf<!TAttributeTypeTraits<AttributeType>::WithCustomBlendOperator, bool>::Type CustomBlendedType()
-		{
-			return false;
-		}
-
-		template<class AttributeType>
-		FORCEINLINE typename TEnableIf<TAttributeTypeTraits<AttributeType>::WithCustomBlendOperator, bool>::Type CustomBlendedType()
-		{
-			return true;
-		}
-#endif
 	}
 }
diff --git a/Engine/Source/Runtime/Engine/Public/Animation/AttributeTypes.h b/Engine/Source/Runtime/Engine/Public/Animation/AttributeTypes.h
--- a/Engine/Source/Runtime/Engine/Public/Animation/AttributeTypes.h
+++ b/Engine/Source/Runtime/Engine/Public/Animation/AttributeTypes.h
@@ -51,18 +51,13 @@ namespace UE
 				AttributeTypes::RegisteredTypes.Add(ScriptStruct);
 				AttributeTypes::Operators.Add(MakeUnique<OperatorType>(Forward<OperatorArgs>(args)...));
 
-#if PLATFORM_COMPILER_HAS_IF_CONSTEXPR
 				if constexpr (!UE::Anim::TAttributeTypeTraits<AttributeType>::StepInterpolate)
-#else
-				if (!UE::Anim::IsStepInterpolatedType<AttributeType>())
-#endif // PLATFORM_COMPILER_HAS_IF_CONSTEXPR	
 				{
 					AttributeTypes::InterpolatableTypes.Add(ScriptStruct);
 				}
 			}
 
 			/** Used for registering an attribute type for which TAttributeTypeTraits::WithCustomBlendOperator is set to false, use RegisterTypeWithOperator() otherwise */
-#if PLATFORM_COMPILER_HAS_IF_CONSTEXPR
 			template<typename AttributeType>
 			static void RegisterType()
 			{
@@ -83,36 +78,7 @@ namespace UE
 
 				AttributeTypes::Operators.Add(MakeUnique<TAttributeBlendOperator<AttributeType>>());
 			}
-#else
-			template<typename AttributeType>
-			static typename TEnableIf<TAttributeTypeTraits<AttributeType>::IsBlendable, void>::Type RegisterType()
-			{
-				static_assert(!TAttributeTypeTraits<AttributeType>::WithCustomBlendOperator, "Attribute type requires a custom blend operation");
-				static_assert(TModels<CBlendableAttribute, AttributeType>::Value, "Missing arithmetic operators required for Attribute blending");
-
-				UScriptStruct* ScriptStruct = AttributeType::StaticStruct();
-				AttributeTypes::RegisteredTypes.Add(ScriptStruct);
 
-				if (!UE::Anim::IsStepInterpolatedType<AttributeType>())
-				{
-					
-					AttributeTypes::InterpolatableTypes.Add(ScriptStruct);
-				}
-
-				AttributeTypes::Operators.Add(MakeUnique<TAttributeBlendOperator<AttributeType>>());
-			}
-
-			template<typename AttributeType>
-			static typename TEnableIf<!TAttributeTypeTraits<AttributeType>::IsBlendable, void>::Type RegisterType()
-			{
-				static_assert(!TAttributeTypeTraits<AttributeType>::WithCustomBlendOperator, "Attribute type requires a custom blend operation");
-
-				UScriptStruct* ScriptStruct = AttributeType::StaticStruct();
-				AttributeTypes::RegisteredTypes.Add(ScriptStruct);
-
-				AttributeTypes::Operators.Add(MakeUnique<TAttributeBlendOperator<AttributeType>>());
-			}
-#endif			
 			/** Unregisters a specific attribute type and deletes its associated blend operator */
 			template<typename AttributeType>
 			static void UnregisterType()
diff --git a/Engine/Source/Runtime/Slate/Private/Framework/Application/SlateApplication.cpp b/Engine/Source/Runtime/Slate/Private/Framework/Application/SlateApplication.cpp
--- a/Engine/Source/Runtime/Slate/Private/Framework/Application/SlateApplication.cpp
+++ b/Engine/Source/Runtime/Slate/Private/Framework/Application/SlateApplication.cpp
@@ -392,7 +392,6 @@ public:
 		{
 			const FWidgetAndPointer& ArrangedWidget = RoutingPolicy.GetWidget();
 
-#if PLATFORM_COMPILER_HAS_IF_CONSTEXPR
 			if constexpr (Translate<EventType>::TranslationNeeded())
 			{
 				const EventType TranslatedEvent = Translate<EventType>::PointerEvent(ArrangedWidget, EventCopy);
@@ -404,11 +403,6 @@ public:
 				Reply = Lambda(ArrangedWidget, EventCopy).SetHandler(ArrangedWidget.Widget);
 				ProcessReply(ThisApplication, RoutingPath, Reply, WidgetsUnderCursor, &EventCopy);
 			}
-#else
-			const EventType TranslatedEvent = Translate<EventType>::PointerEvent(ArrangedWidget, EventCopy);
-			Reply = Lambda(ArrangedWidget, TranslatedEvent).SetHandler(ArrangedWidget.Widget);
-			ProcessReply(ThisApplication, RoutingPath, Reply, WidgetsUnderCursor, &TranslatedEvent);
-#endif
 		}
 
 		return Reply;
