Thread priority/affinity via ini files
- Remove ThreadConfig cvar. This does't work in ini files due to initialization order problems (results in lockups, incorrect affinities)
- Make thread config configurable via engine ini file array, e.g:
    +ThreadConfigs=Task:TPri_SlightlyBelowNormal:0x07
    +ThreadConfigs=TaskBP:TPri_BelowNormal
- Modify SetupThreadConfig format to use colons to separate params for a given thread. Threads are still separated with spaces
- Simplify the SetupThreadConfig and setDefaults functions (combined into one)

diff --git a/Engine/Source/Runtime/Engine/Private/UnrealEngine.cpp b/Engine/Source/Runtime/Engine/Private/UnrealEngine.cpp
--- a/Engine/Source/Runtime/Engine/Private/UnrealEngine.cpp
+++ b/Engine/Source/Runtime/Engine/Private/UnrealEngine.cpp
@@ -417,6 +417,9 @@ bool GIsTextureMemoryCorrupted = false;
 bool GIsPrepareMapChangeBroken = false;
 #endif
 
+// Thread config 
+void InitThreadConfig();
+
 #if WITH_EDITOR
 // A debugging aid set when we switch out different play worlds during Play In Editor / PIE
 ENGINE_API FString GPlayInEditorContextString(TEXT("invalid"));
@@ -1954,6 +1957,8 @@ void UEngine::Init(IEngineLoop* InEngineLoop)
 #if !UE_BUILD_SHIPPING
 	UE_CLOG(FPlatformMemory::IsExtraDevelopmentMemoryAvailable(), LogInit, Warning, TEXT("Running with %dMB of extra development memory!"),FPlatformMemory::GetExtraDevelopmentMemorySize()/1024ull/1024ull);
 #endif
+
+	InitThreadConfig();
 }
 
 void UEngine::Start()
@@ -16205,49 +16210,34 @@ static FAutoConsoleCommand UAssetLoadTestCmd(
 
 #endif
 
-static uint64 TaskThreadAffinity = 0xFFFFFFFFFFFFFFFF;
-static uint64 TaskThreadBPAffinity = 0xFFFFFFFFFFFFFFFF;
-static uint64 GameThreadAffinity = 0xFFFFFFFFFFFFFFFF;
-static uint64 RenderThreadAffinity = 0xFFFFFFFFFFFFFFFF;
-static uint64 RHIThreadAffinity = 0xFFFFFFFFFFFFFFFF;
-
-static EThreadPriority TaskThreadPriority = EThreadPriority::TPri_Normal;
-static EThreadPriority TaskThreadBPPriority = EThreadPriority::TPri_BelowNormal;
-static EThreadPriority GameThreadPriority = EThreadPriority::TPri_Highest;
-static EThreadPriority RenderThreadPriority = EThreadPriority::TPri_Highest;
-static EThreadPriority RHIThreadPriority = EThreadPriority::TPri_Highest;
+struct FThreadConfig
+{
+	FThreadConfig() { Priority = EThreadPriority::TPri_Normal;  Affinity = 0xFFFFFFFFFFFFFFFF;	}
+	uint64 Affinity;
+	EThreadPriority Priority;
+};
 
-static FString CVar_ThreadConfigString = TEXT("");
+static FThreadConfig GameThreadConfig;
+static FThreadConfig RenderThreadConfig;
+static FThreadConfig RHIThreadConfig;
+static FThreadConfig TaskThreadConfig;
+static FThreadConfig TaskBPThreadConfig;
 
 extern CORE_API int32 GUseNewTaskBackend;
 
-static bool AffinitiesLoadedDefaults = false;
-static bool PrioritiesLoadedDefaults = false;
-
-void GetDefaultAffinities(bool bForce = false)
+void GetDefaultThreadConfigs()
 {
-	if (!AffinitiesLoadedDefaults || bForce)
-	{
-		AffinitiesLoadedDefaults = true;
-		TaskThreadAffinity = FPlatformAffinity::GetTaskGraphThreadMask();
-		TaskThreadBPAffinity = FPlatformAffinity::GetTaskGraphBackgroundTaskMask();
-		GameThreadAffinity = FPlatformAffinity::GetMainGameMask();
-		RenderThreadAffinity = FPlatformAffinity::GetRenderingThreadMask();
-		RHIThreadAffinity = FPlatformAffinity::GetRHIThreadMask();
-	}
-}
+	TaskThreadConfig.Affinity = FPlatformAffinity::GetTaskGraphThreadMask();
+	TaskBPThreadConfig.Affinity = FPlatformAffinity::GetTaskGraphBackgroundTaskMask();
+	GameThreadConfig.Affinity = FPlatformAffinity::GetMainGameMask();
+	RenderThreadConfig.Affinity = FPlatformAffinity::GetRenderingThreadMask();
+	RHIThreadConfig.Affinity = FPlatformAffinity::GetRHIThreadMask();
 
-void GetDefaultPriorities(bool bForce = false)
-{
-	if (!PrioritiesLoadedDefaults || bForce)
-	{
-		PrioritiesLoadedDefaults = true;
-		TaskThreadPriority = EThreadPriority::TPri_Normal;
-		TaskThreadBPPriority = EThreadPriority::TPri_BelowNormal;
-		GameThreadPriority = EThreadPriority::TPri_Highest;
-		RenderThreadPriority = FPlatformAffinity::GetRenderingThreadPriority();
-		RHIThreadPriority = FPlatformAffinity::GetRHIThreadPriority();
-	}
+	TaskThreadConfig.Priority = EThreadPriority::TPri_Normal;
+	TaskBPThreadConfig.Priority = EThreadPriority::TPri_BelowNormal;
+	GameThreadConfig.Priority = EThreadPriority::TPri_Highest;
+	RenderThreadConfig.Priority = FPlatformAffinity::GetRenderingThreadPriority();
+	RHIThreadConfig.Priority = FPlatformAffinity::GetRHIThreadPriority();
 }
 
 EThreadPriority StringToThreadPriority(FString InString)
@@ -16288,7 +16278,6 @@ FString ThreadPriorityToString(EThreadPriority InPriority)
 {
 	switch (InPriority)
 	{
-
 	case EThreadPriority::TPri_Normal:
 		return FString(TEXT("TPri_Normal"));
 	case EThreadPriority::TPri_AboveNormal:
@@ -16313,35 +16302,35 @@ FString ThreadPriorityToString(EThreadPriority InPriority)
 
 void SetPriorityAndAffinityOnGameThread()
 {
-	if (GameThreadPriority != EThreadPriority::TPri_Num)
+	if (GameThreadConfig.Priority != EThreadPriority::TPri_Num)
 	{
-		FPlatformProcess::SetThreadPriority(GameThreadPriority);
-		UE_LOG(LogConsoleResponse, Display, TEXT("GT Priority %s"), *ThreadPriorityToString(GameThreadPriority));
+		FPlatformProcess::SetThreadPriority(GameThreadConfig.Priority);
+		UE_LOG(LogConsoleResponse, Display, TEXT("GT Priority %s"), *ThreadPriorityToString(GameThreadConfig.Priority));
 	}
-	FPlatformProcess::SetThreadAffinityMask(GameThreadAffinity);
-	UE_LOG(LogConsoleResponse, Display, TEXT("GT Affinity     %016llX"), GameThreadAffinity);
+	FPlatformProcess::SetThreadAffinityMask(GameThreadConfig.Affinity);
+	UE_LOG(LogConsoleResponse, Display, TEXT("GT Affinity %016llX"), GameThreadConfig.Affinity);
 }
 
 void SetPriorityAndAffinityOnRenderThread()
 {
-	if (RenderThreadPriority != EThreadPriority::TPri_Num)
+	if (RenderThreadConfig.Priority != EThreadPriority::TPri_Num)
 	{
-		FPlatformProcess::SetThreadPriority(RenderThreadPriority);
-		UE_LOG(LogConsoleResponse, Display, TEXT("RT Priority %s"), *ThreadPriorityToString(RenderThreadPriority));
+		FPlatformProcess::SetThreadPriority(RenderThreadConfig.Priority);
+		UE_LOG(LogConsoleResponse, Display, TEXT("RT Priority %s"), *ThreadPriorityToString(RenderThreadConfig.Priority));
 	}
-	FPlatformProcess::SetThreadAffinityMask(RenderThreadAffinity);
-	UE_LOG(LogConsoleResponse, Display, TEXT("RT Affinity    %016llX"), RenderThreadAffinity);
+	FPlatformProcess::SetThreadAffinityMask(RenderThreadConfig.Affinity);
+	UE_LOG(LogConsoleResponse, Display, TEXT("RT Affinity %016llX"), RenderThreadConfig.Affinity);
 }
 
 void SetPriorityAndAffinityOnRHIThread()
 {
-	if (RHIThreadPriority != EThreadPriority::TPri_Num)
+	if (RHIThreadConfig.Priority != EThreadPriority::TPri_Num)
 	{
-		FPlatformProcess::SetThreadPriority(RHIThreadPriority);
-		UE_LOG(LogConsoleResponse, Display, TEXT("RHI Priority %s"), *ThreadPriorityToString(RHIThreadPriority));
+		FPlatformProcess::SetThreadPriority(RHIThreadConfig.Priority);
+		UE_LOG(LogConsoleResponse, Display, TEXT("RHI Priority %s"), *ThreadPriorityToString(RHIThreadConfig.Priority));
 	}
-	FPlatformProcess::SetThreadAffinityMask(RHIThreadAffinity);
-	UE_LOG(LogConsoleResponse, Display, TEXT("RHI Affinity  %016llX"), RHIThreadAffinity);
+	FPlatformProcess::SetThreadAffinityMask(RHIThreadConfig.Affinity);
+	UE_LOG(LogConsoleResponse, Display, TEXT("RHI Affinity %016llX"), RHIThreadConfig.Affinity);
 }
 
 
@@ -16349,88 +16338,71 @@ void SetPriorityAndAffinityOnRHIThread()
 static void SetupThreadConfig(const TArray<FString>& Args)
 {
 	if (!GUseNewTaskBackend)
+	{
+		UE_LOG(LogConsoleResponse, Warning, TEXT("SetupThreadConfig called, but this requires the new task backend. Ignoring"));
 		return;
+	}
 
-	GetDefaultPriorities((Args.Num() && Args[0] == TEXT("default")));
-	GetDefaultAffinities((Args.Num() && Args[0] == TEXT("default")));
+	UE_LOG(LogConsoleResponse, Display, TEXT("Setting thread configurations"));
 
-	if ((Args.Num() && Args[0] != TEXT("default")))
+	// Lazily load the default thread configs
+	static bool bLoadedDefaults = false;
+	bool bResetToDefaults = (Args.Num() && Args[0] == TEXT("default"));
+	if (!bLoadedDefaults || bResetToDefaults)
 	{
-		for (int32 Index = 0; Index + 1 < Args.Num(); )
-		{
-			int32 OriginalIndex = Index;
-			uint64 Aff = 0xffffffffffffffff;
-			EThreadPriority NewPriority = EThreadPriority::TPri_Num;
-
-			int32 Index2 = Index + 1;
-			for (; Index2 < Args.Num(); Index2++)
-			{
-				bool ValidArgFound = false;
-				// if the arg is a hex number we are setting affinity
-				if (Args[Index2].StartsWith(TEXT("0x")))
-				{
-					Aff = FParse::HexNumber(*Args[Index2]); // this is only 32 bits
-					ValidArgFound = true;
-				}
-				// if the arg starts with TPri we are setting priority
-				if (Args[Index2].StartsWith(TEXT("TPri")))
-				{
-					NewPriority = StringToThreadPriority(Args[Index2]);
-					ValidArgFound = true;
-				}
-				if (!ValidArgFound)
-				{
-					// we didn't find an expected arg so we'll move on 
-					break;
-				}
-			}
-			Index = Index2;
+		GetDefaultThreadConfigs();
+		bLoadedDefaults = true;
+	}
 
-			if (Aff != 0xffffffffffffffff)
+	if (!bResetToDefaults)
+	{
+		for (const FString& ThreadConfigStr : Args)
+		{
+			TArray<FString> ThreadConfigEntries;
+			ThreadConfigStr.ParseIntoArray(ThreadConfigEntries, TEXT(":"), true);
+			if (ThreadConfigEntries.Num() >= 2)
 			{
-				if (Args[OriginalIndex] == TEXT("GT"))
-				{
-					GameThreadAffinity = Aff;
-				}
-				else if (Args[OriginalIndex] == TEXT("RT"))
+				FString ThreadName = ThreadConfigEntries[0];
+				FThreadConfig* ThreadConfigToSet = nullptr;
+				if (ThreadName == TEXT("GT"))
 				{
-					RenderThreadAffinity = Aff;
+					ThreadConfigToSet = &GameThreadConfig;
 				}
-				else if (Args[OriginalIndex] == TEXT("RHI"))
+				else if (ThreadName == TEXT("RT"))
 				{
-					RHIThreadAffinity = Aff;
+					ThreadConfigToSet = &RenderThreadConfig;
 				}
-				else if (Args[OriginalIndex] == TEXT("Task"))
+				else if (ThreadName == TEXT("RHI"))
 				{
-					TaskThreadAffinity = Aff;
+					ThreadConfigToSet = &RHIThreadConfig;
 				}
-				else if (Args[OriginalIndex] == TEXT("TaskBP"))
+				else if (ThreadName == TEXT("Task"))
 				{
-					TaskThreadBPAffinity = Aff;
+					ThreadConfigToSet = &TaskThreadConfig;
 				}
-			}
-
-			if (NewPriority != EThreadPriority::TPri_Num)
-			{
-				if (Args[OriginalIndex] == TEXT("GT"))
+				else if (ThreadName == TEXT("TaskBP"))
 				{
-					GameThreadPriority = NewPriority;
+					ThreadConfigToSet = &TaskBPThreadConfig;
 				}
-				else if (Args[OriginalIndex] == TEXT("RT"))
+				if (ThreadConfigToSet == nullptr)
 				{
-					RenderThreadPriority = NewPriority;
-				}
-				else if (Args[OriginalIndex] == TEXT("RHI"))
-				{
-					RHIThreadPriority = NewPriority;
-				}
-				else if (Args[OriginalIndex] == TEXT("Task"))
-				{
-					TaskThreadPriority = NewPriority;
+					UE_LOG(LogConsoleResponse, Warning, TEXT("Thread name not found: %s"), *ThreadName);
+					continue;
 				}
-				else if (Args[OriginalIndex] == TEXT("TaskBP"))
+				// Read the rest of the thread config args
+				for (int i = 1; i < ThreadConfigEntries.Num(); i++)
 				{
-					TaskThreadBPPriority = NewPriority;
+					// if the arg is a hex number we are setting affinity
+					const FString& Value = ThreadConfigEntries[i];
+					if (Value.StartsWith(TEXT("0x")))
+					{
+						ThreadConfigToSet->Affinity = FParse::HexNumber64(*Value);
+					}
+					// if the arg starts with TPri we are setting priority
+					if (Value.StartsWith(TEXT("TPri")))
+					{
+						ThreadConfigToSet->Priority = StringToThreadPriority(Value);
+					}
 				}
 			}
 		}
@@ -16449,16 +16421,17 @@ static void SetupThreadConfig(const TArray<FString>& Args)
 	check(IsInGameThread());
 	SetPriorityAndAffinityOnGameThread();
 	FlushRenderingCommands();
+	
 	int32 NumWorkerThreads = FPlatformMisc::NumberOfWorkerThreadsToSpawn();
 	int32 NumBackgroundWorkers = FMath::Max(1, NumWorkerThreads - FMath::Min<int>(2, NumWorkerThreads));
 	int32 NumForegroundWorkers = FMath::Max(1, NumWorkerThreads - NumBackgroundWorkers);
-	LowLevelTasks::FScheduler::Get().RestartWorkers(NumForegroundWorkers, NumBackgroundWorkers, FForkProcessHelper::IsForkedMultithreadInstance(), TaskThreadPriority, TaskThreadBPPriority, TaskThreadAffinity, TaskThreadBPAffinity);
+	LowLevelTasks::FScheduler::Get().RestartWorkers(NumForegroundWorkers, NumBackgroundWorkers, FForkProcessHelper::IsForkedMultithreadInstance(), TaskThreadConfig.Priority, TaskBPThreadConfig.Priority, TaskThreadConfig.Affinity, TaskBPThreadConfig.Affinity);
 
-	UE_LOG(LogConsoleResponse, Display, TEXT("Task Priority %s"), *ThreadPriorityToString(TaskThreadPriority));
-	UE_LOG(LogConsoleResponse, Display, TEXT("Task Affinity    %016llX"), TaskThreadAffinity);
+	UE_LOG(LogConsoleResponse, Display, TEXT("Task Priority %s"), *ThreadPriorityToString(TaskThreadConfig.Priority));
+	UE_LOG(LogConsoleResponse, Display, TEXT("Task Affinity %016llX"), TaskThreadConfig.Affinity);
 
-	UE_LOG(LogConsoleResponse, Display, TEXT("TaskBP Priority %s"), *ThreadPriorityToString(TaskThreadBPPriority));
-	UE_LOG(LogConsoleResponse, Display, TEXT("TaskBP Affinity    %016llX"), TaskThreadBPAffinity);
+	UE_LOG(LogConsoleResponse, Display, TEXT("TaskBP Priority %s"), *ThreadPriorityToString(TaskBPThreadConfig.Priority));
+	UE_LOG(LogConsoleResponse, Display, TEXT("TaskBP Affinity %016llX"), TaskBPThreadConfig.Affinity);
 
 	GLog->FlushThreadedLogs();
 }
@@ -16470,22 +16443,19 @@ static FAutoConsoleCommand SetupThreadConfigCmd(
 	FConsoleCommandWithArgsDelegate::CreateStatic(&SetupThreadConfig)
 );
 
-void ThreadConfigStringChanged(IConsoleVariable* Variable)
+void InitThreadConfig()
 {
-	TArray<FString> Args;
-	CVar_ThreadConfigString.ParseIntoArray(Args, TEXT(" "), true);
-	SetupThreadConfig(Args);
+	if (GConfig)
+	{
+		TArray<FString> ThreadConfigEntries;
+		GConfig->GetArray(TEXT("/Script/Engine.Engine"), TEXT("ThreadConfigs"), ThreadConfigEntries, GEngineIni);
+		if (!ThreadConfigEntries.IsEmpty())
+		{
+			SetupThreadConfig(ThreadConfigEntries);
+		}
+	}
 }
 
-static FAutoConsoleVariableRef CVarThreadConfigString(
-	TEXT("ThreadConfig"),
-	CVar_ThreadConfigString,
-	TEXT("Configures the Thread Priorites and Affinites of GT/RT/RHI/Task/TaskBP"),
-	FConsoleVariableDelegate::CreateStatic(&ThreadConfigStringChanged),
-	ECVF_Default
-);
-
-
 #if !UE_BUILD_SHIPPING
 
 static void PakFileTest(const TArray<FString>& Args)
