Fix bug where GT priority gets bumped to Highest if we update thread config
- Define the game thread priority centrally in FPlatformAffinity. Use this as the default for thread config

diff --git a/Engine/Source/Runtime/Core/Public/GenericPlatform/GenericPlatformAffinity.h b/Engine/Source/Runtime/Core/Public/GenericPlatform/GenericPlatformAffinity.h
--- a/Engine/Source/Runtime/Core/Public/GenericPlatform/GenericPlatformAffinity.h
+++ b/Engine/Source/Runtime/Core/Public/GenericPlatform/GenericPlatformAffinity.h
@@ -140,6 +140,11 @@ public:
 		return TPri_Normal;
 	}
 
+	static EThreadPriority GetGameThreadPriority()
+	{
+		return TPri_Normal;
+	}
+
 	static EThreadCreateFlags GetRHIThreadFlags()
 	{
 		return EThreadCreateFlags::None;
diff --git a/Engine/Source/Runtime/Engine/Private/UnrealEngine.cpp b/Engine/Source/Runtime/Engine/Private/UnrealEngine.cpp
--- a/Engine/Source/Runtime/Engine/Private/UnrealEngine.cpp
+++ b/Engine/Source/Runtime/Engine/Private/UnrealEngine.cpp
@@ -16235,7 +16235,7 @@ void GetDefaultThreadConfigs()
 
 	TaskThreadConfig.Priority = EThreadPriority::TPri_Normal;
 	TaskBPThreadConfig.Priority = EThreadPriority::TPri_BelowNormal;
-	GameThreadConfig.Priority = EThreadPriority::TPri_Highest;
+	GameThreadConfig.Priority = FPlatformAffinity::GetGameThreadPriority();
 	RenderThreadConfig.Priority = FPlatformAffinity::GetRenderingThreadPriority();
 	RHIThreadConfig.Priority = FPlatformAffinity::GetRHIThreadPriority();
 }
