Added multicast delegate to be broadcast when a PlayerState's pawn changes, allowing Blueprint & C++ to perform actions based on the event.  These changes are slightly different than the PR, in keeping with style of other similar delegates in UE.

PR #8272

diff --git a/Engine/Source/Runtime/Engine/Classes/GameFramework/PlayerState.h b/Engine/Source/Runtime/Engine/Classes/GameFramework/PlayerState.h
--- a/Engine/Source/Runtime/Engine/Classes/GameFramework/PlayerState.h
+++ b/Engine/Source/Runtime/Engine/Classes/GameFramework/PlayerState.h
@@ -58,6 +58,8 @@ struct UE_DEPRECATED(4.27, "ExactPingV2 is no longer used. Please use ExactPing
 	PRAGMA_ENABLE_DEPRECATION_WARNINGS
 };
 
+DECLARE_DYNAMIC_MULTICAST_DELEGATE_ThreeParams(FOnPlayerStatePawnSet, APlayerState*, Player, APawn*, NewPawn, APawn*, OldPawn);
+
 /**
  * A PlayerState is created for every player on a server (or in a standalone game).
  * PlayerStates are replicated to all clients, and contain network game relevant information about the player, such as playername, score, etc.
@@ -163,6 +165,10 @@ public:
 	/** The session that the player needs to join/remove from as it is created/leaves */
 	FName SessionName;
 
+	/** Broadcast whenever this player's possessed pawn is set */
+	UPROPERTY(BlueprintAssignable, Category = "Events")
+	FOnPlayerStatePawnSet OnPawnSet;
+
 private:
 
 	friend struct FSetPlayerStatePawn;
@@ -437,6 +443,8 @@ private:
 
 	FSetPlayerStatePawn(APlayerState* PlayerState, APawn* Pawn)
 	{
+		APawn* OldPawn = PlayerState->PawnPrivate;
 		PlayerState->PawnPrivate = Pawn;
+		PlayerState->OnPawnSet.Broadcast(PlayerState, PlayerState->PawnPrivate, OldPawn); 
 	}
 };
