Fix bug where index buffers would not create SRVs even though they were needed.  This bug would only happen when *r.FreeSkeletalMeshBuffers=1* was on.  This causes index buffer resource creation to empty the CPU-side Indices array.  Then we would conditionally create the SRV based on the size of the array (which is now empty).  To fix this we should check the *CachedNumIndices* value.

- Fix a potential future bug with *FRawStaticIndexBuffer* where it wasn't checking *DoSkeletalMeshIndexBuffersNeedSRV* to see if we need to create a SRV for the index buffer.  We don't actually use *FRawStaticIndexBuffer* for skeletal meshes so this is currently a non-issue.
- *CreateRHIBuffer_Internal* no longer updates *CachedNumIndices* with the *Indices.Num()*.  This was wrong because the call to *RHICreateIndexBuffer* would empty the array and we would set CachedNumIndices to zero.  I've added a check here now to ensure *CachedNumIndices* equals *Indices.Num()*.

diff --git a/Engine/Source/Runtime/Engine/Private/RawIndexBuffer.cpp b/Engine/Source/Runtime/Engine/Private/RawIndexBuffer.cpp
--- a/Engine/Source/Runtime/Engine/Private/RawIndexBuffer.cpp
+++ b/Engine/Source/Runtime/Engine/Private/RawIndexBuffer.cpp
@@ -323,10 +323,14 @@ FBufferRHIRef FRawStaticIndexBuffer::CreateRHIBuffer_Internal()
 
 	if (GetNumIndices() > 0)
 	{
+		extern ENGINE_API bool DoSkeletalMeshIndexBuffersNeedSRV();
+		bool bSRV = DoSkeletalMeshIndexBuffersNeedSRV();
+
 		// When bAllowCPUAccess is true, the meshes is likely going to be used for Niagara to spawn particles on mesh surface.
 		// And it can be the case for CPU *and* GPU access: no differenciation today. That is why we create a SRV in this case.
 		// This also avoid setting lots of states on all the members of all the different buffers used by meshes. Follow up: https://jira.it.epicgames.net/browse/UE-69376.
-		bool bSRV = IndexStorage.GetAllowCPUAccess();
+		bSRV |= IndexStorage.GetAllowCPUAccess();
+
 		const EBufferUsageFlags BufferFlags = BUF_Static | (bSRV ? BUF_ShaderResource : BUF_None);
 
 		// Create the index buffer.
diff --git a/Engine/Source/Runtime/Engine/Public/RawIndexBuffer.h b/Engine/Source/Runtime/Engine/Public/RawIndexBuffer.h
--- a/Engine/Source/Runtime/Engine/Public/RawIndexBuffer.h
+++ b/Engine/Source/Runtime/Engine/Public/RawIndexBuffer.h
@@ -92,7 +92,7 @@ class FIndexArrayView
 public:
 	/** Default constructor. */
 	FIndexArrayView()
-		: UntypedIndexData(NULL)
+		: UntypedIndexData(nullptr)
 		, NumIndices(0)
 		, b32Bit(false)
 	{
@@ -398,9 +398,9 @@ public:
 	{
 		IndexBufferRHI = CreateRHIBuffer_RenderThread();
 
-		if (IndexBufferRHI && IsSRVNeeded() && Indices.Num())
+		if (IndexBufferRHI && IsSRVNeeded() && Num())
 		{
-			SRVValue = RHICreateShaderResourceView(Indices.Num() ? IndexBufferRHI : nullptr);
+			SRVValue = RHICreateShaderResourceView(Num() ? IndexBufferRHI : nullptr);
 		}
 	}
 	
@@ -544,6 +544,8 @@ private:
 	{
 		if (CachedNumIndices)
 		{
+			check(CachedNumIndices == Indices.Num());
+
 			// Create the index buffer.
 			FRHIResourceCreateInfo CreateInfo(sizeof(INDEX_TYPE) == 4 ? TEXT("FRawStaticIndexBuffer32") : TEXT("FRawStaticIndexBuffer16"), &Indices);
 			EBufferUsageFlags Flags = BUF_Static;
@@ -566,9 +568,8 @@ private:
 				Ret = RHIAsyncCreateIndexBuffer(sizeof(INDEX_TYPE), Size, Flags, CreateInfo);
 			}
 
-			CachedNumIndices = Indices.Num();
 			return Ret;
 		}
 		return nullptr;
 	}
-};
\ No newline at end of file
+};
