Added min/max scaled virtual texture pool size.
This gives us better control on pool scaling.
Also adjusted the thresholds for vt residency mip map bias.

diff --git a/Engine/Source/Runtime/Engine/Classes/VT/VirtualTexturePoolConfig.h b/Engine/Source/Runtime/Engine/Classes/VT/VirtualTexturePoolConfig.h
--- a/Engine/Source/Runtime/Engine/Classes/VT/VirtualTexturePoolConfig.h
+++ b/Engine/Source/Runtime/Engine/Classes/VT/VirtualTexturePoolConfig.h
@@ -24,7 +24,9 @@ struct FVirtualTextureSpacePoolConfig
 		, SizeInMegabyte(0)
 		, bEnableResidencyMipMapBias(false)
 		, bAllowSizeScale(false)
-		, ScalabilityGroup(0) 
+		, ScalabilityGroup(0)
+		, MinScaledSizeInMegabyte(0)
+		, MaxScaledSizeInMegabyte(0)
 	{}
 
 	/** Minimum tile size to match (including tile border). */
@@ -39,7 +41,7 @@ struct FVirtualTextureSpacePoolConfig
 	UPROPERTY()
 	TArray< TEnumAsByte<EPixelFormat> > Formats;
 
-	/** Upper limit of size in megabytes to allocate for the pool. The allocator will allocate as close as possible to this limit. */
+	/** Size in megabytes to allocate for the pool. The allocator will allocate as close as possible below this limit. */
 	UPROPERTY()
 	int32 SizeInMegabyte;
 
@@ -55,6 +57,14 @@ struct FVirtualTextureSpacePoolConfig
 	UPROPERTY()
 	uint32 ScalabilityGroup;
 
+	/** Lower limit of size in megabytes to allocate for the pool after size scaling. */
+	UPROPERTY()
+	int32 MinScaledSizeInMegabyte;
+
+	/** Upper limit of size in megabytes to allocate for the pool after size scaling. */
+	UPROPERTY()
+	int32 MaxScaledSizeInMegabyte;
+
 	/** Is this the default config? Use this setting when we can't find any other match. */
 	bool IsDefault() const { return Formats.Num() == 0 && SizeInMegabyte > 0; }
 };
diff --git a/Engine/Source/Runtime/Engine/Private/VT/VirtualTexturePoolConfig.cpp b/Engine/Source/Runtime/Engine/Private/VT/VirtualTexturePoolConfig.cpp
--- a/Engine/Source/Runtime/Engine/Private/VT/VirtualTexturePoolConfig.cpp
+++ b/Engine/Source/Runtime/Engine/Private/VT/VirtualTexturePoolConfig.cpp
@@ -35,6 +35,14 @@ void UVirtualTexturePoolConfig::FindPoolConfig(TEnumAsByte<EPixelFormat> const*
 				OutConfig = Config;
 				const float Scale = Config.bAllowSizeScale ? VirtualTextureScalability::GetPoolSizeScale(Config.ScalabilityGroup) : 1.f;
 				OutConfig.SizeInMegabyte = (int32)(Scale * (float)OutConfig.SizeInMegabyte);
+				if (Scale < 1.f && OutConfig.MinScaledSizeInMegabyte > 0)
+				{
+					OutConfig.SizeInMegabyte = FMath::Max(OutConfig.SizeInMegabyte, OutConfig.MinScaledSizeInMegabyte);
+				}
+				if (Scale > 1.f && OutConfig.MaxScaledSizeInMegabyte > 0)
+				{
+					OutConfig.SizeInMegabyte = FMath::Min(OutConfig.SizeInMegabyte, OutConfig.MaxScaledSizeInMegabyte);
+				}
 				return;
 			}
 		}
diff --git a/Engine/Source/Runtime/Renderer/Private/VT/VirtualTexturePhysicalSpace.cpp b/Engine/Source/Runtime/Renderer/Private/VT/VirtualTexturePhysicalSpace.cpp
--- a/Engine/Source/Runtime/Renderer/Private/VT/VirtualTexturePhysicalSpace.cpp
+++ b/Engine/Source/Runtime/Renderer/Private/VT/VirtualTexturePhysicalSpace.cpp
@@ -11,23 +11,23 @@
 
 static TAutoConsoleVariable<float> CVarVTResidencyMaxMipMapBias(
 	TEXT("r.VT.Residency.MaxMipMapBias"),
-	8,
+	4,
 	TEXT("Maximum mip bias to apply to prevent Virtual Texture pool residency over-subscription.\n")
-	TEXT("Default 8"),
+	TEXT("Default 4"),
 	ECVF_RenderThreadSafe);
 
 static TAutoConsoleVariable<float> CVarVTResidencyUpperBound(
 	TEXT("r.VT.Residency.UpperBound"),
-	0.9f,
+	0.95f,
 	TEXT("Virtual Texture pool residency above which we increase mip bias.\n")
 	TEXT("Default 0.95"),
 	ECVF_RenderThreadSafe);
 
 static TAutoConsoleVariable<float> CVarVTResidencyLowerBound(
 	TEXT("r.VT.Residency.LowerBound"),
-	0.75f,
+	0.95f,
 	TEXT("Virtual Texture pool residency below which we decrease mip bias.\n")
-	TEXT("Default 0.85"),
+	TEXT("Default 0.95"),
 	ECVF_RenderThreadSafe);
 
 static TAutoConsoleVariable<float> CVarVTResidencyLockedUpperBound(
@@ -42,7 +42,7 @@ static TAutoConsoleVariable<float> CVarVTResidencyAdjustmentRate(
 	TEXT("r.VT.Residency.AdjustmentRate"),
 	0.2,
 	TEXT("Rate at which we adjust mip bias due to Virtual Texture pool residency.\n")
-	TEXT("Default 0.25"),
+	TEXT("Default 0.2"),
 	ECVF_RenderThreadSafe);
 
 
