Added nullptr check to fix crash when UWorld::UpdateWorldComponents is called on memory-only world, which does not have Scene initialised.

diff --git a/Engine/Source/Runtime/Engine/Private/World.cpp b/Engine/Source/Runtime/Engine/Private/World.cpp
--- a/Engine/Source/Runtime/Engine/Private/World.cpp
+++ b/Engine/Source/Runtime/Engine/Private/World.cpp
@@ -2123,7 +2123,7 @@ void UWorld::UpdateWorldComponents(bool bRerunConstructionScripts, bool bCurrent
 		}
 
 		static IConsoleVariable* PhysicsFieldEnableClipmapCVar = IConsoleManager::Get().FindConsoleVariable(TEXT("r.PhysicsField.EnableField"));
-		if (PhysicsFieldEnableClipmapCVar && PhysicsFieldEnableClipmapCVar->GetInt() == 1 && !IsMobilePlatform(Scene->GetShaderPlatform()))
+		if (PhysicsFieldEnableClipmapCVar && PhysicsFieldEnableClipmapCVar->GetInt() == 1 && Scene && !IsMobilePlatform(Scene->GetShaderPlatform()))
 		{
 			if (!PhysicsField)
 			{
