LWC: Fix LocalVertexFactory lod data shader binding

diff --git a/Engine/Source/Runtime/Engine/Private/LocalVertexFactory.cpp b/Engine/Source/Runtime/Engine/Private/LocalVertexFactory.cpp
--- a/Engine/Source/Runtime/Engine/Private/LocalVertexFactory.cpp
+++ b/Engine/Source/Runtime/Engine/Private/LocalVertexFactory.cpp
@@ -150,7 +150,7 @@ void FLocalVertexFactoryShaderParametersBase::GetElementShaderBindingsBase(
 
 		if (LODParameter.IsBound())
 		{
-			FVector LODData(BatchElement.MinScreenSize, BatchElement.MaxScreenSize, BatchElement.MaxScreenSize - BatchElement.MinScreenSize);
+			FVector3f LODData(BatchElement.MinScreenSize, BatchElement.MaxScreenSize, BatchElement.MaxScreenSize - BatchElement.MinScreenSize);
 			ShaderBindings.Add(LODParameter, LODData);
 		}
 	}
