Fix compile errors in VisualLoggerTraceDevice.cpp

diff --git a/Engine/Source/Runtime/Engine/Private/VisualLogger/VisualLoggerTraceDevice.cpp b/Engine/Source/Runtime/Engine/Private/VisualLogger/VisualLoggerTraceDevice.cpp
--- a/Engine/Source/Runtime/Engine/Private/VisualLogger/VisualLoggerTraceDevice.cpp
+++ b/Engine/Source/Runtime/Engine/Private/VisualLogger/VisualLoggerTraceDevice.cpp
@@ -28,12 +28,16 @@ void FVisualLoggerTraceDevice::Cleanup(bool bReleaseMemory)
 
 void FVisualLoggerTraceDevice::StartRecordingToFile(float TimeStamp)
 {
+#if UE_TRACE_ENABLED
 	UE::Trace::ToggleChannel(TEXT("VisualLogger"), true); 
+#endif
 }
 
 void FVisualLoggerTraceDevice::StopRecordingToFile(float TimeStamp)
 {
+#if UE_TRACE_ENABLED
 	UE::Trace::ToggleChannel(TEXT("VisualLogger"), false); 
+#endif
 }
 
 void FVisualLoggerTraceDevice::DiscardRecordingToFile()
@@ -46,6 +50,7 @@ void FVisualLoggerTraceDevice::SetFileName(const FString& InFileName)
 
 void FVisualLoggerTraceDevice::Serialize(const UObject* LogOwner, FName OwnerName, FName OwnerClassName, const FVisualLogEntry& LogEntry)
 {
+#if OBJECT_TRACE_ENABLED
 	if (UE_TRACE_CHANNELEXPR_IS_ENABLED(VisualLoggerChannel))
 	{
 		FBufferArchive Archive;
@@ -57,5 +62,7 @@ void FVisualLoggerTraceDevice::Serialize(const UObject* LogOwner, FName OwnerNam
 			<< VisualLogEntry.OwnerId(FObjectTrace::GetObjectId(LogOwner))
 			<< VisualLogEntry.LogEntry(Archive.GetData(), Archive.Num());
 	}
+#endif
 }
+
 #endif
