Added extra optionsto UStaticMesh::FBuildMeshDescriptionsParams  when building mesh through UStaticMesh::BuildFromMeshDescriptions() :
- bUseHighPrecisionTangentBasis :  If true, Tangents will be stored at 16 bit vs 8 bit precision.
- bUseFullPrecisionUVs : If true, UVs will be stored at full floating point precision.

diff --git a/Engine/Source/Runtime/Engine/Classes/Engine/StaticMesh.h b/Engine/Source/Runtime/Engine/Classes/Engine/StaticMesh.h
--- a/Engine/Source/Runtime/Engine/Classes/Engine/StaticMesh.h
+++ b/Engine/Source/Runtime/Engine/Classes/Engine/StaticMesh.h
@@ -1425,6 +1425,16 @@ public:
 		 * Ored with the value of bAllowCpuAccess on the static mesh. Set to false by default.
 		 */
 		bool bAllowCpuAccess;
+
+		/**
+		 * If true, Tangents will be stored at 16 bit vs 8 bit precision. 
+		 */
+		bool bUseHighPrecisionTangentBasis = false;
+
+		/**
+		 * If true, UVs will be stored at full floating point precision.
+		 */
+		bool bUseFullPrecisionUVs = false;
 	};
 
 	/**
diff --git a/Engine/Source/Runtime/Engine/Private/StaticMesh.cpp b/Engine/Source/Runtime/Engine/Private/StaticMesh.cpp
--- a/Engine/Source/Runtime/Engine/Private/StaticMesh.cpp
+++ b/Engine/Source/Runtime/Engine/Private/StaticMesh.cpp
@@ -5940,6 +5940,8 @@ bool UStaticMesh::BuildFromMeshDescriptions(const TArray<const FMeshDescription*
 {
 	TRACE_CPUPROFILER_EVENT_SCOPE(UStaticMesh::BuildFromMeshDescriptions);
 
+	const int32 NewNumLODs = MeshDescriptions.Num();
+
 	// Set up
 	NeverStream = true;
 
@@ -5974,9 +5976,12 @@ bool UStaticMesh::BuildFromMeshDescriptions(const TArray<const FMeshDescription*
 		CommitParams.bMarkPackageDirty = Params.bMarkPackageDirty;
 		CommitParams.bUseHashAsGuid = Params.bUseHashAsGuid;
 
-		SetNumSourceModels(MeshDescriptions.Num());
-		for (int32 LODIndex = 0; LODIndex < MeshDescriptions.Num(); LODIndex++)
+		SetNumSourceModels(NewNumLODs);
+		for (int32 LODIndex = 0; LODIndex < NewNumLODs; LODIndex++)
 		{
+			FMeshBuildSettings& LODBuildSettings = GetSourceModel(LODIndex).BuildSettings;
+			LODBuildSettings.bUseHighPrecisionTangentBasis = Params.bUseHighPrecisionTangentBasis;
+			LODBuildSettings.bUseFullPrecisionUVs = Params.bUseFullPrecisionUVs;
 			CreateMeshDescription(LODIndex, *MeshDescriptions[LODIndex]);
 			CommitMeshDescription(LODIndex, CommitParams);
 		}
@@ -5984,12 +5989,14 @@ bool UStaticMesh::BuildFromMeshDescriptions(const TArray<const FMeshDescription*
 	#endif
 
 	SetRenderData(MakeUnique<FStaticMeshRenderData>());
-	GetRenderData()->AllocateLODResources(MeshDescriptions.Num());
+	GetRenderData()->AllocateLODResources(NewNumLODs);
 
 	FStaticMeshLODResourcesArray& LODResourcesArray = GetRenderData()->LODResources;
 	for (int32 LODIndex = 0; LODIndex < LODResourcesArray.Num(); ++LODIndex)
 	{
 		LODResourcesArray[LODIndex].IndexBuffer.TrySetAllowCPUAccess(bAllowCPUAccess || Params.bAllowCpuAccess);
+		LODResourcesArray[LODIndex].VertexBuffers.StaticMeshVertexBuffer.SetUseHighPrecisionTangentBasis(Params.bUseHighPrecisionTangentBasis);
+		LODResourcesArray[LODIndex].VertexBuffers.StaticMeshVertexBuffer.SetUseFullPrecisionUVs(Params.bUseFullPrecisionUVs);
 	}
 
 	// Build render data from each mesh description
@@ -5998,7 +6005,7 @@ bool UStaticMesh::BuildFromMeshDescriptions(const TArray<const FMeshDescription*
 	if (Params.bFastBuild)
 #endif
 	{
-		for (int32 LODIndex = 0; LODIndex < MeshDescriptions.Num(); LODIndex++)
+		for (int32 LODIndex = 0; LODIndex < NewNumLODs; LODIndex++)
 		{
 			check(MeshDescriptions[LODIndex] != nullptr);
 			FStaticMeshLODResources& LODResources = GetRenderData()->LODResources[LODIndex];
@@ -6012,7 +6019,7 @@ bool UStaticMesh::BuildFromMeshDescriptions(const TArray<const FMeshDescription*
 		GetRenderData()->Bounds = MeshDescriptions[0]->GetBounds();
 		CalculateExtendedBounds();
 
-		for (int32 LOD = 0; LOD < MeshDescriptions.Num(); ++LOD)
+		for (int32 LOD = 0; LOD < NewNumLODs; ++LOD)
 		{
 			// @todo: some way of customizing LOD screen size and/or calculate it based on mesh bounds
 			if (true)
@@ -6059,7 +6066,7 @@ bool UStaticMesh::BuildFromMeshDescriptions(const TArray<const FMeshDescription*
 		Build(true);
 	}
 
-	for (int32 LODIndex = 0; LODIndex < MeshDescriptions.Num(); LODIndex++)
+	for (int32 LODIndex = 0; LODIndex < NewNumLODs; LODIndex++)
 	{
 		FStaticMeshLODResources& LODResources = GetRenderData()->LODResources[LODIndex];
 		for (int32 SectionIndex = 0; SectionIndex < LODResources.Sections.Num(); SectionIndex++)
