Fix to assert instead of crash on LOD mistmach between post process and main anim bp

diff --git a/Engine/Source/Runtime/Engine/Private/Animation/AnimNode_LinkedInputPose.cpp b/Engine/Source/Runtime/Engine/Private/Animation/AnimNode_LinkedInputPose.cpp
--- a/Engine/Source/Runtime/Engine/Private/Animation/AnimNode_LinkedInputPose.cpp
+++ b/Engine/Source/Runtime/Engine/Private/Animation/AnimNode_LinkedInputPose.cpp
@@ -48,7 +48,7 @@ void FAnimNode_LinkedInputPose::Evaluate_AnyThread(FPoseContext& Output)
 		Output.Curve.MoveFrom(InputContext.Curve);
 		Output.CustomAttributes.MoveFrom(InputContext.CustomAttributes);
 	}
-	else if(CachedInputPose.IsValid() && CachedInputCurve.IsValid())
+	else if(CachedInputPose.IsValid() && CachedInputCurve.IsValid() && ensure(Output.Pose.GetNumBones() == CachedInputPose.GetNumBones()))
 	{
 		Output.Pose.CopyBonesFrom(CachedInputPose);
 		Output.Curve.CopyFrom(CachedInputCurve);
