Optimus: Added early out to defend against invalidated data providers
Added ensure()/comment to help with code fixing when it happens.

diff --git a/Engine/Source/Runtime/Engine/Private/ComputeFramework/ComputeGraphComponent.cpp b/Engine/Source/Runtime/Engine/Private/ComputeFramework/ComputeGraphComponent.cpp
--- a/Engine/Source/Runtime/Engine/Private/ComputeFramework/ComputeGraphComponent.cpp
+++ b/Engine/Source/Runtime/Engine/Private/ComputeFramework/ComputeGraphComponent.cpp
@@ -58,7 +58,9 @@ void UComputeGraphComponent::SendRenderDynamicData_Concurrent()
 {
 	Super::SendRenderDynamicData_Concurrent();
 	
-	if (!bValidProviders)
+	// If we hit the ensure then something invalidated providers without calling CreateDataProviders().
+	// Those paths DO need fixing. We can remove the ensure() if we ever feel safe enough!
+	if (!bValidProviders || !ensure(ComputeGraph->ValidateProviders(DataProviders)))
 	{
 		// Probably we marked for update just before invalidating providers.
 		return;
