Test Autoamtion - fix incoming screenshot image saved as JPEG instead of PNG since jpg format was enabled for thumbnail
* Isolate code for compressing in PNG and use this code for screenshot image compression
* Mark screenshot as not similar when error message is not empty when considering comparison result

#rnx

diff --git a/Engine/Plugins/Experimental/AutomationUtils/Source/AutomationUtils/Private/AutomationUtilsBlueprintLibrary.cpp b/Engine/Plugins/Experimental/AutomationUtils/Source/AutomationUtils/Private/AutomationUtilsBlueprintLibrary.cpp
--- a/Engine/Plugins/Experimental/AutomationUtils/Source/AutomationUtils/Private/AutomationUtilsBlueprintLibrary.cpp
+++ b/Engine/Plugins/Experimental/AutomationUtils/Source/AutomationUtils/Private/AutomationUtilsBlueprintLibrary.cpp
@@ -256,8 +256,8 @@ void FAutomationUtilsGameplayAutomationScreenshotInstance::HandleScreenshotData(
 	IFileManager::Get().MakeDirectory(*FPaths::GetPath(DeterminedPath), true);
 
 	//Save Image File
-	TArray<uint8> CompressedBitmap;
-	FImageUtils::CompressImageArray(InSizeX, InSizeY, InImageData, CompressedBitmap);
+	TArray64<uint8> CompressedBitmap;
+	FImageUtils::PNGCompressImageArray(InSizeX, InSizeY, TArrayView64<const FColor>(InImageData.GetData(), InImageData.Num()), CompressedBitmap);
 	FFileHelper::SaveArrayToFile(CompressedBitmap, *DeterminedPath);
 	GLog->Log(FString::Printf(TEXT("Saved %d bytes of screenshot image to %s"), CompressedBitmap.Num(), *DeterminedPath));
 
diff --git a/Engine/Source/Developer/ScreenShotComparisonTools/Private/ImageComparer.cpp b/Engine/Source/Developer/ScreenShotComparisonTools/Private/ImageComparer.cpp
--- a/Engine/Source/Developer/ScreenShotComparisonTools/Private/ImageComparer.cpp
+++ b/Engine/Source/Developer/ScreenShotComparisonTools/Private/ImageComparer.cpp
@@ -268,7 +268,7 @@ TSharedPtr<FComparableImage> FImageComparer::Open(const FString& ImagePath, FTex
 	
 	if ( !ImageReader->GetRaw(ERGBFormat::RGBA, 8, Image->Bytes) )
 	{
-		OutError = LOCTEXT("ErrorReadingRawDataA", "Unable decompress ImageA");
+		OutError = LOCTEXT("ErrorReadingRawDataA", "Unable to decompress image");
 		return nullptr;
 	}
 	else
diff --git a/Engine/Source/Developer/ScreenShotComparisonTools/Public/ImageComparer.h b/Engine/Source/Developer/ScreenShotComparisonTools/Public/ImageComparer.h
--- a/Engine/Source/Developer/ScreenShotComparisonTools/Public/ImageComparer.h
+++ b/Engine/Source/Developer/ScreenShotComparisonTools/Public/ImageComparer.h
@@ -298,7 +298,7 @@ public:
 	FText ErrorMessage;
 
 	/*
-		Error message that can be set during a comparison
+		Version of the image comparision result 
 	*/
 	UPROPERTY()
 	int32 Version;
@@ -367,6 +367,11 @@ public:
 			return false;
 		}
 
+		if (!ErrorMessage.IsEmpty())
+		{
+			return false;
+		}
+
 		if ( MaxLocalDifference > Tolerance.MaximumLocalError || GlobalDifference > Tolerance.MaximumGlobalError )
 		{
 			return false;
diff --git a/Engine/Source/Runtime/AutomationWorker/Private/AutomationWorkerModule.cpp b/Engine/Source/Runtime/AutomationWorker/Private/AutomationWorkerModule.cpp
--- a/Engine/Source/Runtime/AutomationWorker/Private/AutomationWorkerModule.cpp
+++ b/Engine/Source/Runtime/AutomationWorker/Private/AutomationWorkerModule.cpp
@@ -448,8 +448,8 @@ void FAutomationWorkerModule::HandleScreenShotAndTraceCapturedWithName(const TAr
 	int32 NewHeight = Data.Height;
 	int32 NewWidth = Data.Width;
 
-	TArray<uint8> CompressedBitmap;
-	FImageUtils::CompressImageArray(NewWidth, NewHeight, RawImageData, CompressedBitmap);
+	TArray64<uint8> CompressedBitmap;
+	FImageUtils::PNGCompressImageArray(NewWidth, NewHeight, TArrayView64<const FColor>(RawImageData.GetData(), RawImageData.Num()), CompressedBitmap);
 
 	FAutomationScreenshotMetadata Metadata(Data);
 		
diff --git a/Engine/Source/Runtime/Engine/Private/ImageUtils.cpp b/Engine/Source/Runtime/Engine/Private/ImageUtils.cpp
--- a/Engine/Source/Runtime/Engine/Private/ImageUtils.cpp
+++ b/Engine/Source/Runtime/Engine/Private/ImageUtils.cpp
@@ -397,7 +397,7 @@ void FImageUtils::CompressImageArray( int32 ImageWidth, int32 ImageHeight, const
 {
 	TArray<FColor> MutableSrcData = SrcData;
 
-	// PNGs are saved as RGBA but FColors are stored as BGRA. An option to swap the order upon compression may be added at 
+	// Thumbnails are saved as RGBA but FColors are stored as BGRA. An option to swap the order upon compression may be added at 
 	// some point. At the moment, manually swapping Red and Blue 
 	for ( int32 Index = 0; Index < ImageWidth*ImageHeight; Index++ )
 	{
@@ -415,12 +415,35 @@ void FImageUtils::CompressImageArray( int32 ImageWidth, int32 ImageHeight, const
 	ThumbnailByteArray.AddUninitialized(MemorySize);
 	FMemory::Memcpy(ThumbnailByteArray.GetData(), MutableSrcData.GetData(), MemorySize);
 
-	// Compress data - convert into a .png
+	// Compress data - convert into thumbnail current format
 	TempThumbnail.CompressImageData();
-	TArray<uint8>& CompressedByteArray = TempThumbnail.AccessCompressedImageData();
 	DstData = TempThumbnail.AccessCompressedImageData();
 }
 
+void FImageUtils::PNGCompressImageArray(int32 ImageWidth, int32 ImageHeight, const TArrayView64<const FColor>& SrcData, TArray64<uint8>& DstData)
+{
+	const int64 PixelsNum = ImageWidth * ImageHeight;
+	check(SrcData.Num() == PixelsNum);
+
+	const uint8* SrcFirstByte = static_cast<const uint8*>(static_cast<const void*>(SrcData.GetData()));
+	const int64 MemorySize = PixelsNum * sizeof(FColor);
+
+	DstData.Reset();
+
+	if (SrcData.Num() > 0 && ImageWidth > 0 && ImageHeight > 0)
+	{
+		if (DstData.Num() == 0)
+		{
+			IImageWrapperModule& ImageWrapperModule = FModuleManager::LoadModuleChecked<IImageWrapperModule>(FName("ImageWrapper"));
+			TSharedPtr<IImageWrapper> ImageWrapper = ImageWrapperModule.CreateImageWrapper(EImageFormat::PNG);
+			if (ImageWrapper.IsValid() && ImageWrapper->SetRaw(SrcFirstByte, MemorySize, ImageWidth, ImageHeight, ERGBFormat::BGRA, 8))
+			{
+				DstData = ImageWrapper->GetCompressed();
+			}
+		}
+	}
+}
+
 UTexture2D* FImageUtils::CreateCheckerboardTexture(FColor ColorOne, FColor ColorTwo, int32 CheckerSize)
 {
 	CheckerSize = FMath::Min<uint32>( FMath::RoundUpToPowerOfTwo(CheckerSize), 4096 );
diff --git a/Engine/Source/Runtime/Engine/Public/ImageUtils.h b/Engine/Source/Runtime/Engine/Public/ImageUtils.h
--- a/Engine/Source/Runtime/Engine/Public/ImageUtils.h
+++ b/Engine/Source/Runtime/Engine/Public/ImageUtils.h
@@ -141,7 +141,7 @@ public:
 	ENGINE_API static void CropAndScaleImage( int32 SrcWidth, int32 SrcHeight, int32 DesiredWidth, int32 DesiredHeight, const TArray<FColor> &SrcData, TArray<FColor> &DstData  );
 
 	/**
-	 * Compress image to .png uint8 array.
+	 * Compress image to thumbnail enabled format (png or jpg) uint8 array.
 	 *
 	 * @param ImageHeight		Source image width.
 	 * @param ImageWidth		Source image height.
@@ -151,6 +151,17 @@ public:
 	 */
 	ENGINE_API static void CompressImageArray( int32 ImageWidth, int32 ImageHeight, const TArray<FColor> &SrcData, TArray<uint8> &DstData );
 
+	/**
+	 * Compress image to PNG format uint8 array.
+	 *
+	 * @param ImageHeight		Source image width.
+	 * @param ImageWidth		Source image height.
+	 * @param SrcData			Raw image array.
+	 * @param DstData			compressed image array.
+	 *
+	 */
+	ENGINE_API static void PNGCompressImageArray(int32 ImageWidth, int32 ImageHeight, const TArrayView64<const FColor>& SrcData, TArray64<uint8>& DstData);
+
 	/**
 	 * Creates a new UTexture2D with a checkerboard pattern.
 	 *
