Adding protection against nullptr deref in SubmixUtils::AreSubmixFormatsCompatible()

diff --git a/Engine/Source/Runtime/Engine/Private/SoundSubmix.cpp b/Engine/Source/Runtime/Engine/Private/SoundSubmix.cpp
--- a/Engine/Source/Runtime/Engine/Private/SoundSubmix.cpp
+++ b/Engine/Source/Runtime/Engine/Private/SoundSubmix.cpp
@@ -991,6 +991,11 @@ void USoundfieldEndpointSubmix::PostEditChangeProperty(struct FPropertyChangedEv
 
 ENGINE_API bool SubmixUtils::AreSubmixFormatsCompatible(const USoundSubmixBase* ChildSubmix, const USoundSubmixBase* ParentSubmix)
 {
+	if (!ChildSubmix || !ParentSubmix)
+	{
+		return false;
+	}
+
 	const USoundfieldSubmix* ChildSoundfieldSubmix = Cast<const USoundfieldSubmix>(ChildSubmix);
 
 	// If both the child and parent are soundfield submixes, ensure that their formats are compatible.
