Temporarily remove check and add previous code back in until I can figure why the assert is triggering.

diff --git a/Engine/Source/Runtime/Engine/Public/RawIndexBuffer.h b/Engine/Source/Runtime/Engine/Public/RawIndexBuffer.h
--- a/Engine/Source/Runtime/Engine/Public/RawIndexBuffer.h
+++ b/Engine/Source/Runtime/Engine/Public/RawIndexBuffer.h
@@ -544,8 +544,6 @@ private:
 	{
 		if (CachedNumIndices)
 		{
-			check(CachedNumIndices == Indices.Num());
-
 			// Create the index buffer.
 			FRHIResourceCreateInfo CreateInfo(sizeof(INDEX_TYPE) == 4 ? TEXT("FRawStaticIndexBuffer32") : TEXT("FRawStaticIndexBuffer16"), &Indices);
 			EBufferUsageFlags Flags = BUF_Static;
@@ -556,8 +554,12 @@ private:
 				Flags = (EBufferUsageFlags)(Flags | BUF_ShaderResource);
 			}
 
+			// Need to cache number of indices from the source array *before* RHICreateIndexBuffer is called
+			// because it will empty the source array.
+			CachedNumIndices = Indices.Num();
+
 			FBufferRHIRef Ret;
-			const uint32 Size = Indices.Num() * sizeof(INDEX_TYPE);
+			const uint32 Size = CachedNumIndices * sizeof(INDEX_TYPE);
 			CreateInfo.bWithoutNativeResource = !Size;
 			if (bRenderThread)
 			{
