Merging pull requests to improve rendering determinism through more stable draw sorting.

* Use a shader code hash instead of sorting by shader pointer in CalculateBasePassMeshStaticSortKey()
* Added a serial number to UPrimitiveComponent, which is generated in OnRegister() and then used in FPrimitiveArraySortKey as an extra sort key.

Contributed by [REDACTED_PERSON].

#github 7934
#github 8273

diff --git a/Engine/Shaders/Public/ShaderVersion.ush b/Engine/Shaders/Public/ShaderVersion.ush
--- a/Engine/Shaders/Public/ShaderVersion.ush
+++ b/Engine/Shaders/Public/ShaderVersion.ush
@@ -3,4 +3,4 @@
 // in Platform.ush (which should be included in any shader) it allows to invalidate the shader DDC.
 // 
 // If you are merging streams and there is a conflict with this GUID you should make a new GUID rather than taking one or the other.
-// GUID = 70C261F895B946BB9B4960B1B0E15893
+// GUID = DCA69B23FEF84CBDB4BED51526ECE280
diff --git a/Engine/Source/Runtime/Engine/Classes/Components/PrimitiveComponent.h b/Engine/Source/Runtime/Engine/Classes/Components/PrimitiveComponent.h
--- a/Engine/Source/Runtime/Engine/Classes/Components/PrimitiveComponent.h
+++ b/Engine/Source/Runtime/Engine/Classes/Components/PrimitiveComponent.h
@@ -749,6 +749,12 @@ public:
 	/** Used by the renderer, to identify a component across re-registers. */
 	FPrimitiveComponentId ComponentId;
 
+	/**
+	* Identifier used to track the time that this component was registered with the world / renderer.
+	* Updated to unique incremental value each time OnRegister() is called. The value of 0 is unused.
+	* */
+	int32 RegistrationSerialNumber = -1;
+
 	/**
 	* Incremented by the main thread before being attached to the scene, decremented
 	* by the rendering thread after removal. This counter exists to assert that 
@@ -784,6 +790,9 @@ protected:
 	/** Next id to be used by a component. */
 	static FThreadSafeCounter NextComponentId;
 
+	/** Next registration serial number to be assigned to a component when it is registered. */
+	static FThreadSafeCounter NextRegistrationSerialNumber;
+
 public:
 	/** 
 	 * Scales the bounds of the object.
diff --git a/Engine/Source/Runtime/Engine/Private/Components/PrimitiveComponent.cpp b/Engine/Source/Runtime/Engine/Private/Components/PrimitiveComponent.cpp
--- a/Engine/Source/Runtime/Engine/Private/Components/PrimitiveComponent.cpp
+++ b/Engine/Source/Runtime/Engine/Private/Components/PrimitiveComponent.cpp
@@ -284,6 +284,9 @@ FORCEINLINE_DEBUGGABLE static void GetPointersToArrayDataByPredicate(TArray<cons
 
 uint32 UPrimitiveComponent::GlobalOverlapEventsCounter = 0;
 
+// 0 is reserved to mean invalid
+FThreadSafeCounter UPrimitiveComponent::NextRegistrationSerialNumber;
+
 // 0 is reserved to mean invalid
 FThreadSafeCounter UPrimitiveComponent::NextComponentId;
 
@@ -628,6 +631,9 @@ void UPrimitiveComponent::OnRegister()
 {
 	Super::OnRegister();
 
+	// Deterministically track primitives via registration sequence numbers.
+	RegistrationSerialNumber = NextRegistrationSerialNumber.Increment();
+
 	if (bCanEverAffectNavigation)
 	{
 		const bool bNavRelevant = bNavigationRelevant = IsNavigationRelevant();
diff --git a/Engine/Source/Runtime/Engine/Private/ShaderDerivedDataVersion.h b/Engine/Source/Runtime/Engine/Private/ShaderDerivedDataVersion.h
--- a/Engine/Source/Runtime/Engine/Private/ShaderDerivedDataVersion.h
+++ b/Engine/Source/Runtime/Engine/Private/ShaderDerivedDataVersion.h
@@ -10,5 +10,5 @@
 
 // In case of merge conflicts with DDC versions, you MUST generate a new GUID and set this new
 // guid as version
-#define GLOBALSHADERMAP_DERIVEDDATA_VER			TEXT("0357918925494B80A54D1CC105798A81")
-#define MATERIALSHADERMAP_DERIVEDDATA_VER		TEXT("BCC06CA191E448B99C80CA33353AE671")
\ No newline at end of file
+#define GLOBALSHADERMAP_DERIVEDDATA_VER			TEXT("60CD2243B26E480D898BAA7C66094C0A")
+#define MATERIALSHADERMAP_DERIVEDDATA_VER		TEXT("20297056A97D43148C19F0F2D16CFE78")
\ No newline at end of file
diff --git a/Engine/Source/Runtime/RenderCore/Private/Shader.cpp b/Engine/Source/Runtime/RenderCore/Private/Shader.cpp
--- a/Engine/Source/Runtime/RenderCore/Private/Shader.cpp
+++ b/Engine/Source/Runtime/RenderCore/Private/Shader.cpp
@@ -536,9 +536,13 @@ FShader::FShader(const CompiledShaderInitializerType& Initializer)
 	, CodeSize(Initializer.CodeSize)
 #endif // WITH_EDITORONLY_DATA
 {
-#if WITH_EDITORONLY_DATA
 	checkSlow(Initializer.OutputHash != FSHAHash());
-	
+
+	// Only store a truncated hash to minimize memory overhead
+	static_assert(sizeof(SortKey) <= sizeof(Initializer.OutputHash.Hash));
+	FMemory::Memcpy(&SortKey, Initializer.OutputHash.Hash, sizeof(SortKey));
+
+#if WITH_EDITORONLY_DATA
 	OutputHash = Initializer.OutputHash;
 
 	// Store off the source hash that this shader was compiled with
diff --git a/Engine/Source/Runtime/RenderCore/Public/Shader.h b/Engine/Source/Runtime/RenderCore/Public/Shader.h
--- a/Engine/Source/Runtime/RenderCore/Public/Shader.h
+++ b/Engine/Source/Runtime/RenderCore/Public/Shader.h
@@ -715,6 +715,9 @@ public:
 	
 	const FSHAHash& GetOutputHash() const;
 
+	/** Returns an identifier suitable for deterministic sorting of shaders between sessions. */
+	uint32 GetSortKey() const { return SortKey; }
+
 	void Finalize(const FShaderMapResourceCode* Code);
 
 	// Accessors.
@@ -835,6 +838,9 @@ private:
 	/** The number of instructions the shader takes to execute. */
 	LAYOUT_FIELD(uint32, NumInstructions);
 
+	/** Truncated version of OutputHash, intended for sorting. Not suitable for unique shader identification. */
+	LAYOUT_FIELD(uint32, SortKey);
+
 	/** Number of texture samplers the shader uses. */
 	LAYOUT_FIELD_EDITORONLY(uint32, NumTextureSamplers);
 
diff --git a/Engine/Source/Runtime/Renderer/Private/BasePassRendering.cpp b/Engine/Source/Runtime/Renderer/Private/BasePassRendering.cpp
--- a/Engine/Source/Runtime/Renderer/Private/BasePassRendering.cpp
+++ b/Engine/Source/Runtime/Renderer/Private/BasePassRendering.cpp
@@ -279,8 +279,8 @@ FMeshDrawCommandSortKey CalculateTranslucentMeshStaticSortKey(const FPrimitiveSc
 FMeshDrawCommandSortKey CalculateBasePassMeshStaticSortKey(EDepthDrawingMode EarlyZPassMode, EBlendMode BlendMode, const FMeshMaterialShader* VertexShader, const FMeshMaterialShader* PixelShader)
 {
 	FMeshDrawCommandSortKey SortKey;
-	SortKey.BasePass.VertexShaderHash = PointerHash(VertexShader) & 0xFFFF;
-	SortKey.BasePass.PixelShaderHash = PointerHash(PixelShader);
+	SortKey.BasePass.VertexShaderHash = (VertexShader ? VertexShader->GetSortKey() : 0) & 0xFFFF;
+	SortKey.BasePass.PixelShaderHash = PixelShader ? PixelShader->GetSortKey() : 0;
 	if (EarlyZPassMode != DDM_None)
 	{
 		SortKey.BasePass.Masked = BlendMode == EBlendMode::BLEND_Masked ? 0 : 1;
diff --git a/Engine/Source/Runtime/Renderer/Private/DepthRendering.cpp b/Engine/Source/Runtime/Renderer/Private/DepthRendering.cpp
--- a/Engine/Source/Runtime/Renderer/Private/DepthRendering.cpp
+++ b/Engine/Source/Runtime/Renderer/Private/DepthRendering.cpp
@@ -695,14 +695,14 @@ FMeshDrawCommandSortKey CalculateDepthPassMeshStaticSortKey(EBlendMode BlendMode
 	FMeshDrawCommandSortKey SortKey;
 	if (GEarlyZSortMasked)
 	{
-		SortKey.BasePass.VertexShaderHash = PointerHash(VertexShader) & 0xFFFF;
-		SortKey.BasePass.PixelShaderHash = PointerHash(PixelShader);
+		SortKey.BasePass.VertexShaderHash = (VertexShader ? VertexShader->GetSortKey() : 0) & 0xFFFF;
+		SortKey.BasePass.PixelShaderHash = PixelShader ? PixelShader->GetSortKey() : 0;
 		SortKey.BasePass.Masked = BlendMode == EBlendMode::BLEND_Masked ? 1 : 0;
 	}
 	else
 	{
-		SortKey.Generic.VertexShaderHash = PointerHash(VertexShader);
-		SortKey.Generic.PixelShaderHash = PointerHash(PixelShader);
+		SortKey.Generic.VertexShaderHash = VertexShader ? VertexShader->GetSortKey() : 0;
+		SortKey.Generic.PixelShaderHash = PixelShader ? PixelShader->GetSortKey() : 0;
 	}
 	
 	return SortKey;
diff --git a/Engine/Source/Runtime/Renderer/Private/MeshPassProcessor.cpp b/Engine/Source/Runtime/Renderer/Private/MeshPassProcessor.cpp
--- a/Engine/Source/Runtime/Renderer/Private/MeshPassProcessor.cpp
+++ b/Engine/Source/Runtime/Renderer/Private/MeshPassProcessor.cpp
@@ -1359,8 +1359,8 @@ void DrawDynamicMeshPassPrivate(
 FMeshDrawCommandSortKey CalculateMeshStaticSortKey(const FMeshMaterialShader* VertexShader, const FMeshMaterialShader* PixelShader)
 {
 	FMeshDrawCommandSortKey SortKey;
-	SortKey.Generic.VertexShaderHash = PointerHash(VertexShader);
-	SortKey.Generic.PixelShaderHash = PointerHash(PixelShader);
+	SortKey.Generic.VertexShaderHash = VertexShader ? VertexShader->GetSortKey() : 0;
+	SortKey.Generic.PixelShaderHash = PixelShader ? PixelShader->GetSortKey() : 0;
 
 	return SortKey;
 }
diff --git a/Engine/Source/Runtime/Renderer/Private/PrimitiveSceneInfo.cpp b/Engine/Source/Runtime/Renderer/Private/PrimitiveSceneInfo.cpp
--- a/Engine/Source/Runtime/Renderer/Private/PrimitiveSceneInfo.cpp
+++ b/Engine/Source/Runtime/Renderer/Private/PrimitiveSceneInfo.cpp
@@ -134,6 +134,7 @@ FPrimitiveSceneInfoCompact::FPrimitiveSceneInfoCompact(FPrimitiveSceneInfo* InPr
 FPrimitiveSceneInfo::FPrimitiveSceneInfo(UPrimitiveComponent* InComponent,FScene* InScene):
 	Proxy(InComponent->SceneProxy),
 	PrimitiveComponentId(InComponent->ComponentId),
+	RegistrationSerialNumber(InComponent->RegistrationSerialNumber),
 	OwnerLastRenderTime(FActorLastRenderTime::GetPtr(InComponent->GetOwner())),
 	IndirectLightingCacheAllocation(NULL),
 	CachedPlanarReflectionProxy(NULL),
diff --git a/Engine/Source/Runtime/Renderer/Private/RendererScene.cpp b/Engine/Source/Runtime/Renderer/Private/RendererScene.cpp
--- a/Engine/Source/Runtime/Renderer/Private/RendererScene.cpp
+++ b/Engine/Source/Runtime/Renderer/Private/RendererScene.cpp
@@ -3911,7 +3911,17 @@ struct FPrimitiveArraySortKey
 {
 	inline bool operator()(const FPrimitiveSceneInfo& A, const FPrimitiveSceneInfo& B) const
 	{
-		return A.Proxy->GetTypeHash() < B.Proxy->GetTypeHash();
+		uint32 AHash = A.Proxy->GetTypeHash();
+		uint32 BHash = B.Proxy->GetTypeHash();
+
+		if (AHash == BHash) 
+		{
+			return A.RegistrationSerialNumber < B.RegistrationSerialNumber;
+		}
+		else
+		{
+			return AHash < BHash;
+		}
 	}
 };
 
diff --git a/Engine/Source/Runtime/Renderer/Public/PrimitiveSceneInfo.h b/Engine/Source/Runtime/Renderer/Public/PrimitiveSceneInfo.h
--- a/Engine/Source/Runtime/Renderer/Public/PrimitiveSceneInfo.h
+++ b/Engine/Source/Runtime/Renderer/Public/PrimitiveSceneInfo.h
@@ -263,6 +263,13 @@ public:
 	 */
 	FPrimitiveComponentId PrimitiveComponentId;
 
+	/**
+	 * Number assigned to this component when it was registered with the world.
+	 * This will only ever be updated if the object is re-registered.
+	 * Used by FPrimitiveArraySortKey for deterministic ordering.
+	 */
+	int32 RegistrationSerialNumber;
+
 	/** 
 	 * Pointer to the last render time variable on the primitive's owning actor (if owned), which is written to by the RT and read by the GT.
 	 * The value of LastRenderTime will therefore not be deterministic due to race conditions, but the GT uses it in a way that allows this.
