Data Layers debug display tweaks: bigger fonts, show unloaded layers, color changes, etc.

diff --git a/Engine/Source/Runtime/Engine/Private/WorldPartition/DataLayer/DataLayerSubsystem.cpp b/Engine/Source/Runtime/Engine/Private/WorldPartition/DataLayer/DataLayerSubsystem.cpp
--- a/Engine/Source/Runtime/Engine/Private/WorldPartition/DataLayer/DataLayerSubsystem.cpp
+++ b/Engine/Source/Runtime/Engine/Private/WorldPartition/DataLayer/DataLayerSubsystem.cpp
@@ -179,28 +179,46 @@ void UDataLayerSubsystem::DrawDataLayersStatus(UCanvas* Canvas, FVector2D& Offse
 	FVector2D Pos = Offset;
 	float MaxTextWidth = 0.f;
 
-	TMap<FName, FColor> ColorMapping;
-	GetDataLayerDebugColors(ColorMapping);
-
-	auto DrawLayerNames = [this, Canvas, &Pos, &MaxTextWidth, &ColorMapping](const FString& Title, FColor Color, const TSet<FName>& LayerNames)
+	auto DrawLayerNames = [this, Canvas, &Pos, &MaxTextWidth](const FString& Title, FColor HeaderColor, FColor TextColor, const TSet<FName>& LayerNames)
 	{
 		if (LayerNames.Num() > 0)
 		{
-			FWorldPartitionDebugHelper::DrawText(Canvas, Title, GEngine->GetSmallFont(), Color, Pos, &MaxTextWidth);
+			FWorldPartitionDebugHelper::DrawText(Canvas, Title, GEngine->GetSmallFont(), HeaderColor, Pos, &MaxTextWidth);
 
-			UFont* DataLayerFont = GEngine->GetTinyFont();
+			UFont* DataLayerFont = GEngine->GetSmallFont();
 			for (const FName& DataLayerName : LayerNames)
 			{
 				if (UDataLayer* DataLayer = GetDataLayerFromName(DataLayerName))
 				{
-					FWorldPartitionDebugHelper::DrawLegendItem(Canvas, *DataLayer->GetDataLayerLabel().ToString(), DataLayerFont, ColorMapping[DataLayerName], Pos, &MaxTextWidth);
+					FWorldPartitionDebugHelper::DrawLegendItem(Canvas, *DataLayer->GetDataLayerLabel().ToString(), DataLayerFont, DataLayer->GetDebugColor(), TextColor, Pos, &MaxTextWidth);
 				}
 			}
 		}
 	};
 
-	DrawLayerNames(TEXT("Loaded Data Layers"), FColor::Cyan, GetLoadedDataLayerNames());
-	DrawLayerNames(TEXT("Active Data Layers"), FColor::Green, GetActiveDataLayerNames());
+	const TSet<FName> LoadedDataLayers = GetLoadedDataLayerNames();
+	const TSet<FName> ActiveDataLayers = GetActiveDataLayerNames();
+
+	DrawLayerNames(TEXT("Loaded Data Layers"), FColor::Cyan, FColor::White, LoadedDataLayers);
+	DrawLayerNames(TEXT("Active Data Layers"), FColor::Green, FColor::White, ActiveDataLayers);
+
+	if (AWorldDataLayers* WorldDataLayers = GetWorld()->GetWorldDataLayers())
+	{
+		TSet<FName> UnloadedDataLayers;
+		WorldDataLayers->ForEachDataLayer([&LoadedDataLayers, &ActiveDataLayers, &UnloadedDataLayers](UDataLayer* DataLayer)
+		{
+			if (DataLayer->IsDynamicallyLoaded())
+			{
+				const FName DataLayerName = DataLayer->GetFName();
+				if (!LoadedDataLayers.Contains(DataLayerName) && !ActiveDataLayers.Contains(DataLayerName))
+				{
+					UnloadedDataLayers.Add(DataLayerName);
+				}
+			}
+			return true;
+		});
+		DrawLayerNames(TEXT("Unloaded Data Layers"), FColor::Silver, FColor(192,192,192), UnloadedDataLayers);
+	}
 
 	Offset.X += MaxTextWidth + 10;
 }
diff --git a/Engine/Source/Runtime/Engine/Private/WorldPartition/WorldPartitionDebugHelper.cpp b/Engine/Source/Runtime/Engine/Private/WorldPartition/WorldPartitionDebugHelper.cpp
--- a/Engine/Source/Runtime/Engine/Private/WorldPartition/WorldPartitionDebugHelper.cpp
+++ b/Engine/Source/Runtime/Engine/Private/WorldPartition/WorldPartitionDebugHelper.cpp
@@ -152,7 +152,7 @@ void FWorldPartitionDebugHelper::DrawText(UCanvas* Canvas, const FString& Text,
 	}
 }
 
-void FWorldPartitionDebugHelper::DrawLegendItem(UCanvas* Canvas, const FString& Text, const UFont* Font, const FColor& Color, FVector2D& Pos, float* MaxItemWidth)
+void FWorldPartitionDebugHelper::DrawLegendItem(UCanvas* Canvas, const FString& Text, const UFont* Font, const FColor& Color, const FColor& TextColor, FVector2D& Pos, float* MaxItemWidth)
 {
 	static const FVector2D ItemSize(12, 12);
 	
@@ -163,7 +163,7 @@ void FWorldPartitionDebugHelper::DrawLegendItem(UCanvas* Canvas, const FString&
 
 	FVector2D TextPos(Pos.X + ItemSize.X + 10, Pos.Y);
 	float TextWidth = 0.f;
-	FWorldPartitionDebugHelper::DrawText(Canvas, Text, Font, FColor::White, TextPos, &TextWidth);
+	FWorldPartitionDebugHelper::DrawText(Canvas, Text, Font, TextColor, TextPos, &TextWidth);
 	if (MaxItemWidth)
 	{
 		*MaxItemWidth = FMath::Max(*MaxItemWidth, TextWidth + ItemSize.X + 10);
diff --git a/Engine/Source/Runtime/Engine/Private/WorldPartition/WorldPartitionLevelStreamingPolicy.cpp b/Engine/Source/Runtime/Engine/Private/WorldPartition/WorldPartitionLevelStreamingPolicy.cpp
--- a/Engine/Source/Runtime/Engine/Private/WorldPartition/WorldPartitionLevelStreamingPolicy.cpp
+++ b/Engine/Source/Runtime/Engine/Private/WorldPartition/WorldPartitionLevelStreamingPolicy.cpp
@@ -270,7 +270,7 @@ void UWorldPartitionLevelStreamingPolicy::DrawStreamingStatusLegend(UCanvas* Can
 		{
 			DebugString += *FString::Printf(TEXT(" (%d)"), StatusCount[(int32)Status]);
 		}
-		FWorldPartitionDebugHelper::DrawLegendItem(Canvas, *DebugString, GEngine->GetTinyFont(), StatusColor, Pos, &MaxTextWidth);
+		FWorldPartitionDebugHelper::DrawLegendItem(Canvas, *DebugString, GEngine->GetSmallFont(), StatusColor, FColor::White, Pos, &MaxTextWidth);
 	}
 
 	Offset.X += MaxTextWidth + 10;
diff --git a/Engine/Source/Runtime/Engine/Public/WorldPartition/WorldPartitionDebugHelper.h b/Engine/Source/Runtime/Engine/Public/WorldPartition/WorldPartitionDebugHelper.h
--- a/Engine/Source/Runtime/Engine/Public/WorldPartition/WorldPartitionDebugHelper.h
+++ b/Engine/Source/Runtime/Engine/Public/WorldPartition/WorldPartitionDebugHelper.h
@@ -19,7 +19,7 @@ public:
 	static bool IsDebugRuntimeHashGridShown(FName Name);
 	static bool IsDebugCellNameShow(const FString& Name);
 	static void DrawText(UCanvas* Canvas, const FString& Text, const UFont* Font, const FColor& Color, FVector2D& Pos, float* MaxTextWidth = nullptr);
-	static void DrawLegendItem(UCanvas* Canvas, const FString& Text, const UFont* Font, const FColor& Color, FVector2D& Pos, float* MaxItemWidth = nullptr);
+	static void DrawLegendItem(UCanvas* Canvas, const FString& Text, const UFont* Font, const FColor& Color, const FColor& TextColor, FVector2D& Pos, float* MaxItemWidth = nullptr);
 
 private:
 	static FAutoConsoleCommand DebugFilterByStreamingStatusCommand;
