Texture assets turn on Oodle with PNGz0
replaces PNG zlib compression with Oodle
should be smaller and faster to load

diff --git a/Engine/Source/Runtime/Engine/Private/Texture.cpp b/Engine/Source/Runtime/Engine/Private/Texture.cpp
--- a/Engine/Source/Runtime/Engine/Private/Texture.cpp
+++ b/Engine/Source/Runtime/Engine/Private/Texture.cpp
@@ -1281,8 +1281,8 @@ void FTextureSource::Compress()
 #endif
 
 	// if bUseOodleOnPNGz0 , do PNG filters but then use Oodle instead of zlib back-end LZ
-	//	no point turning this on until we drill Oodle compressor options through BulkData
-	bool bUseOodleOnPNGz0 = false;
+	//	should be faster to load and also smaller files (than traditional PNG+zlib)
+	bool bUseOodleOnPNGz0 = true;
 
 	// may already have bPNGCompressed or "CompressionFormat" set
 
@@ -1306,7 +1306,6 @@ void FTextureSource::Compress()
 			if ( CompressedData.Num() > 0 )
 			{
 				BulkData.UpdatePayload(MakeSharedBufferFromArray(MoveTemp(CompressedData)));
-				BulkData.SetCompressionOptions(UE::Virtualization::ECompressionOptions::Disabled);
 
 				bPNGCompressed = true;
 				CompressionFormat = TSCF_PNG;
@@ -1320,17 +1319,14 @@ void FTextureSource::Compress()
 		CompressionFormat = TSCF_PNG;
 	}
 
-	if (CompressionFormat == TSCF_PNG && bUseOodleOnPNGz0)
+	if ( ( CompressionFormat == TSCF_PNG && bUseOodleOnPNGz0 ) ||
+		CompressionFormat == TSCF_None )
 	{
-		BulkData.SetCompressionOptions(UE::Virtualization::ECompressionOptions::Default);
+		BulkData.SetCompressionOptions(FOodleDataCompression::ECompressor::Kraken,FOodleDataCompression::ECompressionLevel::Fast);
 	}
-	else if (CompressionFormat != TSCF_None)
+	else
 	{
-		// There was a brief period when textures stored with a compressed image format could also have ZLib compression
-		// applied via bulkdata serialization, so we need to call ::SetCompressionFormat directly to make sure that the 
-		// disable compression flag is set internally if the image storage format is compressed.
-		// TODO: We should consider versioning this so that we can remove this code path at some point.
-		BulkData.SetCompressionOptions(FOodleDataCompression::ECompressor::NotSet, FOodleDataCompression::ECompressionLevel::None);
+		BulkData.SetCompressionOptions(UE::Virtualization::ECompressionOptions::Disabled);
 	}
 }
 
