Add function to set the velocity on child bones from a specified parent bone on a skeletal mesh

diff --git a/Engine/Source/Runtime/Engine/Classes/Components/SkeletalMeshComponent.h b/Engine/Source/Runtime/Engine/Classes/Components/SkeletalMeshComponent.h
--- a/Engine/Source/Runtime/Engine/Classes/Components/SkeletalMeshComponent.h
+++ b/Engine/Source/Runtime/Engine/Classes/Components/SkeletalMeshComponent.h
@@ -1986,6 +1986,13 @@ public:
 	UFUNCTION(BlueprintCallable, Category = "Physics")
 	void SetAllBodiesBelowPhysicsDisabled(const FName& InBoneName, bool bDisabled, bool bIncludeSelf = true);
 
+	/** set the linear velocity of the child bodies to match that of the given parent bone */
+	UFUNCTION(BlueprintCallable, Category = "Physics")
+	void SetAllBodiesBelowLinearVelocity(const FName& InBoneName, const FVector& LinearVelocity, bool bIncludeSelf = true);
+
+	UFUNCTION(BlueprintCallable, Category = "Physics")
+	FVector GetBoneLinearVelocity(const FName& InBoneName);
+
 	/** Set all of the bones below passed in bone to be simulated */
 	UFUNCTION(BlueprintCallable, Category="Physics")
 	void SetAllBodiesBelowSimulatePhysics(const FName& InBoneName, bool bNewSimulate, bool bIncludeSelf = true );
diff --git a/Engine/Source/Runtime/Engine/Private/SkeletalMeshComponentPhysics.cpp b/Engine/Source/Runtime/Engine/Private/SkeletalMeshComponentPhysics.cpp
--- a/Engine/Source/Runtime/Engine/Private/SkeletalMeshComponentPhysics.cpp
+++ b/Engine/Source/Runtime/Engine/Private/SkeletalMeshComponentPhysics.cpp
@@ -1133,6 +1133,30 @@ void USkeletalMeshComponent::SetAllBodiesBelowPhysicsDisabled(const FName& InBon
 	});
 }
 
+void USkeletalMeshComponent::SetAllBodiesBelowLinearVelocity(const FName& InBoneName, const FVector& LinearVelocity, bool bIncludeSelf)
+{
+	if (const FBodyInstance* ParentBodyInstance = GetBodyInstance(InBoneName))
+	{
+		ForEachBodyBelow(InBoneName, bIncludeSelf, /*bSkipCustomPhysicsType=*/ false, [LinearVelocity, this](FBodyInstance* BI)
+			{
+				BI->SetLinearVelocity(LinearVelocity, false);
+			});
+	}
+}
+
+FVector USkeletalMeshComponent::GetBoneLinearVelocity(const FName& InBoneName)
+{
+	FVector OutVelocity = FVector::ZeroVector;
+
+	if (const FBodyInstance* ParentBodyInstance = GetBodyInstance(InBoneName))
+	{
+		OutVelocity = ParentBodyInstance->GetUnrealWorldVelocity();
+	}
+
+	return OutVelocity;
+}
+
+
 void USkeletalMeshComponent::SetAllBodiesBelowSimulatePhysics( const FName& InBoneName, bool bNewSimulate, bool bIncludeSelf )
 {
 	int32 NumBodiesFound = ForEachBodyBelow(InBoneName, bIncludeSelf, /*bSkipCustomPhysicsType=*/ false, [bNewSimulate](FBodyInstance* BI)
