HLODs builder changes:
- Disable collisions on generated static meshes.
- Rename generated static meshes so we know from which layer they come from.

diff --git a/Engine/Source/Developer/WorldPartitionHLODUtilities/Private/WorldPartitionHLODUtilities.cpp b/Engine/Source/Developer/WorldPartitionHLODUtilities/Private/WorldPartitionHLODUtilities.cpp
--- a/Engine/Source/Developer/WorldPartitionHLODUtilities/Private/WorldPartitionHLODUtilities.cpp
+++ b/Engine/Source/Developer/WorldPartitionHLODUtilities/Private/WorldPartitionHLODUtilities.cpp
@@ -93,7 +93,7 @@ static uint32 ComputeHLODHash(AWorldPartitionHLOD* InHLODActor, const TArray<FWo
 	FArchiveCrc32 Ar;
 
 	// Base key, changing this will force a rebuild of all HLODs
-	FString HLODBaseKey = "E0745F6E31AC49EC89BDD6D1F6A87968";
+	FString HLODBaseKey = "29766A60F4A74BF9B10B77095F4B3B16";
 	Ar << HLODBaseKey;
 
 	// HLOD Layer
diff --git a/Engine/Source/Runtime/Engine/Private/WorldPartition/HLOD/HLODBuilder.cpp b/Engine/Source/Runtime/Engine/Private/WorldPartition/HLOD/HLODBuilder.cpp
--- a/Engine/Source/Runtime/Engine/Private/WorldPartition/HLOD/HLODBuilder.cpp
+++ b/Engine/Source/Runtime/Engine/Private/WorldPartition/HLOD/HLODBuilder.cpp
@@ -4,6 +4,7 @@
 
 #include "Engine/StaticMesh.h"
 #include "Components/InstancedStaticMeshComponent.h"
+#include "PhysicsEngine/BodySetup.h"
 
 #include "WorldPartition/HLOD/HLODActor.h"
 #include "WorldPartition/HLOD/HLODLayer.h"
@@ -53,8 +54,15 @@ void UHLODBuilder::Build(AWorldPartitionHLOD* InHLODActor, const UHLODLayer* InH
 
 	for (UPrimitiveComponent* HLODPrimitive : HLODPrimitives)
 	{
-		DisableCollisions(HLODPrimitive);
-
+		// Disable collisions
+		HLODPrimitive->SetCollisionProfileName(UCollisionProfile::NoCollision_ProfileName);
+		HLODPrimitive->SetGenerateOverlapEvents(false);
+		HLODPrimitive->SetCanEverAffectNavigation(false);
+		HLODPrimitive->CanCharacterStepUpOn = ECanBeCharacterBase::ECB_No;
+		HLODPrimitive->SetCanEverAffectNavigation(false);
+		HLODPrimitive->SetCollisionEnabled(ECollisionEnabled::NoCollision);
+
+		// Enable optimizations
 		HLODPrimitive->bComputeFastLocalBounds = true;
 		HLODPrimitive->bComputeBoundsOnceDuringCook = true;
 
@@ -62,7 +70,24 @@ void UHLODBuilder::Build(AWorldPartitionHLOD* InHLODActor, const UHLODLayer* InH
 		{
 			if (UStaticMeshComponent* StaticMeshComponent = Cast<UStaticMeshComponent>(HLODPrimitive))
 			{
-				StaticMeshComponent->bRayTracingFarField = (StaticMeshComponent->GetStaticMesh() && StaticMeshComponent->GetStaticMesh()->bSupportRayTracing);
+				if (UStaticMesh* StaticMesh = StaticMeshComponent->GetStaticMesh())
+				{
+					// Set up ray tracing far fields
+					StaticMeshComponent->bRayTracingFarField = StaticMesh->bSupportRayTracing;
+
+					if (HLODPrimitive->GetPackage() == StaticMesh->GetPackage())
+					{
+						// Disable collisions on owned static mesh
+						if (UBodySetup* BodySetup = StaticMesh->GetBodySetup())
+						{							
+							BodySetup->DefaultInstance.SetCollisionProfileName(UCollisionProfile::NoCollision_ProfileName);
+							BodySetup->CollisionTraceFlag = CTF_UseSimpleAsComplex;
+						}
+						
+						// Rename owned static mesh
+						StaticMesh->Rename(*MakeUniqueObjectName(StaticMesh->GetOuter(), StaticMesh->GetClass(), *FString::Printf(TEXT("StaticMesh_%s"), *InHLODLayer->GetName())).ToString());
+					}
+				}
 			}
 		}
 	}
@@ -137,14 +162,4 @@ TArray<UPrimitiveComponent*> UHLODBuilder::GatherPrimitiveComponents(const TArra
 	return PrimitiveComponents;
 }
 
-void UHLODBuilder::DisableCollisions(UPrimitiveComponent* Component)
-{
-	Component->SetCollisionProfileName(UCollisionProfile::NoCollision_ProfileName);
-	Component->SetGenerateOverlapEvents(false);
-	Component->SetCanEverAffectNavigation(false);
-	Component->CanCharacterStepUpOn = ECanBeCharacterBase::ECB_No;
-	Component->SetCanEverAffectNavigation(false);
-	Component->SetCollisionEnabled(ECollisionEnabled::NoCollision);
-}
-
 #endif // WITH_EDITOR
diff --git a/Engine/Source/Runtime/Engine/Public/WorldPartition/HLOD/HLODBuilder.h b/Engine/Source/Runtime/Engine/Public/WorldPartition/HLOD/HLODBuilder.h
--- a/Engine/Source/Runtime/Engine/Public/WorldPartition/HLOD/HLODBuilder.h
+++ b/Engine/Source/Runtime/Engine/Public/WorldPartition/HLOD/HLODBuilder.h
@@ -43,7 +43,5 @@ public:
 	void Build(AWorldPartitionHLOD* InHLODActor, const UHLODLayer* InHLODLayer, const TArray<FWorldPartitionReference>& InSubActors);
 
 	static TArray<UPrimitiveComponent*> GatherPrimitiveComponents(const TArray<FWorldPartitionReference>& InActors);
-
-	static void DisableCollisions(UPrimitiveComponent* Component);
 #endif // WITH_EDITOR
 };
