Fix breaking recursive calls to UWorld::BlockTillLevelStreamingCompleted introduced in 16641121.

diff --git a/Engine/Source/Runtime/Engine/Classes/Engine/World.h b/Engine/Source/Runtime/Engine/Classes/Engine/World.h
--- a/Engine/Source/Runtime/Engine/Classes/Engine/World.h
+++ b/Engine/Source/Runtime/Engine/Classes/Engine/World.h
@@ -1055,7 +1055,7 @@ public:
 	void PopulateStreamingLevelsToConsider();
 
 	/** Whether the world is currently in a BlockTillLevelStreamingCompleted() call */
-	bool GetIsInBlockTillLevelStreamingCompleted() const { return bIsInBlockTillLevelStreamingCompleted; }
+	bool GetIsInBlockTillLevelStreamingCompleted() const { return IsInBlockTillLevelStreamingCompleted > 0; }
 
 	/** Prefix we used to rename streaming levels, non empty in PIE and standalone preview */
 	UPROPERTY()
@@ -1177,7 +1177,7 @@ public:
 	
 	/** Whether it was requested that the engine bring up a loading screen and block on async loading. */   
 	uint8 bRequestedBlockOnAsyncLoading:1;
-		
+
 	/** Whether actors have been initialized for play */
 	uint8 bActorsInitialized:1;
 
@@ -1224,9 +1224,6 @@ public:
 	uint8 bAreConstraintsDirty:1;
 
 private:
-	/** Whether the world is currently in a BlockTillLevelStreamingCompleted() call */
-	uint8 bIsInBlockTillLevelStreamingCompleted : 1;
-
 	/** Whether the render scene for this World should be created with HitProxies or not */
 	uint8 bRequiresHitProxies:1;
 
@@ -1248,6 +1245,9 @@ private:
 	/** Whether world object has been initialized via Init and has not yet had CleanupWorld called								*/
 	uint8 bInitializedAndNeedsCleanup : 1;
 
+	/** Whether the world is currently in a BlockTillLevelStreamingCompleted() call */
+	uint32 IsInBlockTillLevelStreamingCompleted;
+
 	/** The world's navigation data manager */
 	UPROPERTY(Transient)
 	TObjectPtr<class UNavigationSystemBase>				NavigationSystem;
diff --git a/Engine/Source/Runtime/Engine/Private/World.cpp b/Engine/Source/Runtime/Engine/Private/World.cpp
--- a/Engine/Source/Runtime/Engine/Private/World.cpp
+++ b/Engine/Source/Runtime/Engine/Private/World.cpp
@@ -12,7 +12,6 @@
 #include "Logging/LogScopedVerbosityOverride.h"
 #include "Stats/StatsMisc.h"
 #include "Misc/ScopedSlowTask.h"
-#include "Misc/ScopeExit.h"
 #include "Misc/CoreDelegates.h"
 #include "Modules/ModuleManager.h"
 #include "UObject/UObjectIterator.h"
@@ -401,6 +400,8 @@ UWorld::UWorld( const FObjectInitializer& ObjectInitializer )
 			SetAudioDevice(EmptyHandle);
 		}
 	});
+
+	IsInBlockTillLevelStreamingCompleted = 0;
 }
 
 UWorld::~UWorld()
@@ -3632,18 +3633,8 @@ void FStreamingLevelsToConsider::AddReferencedObjects(UObject* InThis, FReferenc
 
 void UWorld::BlockTillLevelStreamingCompleted()
 {
-	if (bIsInBlockTillLevelStreamingCompleted)
-	{
-		return;
-	}
-
-	bIsInBlockTillLevelStreamingCompleted = true;
-	ON_SCOPE_EXIT
-	{
-		bIsInBlockTillLevelStreamingCompleted = false;
-	};
-
 	const double StartTime = FPlatformTime::Seconds();
+	TScopeCounter<uint32> IsInBlockTillLevelStreamingCompletedCounter(IsInBlockTillLevelStreamingCompleted);
 	
 	bool bWorkToDo = false;
 	do
