Fix actor labels for shipping editor build.

diff --git a/Engine/Source/Runtime/Engine/Classes/GameFramework/Actor.h b/Engine/Source/Runtime/Engine/Classes/GameFramework/Actor.h
--- a/Engine/Source/Runtime/Engine/Classes/GameFramework/Actor.h
+++ b/Engine/Source/Runtime/Engine/Classes/GameFramework/Actor.h
@@ -890,15 +890,34 @@ public:
 		return HasAnyFlags(RF_HasExternalPackage);
 	}
 
-#if WITH_EDITORONLY_DATA
 private:
+#if WITH_EDITORONLY_DATA
 	/**
 	 * The friendly name for this actor, displayed in the editor.  You should always use AActor::GetActorLabel() to access the actual label to display,
 	 * and call AActor::SetActorLabel() or FActorLabelUtilities::SetActorLabelUnique() to change the label.  Never set the label directly.
 	 */
 	UPROPERTY()
 	mutable FString ActorLabel;
+#endif
+
+#if !WITH_EDITOR && ACTOR_HAS_LABELS
+	FString ActorLabel;
+#endif
+
+public:
+	const FString GetActorNameOrLabel() const
+	{
+#if WITH_EDITORONLY_DATA || (!WITH_EDITOR && ACTOR_HAS_LABELS)
+		if (!ActorLabel.IsEmpty())
+		{
+			return ActorLabel;
+		}
+#endif
+		return GetName();
+	}
 
+#if WITH_EDITORONLY_DATA
+private:
 	/** The folder path of this actor in the world (empty=root, / separated)*/
 	UPROPERTY()
 	FName FolderPath;
@@ -2272,23 +2291,6 @@ public:
 	}
 #endif		// WITH_EDITOR
 
-#if !WITH_EDITOR && ACTOR_HAS_LABELS
-private:
-	FString ActorLabel;
-
-public:
-	const FString& GetActorLabel() const { return ActorLabel; }
-#endif
-
-	const FString GetActorNameOrLabel() const
-	{
-#if WITH_EDITORONLY_DATA || (!WITH_EDITOR && ACTOR_HAS_LABELS)
-		return ActorLabel.IsEmpty() ? GetName() : ActorLabel;
-#else
-		return GetName();
-#endif
-	}
-
 	/**
 	 * Function used to prioritize actors when deciding which to replicate
 	 * @param ViewPos		Position of the viewer
diff --git a/Engine/Source/Runtime/Engine/Private/Actor.cpp b/Engine/Source/Runtime/Engine/Private/Actor.cpp
--- a/Engine/Source/Runtime/Engine/Private/Actor.cpp
+++ b/Engine/Source/Runtime/Engine/Private/Actor.cpp
@@ -798,7 +798,7 @@ void AActor::Serialize(FArchive& Ar)
 
 		if (bIsCooked)
 		{
-#if !ACTOR_HAS_LABELS
+#if !(WITH_EDITORONLY_DATA || (!WITH_EDITOR && ACTOR_HAS_LABELS))
 			// In non-development builds, just skip over the actor labels. We need to figure out a way to either strip that data from shipping builds,
 			// or skip over the string data without doing any memory allocations, probably with a custom FString::SerializeToNull implementation.
 			FString ActorLabel;
