Deprecate old input scales on the player controller (InputYawScale, InputPitchScale, and InputRollScale)

#p4v-preflight-copy 17107710

diff --git a/Engine/Source/Developer/LogVisualizer/Private/VisualLoggerCameraController.cpp b/Engine/Source/Developer/LogVisualizer/Private/VisualLoggerCameraController.cpp
--- a/Engine/Source/Developer/LogVisualizer/Private/VisualLoggerCameraController.cpp
+++ b/Engine/Source/Developer/LogVisualizer/Private/VisualLoggerCameraController.cpp
@@ -63,8 +63,8 @@ void AVisualLoggerCameraController::PostInitializeComponents()
 	MyHUD = GetWorld()->SpawnActor<AVisualLoggerHUD>(SpawnInfo);
 
 	/** TO DO: somehow this two configuration settings are set to zero for this controller, have to fix it later */
-	InputYawScale=2.5;
-	InputPitchScale=-1.75;
+	InputYawScale_DEPRECATED = 2.5f;
+	InputPitchScale_DEPRECATED = -1.75f;
 }
 
 
diff --git a/Engine/Source/Runtime/CoreUObject/Private/UObject/CoreRedirects.cpp b/Engine/Source/Runtime/CoreUObject/Private/UObject/CoreRedirects.cpp
--- a/Engine/Source/Runtime/CoreUObject/Private/UObject/CoreRedirects.cpp
+++ b/Engine/Source/Runtime/CoreUObject/Private/UObject/CoreRedirects.cpp
@@ -1428,6 +1428,9 @@ static void RegisterNativeRedirects40(TArray<FCoreRedirect>& Redirects)
 	PROPERTY_REDIRECT("PawnSensingComponent.bWantsSeePlayerNotify", "PawnSensingComponent.bSeePawns");
 	PROPERTY_REDIRECT("PlayerController.LookRightScale", "PlayerController.InputYawScale");
 	PROPERTY_REDIRECT("PlayerController.LookUpScale", "PlayerController.InputPitchScale");
+	PROPERTY_REDIRECT("PlayerController.InputYawScale", "PlayerController.InputYawScale_DEPRECATED");
+	PROPERTY_REDIRECT("PlayerController.InputPitchScale", "PlayerController.InputPitchScale_DEPRECATED");
+	PROPERTY_REDIRECT("PlayerController.InputRollScale", "PlayerController.InputRollScale_DEPRECATED");
 	PROPERTY_REDIRECT("PlayerController.PlayerCamera", "PlayerController.PlayerCameraManager");
 	PROPERTY_REDIRECT("PlayerController.PlayerCameraClass", "PlayerController.PlayerCameraManagerClass");
 	PROPERTY_REDIRECT("PointLightComponent.Radius", "PointLightComponent.AttenuationRadius");
diff --git a/Engine/Source/Runtime/Engine/Classes/GameFramework/InputSettings.h b/Engine/Source/Runtime/Engine/Classes/GameFramework/InputSettings.h
--- a/Engine/Source/Runtime/Engine/Classes/GameFramework/InputSettings.h
+++ b/Engine/Source/Runtime/Engine/Classes/GameFramework/InputSettings.h
@@ -53,6 +53,10 @@ class ENGINE_API UInputSettings
 	UPROPERTY(config)
 	uint8 bDefaultViewportMouseLock_DEPRECATED:1;
 
+	/** Enable the use of legacy input scales on the player controller (InputYawScale, InputPitchScale, and InputRollScale) */
+	UPROPERTY(config, EditAnywhere, Category = "Input")
+	uint8 bEnableLegacyInputScales:1;
+	
 	/** Should the touch input interface be shown always, or only when the platform has a touch screen? */
 	UPROPERTY(config, EditAnywhere, Category="Mobile")
 	uint8 bAlwaysShowTouchInterface:1;
@@ -67,7 +71,7 @@ class ENGINE_API UInputSettings
 
 	/** If enabled, virtual keyboards will have autocorrect enabled. Currently only supported on mobile devices. */
 	UPROPERTY(config, EditAnywhere, Category = "Virtual Keyboard (Mobile)")
-	bool bUseAutocorrect;
+	uint8 bUseAutocorrect:1;
 
 	/** 
 	 * Disables autocorrect for these operating systems, even if autocorrect is enabled. Use the format "[platform] [osversion]"
diff --git a/Engine/Source/Runtime/Engine/Classes/GameFramework/PlayerController.h b/Engine/Source/Runtime/Engine/Classes/GameFramework/PlayerController.h
--- a/Engine/Source/Runtime/Engine/Classes/GameFramework/PlayerController.h
+++ b/Engine/Source/Runtime/Engine/Classes/GameFramework/PlayerController.h
@@ -438,16 +438,16 @@ public:
 	FRotator RotationInput;
 
 	/** Yaw input speed scaling */
-	UPROPERTY(EditAnywhere, BlueprintReadWrite, config, Category=PlayerController)
-	float InputYawScale;
+	UPROPERTY(config, meta = (DeprecatedProperty, DeprecationMessage = "Use the Enhanced Input plugin Scalar Modifier instead. See UInputSettings::bEnableLegacyInputScales to enable legacy behavior"))
+	float InputYawScale_DEPRECATED = 1.0f;
 
 	/** Pitch input speed scaling */
-	UPROPERTY(EditAnywhere, BlueprintReadWrite, config, Category=PlayerController)
-	float InputPitchScale;
+	UPROPERTY(config, meta = (DeprecatedProperty, DeprecationMessage = "Use the Enhanced Input plugin Scalar Modifier instead. See UInputSettings::bEnableLegacyInputScales to enable legacy behavior"))
+	float InputPitchScale_DEPRECATED = 1.0f;
 
 	/** Roll input speed scaling */
-	UPROPERTY(EditAnywhere, BlueprintReadWrite, config, Category=PlayerController)
-	float InputRollScale;
+	UPROPERTY(config, meta = (DeprecatedProperty, DeprecationMessage = "Use the Enhanced Input plugin Scalar Modifier instead. See UInputSettings::bEnableLegacyInputScales to enable legacy behavior"))
+	float InputRollScale_DEPRECATED = 1.0f;
 
 	/** Whether the mouse cursor should be displayed. */
 	UPROPERTY(EditAnywhere, BlueprintReadWrite, Category=MouseInterface)
diff --git a/Engine/Source/Runtime/Engine/Private/PlayerController.cpp b/Engine/Source/Runtime/Engine/Private/PlayerController.cpp
--- a/Engine/Source/Runtime/Engine/Private/PlayerController.cpp
+++ b/Engine/Source/Runtime/Engine/Private/PlayerController.cpp
@@ -5149,18 +5149,18 @@ bool APlayerController::PopInputComponent(UInputComponent* InInputComponent)
 }
 
 void APlayerController::AddPitchInput(float Val)
-{
-	RotationInput.Pitch += !IsLookInputIgnored() ? Val * InputPitchScale : 0.f;
+{	
+	RotationInput.Pitch += !IsLookInputIgnored() ? Val * (GetDefault<UInputSettings>()->bEnableLegacyInputScales ? InputPitchScale_DEPRECATED : 1.0f) : 0.0f;
 }
 
 void APlayerController::AddYawInput(float Val)
 {
-	RotationInput.Yaw += !IsLookInputIgnored() ? Val * InputYawScale : 0.f;
+	RotationInput.Yaw += !IsLookInputIgnored() ? Val * (GetDefault<UInputSettings>()->bEnableLegacyInputScales ? InputYawScale_DEPRECATED : 1.0f) : 0.0f;
 }
 
 void APlayerController::AddRollInput(float Val)
 {
-	RotationInput.Roll += !IsLookInputIgnored() ? Val * InputRollScale : 0.f;
+	RotationInput.Roll += !IsLookInputIgnored() ? Val * (GetDefault<UInputSettings>()->bEnableLegacyInputScales ? InputRollScale_DEPRECATED : 1.0f) : 0.0f;
 }
 
 bool APlayerController::IsInputKeyDown(const FKey Key) const
diff --git a/Engine/Source/Runtime/Engine/Private/UserInterface/InputSettings.cpp b/Engine/Source/Runtime/Engine/Private/UserInterface/InputSettings.cpp
--- a/Engine/Source/Runtime/Engine/Private/UserInterface/InputSettings.cpp
+++ b/Engine/Source/Runtime/Engine/Private/UserInterface/InputSettings.cpp
@@ -21,6 +21,7 @@ UInputSettings::UInputSettings(const FObjectInitializer& ObjectInitializer)
 	: Super(ObjectInitializer)
 	, bCaptureMouseOnLaunch(true)
 	, bDefaultViewportMouseLock_DEPRECATED(false)
+	, bEnableLegacyInputScales(false)
 	, DefaultViewportMouseCaptureMode(EMouseCaptureMode::CapturePermanently_IncludingInitialMouseDown)
 	, DefaultViewportMouseLockMode(EMouseLockMode::LockOnCapture)
 	, DefaultPlayerInputClass(UPlayerInput::StaticClass())
