Add verbose shader compilation stats to drill into issues.

diff --git a/Engine/Source/Runtime/Engine/Private/ShaderCompiler/ShaderCompiler.cpp b/Engine/Source/Runtime/Engine/Private/ShaderCompiler/ShaderCompiler.cpp
--- a/Engine/Source/Runtime/Engine/Private/ShaderCompiler/ShaderCompiler.cpp
+++ b/Engine/Source/Runtime/Engine/Private/ShaderCompiler/ShaderCompiler.cpp
@@ -235,10 +235,10 @@ FShaderCompileJobCollection::FShaderCompileJobCollection()
 	FMemory::Memzero(PendingJobs);
 	FMemory::Memzero(NumPendingJobs);
 
-	LogJobsCacheStatsCmd = IConsoleManager::Get().RegisterConsoleCommand(
-		TEXT("r.ShaderCompiler.LogCacheStats"),
-		TEXT("Prints out the stats for the in-memory shader job cache."),
-		FConsoleCommandDelegate::CreateRaw(this, &FShaderCompileJobCollection::HandleLogJobsCacheStats),
+	PrintStatsCmd = IConsoleManager::Get().RegisterConsoleCommand(
+		TEXT("r.ShaderCompiler.PrintStats"),
+		TEXT("Prints out to the log the stats for the shader compiler."),
+		FConsoleCommandDelegate::CreateRaw(this, &FShaderCompileJobCollection::HandlePrintStats),
 		ECVF_Default
 	);
 }
@@ -342,6 +342,7 @@ void FShaderCompileJobCollection::InternalSetPriority(FShaderCommonCompileJob* J
 		checkf(NewNumPendingJobs > 1, TEXT("Invalid number of pending jobs %d, should have had at least 1 job previously"), NewNumPendingJobs);
 		InternalAddJob(NewJob);
 
+		GShaderCompilerStats->RegisterNewPendingJob(*NewJob);
 		NumPendingJobs[PriorityIndex]++;
 		ensureMsgf(NewJob->bInputHashSet == Job->bInputHashSet, TEXT("Cloned and original jobs should either both have input hash, or both not have it. Job->bInputHashSet=%d, NewJob->bInputHashSet=%d"),
 			Job->bInputHashSet,
@@ -528,6 +529,7 @@ void FShaderCompileJobCollection::SubmitJobs(const TArray<FShaderCommonCompileJo
 				// new job
 				if (bNewJob)
 				{
+					GShaderCompilerStats->RegisterNewPendingJob(*Job);
 					ensure(!ShaderCompiler::IsJobCacheEnabled() || Job->bInputHashSet);
 					Job->LinkHead(PendingJobs[PriorityIndex]);
 
@@ -553,13 +555,18 @@ void FShaderCompileJobCollection::SubmitJobs(const TArray<FShaderCommonCompileJo
 	}
 }
 
-void FShaderCompileJobCollection::HandleLogJobsCacheStats()
+void FShaderCompileJobCollection::HandlePrintStats()
 {
 	GShaderCompilingManager->PrintStats(true);
 }
 
 void FShaderCompileJobCollection::ProcessFinishedJob(FShaderCommonCompileJob* FinishedJob, bool bWasCached)
 {
+	if (!bWasCached)
+	{
+		GShaderCompilerStats->RegisterFinishedJob(*FinishedJob);
+	}
+
 	// TODO: have a pending shader map critical section? not clear at this point if we can be accessing the results on another thread at the same time
 	FShaderMapCompileResults& ShaderMapResults = *(FinishedJob->PendingShaderMap);
 	ShaderMapResults.FinishedJobs.Add(FinishedJob);
@@ -675,6 +682,8 @@ int32 FShaderCompileJobCollection::GetPendingJobs(EShaderCompilerWorkerType InWo
 		for (int32 i = 0; i < NumJobs; ++i)
 		{
 			FShaderCommonCompileJob& Job = *It;
+
+			GShaderCompilerStats->RegisterAssignedJob(Job);
 			// Temporary commented out until r.ShaderDevelopmentMode=1 shader error retry crash gets fixed
 			//check(Job.CurrentWorker == EShaderCompilerWorkerType::None);
 			//check(Job.PendingPriority == InPriority);
@@ -1819,7 +1828,7 @@ struct FShaderCompileWorkerInfo
 	/** Time at which the worker started the most recent batch of tasks. */
 	double StartTime;
 
-	/** Jobs that this worker is responsible for compiling. */
+	/** Time at which the worker ended the most recent batch of tasks. */
 	double FinishTime = 0.0;
 
 	/** Jobs that this worker is responsible for compiling. */
@@ -1936,12 +1945,16 @@ int32 FShaderCompileThreadRunnable::PullTasksFromQueue()
 						CurrentWorkerInfo.bLaunchedWorker = false;
 						CurrentWorkerInfo.StartTime = FPlatformTime::Seconds();
 						NumActiveThreads++;
-					}
 
-					if (Manager->bLogJobCompletionTimes && CurrentWorkerInfo.FinishTime > 0.0)
-					{
-						const double WorkerIdleTime = CurrentWorkerInfo.StartTime - CurrentWorkerInfo.FinishTime;
-						UE_LOG(LogShaders, Display, TEXT("  Worker was idle for %fs"), WorkerIdleTime);
+						if (CurrentWorkerInfo.FinishTime > 0.0)
+						{
+							const double WorkerIdleTime = CurrentWorkerInfo.StartTime - CurrentWorkerInfo.FinishTime;
+							GShaderCompilerStats->RegisterLocalWorkerIdleTime(WorkerIdleTime);
+							if (Manager->bLogJobCompletionTimes)
+							{
+								UE_LOG(LogShaders, Display, TEXT("  Worker (%d/%d) starts working after being idle for %fs"), WorkerIndex + 1, WorkerInfos.Num(), WorkerIdleTime);
+							}
+						}
 					}
 				}
 			}
@@ -2684,6 +2697,193 @@ void FShaderCompilerStats::WriteStatSummary()
 
 	UE_LOG(LogShaderCompilers, Display, TEXT("=== Shader Compilation stats ==="));
 	UE_LOG(LogShaderCompilers, Display, TEXT("Shaders Compiled: %u"), TotalCompiled);
+
+	FScopeLock Lock(&CompileStatsLock);	// make a local copy for all the stats?
+	UE_LOG(LogShaderCompilers, Display, TEXT("Jobs assigned: %.0f"), JobsAssigned);
+	UE_LOG(LogShaderCompilers, Display, TEXT("Jobs completed: %.0f (%.2f%%)"), JobsCompleted, (JobsAssigned > 0.0) ? 100.0 * JobsCompleted / JobsAssigned : 0.0);
+
+	if (TimesLocalWorkersWereIdle > 0.0)
+	{
+		UE_LOG(LogShaderCompilers, Display, TEXT("Average time worker was idle: %.2f s"), AccumulatedLocalWorkerIdleTime / TimesLocalWorkersWereIdle);
+	}
+
+	if (JobsAssigned > 0.0)
+	{
+		UE_LOG(LogShaderCompilers, Display, TEXT("Average time job spent in pending queue: %.2f s"), AccumulatedPendingTime / JobsAssigned);
+	}
+
+	if (JobsCompleted > 0.0)
+	{
+		UE_LOG(LogShaderCompilers, Display, TEXT("Average job execution time: %.2f s"), AccumulatedJobExecutionTime / JobsCompleted);
+		UE_LOG(LogShaderCompilers, Display, TEXT("Average job life time (pending + execution): %.2f s"), AccumulatedJobLifeTime / JobsCompleted);
+	}
+
+	auto SumUpInterval = [](TArray<TInterval<double>>& Accumulator) -> double
+	{
+		double Sum = 0;
+		for (int32 Idx = 0; Idx < Accumulator.Num(); ++Idx)
+		{
+			const TInterval<double>& Existing = Accumulator[Idx];
+			Sum += Existing.Size();
+		}
+		return Sum;
+	};
+
+	double TotalTimeExecutingAtLeastOneJob = SumUpInterval(JobPendingTimeIntervals);
+	UE_LOG(LogShaderCompilers, Display, TEXT("Time at least one job was waiting to be executed: %.2f s"), TotalTimeExecutingAtLeastOneJob);
+
+	double TotalTimeAtLeastOneJobWasInFlight = SumUpInterval(JobLifeTimeIntervals);
+	UE_LOG(LogShaderCompilers, Display, TEXT("Time at least one job was in flight (either pending or executed): %.2f s"), TotalTimeAtLeastOneJobWasInFlight);
+	if (TotalTimeAtLeastOneJobWasInFlight > 0.0)
+	{
+		UE_LOG(LogShaderCompilers, Display, TEXT("Average processing rate: %.2f jobs/sec"), JobsCompleted / TotalTimeAtLeastOneJobWasInFlight);
+	}
+
+	if (ShaderTimings.Num())
+	{
+		// sort by avg time
+		ShaderTimings.ValueSort([](const FShaderTimings& A, const FShaderTimings& B) { return A.AverageCompileTime > B.AverageCompileTime; });
+
+		const int32 MaxShadersToPrint = FMath::Min(ShaderTimings.Num(), 5);
+		UE_LOG(LogShaderCompilers, Display, TEXT("Top %d most expensive shader types by average time:"), MaxShadersToPrint);
+
+		int32 Idx = 0;
+		for (TMap<FString, FShaderTimings>::TConstIterator Iter(ShaderTimings); Iter; ++Iter)
+		{
+			const FShaderTimings& Timings = Iter.Value();
+
+			UE_LOG(LogShaderCompilers, Display, TEXT("%60s (compiled %4d times, average %4.2f sec, max %4.2f sec, min %4.2f sec)"), *Iter.Key(), Timings.NumCompiled, Timings.AverageCompileTime, Timings.MaxCompileTime, Timings.MinCompileTime);
+			if (++Idx >= MaxShadersToPrint)
+			{
+				break;
+			}
+		}
+
+		// sort by total time
+		ShaderTimings.ValueSort([](const FShaderTimings& A, const FShaderTimings& B) { return A.TotalCompileTime > B.TotalCompileTime; });
+
+		UE_LOG(LogShaderCompilers, Display, TEXT("Top %d shader types by total compile time:"), MaxShadersToPrint);
+
+		Idx = 0;
+		for (TMap<FString, FShaderTimings>::TConstIterator Iter(ShaderTimings); Iter; ++Iter)
+		{
+			const FShaderTimings& Timings = Iter.Value();
+
+			UE_LOG(LogShaderCompilers, Display, TEXT("%60s (compiled %4d times, average %4.2f sec, max %4.2f sec, min %4.2f sec)"), *Iter.Key(), Timings.NumCompiled, Timings.AverageCompileTime, Timings.MaxCompileTime, Timings.MinCompileTime);
+			if (++Idx >= MaxShadersToPrint)
+			{
+				break;
+			}
+		}
+	}
+}
+
+void FShaderCompilerStats::RegisterLocalWorkerIdleTime(double IdleTime)
+{
+	FScopeLock Lock(&CompileStatsLock);
+	AccumulatedLocalWorkerIdleTime += IdleTime;
+	TimesLocalWorkersWereIdle++;
+}
+
+void FShaderCompilerStats::RegisterNewPendingJob(FShaderCommonCompileJob& Job)
+{
+	// accessing job timestamps isn't arbitrated by any lock. It is assumed that the registration of a job at one of the stages
+	// of its lifetime happens before the code can move it to another stage (i.e. new pending job is registered before it is added to the pending queue,
+	// so it cannot be given away to a worker while it's still being registered, and an assigned job is registered before it is actually given to the worker,
+	// so it cannot end up being registered as finished at the same time on some other thread).
+	Job.TimeAddedToPendingQueue = FPlatformTime::Seconds();
+}
+
+void FShaderCompilerStats::RegisterAssignedJob(FShaderCommonCompileJob& Job)
+{
+	ensure(Job.TimeAddedToPendingQueue != 0.0);
+	Job.TimeAssignedToExecution = FPlatformTime::Seconds();
+
+	FScopeLock Lock(&CompileStatsLock);
+	JobsAssigned++;
+	AccumulatedPendingTime += (Job.TimeAssignedToExecution - Job.TimeAddedToPendingQueue);
+}
+
+void FShaderCompilerStats::RegisterFinishedJob(FShaderCommonCompileJob& Job)
+{
+	ensure(Job.TimeAssignedToExecution != 0.0);
+	Job.TimeExecutionCompleted = FPlatformTime::Seconds();
+
+	FScopeLock Lock(&CompileStatsLock);
+	JobsCompleted++;
+	AccumulatedJobExecutionTime += (Job.TimeExecutionCompleted - Job.TimeAssignedToExecution);
+	AccumulatedJobLifeTime += (Job.TimeExecutionCompleted - Job.TimeAddedToPendingQueue);
+
+	auto AddToInterval = [](TArray<TInterval<double>>& Accumulator, const TInterval<double>& NewInterval)
+	{
+		bool bFoundOverlap = false;
+		TInterval<double> New = NewInterval;
+		do
+		{
+			bFoundOverlap = false;
+			for (int32 Idx = 0; Idx < Accumulator.Num(); ++Idx)
+			{
+				const TInterval<double>& Existing = Accumulator[Idx];
+				if (Existing.Max < New.Min || New.Max < Existing.Min)
+				{
+					continue;	// no overlap
+				}
+
+				// if fully contained within existing interval, just ignore
+				if (Existing.Min <= New.Min && New.Max <= Existing.Max)
+				{
+					return;
+				}
+
+				bFoundOverlap = true;
+				// if there's an overlap, remove the existing interval, merge with the new one and attempt to add again
+				TInterval<double> Merged(FMath::Min(Existing.Min, New.Min), FMath::Max(Existing.Max, New.Max));
+				check(Merged.Size() >= Existing.Size());
+				check(Merged.Size() >= New.Size());
+				Accumulator.RemoveAt(Idx);
+				New = Merged;
+				break;
+			}
+		} while (bFoundOverlap);
+
+		// if we arrived here without an overlap, we have a new one
+		if (!bFoundOverlap)
+		{
+			Accumulator.Add(New);
+		}
+	};
+
+	// estimate lifetime without an overlap
+	ensure(Job.TimeAddedToPendingQueue != 0.0 && Job.TimeAddedToPendingQueue <= Job.TimeAssignedToExecution);
+	AddToInterval(JobPendingTimeIntervals, TInterval<double>(Job.TimeAddedToPendingQueue, Job.TimeAssignedToExecution));
+
+	ensure(Job.TimeAddedToPendingQueue != 0.0 && Job.TimeAssignedToExecution <= Job.TimeExecutionCompleted);
+	AddToInterval(JobLifeTimeIntervals, TInterval<double>(Job.TimeAddedToPendingQueue, Job.TimeExecutionCompleted));
+
+	if (const FShaderCompileJob* SingleJob = Job.GetSingleShaderJob())
+	{
+		const FString ShaderName(SingleJob->Key.ShaderType->GetName());
+		if (FShaderTimings* Existing = ShaderTimings.Find(ShaderName))
+		{
+			Existing->MinCompileTime = FMath::Min(Existing->MinCompileTime, static_cast<float>(SingleJob->Output.CompileTime));
+			Existing->MaxCompileTime = FMath::Max(Existing->MaxCompileTime, static_cast<float>(SingleJob->Output.CompileTime));
+			Existing->TotalCompileTime += SingleJob->Output.CompileTime;
+			Existing->NumCompiled++;
+			// calculate as an optimization to make sorting later faster
+			Existing->AverageCompileTime = Existing->TotalCompileTime / static_cast<float>(Existing->NumCompiled);
+		}
+		else
+		{
+			FShaderTimings New;
+			New.MinCompileTime = SingleJob->Output.CompileTime;
+			New.MaxCompileTime = New.MinCompileTime;
+			New.TotalCompileTime = New.MinCompileTime;
+			New.AverageCompileTime = New.MinCompileTime;
+			New.NumCompiled = 1;
+			ShaderTimings.Add(ShaderName, New);
+		}
+	}
+
 }
 
 void FShaderCompilerStats::RegisterCookedShaders(uint32 NumCooked, float CompileTime, EShaderPlatform Platform, const FString MaterialPath, FString PermutationString)
diff --git a/Engine/Source/Runtime/Engine/Public/ShaderCompiler.h b/Engine/Source/Runtime/Engine/Public/ShaderCompiler.h
--- a/Engine/Source/Runtime/Engine/Public/ShaderCompiler.h
+++ b/Engine/Source/Runtime/Engine/Public/ShaderCompiler.h
@@ -141,6 +141,13 @@ public:
 	/** Hash of all the job inputs */
 	FInputHash InputHash;
 
+	/** In-engine timestamp of being added to a pending queue. Not set for jobs that are satisfied from the jobs cache */
+	double TimeAddedToPendingQueue = 0.0;
+	/** In-engine timestamp of being assigned to a worker. Not set for jobs that are satisfied from the jobs cache */
+	double TimeAssignedToExecution = 0.0;
+	/** In-engine timestamp of job being completed. Encompasses the compile time. Not set for jobs that are satisfied from the jobs cache */
+	double TimeExecutionCompleted = 0.0;
+
 	uint32 AddRef() const
 	{
 		return uint32(NumRefs.Increment());
@@ -489,8 +496,8 @@ private:
 		return NewJob;
 	}
 
-	/** Handles the console command to log jobs cache stats */
-	void HandleLogJobsCacheStats();
+	/** Handles the console command to log shader compiler stats */
+	void HandlePrintStats();
 
 	/** Queue of tasks that haven't been assigned to a worker yet. */
 	FShaderCommonCompileJob* PendingJobs[NumShaderCompileJobPriorities];
@@ -515,7 +522,7 @@ private:
 	FShaderJobCache CompletedJobsCache;
 
 	/** Debugging - console command to print stats. */
-	class IConsoleObject* LogJobsCacheStatsCmd;
+	class IConsoleObject* PrintStatsCmd;
 };
 
 class FGlobalShaderTypeCompiler
@@ -719,6 +726,15 @@ public:
 	};
 	using ShaderCompilerStats = TMap<FString, FShaderStats>;
 
+	/** Structure used to describe compiling time of a shader type (for all the instances of it that we have seen). */
+	struct FShaderTimings
+	{
+		float MinCompileTime;
+		float MaxCompileTime;
+		float TotalCompileTime;
+		int32 NumCompiled;
+		float AverageCompileTime;	// stored explicitly as an optimization
+	};
 
 	ENGINE_API void RegisterCookedShaders(uint32 NumCooked, float CompileTime, EShaderPlatform Platform, const FString MaterialPath, FString PermutationString = FString(""));
 	ENGINE_API void RegisterCompiledShaders(uint32 NumPermutations, EShaderPlatform Platform, const FString MaterialPath, FString PermutationString = FString(""));
@@ -726,9 +742,54 @@ public:
 	ENGINE_API void WriteStats(class FOutputDevice* Ar = nullptr);
 	ENGINE_API void WriteStatSummary();
 
+	/** Informs statistics about a time a local ShaderCompileWorker spent idle. */
+	void RegisterLocalWorkerIdleTime(double IdleTime);
+
+	/** Lets the stats to know about a newly added job. Job will be modified to include the current timestamp. */
+	void RegisterNewPendingJob(FShaderCommonCompileJob& InOutJob);
+
+	/** Marks the job as given out to a worker for execution for the stats purpose. Job will be modified to include the current timestamp. */
+	void RegisterAssignedJob(FShaderCommonCompileJob& InOutJob);
+
+	/** Marks the job as finished for the stats purpose. Job will be modified to include the current timestamp. */
+	void RegisterFinishedJob(FShaderCommonCompileJob& InOutJob);
+
 private:
 	FCriticalSection CompileStatsLock;
 	TSparseArray<ShaderCompilerStats> CompileStats;
+
+	/** This tracks accumulated wait time from local workers during the lifetime of the stats.
+	 *
+	 * Wait time is only counted for local workers that are alive and not between their invocations
+	 */
+	double AccumulatedLocalWorkerIdleTime = 0.0;
+
+	/** How many times we registered idle time? */
+	double TimesLocalWorkersWereIdle = 0;
+
+	/** Number of jobs assigned to workers, no matter if they completed or not - used to average pending time. */
+	double JobsAssigned = 0;
+
+	/** Amount of time a job had to spent in pending queue (i.e. waiting to be assigned to a worker). */
+	double AccumulatedPendingTime = 0;
+
+	/** Total number jobs completed. */
+	double JobsCompleted = 0;
+
+	/** Amount of time job spent being processed by the worker. */
+	double AccumulatedJobExecutionTime = 0;
+
+	/** Amount of time job spent being processed by the worker. */
+	double AccumulatedJobLifeTime = 0;
+
+	/** Accumulates the job pending times without overlaps */
+	TArray<TInterval<double>> JobPendingTimeIntervals;
+
+	/** Accumulates the job lifetimes without overlaps */
+	TArray<TInterval<double>> JobLifeTimeIntervals;
+
+	/** Map of shader names to their compilation timings */
+	TMap<FString, FShaderTimings> ShaderTimings;
 };
 
 
