Ensure ispc::FTransform and FTransform sizes match in anim paths
#rnx

diff --git a/Engine/Source/Runtime/Engine/Private/Animation/AnimEncoding_ConstantKeyLerp.cpp b/Engine/Source/Runtime/Engine/Private/Animation/AnimEncoding_ConstantKeyLerp.cpp
--- a/Engine/Source/Runtime/Engine/Private/Animation/AnimEncoding_ConstantKeyLerp.cpp
+++ b/Engine/Source/Runtime/Engine/Private/Animation/AnimEncoding_ConstantKeyLerp.cpp
@@ -9,6 +9,9 @@
 #include "Serialization/MemoryWriter.h"
 #if INTEL_ISPC
 #include "AnimEncoding_ConstantKeyLerp.ispc.generated.h"
+
+static_assert(sizeof(ispc::FTransform) == sizeof(FTransform), "sizeof(ispc::FTransform) != sizeof(FTransform)");
+static_assert(sizeof(ispc::BoneTrackPair) == sizeof(BoneTrackPair), "sizeof(ispc::BoneTrackPair) != sizeof(BoneTrackPair)");
 #endif
 
 /**
diff --git a/Engine/Source/Runtime/Engine/Private/Animation/AnimEncoding_PerTrackCompression.cpp b/Engine/Source/Runtime/Engine/Private/Animation/AnimEncoding_PerTrackCompression.cpp
--- a/Engine/Source/Runtime/Engine/Private/Animation/AnimEncoding_PerTrackCompression.cpp
+++ b/Engine/Source/Runtime/Engine/Private/Animation/AnimEncoding_PerTrackCompression.cpp
@@ -11,6 +11,9 @@
 #include "AnimationCompression.h"
 #if INTEL_ISPC
 #include "AnimEncoding_PerTrackCompression.ispc.generated.h"
+
+static_assert(sizeof(ispc::FTransform) == sizeof(FTransform), "sizeof(ispc::FTransform) != sizeof(FTransform)");
+static_assert(sizeof(ispc::BoneTrackPair) == sizeof(BoneTrackPair), "sizeof(ispc::BoneTrackPair) != sizeof(BoneTrackPair)");
 #endif
 
 // This define controls whether scalar or vector code is used to decompress keys.  Note that not all key decompression code
@@ -713,10 +716,6 @@ void AEFPerTrackCompressionCodec::GetBoneAtomScale(
 
 #if USE_ANIMATION_CODEC_BATCH_SOLVER
 
-#if INTEL_ISPC
-static_assert(sizeof(ispc::BoneTrackPair) == sizeof(BoneTrackPair), "sizeof(ispc::BoneTrackPair) != sizeof(BoneTrackPair)");
-#endif
-
 /**
  * Decompress all requested rotation components from an Animation Sequence
  *
diff --git a/Engine/Source/Runtime/Engine/Private/Animation/AnimEncoding_VariableKeyLerp.cpp b/Engine/Source/Runtime/Engine/Private/Animation/AnimEncoding_VariableKeyLerp.cpp
--- a/Engine/Source/Runtime/Engine/Private/Animation/AnimEncoding_VariableKeyLerp.cpp
+++ b/Engine/Source/Runtime/Engine/Private/Animation/AnimEncoding_VariableKeyLerp.cpp
@@ -9,6 +9,9 @@
 #include "Serialization/MemoryWriter.h"
 #if INTEL_ISPC
 #include "AnimEncoding_VariableKeyLerp.ispc.generated.h"
+
+static_assert(sizeof(ispc::FTransform) == sizeof(FTransform), "sizeof(ispc::FTransform) != sizeof(FTransform)");
+static_assert(sizeof(ispc::BoneTrackPair) == sizeof(BoneTrackPair), "sizeof(ispc::BoneTrackPair) != sizeof(BoneTrackPair)");
 #endif
 
 /**
diff --git a/Engine/Source/Runtime/Engine/Private/Animation/BonePose.cpp b/Engine/Source/Runtime/Engine/Private/Animation/BonePose.cpp
--- a/Engine/Source/Runtime/Engine/Private/Animation/BonePose.cpp
+++ b/Engine/Source/Runtime/Engine/Private/Animation/BonePose.cpp
@@ -6,6 +6,8 @@
 #include "HAL/ThreadSingleton.h"
 #if INTEL_ISPC
 #include "BonePose.ispc.generated.h"
+
+static_assert(sizeof(ispc::FTransform) == sizeof(FTransform), "sizeof(ispc::FTransform) != sizeof(FTransform)");
 #endif
 
 // Normalizes all rotations in this pose
diff --git a/Engine/Source/Runtime/Engine/Private/Components/SkeletalMeshComponent.cpp b/Engine/Source/Runtime/Engine/Private/Components/SkeletalMeshComponent.cpp
--- a/Engine/Source/Runtime/Engine/Private/Components/SkeletalMeshComponent.cpp
+++ b/Engine/Source/Runtime/Engine/Private/Components/SkeletalMeshComponent.cpp
@@ -45,6 +45,7 @@
 
 #if INTEL_ISPC
 #include "SkeletalMeshComponent.ispc.generated.h"
+static_assert(sizeof(ispc::FTransform) == sizeof(FTransform), "sizeof(ispc::FTransform) != sizeof(FTransform)");
 #endif
 
 #define LOCTEXT_NAMESPACE "SkeletalMeshComponent"
