Fixed bug where additive root motion's influence would be cut in half if applied while falling.
This code is not from the PR, but the bug was discovered because of it.

PR #8206

diff --git a/Engine/Source/Runtime/Engine/Private/Components/CharacterMovementComponent.cpp b/Engine/Source/Runtime/Engine/Private/Components/CharacterMovementComponent.cpp
--- a/Engine/Source/Runtime/Engine/Private/Components/CharacterMovementComponent.cpp
+++ b/Engine/Source/Runtime/Engine/Private/Components/CharacterMovementComponent.cpp
@@ -4311,6 +4311,8 @@ void UCharacterMovementComponent::PhysFalling(float deltaTime, int32 Iterations)
 		const FQuat PawnRotation = UpdatedComponent->GetComponentQuat();
 		bJustTeleported = false;
 
+		const FVector OldVelocityWithRootMotion = Velocity;
+
 		RestorePreAdditiveRootMotionVelocity();
 
 		const FVector OldVelocity = Velocity;
@@ -4378,7 +4380,7 @@ void UCharacterMovementComponent::PhysFalling(float deltaTime, int32 Iterations)
 			}
 		}
 
-		//UE_LOG(LogCharacterMovement, Log, TEXT("dt=(%.6f) OldLocation=(%s) OldVelocity=(%s) NewVelocity=(%s)"), timeTick, *(UpdatedComponent->GetComponentLocation()).ToString(), *OldVelocity.ToString(), *Velocity.ToString());
+		//UE_LOG(LogCharacterMovement, Log, TEXT("dt=(%.6f) OldLocation=(%s) OldVelocity=(%s) OldVelocityWithRootMotion=(%s) NewVelocity=(%s)"), timeTick, *(UpdatedComponent->GetComponentLocation()).ToString(), *OldVelocity.ToString(), *OldVelocityWithRootMotion.ToString(), *Velocity.ToString());
 		ApplyRootMotionToVelocity(timeTick);
 
 		if (bNotifyApex && (Velocity.Z < 0.f))
@@ -4389,7 +4391,7 @@ void UCharacterMovementComponent::PhysFalling(float deltaTime, int32 Iterations)
 		}
 
 		// Compute change in position (using midpoint integration method).
-		FVector Adjusted = 0.5f*(OldVelocity + Velocity) * timeTick;
+		FVector Adjusted = 0.5f * (OldVelocityWithRootMotion + Velocity) * timeTick;
 		
 		// Special handling if ending the jump force where we didn't apply gravity during the jump.
 		if (bEndingJumpForce && !bApplyGravityWhileJumping)
@@ -4397,7 +4399,7 @@ void UCharacterMovementComponent::PhysFalling(float deltaTime, int32 Iterations)
 			// We had a portion of the time at constant speed then a portion with acceleration due to gravity.
 			// Account for that here with a more correct change in position.
 			const float NonGravityTime = FMath::Max(0.f, timeTick - GravityTime);
-			Adjusted = (OldVelocity * NonGravityTime) + (0.5f*(OldVelocity + Velocity) * GravityTime);
+			Adjusted = (OldVelocityWithRootMotion * NonGravityTime) + (0.5f*(OldVelocityWithRootMotion + Velocity) * GravityTime);
 		}
 
 		// Move
