BodyInstance
-- Implemented prioritized hierarchy for PhysMat overrides for Complex collisions.

Ran preflight 611571ea0d783d00015b7068
Failed (numerous times) due to unrelated issues.

diff --git a/Engine/Source/Runtime/Engine/Private/PhysicsEngine/BodyInstance.cpp b/Engine/Source/Runtime/Engine/Private/PhysicsEngine/BodyInstance.cpp
--- a/Engine/Source/Runtime/Engine/Private/PhysicsEngine/BodyInstance.cpp
+++ b/Engine/Source/Runtime/Engine/Private/PhysicsEngine/BodyInstance.cpp
@@ -3026,7 +3026,7 @@ void FBodyInstance::GetComplexPhysicalMaterials(TArray<UPhysicalMaterial*>& OutP
 	GetComplexPhysicalMaterials(this, OwnerComponent, OutPhysMaterials, &OutPhysMaterialMasks);
 }
 
-void FBodyInstance::GetComplexPhysicalMaterials(const FBodyInstance*, TWeakObjectPtr<UPrimitiveComponent> OwnerComp, TArray<UPhysicalMaterial*>& OutPhysMaterials, TArray<FPhysicalMaterialMaskParams>* OutPhysMaterialMasks)
+void FBodyInstance::GetComplexPhysicalMaterials(const FBodyInstance* BodyInstance, TWeakObjectPtr<UPrimitiveComponent> OwnerComp, TArray<UPhysicalMaterial*>& OutPhysMaterials, TArray<FPhysicalMaterialMaskParams>* OutPhysMaterialMasks)
 {
 	if(!GEngine || !GEngine->DefaultPhysMaterial)
 	{
@@ -3035,46 +3035,69 @@ void FBodyInstance::GetComplexPhysicalMaterials(const FBodyInstance*, TWeakObjec
 		return;
 	}
 
-	// See if the Material has a PhysicalMaterial
-	UPrimitiveComponent* PrimComp = OwnerComp.Get();
-	if (PrimComp)
+	// Find the PhysicalMaterial we need to apply to the physics bodies.
+	// (LOW priority) Engine Mat, Material PhysMat, Component Override, Body Override (HIGH priority)
+	
+	// BodyInstance override
+	if (BodyInstance && BodyInstance->PhysMaterialOverride != nullptr)
 	{
-		const int32 NumMaterials = PrimComp->GetNumMaterials();
-		OutPhysMaterials.SetNum(NumMaterials);
-
-		if (OutPhysMaterialMasks)
+		OutPhysMaterials.SetNum(1);
+		OutPhysMaterials[0] = BodyInstance->PhysMaterialOverride;
+		check(!OutPhysMaterials[0] || OutPhysMaterials[0]->IsValidLowLevel());
+	}
+	else
+	{
+		// Component override
+		UPrimitiveComponent* OwnerPrimComponent = OwnerComp.Get();
+		if (OwnerPrimComponent && OwnerPrimComponent->BodyInstance.PhysMaterialOverride != nullptr)
 		{
-			OutPhysMaterialMasks->SetNum(NumMaterials);
+			OutPhysMaterials.SetNum(1);
+			OutPhysMaterials[0] = OwnerComp->BodyInstance.PhysMaterialOverride;
+			check(!OutPhysMaterials[0] || OutPhysMaterials[0]->IsValidLowLevel());
 		}
-
-		for (int32 MatIdx = 0; MatIdx < NumMaterials; MatIdx++)
+		else
 		{
-			UPhysicalMaterial* PhysMat = GEngine->DefaultPhysMaterial;
-			UMaterialInterface* Material = PrimComp->GetMaterial(MatIdx);
-			if (Material)
+			// See if the Material has a PhysicalMaterial
+			if (OwnerPrimComponent)
 			{
-				PhysMat = Material->GetPhysicalMaterial();
-			}
-			
-			OutPhysMaterials[MatIdx] = PhysMat;
+				const int32 NumMaterials = OwnerPrimComponent->GetNumMaterials();
+				OutPhysMaterials.SetNum(NumMaterials);
 
-			if (OutPhysMaterialMasks)
-			{
-				UPhysicalMaterialMask* PhysMatMask = nullptr;
-				UMaterialInterface* PhysMatMap = nullptr;
+				if (OutPhysMaterialMasks)
+				{
+					OutPhysMaterialMasks->SetNum(NumMaterials);
+				}
 
-				if (Material)
+				for (int32 MatIdx = 0; MatIdx < NumMaterials; MatIdx++)
 				{
-					PhysMatMask = Material->GetPhysicalMaterialMask();
-					if (PhysMatMask)
+					UPhysicalMaterial* PhysMat = GEngine->DefaultPhysMaterial;
+					UMaterialInterface* Material = OwnerPrimComponent->GetMaterial(MatIdx);
+					if (Material)
 					{
-						PhysMatMap = Material;
+						PhysMat = Material->GetPhysicalMaterial();
 					}
-				}
 
-				(*OutPhysMaterialMasks)[MatIdx].PhysicalMaterialMask = PhysMatMask;
-				(*OutPhysMaterialMasks)[MatIdx].PhysicalMaterialMap = PhysMatMap;
-			}
+					OutPhysMaterials[MatIdx] = PhysMat;
+
+					if (OutPhysMaterialMasks)
+					{
+						UPhysicalMaterialMask* PhysMatMask = nullptr;
+						UMaterialInterface* PhysMatMap = nullptr;
+
+						if (Material)
+						{
+							PhysMatMask = Material->GetPhysicalMaterialMask();
+							if (PhysMatMask)
+							{
+								PhysMatMap = Material;
+							}
+						}
+
+						(*OutPhysMaterialMasks)[MatIdx].PhysicalMaterialMask = PhysMatMask;
+						(*OutPhysMaterialMasks)[MatIdx].PhysicalMaterialMap = PhysMatMap;
+					}
+				}
+			}			
 		}
 	}
 }
