bIsVisibleInRayTracing is now exposed in blueprints through the SetVisibleInRayTracing method.

diff --git a/Engine/Source/Runtime/Engine/Classes/Components/PrimitiveComponent.h b/Engine/Source/Runtime/Engine/Classes/Components/PrimitiveComponent.h
--- a/Engine/Source/Runtime/Engine/Classes/Components/PrimitiveComponent.h
+++ b/Engine/Source/Runtime/Engine/Classes/Components/PrimitiveComponent.h
@@ -1546,6 +1546,10 @@ public:
 	UFUNCTION(BlueprintCallable, Category="Rendering")
 	void SetOnlyOwnerSee(bool bNewOnlyOwnerSee);
 
+	/** Changes the value of bIsVisibleInRayTracing. */
+	UFUNCTION(BlueprintCallable, Category = "Rendering")
+	void SetVisibleInRayTracing(bool bNewVisibleInRayTracing);
+
 	/** Changes the value of CastShadow. */
 	UFUNCTION(BlueprintCallable, Category="Rendering")
 	void SetCastShadow(bool NewCastShadow);
diff --git a/Engine/Source/Runtime/Engine/Private/Components/PrimitiveComponent.cpp b/Engine/Source/Runtime/Engine/Private/Components/PrimitiveComponent.cpp
--- a/Engine/Source/Runtime/Engine/Private/Components/PrimitiveComponent.cpp
+++ b/Engine/Source/Runtime/Engine/Private/Components/PrimitiveComponent.cpp
@@ -1581,6 +1581,15 @@ bool UPrimitiveComponent::GetLevelInstanceEditingState() const
 	return false;
 }
 
+void UPrimitiveComponent::SetVisibleInRayTracing(bool bNewVisibleInRayTracing)
+{
+	if (bNewVisibleInRayTracing != bVisibleInRayTracing)
+	{
+		bVisibleInRayTracing = bNewVisibleInRayTracing;
+		MarkRenderStateDirty();
+	}
+}
+
 void UPrimitiveComponent::SetCastShadow(bool NewCastShadow)
 {
 	if(NewCastShadow != CastShadow)
