Fix for overshoot w/ blendspace smoothing
Renamed GetWeight (deprecated) to GetClampedWeight

diff --git a/Engine/Source/Editor/Persona/Private/SAnimationBlendSpaceGridWidget.cpp b/Engine/Source/Editor/Persona/Private/SAnimationBlendSpaceGridWidget.cpp
--- a/Engine/Source/Editor/Persona/Private/SAnimationBlendSpaceGridWidget.cpp
+++ b/Engine/Source/Editor/Persona/Private/SAnimationBlendSpaceGridWidget.cpp
@@ -575,7 +575,7 @@ float SBlendSpaceGridWidget::GetSampleLookupWeight(int32 SampleIndex) const
 		{
 			if (LookedUpSample.SampleDataIndex == SampleIndex)
 			{
-				LookedUpSampleWeight += LookedUpSample.GetWeight();
+				LookedUpSampleWeight += LookedUpSample.GetClampedWeight();
 			}
 		}
 		return FMath::Clamp(LookedUpSampleWeight, 0.0f, 1.0f);
diff --git a/Engine/Source/Runtime/Engine/Classes/Animation/AnimationAsset.h b/Engine/Source/Runtime/Engine/Classes/Animation/AnimationAsset.h
--- a/Engine/Source/Runtime/Engine/Classes/Animation/AnimationAsset.h
+++ b/Engine/Source/Runtime/Engine/Classes/Animation/AnimationAsset.h
@@ -161,7 +161,14 @@ struct FBlendSampleData
 	{
 		TotalWeight += Weight;
 	}
+
+	UE_DEPRECATED(5.0, "GetWeight() was renamed to GetClampedWeight()")
 	float GetWeight() const
+	{
+		return GetClampedWeight();
+	}
+
+	float GetClampedWeight() const
 	{
 		return FMath::Clamp<float>(TotalWeight, 0.f, 1.f);
 	}
diff --git a/Engine/Source/Runtime/Engine/Private/Animation/AnimInstance.cpp b/Engine/Source/Runtime/Engine/Private/Animation/AnimInstance.cpp
--- a/Engine/Source/Runtime/Engine/Private/Animation/AnimInstance.cpp
+++ b/Engine/Source/Runtime/Engine/Private/Animation/AnimInstance.cpp
@@ -895,7 +895,7 @@ void OutputTickRecords(const TArray<FAnimTickRecord>& Records, UCanvas* Canvas,
 						FBlendSampleData& WeightedSample = SampleData[WeightedSampleIndex];
 						if (WeightedSample.SampleDataIndex == SampleIndex)
 						{
-							Weight += WeightedSample.GetWeight();
+							Weight += WeightedSample.GetClampedWeight();
 						}
 						else if (WeightedSample.SampleDataIndex > SampleIndex)
 						{
diff --git a/Engine/Source/Runtime/Engine/Private/Animation/AnimSingleNodeInstanceProxy.cpp b/Engine/Source/Runtime/Engine/Private/Animation/AnimSingleNodeInstanceProxy.cpp
--- a/Engine/Source/Runtime/Engine/Private/Animation/AnimSingleNodeInstanceProxy.cpp
+++ b/Engine/Source/Runtime/Engine/Private/Animation/AnimSingleNodeInstanceProxy.cpp
@@ -288,7 +288,7 @@ float FAnimSingleNodeInstanceProxy::GetBlendSpaceLength() const
 	{
 		const FBlendSampleData& Data = BlendSampleData[Index];
 		float AnimLength = Data.Animation->GetPlayLength();
-		float Weight = Data.GetWeight();
+		float Weight = Data.GetClampedWeight();
 		TotalLength += AnimLength * Weight;
 		TotalWeight += Weight;
 	}
diff --git a/Engine/Source/Runtime/Engine/Private/Animation/AnimationRuntime.cpp b/Engine/Source/Runtime/Engine/Private/Animation/AnimationRuntime.cpp
--- a/Engine/Source/Runtime/Engine/Private/Animation/AnimationRuntime.cpp
+++ b/Engine/Source/Runtime/Engine/Private/Animation/AnimationRuntime.cpp
@@ -504,7 +504,7 @@ void FAnimationRuntime::BlendTwoPosesTogetherPerBone(const FAnimationPoseData& S
 template <int32 TRANSFORM_BLEND_MODE>
 void BlendPosePerBone(const TArray<FBoneIndexType>& RequiredBoneIndices, const TArray<int32>& PerBoneIndices, const FBlendSampleData& BlendSampleDataCache, FTransformArrayA2& ResultAtoms, const FTransformArrayA2& SourceAtoms)
 {
-	const float BlendWeight = BlendSampleDataCache.GetWeight();
+	const float BlendWeight = BlendSampleDataCache.GetClampedWeight();
 	for (int32 i = 0; i < RequiredBoneIndices.Num(); ++i)
 	{
 		const int32 BoneIndex = RequiredBoneIndices[i];
@@ -523,7 +523,7 @@ void BlendPosePerBone(const TArray<FBoneIndexType>& RequiredBoneIndices, const T
 template <int32 TRANSFORM_BLEND_MODE>
 void BlendPosePerBone(const TArray<int32>& PerBoneIndices, const FBlendSampleData& BlendSampleDataCache, FCompactPose& ResultPose, const FCompactPose& SourcePose)
 {
-	const float BlendWeight = BlendSampleDataCache.GetWeight();
+	const float BlendWeight = BlendSampleDataCache.GetClampedWeight();
 	for (FCompactPoseBoneIndex BoneIndex : SourcePose.ForEachBoneIndex())
 	{
 		const int32 PerBoneIndex = PerBoneIndices[BoneIndex.GetInt()];
diff --git a/Engine/Source/Runtime/Engine/Private/Animation/AttributeBlendData.cpp b/Engine/Source/Runtime/Engine/Private/Animation/AttributeBlendData.cpp
--- a/Engine/Source/Runtime/Engine/Private/Animation/AttributeBlendData.cpp
+++ b/Engine/Source/Runtime/Engine/Private/Animation/AttributeBlendData.cpp
@@ -345,7 +345,7 @@ namespace UE { namespace Anim {
 			}
 			else
 			{
-				return BlendSampleData.GetWeight();
+				return BlendSampleData.GetClampedWeight();
 			}
 		}
 
diff --git a/Engine/Source/Runtime/Engine/Private/Animation/AttributesRuntime.cpp b/Engine/Source/Runtime/Engine/Private/Animation/AttributesRuntime.cpp
--- a/Engine/Source/Runtime/Engine/Private/Animation/AttributesRuntime.cpp
+++ b/Engine/Source/Runtime/Engine/Private/Animation/AttributesRuntime.cpp
@@ -368,7 +368,7 @@ void Attributes::BlendAttributesPerBone(TArrayView<const FStackAttributeContaine
 		TArray<float> Weights;
 		for (const FBlendSampleData& SampleData : BlendSampleDataCache)
 		{
-			Weights.Add(SampleData.GetWeight());
+			Weights.Add(SampleData.GetClampedWeight());
 		}
 
 		for (const TWeakObjectPtr<const UScriptStruct>& WeakScriptStruct : UniqueStructs)
@@ -409,7 +409,7 @@ void Attributes::BlendAttributesPerBone(TArrayView<const FStackAttributeContaine
 		for (const int32& SampleIndex : BlendSampleDataCacheIndices)
 		{
 			const FBlendSampleData& SampleData = BlendSampleDataCache[SampleIndex];
-			Weights.Add(SampleData.GetWeight());
+			Weights.Add(SampleData.GetClampedWeight());
 		}
 
 		for (const TWeakObjectPtr<const UScriptStruct>& WeakScriptStruct : UniqueStructs)
diff --git a/Engine/Source/Runtime/Engine/Private/Animation/BlendSpace.cpp b/Engine/Source/Runtime/Engine/Private/Animation/BlendSpace.cpp
--- a/Engine/Source/Runtime/Engine/Private/Animation/BlendSpace.cpp
+++ b/Engine/Source/Runtime/Engine/Private/Animation/BlendSpace.cpp
@@ -560,14 +560,14 @@ void UBlendSpace::TickAssetPlayer(FAnimTickRecord& Instance, struct FAnimNotifyQ
 
 							if (Context.RootMotionMode == ERootMotionMode::RootMotionFromEverything && Sample.Animation->bEnableRootMotion)
 							{
-								Context.RootMotionMovementParams.AccumulateWithBlend(Sample.Animation->ExtractRootMotion(PrevSampleDataTime, DeltaTimePosition, Instance.bLooping), SampleEntry.GetWeight());
+								Context.RootMotionMovementParams.AccumulateWithBlend(Sample.Animation->ExtractRootMotion(PrevSampleDataTime, DeltaTimePosition, Instance.bLooping), SampleEntry.GetClampedWeight());
 							}
 
 							// Capture the final adjusted delta time and previous frame time as an asset player record
 							SampleEntry.DeltaTimeRecord.Previous = PrevSampleDataTime;
 							SampleEntry.DeltaTimeRecord.Delta = DeltaTimePosition;
 
-							UE_LOG(LogAnimation, Verbose, TEXT("%d. Blending animation(%s) with %f weight at time %0.2f"), I + 1, *Sample.Animation->GetName(), SampleEntry.GetWeight(), CurrentSampleDataTime);
+							UE_LOG(LogAnimation, Verbose, TEXT("%d. Blending animation(%s) with %f weight at time %0.2f"), I + 1, *Sample.Animation->GetName(), SampleEntry.GetClampedWeight(), CurrentSampleDataTime);
 						}
 					}
 				}
@@ -827,7 +827,7 @@ void UBlendSpace::GetAnimationPose_Internal(TArray<FBlendSampleData>& BlendSampl
 		if (SampleData.IsValidIndex(BlendSampleDataCache[I].SampleDataIndex))
 		{
 			const FBlendSample& Sample = SampleData[BlendSampleDataCache[I].SampleDataIndex];
-			ChildrenWeights[I] = BlendSampleDataCache[I].GetWeight();
+			ChildrenWeights[I] = BlendSampleDataCache[I].GetClampedWeight();
 
 			if (bNested)
 			{
@@ -982,24 +982,24 @@ bool UBlendSpace::GetSamplesFromBlendInput(
 					(FirstSample->Animation != nullptr && FirstSample->Animation == SecondSample->Animation))
 				{
 					//Calc New Sample Playrate
-					const float TotalWeight = FirstSample->GetWeight() + SecondSample->GetWeight();
+					const float TotalWeight = FirstSample->GetClampedWeight() + SecondSample->GetClampedWeight();
 
 					// Only combine playrates if total weight > 0
 					if (!FMath::IsNearlyZero(TotalWeight))
 					{
-						if (FirstSample->GetWeight() < SecondSample->GetWeight())
+						if (FirstSample->GetClampedWeight() < SecondSample->GetClampedWeight())
 						{
 							// Not strictly necessary, but if we swap here then we keep the one that has a higher
 							// weight, which can make debugging/viewing the blend space more intuitive.
 							OutSampleDataList.Swap(Index1, Index2);
 						}
 
-						const float OriginalWeightedPlayRate = FirstSample->SamplePlayRate * (FirstSample->GetWeight() / TotalWeight);
-						const float SecondSampleWeightedPlayRate = SecondSample->SamplePlayRate * (SecondSample->GetWeight() / TotalWeight);
+						const float OriginalWeightedPlayRate = FirstSample->SamplePlayRate * (FirstSample->GetClampedWeight() / TotalWeight);
+						const float SecondSampleWeightedPlayRate = SecondSample->SamplePlayRate * (SecondSample->GetClampedWeight() / TotalWeight);
 						FirstSample->SamplePlayRate = OriginalWeightedPlayRate + SecondSampleWeightedPlayRate;
 
 						// add weight
-						FirstSample->AddWeight(SecondSample->GetWeight());
+						FirstSample->AddWeight(SecondSample->GetClampedWeight());
 					}
 
 					// as for time or previous time will be the master one(Index1)
@@ -1520,8 +1520,8 @@ float UBlendSpace::GetAnimationLengthFromSampleData(const TArray<FBlendSampleDat
 				//Multiple samples contribution which we would otherwise lose
 				const float MultipliedSampleRateScale = Sample.Animation->RateScale * SampleDataList[I].SamplePlayRate;
 				// apply rate scale to get actual playback time
-				BlendAnimLength += (Sample.Animation->GetPlayLength() / ((MultipliedSampleRateScale) != 0.0f ? FMath::Abs(MultipliedSampleRateScale) : 1.0f)) * SampleDataList[I].GetWeight();
-				UE_LOG(LogAnimation, Verbose, TEXT("[%d] - Sample Animation(%s) : Weight(%0.5f) "), I + 1, *Sample.Animation->GetName(), SampleDataList[I].GetWeight());
+				BlendAnimLength += (Sample.Animation->GetPlayLength() / ((MultipliedSampleRateScale) != 0.0f ? FMath::Abs(MultipliedSampleRateScale) : 1.0f)) * SampleDataList[I].GetClampedWeight();
+				UE_LOG(LogAnimation, Verbose, TEXT("[%d] - Sample Animation(%s) : Weight(%0.5f) "), I + 1, *Sample.Animation->GetName(), SampleDataList[I].GetClampedWeight());
 			}
 		}
 	}
@@ -1670,7 +1670,7 @@ bool UBlendSpace::InterpolateWeightOfSampleData(float DeltaTime, const TArray<FB
 				if (InterpData.TotalWeight > ZERO_ANIMWEIGHT_THRESH || TotalPerBoneWeight > ZERO_ANIMWEIGHT_THRESH)
 				{
 					FinalSampleDataList.Add(InterpData);
-					TotalFinalWeight += InterpData.GetWeight();
+					TotalFinalWeight += InterpData.GetClampedWeight();
 					TotalFinalPerBoneWeight += TotalPerBoneWeight;
 					bTargetSampleExists = true;
 					break;
@@ -1698,7 +1698,7 @@ bool UBlendSpace::InterpolateWeightOfSampleData(float DeltaTime, const TArray<FB
 			if (InterpData.TotalWeight > ZERO_ANIMWEIGHT_THRESH || TotalPerBoneWeight > ZERO_ANIMWEIGHT_THRESH)
 			{
 				FinalSampleDataList.Add(InterpData);
-				TotalFinalWeight += InterpData.GetWeight();
+				TotalFinalWeight += InterpData.GetClampedWeight();
 				TotalFinalPerBoneWeight += TotalPerBoneWeight;
 			}
 		}
@@ -1751,7 +1751,7 @@ bool UBlendSpace::InterpolateWeightOfSampleData(float DeltaTime, const TArray<FB
 			if (InterpData.TotalWeight > ZERO_ANIMWEIGHT_THRESH || TotalPerBoneWeight > ZERO_ANIMWEIGHT_THRESH)
 			{
 				FinalSampleDataList.Add(InterpData);
-				TotalFinalWeight += InterpData.GetWeight();
+				TotalFinalWeight += InterpData.GetClampedWeight();
 				TotalFinalPerBoneWeight += TotalPerBoneWeight;
 			}
 		}
diff --git a/Engine/Source/Runtime/Engine/Private/Animation/BlendSpaceUtilities.cpp b/Engine/Source/Runtime/Engine/Private/Animation/BlendSpaceUtilities.cpp
--- a/Engine/Source/Runtime/Engine/Private/Animation/BlendSpaceUtilities.cpp
+++ b/Engine/Source/Runtime/Engine/Private/Animation/BlendSpaceUtilities.cpp
@@ -54,13 +54,13 @@ bool FSyncPattern::DoesPatternMatch(const TArray<FName>& TestMarkerNames)
 int32 FBlendSpaceUtilities::GetHighestWeightSample(const TArray<struct FBlendSampleData> &SampleDataList)
 {
 	int32 HighestWeightIndex = 0;
-	float HighestWeight = SampleDataList[HighestWeightIndex].GetWeight();
+	float HighestWeight = SampleDataList[HighestWeightIndex].GetClampedWeight();
 	for (int32 I = 1; I < SampleDataList.Num(); I++)
 	{
-		if (SampleDataList[I].GetWeight() > HighestWeight)
+		if (SampleDataList[I].GetClampedWeight() > HighestWeight)
 		{
 			HighestWeightIndex = I;
-			HighestWeight = SampleDataList[I].GetWeight();
+			HighestWeight = SampleDataList[I].GetClampedWeight();
 		}
 	}
 	return HighestWeightIndex;
@@ -74,11 +74,11 @@ int32 FBlendSpaceUtilities::GetHighestWeightMarkerSyncSample(const TArray<struct
 	for (int32 I = 0; I < SampleDataList.Num(); I++)
 	{
 		const FBlendSampleData& SampleData = SampleDataList[I];
-		if (SampleData.GetWeight() > HighestWeight &&
+		if (SampleData.GetClampedWeight() > HighestWeight &&
 			BlendSamples[SampleData.SampleDataIndex].Animation->AuthoredSyncMarkers.Num() > 0)
 		{
 			HighestWeightIndex = I;
-			HighestWeight = SampleData.GetWeight();
+			HighestWeight = SampleData.GetClampedWeight();
 		}
 	}
 	return HighestWeightIndex;
diff --git a/Engine/Source/Runtime/Engine/Private/AnimationAsset.cpp b/Engine/Source/Runtime/Engine/Private/AnimationAsset.cpp
--- a/Engine/Source/Runtime/Engine/Private/AnimationAsset.cpp
+++ b/Engine/Source/Runtime/Engine/Private/AnimationAsset.cpp
@@ -696,7 +696,7 @@ void FBlendSampleData::NormalizeDataWeight(TArray<FBlendSampleData>& SampleDataL
 	{
 		checkf(SampleDataList[PoseIndex].PerBoneBlendData.Num() == NumBones, TEXT("Attempted to normalise a blend sample list, but the samples have differing numbers of bones."));
 
-		TotalSum += SampleDataList[PoseIndex].GetWeight();
+		TotalSum += SampleDataList[PoseIndex].TotalWeight;
 
 		if (SampleDataList[PoseIndex].PerBoneBlendData.Num() > 0)
 		{
