VT builder was missing TBSetting bHasEditorOnlyData
this meant VT encoded in Editor was not using fast cook settings like it should

diff --git a/Engine/Source/Runtime/Engine/Private/VT/VirtualTextureDataBuilder.cpp b/Engine/Source/Runtime/Engine/Private/VT/VirtualTextureDataBuilder.cpp
--- a/Engine/Source/Runtime/Engine/Private/VT/VirtualTextureDataBuilder.cpp
+++ b/Engine/Source/Runtime/Engine/Private/VT/VirtualTextureDataBuilder.cpp
@@ -634,6 +634,7 @@ void FVirtualTextureDataBuilder::BuildTiles(const TArray<FVTSourceTileEntry>& Ti
 		TBSettings.MipGenSettings = TMGS_NoMipmaps;
 		TBSettings.bForceAlphaChannel = BuildSettingsForLayer.bForceAlphaChannel;
 		TBSettings.bVirtualStreamable = true;
+		TBSettings.bHasEditorOnlyData = BuildSettingsForLayer.bHasEditorOnlyData;
 
 		check(TBSettings.GetGammaSpace() == BuildSettingsForLayer.GetGammaSpace());
  
