Print effective parallelizaiton of shader compilation.

- Also made unused thread count to not apply to build machines (there are jobs where they run non-commandlets).

diff --git a/Engine/Source/Runtime/Engine/Private/ShaderCompiler/ShaderCompiler.cpp b/Engine/Source/Runtime/Engine/Private/ShaderCompiler/ShaderCompiler.cpp
--- a/Engine/Source/Runtime/Engine/Private/ShaderCompiler/ShaderCompiler.cpp
+++ b/Engine/Source/Runtime/Engine/Private/ShaderCompiler/ShaderCompiler.cpp
@@ -2741,6 +2741,20 @@ void FShaderCompilerStats::WriteStatSummary()
 
 	if (ShaderTimings.Num())
 	{
+		// calculate effective parallelization (total time needed to compile all shaders divided by actual wall clock time spent processing at least 1 shader)
+		if (TotalTimeAtLeastOneJobWasInFlight > 0.0)
+		{
+			double TotalTimeForAllShaders = 0;
+			for (TMap<FString, FShaderTimings>::TConstIterator Iter(ShaderTimings); Iter; ++Iter)
+			{
+				TotalTimeForAllShaders += Iter.Value().TotalCompileTime;
+			}
+
+			double EffectiveParallelization = TotalTimeForAllShaders / TotalTimeAtLeastOneJobWasInFlight;
+			UE_LOG(LogShaderCompilers, Display, TEXT("Effective parallelization: %.2f (times faster than compiling all shaders on one thread)."), EffectiveParallelization);
+		}
+
+
 		// sort by avg time
 		ShaderTimings.ValueSort([](const FShaderTimings& A, const FShaderTimings& B) { return A.AverageCompileTime > B.AverageCompileTime; });
 
@@ -3098,7 +3112,7 @@ FShaderCompilingManager::FShaderCompilingManager() :
 
 	// Don't reserve threads based on a percentage if we are in a commandlet or on a low core machine.
 	// In these scenarios we should try to use as many threads as possible.
-	if (!IsRunningCommandlet() && NumVirtualCores > ShaderCompilerCoreCountThreshold)
+	if (!IsRunningCommandlet() && !GIsBuildMachine && NumVirtualCores > ShaderCompilerCoreCountThreshold)
 	{
 		// Reserve a percentage of the threads for general background work.
 		float PercentageUnusedShaderCompilingThreads;
@@ -3264,7 +3278,7 @@ FShaderCompilingManager::FShaderCompilingManager() :
 	}
 	else
 	{
-		UE_LOG(LogShaderCompilers, Display, TEXT("Using Local Shader Compiler."));
+		UE_LOG(LogShaderCompilers, Display, TEXT("Using Local Shader Compiler with %d workers."), NumShaderCompilingThreads);
 
 		if (GIsBuildMachine)
 		{
