Enable legacy input scales by default to fix some internal blockers while upgrade paths of ini files are considered

#rnx

diff --git a/Engine/Source/Runtime/Engine/Private/UserInterface/InputSettings.cpp b/Engine/Source/Runtime/Engine/Private/UserInterface/InputSettings.cpp
--- a/Engine/Source/Runtime/Engine/Private/UserInterface/InputSettings.cpp
+++ b/Engine/Source/Runtime/Engine/Private/UserInterface/InputSettings.cpp
@@ -21,7 +21,7 @@ UInputSettings::UInputSettings(const FObjectInitializer& ObjectInitializer)
 	: Super(ObjectInitializer)
 	, bCaptureMouseOnLaunch(true)
 	, bDefaultViewportMouseLock_DEPRECATED(false)
-	, bEnableLegacyInputScales(false)
+	, bEnableLegacyInputScales(true)
 	, DefaultViewportMouseCaptureMode(EMouseCaptureMode::CapturePermanently_IncludingInitialMouseDown)
 	, DefaultViewportMouseLockMode(EMouseLockMode::LockOnCapture)
 	, DefaultPlayerInputClass(UPlayerInput::StaticClass())
