FStatMessage clocks values using FPlatformTime::Cycles().  Use 32-bit FPlatformTime::ToMilliseconds() rather than FPlatformTime::ToMilliseconds64().

diff --git a/Engine/Source/Runtime/Engine/Private/StatsRender2.cpp b/Engine/Source/Runtime/Engine/Private/StatsRender2.cpp
--- a/Engine/Source/Runtime/Engine/Private/StatsRender2.cpp
+++ b/Engine/Source/Runtime/Engine/Private/StatsRender2.cpp
@@ -347,7 +347,7 @@ static int32 RenderCycle( const FComplexStatMessage& Item, class FCanvas* Canvas
 
 	if( bIsInitialized )
 	{
-		const float InMs = FPlatformTime::ToMilliseconds64(Item.GetValue_Duration(EComplexStatField::IncAve));
+		const float InMs = FPlatformTime::ToMilliseconds(Item.GetValue_Duration(EComplexStatField::IncAve));
 		// Color will be determined by the average value of history
 		// If show inclusive and and show exclusive is on, then it will choose color based on inclusive average
 		// #Stats: 2015-06-09 This is slow, fix this
@@ -394,13 +394,13 @@ static int32 RenderCycle( const FComplexStatMessage& Item, class FCanvas* Canvas
 	// Add the two inclusive columns if asked
 	if( bIsInitialized )
 	{
-		RightJustify(Canvas,CurrX,Y,*FString::Printf(TEXT("%1.2f ms"),FPlatformTime::ToMilliseconds64(Item.GetValue_Duration(EComplexStatField::IncAve))),Color); 
+		RightJustify(Canvas,CurrX,Y,*FString::Printf(TEXT("%1.2f ms"),FPlatformTime::ToMilliseconds(Item.GetValue_Duration(EComplexStatField::IncAve))),Color); 
 	}
 	CurrX += Globals.InterColumnOffset;
 
 	if( bIsInitialized )
 	{
-		RightJustify(Canvas,CurrX,Y,*FString::Printf(TEXT("%1.2f ms"),FPlatformTime::ToMilliseconds64(Item.GetValue_Duration(EComplexStatField::IncMax))),Color);
+		RightJustify(Canvas,CurrX,Y,*FString::Printf(TEXT("%1.2f ms"),FPlatformTime::ToMilliseconds(Item.GetValue_Duration(EComplexStatField::IncMax))),Color);
 		
 	}
 	CurrX += Globals.InterColumnOffset;
@@ -410,13 +410,13 @@ static int32 RenderCycle( const FComplexStatMessage& Item, class FCanvas* Canvas
 		// And the exclusive if asked
 		if( bIsInitialized )
 		{
-			RightJustify(Canvas,CurrX,Y,*FString::Printf(TEXT("%1.2f ms"),FPlatformTime::ToMilliseconds64(Item.GetValue_Duration(EComplexStatField::ExcAve))),Color);
+			RightJustify(Canvas,CurrX,Y,*FString::Printf(TEXT("%1.2f ms"),FPlatformTime::ToMilliseconds(Item.GetValue_Duration(EComplexStatField::ExcAve))),Color);
 		}
 		CurrX += Globals.InterColumnOffset;
 
 		if( bIsInitialized )
 		{
-			RightJustify(Canvas,CurrX,Y,*FString::Printf(TEXT("%1.2f ms"),FPlatformTime::ToMilliseconds64(Item.GetValue_Duration(EComplexStatField::ExcMax))),Color);
+			RightJustify(Canvas,CurrX,Y,*FString::Printf(TEXT("%1.2f ms"),FPlatformTime::ToMilliseconds(Item.GetValue_Duration(EComplexStatField::ExcMax))),Color);
 		}
 		CurrX += Globals.InterColumnOffset;
 	}
@@ -686,8 +686,8 @@ int32 RenderGroupBudget( FCanvas* Canvas, const  int32 X, const int32 Y, const u
 	// Stat [32chars]	Value [8chars]	Average [8chars]
 	const FStatRenderGlobals& Globals = GetStatRenderGlobals();
 
-	const float AvgTotalMs = FPlatformTime::ToMilliseconds64(AvgTotalTime);
-	const float MaxTotalMs = FPlatformTime::ToMilliseconds64(MaxTotalTime);
+	const float AvgTotalMs = FPlatformTime::ToMilliseconds(AvgTotalTime);
+	const float MaxTotalMs = FPlatformTime::ToMilliseconds(MaxTotalTime);
 
 	FString BudgetString = FString::Printf(TEXT("Total (of %1.2f ms)"), GroupBudget);
 
