Properly re-initialize soundfield submixes if the encoding format changes (in editor).
#p4v-preflight-copy 17142211

diff --git a/Engine/Source/Runtime/AudioMixer/Private/AudioMixerDevice.cpp b/Engine/Source/Runtime/AudioMixer/Private/AudioMixerDevice.cpp
--- a/Engine/Source/Runtime/AudioMixer/Private/AudioMixerDevice.cpp
+++ b/Engine/Source/Runtime/AudioMixer/Private/AudioMixerDevice.cpp
@@ -998,7 +998,7 @@ namespace Audio
 				if (!IsMasterSubmixType(SubmixToLoad))
 				{
 					LoadSoundSubmix(*SubmixToLoad);
-					InitSoundfieldAndEndpointDataForSubmix(*SubmixToLoad, GetSubmixInstance(SubmixToLoad).Pin(), false);
+					InitSoundfieldAndEndpointDataForSubmix(*SubmixToLoad, GetSubmixInstance(SubmixToLoad).Pin(), true);
 				}
 			}
 			bSubmixRegistrationDisabled = false;
diff --git a/Engine/Source/Runtime/Engine/Private/SoundSubmix.cpp b/Engine/Source/Runtime/Engine/Private/SoundSubmix.cpp
--- a/Engine/Source/Runtime/Engine/Private/SoundSubmix.cpp
+++ b/Engine/Source/Runtime/Engine/Private/SoundSubmix.cpp
@@ -844,7 +844,7 @@ void USoundfieldSubmix::SanitizeLinks()
 	}
 
 	// If this submix is now incompatible with the parent submix, disconnect it.
-	if (!SubmixUtils::AreSubmixFormatsCompatible(this, ParentSubmix))
+	if (ParentSubmix && !SubmixUtils::AreSubmixFormatsCompatible(this, ParentSubmix))
 	{
 		ParentSubmix->ChildSubmixes.RemoveSwap(this);
 		ParentSubmix->Modify();
@@ -872,6 +872,12 @@ void USoundfieldSubmix::PostEditChangeProperty(struct FPropertyChangedEvent& Pro
 		if (PropertyChangedEvent.Property->GetFName() == GET_MEMBER_NAME_CHECKED(USoundfieldSubmix, SoundfieldEncodingFormat))
 		{
 			bShouldSanitizeLinks = true;
+			
+			FAudioDeviceManager* AudioDeviceManager = FAudioDeviceManager::Get();
+			if (AudioDeviceManager)
+			{
+				AudioDeviceManager->InitSoundSubmixes();
+			}
 		}
 	}
 
