Lights can now enable ray traced shadows independently of the r.raytracing.shadows project settings. There are three stages for this setting: Disabled/UseProjectSetting/Enabled. For old scenes, lights that had bCastRaytracingShadows enabled are mapped to the new UseProjectSetting value

diff --git a/Engine/Source/Runtime/Core/Public/UObject/UE5MainStreamObjectVersion.h b/Engine/Source/Runtime/Core/Public/UObject/UE5MainStreamObjectVersion.h
--- a/Engine/Source/Runtime/Core/Public/UObject/UE5MainStreamObjectVersion.h
+++ b/Engine/Source/Runtime/Core/Public/UObject/UE5MainStreamObjectVersion.h
@@ -169,6 +169,9 @@ struct CORE_API FUE5MainStreamObjectVersion
 		// Introduce RigVM Memory Class Object
 		RigVMMemoryStorageObject,
 
+		// Ray tracing shadows have three states now (Disabled, Use Project Settings, Enabled)
+		RayTracedShadowsType,
+
 		// -----<new versions can be added above this line>-------------------------------------------------
 		VersionPlusOne,
 		LatestVersion = VersionPlusOne - 1
diff --git a/Engine/Source/Runtime/Engine/Classes/Components/LightComponentBase.h b/Engine/Source/Runtime/Engine/Classes/Components/LightComponentBase.h
--- a/Engine/Source/Runtime/Engine/Classes/Components/LightComponentBase.h
+++ b/Engine/Source/Runtime/Engine/Classes/Components/LightComponentBase.h
@@ -84,8 +84,11 @@ class ENGINE_API ULightComponentBase : public USceneComponent
 	uint32 bCastDeepShadow : 1;
 
 	/** Whether the light shadows are computed with shadow-mapping or ray-tracing (when available). */
-	UPROPERTY(EditAnywhere, BlueprintReadOnly, Category = Light, meta = (DisplayName = "Cast Ray Tracing Shadows"), AdvancedDisplay)
-	uint32 bCastRaytracedShadow : 1;
+	UPROPERTY()
+	uint32 bCastRaytracedShadow_DEPRECATED : 1;
+
+	UPROPERTY(EditAnywhere, BlueprintReadOnly, Category = Light, meta = (DisplayName = "Cast Ray Traced Shadows"), AdvancedDisplay)
+	TEnumAsByte<ECastRayTracedShadow::Type> CastRaytracedShadow;
 
 	/** Whether the light affects objects in reflections, when ray-traced reflection is enabled. */
 	UPROPERTY(EditAnywhere, BlueprintReadOnly, Category = Light, AdvancedDisplay, meta = (DisplayName = "Affect Ray Tracing Reflections"))
@@ -154,9 +157,12 @@ class ENGINE_API ULightComponentBase : public USceneComponent
 	UFUNCTION(BlueprintCallable, Category = "Rendering|Components|Light")
 	void SetAffectGlobalIllumination(bool bNewValue);
 
-	UFUNCTION(BlueprintCallable, Category = "Rendering|Components|Light")
+	UFUNCTION(BlueprintCallable, Category = "Rendering|Components|Light", meta = (DeprecatedFunction, DeprecationMessage = "ULightComponentBase::SetCastRaytracedShadow is deprecated. Use ULightComponentBase::SetCastRaytracedShadows instead."))
 	void SetCastRaytracedShadow(bool bNewValue);
 
+	UFUNCTION(BlueprintCallable, Category = "Rendering|Components|Light")
+	void SetCastRaytracedShadows(ECastRayTracedShadow::Type bNewValue);
+
 	UFUNCTION(BlueprintCallable, Category = "Rendering|Components|Light")
 	void SetSamplesPerPixel(int NewValue);
 
diff --git a/Engine/Source/Runtime/Engine/Classes/Engine/EngineTypes.h b/Engine/Source/Runtime/Engine/Classes/Engine/EngineTypes.h
--- a/Engine/Source/Runtime/Engine/Classes/Engine/EngineTypes.h
+++ b/Engine/Source/Runtime/Engine/Classes/Engine/EngineTypes.h
@@ -432,6 +432,21 @@ namespace EShadowMapMethod
 	};
 }
 
+/** Ray Tracing Shadows type. */
+UENUM()
+namespace ECastRayTracedShadow 
+{
+	enum Type
+	{
+		/** Ray traced shadows disabled for this light */
+		Disabled,
+		/** Ray traced shadows follow Cast Ray Traced Shadows project setting */
+		UseProjectSetting,
+		/** Ray traced shadows enabled for this light */
+		Enabled,
+	};
+}
+
 /** Specifies which component of the scene rendering should be output to the final render target. */
 UENUM()
 enum ESceneCaptureSource 
diff --git a/Engine/Source/Runtime/Engine/Classes/Engine/RendererSettings.h b/Engine/Source/Runtime/Engine/Classes/Engine/RendererSettings.h
--- a/Engine/Source/Runtime/Engine/Classes/Engine/RendererSettings.h
+++ b/Engine/Source/Runtime/Engine/Classes/Engine/RendererSettings.h
@@ -407,8 +407,8 @@ class ENGINE_API URendererSettings : public UDeveloperSettings
 		uint32 bEnableRayTracing : 1;
 
 	UPROPERTY(config, EditAnywhere, Category = HardwareRayTracing, meta = (
-		ConsoleVariable = "r.RayTracing.Shadows", DisplayName = "Ray Tracing Shadows",
-		ToolTip = "Enables hardware ray-traced shadows. Requires Hardware Ray Tracing to be enabled."))
+		ConsoleVariable = "r.RayTracing.Shadows", DisplayName = "Ray Traced Shadows",
+		ToolTip = "Controls whether Ray Traced Shadows are used by default. Lights can still override and force Ray Traced shadows on or off. Requires Hardware Ray Tracing to be enabled."))
 		uint32 bEnableRayTracingShadows : 1;
 
 	UPROPERTY(config, EditAnywhere, Category = HardwareRayTracing, meta = (
diff --git a/Engine/Source/Runtime/Engine/Private/Components/LightComponent.cpp b/Engine/Source/Runtime/Engine/Private/Components/LightComponent.cpp
--- a/Engine/Source/Runtime/Engine/Private/Components/LightComponent.cpp
+++ b/Engine/Source/Runtime/Engine/Private/Components/LightComponent.cpp
@@ -12,6 +12,7 @@
 #include "Materials/Material.h"
 #include "UObject/ObjectSaveContext.h"
 #include "UObject/RenderingObjectVersion.h"
+#include "UObject/UE5MainStreamObjectVersion.h"
 #include "UObject/UObjectHash.h"
 #include "UObject/UObjectIterator.h"
 #include "UObject/Package.h"
@@ -106,12 +107,30 @@ void ULightComponentBase::SetAffectGlobalIllumination(bool bNewValue)
 	}
 }
 
+// Deprecated
 void ULightComponentBase::SetCastRaytracedShadow(bool bNewValue)
+{
+	if (AreDynamicDataChangesAllowed())
+	{
+		if (bNewValue && CastRaytracedShadow == ECastRayTracedShadow::Disabled)
+		{
+			CastRaytracedShadow = ECastRayTracedShadow::UseProjectSetting;
+			MarkRenderStateDirty();
+		}
+		else if (!bNewValue && CastRaytracedShadow > ECastRayTracedShadow::Disabled)
+		{
+			CastRaytracedShadow = ECastRayTracedShadow::Disabled;
+			MarkRenderStateDirty();
+		}
+	}
+}
+
+void ULightComponentBase::SetCastRaytracedShadows(ECastRayTracedShadow::Type bNewValue)
 {
 	if (AreDynamicDataChangesAllowed()
-		&& bCastRaytracedShadow != bNewValue)
+		&& CastRaytracedShadow != bNewValue)
 	{
-		bCastRaytracedShadow = bNewValue;
+		CastRaytracedShadow = bNewValue;
 		MarkRenderStateDirty();
 	}
 }
@@ -134,6 +153,13 @@ void ULightComponentBase::Serialize(FArchive& Ar)
 	{
 		Intensity = Brightness_DEPRECATED;
 	}
+
+	Ar.UsingCustomVersion(FUE5MainStreamObjectVersion::GUID);
+
+	if (Ar.IsLoading() && (Ar.CustomVer(FUE5MainStreamObjectVersion::GUID) < FUE5MainStreamObjectVersion::RayTracedShadowsType))
+	{
+		CastRaytracedShadow = bCastRaytracedShadow_DEPRECATED == 0? ECastRayTracedShadow::Disabled : ECastRayTracedShadow::UseProjectSetting;
+	}
 }
 
 /**
@@ -310,7 +336,7 @@ FLightSceneProxy::FLightSceneProxy(const ULightComponent* InLightComponent)
 	, bCastHairStrandsDeepShadow(InLightComponent->bCastDeepShadow)
 	, bCastShadowsFromCinematicObjectsOnly(InLightComponent->bCastShadowsFromCinematicObjectsOnly)
 	, bForceCachedShadowsForMovablePrimitives(InLightComponent->bForceCachedShadowsForMovablePrimitives)
-	, bCastRaytracedShadow(InLightComponent->bCastRaytracedShadow)
+	, CastRaytracedShadow(InLightComponent->CastRaytracedShadow)
 	, bAffectReflection(InLightComponent->bAffectReflection)
 	, bAffectGlobalIllumination(InLightComponent->bAffectGlobalIllumination)
 	, bAffectTranslucentLighting(InLightComponent->bAffectTranslucentLighting)
@@ -432,7 +458,8 @@ ULightComponentBase::ULightComponentBase(const FObjectInitializer& ObjectInitial
 	CastShadows = true;
 	CastStaticShadows = true;
 	CastDynamicShadows = true;
-	bCastRaytracedShadow = true;
+	CastRaytracedShadow = ECastRayTracedShadow::UseProjectSetting;
+	bCastRaytracedShadow_DEPRECATED = true;
 	bAffectReflection = true;
 	bAffectGlobalIllumination = true;
 #if WITH_EDITORONLY_DATA
@@ -659,10 +686,9 @@ bool ULightComponent::CanEditChange(const FProperty* InProperty) const
 			return bCanEdit;
 		}
 
-		if (PropertyName == GET_MEMBER_NAME_STRING_CHECKED(ULightComponent, bCastRaytracedShadow))
+		if (PropertyName == GET_MEMBER_NAME_STRING_CHECKED(ULightComponent, CastRaytracedShadow))
 		{
-			static const auto CVar = IConsoleManager::Get().FindTConsoleVariableDataInt(TEXT("r.RayTracing.Shadows"));
-			return CVar->GetValueOnGameThread() != 0;
+			return IsRayTracingEnabled();
 		}
 
 		if (PropertyName == GET_MEMBER_NAME_STRING_CHECKED(ULightComponent, LightFunctionScale)
@@ -739,7 +765,7 @@ void ULightComponent::PostEditChangeProperty(FPropertyChangedEvent& PropertyChan
 		PropertyName != GET_MEMBER_NAME_STRING_CHECKED(ULightComponent, VolumetricScatteringIntensity) &&
 		PropertyName != GET_MEMBER_NAME_STRING_CHECKED(ULightComponent, bCastVolumetricShadow) &&
 		PropertyName != GET_MEMBER_NAME_STRING_CHECKED(ULightComponent, bCastDeepShadow) &&
-		PropertyName != GET_MEMBER_NAME_STRING_CHECKED(ULightComponent, bCastRaytracedShadow) &&
+		PropertyName != GET_MEMBER_NAME_STRING_CHECKED(ULightComponent, CastRaytracedShadow) &&
 		PropertyName != GET_MEMBER_NAME_STRING_CHECKED(ULightComponent, bAffectReflection) &&
 		PropertyName != GET_MEMBER_NAME_STRING_CHECKED(ULightComponent, bAffectGlobalIllumination) &&
 		// Point light properties that shouldn't unbuild lighting
diff --git a/Engine/Source/Runtime/Engine/Private/Components/SkyLightComponent.cpp b/Engine/Source/Runtime/Engine/Private/Components/SkyLightComponent.cpp
--- a/Engine/Source/Runtime/Engine/Private/Components/SkyLightComponent.cpp
+++ b/Engine/Source/Runtime/Engine/Private/Components/SkyLightComponent.cpp
@@ -20,6 +20,7 @@
 #include "ShaderCompiler.h"
 #include "Components/BillboardComponent.h"
 #include "UObject/ReleaseObjectVersion.h"
+#include "UObject/UE5MainStreamObjectVersion.h"
 #include "Modules/ModuleManager.h"
 #include "Internationalization/Text.h"
 #include "CoreGlobals.h"
@@ -197,7 +198,7 @@ FSkyLightSceneProxy::FSkyLightSceneProxy(const USkyLightComponent* InLightCompon
 	, bWantsStaticShadowing(InLightComponent->Mobility == EComponentMobility::Stationary)
 	, bHasStaticLighting(InLightComponent->HasStaticLighting())
 	, bCastVolumetricShadow(InLightComponent->bCastVolumetricShadow)
-	, bCastRayTracedShadow(InLightComponent->bCastRaytracedShadow)
+	, CastRayTracedShadow(InLightComponent->CastRaytracedShadow)
 	, bAffectReflection(InLightComponent->bAffectReflection)
 	, bAffectGlobalIllumination(InLightComponent->bAffectGlobalIllumination)
 	, bTransmission(InLightComponent->bTransmission)
@@ -275,7 +276,7 @@ USkyLightComponent::USkyLightComponent(const FObjectInitializer& ObjectInitializ
 	AverageBrightness = 1.0f;
 	BlendDestinationAverageBrightness = 1.0f;
 	bCastVolumetricShadow = true;
-	bCastRaytracedShadow = false;
+	CastRaytracedShadow = ECastRayTracedShadow::UseProjectSetting;
 	bAffectReflection = true;
 	bAffectGlobalIllumination = true;
 	SamplesPerPixel = 4;
@@ -1088,6 +1089,7 @@ void USkyLightComponent::RecaptureSky()
 void USkyLightComponent::Serialize(FArchive& Ar)
 {
 	Ar.UsingCustomVersion(FReleaseObjectVersion::GUID);
+	Ar.UsingCustomVersion(FUE5MainStreamObjectVersion::GUID);
 
 	Super::Serialize(Ar);
 
@@ -1097,6 +1099,12 @@ void USkyLightComponent::Serialize(FArchive& Ar)
 		FSHVectorRGB3 DummyIrradianceEnvironmentMap;
 		Ar << DummyIrradianceEnvironmentMap;
 	}
+
+	if (Ar.IsLoading() && (Ar.CustomVer(FUE5MainStreamObjectVersion::GUID) < FUE5MainStreamObjectVersion::RayTracedShadowsType))
+	{
+		// Skylight is not driven by r.raytracing.shadows but if it was enabled before, it should be enabled now independently of the project setting
+		CastRaytracedShadow = bCastRaytracedShadow_DEPRECATED == 0 ? ECastRayTracedShadow::Disabled : ECastRayTracedShadow::Enabled;
+	}
 }
 
 
diff --git a/Engine/Source/Runtime/Engine/Public/SceneManagement.h b/Engine/Source/Runtime/Engine/Public/SceneManagement.h
--- a/Engine/Source/Runtime/Engine/Public/SceneManagement.h
+++ b/Engine/Source/Runtime/Engine/Public/SceneManagement.h
@@ -1144,7 +1144,7 @@ public:
 	uint8 bWantsStaticShadowing:1;
 	uint8 bHasStaticLighting:1;
 	uint8 bCastVolumetricShadow:1;
-	uint8 bCastRayTracedShadow:1;
+	TEnumAsByte<ECastRayTracedShadow::Type> CastRayTracedShadow;
 	uint8 bAffectReflection:1;
 	uint8 bAffectGlobalIllumination:1;
 	uint8 bTransmission:1;
@@ -1506,7 +1506,7 @@ public:
 	inline bool CastsTranslucentShadows() const { return bCastTranslucentShadows; }
 	inline bool CastsVolumetricShadow() const { return bCastVolumetricShadow; }
 	inline bool CastsHairStrandsDeepShadow() const { return bCastHairStrandsDeepShadow; }
-	inline bool CastsRaytracedShadow() const { return bCastRaytracedShadow; }
+	inline TEnumAsByte<ECastRayTracedShadow::Type> CastsRaytracedShadow() const { return CastRaytracedShadow; }
 	inline bool AffectReflection() const { return bAffectReflection; }
 	inline bool AffectGlobalIllumination() const { return bAffectGlobalIllumination; }
 	inline bool CastsShadowsFromCinematicObjectsOnly() const { return bCastShadowsFromCinematicObjectsOnly; }
@@ -1707,7 +1707,7 @@ protected:
 	const uint8 bForceCachedShadowsForMovablePrimitives : 1;
 
 	/** Whether the light shadows are computed with shadow-mapping or ray-tracing (when available). */
-	const uint8 bCastRaytracedShadow : 1;
+	const TEnumAsByte<ECastRayTracedShadow::Type> CastRaytracedShadow;
 
 	/** Whether the light affects objects in reflections, when ray-traced reflection is enabled. */
 	const uint8 bAffectReflection : 1;
diff --git a/Engine/Source/Runtime/Renderer/Private/DeferredShadingRenderer.cpp b/Engine/Source/Runtime/Renderer/Private/DeferredShadingRenderer.cpp
--- a/Engine/Source/Runtime/Renderer/Private/DeferredShadingRenderer.cpp
+++ b/Engine/Source/Runtime/Renderer/Private/DeferredShadingRenderer.cpp
@@ -3256,6 +3256,7 @@ bool AnyRayTracingPassEnabled(const FScene* Scene, const FViewInfo& View)
 		|| ShouldRenderRayTracingTranslucency(View)
 		|| ShouldRenderRayTracingSkyLight(Scene ? Scene->SkyLight : nullptr)
 		|| ShouldRenderRayTracingShadows()
+		|| Scene->RayTracedLights.Num() > 0
         || Lumen::AnyLumenHardwareRayTracingPassEnabled(Scene, View)
 		|| HasRayTracedOverlay(*View.Family);
 }
diff --git a/Engine/Source/Runtime/Renderer/Private/LightRendering.cpp b/Engine/Source/Runtime/Renderer/Private/LightRendering.cpp
--- a/Engine/Source/Runtime/Renderer/Private/LightRendering.cpp
+++ b/Engine/Source/Runtime/Renderer/Private/LightRendering.cpp
@@ -163,13 +163,13 @@ bool ShouldRenderRayTracingShadows()
 
 bool ShouldRenderRayTracingShadowsForLight(const FLightSceneProxy& LightProxy)
 {
-	return ShouldRenderRayTracingShadows() && LightProxy.CastsRaytracedShadow()
+	return (LightProxy.CastsRaytracedShadow() == ECastRayTracedShadow::Enabled || (ShouldRenderRayTracingShadows() && LightProxy.CastsRaytracedShadow() == ECastRayTracedShadow::UseProjectSetting))
 		&& ShouldRenderRayTracingShadowsForLightType((ELightComponentType)LightProxy.GetLightType());
 }
 
 bool ShouldRenderRayTracingShadowsForLight(const FLightSceneInfoCompact& LightInfo)
 {
-	return ShouldRenderRayTracingShadows() && LightInfo.bCastRaytracedShadow
+	return (LightInfo.CastRaytracedShadow == ECastRayTracedShadow::Enabled || (ShouldRenderRayTracingShadows() && LightInfo.CastRaytracedShadow == ECastRayTracedShadow::UseProjectSetting))
 		&& ShouldRenderRayTracingShadowsForLightType((ELightComponentType)LightInfo.LightType);
 }
 #endif // RHI_RAYTRACING
diff --git a/Engine/Source/Runtime/Renderer/Private/LightSceneInfo.cpp b/Engine/Source/Runtime/Renderer/Private/LightSceneInfo.cpp
--- a/Engine/Source/Runtime/Renderer/Private/LightSceneInfo.cpp
+++ b/Engine/Source/Runtime/Renderer/Private/LightSceneInfo.cpp
@@ -39,7 +39,7 @@ void FLightSceneInfoCompact::Init(FLightSceneInfo* InLightSceneInfo)
 	bStaticLighting = InLightSceneInfo->Proxy->HasStaticLighting();
 	bAffectReflection = InLightSceneInfo->Proxy->AffectReflection();
 	bAffectGlobalIllumination = InLightSceneInfo->Proxy->AffectGlobalIllumination();
-	bCastRaytracedShadow = InLightSceneInfo->Proxy->CastsRaytracedShadow();
+	CastRaytracedShadow = InLightSceneInfo->Proxy->CastsRaytracedShadow();
 }
 
 FLightSceneInfo::FLightSceneInfo(FLightSceneProxy* InProxy, bool InbVisible)
diff --git a/Engine/Source/Runtime/Renderer/Private/LightSceneInfo.h b/Engine/Source/Runtime/Renderer/Private/LightSceneInfo.h
--- a/Engine/Source/Runtime/Renderer/Private/LightSceneInfo.h
+++ b/Engine/Source/Runtime/Renderer/Private/LightSceneInfo.h
@@ -38,7 +38,7 @@ public:
 	uint32 bStaticLighting : 1;
 	uint32 bAffectReflection : 1;
 	uint32 bAffectGlobalIllumination : 1;
-	uint32 bCastRaytracedShadow : 1;
+	TEnumAsByte<ECastRayTracedShadow::Type> CastRaytracedShadow;
 
 	/** Initializes the compact scene info from the light's full scene info. */
 	void Init(FLightSceneInfo* InLightSceneInfo);
diff --git a/Engine/Source/Runtime/Renderer/Private/RayTracing/RayTracingShadows.cpp b/Engine/Source/Runtime/Renderer/Private/RayTracing/RayTracingShadows.cpp
--- a/Engine/Source/Runtime/Renderer/Private/RayTracing/RayTracingShadows.cpp
+++ b/Engine/Source/Runtime/Renderer/Private/RayTracing/RayTracingShadows.cpp
@@ -178,10 +178,9 @@ float GetRaytracingMaxNormalBias()
 void FDeferredShadingSceneRenderer::PrepareRayTracingShadows(const FViewInfo& View, TArray<FRHIRayTracingShader*>& OutRayGenShaders)
 {
 	// Declare all RayGen shaders that require material closest hit shaders to be bound
+	// Ray tracing shadows shaders should be properly configured even if r.RayTracing.Shadows is 0 because lights can have raytracing shadows enabled independently of that CVar
 
-	static auto CVarRayTracingShadows = IConsoleManager::Get().FindConsoleVariable(TEXT("r.RayTracing.Shadows"));
-
-	const bool bRayTracingShadows = ShouldRenderRayTracingEffect(CVarRayTracingShadows != nullptr && CVarRayTracingShadows->GetInt() > 0);
+	const bool bRayTracingShadows = ShouldRenderRayTracingEffect(true);
 
 	if (!bRayTracingShadows)
 	{
diff --git a/Engine/Source/Runtime/Renderer/Private/RayTracing/RaytracingSkylight.cpp b/Engine/Source/Runtime/Renderer/Private/RayTracing/RaytracingSkylight.cpp
--- a/Engine/Source/Runtime/Renderer/Private/RayTracing/RaytracingSkylight.cpp
+++ b/Engine/Source/Runtime/Renderer/Private/RayTracing/RaytracingSkylight.cpp
@@ -134,8 +134,9 @@ bool ShouldRenderRayTracingSkyLight(const FSkyLightSceneProxy* SkyLightSceneProx
 		return false;
 	}
 
+	// SkyLightSceneProxy has a separate r.raytracing.skylight Cvar and does not check r.raytracing.shadows when CastRayTracedShadow == ECastRayTracedShadow::UseProjectSetting
 	bool bRayTracingSkyEnabled = GRayTracingSkyLight  < 0 
-		? SkyLightSceneProxy->bCastRayTracedShadow
+		? SkyLightSceneProxy->CastRayTracedShadow == ECastRayTracedShadow::Enabled
 		: GRayTracingSkyLight != 0;
 
 	bRayTracingSkyEnabled = bRayTracingSkyEnabled && (GetSkyLightSamplesPerPixel(SkyLightSceneProxy) > 0);
diff --git a/Engine/Source/Runtime/Renderer/Private/RendererScene.cpp b/Engine/Source/Runtime/Renderer/Private/RendererScene.cpp
--- a/Engine/Source/Runtime/Renderer/Private/RendererScene.cpp
+++ b/Engine/Source/Runtime/Renderer/Private/RendererScene.cpp
@@ -1735,6 +1735,13 @@ void FScene::AddLightSceneInfo_RenderThread(FLightSceneInfo* LightSceneInfo)
 
 	ProcessAtmosphereLightAddition_RenderThread(LightSceneInfo);
 
+#if RHI_RAYTRACING
+	if (ShouldRenderRayTracingShadowsForLight(LightSceneInfo))
+	{
+		RayTracedLights.Add(LightSceneInfo);
+	}
+#endif 
+
 	// Add the light to the scene.
 	LightSceneInfo->AddToScene();
 }
@@ -2951,6 +2958,13 @@ void FScene::RemoveLightSceneInfo_RenderThread(FLightSceneInfo* LightSceneInfo)
 
 		ProcessAtmosphereLightRemoval_RenderThread(LightSceneInfo);
 
+
+#if RHI_RAYTRACING
+		if (ShouldRenderRayTracingShadowsForLight(LightSceneInfo))
+		{
+			RayTracedLights.Remove(LightSceneInfo);
+		}
+#endif 
 		// Remove the light from the scene.
 		LightSceneInfo->RemoveFromScene();
 
diff --git a/Engine/Source/Runtime/Renderer/Private/ScenePrivate.h b/Engine/Source/Runtime/Renderer/Private/ScenePrivate.h
--- a/Engine/Source/Runtime/Renderer/Private/ScenePrivate.h
+++ b/Engine/Source/Runtime/Renderer/Private/ScenePrivate.h
@@ -2811,6 +2811,7 @@ public:
 #if RHI_RAYTRACING
 	/** Persistently-allocated ray tracing scene data. */
 	FRayTracingScene RayTracingScene;
+	TArray<FLightSceneInfo*, TInlineAllocator<4>> RayTracedLights;
 #endif // RHI_RAYTRACING
 
 	/** Distance field object scene data. */
