Add blueprint accessors for deprecated InputScale properties on the player controller in case old projects rely on them.

diff --git a/Engine/Source/Runtime/Engine/Classes/GameFramework/PlayerController.h b/Engine/Source/Runtime/Engine/Classes/GameFramework/PlayerController.h
--- a/Engine/Source/Runtime/Engine/Classes/GameFramework/PlayerController.h
+++ b/Engine/Source/Runtime/Engine/Classes/GameFramework/PlayerController.h
@@ -449,6 +449,36 @@ public:
 	UPROPERTY(config, meta = (DeprecatedProperty, DeprecationMessage = "Use the Enhanced Input plugin Scalar Modifier instead. See UInputSettings::bEnableLegacyInputScales to enable legacy behavior"))
 	float InputRollScale_DEPRECATED = 1.0f;
 
+	/** A getter for the deprecated InputYawScale property. This should only be used if UInputSettings::bEnableLegacyInputScales is turned on. */
+	UE_DEPRECATED(5.0, "GetDeprecatedInputYawScale is deprecated, please use the Enhanced Input plugin Scalar Modifier instead.")
+	UFUNCTION(BlueprintCallable, Category = PlayerController)
+	float GetDeprecatedInputYawScale() const;
+	
+	/** A getter for the deprecated InputPitchScale property. This should only be used if UInputSettings::bEnableLegacyInputScales is turned on. */
+	UE_DEPRECATED(5.0, "GetDeprecatedInputPitchScale is deprecated, please use the Enhanced Input plugin Scalar Modifier instead.")
+	UFUNCTION(BlueprintCallable, Category = PlayerController)
+	float GetDeprecatedInputPitchScale() const;
+	
+	/** A getter for the deprecated InputRollScale property. This should only be used if UInputSettings::bEnableLegacyInputScales is turned on. */
+	UE_DEPRECATED(5.0, "GetDeprecatedInputRollScale is deprecated, please use the Enhanced Input plugin Scalar Modifier instead.)")
+	UFUNCTION(BlueprintCallable, Category = PlayerController)
+	float GetDeprecatedInputRollScale() const;
+
+	/** A getter for the deprecated InputYawScale property. This should only be used if UInputSettings::bEnableLegacyInputScales is turned on. */
+	UE_DEPRECATED(5.0, "SetDeprecatedInputYawScale is deprecated, please use the Enhanced Input plugin Scalar Modifier instead.")
+	UFUNCTION(BlueprintCallable, Category = PlayerController)
+	void SetDeprecatedInputYawScale(float NewValue);
+
+	/** A getter for the deprecated InputPitchScale property. This should only be used if UInputSettings::bEnableLegacyInputScales is turned on. */
+	UE_DEPRECATED(5.0, "SetDeprecatedInputPitchScale is deprecated, please use the Enhanced Input plugin Scalar Modifier instead.")
+	UFUNCTION(BlueprintCallable, Category = PlayerController)
+	void SetDeprecatedInputPitchScale(float NewValue);
+
+	/** A getter for the deprecated InputRollScale property. This should only be used if UInputSettings::bEnableLegacyInputScales is turned on. */
+	UE_DEPRECATED(5.0, "SetDeprecatedInputRollScale is deprecated, please use the Enhanced Input plugin Scalar Modifier instead.)")
+	UFUNCTION(BlueprintCallable, Category = PlayerController)
+	void SetDeprecatedInputRollScale(float NewValue);
+
 	/** Whether the mouse cursor should be displayed. */
 	UPROPERTY(EditAnywhere, BlueprintReadWrite, Category=MouseInterface)
 	uint32 bShowMouseCursor:1;
diff --git a/Engine/Source/Runtime/Engine/Private/PlayerController.cpp b/Engine/Source/Runtime/Engine/Private/PlayerController.cpp
--- a/Engine/Source/Runtime/Engine/Private/PlayerController.cpp
+++ b/Engine/Source/Runtime/Engine/Private/PlayerController.cpp
@@ -2889,6 +2889,81 @@ void APlayerController::ClientSetSpectatorWaiting_Implementation(bool bWaiting)
 	}
 }
 
+float APlayerController::GetDeprecatedInputYawScale() const
+{
+	if (GetDefault<UInputSettings>()->bEnableLegacyInputScales)
+	{
+		return InputYawScale_DEPRECATED;
+	}
+	else
+	{
+		UE_LOG(LogPlayerController, Warning, TEXT("Attempting to access legacy input scales without the setting enabled! See UInputSettings::bEnableLegacyInputScales."));
+		return 1.0f;
+	}
+}
+
+float APlayerController::GetDeprecatedInputPitchScale() const
+{
+	if (GetDefault<UInputSettings>()->bEnableLegacyInputScales)
+	{
+		return InputPitchScale_DEPRECATED;
+	}
+	else
+	{
+		UE_LOG(LogPlayerController, Warning, TEXT("Attempting to access legacy input scales without the setting enabled! See UInputSettings::bEnableLegacyInputScales."));
+		return 1.0f;
+	}
+}
+
+float APlayerController::GetDeprecatedInputRollScale() const
+{
+	if (GetDefault<UInputSettings>()->bEnableLegacyInputScales)
+	{
+		return InputRollScale_DEPRECATED;
+	}
+	else
+	{
+		UE_LOG(LogPlayerController, Warning, TEXT("Attempting to access legacy input scales without the setting enabled! See UInputSettings::bEnableLegacyInputScales."));
+		return 1.0f;
+	}
+}
+
+void APlayerController::SetDeprecatedInputYawScale(float NewValue)
+{
+	if (GetDefault<UInputSettings>()->bEnableLegacyInputScales)
+	{
+		InputYawScale_DEPRECATED = NewValue;
+	}
+	else
+	{
+		UE_LOG(LogPlayerController, Warning, TEXT("Attempting to access legacy input scales without the setting enabled! See UInputSettings::bEnableLegacyInputScales."));
+	}
+}
+
+void APlayerController::SetDeprecatedInputPitchScale(float NewValue)
+{
+	if (GetDefault<UInputSettings>()->bEnableLegacyInputScales)
+	{
+		InputPitchScale_DEPRECATED = NewValue;
+	}
+	else
+	{
+		UE_LOG(LogPlayerController, Warning, TEXT("Attempting to access legacy input scales without the setting enabled! See UInputSettings::bEnableLegacyInputScales."));
+	}
+}
+
+void APlayerController::SetDeprecatedInputRollScale(float NewValue)
+{
+	if (GetDefault<UInputSettings>()->bEnableLegacyInputScales)
+	{
+		InputRollScale_DEPRECATED = NewValue;
+	}
+	else
+	{
+		UE_LOG(LogPlayerController, Warning, TEXT("Attempting to access legacy input scales without the setting enabled! See UInputSettings::bEnableLegacyInputScales."));
+	}
+}
+
 bool APlayerController::ServerViewNextPlayer_Validate()
 {
 	return true;
