GitHub #8215 : Fix check in TPlayerControllerIterator.

diff --git a/Engine/Source/Runtime/Engine/Public/UnrealEngine.h b/Engine/Source/Runtime/Engine/Public/UnrealEngine.h
--- a/Engine/Source/Runtime/Engine/Public/UnrealEngine.h
+++ b/Engine/Source/Runtime/Engine/Public/UnrealEngine.h
@@ -104,7 +104,7 @@ public:
 	TBasePlayerControllerIterator(class UWorld* InWorld)
 		: Iter(InWorld->GetPlayerControllerIterator())
 	{
-		check(!LocalOnly || InWorld->GetNetMode() != NM_Client);	// You should only iterate on non local player controllers if you are the server
+		check(LocalOnly || InWorld->GetNetMode() != NM_Client);	// You should only iterate on non local player controllers if you are the server
 		AdvanceCurrent();
 	}
 
