Shader symbol terminology changes
- Splitting symbol generation and shader "debug info" generation into two different cvars. This should now allow us to generate symbols post-packaging without generating different shader variants. This should also enable us to always generate symbols and still have the final deduplicated shader sizes.
- r.Shaders.Symbols is now the primary way to enable shader symbols for debugging.
- r.Shaders.GenerateSymbols can be enabled to always generate symbols without writing them to disk.
- r.Shaders.WriteSymbols can be used to force write symbols if they were previously generated.
- r.Shaders.ExtraData can be used to generate shader names for engine runtime systems to use.
- NotifyShaderCooked replaced with NotifyShaderCompiled, prep for upcoming non-cook symbol writing.
- PC DXC shader compiles should only generate debug info when requested, not all the time.
- PC FXC shader compiles should generate debug info when requested, not just when r.Shaders.Optimize is 0.

diff --git a/Engine/Config/BaseEngine.ini b/Engine/Config/BaseEngine.ini
--- a/Engine/Config/BaseEngine.ini
+++ b/Engine/Config/BaseEngine.ini
@@ -1541,13 +1541,34 @@ bUseVirtualCores = False
 bForceUseSCWMemoryPressureLimits = False
 
 [ShaderCompiler]
-; Shader compiler debug info cvars that can be overridden per platform.
-; Build shader debug info and export any external debug info.
-r.Shaders.KeepDebugInfo=0
-; Build, but do not export external debug info.
-r.Shaders.PrepareExportedDebugInfo=0
-; Export external debug info if it is built.
-r.Shaders.ExportDebugInfo=0
+; Shader symbol cvars that can be overridden per platform.
+; Generate and write shader symbols
+r.Shaders.Symbols=0
+; Generate shader symbols but don't write to disk
+r.Shaders.GenerateSymbols=0
+; Write generated shader symbols to disk
+r.Shaders.WriteSymbols=0
+; Write generated symbols to a single zip file
+r.Shaders.WriteSymbols.Zip=0
+; Override the path that symbols are written to.
+r.Shaders.SymbolPathOverride=""
+; Base shader symbols off of source and not generated shaders
+r.Shaders.AllowUniqueSymbols=0
+
+[ShaderCompiler_BuildMachine]
+; Build machine specific shader symbol cvars that can be overridden per platform.
+; Generate and write shader symbols
+r.Shaders.Symbols=0
+; Generate shader symbols but don't write to disk
+r.Shaders.GenerateSymbols=0
+; Write generated shader symbols to disk
+r.Shaders.WriteSymbols=0
+; Write generated symbols to a single zip file
+r.Shaders.WriteSymbols.Zip=0
+; Override the path that symbols are written to.
+r.Shaders.SymbolPathOverride=""
+; Base shader symbols off of source and not generated shaders
+r.Shaders.AllowUniqueSymbols=0
 
 [LogFiles]
 PurgeLogsDays=5
@@ -2571,10 +2592,11 @@ DefaultASTCQualityBySize=3
 +VersionedIntRValues=r.BasePassOutputsVelocity
 +VersionedIntRValues=r.SelectiveBasePassOutputs
 +VersionedIntRValues=r.DBuffer
-+VersionedIntRValues=r.Shaders.KeepDebugInfo
-+VersionedIntRValues=r.Shaders.ExportDebugInfo
-+VersionedIntRValues=r.Shaders.AllowUniqueDebugInfo
-+VersionedIntRValues=r.Shaders.PrepareExportedDebugInfo
++VersionedIntRValues=r.Shaders.Symbols
++VersionedIntRValues=r.Shaders.GenerateSymbols
++VersionedIntRValues=r.Shaders.WriteSymbols
++VersionedIntRValues=r.Shaders.AllowUniqueSymbols
++VersionedIntRValues=r.Shaders.ExtraData
 +VersionedIntRValues=r.Shaders.Optimize
 +VersionedIntRValues=r.CompileShadersForDevelopment
 +VersionedIntRValues=r.MobileHDR
diff --git a/Engine/Config/ConsoleVariables.ini b/Engine/Config/ConsoleVariables.ini
--- a/Engine/Config/ConsoleVariables.ini
+++ b/Engine/Config/ConsoleVariables.ini
@@ -71,7 +71,7 @@ r.DumpShaderDebugInfo=2
 ; Uncomment when running with a graphical debugger (but not when profiling)
 ;r.Shaders.Optimize=0
 ; When this is enabled, shaders will have extra debugging info. This could change patch sizes, uniqueness, etc and will recompile the shaders
-;r.Shaders.KeepDebugInfo=1
+;r.Shaders.Symbols=1
 ; Uncomment to skip shader compression. Can save a significant time when using debug shaders.
 ;r.Shaders.SkipCompression=1
 ; Uncomment and change to set a custom folder to write shader symbols to when they are generated.
diff --git a/Engine/Shaders/Public/ShaderVersion.ush b/Engine/Shaders/Public/ShaderVersion.ush
--- a/Engine/Shaders/Public/ShaderVersion.ush
+++ b/Engine/Shaders/Public/ShaderVersion.ush
@@ -3,4 +3,4 @@
 // in Platform.ush (which should be included in any shader) it allows to invalidate the shader DDC.
 // 
 // If you are merging streams and there is a conflict with this GUID you should make a new GUID rather than taking one or the other.
-// GUID = DCA69B23FEF84CBDB4BED51526ECE280
+// GUID = 04937C71BC6C4F878D4CA02CC386BF6E
diff --git a/Engine/Source/Developer/Apple/MetalShaderFormat/Private/MetalShaderCompiler.cpp b/Engine/Source/Developer/Apple/MetalShaderFormat/Private/MetalShaderCompiler.cpp
--- a/Engine/Source/Developer/Apple/MetalShaderFormat/Private/MetalShaderCompiler.cpp
+++ b/Engine/Source/Developer/Apple/MetalShaderFormat/Private/MetalShaderCompiler.cpp
@@ -258,7 +258,7 @@ void BuildMetalShaderOutput(
 	FMetalCodeHeader Header;
 	Header.CompileFlags = (ShaderInput.Environment.CompilerFlags.Contains(CFLAG_Debug) ? (1 << CFLAG_Debug) : 0);
 	Header.CompileFlags |= (bNoFastMath ? (1 << CFLAG_NoFastMath) : 0);
-	Header.CompileFlags |= (ShaderInput.Environment.CompilerFlags.Contains(CFLAG_KeepDebugInfo) ? (1 << CFLAG_KeepDebugInfo) : 0);
+	Header.CompileFlags |= (ShaderInput.Environment.CompilerFlags.Contains(CFLAG_ExtraShaderData) ? (1 << CFLAG_ExtraShaderData) : 0);
 	Header.CompileFlags |= (ShaderInput.Environment.CompilerFlags.Contains(CFLAG_ZeroInitialise) ? (1 <<  CFLAG_ZeroInitialise) : 0);
 	Header.CompileFlags |= (ShaderInput.Environment.CompilerFlags.Contains(CFLAG_BoundsChecking) ? (1 << CFLAG_BoundsChecking) : 0);
 	Header.CompileFlags |= (ShaderInput.Environment.CompilerFlags.Contains(CFLAG_Archive) ? (1 << CFLAG_Archive) : 0);
@@ -533,7 +533,7 @@ void BuildMetalShaderOutput(
 	}
 
 	FString MetalCode = FString(USFSource);
-	if (ShaderInput.Environment.CompilerFlags.Contains(CFLAG_KeepDebugInfo) || ShaderInput.Environment.CompilerFlags.Contains(CFLAG_Debug))
+	if (ShaderInput.Environment.CompilerFlags.Contains(CFLAG_ExtraShaderData) || ShaderInput.Environment.CompilerFlags.Contains(CFLAG_Debug))
 	{
 		MetalCode.InsertAt(0, FString::Printf(TEXT("// %s\n"), *CCHeader.Name));
 		Header.ShaderName = CCHeader.Name;
@@ -581,7 +581,7 @@ void BuildMetalShaderOutput(
 	{
 		// TODO technically should probably check the version of the metal compiler to make sure it's recent enough to support -MO.
         FString DebugInfo = TEXT("");
-		if (ShaderInput.Environment.CompilerFlags.Contains(CFLAG_KeepDebugInfo))
+		if (ShaderInput.Environment.CompilerFlags.Contains(CFLAG_ExtraShaderData))
 		{
 			DebugInfo = TEXT("-gline-tables-only -MO");
 		}
@@ -774,7 +774,7 @@ void BuildMetalShaderOutput(
 				ShaderOutput.ShaderCode.AddOptionalData('u', (const uint8*)&DebugCode.UncompressedSize, sizeof(DebugCode.UncompressedSize));
 			}
 			
-			if (ShaderInput.Environment.CompilerFlags.Contains(CFLAG_KeepDebugInfo))
+			if (ShaderInput.Environment.CompilerFlags.Contains(CFLAG_ExtraShaderData))
 			{
 				// store data we can pickup later with ShaderCode.FindOptionalData(FShaderCodeName::Key), could be removed for shipping
 				ShaderOutput.ShaderCode.AddOptionalData(FShaderCodeName::Key, TCHAR_TO_UTF8(*ShaderInput.GenerateShaderName()));
diff --git a/Engine/Source/Developer/ShaderCompilerCommon/Public/ShaderSymbolExport.h b/Engine/Source/Developer/ShaderCompilerCommon/Public/ShaderSymbolExport.h
--- a/Engine/Source/Developer/ShaderCompilerCommon/Public/ShaderSymbolExport.h
+++ b/Engine/Source/Developer/ShaderCompilerCommon/Public/ShaderSymbolExport.h
@@ -48,7 +48,7 @@ inline FShaderSymbolExport::FShaderSymbolExport(FName InShaderFormat)
 
 inline void FShaderSymbolExport::Initialize()
 {
-	const bool bSymbolsEnabled = ShouldExportShaderDebugInfo(ShaderFormat);
+	const bool bSymbolsEnabled = ShouldWriteShaderSymbols(ShaderFormat);
 	const bool bForceSymbols = FParse::Value(FCommandLine::Get(), TEXT("-ShaderSymbolsExport="), ExportPath);
 
 	if (bSymbolsEnabled || bForceSymbols)
@@ -56,7 +56,7 @@ inline void FShaderSymbolExport::Initialize()
 		// if no command line path is provided, look to the cvar first
 		if (ExportPath.IsEmpty())
 		{
-			if (GetShaderDebugInfoPathOverride(ExportPath, ShaderFormat))
+			if (GetShaderSymbolPathOverride(ExportPath, ShaderFormat))
 			{
 				ExportPath = IFileManager::Get().ConvertToAbsolutePathForExternalAppForWrite(*ExportPath);
 			}
@@ -79,7 +79,7 @@ inline void FShaderSymbolExport::Initialize()
 		else
 		{
 			// Check if the export mode is to an uncompressed archive or loose files.
-			bool bExportAsZip = ShouldExportShaderDebugInfoAsZip(ShaderFormat);
+			bool bExportAsZip = ShouldWriteShaderSymbolsAsZip(ShaderFormat);
 			if (bExportAsZip || FParse::Param(FCommandLine::Get(), TEXT("ShaderSymbolsExportZip")))
 			{
 				FString SingleFilePath = ExportPath / TEXT("ShaderSymbols.zip");
diff --git a/Engine/Source/Developer/ShaderFormatOpenGL/Private/OpenGLShaderCompiler.cpp b/Engine/Source/Developer/ShaderFormatOpenGL/Private/OpenGLShaderCompiler.cpp
--- a/Engine/Source/Developer/ShaderFormatOpenGL/Private/OpenGLShaderCompiler.cpp
+++ b/Engine/Source/Developer/ShaderFormatOpenGL/Private/OpenGLShaderCompiler.cpp
@@ -995,7 +995,7 @@ void FOpenGLFrontend::BuildShaderOutput(
 			ShaderOutput.OptionalFinalShaderSource = FString(GlslCodeOriginal.GetData());
 		}
 
-		if (ShaderInput.Environment.CompilerFlags.Contains(CFLAG_KeepDebugInfo))
+		if (ShaderInput.Environment.CompilerFlags.Contains(CFLAG_ExtraShaderData))
 		{
 			ShaderOutput.ShaderCode.AddOptionalData(FShaderCodeName::Key, TCHAR_TO_UTF8(*ShaderInput.GenerateShaderName()));
 		}
diff --git a/Engine/Source/Developer/TargetPlatform/Public/Interfaces/IShaderFormat.h b/Engine/Source/Developer/TargetPlatform/Public/Interfaces/IShaderFormat.h
--- a/Engine/Source/Developer/TargetPlatform/Public/Interfaces/IShaderFormat.h
+++ b/Engine/Source/Developer/TargetPlatform/Public/Interfaces/IShaderFormat.h
@@ -111,7 +111,14 @@ public:
 	 * Called when a shader resource is cooked, so the shader format can perform platform-specific operations on the debug data.
 	 * Does nothing on platforms that make no use of the platform debug data.
 	 */
-	virtual void NotifyShaderCooked(const TConstArrayView<uint8>& PlatformDebugData, FName Format) const { }
+	virtual void NotifyShaderCompiled(const TConstArrayView<uint8>& PlatformDebugData, FName Format) const { }
+
+	UE_DEPRECATED(5.0, "NotifyShaderCompiled should be called")
+	virtual void NotifyShaderCooked(const TConstArrayView<uint8>& PlatformDebugData, FName Format) const
+	{
+		NotifyShaderCompiled(PlatformDebugData, Format);
+	}
+
 
 	/**
 	 * Appends shader key text to the provided key string for use in DDC unique key construction.
diff --git a/Engine/Source/Developer/VulkanShaderFormat/Private/VulkanShaderCompiler.cpp b/Engine/Source/Developer/VulkanShaderFormat/Private/VulkanShaderCompiler.cpp
--- a/Engine/Source/Developer/VulkanShaderFormat/Private/VulkanShaderCompiler.cpp
+++ b/Engine/Source/Developer/VulkanShaderFormat/Private/VulkanShaderCompiler.cpp
@@ -1488,7 +1488,7 @@ static void BuildShaderOutput(
 	FVulkanShaderHeader NEWHeader(FVulkanShaderHeader::EZero);
 	ConvertToNEWHeader(OLDHeader, BindingTable, Spirv, NEWEntryTypes, ShaderInput, CCHeader, ShaderOutput.ParameterMap, NEWHeader);
 
-	if (ShaderInput.Environment.CompilerFlags.Contains(CFLAG_KeepDebugInfo))
+	if (ShaderInput.Environment.CompilerFlags.Contains(CFLAG_ExtraShaderData))
 	{
 		NEWHeader.DebugName = ShaderInput.GenerateShaderName();
 	}
@@ -2352,7 +2352,7 @@ static bool CompileWithShaderConductor(
 	// Build shader output and binding table
 	Output.bSucceeded = BuildShaderOutputFromSpirv(CompilerContext, Spirv, Input, Output, BindingTable, EntryPointName, bStripReflect, bIsRayTracingShader, bDebugDump);
 
-	if (Input.Environment.CompilerFlags.Contains(CFLAG_KeepDebugInfo))
+	if (Input.Environment.CompilerFlags.Contains(CFLAG_ExtraShaderData))
 	{
 		Output.ShaderCode.AddOptionalData(FShaderCodeName::Key, TCHAR_TO_UTF8(*Input.GenerateShaderName()));
 	}
diff --git a/Engine/Source/Developer/Windows/ShaderFormatD3D/Private/D3DShaderCompiler.cpp b/Engine/Source/Developer/Windows/ShaderFormatD3D/Private/D3DShaderCompiler.cpp
--- a/Engine/Source/Developer/Windows/ShaderFormatD3D/Private/D3DShaderCompiler.cpp
+++ b/Engine/Source/Developer/Windows/ShaderFormatD3D/Private/D3DShaderCompiler.cpp
@@ -956,7 +956,7 @@ bool CompileAndProcessD3DShaderFXC(FString& PreprocessedShaderSource, const FStr
 		{
 			TRefCountPtr<ID3DBlob> CompressedData;
 
-			if (Input.Environment.CompilerFlags.Contains(CFLAG_KeepDebugInfo))
+			if (Input.Environment.CompilerFlags.Contains(CFLAG_GenerateSymbols))
 			{
 				CompressedData = Shader;
 			}
@@ -1143,10 +1143,14 @@ void CompileD3DShader(const FShaderCompilerInput& Input, FShaderCompilerOutput&
 		// Unpack uniform matrices as row-major to match the CPU layout.
 		| D3DCOMPILE_PACK_MATRIX_ROW_MAJOR;
 
+	if (Input.Environment.CompilerFlags.Contains(CFLAG_GenerateSymbols))
+	{
+		CompileFlags |= D3DCOMPILE_DEBUG;
+	}
+
 	if (Input.Environment.CompilerFlags.Contains(CFLAG_Debug)) 
 	{
-		//add the debug flags
-		CompileFlags |= D3DCOMPILE_DEBUG | D3DCOMPILE_SKIP_OPTIMIZATION;
+		CompileFlags |= D3DCOMPILE_SKIP_OPTIMIZATION;
 	}
 	else
 	{
diff --git a/Engine/Source/Developer/Windows/ShaderFormatD3D/Private/D3DShaderCompiler.inl b/Engine/Source/Developer/Windows/ShaderFormatD3D/Private/D3DShaderCompiler.inl
--- a/Engine/Source/Developer/Windows/ShaderFormatD3D/Private/D3DShaderCompiler.inl
+++ b/Engine/Source/Developer/Windows/ShaderFormatD3D/Private/D3DShaderCompiler.inl
@@ -25,7 +25,8 @@ public:
 		const FString& InDumpDebugInfoPath,	// Optional, empty when not dumping shader debug info
 		const FString& InBaseFilename,
 		bool bInEnable16BitTypes,
-		bool bKeepDebugInfo,
+		bool bGenerateSymbols,
+		bool bSymbolsBasedOnSource,
 		uint32 D3DCompileFlags,
 		uint32 AutoBindingSpace,
 		const TCHAR* InOptValidatorVersion
@@ -129,7 +130,7 @@ public:
 		if (D3DCompileFlags & D3DCOMPILE_DEBUG)
 		{
 			D3DCompileFlags &= ~D3DCOMPILE_DEBUG;
-			bKeepDebugInfo = true;
+			bGenerateSymbols = true;
 		}
 
 		if (bEnable16BitTypes)
@@ -145,11 +146,18 @@ public:
 			ExtraArguments.Add(FString(InOptValidatorVersion));
 		}
 
-		ExtraArguments.Add(L"/Zss");
-		ExtraArguments.Add(L"/Qembed_debug");
-		ExtraArguments.Add(L"/Zi");
-		ExtraArguments.Add(L"/Fd");
-		ExtraArguments.Add(L".\\");
+		if (bGenerateSymbols)
+		{
+			// -Zsb Compute Shader Hash considering only output binary
+			// -Zss Compute Shader Hash considering source information
+			ExtraArguments.Add(bSymbolsBasedOnSource ? L"/Zss" : L"/Zsb");
+
+			ExtraArguments.Add(L"/Qembed_debug");
+			ExtraArguments.Add(L"/Zi");
+
+			ExtraArguments.Add(L"/Fd");
+			ExtraArguments.Add(L".\\");
+		}
 
 		// Reflection will be removed later, otherwise the disassembly won't contain variables
 		//ExtraArguments.Add(L"/Qstrip_reflect");
@@ -557,7 +565,7 @@ inline void GenerateFinalOutput(TRefCountPtr<TBlob>& CompressedData,
 		}
 	}
 
-	if (Input.Environment.CompilerFlags.Contains(CFLAG_KeepDebugInfo))
+	if (Input.Environment.CompilerFlags.Contains(CFLAG_ExtraShaderData))
 	{
 		Output.ShaderCode.AddOptionalData(FShaderCodeName::Key, TCHAR_TO_UTF8(*Input.GenerateShaderName()));
 	}
diff --git a/Engine/Source/Developer/Windows/ShaderFormatD3D/Private/D3DShaderCompilerDXC.cpp b/Engine/Source/Developer/Windows/ShaderFormatD3D/Private/D3DShaderCompilerDXC.cpp
--- a/Engine/Source/Developer/Windows/ShaderFormatD3D/Private/D3DShaderCompilerDXC.cpp
+++ b/Engine/Source/Developer/Windows/ShaderFormatD3D/Private/D3DShaderCompilerDXC.cpp
@@ -673,7 +673,8 @@ bool CompileAndProcessD3DShaderDXC(FString& PreprocessedShaderSource,
 
 	const TCHAR* ValidatorVersion = nullptr; // Whatever default validator chosen by DXC
 
-	const bool bKeepDebugInfo = Input.Environment.CompilerFlags.Contains(CFLAG_KeepDebugInfo);
+	const bool bGenerateSymbols = Input.Environment.CompilerFlags.Contains(CFLAG_GenerateSymbols);
+	const bool bSymbolsBasedOnSource = Input.Environment.CompilerFlags.Contains(CFLAG_AllowUniqueSymbols);
 
 	FDxcArguments Args
 	(
@@ -683,7 +684,8 @@ bool CompileAndProcessD3DShaderDXC(FString& PreprocessedShaderSource,
 		Input.DumpDebugInfoPath,
 		Filename,
 		bEnable16BitTypes,
-		bKeepDebugInfo,
+		bGenerateSymbols,
+		bSymbolsBasedOnSource,
 		DXCFlags,
 		AutoBindingSpace,
 		ValidatorVersion
diff --git a/Engine/Source/Developer/Windows/ShaderFormatD3D/Private/ShaderFormatD3D.cpp b/Engine/Source/Developer/Windows/ShaderFormatD3D/Private/ShaderFormatD3D.cpp
--- a/Engine/Source/Developer/Windows/ShaderFormatD3D/Private/ShaderFormatD3D.cpp
+++ b/Engine/Source/Developer/Windows/ShaderFormatD3D/Private/ShaderFormatD3D.cpp
@@ -17,8 +17,8 @@ class FShaderFormatD3D : public IShaderFormat
 	enum
 	{
 		/** Version for shader format, this becomes part of the DDC key. */
-		UE_SHADER_PCD3D_SM6_VER = 1,
-		UE_SHADER_PCD3D_SM5_VER = 8,
+		UE_SHADER_PCD3D_SM6_VER = 2,
+		UE_SHADER_PCD3D_SM5_VER = 9,
 		UE_SHADER_PCD3D_ES3_1_VER = 8,
 	};
 
diff --git a/Engine/Source/Runtime/Engine/Private/ShaderCompiler/ShaderCompiler.cpp b/Engine/Source/Runtime/Engine/Private/ShaderCompiler/ShaderCompiler.cpp
--- a/Engine/Source/Runtime/Engine/Private/ShaderCompiler/ShaderCompiler.cpp
+++ b/Engine/Source/Runtime/Engine/Private/ShaderCompiler/ShaderCompiler.cpp
@@ -825,51 +825,60 @@ static FAutoConsoleVariableRef CVarForceAllCoresForShaderCompiling(
 	TEXT("Improves shader throughput but for big projects it can make the machine run OOM")
 );
 
-static TAutoConsoleVariable<int32> CVarKeepShaderDebugData(
-	TEXT("r.Shaders.KeepDebugInfo"),
+static TAutoConsoleVariable<int32> CVarShadersSymbols(
+	TEXT("r.Shaders.Symbols"),
 	0,
-	TEXT("Whether to keep shader reflection and debug data from shader bytecode, default is to strip.  When using graphical debuggers like Nsight it can be useful to enable this on startup.")
-	TEXT("For some platforms this cvar can be overriden in the Engine.ini, under the [ShaderCompiler] section."),
+	TEXT("Enables debugging of shaders in platform specific graphics debuggers. This will generate and write shader symbols.\n")
+	TEXT("This enables the behavior of both r.Shaders.GenerateSymbols and r.Shaders.WriteSymbols.\n")
+	TEXT("Enables shader debugging features that require shaders to be recompiled. This compiles shaders with symbols and also includes extra runtime information like shader names. When using graphical debuggers it can be useful to enable this on startup.\n")
+	TEXT("This setting can be overriden in any Engine.ini under the [ShaderCompiler] section."),
 	ECVF_ReadOnly);
 
-static TAutoConsoleVariable<int32> CVarPrepareExportedDebugInfo(
-	TEXT("r.Shaders.PrepareExportedDebugInfo"),
+static TAutoConsoleVariable<int32> CVarShadersGenerateSymbols(
+	TEXT("r.Shaders.GenerateSymbols"),
 	0,
-	TEXT("For platforms that work on exported shader debug info, this should be set to keep shader reflection and debug data from shader bytecode, default is to strip.")
-	TEXT("For these platforms, r.Shaders.KeepDebugInfo has the same behavior but will also export the debug info.")
-	TEXT("This will only be read from the platform Engine.ini, under the [ShaderCompiler] section."),
+	TEXT("Enables generation of data for shader debugging when compiling shaders. This explicitly does not write any shader symbols to disk.\n")
+	TEXT("This setting can be overriden in any Engine.ini under the [ShaderCompiler] section."),
 	ECVF_ReadOnly);
 
-static TAutoConsoleVariable<int32> CVarExportShaderDebugData(
-	TEXT("r.Shaders.ExportDebugInfo"),
+static TAutoConsoleVariable<int32> CVarShadersWriteSymbols(
+	TEXT("r.Shaders.WriteSymbols"),
 	0,
-	TEXT("Whether to export the shader reflection and debug data from shader bytecode as separate files.")
-	TEXT("r.Shaders.KeepDebugInfo must be enabled and r.DumpShaderDebugInfo will enable this cvar.")
-	TEXT("For some platforms this cvar can be overriden in the Engine.ini, under the [ShaderCompiler] section."),
+	TEXT("Enables writing shader symbols to disk for platforms that support that. This explicitly does not enable generation of shader symbols.\n")
+	TEXT("This setting can be overriden in any Engine.ini under the [ShaderCompiler] section."),
 	ECVF_ReadOnly);
 
 static TAutoConsoleVariable<FString> CVarShadersSymbolPathOverride(
 	TEXT("r.Shaders.SymbolPathOverride"),
 	"",
 	TEXT("Override output location of shader symbols. If the path contains the text '{Platform}', that will be replaced with the shader platform string.\n")
-	TEXT("Empty: use default location Saved/ShaderSymbols/{Platform}"),
+	TEXT("Empty: use default location Saved/ShaderSymbols/{Platform}\n")
+	TEXT("This setting can be overriden in any Engine.ini under the [ShaderCompiler] section."),
 	ECVF_ReadOnly);
 
 static TAutoConsoleVariable<int32> CVarAllowUniqueDebugInfo(
-	TEXT("r.Shaders.AllowUniqueDebugInfo"),
+	TEXT("r.Shaders.AllowUniqueSymbols"),
 	0,
-	TEXT("Whether to allow (exported) shader debug info to be unique per shader.\n")
-	TEXT("Enabling this can cause a drastic increase in debug data, enable only if absolutely necessary.\n")
-	TEXT("For some platforms this cvar can be overriden in the Engine.ini, under the [ShaderCompiler] section."),
+	TEXT("When enabled, this tells supported shader compilers to generate symbols based on source files.\n")
+	TEXT("Enabling this can cause a drastic increase in the number of symbol files, enable only if absolutely necessary.\n")
+	TEXT("This setting can be overriden in any Engine.ini under the [ShaderCompiler] section."),
 	ECVF_ReadOnly);
 
-static TAutoConsoleVariable<int32> CVarExportShaderDebugInfoZip(
-	TEXT("r.Shaders.ExportDebugInfo.Zip"),
+static TAutoConsoleVariable<int32> CVarShadersWriteSymbolsZip(
+	TEXT("r.Shaders.WriteSymbols.Zip"),
 	0,
 	TEXT(" 0: Export as loose files.\n")
 	TEXT(" 1: Export as an uncompressed archive.\n"),
 	ECVF_ReadOnly);
 
+static TAutoConsoleVariable<int32> CVarShadersEnableExtraData(
+	TEXT("r.Shaders.ExtraData"),
+	0,
+	TEXT("Enables generation of extra shader data that can be used at runtime. This includes shader names and other platform specific data.\n")
+	TEXT("This can add bloat to compiled shaders and can prevent shaders from being deduplicated.\n")
+	TEXT("This setting can be overriden in any Engine.ini under the [ShaderCompiler] section."),
+	ECVF_ReadOnly);
+
 static TAutoConsoleVariable<int32> CVarOptimizeShaders(
 	TEXT("r.Shaders.Optimize"),
 	1,
@@ -5154,14 +5163,22 @@ void GlobalBeginCompileShader(
 		}
 	}
 
-	if (ShouldKeepShaderDebugInfo(ShaderFormatName))
+	// Extra data (names, etc)
+	if (ShouldEnableExtraShaderData(ShaderFormatName))
+	{
+		Input.Environment.CompilerFlags.Add(CFLAG_ExtraShaderData);
+	}
+
+	// Symbols
+	if (ShouldGenerateShaderSymbols(ShaderFormatName))
 	{
-		Input.Environment.CompilerFlags.Add(CFLAG_KeepDebugInfo);
+		Input.Environment.CompilerFlags.Add(CFLAG_GenerateSymbols);
 	}
 
-	if (ShouldAllowUniqueDebugInfo(ShaderFormatName))
+	// Are symbols based on source or results
+	if (ShouldAllowUniqueShaderSymbols(ShaderFormatName))
 	{
-		Input.Environment.CompilerFlags.Add(CFLAG_AllowUniqueDebugInfo);
+		Input.Environment.CompilerFlags.Add(CFLAG_AllowUniqueSymbols);
 	}
 
 	if (CVarShaderFastMath.GetValueOnAnyThread() == 0)
diff --git a/Engine/Source/Runtime/RenderCore/Private/Shader.cpp b/Engine/Source/Runtime/RenderCore/Private/Shader.cpp
--- a/Engine/Source/Runtime/RenderCore/Private/Shader.cpp
+++ b/Engine/Source/Runtime/RenderCore/Private/Shader.cpp
@@ -1432,8 +1432,12 @@ void ShaderMapAppendKeyString(EShaderPlatform Platform, FString& KeyString)
 	}
 
 	{
-		KeyString += ShouldKeepShaderDebugInfo(ShaderFormatName) ? TEXT("_NoStrip") : TEXT("");
-		KeyString += ShouldAllowUniqueDebugInfo(ShaderFormatName) ? TEXT("_FullDbg") : TEXT("");
+		// Extra data (names, etc)
+		KeyString += ShouldEnableExtraShaderData(ShaderFormatName) ? TEXT("_ExtraData") : TEXT("");
+		// Symbols
+		KeyString += ShouldGenerateShaderSymbols(ShaderFormatName) ? TEXT("_Symbols") : TEXT("");
+		// Are symbols based on source or results
+		KeyString += ShouldAllowUniqueShaderSymbols(ShaderFormatName) ? TEXT("_FullDbg") : TEXT("");
 	}
 
 	{
diff --git a/Engine/Source/Runtime/RenderCore/Private/ShaderCore.cpp b/Engine/Source/Runtime/RenderCore/Private/ShaderCore.cpp
--- a/Engine/Source/Runtime/RenderCore/Private/ShaderCore.cpp
+++ b/Engine/Source/Runtime/RenderCore/Private/ShaderCore.cpp
@@ -381,11 +381,12 @@ public:
 		: SettingName(InSettingName)
 		, SettingCVar(!bPlatformOnly ? IConsoleManager::Get().FindConsoleVariable(InSettingName) : nullptr)
 	{
+		check(SettingCVar || bPlatformOnly);
 	}
 
 	bool IsEnabled(FName ShaderFormat) const
 	{
-		return  IsShaderCompilerConfigEnabledForPlatform(SettingCVar, SettingName, ShaderFormat);
+		return IsShaderCompilerConfigEnabledForPlatform(SettingCVar, SettingName, ShaderFormat);
 	}
 
 	bool GetString(FString& OutString, FName ShaderFormat) const
@@ -394,31 +395,27 @@ public:
 	}
 };
 
-bool ShouldKeepShaderDebugInfo(FName ShaderFormat)
+bool ShouldGenerateShaderSymbols(FName ShaderFormat)
 {
-	static const FShaderSymbolSettingHelper GenerateSymbols(TEXT("r.Shaders.KeepDebugInfo"));
-	static const FShaderSymbolSettingHelper PrepareExportedSymbols(TEXT("r.Shaders.PrepareExportedDebugInfo"), true);
-	return GenerateSymbols.IsEnabled(ShaderFormat) || PrepareExportedSymbols.IsEnabled(ShaderFormat);
+	static const FShaderSymbolSettingHelper Symbols(TEXT("r.Shaders.Symbols"));
+	static const FShaderSymbolSettingHelper GenerateSymbols(TEXT("r.Shaders.GenerateSymbols"), true);
+	return Symbols.IsEnabled(ShaderFormat) || GenerateSymbols.IsEnabled(ShaderFormat);
 }
 
-bool ShouldKeepShaderDebugInfo(EShaderPlatform ShaderPlatform)
+bool ShouldWriteShaderSymbols(FName ShaderFormat)
 {
-	return ShouldKeepShaderDebugInfo(LegacyShaderPlatformToShaderFormat(ShaderPlatform));
+	static const FShaderSymbolSettingHelper Symbols(TEXT("r.Shaders.Symbols"));
+	static const FShaderSymbolSettingHelper WriteSymbols(TEXT("r.Shaders.WriteSymbols"), true);
+	return Symbols.IsEnabled(ShaderFormat) || WriteSymbols.IsEnabled(ShaderFormat);
 }
 
-bool ShouldExportShaderDebugInfo(FName ShaderFormat)
+bool ShouldAllowUniqueShaderSymbols(FName ShaderFormat)
 {
-	static const FShaderSymbolSettingHelper ExportSymbols(TEXT("r.Shaders.ExportDebugInfo"));
-	static const FShaderSymbolSettingHelper GenerateSymbols(TEXT("r.Shaders.KeepDebugInfo"));
-	return ExportSymbols.IsEnabled(ShaderFormat) || GenerateSymbols.IsEnabled(ShaderFormat);
+	static const FShaderSymbolSettingHelper AllowUniqueSymbols(TEXT("r.Shaders.AllowUniqueSymbols"));
+	return AllowUniqueSymbols.IsEnabled(ShaderFormat);
 }
 
-bool ShouldExportShaderDebugInfo(EShaderPlatform ShaderPlatform)
-{
-	return ShouldExportShaderDebugInfo(LegacyShaderPlatformToShaderFormat(ShaderPlatform));
-}
-
-bool GetShaderDebugInfoPathOverride(FString& OutPathOverride, FName ShaderFormat)
+bool GetShaderSymbolPathOverride(FString& OutPathOverride, FName ShaderFormat)
 {
 	static const FShaderSymbolSettingHelper SymbolPathOverride(TEXT("r.Shaders.SymbolPathOverride"));
 	if (SymbolPathOverride.GetString(OutPathOverride, ShaderFormat))
@@ -437,10 +434,16 @@ bool GetShaderDebugInfoPathOverride(FString& OutPathOverride, FName ShaderFormat
 	return false;
 }
 
-bool ShouldExportShaderDebugInfoAsZip(FName ShaderFormat)
+bool ShouldWriteShaderSymbolsAsZip(FName ShaderFormat)
+{
+	static const FShaderSymbolSettingHelper WriteSymbolsZip(TEXT("r.Shaders.WriteSymbols.Zip"));
+	return WriteSymbolsZip.IsEnabled(ShaderFormat);
+}
+
+bool ShouldEnableExtraShaderData(FName ShaderFormat)
 {
-	static const FShaderSymbolSettingHelper ExportSymbolsZip(TEXT("r.Shaders.ExportDebugInfo.Zip"));
-	return ExportSymbolsZip.IsEnabled(ShaderFormat);
+	static const FShaderSymbolSettingHelper ExtraData(TEXT("r.Shaders.ExtraData"));
+	return ExtraData.IsEnabled(ShaderFormat);
 }
 
 #ifndef UE_ALLOW_SHADER_COMPILING
@@ -480,13 +483,6 @@ bool AllowGlobalShaderLoad()
 
 }
 
-
-bool ShouldAllowUniqueDebugInfo(FName ShaderFormat)
-{
-	static const FShaderSymbolSettingHelper AllowUniqueDebugInfo(TEXT("r.Shaders.AllowUniqueDebugInfo"));
-	return AllowUniqueDebugInfo.IsEnabled(ShaderFormat);
-}
-
 bool FShaderParameterMap::FindParameterAllocation(const TCHAR* ParameterName,uint16& OutBufferIndex,uint16& OutBaseIndex,uint16& OutSize) const
 {
 	const FParameterAllocation* Allocation = ParameterMap.Find(ParameterName);
diff --git a/Engine/Source/Runtime/RenderCore/Private/ShaderMap.cpp b/Engine/Source/Runtime/RenderCore/Private/ShaderMap.cpp
--- a/Engine/Source/Runtime/RenderCore/Private/ShaderMap.cpp
+++ b/Engine/Source/Runtime/RenderCore/Private/ShaderMap.cpp
@@ -161,7 +161,11 @@ bool FShaderMapBase::Serialize(FArchive& Ar, bool bInlineShaderResources, bool b
 
 		if (Ar.IsCooking())
 		{
-			Code->NotifyShadersCooked(Ar.CookingTarget());
+			const FName ShaderFormat = LegacyShaderPlatformToShaderFormat(ShaderPlatform);
+			if (ShaderFormat != NAME_None)
+			{
+				Code->NotifyShadersCompiled(ShaderFormat);
+			}
 		}
 
 		if (bShareCode)
diff --git a/Engine/Source/Runtime/RenderCore/Private/ShaderResource.cpp b/Engine/Source/Runtime/RenderCore/Private/ShaderResource.cpp
--- a/Engine/Source/Runtime/RenderCore/Private/ShaderResource.cpp
+++ b/Engine/Source/Runtime/RenderCore/Private/ShaderResource.cpp
@@ -304,30 +304,35 @@ void FShaderMapResourceCode::Serialize(FArchive& Ar, bool bLoadedByCookedMateria
 }
 
 #if WITH_EDITORONLY_DATA
-void FShaderMapResourceCode::NotifyShadersCooked(const ITargetPlatform* TargetPlatform)
+void FShaderMapResourceCode::NotifyShadersCompiled(FName FormatName)
 {
 #if WITH_ENGINE
 	// Notify the platform shader format that this particular shader is being used in the cook.
 	// We discard this data in cooked builds unless Ar.CookingTarget()->HasEditorOnlyData() is true.
-	check(TargetPlatform);
 	if (PlatformDebugData.Num())
 	{
-		TArray<FName> ShaderFormatNames;
-		TargetPlatform->GetAllTargetedShaderFormats(ShaderFormatNames);
-		for (FName FormatName : ShaderFormatNames)
+		if (const IShaderFormat* ShaderFormat = GetTargetPlatformManagerRef().FindShaderFormat(FormatName))
 		{
-			const IShaderFormat* ShaderFormat = GetTargetPlatformManagerRef().FindShaderFormat(FormatName);
-			if (ShaderFormat)
+			for (const TArray<uint8>& Entry : PlatformDebugData)
 			{
-				for (const auto& Entry : PlatformDebugData)
-				{
-					ShaderFormat->NotifyShaderCooked(Entry, FormatName);
-				}
+				ShaderFormat->NotifyShaderCompiled(Entry, FormatName);
 			}
 		}
 	}
 #endif // WITH_ENGINE
 }
+
+void FShaderMapResourceCode::NotifyShadersCooked(const ITargetPlatform* TargetPlatform)
+{
+#if WITH_ENGINE
+	TArray<FName> ShaderFormatNames;
+	TargetPlatform->GetAllTargetedShaderFormats(ShaderFormatNames);
+	for (FName FormatName : ShaderFormatNames)
+	{
+		NotifyShadersCompiled(FormatName);
+	}
+#endif
+}
 #endif // WITH_EDITORONLY_DATA
 
 FShaderMapResource::FShaderMapResource(EShaderPlatform InPlatform, int32 NumShaders)
diff --git a/Engine/Source/Runtime/RenderCore/Public/Shader.h b/Engine/Source/Runtime/RenderCore/Public/Shader.h
--- a/Engine/Source/Runtime/RenderCore/Public/Shader.h
+++ b/Engine/Source/Runtime/RenderCore/Public/Shader.h
@@ -352,6 +352,8 @@ public:
 
 	RENDERCORE_API void Serialize(FArchive& Ar, bool bLoadedByCookedMaterial);
 #if WITH_EDITORONLY_DATA
+	RENDERCORE_API void NotifyShadersCompiled(FName FormatName);
+	UE_DEPRECATED(5.0, "NotifyShadersCompiled should be called")
 	RENDERCORE_API void NotifyShadersCooked(const ITargetPlatform* TargetPlatform);
 #endif // WITH_EDITORONLY_DATA
 
diff --git a/Engine/Source/Runtime/RenderCore/Public/ShaderCompilerCore.h b/Engine/Source/Runtime/RenderCore/Public/ShaderCompilerCore.h
--- a/Engine/Source/Runtime/RenderCore/Public/ShaderCompilerCore.h
+++ b/Engine/Source/Runtime/RenderCore/Public/ShaderCompilerCore.h
@@ -24,21 +24,28 @@ const int32 ShaderCompileWorkerSingleJobHeader = 'S';
 // this is for the protocol, not the data.
 const int32 ShaderCompileWorkerPipelineJobHeader = 'P';
 
-/** Returns true if debug information should be kept for a given platform. */
-extern RENDERCORE_API bool ShouldKeepShaderDebugInfo(EShaderPlatform Platform);
-extern RENDERCORE_API bool ShouldKeepShaderDebugInfo(FName ShaderFormat);
+/** Returns true if shader symbols should be kept for a given platform. */
+extern RENDERCORE_API bool ShouldGenerateShaderSymbols(FName ShaderFormat);
 
-/** Returns true if debug information should be exported to separate files for a given platform . */
-extern RENDERCORE_API bool ShouldExportShaderDebugInfo(EShaderPlatform Platform);
-extern RENDERCORE_API bool ShouldExportShaderDebugInfo(FName ShaderFormat);
+/** Returns true if shader symbols should be exported to separate files for a given platform. */
+extern RENDERCORE_API bool ShouldWriteShaderSymbols(FName ShaderFormat);
 
-extern RENDERCORE_API bool GetShaderDebugInfoPathOverride(FString& OutPathOverride, FName ShaderFormat);
-
-/** Returns true if shader symbols should be exported to a single zip file for a given platform. */
-extern RENDERCORE_API bool ShouldExportShaderDebugInfoAsZip(FName ShaderFormat);
+/** Returns true if the shader symbol path is overriden and OutPathOverride contains the override path. */
+extern RENDERCORE_API bool GetShaderSymbolPathOverride(FString& OutPathOverride, FName ShaderFormat);
 
 /** Returns true if (external) shader symbols should be specific to each shader rather than be de-duplicated. */
-extern RENDERCORE_API bool ShouldAllowUniqueDebugInfo(FName ShaderFormat);
+extern RENDERCORE_API bool ShouldAllowUniqueShaderSymbols(FName ShaderFormat);
+
+/** Returns true if shaders should be combined into a single zip file instead of individual files. */
+extern RENDERCORE_API bool ShouldWriteShaderSymbolsAsZip(FName ShaderFormat);
+
+/** Returns true if the user wants more runtime shader data (names, extra info) */
+extern RENDERCORE_API bool ShouldEnableExtraShaderData(FName Platform);
+
+UE_DEPRECATED(5.0, "ShouldGenerateShaderSymbols should be called to determine if symbols (debug data) should be generated")
+bool ShouldKeepShaderDebugInfo(EShaderPlatform Platform);
+UE_DEPRECATED(5.0, "ShouldWriteShaderSymbols should be called to determine if symbols (debug data) should be written")
+bool ShouldExportShaderDebugInfo(EShaderPlatform Platform);
 
 /** Returns true is shader compiling is allowed */
 extern RENDERCORE_API bool AllowShaderCompiling();
@@ -58,10 +65,13 @@ enum ECompilerFlags
 	CFLAG_StandardOptimization,
 	// Always optimize, even when CFLAG_Debug is set. Required for some complex shaders and features.
 	CFLAG_ForceOptimization,
+	// Shader should generate symbols for debugging.
+	CFLAG_GenerateSymbols,
+	CFLAG_KeepDebugInfo UE_DEPRECATED(5.0, "CFLAG_GenerateSymbols should be used to signal if debug data needs to be generated") = CFLAG_GenerateSymbols,
 	// Shader should insert debug/name info at the risk of generating non-deterministic libraries
-	CFLAG_KeepDebugInfo,
-	// Allows the (external) debug info to be specific to each shader rather than trying to deduplicate.
-	CFLAG_AllowUniqueDebugInfo,
+	CFLAG_ExtraShaderData,
+	// Allows the (external) symbols to be specific to each shader rather than trying to deduplicate.
+	CFLAG_AllowUniqueSymbols,
 	CFLAG_NoFastMath,
 	// Explicitly enforce zero initialization on shader platforms that may omit it.
 	CFLAG_ZeroInitialise,
