GitHub #8210 : Remove unused function.

diff --git a/Engine/Source/Runtime/Engine/Private/Components/PrimitiveComponent.cpp b/Engine/Source/Runtime/Engine/Private/Components/PrimitiveComponent.cpp
--- a/Engine/Source/Runtime/Engine/Private/Components/PrimitiveComponent.cpp
+++ b/Engine/Source/Runtime/Engine/Private/Components/PrimitiveComponent.cpp
@@ -3514,11 +3514,6 @@ bool UPrimitiveComponent::UpdateOverlapsImpl(const TOverlapArrayView* NewPending
 	return bCanSkipUpdateOverlaps;
 }
 
-bool RequiresUpdateOverlaps(bool bGenerateOverlapEvents)
-{
-	return bGenerateOverlapEvents;
-}
-
 void UPrimitiveComponent::SetGenerateOverlapEvents(bool bInGenerateOverlapEvents)
 {
 	if (bGenerateOverlapEvents != bInGenerateOverlapEvents)
