Added the possibility to allow performing a GC pass while there active async loading through the gc.PerformGCWhileAsyncLoading cvar (off by default to keep original UE4 behavior).

diff --git a/Engine/Source/Runtime/Engine/Private/UnrealEngine.cpp b/Engine/Source/Runtime/Engine/Private/UnrealEngine.cpp
--- a/Engine/Source/Runtime/Engine/Private/UnrealEngine.cpp
+++ b/Engine/Source/Runtime/Engine/Private/UnrealEngine.cpp
@@ -1334,6 +1334,14 @@ static TAutoConsoleVariable<int32> CVarForceCollectGarbageEveryFrame(
 	TEXT("If set to 1, the engine will force GC each frame."));
 #endif
 
+int32 GPerformGCWhileAsyncLoading = 0;
+static FAutoConsoleVariableRef CVarPerformGCWhileAsyncLoading(
+	TEXT("gc.PerformGCWhileAsyncLoading"),
+	GPerformGCWhileAsyncLoading,
+	TEXT("Allow performing GC even if there's async loading in progress."),
+	ECVF_Default
+);
+
 static TAutoConsoleVariable<int32> CVarCollectGarbageEveryFrame(
 	TEXT("gc.CollectGarbageEveryFrame"),
 	0,
@@ -1564,7 +1572,7 @@ void UEngine::PerformGarbageCollectionAndCleanupActors()
 {
 	// We don't collect garbage while there are outstanding async load requests as we would need
 	// to block on loading the remaining data.
-	if (!IsAsyncLoading())
+	if (GPerformGCWhileAsyncLoading || !IsAsyncLoading())
 	{
 		bool bForcePurge = true;
 		for (FWorldContext& Context : WorldList)
