Fix crash when trying to get an actor's data layers when the actor's owning world is not fully initialized (happens when refreshing the world outliner while streaming cells).

diff --git a/Engine/Source/Runtime/Engine/Private/ActorEditor.cpp b/Engine/Source/Runtime/Engine/Private/ActorEditor.cpp
--- a/Engine/Source/Runtime/Engine/Private/ActorEditor.cpp
+++ b/Engine/Source/Runtime/Engine/Private/ActorEditor.cpp
@@ -1185,7 +1185,7 @@ TArray<FName> AActor::GetDataLayerNames() const
 
 TArray<const UDataLayer*> AActor::GetDataLayerObjects() const
 {
-	return GetDataLayerObjects(GetWorld()->GetWorldDataLayers());
+	return GetWorld() ? GetDataLayerObjects(GetWorld()->GetWorldDataLayers()) : TArray<const UDataLayer*>();
 }
 
 TArray<const UDataLayer*> AActor::GetDataLayerObjects(const AWorldDataLayers* WorldDataLayers) const
