Add build information to a memreport when it is written to a file

diff --git a/Engine/Source/Runtime/Engine/Private/UnrealEngine.cpp b/Engine/Source/Runtime/Engine/Private/UnrealEngine.cpp
--- a/Engine/Source/Runtime/Engine/Private/UnrealEngine.cpp
+++ b/Engine/Source/Runtime/Engine/Private/UnrealEngine.cpp
@@ -253,6 +253,7 @@ UnrealEngine.cpp: Implements the UEngine class and helpers.
 #include "Particles/ParticlePerfStatsManager.h"
 
 #include "Engine/InstancedStaticMesh.h"
+#include "IDeviceProfileSelectorModule.h"
 
 DEFINE_LOG_CATEGORY(LogEngine);
 IMPLEMENT_MODULE( FEngineModule, Engine );
@@ -6743,8 +6744,35 @@ bool UEngine::HandleMemReportDeferredCommand( const TCHAR* Cmd, FOutputDevice& A
 		UE_LOG(LogEngine, Log, TEXT("MemReportDeferred: saving to %s"), *FilenameFull);		
 	}
 
-	ReportAr->Logf(TEXT( "CommandLine Options: %s" ),  FCommandLine::Get() );
-	ReportAr->Logf(TEXT("Time Since Boot: %.02f Seconds") LINE_TERMINATOR, FPlatformTime::Seconds() - GStartTime);
+
+	// log out some useful information in the header
+	{
+		FString ConfigString = LexToString(FApp::GetBuildConfiguration());
+		uint32 ChangelistNumber = FEngineVersion::Current().GetChangelist();
+
+		// Set the device name to the platform name as a fallback
+		FString DeviceName = FPlatformProperties::PlatformName();
+
+		// Attempt to get the specific device name from the runtime device profile selector
+		FString DeviceProfile = TEXT("None");
+		FString DeviceProfileSelectionModule;
+		if (GConfig->GetString(TEXT("DeviceProfileManager"), TEXT("DeviceProfileSelectionModule"), DeviceProfileSelectionModule, GEngineIni))
+		{
+			if (IDeviceProfileSelectorModule* DPSelectorModule = FModuleManager::LoadModulePtr<IDeviceProfileSelectorModule>(*DeviceProfileSelectionModule))
+			{
+				DeviceName = DPSelectorModule->GetRuntimeDeviceProfileName();
+			}
+		}
+
+		ReportAr->Logf(TEXT("Changelist: %d"), ChangelistNumber);
+		ReportAr->Logf(TEXT("Config: %s"), *ConfigString);
+		ReportAr->Logf(TEXT("Device Name: %s"), *DeviceName);
+		ReportAr->Logf(TEXT("Device Profile: %s"), *DeviceProfile);
+
+		ReportAr->Logf(TEXT("CommandLine Options: %s"), FCommandLine::Get());
+		ReportAr->Logf(TEXT("Time Since Boot: %.02f Seconds") LINE_TERMINATOR, FPlatformTime::Seconds() - GStartTime);
+	}
+
 
 	// Run commands from the ini
 	FConfigSection* CommandsToRun = GConfig->GetSectionPrivate(TEXT("MemReportCommands"), 0, 1, GEngineIni);
