ISM partition actors create instance components.

This allows for individual instance manipulation using typed element interfaces instead of custom rolled code reaching into the partition actors, since the component will be treated as if the user had created it in the details panel.

diff --git a/Engine/Source/Runtime/Engine/Private/ISMPartition/ISMPartitionActor.cpp b/Engine/Source/Runtime/Engine/Private/ISMPartition/ISMPartitionActor.cpp
--- a/Engine/Source/Runtime/Engine/Private/ISMPartition/ISMPartitionActor.cpp
+++ b/Engine/Source/Runtime/Engine/Private/ISMPartition/ISMPartitionActor.cpp
@@ -360,6 +360,7 @@ void AISMPartitionActor::CreateComponent(const FISMComponentDescriptor& Componen
 	ComponentData.Component->bSelectable = true;
 	ComponentData.Component->bHasPerInstanceHitProxies = true;
 	ComponentData.RegisterDelegates();
+	AddInstanceComponent(ComponentData.Component);
 
 	ComponentData.Component->SetupAttachment(RootComponent);
 
diff --git a/Engine/Source/Runtime/Foliage/Private/InstancedFoliage.cpp b/Engine/Source/Runtime/Foliage/Private/InstancedFoliage.cpp
--- a/Engine/Source/Runtime/Foliage/Private/InstancedFoliage.cpp
+++ b/Engine/Source/Runtime/Foliage/Private/InstancedFoliage.cpp
@@ -1418,6 +1418,7 @@ void FFoliageStaticMesh::CreateNewComponent(const UFoliageType* InSettings)
 
 	AInstancedFoliageActor* IFA = GetIFA();
 	UFoliageInstancedStaticMeshComponent* FoliageComponent = NewObject<UFoliageInstancedStaticMeshComponent>(IFA, ComponentClass, NAME_None, RF_Transactional);
+	IFA->AddInstanceComponent(FoliageComponent);
 	
 	check(FoliageType_InstancedStaticMesh);
 
