Fix for morph targets not working on cooked builds

diff --git a/Engine/Source/Runtime/Engine/Private/SkeletalMesh.cpp b/Engine/Source/Runtime/Engine/Private/SkeletalMesh.cpp
--- a/Engine/Source/Runtime/Engine/Private/SkeletalMesh.cpp
+++ b/Engine/Source/Runtime/Engine/Private/SkeletalMesh.cpp
@@ -3422,6 +3422,10 @@ void USkeletalMesh::FinishPostLoadInternal(FSkeletalMeshPostLoadContext& Context
 
 #if WITH_EDITOR
 	ApplyFinishBuildInternalData(&Context);
+#else
+	// init morph targets here as in non-editor builds we don't use DDC or build the mesh. 
+	// should do this before InitResources, so that we clear invalid morph targets.
+	InitMorphTargets();
 #endif
 
 	// initialize rendering resources
