Adding logging for when command line overrides audio rendering settings

diff --git a/Engine/Source/Runtime/Engine/Private/AudioDevice.cpp b/Engine/Source/Runtime/Engine/Private/AudioDevice.cpp
--- a/Engine/Source/Runtime/Engine/Private/AudioDevice.cpp
+++ b/Engine/Source/Runtime/Engine/Private/AudioDevice.cpp
@@ -361,8 +361,14 @@ bool FAudioDevice::Init(Audio::FDeviceId InDeviceID, int32 InMaxSources)
 	PlatformSettings = GetPlatformSettings();
 
 	// Override platform settings with the command line
-	FParse::Value(FCommandLine::Get(), TEXT("-AudioCallbackBufferFrameSize="), PlatformSettings.CallbackBufferFrameSize);
-	FParse::Value(FCommandLine::Get(), TEXT("-AudioNumBuffersToEnqueue="), PlatformSettings.NumBuffers);
+	if (FParse::Value(FCommandLine::Get(), TEXT("-AudioCallbackBufferFrameSize="), PlatformSettings.CallbackBufferFrameSize))
+	{
+		UE_LOG(LogAudioMixer, Display, TEXT("Command Line CallbackBufferFrameSize Override: %d"), PlatformSettings.CallbackBufferFrameSize);
+	}
+	if (FParse::Value(FCommandLine::Get(), TEXT("-AudioNumBuffersToEnqueue="), PlatformSettings.NumBuffers))
+	{
+		UE_LOG(LogAudioMixer, Display, TEXT("Command Line NumBuffersToEnqueue Override: %d"), PlatformSettings.NumBuffers);
+	}
 
 	// MaxSources is the max value supplied to Init call (quality settings), unless overwritten by the platform settings.
 	// This does not have to be the minimum value in this case (nor is it desired, so platforms can potentially scale up)
