[NO NON-PII COMMIT MESSAGE]

diff --git a/Engine/Source/Runtime/Engine/Classes/GameFramework/PlayerController.h b/Engine/Source/Runtime/Engine/Classes/GameFramework/PlayerController.h
--- a/Engine/Source/Runtime/Engine/Classes/GameFramework/PlayerController.h
+++ b/Engine/Source/Runtime/Engine/Classes/GameFramework/PlayerController.h
@@ -1791,7 +1791,7 @@ public:
 	/** Gives the PlayerController an opportunity to cleanup any changes it applied to the game viewport, primarily for the touch interface */
 	virtual void CleanupGameViewport();
 
-	/** Returns true if input should be frozen (whether UnFreeze timer is active) */
+	/** Called on the client to do local pawn setup after possession, before calling ServerAcknowledgePossession */
 	virtual void AcknowledgePossession(class APawn* P);
 
 	/** Clean up when a Pawn's player is leaving a game. Base implementation destroys the pawn. */
