Functions like GetPlayerController now use a stable PlayerIndex that puts all local controllers first. The old method with player controller iterator order changed for map transfer and network changes, and users expect index 0 to be local
Add functions for getting the local and total number of player controllers
Add functions for accessing PlayerStates, including via Unique Net Id (which now has a display name)
Add native functions to better access the local players list on game instance
Add OnSetUniqueId virtual to PlayerState to allow games to update things properly
Many comment fixes to explain how these functions actually work

diff --git a/Engine/Source/Runtime/Engine/Classes/Engine/GameInstance.h b/Engine/Source/Runtime/Engine/Classes/Engine/GameInstance.h
--- a/Engine/Source/Runtime/Engine/Classes/Engine/GameInstance.h
+++ b/Engine/Source/Runtime/Engine/Classes/Engine/GameInstance.h
@@ -14,6 +14,7 @@
 
 #include "Subsystems/GameInstanceSubsystem.h"
 #include "Subsystems/SubsystemCollection.h"
+#include "GameFramework/OnlineReplStructs.h"
 
 #if WITH_EDITOR
 #include "Settings/LevelEditorPlaySettings.h"
@@ -309,17 +310,33 @@ public:
 	 * @return whether the player was successfully removed. Removal is not allowed while connected to a server.
 	 */
 	virtual bool			RemoveLocalPlayer(ULocalPlayer * ExistingPlayer);
-
+	
+	/** Returns number of fully registered local players */
 	int32					GetNumLocalPlayers() const;
+
+	/** Returns the local player at a certain index, or null if not found */
 	ULocalPlayer*			GetLocalPlayerByIndex(const int32 Index) const;
+
+	/** Returns the first local player, will not be null during normal gameplay */
+	ULocalPlayer*			GetFirstGamePlayer() const;
+
+	/** Returns the player controller assigned to the first local player. If World is specified it will search within that specific world */
 	APlayerController*		GetFirstLocalPlayerController(const UWorld* World = nullptr) const;
+
+	/** Returns the local player assigned to a physical Controller Id, or null if not found */
 	ULocalPlayer*			FindLocalPlayerFromControllerId(const int32 ControllerId) const;
+
+	/** Returns the local player that has been assigned the specific unique net id */
 	ULocalPlayer*			FindLocalPlayerFromUniqueNetId(FUniqueNetIdPtr UniqueNetId) const;
 	ULocalPlayer*			FindLocalPlayerFromUniqueNetId(const FUniqueNetId& UniqueNetId) const;
-	ULocalPlayer*			GetFirstGamePlayer() const;
+	ULocalPlayer*			FindLocalPlayerFromUniqueNetId(const FUniqueNetIdRepl& UniqueNetId) const;
 
+	/** Returns const iterator for searching list of local players */
 	TArray<ULocalPlayer*>::TConstIterator	GetLocalPlayerIterator() const;
+
+	/** Returns reference to entire local player list */
 	const TArray<ULocalPlayer*> &			GetLocalPlayers() const;
+
 	/**
 	 * Get the primary player controller on this machine (others are splitscreen children)
 	 * (must have valid player state)
diff --git a/Engine/Source/Runtime/Engine/Classes/Engine/LocalPlayer.h b/Engine/Source/Runtime/Engine/Classes/Engine/LocalPlayer.h
--- a/Engine/Source/Runtime/Engine/Classes/Engine/LocalPlayer.h
+++ b/Engine/Source/Runtime/Engine/Classes/Engine/LocalPlayer.h
@@ -282,6 +282,13 @@ public:
 	 */
 	UGameInstance* GetGameInstance() const;
 
+	/**
+	 * Returns the index of this player in the Game instances local players array
+	 *
+	 * @return Index in array, will be >= 0 if this is a fully registered player
+	 */
+	int32 GetIndexInGameInstance() const;
+
 	/**
 	 * Get a Subsystem of specified type
 	 */
@@ -391,7 +398,7 @@ public:
 	virtual void SendSplitJoin(TArray<FString>& Options);
 	
 	/**
-	 * Change the ControllerId for this player; if the specified ControllerId is already taken by another player, changes the ControllerId
+	 * Change the physical ControllerId for this player; if the specified ControllerId is already taken by another player, changes the ControllerId
 	 * for the other player to the ControllerId currently in use by this player.
 	 *
 	 * @param	NewControllerId		the ControllerId to assign to this player.
@@ -399,7 +406,7 @@ public:
 	virtual void SetControllerId(int32 NewControllerId);
 
 	/**
-	 * Returns the controller ID for the player
+	 * Returns the physical controller ID for the player
 	 */
 	int32 GetControllerId() const { return ControllerId; }
 
diff --git a/Engine/Source/Runtime/Engine/Classes/GameFramework/OnlineReplStructs.h b/Engine/Source/Runtime/Engine/Classes/GameFramework/OnlineReplStructs.h
--- a/Engine/Source/Runtime/Engine/Classes/GameFramework/OnlineReplStructs.h
+++ b/Engine/Source/Runtime/Engine/Classes/GameFramework/OnlineReplStructs.h
@@ -20,7 +20,7 @@ class FJsonValue;
  * Makes sure that the opaque aspects of FUniqueNetId are properly handled/serialized 
  * over network RPC and actor replication
  */
-USTRUCT(BlueprintType)
+USTRUCT(BlueprintType, DisplayName = "Unique Net Id")
 struct FUniqueNetIdRepl : public FUniqueNetIdWrapper
 {
 	GENERATED_USTRUCT_BODY()
diff --git a/Engine/Source/Runtime/Engine/Classes/GameFramework/PlayerState.h b/Engine/Source/Runtime/Engine/Classes/GameFramework/PlayerState.h
--- a/Engine/Source/Runtime/Engine/Classes/GameFramework/PlayerState.h
+++ b/Engine/Source/Runtime/Engine/Classes/GameFramework/PlayerState.h
@@ -223,12 +223,20 @@ public:
 	//~ End AActor Interface
 
 	/** Return the pawn controlled by this Player State. */
+	UFUNCTION(BlueprintCallable, Category = "PlayerState")
 	APawn* GetPawn() const { return PawnPrivate; }
 
 	/** Convenience helper to return a cast version of the pawn controlled by this Player State. */
 	template<class T>
 	T* GetPawn() const { return Cast<T>(PawnPrivate); }
 
+	/** Returns the AI or player controller that created this player state, or null for remote clients */
+	class AController* GetOwningController() const;
+
+	/** Return the player controller that created this player state, or null for remote clients */
+	UFUNCTION(BlueprintCallable, Category = "PlayerState")
+	class APlayerController* GetPlayerController() const;
+	
 	/** Called by Controller when its PlayerState is initially replicated. */
 	virtual void ClientInitialize(class AController* C);
 
@@ -266,12 +274,6 @@ public:
 	/** set the player name to S */
 	virtual void SetOldPlayerName(const FString& S);
 
-	/** 
-	 * Associate an online unique id with this player
-	 * @param InUniqueId the unique id associated with this player
-	 */
-	virtual void SetUniqueId(const FUniqueNetIdPtr& InUniqueId);
-
 	/** 
 	 * Register a player with the online subsystem
 	 * @param bWasFromInvite was this player invited directly
@@ -428,9 +430,23 @@ public:
 		return UniqueId;
 	}
 
+	/** Gets the online unique id for a player. If a player is logged in this will be consistent across all clients and servers. */
+	UFUNCTION(BlueprintCallable, Category = "PlayerState", meta = (DisplayName = "Get Unique Net Id"))
+	FUniqueNetIdRepl BP_GetUniqueId() const;
+
+	/**
+	 * Associate an online unique id with this player
+	 * @param InUniqueId the unique id associated with this player
+	 */
+	virtual void SetUniqueId(const FUniqueNetIdPtr& InUniqueId);
+
 	/** Sets the value of UniqueId without causing other side effects to this instance. */
 	void SetUniqueId(const FUniqueNetIdRepl& NewUniqueId);
 	void SetUniqueId(FUniqueNetIdRepl&& NewUniqueId);
+
+	/** Called on both the client and server when unique ID has been modified */
+	virtual void OnSetUniqueId();
+
 	//~ End Methods for Replicated Members.
 PRAGMA_ENABLE_DEPRECATION_WARNINGS
 };
diff --git a/Engine/Source/Runtime/Engine/Classes/Kismet/GameplayStatics.h b/Engine/Source/Runtime/Engine/Classes/Kismet/GameplayStatics.h
--- a/Engine/Source/Runtime/Engine/Classes/Kismet/GameplayStatics.h
+++ b/Engine/Source/Runtime/Engine/Classes/Kismet/GameplayStatics.h
@@ -139,55 +139,126 @@ class ENGINE_API UGameplayStatics : public UBlueprintFunctionLibrary
 	UFUNCTION(BlueprintPure, Category="Game", meta=(WorldContext="WorldContextObject"))
 	static class UGameInstance* GetGameInstance(const UObject* WorldContextObject);
 
-	/** Returns the player controller at the specified player index */
-	UFUNCTION(BlueprintPure, Category="Game", meta=(WorldContext="WorldContextObject", UnsafeDuringActorConstruction="true"))
-	static class APlayerController* GetPlayerController(const UObject* WorldContextObject, int32 PlayerIndex);
+	/**
+	 * Returns the number of active player states, there is one player state for every connected player even if they are a remote client.
+	 * Indexes up to this can be use as PlayerStateIndex parameters for other functions.
+	 */
+	UFUNCTION(BlueprintPure, Category="Game", meta = (WorldContext="WorldContextObject", UnsafeDuringActorConstruction="true"))
+	static int32 GetNumPlayerStates(const UObject* WorldContextObject);
+
+	/**
+	 * Returns the player state at the given index in the game state's PlayerArray. 
+	 * This will work on both the client and server and the index will be consistent.
+	 * After initial replication, all clients and the server will have access to PlayerStates for all connected players.
+	 *
+	 * @param PlayerStateIndex	Index into the game state's PlayerArray
+	 */
+	UFUNCTION(BlueprintPure, Category = "Game", meta = (WorldContext = "WorldContextObject", UnsafeDuringActorConstruction = "true"))
+	static class APlayerState* GetPlayerState(const UObject* WorldContextObject, int32 PlayerStateIndex);
 
-	/** Returns the player controller that has the given controller ID */
+	/**
+	 * Returns the player state that matches the passed in online id, or null for an invalid one.
+	 * This will work on both the client and server for local and remote players.
+	 *
+	 * @param UniqueId	The player's unique net/online id
+	 */
 	UFUNCTION(BlueprintPure, Category = "Game", meta = (WorldContext = "WorldContextObject", UnsafeDuringActorConstruction = "true"))
-	static class APlayerController* GetPlayerControllerFromID(const UObject* WorldContextObject, int32 ControllerID);
+	static class APlayerState* GetPlayerStateFromUniqueNetId(const UObject* WorldContextObject, const FUniqueNetIdRepl& UniqueId);
 
-	/** Returns the player pawn at the specified player index */
+	/**
+	 * Returns the total number of available player controllers, including remote players when called on a server.
+	 * Indexes up to this can be used as PlayerIndex parameters for the following functions.
+	 */
+	UFUNCTION(BlueprintPure, Category = "Game", meta = (WorldContext = "WorldContextObject", UnsafeDuringActorConstruction = "true"))
+	static int32 GetNumPlayerControllers(const UObject* WorldContextObject);
+
+	/**
+	 * Returns the number of fully initialized local players, this will be 0 on dedicated servers.
+	 * Indexes up to this can be used as PlayerIndex parameters for the following functions, and you are guaranteed to get a local player controller.
+	 */
+	UFUNCTION(BlueprintPure, Category = "Game", meta = (WorldContext = "WorldContextObject", UnsafeDuringActorConstruction = "true"))
+	static int32 GetNumLocalPlayerControllers(const UObject* WorldContextObject);
+
+	/** 
+	 * Returns the player controller found while iterating through the local and available remote player controllers.
+	 * On a network client, this will only include local players as remote player controllers are not available.
+	 * The index will be consistent as long as no new players join or leave, but it will not be the same across different clients and servers.
+	 *
+	 * @param PlayerIndex	Index in the player controller list, starting first with local players and then available remote ones
+	 */
+	UFUNCTION(BlueprintPure, Category="Game", meta=(WorldContext="WorldContextObject", UnsafeDuringActorConstruction="true"))
+	static class APlayerController* GetPlayerController(const UObject* WorldContextObject, int32 PlayerIndex);
+
+	/**
+	 * Returns the pawn for the player controller at the specified player index.
+	 * This will not include pawns of remote clients with no available player controller, you can use the player states list for that.
+	 *
+	 * @param PlayerIndex	Index in the player controller list, starting first with local players and then available remote ones
+	 */
 	UFUNCTION(BlueprintPure, Category="Game", meta=(WorldContext="WorldContextObject", UnsafeDuringActorConstruction="true"))
 	static class APawn* GetPlayerPawn(const UObject* WorldContextObject, int32 PlayerIndex);
 
-	/** Returns the player character (NULL if the player pawn doesn't exist OR is not a character) at the specified player index */
+	/**
+	 * Returns the pawn for the player controller at the specified player index, will return null if the pawn is not a character.
+	 * This will not include characters of remote clients with no available player controller, you can iterate the PlayerStates list for that.
+	 *
+	 * @param PlayerIndex	Index in the player controller list, starting first with local players and then available remote ones
+	 */
 	UFUNCTION(BlueprintPure, Category="Game", meta=(WorldContext="WorldContextObject", UnsafeDuringActorConstruction="true"))
 	static class ACharacter* GetPlayerCharacter(const UObject* WorldContextObject, int32 PlayerIndex);
 
-	/** Returns the player's camera manager for the specified player index */
+	/**
+	 * Returns the camera manager for the Player Controller at the specified player index.
+	 * This will not include remote clients with no player controller.
+	 *
+	 * @param PlayerIndex	Index in the player controller list, starting first with local players and then available remote ones
+	 */
 	UFUNCTION(BlueprintPure, Category="Game", meta=(WorldContext="WorldContextObject", UnsafeDuringActorConstruction="true"))
 	static class APlayerCameraManager* GetPlayerCameraManager(const UObject* WorldContextObject, int32 PlayerIndex);
 
-	/** Create a new player for this game.  
-	 *  @param ControllerId             The ID of the controller that the should control the newly created player.  A value of -1 specifies to use the next available ID
-	 *  @param bSpawnPlayerController   Whether a player controller should be spawned immediately for this player. If false a player controller will not be created automatically until transition to the next map.
-	 *  @return                         The created player controller if one is created. 
+	/** 
+	 * Returns the player controller with the specified physical controller ID. This only works for local players.
+	 *
+	 * @param ControllerID	Physical controller ID, the same value returned from Get Player Controller ID
+	 */
+	UFUNCTION(BlueprintPure, Category="Game", meta=(DisplayName="Get Local Player Controller From ID", WorldContext="WorldContextObject", UnsafeDuringActorConstruction="true"))
+	static class APlayerController* GetPlayerControllerFromID(const UObject* WorldContextObject, int32 ControllerID);
+
+	/** 
+	 * Create a new local player for this game, for cases like local multiplayer.
+	 *
+	 * @param ControllerId             The ID of the physical controller that the should control the newly created player. A value of -1 specifies to use the next available ID
+	 * @param bSpawnPlayerController   Whether a player controller should be spawned immediately for this player. If false a player controller will not be created automatically until transition to the next map.
+	 * @return                         The created player controller if one is created.
 	 */
-	UFUNCTION(BlueprintCallable, Category="Game", meta=(WorldContext="WorldContextObject", AdvancedDisplay="2", UnsafeDuringActorConstruction="true"))
+	UFUNCTION(BlueprintCallable, Category = "Game", meta = (DisplayName="Create Local Player", WorldContext = "WorldContextObject", AdvancedDisplay = "2", UnsafeDuringActorConstruction = "true"))
 	static class APlayerController* CreatePlayer(const UObject* WorldContextObject, int32 ControllerId = -1, bool bSpawnPlayerController = true);
 
-	/** Removes a player from this game.  
-	 *  @param Player			The player controller of the player to be removed
-	 *  @param bDestroyPawn		Whether the controlled pawn should be deleted as well
+	/** 
+	 * Removes a local player from this game.
+	 *
+	 * @param Player			The player controller of the player to be removed
+	 * @param bDestroyPawn		Whether the controlled pawn should be deleted as well
 	 */
-	UFUNCTION(BlueprintCallable, Category="Game", meta=(UnsafeDuringActorConstruction="true"))
+	UFUNCTION(BlueprintCallable, Category = "Game", meta = (DisplayName="Remove Local Player", UnsafeDuringActorConstruction = "true"))
 	static void RemovePlayer(APlayerController* Player, bool bDestroyPawn);
 
-	/** 
-	* Gets what controller ID a Player is using
-	* @param Player		The player controller of the player to get the ID of
-	* @return			The ID of the passed in player. -1 if there is no controller for the passed in player
-	*/
-	UFUNCTION(BlueprintPure, Category="Game", meta=(UnsafeDuringActorConstruction="true"))
+	/**
+	 * Gets what physical controller ID a player is using. This only works for local player controllers.
+	 *
+	 * @param Player	The player controller of the player to get the ID of
+	 * @return			The ID of the passed in player. -1 if there is no physical controller assigned to the passed in player
+	 */
+	UFUNCTION(BlueprintPure, Category="Game", meta=(DisplayName="Get Local Player Controller ID", UnsafeDuringActorConstruction="true"))
 	static int32 GetPlayerControllerID(APlayerController* Player);
 
-	/** 
-	 * Sets what controller ID a Player should be using
+	/**
+	 * Sets what physical controller ID a player should be using. This only works for local player controllers.
+	 *
 	 * @param Player			The player controller of the player to change the controller ID of
 	 * @param ControllerId		The controller ID to assign to this player
 	 */
-	UFUNCTION(BlueprintCallable, Category="Game", meta=(UnsafeDuringActorConstruction="true"))
+	UFUNCTION(BlueprintCallable, Category="Game", meta=(DisplayName="Set Local Player Controller ID",UnsafeDuringActorConstruction="true"))
 	static void SetPlayerControllerID(APlayerController* Player, int32 ControllerId);
 
 	// --- Level Streaming functions ------------------------
diff --git a/Engine/Source/Runtime/Engine/Private/GameInstance.cpp b/Engine/Source/Runtime/Engine/Private/GameInstance.cpp
--- a/Engine/Source/Runtime/Engine/Private/GameInstance.cpp
+++ b/Engine/Source/Runtime/Engine/Private/GameInstance.cpp
@@ -866,7 +866,8 @@ ULocalPlayer* UGameInstance::GetLocalPlayerByIndex(const int32 Index) const
 
 APlayerController* UGameInstance::GetFirstLocalPlayerController(const UWorld* World) const
 {
-	if (World == nullptr)
+	// Use the consistent local players order if possible
+	if (World == nullptr || World == GetWorld())
 	{
 		for (ULocalPlayer* Player : LocalPlayers)
 		{
@@ -982,6 +983,16 @@ ULocalPlayer* UGameInstance::FindLocalPlayerFromUniqueNetId(FUniqueNetIdPtr Uniq
 	return FindLocalPlayerFromUniqueNetId(*UniqueNetId);
 }
 
+ULocalPlayer* UGameInstance::FindLocalPlayerFromUniqueNetId(const FUniqueNetIdRepl& UniqueNetId) const
+{
+	if (!UniqueNetId.IsValid())
+	{
+		return nullptr;
+	}
+
+	return FindLocalPlayerFromUniqueNetId(*UniqueNetId);
+}
+
 ULocalPlayer* UGameInstance::GetFirstGamePlayer() const
 {
 	return (LocalPlayers.Num() > 0) ? LocalPlayers[0] : nullptr;
diff --git a/Engine/Source/Runtime/Engine/Private/GameplayStatics.cpp b/Engine/Source/Runtime/Engine/Private/GameplayStatics.cpp
--- a/Engine/Source/Runtime/Engine/Private/GameplayStatics.cpp
+++ b/Engine/Source/Runtime/Engine/Private/GameplayStatics.cpp
@@ -31,6 +31,8 @@
 #include "Sound/SoundBase.h"
 #include "Sound/DialogueWave.h"
 #include "GameFramework/SaveGame.h"
+#include "GameFramework/GameStateBase.h"
+#include "GameFramework/PlayerState.h"
 #include "PhysicalMaterials/PhysicalMaterial.h"
 #include "Components/DecalComponent.h"
 #include "Components/ForceFeedbackComponent.h"
@@ -208,21 +210,123 @@ class UGameInstance* UGameplayStatics::GetGameInstance(const UObject* WorldConte
 	return World ? World->GetGameInstance() : nullptr;
 }
 
-class APlayerController* UGameplayStatics::GetPlayerController(const UObject* WorldContextObject, int32 PlayerIndex ) 
+int32 UGameplayStatics::GetNumPlayerStates(const UObject* WorldContextObject)
+{
+	AGameStateBase* GameState = GetGameState(WorldContextObject);
+
+	if (GameState)
+	{
+		return GameState->PlayerArray.Num();
+	}
+
+	return 0;
+}
+
+class APlayerState* UGameplayStatics::GetPlayerState(const UObject* WorldContextObject, int32 PlayerStateIndex)
+{
+	AGameStateBase* GameState = GetGameState(WorldContextObject);
+
+	if (GameState && GameState->PlayerArray.IsValidIndex(PlayerStateIndex))
+	{
+		return GameState->PlayerArray[PlayerStateIndex];
+	}
+
+	return nullptr;
+}
+
+class APlayerState* UGameplayStatics::GetPlayerStateFromUniqueNetId(const UObject* WorldContextObject, const FUniqueNetIdRepl& UniqueId)
+{
+	AGameStateBase* GameState = GetGameState(WorldContextObject);
+
+	if (GameState)
+	{
+		return GameState->GetPlayerStateFromUniqueNetId(UniqueId);
+	}
+
+	return nullptr;
+}
+
+int32 UGameplayStatics::GetNumPlayerControllers(const UObject* WorldContextObject)
 {
 	if (UWorld* World = GEngine->GetWorldFromContextObject(WorldContextObject, EGetWorldErrorMode::LogAndReturnNull))
 	{
-		uint32 Index = 0;
-		for (FConstPlayerControllerIterator Iterator = World->GetPlayerControllerIterator(); Iterator; ++Iterator)
+		return World->GetNumPlayerControllers();
+	}
+	return 0;
+}
+
+int32 UGameplayStatics::GetNumLocalPlayerControllers(const UObject* WorldContextObject)
+{
+	int32 Count = 0;
+	UGameInstance* GameInstance = GetGameInstance(WorldContextObject);
+
+	// We only want Local Players that have valid player controllers
+	if (GameInstance)
+	{
+		const TArray<ULocalPlayer*>& LocalPlayers = GameInstance->GetLocalPlayers();
+		for (ULocalPlayer* LocalPlayer : LocalPlayers)
 		{
-			APlayerController* PlayerController = Iterator->Get();
-			if (Index == PlayerIndex)
+			if (APlayerController* PC = LocalPlayer->PlayerController)
 			{
-				return PlayerController;
+				Count++;
+			}
+		}
+	}
+	return Count;
+}
+
+class APlayerController* UGameplayStatics::GetPlayerController(const UObject* WorldContextObject, int32 PlayerIndex) 
+{
+	// The order for the player controller iterator is not consistent across map transfer/etc so we don't want to use that index
+	// 99% of the time people pass in index 0 and want the primary local player controller
+	// After we've finished iterating the local player controllers, iterate the GameState list to find remote ones in a consistent order
+
+	int32 Index = 0;
+	UGameInstance* GameInstance = GetGameInstance(WorldContextObject);
+	if (GameInstance)
+	{		
+		const TArray<ULocalPlayer*>& LocalPlayers = GameInstance->GetLocalPlayers();
+		for (ULocalPlayer* LocalPlayer : LocalPlayers)
+		{
+			// Only count local players with an actual PC as part of the indexing
+			if (APlayerController* PC = LocalPlayer->PlayerController)
+			{
+				if (Index == PlayerIndex)
+				{
+					return PC;
+				}
+				Index++;
 			}
-			Index++;
 		}
 	}
+
+	AGameStateBase* GameState = GetGameState(WorldContextObject);
+	if (GameState)
+	{
+		for (APlayerState* PlayerState : GameState->PlayerArray)
+		{
+			// Only count valid player controllers that we skipped over in the last list
+			APlayerController* PC = PlayerState->GetPlayerController();
+			if (PC && !PC->GetLocalPlayer())
+			{
+				if (Index == PlayerIndex)
+				{
+					return PC;
+				}
+				Index++;
+			}
+		}
+	}
+
+	// Fallback case for weird world contexts with no game instance or state, only support index 0
+	if (PlayerIndex == 0)
+	{
+		if (UWorld* World = GEngine->GetWorldFromContextObject(WorldContextObject, EGetWorldErrorMode::LogAndReturnNull))
+		{
+			return World->GetFirstPlayerController();
+		}
+	}
+
 	return nullptr;
 }
 
diff --git a/Engine/Source/Runtime/Engine/Private/LocalPlayer.cpp b/Engine/Source/Runtime/Engine/Private/LocalPlayer.cpp
--- a/Engine/Source/Runtime/Engine/Private/LocalPlayer.cpp
+++ b/Engine/Source/Runtime/Engine/Private/LocalPlayer.cpp
@@ -1618,6 +1618,19 @@ UGameInstance* ULocalPlayer::GetGameInstance() const
 	return ViewportClient ? ViewportClient->GetGameInstance() : nullptr;
 }
 
+int32 ULocalPlayer::GetIndexInGameInstance() const
+{
+	int32 FoundIndex = INDEX_NONE;
+	UGameInstance* GameInstance = GetGameInstance();
+	if (GameInstance)
+	{
+		const TArray<ULocalPlayer*>& LocalPlayers = GameInstance->GetLocalPlayers();
+		LocalPlayers.Find(const_cast<ULocalPlayer*>(this), FoundIndex);
+	}
+
+	return FoundIndex;
+}
+
 void ULocalPlayer::AddReferencedObjects(UObject* InThis, FReferenceCollector& Collector)
 {
 	ULocalPlayer* This = CastChecked<ULocalPlayer>(InThis);
diff --git a/Engine/Source/Runtime/Engine/Private/PlayerState.cpp b/Engine/Source/Runtime/Engine/Private/PlayerState.cpp
--- a/Engine/Source/Runtime/Engine/Private/PlayerState.cpp
+++ b/Engine/Source/Runtime/Engine/Private/PlayerState.cpp
@@ -207,8 +207,8 @@ void APlayerState::PostInitializeComponents()
 		return;
 	}
 
-	AController* OwningController = Cast<AController>(GetOwner());
-	if (OwningController != NULL)
+	AController* OwningController = GetOwningController();
+	if (OwningController != nullptr)
 	{
 		SetIsABot(Cast<APlayerController>(OwningController) == nullptr);
 	}
@@ -219,6 +219,16 @@ void APlayerState::PostInitializeComponents()
 	}
 }
 
+class AController* APlayerState::GetOwningController() const
+{
+	return Cast<AController>(GetOwner());
+}
+
+class APlayerController* APlayerState::GetPlayerController() const
+{
+	return Cast<APlayerController>(GetOwner());
+}
+
 void APlayerState::ClientInitialize(AController* C)
 {
 	SetOwner(C);
@@ -378,18 +388,13 @@ void APlayerState::OnRep_PlayerId()
 
 void APlayerState::OnRep_UniqueId()
 {
+	// First notify it's changed
+	OnSetUniqueId();
+
 	// Register player with session
 	RegisterPlayerWithSession(false);
 }
 
-PRAGMA_DISABLE_DEPRECATION_WARNINGS
-void APlayerState::SetUniqueId(const FUniqueNetIdPtr& InUniqueId)
-{
-	MARK_PROPERTY_DIRTY_FROM_NAME(APlayerState, UniqueId, this);
-	UniqueId.SetUniqueNetId(InUniqueId);
-}
-PRAGMA_ENABLE_DEPRECATION_WARNINGS
-
 void APlayerState::RegisterPlayerWithSession(bool bWasFromInvite)
 {
 	if (GetNetMode() != NM_Standalone)
@@ -522,15 +527,35 @@ void APlayerState::SetStartTime(const int32 NewStartTime)
 	StartTime = NewStartTime;
 }
 
+FUniqueNetIdRepl APlayerState::BP_GetUniqueId() const
+{
+	return GetUniqueId();
+}
+
+void APlayerState::SetUniqueId(const FUniqueNetIdPtr& InUniqueId)
+{
+	MARK_PROPERTY_DIRTY_FROM_NAME(APlayerState, UniqueId, this);
+	UniqueId.SetUniqueNetId(InUniqueId);
+	OnSetUniqueId();
+}
+
 void APlayerState::SetUniqueId(const FUniqueNetIdRepl& NewUniqueId)
 {
 	MARK_PROPERTY_DIRTY_FROM_NAME(APlayerState, UniqueId, this);
 	UniqueId = NewUniqueId;
+	OnSetUniqueId();
 }
 
 void APlayerState::SetUniqueId(FUniqueNetIdRepl&& NewUniqueId)
 {
 	MARK_PROPERTY_DIRTY_FROM_NAME(APlayerState, UniqueId, this);
 	UniqueId = MoveTemp(NewUniqueId);
+	OnSetUniqueId();
 }
-PRAGMA_ENABLE_DEPRECATION_WARNINGS
\ No newline at end of file
+
+void APlayerState::OnSetUniqueId()
+{
+
+}
+
+PRAGMA_ENABLE_DEPRECATION_WARNINGS
