Fix the mismatch in data allocated via FMipAllocation (allocated with new [] operator but deleted via the ::Free function)

#rnx

- Changed ReadWriteBuffer to be allocated with FMemory::Malloc and freed via FMemory::Free
- Removed unused code paths and constructors that were used when we supported both old and new style bulkdata.
- There remains the issue that the calling code could request readwrite buffer via ::GetDataReadWrite and then request a read only FSharedBuffer via ::GetDataReadOnly() and effectively change the data that the FSharedBuffer represent via the pointer.
-- This is supposed to prevented by the higher level calling code in FTextureSource and is a pre-existing problem that I am choosing not to address at this point.

diff --git a/Engine/Source/Runtime/Engine/Classes/Engine/Texture.h b/Engine/Source/Runtime/Engine/Classes/Engine/Texture.h
--- a/Engine/Source/Runtime/Engine/Classes/Engine/Texture.h
+++ b/Engine/Source/Runtime/Engine/Classes/Engine/Texture.h
@@ -343,10 +343,6 @@ struct FTextureSource
 		FMipAllocation() = default;
 		/** Take a read only FSharedBuffer, will allocate a new buffer and copy from this if Read/Write access is requested */
 		FMipAllocation(FSharedBuffer SrcData);
-		/** Allocate a new buffer and copy from the source data */
-		FMipAllocation(const void* SrcData, int64 DataLength);
-		/** Lock the bulkdata and use it's memory buffer */
-		FMipAllocation(FByteBulkData& BulkData);
 
 		// Do not actually do anything for copy constructor or assignments, this is required for as long as
 		// we need to support the old bulkdata code path (although previously storing these allocations as
@@ -382,9 +378,19 @@ struct FTextureSource
 	private:
 		void CreateReadWriteBuffer(const void* SrcData, int64 DataLength);
 
-		FSharedBuffer		ReadOnlyReference;
-		TUniquePtr<uint8[]>	ReadWriteBuffer;
-		uint8* BulkDataPtr{ nullptr };
+		struct FDeleterFree
+		{
+			void operator()(uint8* Ptr) const
+			{
+				if (Ptr)
+				{
+					FMemory::Free(Ptr);
+				}
+			}
+		};
+
+		FSharedBuffer					ReadOnlyReference;
+		TUniquePtr<uint8, FDeleterFree>	ReadWriteBuffer;
 	};
 
 	/** Structure that encapsulates the decompressed texture data and can be accessed per mip */
diff --git a/Engine/Source/Runtime/Engine/Private/Texture.cpp b/Engine/Source/Runtime/Engine/Private/Texture.cpp
--- a/Engine/Source/Runtime/Engine/Private/Texture.cpp
+++ b/Engine/Source/Runtime/Engine/Private/Texture.cpp
@@ -2289,19 +2289,6 @@ FTextureSource::FMipAllocation::FMipAllocation(FSharedBuffer SrcData)
 {
 }
 
-FTextureSource::FMipAllocation::FMipAllocation(const void* SrcData, int64 DataLength)
-{
-	CreateReadWriteBuffer(SrcData, DataLength);
-}
-
-FTextureSource::FMipAllocation::FMipAllocation(FByteBulkData& BulkData)
-{
-	// Note that we DO NOT unlock the bulkdata via FMipAllocation as some areas of code keep the lock open for longer
-	// than the memory buffer is actually used.
-	BulkDataPtr = (uint8*)BulkData.Lock(LOCK_READ_WRITE);
-	ReadOnlyReference = FSharedBuffer::MakeView(BulkDataPtr, BulkData.GetBulkDataSize());
-}
-
 FTextureSource::FMipAllocation::FMipAllocation(FTextureSource::FMipAllocation&& Other)
 {
 	*this = MoveTemp(Other);
@@ -2312,9 +2299,6 @@ FTextureSource::FMipAllocation& FTextureSource::FMipAllocation::operator =(FText
 	ReadOnlyReference = MoveTemp(Other.ReadOnlyReference);
 	ReadWriteBuffer = MoveTemp(Other.ReadWriteBuffer);
 
-	BulkDataPtr = Other.BulkDataPtr;
-	Other.BulkDataPtr = nullptr;
-
 	return *this;
 }
 
@@ -2322,17 +2306,11 @@ void FTextureSource::FMipAllocation::Reset()
 {
 	ReadOnlyReference.Reset();
 	ReadWriteBuffer = nullptr;
-	BulkDataPtr = nullptr;
 }
 
 uint8* FTextureSource::FMipAllocation::GetDataReadWrite()
 {
-	if (BulkDataPtr != nullptr)
-	{
-		return BulkDataPtr;
-	}
-
-	if (ReadWriteBuffer == nullptr)
+	if (!ReadWriteBuffer.IsValid())
 	{
 		CreateReadWriteBuffer(ReadOnlyReference.GetData(), ReadOnlyReference.GetSize());
 	}
@@ -2342,8 +2320,7 @@ uint8* FTextureSource::FMipAllocation::GetDataReadWrite()
 
 FSharedBuffer FTextureSource::FMipAllocation::Release()
 {
-	check(BulkDataPtr == nullptr); // Should not be called if we are using old bulkdata
-	if (ReadWriteBuffer != nullptr)
+	if (ReadWriteBuffer.IsValid())
 	{
 		const int64 DataSize = ReadOnlyReference.GetSize();
 		ReadOnlyReference.Reset();
@@ -2357,7 +2334,7 @@ FSharedBuffer FTextureSource::FMipAllocation::Release()
 
 void FTextureSource::FMipAllocation::CreateReadWriteBuffer(const void* SrcData, int64 DataLength)
 {
-	ReadWriteBuffer = MakeUnique<uint8[]>(DataLength);
+	ReadWriteBuffer = TUniquePtr<uint8, FDeleterFree>((uint8*)FMemory::Malloc(DataLength));
 	FMemory::Memcpy(ReadWriteBuffer.Get(), SrcData, DataLength);
 
 	ReadOnlyReference = FSharedBuffer::MakeView(ReadWriteBuffer.Get(), DataLength);
