Added support for spawning Level Instances in cooked builds.

diff --git a/Engine/Source/Runtime/Engine/Private/LevelInstance/LevelInstanceActor.cpp b/Engine/Source/Runtime/Engine/Private/LevelInstance/LevelInstanceActor.cpp
--- a/Engine/Source/Runtime/Engine/Private/LevelInstance/LevelInstanceActor.cpp
+++ b/Engine/Source/Runtime/Engine/Private/LevelInstance/LevelInstanceActor.cpp
@@ -67,6 +67,7 @@ void ALevelInstance::Serialize(FArchive& Ar)
 		else if (Ar.IsPersistent())
 		{
 			Ar << LevelInstanceActorGuid;
+			check(LevelInstanceActorGuid.IsValid());
 		}
 	}
 #endif
@@ -78,6 +79,13 @@ void ALevelInstance::PostRegisterAllComponents()
 
 	if (ULevelInstanceSubsystem* LevelInstanceSubsystem = GetLevelInstanceSubsystem())
 	{
+#if !WITH_EDITOR
+		// If the level instance was spawned, not loaded
+		if (!LevelInstanceActorGuid.IsValid())
+		{
+			LevelInstanceActorGuid = FGuid::NewGuid();
+		}
+#endif
 		LevelInstanceID = LevelInstanceSubsystem->RegisterLevelInstance(this);
 
 		LoadLevelInstance();
