Fix for Audio memory regression w/ USoundWave data.

diff --git a/Engine/Source/Runtime/Engine/Private/AudioStreamingCache.cpp b/Engine/Source/Runtime/Engine/Private/AudioStreamingCache.cpp
--- a/Engine/Source/Runtime/Engine/Private/AudioStreamingCache.cpp
+++ b/Engine/Source/Runtime/Engine/Private/AudioStreamingCache.cpp
@@ -235,24 +235,24 @@ static FAutoConsoleCommand GDisableProfilingAudioCacheCommand(
 
 
 FAudioChunkCache::FChunkKey::FChunkKey(const FChunkKey& Other)
-	: SoundWaveProxyPtr(Other.SoundWaveProxyPtr)
+	: SoundWaveProxyWeakPtr(Other.SoundWaveProxyWeakPtr)
 	, ChunkIndex(Other.ChunkIndex)
 	, ObjectKey(Other.ObjectKey)
 #if WITH_EDITOR
 	, ChunkRevision(Other.ChunkRevision)
 #endif // #if WITH_EDITOR
 {
-	if (SoundWaveProxyPtr.IsValid())
+	if (FSoundWaveProxyPtr ProxyPtr = SoundWaveProxyWeakPtr.Pin())
 	{
-		SoundWaveName = SoundWaveProxyPtr->GetFName();
-		SoundWaveProxyPtr->ReleaseCompressedAudio();
+		SoundWaveName = ProxyPtr->GetFName();
+		ProxyPtr->ReleaseCompressedAudio();
 	}
 }
 
 
 FAudioChunkCache::FChunkKey& FAudioChunkCache::FChunkKey::operator=(const FChunkKey& Other)
 {
-	SoundWaveProxyPtr = Other.SoundWaveProxyPtr;
+	SoundWaveProxyWeakPtr = Other.SoundWaveProxyWeakPtr;
 	SoundWaveName = Other.SoundWaveName;
 	ChunkIndex = Other.ChunkIndex;
 	ObjectKey = Other.ObjectKey;
@@ -271,24 +271,19 @@ FAudioChunkCache::FChunkKey::FChunkKey(const FSoundWaveProxyPtr& InSoundWave, ui
 #endif // #if WITH_EDITOR
 )
 
-	: SoundWaveProxyPtr(InSoundWave)
+	: SoundWaveProxyWeakPtr(InSoundWave)
 	, ChunkIndex(InChunkIndex)
 #if WITH_EDITOR
 	, ChunkRevision(InChunkRevision)
 #endif // #if WITH_EDITOR
 {
-	if (FSoundWaveProxyPtr ProxyPtr = SoundWaveProxyPtr)
+	if (FSoundWaveProxyPtr ProxyPtr = SoundWaveProxyWeakPtr.Pin())
 	{
 		SoundWaveName = ProxyPtr->GetFName();
 		ObjectKey = ProxyPtr->GetFObjectKey();
 	}
 }
 
-FAudioChunkCache::FChunkKey::FChunkKey(FName InSoundWaveName, uint32 InChunkIndex)
-	: SoundWaveName(InSoundWaveName)
-	, ChunkIndex(InChunkIndex)
-{
-}
 
 bool FAudioChunkCache::FChunkKey::operator==(const FChunkKey& Other) const
 {
@@ -313,26 +308,34 @@ bool FAudioChunkCache::FChunkKey::operator==(const FChunkKey& Other) const
 
 bool FAudioChunkCache::FChunkKey::IsChunkStale()
 {
-	FSoundWaveProxyPtr ProxyPtr = SoundWaveProxyPtr;
-
-	return !ProxyPtr.IsValid()
+	if(FSoundWaveProxyPtr ProxyPtr = SoundWaveProxyWeakPtr.Pin())
+	{
+		return !ProxyPtr.IsValid()
 #if WITH_EDITOR
 		|| (ProxyPtr->GetCurrentChunkRevision() != ChunkRevision)
 #endif // #if WITH_EDITOR
 		;
+
+	}
+
+	return true;
 }
 
 ESoundWaveLoadingBehavior FAudioChunkCache::FChunkKey::GetLoadingBehavior() const
 {
-	FSoundWaveProxyPtr ProxyPtr = SoundWaveProxyPtr;
-	check(ProxyPtr.IsValid());
-	return ProxyPtr->GetLoadingBehavior();
+	if(FSoundWaveProxyPtr ProxyPtr = SoundWaveProxyWeakPtr.Pin())
+	{
+		return ProxyPtr->GetLoadingBehavior();
+	}
+
+	return ESoundWaveLoadingBehavior::Uninitialized;
 }
 
 FStreamedAudioChunk& FAudioChunkCache::FChunkKey::GetChunk(uint32 InChunkIndex) const
 {
-	FSoundWaveProxyPtr ProxyPtr = SoundWaveProxyPtr;
+	FSoundWaveProxyPtr ProxyPtr = SoundWaveProxyWeakPtr.Pin();
 
+	// the Weakptr should be valid here since it's from a shared ptr up the stack
 	check(ProxyPtr.IsValid());
 	check((ChunkIndex < (uint32)ProxyPtr->GetNumChunks()));
 
@@ -347,34 +350,41 @@ FStreamedAudioChunk& FAudioChunkCache::FChunkKey::GetChunk(uint32 InChunkIndex)
 
 uint32 FAudioChunkCache::FChunkKey::GetNumChunks() const
 {
-	FSoundWaveProxyPtr ProxyPtr = SoundWaveProxyPtr;
-
-	check(ProxyPtr.IsValid());
-	return ProxyPtr->GetNumChunks();
+	if (FSoundWaveProxyPtr ProxyPtr = SoundWaveProxyWeakPtr.Pin())
+	{
+		return ProxyPtr->GetNumChunks();
+	}
+	
+	return 0;
 }
 
 bool FAudioChunkCache::FChunkKey::IsRetainingAudio() const
 {
-	FSoundWaveProxyPtr ProxyPtr = SoundWaveProxyPtr;
+	if (FSoundWaveProxyPtr ProxyPtr = SoundWaveProxyWeakPtr.Pin())
+	{
+		return ProxyPtr->IsRetainingAudio();
+	}
 
-	check(ProxyPtr.IsValid());
-	return ProxyPtr->IsRetainingAudio();
+	return false;
 }
 
 void FAudioChunkCache::FChunkKey::ReleaseCompressedAudio()
 {
-	FSoundWaveProxyPtr ProxyPtr = SoundWaveProxyPtr;
-
-	check(ProxyPtr.IsValid());
-	ProxyPtr->ReleaseCompressedAudio();
+	if (FSoundWaveProxyPtr ProxyPtr = SoundWaveProxyWeakPtr.Pin())
+	{
+		check(ProxyPtr.IsValid());
+		ProxyPtr->ReleaseCompressedAudio();
+	}
 }
 
 bool FAudioChunkCache::FChunkKey::WasLoadingBehaviorOverridden() const
 {
-	FSoundWaveProxyPtr ProxyPtr = SoundWaveProxyPtr;
+	if(FSoundWaveProxyPtr ProxyPtr = SoundWaveProxyWeakPtr.Pin())
+	{
+		return ProxyPtr->WasLoadingBehaviorOverridden();
+	}
 
-	check(ProxyPtr.IsValid());
-	return ProxyPtr->WasLoadingBehaviorOverridden();
+	return false;
 }
 
 FCachedAudioStreamingManager::FCachedAudioStreamingManager(const FCachedAudioStreamingManagerParams& InitParams)
@@ -1012,6 +1022,7 @@ uint64 FAudioChunkCache::TrimMemory(uint64 BytesToFree, bool bInAllowRetainedChu
 
 			CurrentElement->ChunkData = nullptr;
 			CurrentElement->ChunkDataSize = 0;
+			CacheLookupIdMap.Remove(CurrentElement->Key);
 			CurrentElement->Key = FChunkKey();
 
 #if DEBUG_STREAM_CACHE
@@ -1052,6 +1063,7 @@ uint64 FAudioChunkCache::TrimMemory(uint64 BytesToFree, bool bInAllowRetainedChu
 					FMemory::Free(CurrentElement->ChunkData);
 					CurrentElement->ChunkData = nullptr;
 					CurrentElement->ChunkDataSize = 0;
+					CacheLookupIdMap.Remove(CurrentElement->Key);
 					CurrentElement->Key = FChunkKey();
 
 #if DEBUG_STREAM_CACHE
@@ -1151,7 +1163,7 @@ FString FAudioChunkCache::FlushCacheMissLog()
 		}
 	};
 
-	TMap<FChunkKey, int32> CacheMissCount;
+	TMap<FCacheMissEntry, int32> CacheMissCount;
 
 	TQueue<FCacheMissInfo> BackupQueue;
 
@@ -1163,9 +1175,9 @@ FString FAudioChunkCache::FlushCacheMissLog()
 		ConcatenatedCacheMisses.AppendInt(CacheMissInfo.ChunkIndex);
 		ConcatenatedCacheMisses.Append(TEXT("\n"));
 
-		FChunkKey Chunk(CacheMissInfo.SoundWaveName, CacheMissInfo.ChunkIndex);
+		FCacheMissEntry CacheMissEntry(CacheMissInfo.SoundWaveName, CacheMissInfo.ChunkIndex);
 
-		int32& MissCount = CacheMissCount.FindOrAdd(Chunk);
+		int32& MissCount = CacheMissCount.FindOrAdd(CacheMissEntry);
 		MissCount++;
 
 		if (KeepCacheMissBufferOnFlushCVar)
diff --git a/Engine/Source/Runtime/Engine/Public/AudioStreamingCache.h b/Engine/Source/Runtime/Engine/Public/AudioStreamingCache.h
--- a/Engine/Source/Runtime/Engine/Public/AudioStreamingCache.h
+++ b/Engine/Source/Runtime/Engine/Public/AudioStreamingCache.h
@@ -36,9 +36,6 @@ public:
 
 		FChunkKey(const FChunkKey& Other);
 
-		// used when flushing the cache miss log
-		FChunkKey(FName InSoundWaveName, uint32 InChunkIndex);
-
 		FChunkKey(
 			  const FSoundWaveProxyPtr& InSoundWave
 			, uint32 InChunkIndex
@@ -51,7 +48,7 @@ public:
 
 
 	private:
-		mutable FSoundWaveProxyPtr SoundWaveProxyPtr{ nullptr };
+		mutable TWeakPtr<FSoundWaveProxy, ESPMode::ThreadSafe> SoundWaveProxyWeakPtr;
 
 	public:
 		mutable FName SoundWaveName = FName(); // mutable for GetTypeHash()
@@ -65,7 +62,7 @@ public:
 		inline bool operator==(const FChunkKey& Other) const;
 
 		bool IsChunkStale();
-		bool IsSoundWaveValid() const { return SoundWaveProxyPtr.IsValid(); }
+		bool IsSoundWaveValid() const { return SoundWaveProxyWeakPtr.IsValid(); }
 		// Forward FSoundWaveProxy interface
 		ESoundWaveLoadingBehavior GetLoadingBehavior() const;
 		FStreamedAudioChunk& GetChunk(uint32 InChunkIndex) const;
@@ -81,6 +78,31 @@ public:
 		}
 	};
 
+	struct FCacheMissEntry
+	{
+	public:
+		// ctor
+		FCacheMissEntry(FName InSoundWaveName, uint32 InChunkIndex)
+			: SoundWaveName(InSoundWaveName)
+			, ChunkIndex(InChunkIndex)
+		{}
+
+		inline bool operator==(const FCacheMissEntry& Other) const
+		{
+			return (SoundWaveName == Other.SoundWaveName) && (ChunkIndex == Other.ChunkIndex);
+		}
+
+		FName SoundWaveName;
+
+		uint32 ChunkIndex;
+
+		/** Hash function */
+		friend uint32 GetTypeHash(const FCacheMissEntry& InCacheMissEntry)
+		{
+			return HashCombine(InCacheMissEntry.SoundWaveName.GetNumber(), InCacheMissEntry.ChunkIndex);
+		}
+	};
+
 	FAudioChunkCache(uint32 InMaxChunkSize, uint32 NumChunks, uint64 InMemoryLimitInBytes);
 	
 	~FAudioChunkCache();
