Morph target buffer memory optimization
1) Create VertexIndicesVB as 16bit buffer instead of 32bit when max vertex index doesn't exceed 65536
2) Create MorphDeltaVB  and VertexIndicesVB buffer sizes with TArray element count rather than its max capacity

diff --git a/Engine/Source/Runtime/Engine/Private/Rendering/MorphTargetVertexInfoBuffers.cpp b/Engine/Source/Runtime/Engine/Private/Rendering/MorphTargetVertexInfoBuffers.cpp
--- a/Engine/Source/Runtime/Engine/Private/Rendering/MorphTargetVertexInfoBuffers.cpp
+++ b/Engine/Source/Runtime/Engine/Private/Rendering/MorphTargetVertexInfoBuffers.cpp
@@ -11,17 +11,31 @@ void FMorphTargetVertexInfoBuffers::InitRHI()
 
 	{
 		FRHIResourceCreateInfo CreateInfo(TEXT("VertexIndicesVB"));
-		VertexIndicesVB = RHICreateBuffer(VertexIndices.GetAllocatedSize(), BUF_Static | BUF_VertexBuffer | BUF_ShaderResource, 0, ERHIAccess::VertexOrIndexBuffer | ERHIAccess::SRVMask, CreateInfo);
-		void* VertexIndicesVBData = RHILockBuffer(VertexIndicesVB, 0, VertexIndices.GetAllocatedSize(), RLM_WriteOnly);
-		FMemory::ParallelMemcpy(VertexIndicesVBData, VertexIndices.GetData(), VertexIndices.GetAllocatedSize(), EMemcpyCachePolicy::StoreUncached);
-		RHIUnlockBuffer(VertexIndicesVB);
-		VertexIndicesSRV = RHICreateShaderResourceView(VertexIndicesVB, 4, PF_R32_UINT);
+		if (bMaxVertexIndex32Bits)
+		{
+			VertexIndicesVB = RHICreateBuffer(VertexIndices.Num() * sizeof(uint32), BUF_Static | BUF_VertexBuffer | BUF_ShaderResource, 0, ERHIAccess::VertexOrIndexBuffer | ERHIAccess::SRVMask, CreateInfo);
+			void* VertexIndicesVBData = RHILockBuffer(VertexIndicesVB, 0, VertexIndicesVB->GetSize(), RLM_WriteOnly);
+			FMemory::ParallelMemcpy(VertexIndicesVBData, VertexIndices.GetData(), VertexIndices.Num() * sizeof(uint32), EMemcpyCachePolicy::StoreUncached);
+			RHIUnlockBuffer(VertexIndicesVB);
+			VertexIndicesSRV = RHICreateShaderResourceView(VertexIndicesVB, sizeof(uint32), PF_R32_UINT);
+		}
+		else
+		{
+			VertexIndicesVB = RHICreateBuffer(VertexIndices.Num() * sizeof(uint16), BUF_Static | BUF_VertexBuffer | BUF_ShaderResource, 0, ERHIAccess::VertexOrIndexBuffer | ERHIAccess::SRVMask, CreateInfo);
+			uint16* VertexIndicesVBData = (uint16*)RHILockBuffer(VertexIndicesVB, 0, VertexIndicesVB->GetSize(), RLM_WriteOnly);
+			for (int32 i = 0; i < VertexIndices.Num(); ++i)
+			{
+				VertexIndicesVBData[i] = (uint16)VertexIndices[i];
+			}
+			RHIUnlockBuffer(VertexIndicesVB);
+			VertexIndicesSRV = RHICreateShaderResourceView(VertexIndicesVB, sizeof(uint16), PF_R16_UINT);
+		}
 	}
 	{
 		FRHIResourceCreateInfo CreateInfo(TEXT("MorphDeltasVB"));
-		MorphDeltasVB = RHICreateBuffer(MorphDeltas.GetAllocatedSize(), BUF_Static | BUF_VertexBuffer | BUF_ShaderResource, 0, ERHIAccess::VertexOrIndexBuffer | ERHIAccess::SRVMask, CreateInfo);
-		void* MorphDeltasVBData = RHILockBuffer(MorphDeltasVB, 0, MorphDeltas.GetAllocatedSize(), RLM_WriteOnly);
-		FMemory::ParallelMemcpy(MorphDeltasVBData, MorphDeltas.GetData(), MorphDeltas.GetAllocatedSize(), EMemcpyCachePolicy::StoreUncached);
+		MorphDeltasVB = RHICreateBuffer(MorphDeltas.Num() * MorphDeltas.GetTypeSize(), BUF_Static | BUF_VertexBuffer | BUF_ShaderResource, 0, ERHIAccess::VertexOrIndexBuffer | ERHIAccess::SRVMask, CreateInfo);
+		void* MorphDeltasVBData = RHILockBuffer(MorphDeltasVB, 0, MorphDeltasVB->GetSize(), RLM_WriteOnly);
+		FMemory::ParallelMemcpy(MorphDeltasVBData, MorphDeltas.GetData(), MorphDeltas.Num() * MorphDeltas.GetTypeSize(), EMemcpyCachePolicy::StoreUncached);
 		RHIUnlockBuffer(MorphDeltasVB);
 		MorphDeltasSRV = RHICreateShaderResourceView(MorphDeltasVB, 2, PF_R16F);
 	}
diff --git a/Engine/Source/Runtime/Engine/Private/SkeletalMeshLODRenderData.cpp b/Engine/Source/Runtime/Engine/Private/SkeletalMeshLODRenderData.cpp
--- a/Engine/Source/Runtime/Engine/Private/SkeletalMeshLODRenderData.cpp
+++ b/Engine/Source/Runtime/Engine/Private/SkeletalMeshLODRenderData.cpp
@@ -254,6 +254,8 @@ void FSkeletalMeshLODRenderData::InitResources(bool bNeedsVertexColors, int32 LO
 #endif
 		}
 
+		MorphTargetVertexInfoBuffers.bMaxVertexIndex32Bits = MaxVertexIndex > 65535;
+
 		check(MorphTargetVertexInfoBuffers.WorkItemsPerMorph.Num() == MorphTargetVertexInfoBuffers.StartOffsetPerMorph.Num());
 		check(MorphTargetVertexInfoBuffers.WorkItemsPerMorph.Num() == MorphTargetVertexInfoBuffers.MaximumValuePerMorph.Num());
 		check(MorphTargetVertexInfoBuffers.WorkItemsPerMorph.Num() == MorphTargetVertexInfoBuffers.MinimumValuePerMorph.Num());
diff --git a/Engine/Source/Runtime/Engine/Public/Rendering/MorphTargetVertexInfoBuffers.h b/Engine/Source/Runtime/Engine/Public/Rendering/MorphTargetVertexInfoBuffers.h
--- a/Engine/Source/Runtime/Engine/Public/Rendering/MorphTargetVertexInfoBuffers.h
+++ b/Engine/Source/Runtime/Engine/Public/Rendering/MorphTargetVertexInfoBuffers.h
@@ -84,6 +84,7 @@ public:
 	void Reset()
 	{
 		VertexIndices.Empty();
+		bMaxVertexIndex32Bits = false;
 		MorphDeltas.Empty();
 		MaximumValuePerMorph.Empty();
 		MinimumValuePerMorph.Empty();
@@ -98,6 +99,7 @@ protected:
 
 	// Transient data used while creating the vertex buffers, gets deleted as soon as VB gets initialized.
 	TArray<uint32> VertexIndices;
+	bool bMaxVertexIndex32Bits = false;
 	// Transient data used while creating the vertex buffers, gets deleted as soon as VB gets initialized.
 	TArray<FMorphDelta> MorphDeltas;
 
