Clarify tooltip on function

diff --git a/Engine/Source/Runtime/Engine/Classes/Kismet/GameplayStatics.h b/Engine/Source/Runtime/Engine/Classes/Kismet/GameplayStatics.h
--- a/Engine/Source/Runtime/Engine/Classes/Kismet/GameplayStatics.h
+++ b/Engine/Source/Runtime/Engine/Classes/Kismet/GameplayStatics.h
@@ -529,8 +529,8 @@ public:
 	
 	/**
 	 * Determines if any audio listeners are within range of the specified location
-	 * @param Location		The location to potentially play a sound at
-	 * @param MaximumRange	The maximum distance away from Location that a listener can be
+	 * @param Location		The location from which test if a listener is in range
+	 * @param MaximumRange	The distance away from Location to test if any listener is within
 	 * @note This will always return false if there is no audio device, or the audio device is disabled.
 	 */
 	UFUNCTION(BlueprintCallable, Category = "Audio", meta = (WorldContext = "WorldContextObject"))
