Allowing r.Shaders.Optimize to be set per-platform.

diff --git a/Engine/Source/Runtime/Engine/Private/ShaderCompiler/ShaderCompiler.cpp b/Engine/Source/Runtime/Engine/Private/ShaderCompiler/ShaderCompiler.cpp
--- a/Engine/Source/Runtime/Engine/Private/ShaderCompiler/ShaderCompiler.cpp
+++ b/Engine/Source/Runtime/Engine/Private/ShaderCompiler/ShaderCompiler.cpp
@@ -882,7 +882,8 @@ static TAutoConsoleVariable<int32> CVarShadersEnableExtraData(
 static TAutoConsoleVariable<int32> CVarOptimizeShaders(
 	TEXT("r.Shaders.Optimize"),
 	1,
-	TEXT("Whether to optimize shaders.  When using graphical debuggers like Nsight it can be useful to disable this on startup."),
+	TEXT("Whether to optimize shaders.  When using graphical debuggers like Nsight it can be useful to disable this on startup.\n")
+	TEXT("This setting can be overriden in any Engine.ini under the [ShaderCompiler] section."),
 	ECVF_ReadOnly);
 
 static TAutoConsoleVariable<int32> CVarShaderFastMath(
@@ -5156,8 +5157,7 @@ void GlobalBeginCompileShader(
 		// Check if the compile environment explicitly wants to force optimization
 		const bool bForceOptimization = Input.Environment.CompilerFlags.Contains(CFLAG_ForceOptimization);
 
-		static const auto CVar = IConsoleManager::Get().FindConsoleVariable(TEXT("r.Shaders.Optimize"));
-		if (!bForceOptimization && CVar && CVar->GetInt() == 0)
+		if (!bForceOptimization && !ShouldOptimizeShaders(ShaderFormatName))
 		{
 			Input.Environment.CompilerFlags.Add(CFLAG_Debug);
 		}
diff --git a/Engine/Source/Runtime/RenderCore/Private/Shader.cpp b/Engine/Source/Runtime/RenderCore/Private/Shader.cpp
--- a/Engine/Source/Runtime/RenderCore/Private/Shader.cpp
+++ b/Engine/Source/Runtime/RenderCore/Private/Shader.cpp
@@ -1441,8 +1441,7 @@ void ShaderMapAppendKeyString(EShaderPlatform Platform, FString& KeyString)
 	}
 
 	{
-		static const auto CVar = IConsoleManager::Get().FindConsoleVariable(TEXT("r.Shaders.Optimize"));
-		KeyString += (CVar && CVar->GetInt() != 0) ? TEXT("") : TEXT("_NoOpt");
+		KeyString += ShouldOptimizeShaders(ShaderFormatName) ? TEXT("") : TEXT("_NoOpt");
 	}
 	
 	{
diff --git a/Engine/Source/Runtime/RenderCore/Private/ShaderCore.cpp b/Engine/Source/Runtime/RenderCore/Private/ShaderCore.cpp
--- a/Engine/Source/Runtime/RenderCore/Private/ShaderCore.cpp
+++ b/Engine/Source/Runtime/RenderCore/Private/ShaderCore.cpp
@@ -446,6 +446,12 @@ bool ShouldEnableExtraShaderData(FName ShaderFormat)
 	return ExtraData.IsEnabled(ShaderFormat);
 }
 
+bool ShouldOptimizeShaders(FName ShaderFormat)
+{
+	static const FShaderSymbolSettingHelper Optimize(TEXT("r.Shaders.Optimize"));
+	return Optimize.IsEnabled(ShaderFormat);
+}
+
 #ifndef UE_ALLOW_SHADER_COMPILING
 // is shader compiling allowed at all? (set to 0 in a cooked editor .Target.cs if the target has no shaders available)
 #define UE_ALLOW_SHADER_COMPILING 1
diff --git a/Engine/Source/Runtime/RenderCore/Public/ShaderCompilerCore.h b/Engine/Source/Runtime/RenderCore/Public/ShaderCompilerCore.h
--- a/Engine/Source/Runtime/RenderCore/Public/ShaderCompilerCore.h
+++ b/Engine/Source/Runtime/RenderCore/Public/ShaderCompilerCore.h
@@ -40,7 +40,9 @@ extern RENDERCORE_API bool ShouldAllowUniqueShaderSymbols(FName ShaderFormat);
 extern RENDERCORE_API bool ShouldWriteShaderSymbolsAsZip(FName ShaderFormat);
 
 /** Returns true if the user wants more runtime shader data (names, extra info) */
-extern RENDERCORE_API bool ShouldEnableExtraShaderData(FName Platform);
+extern RENDERCORE_API bool ShouldEnableExtraShaderData(FName ShaderFormat);
+
+extern RENDERCORE_API bool ShouldOptimizeShaders(FName ShaderFormat);
 
 UE_DEPRECATED(5.0, "ShouldGenerateShaderSymbols should be called to determine if symbols (debug data) should be generated")
 bool ShouldKeepShaderDebugInfo(EShaderPlatform Platform);
