Duplicate LevelScriptActor issue: improved logging to identify the package / map that the LSAs reside in

#rnx

diff --git a/Engine/Source/Runtime/Engine/Private/LevelScriptActor.cpp b/Engine/Source/Runtime/Engine/Private/LevelScriptActor.cpp
--- a/Engine/Source/Runtime/Engine/Private/LevelScriptActor.cpp
+++ b/Engine/Source/Runtime/Engine/Private/LevelScriptActor.cpp
@@ -46,8 +46,8 @@ ALevelScriptActor::ALevelScriptActor(const FObjectInitializer& ObjectInitializer
 			bool bHasNewerClassVersion = Sibling->GetClass()->HasAnyClassFlags(CLASS_NewerVersionExists);
 
 			ensureMsgf((bIsNotAnLSA || bIsTheSameObject || bHasNewerClassVersion),
-				TEXT("Detected the creation of more than one LevelScriptActor (%s, %s) within the same outer (%s). This can lead to duplicate level blueprint operations during play."),
-				*GetName(), *Sibling->GetName(), *ThisOuter->GetName());
+				TEXT("Detected the creation of more than one LevelScriptActor (%s, %s) within the same outer (%s) in package (%s). This can lead to duplicate level blueprint operations during play."),
+				*GetName(), *Sibling->GetName(), *ThisOuter->GetName(), *ThisOuter->GetPackage()->GetName());
 		}
 	}
 #endif // WITH_EDITOR
