Fix the issue that deleting cpu morph data in cooked build results in failing to select the morph target vertex factory

diff --git a/Engine/Source/Runtime/Engine/Classes/Animation/MorphTarget.h b/Engine/Source/Runtime/Engine/Classes/Animation/MorphTarget.h
--- a/Engine/Source/Runtime/Engine/Classes/Animation/MorphTarget.h
+++ b/Engine/Source/Runtime/Engine/Classes/Animation/MorphTarget.h
@@ -53,6 +53,9 @@ struct FMorphTargetLODModel
 	/** vertex data for a single LOD morph mesh */
 	TArray<FMorphTargetDelta> Vertices;
 
+	/** Number of elements in Vertices array. This property is set at runtime and is not serialized. */
+	int32 NumVertices;
+
 	/** number of original verts in the base mesh */
 	int32 NumBaseMeshVerts;
 	
@@ -66,7 +69,8 @@ struct FMorphTargetLODModel
 	void GetResourceSizeEx(FResourceSizeEx& CumulativeResourceSize) const;
 
 	FMorphTargetLODModel()
-		: NumBaseMeshVerts(0)
+		: NumVertices(0)
+		, NumBaseMeshVerts(0)
 		, bGeneratedByEngine(false)
 	{ }
 
@@ -94,12 +98,16 @@ struct FMorphTargetLODModel
 			}
 		}
 
+		// Cache Vertices array size, no need to serialize CachedNumVertices value.
+		M.NumVertices = M.Vertices.Num();
+
 		return Ar;
 	}
 
 	void Reset()
 	{
 		Vertices.Reset();
+		NumVertices = 0;
 		NumBaseMeshVerts = 0;
 		SectionIndices.Reset();
 		// since engine cleared it, we mark as engine generated
@@ -107,10 +115,12 @@ struct FMorphTargetLODModel
 		bGeneratedByEngine = true;
 	}
 
-	void DiscardCPUBuffers()
+	void DiscardVertexData()
 	{
+		// Vertices array can be discarded after render data is initialized and if CPU data is no longer required,
+		// but cache the array size so that UMorphTarget::HasValidData() can still query if morph target is valid. 
+		NumVertices = Vertices.Num();
 		Vertices.Empty();
-		SectionIndices.Empty();
 	}
 };
 
@@ -139,7 +149,7 @@ public:
 	ENGINE_API bool HasValidData() const;
 
 	/** Discard CPU Buffers after render resources have been created. */
-	ENGINE_API void DiscardCPUBuffers();
+	ENGINE_API void DiscardVertexData();
 
 #if WITH_EDITOR
 	/** Populates the given morph target LOD model with the provided deltas */
diff --git a/Engine/Source/Runtime/Engine/Private/MorphTools.cpp b/Engine/Source/Runtime/Engine/Private/MorphTools.cpp
--- a/Engine/Source/Runtime/Engine/Private/MorphTools.cpp
+++ b/Engine/Source/Runtime/Engine/Private/MorphTools.cpp
@@ -36,7 +36,12 @@ FMorphTargetDelta* UMorphTarget::GetMorphTargetDelta(int32 LODIndex, int32& OutN
 bool UMorphTarget::HasDataForLOD(int32 LODIndex) 
 {
 	// If we have an entry for this LOD, and it has verts
+#if WITH_EDITOR
 	return (MorphLODModels.IsValidIndex(LODIndex) && MorphLODModels[LODIndex].Vertices.Num() > 0);
+#else
+	// Morph target's vertices array could have been emptied after render data is created, so check NumVertices instead
+	return (MorphLODModels.IsValidIndex(LODIndex) && MorphLODModels[LODIndex].NumVertices > 0);
+#endif
 }
 
 bool UMorphTarget::HasValidData() const
@@ -52,11 +57,11 @@ bool UMorphTarget::HasValidData() const
 	return false;
 }
 
-void UMorphTarget::DiscardCPUBuffers()
+void UMorphTarget::DiscardVertexData()
 {
 	for (FMorphTargetLODModel& Model : MorphLODModels)
 	{
-		Model.DiscardCPUBuffers();
+		Model.DiscardVertexData();
 	}
 }
 
@@ -116,6 +121,7 @@ void UMorphTarget::PopulateDeltas(const TArray<FMorphTargetDelta>& Deltas, const
 
 	// remove array slack
 	MorphModel.Vertices.Shrink();
+	MorphModel.NumVertices = MorphModel.Vertices.Num();
 }
 
 void UMorphTarget::RemoveEmptyMorphTargets()
diff --git a/Engine/Source/Runtime/Engine/Private/SkeletalMeshLODRenderData.cpp b/Engine/Source/Runtime/Engine/Private/SkeletalMeshLODRenderData.cpp
--- a/Engine/Source/Runtime/Engine/Private/SkeletalMeshLODRenderData.cpp
+++ b/Engine/Source/Runtime/Engine/Private/SkeletalMeshLODRenderData.cpp
@@ -249,7 +249,7 @@ void FSkeletalMeshLODRenderData::InitResources(bool bNeedsVertexColors, int32 LO
 			{
 				// A CPU copy of the morph deltas has been made so it is safe to 
 				// discard the original data.  Keep CPU buffers when in the editor.
-				MorphTarget->DiscardCPUBuffers();
+				MorphTarget->DiscardVertexData();
 			}
 #endif
 		}
