Fix crash initializing Texture2DArray resources on some platforms
 - ExtData field was not being forwarded to the RHI, causing incorrect tiling mode selection and memory out-of-bounds writes.

diff --git a/Engine/Source/Runtime/Engine/Private/Rendering/Texture2DArrayResource.cpp b/Engine/Source/Runtime/Engine/Private/Rendering/Texture2DArrayResource.cpp
--- a/Engine/Source/Runtime/Engine/Private/Rendering/Texture2DArrayResource.cpp
+++ b/Engine/Source/Runtime/Engine/Private/Rendering/Texture2DArrayResource.cpp
@@ -74,6 +74,7 @@ void FTexture2DArrayResource::CreateTexture()
 	const FTexture2DMipMap& FirstMip = *MipsView[RequestedFirstLODIdx];
 
 	FRHIResourceCreateInfo CreateInfo(TEXT("FTexture2DArrayResource"));
+	CreateInfo.ExtData = PlatformData->GetExtData();
 	TRefCountPtr<FRHITexture2DArray> TextureArray = RHICreateTexture2DArray(FirstMip.SizeX, FirstMip.SizeY, FirstMip.SizeZ, PixelFormat, State.NumRequestedLODs, 1, CreationFlags, CreateInfo);
 	TextureRHI = TextureArray;
 
