Add mesh name and ray tracing tag to morph target gpu marker

diff --git a/Engine/Source/Runtime/Engine/Private/SkeletalRenderGPUSkin.cpp b/Engine/Source/Runtime/Engine/Private/SkeletalRenderGPUSkin.cpp
--- a/Engine/Source/Runtime/Engine/Private/SkeletalRenderGPUSkin.cpp
+++ b/Engine/Source/Runtime/Engine/Private/SkeletalRenderGPUSkin.cpp
@@ -582,7 +582,7 @@ void FSkeletalMeshObjectGPUSkin::ProcessUpdatedDynamicData(EGPUSkinCacheEntryMod
 						MorphTargetWeights.Add(DynamicData->MorphTargetWeights[i]);
 					}
 				}
-				LOD.UpdateMorphVertexBufferGPU(RHICmdList, MorphTargetWeights, LODData.MorphTargetVertexInfoBuffers, DynamicData->SectionIdsUseByActiveMorphTargets);
+				LOD.UpdateMorphVertexBufferGPU(RHICmdList, MorphTargetWeights, LODData.MorphTargetVertexInfoBuffers, DynamicData->SectionIdsUseByActiveMorphTargets, GetDebugName(), Mode);
 			}
 			else
 			{
@@ -861,7 +861,9 @@ static void CalculateMorphDeltaBounds(const TArray<float>& MorphTargetWeights, c
 	);
 }
 
-void FSkeletalMeshObjectGPUSkin::FSkeletalMeshObjectLOD::UpdateMorphVertexBufferGPU(FRHICommandListImmediate& RHICmdList, const TArray<float>& MorphTargetWeights, const FMorphTargetVertexInfoBuffers& MorphTargetVertexInfoBuffers, const TArray<int32>& SectionIdsUseByActiveMorphTargets)
+void FSkeletalMeshObjectGPUSkin::FSkeletalMeshObjectLOD::UpdateMorphVertexBufferGPU(FRHICommandListImmediate& RHICmdList, const TArray<float>& MorphTargetWeights, 
+																					const FMorphTargetVertexInfoBuffers& MorphTargetVertexInfoBuffers, 
+																					const TArray<int32>& SectionIdsUseByActiveMorphTargets, const FName& OwnerName, EGPUSkinCacheEntryMode Mode)
 {
 	if (IsValidRef(MorphVertexBuffer.VertexBufferRHI))
 	{
@@ -875,8 +877,10 @@ void FSkeletalMeshObjectGPUSkin::FSkeletalMeshObjectLOD::UpdateMorphVertexBuffer
 
 		SCOPED_GPU_STAT(RHICmdList, MorphTargets);
 
+		const FString RayTracingTag = (Mode == EGPUSkinCacheEntryMode::RayTracing ? TEXT("[RT]") : TEXT(""));
+		const FString LODName = OwnerName.ToString() + TEXT("_LOD") + FString::FromInt(LODIndex);
 		SCOPED_DRAW_EVENTF(RHICmdList, MorphUpdate,
-			TEXT("MorphUpdate LodVertices=%d Threads=%d"),
+			TEXT("MorphUpdate%s_%s LodVertices=%d Threads=%d"), *RayTracingTag, *LODName,
 			LodData.GetNumVertices(),
 			MorphTargetVertexInfoBuffers.GetNumWorkItems());
 
diff --git a/Engine/Source/Runtime/Engine/Private/SkeletalRenderGPUSkin.h b/Engine/Source/Runtime/Engine/Private/SkeletalRenderGPUSkin.h
--- a/Engine/Source/Runtime/Engine/Private/SkeletalRenderGPUSkin.h
+++ b/Engine/Source/Runtime/Engine/Private/SkeletalRenderGPUSkin.h
@@ -632,7 +632,8 @@ private:
 		 * @param MorphTargetWeights - All Morph weights
 		 */
 		void UpdateMorphVertexBufferCPU(const TArray<FActiveMorphTarget>& ActiveMorphTargets, const TArray<float>& MorphTargetWeights);
-		void UpdateMorphVertexBufferGPU(FRHICommandListImmediate& RHICmdList, const TArray<float>& MorphTargetWeights, const FMorphTargetVertexInfoBuffers& MorphTargetVertexInfoBuffers, const TArray<int32>& SectionIdsUseByActiveMorphTargets);
+		void UpdateMorphVertexBufferGPU(FRHICommandListImmediate& RHICmdList, const TArray<float>& MorphTargetWeights, const FMorphTargetVertexInfoBuffers& MorphTargetVertexInfoBuffers, 
+										const TArray<int32>& SectionIdsUseByActiveMorphTargets, const FName& OwnerName, EGPUSkinCacheEntryMode Mode);
 
 		void UpdateSkinWeights(FSkelMeshComponentLODInfo* CompLODInfo);
 
