Networking Insights Updates

Engine Frame number is recorded for each packet, and can be seen by hovering over a packet in the packet view of the network profiler of Insights.
Connection state is displayed for each packet in pop up view when the mouse is hovered over a packet in the packet view of the network profiler.
GameInstance names contain their name, role, and status as server or client.
Connection labels include the actor that owns the connection, if the connection is going to a client or server, and the remote address and port of the connection.

EConnectionState UNetConnection::State is deprecated (becoming private), instead use GetConnectionState and SetConnectionState to access State

CL 17071965, 17101672, 17102201, 17102631

diff --git a/Engine/Plugins/Experimental/WebSocketNetworking/Source/WebSocketNetworking/Private/WebsocketConnection.cpp b/Engine/Plugins/Experimental/WebSocketNetworking/Source/WebSocketNetworking/Private/WebsocketConnection.cpp
--- a/Engine/Plugins/Experimental/WebSocketNetworking/Source/WebSocketNetworking/Private/WebsocketConnection.cpp
+++ b/Engine/Plugins/Experimental/WebSocketNetworking/Source/WebSocketNetworking/Private/WebsocketConnection.cpp
@@ -119,10 +119,7 @@ FString UWebSocketConnection::LowLevelDescribe()
 		TEXT(" remote=%s local=%s state: %s"),
 		*WebSocket->RemoteEndPoint(true),
 		*WebSocket->LocalEndPoint(true),
-		State == USOCK_Pending ? TEXT("Pending")
-		: State == USOCK_Open ? TEXT("Open")
-		: State == USOCK_Closed ? TEXT("Closed")
-		: TEXT("Invalid")
+		LexToString(GetConnectionState())
 		);
 }
 
diff --git a/Engine/Plugins/NetcodeUnitTest/NetcodeUnitTest/Source/NetcodeUnitTest/Private/ClientUnitTest.cpp b/Engine/Plugins/NetcodeUnitTest/NetcodeUnitTest/Source/NetcodeUnitTest/Private/ClientUnitTest.cpp
--- a/Engine/Plugins/NetcodeUnitTest/NetcodeUnitTest/Source/NetcodeUnitTest/Private/ClientUnitTest.cpp
+++ b/Engine/Plugins/NetcodeUnitTest/NetcodeUnitTest/Source/NetcodeUnitTest/Private/ClientUnitTest.cpp
@@ -426,7 +426,7 @@ void UClientUnitTest::NotifyProcessLog(TWeakPtr<FUnitTestProcess> InProcess, con
 		{
 			UNetConnection* UnitConn = (MinClient != nullptr ? MinClient->GetConn() : nullptr);
 
-			if (!bTriggerredInitialConnect && (UnitConn == nullptr || UnitConn->State == EConnectionState::USOCK_Pending))
+			if (!bTriggerredInitialConnect && (UnitConn == nullptr || UnitConn->GetConnectionState() == EConnectionState::USOCK_Pending))
 			{
 				if (ServerReadyLogs->ContainsByPredicate(SearchInLogLine))
 				{
diff --git a/Engine/Plugins/NetcodeUnitTest/NetcodeUnitTest/Source/NetcodeUnitTest/Private/MinimalClient.cpp b/Engine/Plugins/NetcodeUnitTest/NetcodeUnitTest/Source/NetcodeUnitTest/Private/MinimalClient.cpp
--- a/Engine/Plugins/NetcodeUnitTest/NetcodeUnitTest/Source/NetcodeUnitTest/Private/MinimalClient.cpp
+++ b/Engine/Plugins/NetcodeUnitTest/NetcodeUnitTest/Source/NetcodeUnitTest/Private/MinimalClient.cpp
@@ -588,7 +588,7 @@ void UMinimalClient::ResetConnTimeout(float Duration)
 {
 	UNetDriver* UnitDriver = (UnitConn != nullptr ? ToRawPtr(UnitConn->Driver) : nullptr);
 
-	if (UnitDriver != nullptr && UnitConn->State != USOCK_Closed)
+	if (UnitDriver != nullptr && UnitConn->GetConnectionState() != USOCK_Closed)
 	{
 		// @todo #JohnBHack: This is a slightly hacky way of setting the timeout to a large value, which will be overridden by newly
 		//				received packets, making it unsuitable for most situations (except crashes - but that could still be subject
diff --git a/Engine/Plugins/Online/OnlineFramework/Source/Lobby/Private/LobbyBeaconHost.cpp b/Engine/Plugins/Online/OnlineFramework/Source/Lobby/Private/LobbyBeaconHost.cpp
--- a/Engine/Plugins/Online/OnlineFramework/Source/Lobby/Private/LobbyBeaconHost.cpp
+++ b/Engine/Plugins/Online/OnlineFramework/Source/Lobby/Private/LobbyBeaconHost.cpp
@@ -213,7 +213,7 @@ void ALobbyBeaconHost::PostLogin(ALobbyBeaconClient* ClientActor)
 
 void ALobbyBeaconHost::KickPlayer(ALobbyBeaconClient* ClientActor, const FText& KickReason)
 {
-	UE_LOG(LogLobbyBeacon, Log, TEXT("KickPlayer for %s. PendingKill %d UNetConnection %s UNetDriver %s State %d"), *GetNameSafe(ClientActor), ClientActor->IsPendingKill(), *GetNameSafe(ClientActor->BeaconConnection), ClientActor->BeaconConnection ? *GetNameSafe(ClientActor->BeaconConnection->Driver) : TEXT("null"), ClientActor->BeaconConnection ? ClientActor->BeaconConnection->State : -1);
+	UE_LOG(LogLobbyBeacon, Log, TEXT("KickPlayer for %s. PendingKill %d UNetConnection %s UNetDriver %s State %d"), *GetNameSafe(ClientActor), ClientActor->IsPendingKill(), *GetNameSafe(ClientActor->BeaconConnection), ClientActor->BeaconConnection ? *GetNameSafe(ClientActor->BeaconConnection->Driver) : TEXT("null"), ClientActor->BeaconConnection ? ClientActor->BeaconConnection->GetConnectionState() : -1);
 	ClientActor->ClientWasKicked(KickReason);
 	DisconnectClient(ClientActor);
 }
diff --git a/Engine/Plugins/Online/OnlineSubsystemOculus/Source/Private/OculusNetConnection.cpp b/Engine/Plugins/Online/OnlineSubsystemOculus/Source/Private/OculusNetConnection.cpp
--- a/Engine/Plugins/Online/OnlineSubsystemOculus/Source/Private/OculusNetConnection.cpp
+++ b/Engine/Plugins/Online/OnlineSubsystemOculus/Source/Private/OculusNetConnection.cpp
@@ -84,7 +84,7 @@ void UOculusNetConnection::LowLevelSend(void* Data, int32 CountBits, FOutPacketT
 
 	// Do not send packets over a closed connection
 	// This can unintentionally re-open the connection
-	if (State == EConnectionState::USOCK_Closed && ovr_Net_IsConnected(PeerID))
+	if (GetConnectionState() == EConnectionState::USOCK_Closed && ovr_Net_IsConnected(PeerID))
 	{
 		return;
 	}
@@ -147,7 +147,7 @@ void UOculusNetConnection::FinishDestroy()
 		return;
 	}
 	// Keep track if it's this call that is closing the connection before cleanup is called
-	const bool bIsClosingOpenConnection = State != EConnectionState::USOCK_Closed;
+	const bool bIsClosingOpenConnection = GetConnectionState() != EConnectionState::USOCK_Closed;
 	UNetConnection::FinishDestroy();
 
 	// If this connection was open, then close it
diff --git a/Engine/Plugins/Online/OnlineSubsystemOculus/Source/Private/OculusNetDriver.cpp b/Engine/Plugins/Online/OnlineSubsystemOculus/Source/Private/OculusNetDriver.cpp
--- a/Engine/Plugins/Online/OnlineSubsystemOculus/Source/Private/OculusNetDriver.cpp
+++ b/Engine/Plugins/Online/OnlineSubsystemOculus/Source/Private/OculusNetDriver.cpp
@@ -268,7 +268,7 @@ void UOculusNetDriver::TickDispatch(float DeltaTime)
 		{
 			auto Connection = Connections[PeerID];
 
-			if (Connection->State == EConnectionState::USOCK_Open)
+			if (Connection->GetConnectionState() == EConnectionState::USOCK_Open)
 			{
 				UE_LOG(LogNetTraffic, VeryVerbose, TEXT("Got a raw packet of size %d"), PacketSize);
 				Connection->ReceivedRawPacket(Data, PacketSize);
@@ -401,7 +401,7 @@ void UOculusNetDriver::OnNetworkingConnectionStateChange(ovrMessageHandle Messag
 		{
 			// Connections in a state of Closed will not have a NetDriver
 			// They will hit a nullptr exception if processing packets
-			if (Connection->State == EConnectionState::USOCK_Closed)
+			if (Connection->GetConnectionState() == EConnectionState::USOCK_Closed)
 			{
 				UE_LOG(LogNet, Warning, TEXT("Cannot reopen a closed connection to %llu"), PeerID);
 
@@ -411,7 +411,7 @@ void UOculusNetDriver::OnNetworkingConnectionStateChange(ovrMessageHandle Messag
 			else
 			{
 				UE_LOG(LogNet, Verbose, TEXT("%llu is connected"), PeerID);
-				Connection->State = EConnectionState::USOCK_Open;
+				Connection->SetConnectionState(EConnectionState::USOCK_Open);
 			}
 		}
 		else
@@ -424,7 +424,7 @@ void UOculusNetDriver::OnNetworkingConnectionStateChange(ovrMessageHandle Messag
 		// Use ovr_Net_IsConnected as the source of truth of the actual connection
 		if (!ovr_Net_IsConnected(PeerID))
 		{
-			if (Connection->State == EConnectionState::USOCK_Pending && !IsServer())
+			if (Connection->GetConnectionState() == EConnectionState::USOCK_Pending && !IsServer())
 			{
 				// Treat the pending case as if the connection timed out and try again
 				UE_LOG(LogNet, Verbose, TEXT("Notification said %llu is closed, but connection is still pending.  Ignoring potentially old notification and retry the connection"), PeerID);
@@ -433,7 +433,7 @@ void UOculusNetDriver::OnNetworkingConnectionStateChange(ovrMessageHandle Messag
 			else
 			{
 				UE_LOG(LogNet, Verbose, TEXT("%llu is closed"), PeerID);
-				Connection->State = EConnectionState::USOCK_Closed;
+				Connection->SetConnectionState(EConnectionState::USOCK_Closed);
 			}
 		}
 		else
@@ -443,7 +443,7 @@ void UOculusNetDriver::OnNetworkingConnectionStateChange(ovrMessageHandle Messag
 	}
 	else if (State == ovrPeerState_Timeout)
 	{
-		if (Connection->State == EConnectionState::USOCK_Pending && !IsServer())
+		if (Connection->GetConnectionState() == EConnectionState::USOCK_Pending && !IsServer())
 		{
 			UE_LOG(LogNet, Verbose, TEXT("Retrying connection to %llu"), PeerID);
 			ovr_Net_Connect(PeerID);
@@ -451,7 +451,7 @@ void UOculusNetDriver::OnNetworkingConnectionStateChange(ovrMessageHandle Messag
 		else
 		{
 			UE_LOG(LogNet, Warning, TEXT("%llu timed out"), PeerID);
-			Connection->State = EConnectionState::USOCK_Closed;
+			Connection->SetConnectionState(EConnectionState::USOCK_Closed);
 		}
 	}
 	else
@@ -514,5 +514,5 @@ bool UOculusNetDriver::IsNetResourceValid()
 	}
 
 	// The clients need to wait until the connection is established before sending packets
-	return ServerConnection->State == EConnectionState::USOCK_Open;
+	return ServerConnection->GetConnectionState() == EConnectionState::USOCK_Open;
 }
diff --git a/Engine/Plugins/Online/OnlineSubsystemUtils/Source/OnlineSubsystemUtils/Private/IpConnection.cpp b/Engine/Plugins/Online/OnlineSubsystemUtils/Source/OnlineSubsystemUtils/Private/IpConnection.cpp
--- a/Engine/Plugins/Online/OnlineSubsystemUtils/Source/OnlineSubsystemUtils/Private/IpConnection.cpp
+++ b/Engine/Plugins/Online/OnlineSubsystemUtils/Source/OnlineSubsystemUtils/Private/IpConnection.cpp
@@ -252,7 +252,7 @@ void UIpConnection::Tick(float DeltaSeconds)
 		ResolutionState = EAddressResolutionState::Done;
 
 		// Host name resolution just now failed.
-		State = USOCK_Closed;
+		SetConnectionState(USOCK_Closed);
 		Close();
 	}
 
@@ -646,9 +646,6 @@ FString UIpConnection::LowLevelDescribe()
 		(RemoteAddr.IsValid() ? *RemoteAddr->ToString(true) : TEXT("nullptr")),
 		*LocalAddr->ToString(true),
 		(PlayerId.IsValid() ? *PlayerId->ToDebugString() : TEXT("nullptr")),
-			State==USOCK_Pending	?	TEXT("Pending")
-		:	State==USOCK_Open		?	TEXT("Open")
-		:	State==USOCK_Closed		?	TEXT("Closed")
-		:								TEXT("Invalid")
+		LexToString(GetConnectionState())
 	);
 }
diff --git a/Engine/Plugins/Online/OnlineSubsystemUtils/Source/OnlineSubsystemUtils/Private/IpNetDriver.cpp b/Engine/Plugins/Online/OnlineSubsystemUtils/Source/OnlineSubsystemUtils/Private/IpNetDriver.cpp
--- a/Engine/Plugins/Online/OnlineSubsystemUtils/Source/OnlineSubsystemUtils/Private/IpNetDriver.cpp
+++ b/Engine/Plugins/Online/OnlineSubsystemUtils/Source/OnlineSubsystemUtils/Private/IpNetDriver.cpp
@@ -1316,7 +1316,7 @@ void UIpNetDriver::TickDispatch(float DeltaTime)
 					// players slide...which means if the client's game crashes, they
 					// might get flooded to some degree with packets until they timeout.
 					// Either way, this should close up the usual DoS attacks.
-					if ((Connection->State != USOCK_Open) || (!AllowPlayerPortUnreach))
+					if ((Connection->GetConnectionState() != USOCK_Open) || (!AllowPlayerPortUnreach))
 					{
 						if (LogPortUnreach)
 						{
diff --git a/Engine/Plugins/Online/OnlineSubsystemUtils/Source/OnlineSubsystemUtils/Private/OnlineBeaconClient.cpp b/Engine/Plugins/Online/OnlineSubsystemUtils/Source/OnlineSubsystemUtils/Private/OnlineBeaconClient.cpp
--- a/Engine/Plugins/Online/OnlineSubsystemUtils/Source/OnlineSubsystemUtils/Private/OnlineBeaconClient.cpp
+++ b/Engine/Plugins/Online/OnlineSubsystemUtils/Source/OnlineSubsystemUtils/Private/OnlineBeaconClient.cpp
@@ -62,7 +62,7 @@ class UNetConnection* AOnlineBeaconClient::GetNetConnection() const
 
 bool AOnlineBeaconClient::DestroyNetworkActorHandled()
 {
-	if (BeaconConnection && BeaconConnection->State != USOCK_Closed)
+	if (BeaconConnection && BeaconConnection->GetConnectionState() != USOCK_Closed)
 	{
 		// This will be cleaned up in ~2 sec by UNetConnection Tick
 		BeaconConnection->bPendingDestroy = true;
@@ -204,7 +204,7 @@ void AOnlineBeaconClient::Tick(float DeltaTime)
 		// Monitor for close bunches sent by the server which close down the connection in UChannel::Cleanup
 		// See similar code in UWorld::TickNetClient
 		if (((ConnectionState == EBeaconConnectionState::Pending) || (ConnectionState == EBeaconConnectionState::Open)) &&
-			(NetDriver->ServerConnection->State == USOCK_Closed))
+			(NetDriver->ServerConnection->GetConnectionState() == USOCK_Closed))
 		{
 			UE_LOG(LogBeacon, Verbose, TEXT("Client beacon (%s) socket has closed, triggering failure."), *GetName());
 			OnFailure();
@@ -268,7 +268,7 @@ void AOnlineBeaconClient::OnFailure()
 void AOnlineBeaconClient::ClientOnConnected_Implementation()
 {
 	SetConnectionState(EBeaconConnectionState::Open);
-	BeaconConnection->State = USOCK_Open;
+	BeaconConnection->SetConnectionState(USOCK_Open);
 
 	SetRole(ROLE_Authority);
 	SetReplicates(true);
@@ -444,7 +444,7 @@ void AOnlineBeaconClient::FinalizeEncryptedConnection(const FEncryptionKeyRespon
 	UNetConnection* Connection = WeakConnection.Get();
 	if (Connection)
 	{
-		if (Connection->State != USOCK_Invalid && Connection->State != USOCK_Closed && Connection->Driver)
+		if (Connection->GetConnectionState() != USOCK_Invalid && Connection->GetConnectionState() != USOCK_Closed && Connection->Driver)
 		{
 			if (Response.Response == EEncryptionResponse::Success)
 			{
diff --git a/Engine/Plugins/Online/OnlineSubsystemUtils/Source/OnlineSubsystemUtils/Private/OnlineBeaconHost.cpp b/Engine/Plugins/Online/OnlineSubsystemUtils/Source/OnlineSubsystemUtils/Private/OnlineBeaconHost.cpp
--- a/Engine/Plugins/Online/OnlineSubsystemUtils/Source/OnlineSubsystemUtils/Private/OnlineBeaconHost.cpp
+++ b/Engine/Plugins/Online/OnlineSubsystemUtils/Source/OnlineSubsystemUtils/Private/OnlineBeaconHost.cpp
@@ -298,10 +298,10 @@ void AOnlineBeaconHost::DisconnectClient(AOnlineBeaconClient* ClientActor)
 			*GetNameSafe(ClientActor), 
 			*GetNameSafe(Connection), 
 			Connection ? *GetNameSafe(Connection->Driver) : TEXT("null"),
-			Connection ? Connection->State : -1);
+			Connection ? Connection->GetConnectionState() : -1);
 
 		// Closing the connection will start the chain of events leading to the removal from lists and destruction of the actor
-		if (Connection && Connection->State != USOCK_Closed)
+		if (Connection && Connection->GetConnectionState() != USOCK_Closed)
 		{
 			Connection->FlushNet(true);
 			Connection->Close();
@@ -338,7 +338,7 @@ void AOnlineBeaconHost::SendWelcomeControlMessage(UNetConnection* Connection)
 {
 	if (Connection)
 	{
-		if (Connection->State != USOCK_Invalid && Connection->State != USOCK_Closed && Connection->Driver)
+		if (Connection->GetConnectionState() != USOCK_Invalid && Connection->GetConnectionState() != USOCK_Closed && Connection->Driver)
 		{
 			Connection->Challenge = FString::Printf(TEXT("%08X"), FPlatformTime::Cycles());
 			FNetControlMessage<NMT_BeaconWelcome>::Send(Connection);
@@ -360,7 +360,7 @@ void AOnlineBeaconHost::SendWelcomeControlMessage(const FEncryptionKeyResponse&
 	UNetConnection* Connection = WeakConnection.Get();
 	if (Connection)
 	{
-		if (Connection->State != USOCK_Invalid && Connection->State != USOCK_Closed && Connection->Driver)
+		if (Connection->GetConnectionState() != USOCK_Invalid && Connection->GetConnectionState() != USOCK_Closed && Connection->Driver)
 		{
 			if (Response.Response == EEncryptionResponse::Success)
 			{
diff --git a/Engine/Plugins/Online/VoiceChat/EOSVoiceChat/Source/EOSVoiceChat/Public/EOSVoiceChat.h b/Engine/Plugins/Online/VoiceChat/EOSVoiceChat/Source/EOSVoiceChat/Public/EOSVoiceChat.h
--- a/Engine/Plugins/Online/VoiceChat/EOSVoiceChat/Source/EOSVoiceChat/Public/EOSVoiceChat.h
+++ b/Engine/Plugins/Online/VoiceChat/EOSVoiceChat/Source/EOSVoiceChat/Public/EOSVoiceChat.h
@@ -231,7 +231,7 @@ protected:
 
 	FEOSVoiceChatWeakPtr CreateWeakThis();
 
-	friend const TCHAR* LexToString(EConnectionState State);
+	friend const TCHAR* LexToString(FEOSVoiceChat::EConnectionState State);
 };
 
 #endif // WITH_EOS_RTC
\ No newline at end of file
diff --git a/Engine/Plugins/Runtime/ReplicationGraph/Source/Private/ReplicationGraph.cpp b/Engine/Plugins/Runtime/ReplicationGraph/Source/Private/ReplicationGraph.cpp
--- a/Engine/Plugins/Runtime/ReplicationGraph/Source/Private/ReplicationGraph.cpp
+++ b/Engine/Plugins/Runtime/ReplicationGraph/Source/Private/ReplicationGraph.cpp
@@ -1089,7 +1089,7 @@ int32 UReplicationGraph::ServerReplicateActors(float DeltaSeconds)
 						if (ensureMsgf(Actor,
 							TEXT("Stale Connection Actor Info with Valid Channel but Invalid Actor. RelevantTime=%f, LastUpdateTime=%f, LastRepFrameNum=%d, RepPeriod=%d, CloseFrame=%d, CurrentRepFrame=%d, bTearOff=%d, bDormant=%d, Channel=%s, State=%d"),
 							Channel->RelevantTime, Channel->LastUpdateTime, ConnectionActorInfo.LastRepFrameNum, ConnectionActorInfo.ReplicationPeriodFrame, ConnectionActorInfo.ActorChannelCloseFrameNum,
-							FrameNum, !!ConnectionActorInfo.bTearOff, !!ConnectionActorInfo.bDormantOnConnection, *(Channel->Describe()), static_cast<int32>(NetConnection->State)))
+							FrameNum, !!ConnectionActorInfo.bTearOff, !!ConnectionActorInfo.bDormantOnConnection, *(Channel->Describe()), static_cast<int32>(NetConnection->GetConnectionState())))
 						{
 							if (Actor->IsNetStartupActor())
 								continue;
@@ -2207,7 +2207,7 @@ bool UReplicationGraph::ProcessRemoteFunction(class AActor* Actor, UFunction* Fu
 			Connection = ((UChildConnection*)Connection)->Parent;
 		}
 	
-		if (Connection->State == USOCK_Closed)
+		if (Connection->GetConnectionState() == USOCK_Closed)
 		{
 			return true;
 		}
@@ -2507,7 +2507,7 @@ bool UNetReplicationGraphConnection::PrepareForReplication()
 		CurChild->ViewTarget = CurChild->PlayerController ? CurChild->PlayerController->GetViewTarget() : ToRawPtr(CurChild->OwningActor);
 	}
 
-	return (NetConnection->State != USOCK_Closed) && (NetConnection->ViewTarget != nullptr);
+	return (NetConnection->GetConnectionState() != USOCK_Closed) && (NetConnection->ViewTarget != nullptr);
 }
 
 void UNetReplicationGraphConnection::NotifyAddDestructionInfo(FActorDestructionInfo* DestructInfo)
diff --git a/Engine/Plugins/Runtime/Steam/SteamSockets/Source/SteamSockets/Private/SteamSocketsNetDriver.cpp b/Engine/Plugins/Runtime/Steam/SteamSockets/Source/SteamSockets/Private/SteamSocketsNetDriver.cpp
--- a/Engine/Plugins/Runtime/Steam/SteamSockets/Source/SteamSockets/Private/SteamSocketsNetDriver.cpp
+++ b/Engine/Plugins/Runtime/Steam/SteamSockets/Source/SteamSockets/Private/SteamSocketsNetDriver.cpp
@@ -403,7 +403,7 @@ class ISocketSubsystem* USteamSocketsNetDriver::GetSocketSubsystem()
 
 bool USteamSocketsNetDriver::IsNetResourceValid(void)
 {
-	return Socket != nullptr && (!ServerConnection || (ServerConnection && ServerConnection->State == USOCK_Open));
+	return Socket != nullptr && (!ServerConnection || (ServerConnection && ServerConnection->GetConnectionState() == USOCK_Open));
 }
 
 bool USteamSocketsNetDriver::ArePacketHandlersDisabled() const
@@ -543,7 +543,7 @@ void USteamSocketsNetDriver::OnConnectionUpdated(SteamSocketHandles SocketHandle
 		UNetConnection* SocketConnection = (ServerConnection) ? ToRawPtr(ServerConnection) : FindClientConnectionForHandle(SocketHandle);
 		if (SocketConnection)
 		{
-			SocketConnection->State = USOCK_Open;
+			SocketConnection->SetConnectionState(USOCK_Open);
 		}
 
 		UE_LOG(LogNet, Verbose, TEXT("SteamSockets: Connection established with user with socket id: %u"), SocketHandle);
@@ -566,7 +566,7 @@ void USteamSocketsNetDriver::OnConnectionDisconnected(SteamSocketHandles SocketH
 	UNetConnection* SocketConnection = ServerConnection ? ToRawPtr(ServerConnection) : FindClientConnectionForHandle(SocketHandle);
 	if (SocketConnection)
 	{
-		SocketConnection->State = USOCK_Closed;
+		SocketConnection->SetConnectionState(USOCK_Closed);
 	}
 
 	UE_LOG(LogNet, Verbose, TEXT("SteamSockets: Connection dropped with user with socket id: %u"), SocketHandle);
diff --git a/Engine/Source/Developer/TraceInsights/Private/Insights/NetworkingProfiler/ViewModels/PacketViewDrawHelper.cpp b/Engine/Source/Developer/TraceInsights/Private/Insights/NetworkingProfiler/ViewModels/PacketViewDrawHelper.cpp
--- a/Engine/Source/Developer/TraceInsights/Private/Insights/NetworkingProfiler/ViewModels/PacketViewDrawHelper.cpp
+++ b/Engine/Source/Developer/TraceInsights/Private/Insights/NetworkingProfiler/ViewModels/PacketViewDrawHelper.cpp
@@ -40,6 +40,7 @@ void FNetworkPacketAggregatedSample::AddPacket(const int32 PacketIndex, const Tr
 		LargestPacket.TotalSizeInBytes = Packet.TotalPacketSizeInBytes;
 		LargestPacket.TimeStamp = TimeStamp;
 		LargestPacket.Status = Packet.DeliveryStatus;
+		LargestPacket.ConnectionState = Packet.ConnectionState;
 	}
 
 	switch (AggregatedStatus)
diff --git a/Engine/Source/Developer/TraceInsights/Private/Insights/NetworkingProfiler/ViewModels/PacketViewDrawHelper.h b/Engine/Source/Developer/TraceInsights/Private/Insights/NetworkingProfiler/ViewModels/PacketViewDrawHelper.h
--- a/Engine/Source/Developer/TraceInsights/Private/Insights/NetworkingProfiler/ViewModels/PacketViewDrawHelper.h
+++ b/Engine/Source/Developer/TraceInsights/Private/Insights/NetworkingProfiler/ViewModels/PacketViewDrawHelper.h
@@ -27,6 +27,7 @@ struct FNetworkPacket
 	uint32 ContentSizeInBits;
 	uint32 TotalSizeInBytes;
 	double TimeStamp;
+	TraceServices::ENetProfilerConnectionState ConnectionState;
 	TraceServices::ENetProfilerDeliveryStatus Status;
 
 	FNetworkPacket()
@@ -35,6 +36,7 @@ struct FNetworkPacket
 		, ContentSizeInBits(0)
 		, TotalSizeInBytes(0)
 		, TimeStamp(std::numeric_limits<double>::infinity())
+		, ConnectionState(TraceServices::ENetProfilerConnectionState::USOCK_Invalid)
 		, Status(TraceServices::ENetProfilerDeliveryStatus::Unknown)
 	{}
 
diff --git a/Engine/Source/Developer/TraceInsights/Private/Insights/NetworkingProfiler/Widgets/SNetworkingProfilerWindow.cpp b/Engine/Source/Developer/TraceInsights/Private/Insights/NetworkingProfiler/Widgets/SNetworkingProfilerWindow.cpp
--- a/Engine/Source/Developer/TraceInsights/Private/Insights/NetworkingProfiler/Widgets/SNetworkingProfilerWindow.cpp
+++ b/Engine/Source/Developer/TraceInsights/Private/Insights/NetworkingProfiler/Widgets/SNetworkingProfilerWindow.cpp
@@ -757,7 +757,15 @@ FReply SNetworkingProfilerWindow::OnDrop(const FGeometry& MyGeometry, const FDra
 
 FText SNetworkingProfilerWindow::FGameInstanceItem::GetText() const
 {
-	return FText::Format(LOCTEXT("GameInstanceItemFmt", "Game Instance {0}"), FText::AsNumber(GetIndex()));
+	if (IsInstanceNameSet())
+	{
+		return FText::Format(LOCTEXT("GameInstanceItemFmt0", "{0} Game Instance {1} [{2}]"), FText::FromString(GetInstanceName()), FText::AsNumber(GetIndex()),
+			IsServer() ? FText::FromString("Server") : FText::FromString("Client"));
+	}
+	else
+	{
+		return FText::Format(LOCTEXT("GameInstanceItemFmt1", "Game Instance {0}"), FText::AsNumber(GetIndex()));
+	}
 }
 
 ////////////////////////////////////////////////////////////////////////////////////////////////////
@@ -771,13 +779,29 @@ FText SNetworkingProfilerWindow::FGameInstanceItem::GetTooltipText() const
 
 FText SNetworkingProfilerWindow::FConnectionItem::GetText() const
 {
-	if (!Connection.Name)
+	if (IsConnectionNameSet())
 	{
-		return FText::Format(LOCTEXT("ConnectionItemFmt1", "Connection {0}"), FText::AsNumber(GetIndex()));
+		if (IsConnectionAddressSet())
+		{
+			return FText::Format(LOCTEXT("ConnectionItemFmt3", "Connection {0} ({1}) to {2}"), FText::AsNumber(GetIndex()), FText::FromString(GetConnectionName()),
+				FText::FromString(GetConnectionAddress()));
+		}
+		else
+		{
+			return FText::Format(LOCTEXT("ConnectionItemFmt4", "Connection {0} ({1})"), FText::AsNumber(GetIndex()), FText::FromString(GetConnectionName()));
+		}
 	}
 	else
 	{
-		return FText::Format(LOCTEXT("ConnectionItemFmt2", "Connection {0} ({1})"), FText::AsNumber(GetIndex()), FText::FromString(Connection.Name));
+		if (IsConnectionAddressSet())
+		{
+			return FText::Format(LOCTEXT("ConnectionItemFmt1", "Connection {0} to {1}"), FText::AsNumber(GetIndex()),
+				FText::FromString(GetConnectionAddress()));
+		}
+		else
+		{
+			return FText::Format(LOCTEXT("ConnectionItemFmt2", "Connection {0}"), FText::AsNumber(GetIndex()));
+		}
 	}
 }
 
diff --git a/Engine/Source/Developer/TraceInsights/Private/Insights/NetworkingProfiler/Widgets/SNetworkingProfilerWindow.h b/Engine/Source/Developer/TraceInsights/Private/Insights/NetworkingProfiler/Widgets/SNetworkingProfilerWindow.h
--- a/Engine/Source/Developer/TraceInsights/Private/Insights/NetworkingProfiler/Widgets/SNetworkingProfilerWindow.h
+++ b/Engine/Source/Developer/TraceInsights/Private/Insights/NetworkingProfiler/Widgets/SNetworkingProfilerWindow.h
@@ -52,6 +52,9 @@ private:
 		{}
 
 		uint32 GetIndex() const { return GameInstance.GameInstanceIndex; }
+		bool IsInstanceNameSet() const { return GameInstance.InstanceName != nullptr; }
+		const TCHAR* GetInstanceName() const { return GameInstance.InstanceName; }
+		bool IsServer() const { return GameInstance.bIsServer; }
 		FText GetText() const;
 		FText GetTooltipText() const;
 
@@ -66,6 +69,10 @@ private:
 		{}
 
 		uint32 GetIndex() const { return Connection.ConnectionIndex; }
+		bool IsConnectionNameSet() const { return Connection.Name != nullptr; }
+		bool IsConnectionAddressSet() const { return Connection.AddressString != nullptr; }
+		const TCHAR* GetConnectionName() const { return Connection.Name; }
+		const TCHAR* GetConnectionAddress() const { return Connection.AddressString; }
 		FText GetText() const;
 		FText GetTooltipText() const;
 
diff --git a/Engine/Source/Developer/TraceInsights/Private/Insights/NetworkingProfiler/Widgets/SPacketView.cpp b/Engine/Source/Developer/TraceInsights/Private/Insights/NetworkingProfiler/Widgets/SPacketView.cpp
--- a/Engine/Source/Developer/TraceInsights/Private/Insights/NetworkingProfiler/Widgets/SPacketView.cpp
+++ b/Engine/Source/Developer/TraceInsights/Private/Insights/NetworkingProfiler/Widgets/SPacketView.cpp
@@ -711,15 +711,25 @@ int32 SPacketView::OnPaint(const FPaintArgs& Args, const FGeometry& AllottedGeom
 			int NumLines;
 			FString Text;
 			uint32 UnusedBits = HoveredSample.Sample->LargestPacket.TotalSizeInBytes * 8 - HoveredSample.Sample->LargestPacket.ContentSizeInBits;
+			TSharedPtr<const TraceServices::IAnalysisSession> Session = FInsightsManager::Get()->GetSession();
+			uint64 EngineFrameNumber = 0;
+			if (Session.IsValid())
+			{
+				TraceServices::FAnalysisSessionReadScope SessionReadScope(*Session.Get());
+				const TraceServices::IFrameProvider& FrameProvider = TraceServices::ReadFrameProvider(*Session.Get());
+				EngineFrameNumber = FrameProvider.GetFrameNumberForTimestamp(ETraceFrameType::TraceFrameType_Game, HoveredSample.Sample->LargestPacket.TimeStamp);
+			}
 			if (HoveredSample.Sample->NumPackets == 1)
 			{
-				NumLines = 6;
+				NumLines = 8;
 				Text = FString::Format(TEXT("Packet Index: {0}\n"
 											"Sequence Number: {1}\n"
 											"Content Size: {2} bits\n"
 											"Total Size: {3} bytes ({4} unused bits)\n"
 											"Timestamp: {5}\n"
-											"Status: {6}"),
+											"Status: {6}\n"
+											"Connection State: {7}\n"
+											"Engine Frame Number: {8}"),
 					{
 						FText::AsNumber(HoveredSample.Sample->LargestPacket.Index).ToString(),
 						FText::AsNumber(HoveredSample.Sample->LargestPacket.SequenceNumber).ToString(),
@@ -727,12 +737,16 @@ int32 SPacketView::OnPaint(const FPaintArgs& Args, const FGeometry& AllottedGeom
 						FText::AsNumber(HoveredSample.Sample->LargestPacket.TotalSizeInBytes).ToString(),
 						FText::AsNumber(UnusedBits).ToString(),
 						TimeUtils::FormatTimeHMS(HoveredSample.Sample->LargestPacket.TimeStamp, Precision),
-						::StatusToString(HoveredSample.Sample->LargestPacket.Status)
+						::StatusToString(HoveredSample.Sample->LargestPacket.Status),
+						LexToString(HoveredSample.Sample->LargestPacket.ConnectionState),
+						EngineFrameNumber > 0 ?
+						FText::AsNumber(EngineFrameNumber).ToString() :
+						TEXT("N/A")
 					});
 			}
 			else
 			{
-				NumLines = 9;
+				NumLines = 11;
 				Text = FString::Format(TEXT("{0} network packets\n"
 											"({1})\n"
 											"Largest Packet\n"
@@ -741,7 +755,9 @@ int32 SPacketView::OnPaint(const FPaintArgs& Args, const FGeometry& AllottedGeom
 											"    Content Size: {4} bits\n"
 											"    Total Size: {5} bytes ({6} unused bits)\n"
 											"    Timestamp: {7}\n"
-											"    Status: {8}"),
+											"    Status: {8}\n"
+											"    Connection State: {9}\n"
+											"    Engine Frame Number: {10}"),
 					{
 						HoveredSample.Sample->NumPackets,
 						::AggregatedStatusToString(HoveredSample.Sample->AggregatedStatus),
@@ -751,7 +767,11 @@ int32 SPacketView::OnPaint(const FPaintArgs& Args, const FGeometry& AllottedGeom
 						FText::AsNumber(HoveredSample.Sample->LargestPacket.TotalSizeInBytes).ToString(),
 						FText::AsNumber(UnusedBits).ToString(),
 						TimeUtils::FormatTimeHMS(HoveredSample.Sample->LargestPacket.TimeStamp, Precision),
-						::StatusToString(HoveredSample.Sample->LargestPacket.Status)
+						::StatusToString(HoveredSample.Sample->LargestPacket.Status),
+						LexToString(HoveredSample.Sample->LargestPacket.ConnectionState),
+						EngineFrameNumber > 0 ?
+						FText::AsNumber(EngineFrameNumber).ToString() :
+						TEXT("N/A")
 					});
 			}
 
@@ -1306,6 +1326,36 @@ FReply SPacketView::OnKeyDown(const FGeometry& MyGeometry, const FKeyEvent& InKe
 		}
 		return FReply::Handled();
 	}
+	else if (InKeyEvent.GetKey() == EKeys::Equals ||
+			 InKeyEvent.GetKey() == EKeys::Add)
+	{
+		if (InKeyEvent.GetModifierKeys().IsShiftDown())
+		{
+			FAxisViewportDouble& ViewportY = Viewport.GetVerticalAxisViewport();
+			const float ScaleY = ViewportY.GetScale() * 1.25f;
+			ViewportY.SetScale(ScaleY);
+		}
+		else
+		{
+			ZoomHorizontally(1.0f, MousePosition.X);
+		}
+		return FReply::Handled();
+	}
+	else if (InKeyEvent.GetKey() == EKeys::Hyphen ||
+			 InKeyEvent.GetKey() == EKeys::Subtract)
+	{
+		if (InKeyEvent.GetModifierKeys().IsShiftDown())
+		{
+			FAxisViewportDouble& ViewportY = Viewport.GetVerticalAxisViewport();
+			const float ScaleY = ViewportY.GetScale() * 0.8f;
+			ViewportY.SetScale(ScaleY);
+		}
+		else
+		{
+			ZoomHorizontally(-1.0f, MousePosition.X);
+		}
+		return FReply::Handled();
+	}
 
 	return SCompoundWidget::OnKeyDown(MyGeometry, InKeyEvent);
 }
diff --git a/Engine/Source/Developer/TraceServices/Private/Analyzers/CsvProfilerTraceAnalysis.cpp b/Engine/Source/Developer/TraceServices/Private/Analyzers/CsvProfilerTraceAnalysis.cpp
--- a/Engine/Source/Developer/TraceServices/Private/Analyzers/CsvProfilerTraceAnalysis.cpp
+++ b/Engine/Source/Developer/TraceServices/Private/Analyzers/CsvProfilerTraceAnalysis.cpp
@@ -147,7 +147,7 @@ bool FCsvProfilerAnalyzer::OnEvent(uint16 RouteId, EStyle Style, const FOnEventC
 	{
 		RenderThreadId = FTraceAnalyzerUtils::GetThreadIdField(Context, "RenderThreadId");
 		RHIThreadId = FTraceAnalyzerUtils::GetThreadIdField(Context, "RHIThreadId");
-		uint32 CaptureStartFrame = GetFrameNumberForTimestamp(TraceFrameType_Game, Context.EventTime.AsSeconds(EventData.GetValue<uint64>("Cycle")));
+		uint32 CaptureStartFrame = FrameProvider.GetFrameNumberForTimestamp(TraceFrameType_Game, Context.EventTime.AsSeconds(EventData.GetValue<uint64>("Cycle")));
 		bEnableCounts = EventData.GetValue<bool>("EnableCounts");
 		FString FileName = FTraceAnalyzerUtils::LegacyAttachmentString<TCHAR>("FileName", Context);
 		const TCHAR* StoredFileName = Session.StoreString(*FileName);
@@ -156,7 +156,7 @@ bool FCsvProfilerAnalyzer::OnEvent(uint16 RouteId, EStyle Style, const FOnEventC
 	}
 	case RouteId_EndCapture:
 	{
-		uint32 CaptureEndFrame = GetFrameNumberForTimestamp(TraceFrameType_Game, Context.EventTime.AsSeconds(EventData.GetValue<uint64>("Cycle")));
+		uint32 CaptureEndFrame = FrameProvider.GetFrameNumberForTimestamp(TraceFrameType_Game, Context.EventTime.AsSeconds(EventData.GetValue<uint64>("Cycle")));
 		for (FStatSeriesInstance* StatSeries : StatSeriesInstanceArray)
 		{
 			FlushAtEndOfCapture(*StatSeries, CaptureEndFrame);
@@ -341,7 +341,7 @@ void FCsvProfilerAnalyzer::HandleMarker(const FOnEventContext& Context, FThreadS
 	}
 
 	double Timestamp = Context.EventTime.AsSeconds(Marker.Cycle);
-	uint32 FrameNumber = GetFrameNumberForTimestamp(ThreadState.FrameType, Timestamp);
+	uint32 FrameNumber = FrameProvider.GetFrameNumberForTimestamp(ThreadState.FrameType, Timestamp);
 	if (Marker.bIsBegin)
 	{
 		ThreadState.MarkerStack.Push(Marker);
@@ -391,7 +391,7 @@ void FCsvProfilerAnalyzer::HandleCustomStatEvent(const FOnEventContext& Context,
 	FThreadState& ThreadState = GetThreadState(ThreadId);
 	FStatSeriesInstance& StatSeries = GetStatSeries(Context.EventData.GetValue<uint64>("StatId"), bIsFloat ? CsvStatSeriesType_CustomStatFloat : CsvStatSeriesType_CustomStatInt, ThreadState);
 	ECsvOpType OpType = static_cast<ECsvOpType>(Context.EventData.GetValue<uint8>("OpType"));
-	uint32 FrameNumber = GetFrameNumberForTimestamp(ThreadState.FrameType, Context.EventTime.AsSeconds(Context.EventData.GetValue<uint64>("Cycle")));
+	uint32 FrameNumber = FrameProvider.GetFrameNumberForTimestamp(ThreadState.FrameType, Context.EventTime.AsSeconds(Context.EventData.GetValue<uint64>("Cycle")));
 	if (bIsFloat)
 	{
 		float Value = Context.EventData.GetValue<float>("Value");
@@ -409,7 +409,7 @@ void FCsvProfilerAnalyzer::HandleEventEvent(const FOnEventContext& Context)
 	uint32 ThreadId = FTraceAnalyzerUtils::GetThreadIdField(Context);
 	FThreadState& ThreadState = GetThreadState(ThreadId);
 	uint64 Cycle = Context.EventData.GetValue<uint64>("Cycle");
-	uint32 FrameNumber = GetFrameNumberForTimestamp(ThreadState.FrameType, Context.EventTime.AsSeconds(Cycle));
+	uint32 FrameNumber = FrameProvider.GetFrameNumberForTimestamp(ThreadState.FrameType, Context.EventTime.AsSeconds(Cycle));
 	FString EventText = FTraceAnalyzerUtils::LegacyAttachmentString<TCHAR>("Text", Context);
 	int32 CategoryIndex = Context.EventData.GetValue<int32>("CategoryIndex");
 	if (CategoryIndex > 0)
@@ -419,25 +419,6 @@ void FCsvProfilerAnalyzer::HandleEventEvent(const FOnEventContext& Context)
 	CsvProfilerProvider.AddEvent(FrameNumber, Session.StoreString(*EventText));
 }
 
-uint32 FCsvProfilerAnalyzer::GetFrameNumberForTimestamp(ETraceFrameType FrameType, double Timestamp) const
-{
-	const TArray64<double>& FrameStartTimes = FrameProvider.GetFrameStartTimes(FrameType);
-
-	if (FrameStartTimes.Num() == 0 || Timestamp < FrameStartTimes[0])
-	{
-		return 0;
-	}
-	else if (Timestamp >= FrameStartTimes.Last())
-	{
-		return FrameStartTimes.Num();
-	}
-	else
-	{
-		uint32 Index = static_cast<uint32>(Algo::LowerBound(FrameStartTimes, Timestamp));
-		return Index + 1;
-	}
-}
-
 void FCsvProfilerAnalyzer::Flush(FStatSeriesInstance& StatSeries)
 {
 	double CounterTimestamp;
diff --git a/Engine/Source/Developer/TraceServices/Private/Analyzers/CsvProfilerTraceAnalysis.h b/Engine/Source/Developer/TraceServices/Private/Analyzers/CsvProfilerTraceAnalysis.h
--- a/Engine/Source/Developer/TraceServices/Private/Analyzers/CsvProfilerTraceAnalysis.h
+++ b/Engine/Source/Developer/TraceServices/Private/Analyzers/CsvProfilerTraceAnalysis.h
@@ -121,7 +121,6 @@ private:
 	void HandleMarker(const FOnEventContext& Context, FThreadState& ThreadState, const FTimingMarker& Marker);
 	void HandleCustomStatEvent(const FOnEventContext& Context, bool bIsFloat);
 	void HandleEventEvent(const FOnEventContext& Context);
-	uint32 GetFrameNumberForTimestamp(ETraceFrameType FrameType, double Timestamp) const;
 	void Flush(FStatSeriesInstance& StatSeries);
 	void FlushIfNewFrame(FStatSeriesInstance& StatSeries, uint32 FrameNumber);
 	void FlushAtEndOfCapture(FStatSeriesInstance& StatSeries, uint32 CaptureEndFrame);
diff --git a/Engine/Source/Developer/TraceServices/Private/Analyzers/NetTraceAnalyzer.cpp b/Engine/Source/Developer/TraceServices/Private/Analyzers/NetTraceAnalyzer.cpp
--- a/Engine/Source/Developer/TraceServices/Private/Analyzers/NetTraceAnalyzer.cpp
+++ b/Engine/Source/Developer/TraceServices/Private/Analyzers/NetTraceAnalyzer.cpp
@@ -37,11 +37,14 @@ void FNetTraceAnalyzer::OnAnalysisBegin(const FOnAnalysisContext& Context)
 	Builder.RouteEvent(RouteId_PacketEvent, "NetTrace", "PacketEvent");
 	Builder.RouteEvent(RouteId_PacketDroppedEvent, "NetTrace", "PacketDroppedEvent");
 	Builder.RouteEvent(RouteId_ConnectionCreatedEvent, "NetTrace", "ConnectionCreatedEvent");
+	Builder.RouteEvent(RouteId_ConnectionUpdatedEvent, "NetTrace", "ConnectionUpdatedEvent");
 	// Add some default event types that we use for generic type events to make it easier to extend
 	// ConnectionAdded/Removed connections state /name etc?
 	Builder.RouteEvent(RouteId_ConnectionClosedEvent, "NetTrace", "ConnectionClosedEvent");
 	Builder.RouteEvent(RouteId_ObjectCreatedEvent, "NetTrace", "ObjectCreatedEvent");
 	Builder.RouteEvent(RouteId_ObjectDestroyedEvent, "NetTrace", "ObjectDestroyedEvent");
+	Builder.RouteEvent(RouteId_ConnectionStateUpdatedEvent, "NetTrace", "ConnectionStateUpdatedEvent");
+	Builder.RouteEvent(RouteId_InstanceUpdatedEvent, "NetTrace", "InstanceUpdatedEvent");
 
 	// Default names
 	{
@@ -141,7 +144,19 @@ bool FNetTraceAnalyzer::OnEvent(uint16 RouteId, EStyle Style, const FOnEventCont
 
 		case RouteId_ConnectionCreatedEvent:
 		{
-			HandleConnectionCretedEvent(Context, EventData);
+			HandleConnectionCreatedEvent(Context, EventData);
+		}
+		break;
+
+		case RouteId_ConnectionStateUpdatedEvent:
+		{
+			HandleConnectionStateUpdatedEvent(Context, EventData);
+		}
+		break;
+
+		case RouteId_ConnectionUpdatedEvent:
+		{
+			HandleConnectionUpdatedEvent(Context, EventData);
 		}
 		break;
 
@@ -162,6 +177,12 @@ bool FNetTraceAnalyzer::OnEvent(uint16 RouteId, EStyle Style, const FOnEventCont
 			HandleObjectDestroyedEvent(Context, EventData);
 		}
 		break;
+		
+		case RouteId_InstanceUpdatedEvent:
+		{
+			HandleGameInstanceUpdatedEvent(Context, EventData);
+		}
+		break;
 	}
 
 	return true;
@@ -456,6 +477,7 @@ void FNetTraceAnalyzer::HandlePacketEvent(const FOnEventContext& Context, const
 	Packet.TimeStamp = GetLastTimestamp();
 	Packet.SequenceNumber = SequenceNumber;
 	Packet.DeliveryStatus = ENetProfilerDeliveryStatus::Unknown;
+	Packet.ConnectionState = ConnectionState->ConnectionState;
 
 	Packet.ContentSizeInBits = PacketBits;
 	Packet.TotalPacketSizeInBytes = (Packet.ContentSizeInBits + 7u) >> 3u;
@@ -491,7 +513,7 @@ void FNetTraceAnalyzer::HandlePacketDroppedEvent(const FOnEventContext& Context,
 	NetProfilerProvider.EditPacketDeliveryStatus(ConnectionState->ConnectionIndex, ENetProfilerConnectionMode(PacketType), SequenceNumber, ENetProfilerDeliveryStatus::Dropped);
 }
 
-void FNetTraceAnalyzer::HandleConnectionCretedEvent(const FOnEventContext& Context, const FEventData& EventData)
+void FNetTraceAnalyzer::HandleConnectionCreatedEvent(const FOnEventContext& Context, const FEventData& EventData)
 {
 	const uint8 GameInstanceId = EventData.GetValue<uint8>("GameInstanceId");
 	const uint16 ConnectionId = EventData.GetValue<uint16>("ConnectionId");
@@ -515,6 +537,47 @@ void FNetTraceAnalyzer::HandleConnectionCretedEvent(const FOnEventContext& Conte
 	ConnectionState->CurrentPacketBitOffset[ENetProfilerConnectionMode::Incoming] = 0U;
 }
 
+void FNetTraceAnalyzer::HandleConnectionStateUpdatedEvent(const FOnEventContext& Context, const FEventData& EventData)
+{
+	const uint8 GameInstanceId = EventData.GetValue<uint8>("GameInstanceId");
+	const uint16 ConnectionId = EventData.GetValue<uint16>("ConnectionId");
+	const uint8 ConnectionStateValue = EventData.GetValue<uint8>("ConnectionStateValue");
+
+	TSharedRef<FNetTraceGameInstanceState> GameInstanceState = GetOrCreateActiveGameInstanceState(GameInstanceId);
+
+	if (TSharedRef<FNetTraceConnectionState>* ConnectionState = GameInstanceState->ActiveConnections.Find(ConnectionId))
+	{
+		(*ConnectionState)->ConnectionState = ENetProfilerConnectionState(ConnectionStateValue);
+	}
+}
+
+void FNetTraceAnalyzer::HandleConnectionUpdatedEvent(const FOnEventContext& Context, const FEventData& EventData)
+{
+	const uint8 GameInstanceId = EventData.GetValue<uint8>("GameInstanceId");
+	const uint16 ConnectionId = EventData.GetValue<uint16>("ConnectionId");
+	FString Name;
+	EventData.GetString("Name", Name);
+	FString AddressString;
+	EventData.GetString("Address", AddressString);
+
+	TSharedRef<FNetTraceGameInstanceState> GameInstanceState = GetOrCreateActiveGameInstanceState(GameInstanceId);
+	
+	if (TSharedRef<FNetTraceConnectionState>* ConnectionState = GameInstanceState->ActiveConnections.Find(ConnectionId))
+	{
+		if (FNetProfilerConnectionInternal* Connection = NetProfilerProvider.EditConnection((*ConnectionState)->ConnectionIndex))
+		{
+			Connection->Connection.Name = Session.StoreString(Name);
+			Connection->Connection.AddressString = Session.StoreString(AddressString);
+		}
+	}
+	else
+	{
+		// Incomplete trace?  Ignore?
+		UE_LOG(LogNetTrace, Warning, TEXT("Connection %d is missing"), ConnectionId);
+	}
+
+}
+
 void FNetTraceAnalyzer::HandleConnectionClosedEvent(const FOnEventContext& Context, const FEventData& EventData)
 {
 	const uint8 GameInstanceId = EventData.GetValue<uint8>("GameInstanceId");
@@ -650,4 +713,20 @@ FNetTraceAnalyzer::FNetTraceConnectionState* FNetTraceAnalyzer::GetActiveConnect
 	return nullptr;
 }
 
+void FNetTraceAnalyzer::HandleGameInstanceUpdatedEvent(const FOnEventContext& Context, const FEventData& EventData)
+{
+	const uint8 GameInstanceId = EventData.GetValue<uint8>("GameInstanceId");
+	const bool bIsServer = EventData.GetValue<bool>("bIsServer");
+	FString InstanceName;
+	EventData.GetString("Name", InstanceName);
+
+	TSharedRef<FNetTraceGameInstanceState> GameInstanceState = GetOrCreateActiveGameInstanceState(GameInstanceId);
+
+	if (FNetProfilerGameInstanceInternal* InternalGameInstance = NetProfilerProvider.EditGameInstance(GameInstanceState->GameInstanceIndex))
+	{
+		InternalGameInstance->Instance.bIsServer = bIsServer;
+		InternalGameInstance->Instance.InstanceName = Session.StoreString(InstanceName);
+	}
+}
+
 } // namespace TraceServices
diff --git a/Engine/Source/Developer/TraceServices/Private/Analyzers/NetTraceAnalyzer.h b/Engine/Source/Developer/TraceServices/Private/Analyzers/NetTraceAnalyzer.h
--- a/Engine/Source/Developer/TraceServices/Private/Analyzers/NetTraceAnalyzer.h
+++ b/Engine/Source/Developer/TraceServices/Private/Analyzers/NetTraceAnalyzer.h
@@ -31,9 +31,12 @@ private:
 		RouteId_PacketEvent,
 		RouteId_PacketDroppedEvent,
 		RouteId_ConnectionCreatedEvent,
+		RouteId_ConnectionUpdatedEvent,
 		RouteId_ConnectionClosedEvent,
 		RouteId_ObjectCreatedEvent,
 		RouteId_ObjectDestroyedEvent,
+		RouteId_ConnectionStateUpdatedEvent,
+		RouteId_InstanceUpdatedEvent,
 	};
 
 	// This must be kept in sync with Event types in NetTrace.h
@@ -79,6 +82,8 @@ private:
 		TArray<FNetProfilerContentEvent> BunchEvents[ENetProfilerConnectionMode::Count];
 
 		TArray<FBunchInfo> BunchInfos[ENetProfilerConnectionMode::Count];
+		
+		ENetProfilerConnectionState ConnectionState;
 	};
 
 	struct FNetTraceActiveObjectState
@@ -112,11 +117,15 @@ private:
 	void HandlePacketEvent(const FOnEventContext& Context, const FEventData& EventData);	
 	void HandlePacketContentEvent(const FOnEventContext& Context, const FEventData& EventData);
 	void HandlePacketDroppedEvent(const FOnEventContext& Context, const FEventData& EventData);
-	void HandleConnectionCretedEvent(const FOnEventContext& Context, const FEventData& EventData);
+	void HandleConnectionStateUpdatedEvent(const FOnEventContext& Context, const FEventData& EventData);
+	void HandleConnectionCreatedEvent(const FOnEventContext& Context, const FEventData& EventData);
+	void HandleConnectionUpdatedEvent(const FOnEventContext& Context, const FEventData& EventData);
 	void HandleConnectionClosedEvent(const FOnEventContext& Context, const FEventData& EventData);
 	void HandleObjectCreatedEvent(const FOnEventContext& Context, const FEventData& EventData);
 	void HandleObjectDestroyedEvent(const FOnEventContext& Context, const FEventData& EventData);
 
+	void HandleGameInstanceUpdatedEvent(const FOnEventContext& Context, const FEventData& EventData);
+
 	void AddEvent(TPagedArray<FNetProfilerContentEvent>& Events, const FNetProfilerContentEvent& Event, uint32 Offset, uint32 LevelOffset);
 	void AddEvent(TPagedArray<FNetProfilerContentEvent>& Events, uint32 StartPos, uint32 EndPos, uint32 Level, uint32 NameIndex, FNetProfilerBunchInfo Info);
 
diff --git a/Engine/Source/Developer/TraceServices/Private/Model/Frames.cpp b/Engine/Source/Developer/TraceServices/Private/Model/Frames.cpp
--- a/Engine/Source/Developer/TraceServices/Private/Model/Frames.cpp
+++ b/Engine/Source/Developer/TraceServices/Private/Model/Frames.cpp
@@ -97,6 +97,23 @@ void FFrameProvider::EndFrame(ETraceFrameType FrameType, double Time)
 	Session.UpdateDurationSeconds(Time);
 }
 
+uint32 FFrameProvider::GetFrameNumberForTimestamp(ETraceFrameType FrameType, double Timestamp) const
+{
+	if (FrameStartTimes[FrameType].Num() == 0 || Timestamp < FrameStartTimes[FrameType][0])
+	{
+		return 0;
+	}
+	else if (Timestamp >= FrameStartTimes[FrameType].Last())
+	{
+		return FrameStartTimes[FrameType].Num();
+	}
+	else
+	{
+		uint32 Index = static_cast<uint32>(Algo::LowerBound(FrameStartTimes[FrameType], Timestamp));
+		return Index + 1;
+	}
+}
+
 const IFrameProvider& ReadFrameProvider(const IAnalysisSession& Session)
 {
 	return *Session.ReadProvider<IFrameProvider>(FFrameProvider::ProviderName);
diff --git a/Engine/Source/Developer/TraceServices/Private/Model/FramesPrivate.h b/Engine/Source/Developer/TraceServices/Private/Model/FramesPrivate.h
--- a/Engine/Source/Developer/TraceServices/Private/Model/FramesPrivate.h
+++ b/Engine/Source/Developer/TraceServices/Private/Model/FramesPrivate.h
@@ -23,6 +23,8 @@ public:
 	virtual const TArray64<double>& GetFrameStartTimes(ETraceFrameType FrameType) const override { return FrameStartTimes[FrameType]; }
 	virtual bool GetFrameFromTime(ETraceFrameType FrameType, double Time, FFrame& OutFrame) const override;
 	virtual const FFrame* GetFrame(ETraceFrameType FrameType, uint64 Index) const override;
+	virtual uint32 GetFrameNumberForTimestamp(ETraceFrameType FrameType, double Timestamp) const override;
+
 	void BeginFrame(ETraceFrameType FrameType, double Time);
 	void EndFrame(ETraceFrameType FrameType, double Time);
 
diff --git a/Engine/Source/Developer/TraceServices/Private/Model/NetProfilerProvider.cpp b/Engine/Source/Developer/TraceServices/Private/Model/NetProfilerProvider.cpp
--- a/Engine/Source/Developer/TraceServices/Private/Model/NetProfilerProvider.cpp
+++ b/Engine/Source/Developer/TraceServices/Private/Model/NetProfilerProvider.cpp
@@ -27,6 +27,28 @@ const TCHAR* LexToString(const ENetProfilerChannelCloseReason Value)
 	return TEXT("Unknown");
 }
 
+const TCHAR* LexToString(const ENetProfilerConnectionState Value)
+{
+	switch (Value)
+	{
+	case ENetProfilerConnectionState::USOCK_Closed:
+		return TEXT("Closed");
+		break;
+	case ENetProfilerConnectionState::USOCK_Open:
+		return TEXT("Open");
+		break;
+	case ENetProfilerConnectionState::USOCK_Pending:
+		return TEXT("Pending");
+		break;
+	case ENetProfilerConnectionState::USOCK_Invalid:
+	default:
+		return TEXT("Invalid");
+		break;
+	}
+
+	return TEXT("Invalid");
+}
+
 FNetProfilerProvider::FNetProfilerProvider(IAnalysisSession& InSession)
 	: Session(InSession)
 	, NetTraceVersion(0U)
diff --git a/Engine/Source/Developer/TraceServices/Public/TraceServices/Model/Frames.h b/Engine/Source/Developer/TraceServices/Public/TraceServices/Model/Frames.h
--- a/Engine/Source/Developer/TraceServices/Public/TraceServices/Model/Frames.h
+++ b/Engine/Source/Developer/TraceServices/Public/TraceServices/Model/Frames.h
@@ -29,6 +29,7 @@ public:
 	virtual const TArray64<double>& GetFrameStartTimes(ETraceFrameType FrameType) const = 0;
 	virtual bool GetFrameFromTime(ETraceFrameType FrameType, double Time, FFrame& OutFrame) const = 0;
 	virtual const FFrame* GetFrame(ETraceFrameType FrameType, uint64 Index) const = 0;
+	virtual uint32 GetFrameNumberForTimestamp(ETraceFrameType FrameType, double Timestamp) const = 0;
 };
 
 TRACESERVICES_API const IFrameProvider& ReadFrameProvider(const IAnalysisSession& Session);
diff --git a/Engine/Source/Developer/TraceServices/Public/TraceServices/Model/NetProfiler.h b/Engine/Source/Developer/TraceServices/Public/TraceServices/Model/NetProfiler.h
--- a/Engine/Source/Developer/TraceServices/Public/TraceServices/Model/NetProfiler.h
+++ b/Engine/Source/Developer/TraceServices/Public/TraceServices/Model/NetProfiler.h
@@ -35,7 +35,17 @@ enum class ENetProfilerChannelCloseReason : uint8
 	TearOff,
 };
 
+// Mirrored from EConnectionState on the runtime side
+enum class ENetProfilerConnectionState : uint8
+{
+	USOCK_Invalid = 0, // Connection is invalid, possibly uninitialized.
+	USOCK_Closed = 1, // Connection permanently closed.
+	USOCK_Pending = 2, // Connection is awaiting connection.
+	USOCK_Open = 3, // Connection is open.
+};
+
 TRACESERVICES_API const TCHAR* LexToString(const ENetProfilerChannelCloseReason Value);
+TRACESERVICES_API const TCHAR* LexToString(const ENetProfilerConnectionState Value);
 
 struct FNetProfilerName
 {
@@ -123,16 +133,19 @@ struct FNetProfilerPacket
 	uint32 SequenceNumber;
 	uint32 ContentSizeInBits;						// This is the part that is tracked by the PacketContents
 	uint32 TotalPacketSizeInBytes;					// This is the actual size of the packet sent on the socket
-	ENetProfilerDeliveryStatus DeliveryStatus;		// Indicates if the packet was delivered or not, updated as soon as we know
 
 	// Index into Events
 	uint32 StartEventIndex;
 	uint32 EventCount;
+
+	ENetProfilerDeliveryStatus DeliveryStatus;		// Indicates if the packet was delivered or not, updated as soon as we know
+	ENetProfilerConnectionState ConnectionState;
 };
 
 struct FNetProfilerConnection
 {
-	const TCHAR* Name;
+	const TCHAR* Name = nullptr;
+	const TCHAR* AddressString = nullptr;
 	FNetProfilerLifeTime LifeTime;
 	uint32 GameInstanceIndex;
 	uint32 ConnectionIndex : 16;
@@ -146,6 +159,8 @@ struct FNetProfilerGameInstance
 	FNetProfilerLifeTime LifeTime;
 	uint32 GameInstanceIndex;
 	uint32 GameInstanceId;
+	const TCHAR* InstanceName = nullptr;
+	bool bIsServer;
 };
 
 // What do we need?
diff --git a/Engine/Source/Runtime/Engine/Classes/Engine/ChildConnection.h b/Engine/Source/Runtime/Engine/Classes/Engine/ChildConnection.h
--- a/Engine/Source/Runtime/Engine/Classes/Engine/ChildConnection.h
+++ b/Engine/Source/Runtime/Engine/Classes/Engine/ChildConnection.h
@@ -71,7 +71,7 @@ public:
 
 	virtual void Tick(float DeltaSeconds) override
 	{
-		State = Parent->State;
+		SetConnectionState(Parent->GetConnectionState());
 	}
 
 	virtual void HandleClientPlayer(class APlayerController* PC, class UNetConnection* NetConnection) override;
diff --git a/Engine/Source/Runtime/Engine/Classes/Engine/EngineTypes.h b/Engine/Source/Runtime/Engine/Classes/Engine/EngineTypes.h
--- a/Engine/Source/Runtime/Engine/Classes/Engine/EngineTypes.h
+++ b/Engine/Source/Runtime/Engine/Classes/Engine/EngineTypes.h
@@ -1046,6 +1046,8 @@ namespace EWorldType
 	};
 }
 
+const TCHAR* LexToString(const EWorldType::Type Value);
+
 /** Describes what parts of level streaming should be forcibly handled immediately */
 enum class EFlushLevelStreamingType : uint8
 {
diff --git a/Engine/Source/Runtime/Engine/Classes/Engine/NetConnection.h b/Engine/Source/Runtime/Engine/Classes/Engine/NetConnection.h
--- a/Engine/Source/Runtime/Engine/Classes/Engine/NetConnection.h
+++ b/Engine/Source/Runtime/Engine/Classes/Engine/NetConnection.h
@@ -73,6 +73,7 @@ enum EConnectionState
 	USOCK_Pending	= 2, // Connection is awaiting connection.
 	USOCK_Open      = 3, // Connection is open.
 };
+ENGINE_API const TCHAR* LexToString(const EConnectionState Value);
 
 // 
 // Security event types used for UE_SECURITY_LOG
@@ -377,7 +378,11 @@ public:
 public:
 	// Connection information.
 
+	UE_DEPRECATED(5.0, "EConnectionState State will be made private. Use GetConnectionState and SetConnectionState to access State instead.")
 	EConnectionState	State;					// State this connection is in.
+
+	ENGINE_API const EConnectionState GetConnectionState() const;
+	ENGINE_API void SetConnectionState(EConnectionState ConnectionState);
 	
 	uint32 bPendingDestroy:1;    // when true, playercontroller or beaconclient is being destroyed
 
@@ -1350,6 +1355,11 @@ protected:
 	/** This is called whenever a connection has passed the relative time to be considered timed out. */
 	ENGINE_API virtual void HandleConnectionTimeout(const FString& Error);
 
+	/**
+	 * Notification that information about this connection may have been updated
+	 */
+	ENGINE_API void NotifyConnectionUpdated();
+
 private:
 	/**
 	 * The channels that need ticking. This will be a subset of OpenChannels, only including
diff --git a/Engine/Source/Runtime/Engine/Classes/Engine/NetDriver.h b/Engine/Source/Runtime/Engine/Classes/Engine/NetDriver.h
--- a/Engine/Source/Runtime/Engine/Classes/Engine/NetDriver.h
+++ b/Engine/Source/Runtime/Engine/Classes/Engine/NetDriver.h
@@ -1723,6 +1723,9 @@ protected:
 	/** Stream of random numbers to be used by this instance of UNetDriver */
 	FRandomStream UpdateDelayRandomStream;
 
+	/** Creates a trace event that updates the name and properties of the associated Game Instance */
+	void NotifyGameInstanceUpdated();
+
 private:
 
 	ENGINE_API virtual ECreateReplicationChangelistMgrFlags GetCreateReplicationChangelistMgrFlags() const;
diff --git a/Engine/Source/Runtime/Engine/Private/DataChannel.cpp b/Engine/Source/Runtime/Engine/Private/DataChannel.cpp
--- a/Engine/Source/Runtime/Engine/Private/DataChannel.cpp
+++ b/Engine/Source/Runtime/Engine/Private/DataChannel.cpp
@@ -141,7 +141,7 @@ int64 UChannel::Close(EChannelCloseReason Reason)
 
 	int64 NumBits = 0;
 
-	if ( !Closing && ( Connection->State == USOCK_Open || Connection->State == USOCK_Pending ) && !SentClosingBunch)
+	if ( !Closing && ( Connection->GetConnectionState() == USOCK_Open || Connection->GetConnectionState() == USOCK_Pending ) && !SentClosingBunch)
 	{
 		if ( ChIndex == 0 )
 		{
@@ -1499,7 +1499,7 @@ void UControlChannel::ReceivedBunch( FInBunch& Bunch )
 	}
 
 	// Process the packet
-	while (!Bunch.AtEnd() && Connection != NULL && Connection->State != USOCK_Closed) // if the connection got closed, we don't care about the rest
+	while (!Bunch.AtEnd() && Connection != NULL && Connection->GetConnectionState() != USOCK_Closed) // if the connection got closed, we don't care about the rest
 	{
 		uint8 MessageType = 0;
 		Bunch << MessageType;
@@ -1705,7 +1705,7 @@ void UControlChannel::QueueMessage(const FOutBunch* Bunch)
 		// we're out of room in our extra buffer as well, so kill the connection
 		UE_LOG(LogNet, Log, TEXT("Overflowed control channel message queue, disconnecting client"));
 		// intentionally directly setting State as the messaging in Close() is not going to work in this case
-		Connection->State = USOCK_Closed;
+		Connection->SetConnectionState(USOCK_Closed);
 	}
 	else
 	{
@@ -1828,7 +1828,7 @@ int64 UControlChannel::SendDestructionInfo(FActorDestructionInfo* DestructionInf
 	checkf(Connection && Connection->PackageMap, TEXT("SendDestructionInfo requires a valid connection and package map: %s"), *Describe());
 	checkf(DestructionInfo, TEXT("SendDestructionInfo was passed an invalid desctruction info: %s"), *Describe());
 
-	if (!Closing && (Connection->State == USOCK_Open || Connection->State == USOCK_Pending))
+	if (!Closing && (Connection->GetConnectionState() == USOCK_Open || Connection->GetConnectionState() == USOCK_Pending))
 	{
 		// Outer must be valid to call PackageMap->WriteObject. In the case of streaming out levels, this can go null out of from underneath us. In that case, just skip the destruct info.
 		// We assume that if server unloads a level that clients will to and this will implicitly destroy all actors in it, so not worried about leaking actors client side here.
@@ -2397,7 +2397,7 @@ int64 UActorChannel::SetChannelActorForDestroy( FActorDestructionInfo *DestructI
 
 	if
 	(	!Closing
-	&&	(Connection->State==USOCK_Open || Connection->State==USOCK_Pending) )
+	&&	(Connection->GetConnectionState()==USOCK_Open || Connection->GetConnectionState()==USOCK_Pending) )
 	{
 		// Outer must be valid to call PackageMap->WriteObject. In the case of streaming out levels, this can go null out of from underneath us. In that case, just skip the destruct info.
 		// We assume that if server unloads a level that clients will to and this will implicitly destroy all actors in it, so not worried about leaking actors client side here.
@@ -2805,7 +2805,7 @@ void UActorChannel::ProcessBunch( FInBunch & Bunch )
 	// ----------------------------------------------
 	//	Read chunks of actor content
 	// ----------------------------------------------
-	while ( !Bunch.AtEnd() && Connection != NULL && Connection->State != USOCK_Closed )
+	while ( !Bunch.AtEnd() && Connection != NULL && Connection->GetConnectionState() != USOCK_Closed )
 	{
 		FNetBitReader Reader( Bunch.PackageMap, 0 );
 
diff --git a/Engine/Source/Runtime/Engine/Private/DataReplication.cpp b/Engine/Source/Runtime/Engine/Private/DataReplication.cpp
--- a/Engine/Source/Runtime/Engine/Private/DataReplication.cpp
+++ b/Engine/Source/Runtime/Engine/Private/DataReplication.cpp
@@ -2145,7 +2145,7 @@ void FObjectReplicator::UpdateUnmappedObjects(bool & bOutHasMoreUnmapped)
 		return;
 	}
 
-	if (Connection->State == USOCK_Closed)
+	if (Connection->GetConnectionState() == USOCK_Closed)
 	{
 		UE_LOG(LogNet, Verbose, TEXT("FObjectReplicator::UpdateUnmappedObjects: Connection->State == USOCK_Closed"));
 		return;
diff --git a/Engine/Source/Runtime/Engine/Private/DemoNetDriver.cpp b/Engine/Source/Runtime/Engine/Private/DemoNetDriver.cpp
--- a/Engine/Source/Runtime/Engine/Private/DemoNetDriver.cpp
+++ b/Engine/Source/Runtime/Engine/Private/DemoNetDriver.cpp
@@ -1093,13 +1093,13 @@ bool UDemoNetDriver::ContinueListen(FURL& ListenURL)
 
 bool UDemoNetDriver::IsRecording() const
 {
-	return ClientConnections.Num() > 0 && ClientConnections[0] != nullptr && ClientConnections[0]->State != USOCK_Closed;
+	return ClientConnections.Num() > 0 && ClientConnections[0] != nullptr && ClientConnections[0]->GetConnectionState() != USOCK_Closed;
 }
 
 bool UDemoNetDriver::IsPlaying() const
 {
 	// ServerConnection may be deleted / recreated during checkpoint loading.
-	return IsLoadingCheckpoint() || (ServerConnection != nullptr && ServerConnection->State != USOCK_Closed);
+	return IsLoadingCheckpoint() || (ServerConnection != nullptr && ServerConnection->GetConnectionState() != USOCK_Closed);
 }
 
 bool UDemoNetDriver::IsServer() const
@@ -1388,7 +1388,7 @@ void UDemoNetDriver::StopDemo()
 		// flush out any pending network traffic
 		ServerConnection->FlushNet();
 
-		ServerConnection->State = USOCK_Closed;
+		ServerConnection->SetConnectionState(USOCK_Closed);
 		ServerConnection->Close();
 	}
 
@@ -2353,7 +2353,7 @@ bool UDemoNetDriver::ProcessPacket(const uint8* Data, int32 Count)
 		ServerConnection->ReceivedRawPacket(const_cast<uint8*>(Data), Count);
 	}
 
-	if (ServerConnection == nullptr || ServerConnection->State == USOCK_Closed)
+	if (ServerConnection == nullptr || ServerConnection->GetConnectionState() == USOCK_Closed)
 	{
 		// Something we received resulted in the demo being stopped
 		UE_LOG(LogDemo, Error, TEXT("UDemoNetDriver::ProcessPacket: ReceivedRawPacket closed connection"));
@@ -4526,7 +4526,7 @@ void UDemoNetDriver::RestoreConnectionPostScrub(APlayerController* PC, UNetConne
 	NetConnection->LastReceiveTime = GetElapsedTime();
 	NetConnection->LastReceiveRealtime = FPlatformTime::Seconds();
 	NetConnection->LastGoodPacketRealtime = FPlatformTime::Seconds();
-	NetConnection->State = USOCK_Open;
+	NetConnection->SetConnectionState(USOCK_Open);
 	NetConnection->PlayerController = PC;
 	NetConnection->OwningActor = PC;
 }
diff --git a/Engine/Source/Runtime/Engine/Private/EngineTypes.cpp b/Engine/Source/Runtime/Engine/Private/EngineTypes.cpp
--- a/Engine/Source/Runtime/Engine/Private/EngineTypes.cpp
+++ b/Engine/Source/Runtime/Engine/Private/EngineTypes.cpp
@@ -656,4 +656,38 @@ FVector FRepMovement::RebaseOntoZeroOrigin(const FVector& Location, const UActor
 	return RebaseOntoZeroOrigin(Location, WorldContextActorComponent->GetWorld()->OriginLocation);
 }
 
+const TCHAR* LexToString(const EWorldType::Type Value)
+{
+	switch (Value)
+	{
+	case EWorldType::Type::Editor:
+		return TEXT("Editor");
+		break;
+	case EWorldType::Type::EditorPreview:
+		return TEXT("EditorPreview");
+		break;
+	case EWorldType::Type::Game:
+		return TEXT("Game");
+		break;
+	case EWorldType::Type::GamePreview:
+		return TEXT("GamePreview");
+		break;
+	case EWorldType::Type::GameRPC:
+		return TEXT("GameRPC");
+		break;
+	case EWorldType::Type::Inactive:
+		return TEXT("Inactive");
+		break;
+	case EWorldType::Type::PIE:
+		return TEXT("PIE");
+		break;
+	case EWorldType::Type::None:
+		return TEXT("None");
+		break;
+	default:
+		return TEXT("Unknown");
+		break;
+	}
+}
+
 /// @endcond
diff --git a/Engine/Source/Runtime/Engine/Private/LevelTick.cpp b/Engine/Source/Runtime/Engine/Private/LevelTick.cpp
--- a/Engine/Source/Runtime/Engine/Private/LevelTick.cpp
+++ b/Engine/Source/Runtime/Engine/Private/LevelTick.cpp
@@ -437,7 +437,7 @@ void UWorld::TickNetClient( float DeltaSeconds )
 
 	// If our net driver has lost connection to the server,
 	// and there isn't a PendingNetGame, throw a network failure error.
-	if( NetDriver->ServerConnection->State == USOCK_Closed )
+	if( NetDriver->ServerConnection->GetConnectionState() == USOCK_Closed )
 	{
 		if (GEngine->PendingNetGameFromWorld(this) == nullptr)
 		{
@@ -1462,7 +1462,7 @@ void UWorld::Tick( ELevelTick TickType, float DeltaSeconds )
 	bool bDoingActorTicks = 
 		(TickType!=LEVELTICK_TimeOnly)
 		&&	!bIsPaused
-		&&	(!NetDriver || !NetDriver->ServerConnection || NetDriver->ServerConnection->State==USOCK_Open);
+		&&	(!NetDriver || !NetDriver->ServerConnection || NetDriver->ServerConnection->GetConnectionState()==USOCK_Open);
 
 	FLatentActionManager& CurrentLatentActionManager = GetLatentActionManager();
 
diff --git a/Engine/Source/Runtime/Engine/Private/LocalPlayer.cpp b/Engine/Source/Runtime/Engine/Private/LocalPlayer.cpp
--- a/Engine/Source/Runtime/Engine/Private/LocalPlayer.cpp
+++ b/Engine/Source/Runtime/Engine/Private/LocalPlayer.cpp
@@ -333,7 +333,7 @@ void ULocalPlayer::SendSplitJoin(TArray<FString>& Options)
 		NetDriver = World->GetNetDriver();
 	}
 
-	if (World == NULL || NetDriver == NULL || NetDriver->ServerConnection == NULL || NetDriver->ServerConnection->State != USOCK_Open)
+	if (World == NULL || NetDriver == NULL || NetDriver->ServerConnection == NULL || NetDriver->ServerConnection->GetConnectionState() != USOCK_Open)
 	{
 		UE_LOG(LogPlayerManagement, Warning, TEXT("SendSplitJoin(): Not connected to a server"));
 	}
diff --git a/Engine/Source/Runtime/Engine/Private/NetConnection.cpp b/Engine/Source/Runtime/Engine/Private/NetConnection.cpp
--- a/Engine/Source/Runtime/Engine/Private/NetConnection.cpp
+++ b/Engine/Source/Runtime/Engine/Private/NetConnection.cpp
@@ -366,7 +366,7 @@ void UNetConnection::InitBase(UNetDriver* InDriver,class FSocket* InSocket, cons
 	NetConnectionHistogram.InitHitchTracking();
 
 	// Current state
-	State = InState;
+	SetConnectionState(InState);
 	// Copy the URL
 	URL = InURL;
 
@@ -433,6 +433,7 @@ void UNetConnection::InitBase(UNetDriver* InDriver,class FSocket* InSocket, cons
 	}
 
 	UE_NET_TRACE_CONNECTION_CREATED(NetTraceId, GetConnectionId());
+	UE_NET_TRACE_CONNECTION_STATE_UPDATED(NetTraceId, GetConnectionId(), static_cast<uint8>(GetConnectionState()));
 }
 
 /**
@@ -450,7 +451,7 @@ void UNetConnection::InitConnection(UNetDriver* InDriver, EConnectionState InSta
 	// We won't be sending any packets, so use a default size
 	MaxPacket = (InMaxPacket == 0 || InMaxPacket > MAX_PACKET_SIZE) ? MAX_PACKET_SIZE : InMaxPacket;
 	PacketOverhead = 0;
-	State = InState;
+	SetConnectionState(InState);
 
 #if DO_ENABLE_NET_TEST
 	// Copy the command line settings from the net driver
@@ -566,6 +567,14 @@ void UNetConnection::NotifyAnalyticsProvider()
 	}
 }
 
+void UNetConnection::NotifyConnectionUpdated()
+{
+	if (const uint32 MyConnectionId = GetConnectionId() && OwningActor && RemoteAddr)
+	{
+		UE_NET_TRACE_CONNECTION_UPDATED(NetTraceId, MyConnectionId, *(RemoteAddr->ToString(true)), *(OwningActor->GetName()));
+	}
+}
+
 void UNetConnection::EnableEncryptionWithKey(TArrayView<const uint8> Key)
 {
 	FEncryptionData EncryptionData;
@@ -603,7 +612,7 @@ void UNetConnection::EnableEncryptionWithKeyServer(TArrayView<const uint8> Key)
 
 void UNetConnection::EnableEncryptionServer(const FEncryptionData& EncryptionData)
 {
-	if (State != USOCK_Invalid && State != USOCK_Closed && Driver)
+	if (GetConnectionState() != USOCK_Invalid && GetConnectionState() != USOCK_Closed && Driver)
 	{
 		SendClientEncryptionAck();
 		EnableEncryption(EncryptionData);
@@ -616,7 +625,7 @@ void UNetConnection::EnableEncryptionServer(const FEncryptionData& EncryptionDat
 
 void UNetConnection::SendClientEncryptionAck()
 {
-	if (State != USOCK_Invalid && State != USOCK_Closed && Driver)
+	if (GetConnectionState() != USOCK_Invalid && GetConnectionState() != USOCK_Closed && Driver)
 	{
 		FNetControlMessage<NMT_EncryptionAck>::Send(this);
 		FlushNet();
@@ -785,6 +794,41 @@ void UNetConnection::Serialize( FArchive& Ar )
 	}
 }
 
+PRAGMA_DISABLE_DEPRECATION_WARNINGS
+EConnectionState const UNetConnection::GetConnectionState() const
+{
+	return State;
+}
+
+void UNetConnection::SetConnectionState(EConnectionState ConnectionState)
+{
+	State = ConnectionState;
+	UE_NET_TRACE_CONNECTION_STATE_UPDATED(NetTraceId, GetConnectionId(), static_cast<uint8>(State));
+}
+PRAGMA_ENABLE_DEPRECATION_WARNINGS
+
+const TCHAR* LexToString(const EConnectionState Value)
+{
+	switch (Value)
+	{
+	case EConnectionState::USOCK_Closed:
+		return TEXT("Closed");
+		break;
+	case EConnectionState::USOCK_Open:
+		return TEXT("Open");
+		break;
+	case EConnectionState::USOCK_Pending:
+		return TEXT("Pending");
+		break;
+	case EConnectionState::USOCK_Invalid:
+	default:
+		return TEXT("Invalid");
+		break;
+	}
+
+	return TEXT("Invalid");
+}
+
 void UNetConnection::Close()
 {
 	if (IsInternalAck())
@@ -793,7 +837,7 @@ void UNetConnection::Close()
 		SetIgnoreReservedChannels(false);
 	}
 
-	if (Driver != nullptr && State != USOCK_Closed)
+	if (Driver != nullptr && GetConnectionState() != USOCK_Closed)
 	{
 		NETWORK_PROFILER(GNetworkProfiler.TrackEvent(TEXT("CLOSE"), *(GetName() + TEXT(" ") + LowLevelGetRemoteAddress()), this));
 		UE_LOG(LogNet, Log, TEXT("UNetConnection::Close: %s, Channels: %i, Time: %s"), *Describe(), OpenChannels.Num(), *FDateTime::UtcNow().ToString(TEXT("%Y.%m.%d-%H.%M.%S")));
@@ -802,7 +846,7 @@ void UNetConnection::Close()
 		{
 			Channels[0]->Close(EChannelCloseReason::Destroyed);
 		}
-		State = USOCK_Closed;
+		SetConnectionState(EConnectionState::USOCK_Closed);
 
 		if ((Handler == nullptr || Handler->IsFullyInitialized()) && HasReceivedClientPacket())
 		{
@@ -848,7 +892,7 @@ void UNetConnection::CleanUp()
 	}
 	Children.Empty();
 
-	if ( State != USOCK_Closed )
+	if ( GetConnectionState() != USOCK_Closed )
 	{
 		UE_LOG( LogNet, Log, TEXT( "UNetConnection::Cleanup: Closing open connection. %s" ), *Describe() );
 	}
@@ -1049,7 +1093,7 @@ bool UNetConnection::Exec( UWorld* InWorld, const TCHAR* Cmd, FOutputDevice& Ar
 void UNetConnection::AssertValid()
 {
 	// Make sure this connection is in a reasonable state.
-	check(State==USOCK_Closed || State==USOCK_Pending || State==USOCK_Open);
+	check(GetConnectionState()==USOCK_Closed || GetConnectionState()==USOCK_Pending || GetConnectionState()==USOCK_Open);
 
 }
 
@@ -1118,6 +1162,8 @@ void UNetConnection::UpdateLevelVisibility(const FUpdateLevelVisibilityLevelInfo
 		}
 	}
 	UpdateLevelVisibilityInternal(LevelVisibility);
+
+	NotifyConnectionUpdated();
 }
 
 void UNetConnection::UpdateLevelVisibilityInternal(const FUpdateLevelVisibilityLevelInfo& LevelVisibility)
@@ -1489,7 +1535,7 @@ void UNetConnection::FlushNet(bool bIgnoreSimulation)
 	TimeSensitive = 0;
 
 	// If there is any pending data to send, send it.
-	if (SendBuffer.GetNumBits() || HasDirtyAcks || ( Driver->GetElapsedTime() - LastSendTime > Driver->KeepAliveTime && !IsInternalAck() && State != USOCK_Closed))
+	if (SendBuffer.GetNumBits() || HasDirtyAcks || ( Driver->GetElapsedTime() - LastSendTime > Driver->KeepAliveTime && !IsInternalAck() && GetConnectionState() != USOCK_Closed))
 	{
 		// Due to the PacketHandler handshake code, servers must never send the client data,
 		// before first receiving a client control packet (which is taken as an indication of a complete handshake).
@@ -1566,7 +1612,7 @@ void UNetConnection::FlushNet(bool bIgnoreSimulation)
 
 		// if the connection is closing/being destroyed/etc we need to send immediately regardless of settings
 		// because we won't be around to send it delayed
-		if (State != USOCK_Closed && !IsGarbageCollecting() && !bIgnoreSimulation && !IsInternalAck())
+		if (GetConnectionState() != USOCK_Closed && !IsGarbageCollecting() && !bIgnoreSimulation && !IsInternalAck())
 		{
 			bWasPacketEmulated = CheckOutgoingPacketEmulation(Traits);
 		}
@@ -2385,7 +2431,7 @@ void UNetConnection::ReceivedPacket( FBitReader& Reader, bool bIsReinjectedPacke
 	}
 
 	// Disassemble and dispatch all bunches in the packet.
-	while( !Reader.AtEnd() && State!=USOCK_Closed )
+	while( !Reader.AtEnd() && GetConnectionState()!=USOCK_Closed )
 	{
 		// For demo backwards compatibility, old replays still have this bit
 		if (IsInternalAck() && EngineNetworkProtocolVersion < EEngineNetworkVersionHistory::HISTORY_ACKS_INCLUDED_IN_HEADER)
@@ -2849,7 +2895,7 @@ void UNetConnection::ReceivedPacket( FBitReader& Reader, bool bIsReinjectedPacke
 					if( Bunch.ChIndex==0 )
 					{
 						UE_LOG(LogNetTraffic, Log, TEXT("Channel 0 create failed") );
-						State = USOCK_Closed;
+						SetConnectionState(EConnectionState::USOCK_Closed);
 					}
 					continue;
 				}
@@ -3463,7 +3509,7 @@ float UNetConnection::GetTimeoutValue()
 
 	float Timeout = Driver->InitialConnectTimeout;
 
-	if ((State != USOCK_Pending) && (bPendingDestroy || (OwningActor && OwningActor->UseShortConnectTimeout())))
+	if ((GetConnectionState() != USOCK_Pending) && (bPendingDestroy || (OwningActor && OwningActor->UseShortConnectTimeout())))
 	{
 		const float ConnectionTimeout = Driver->ConnectionTimeout;
 
@@ -3773,7 +3819,7 @@ void UNetConnection::Tick(float DeltaSeconds)
 		// If channel 0 has closed, mark the connection as closed.
 		if (Channels[0] == nullptr && (OutReliable[0] != InitOutReliable || InReliable[0] != InitInReliable))
 		{
-			State = USOCK_Closed;
+			SetConnectionState(EConnectionState::USOCK_Closed);
 		}
 	}
 
@@ -3922,7 +3968,7 @@ void UNetConnection::HandleClientPlayer( APlayerController *PC, UNetConnection*
 	PC->SetPlayer(LocalPlayer);
 	UE_LOG(LogNet, Verbose, TEXT("%s setplayer %s"),*PC->GetName(),*LocalPlayer->GetName());
 	LastReceiveTime = Driver->GetElapsedTime();
-	State = USOCK_Open;
+	SetConnectionState(EConnectionState::USOCK_Open);
 	PlayerController = PC;
 	OwningActor = PC;
 
@@ -3959,6 +4005,8 @@ void UNetConnection::HandleClientPlayer( APlayerController *PC, UNetConnection*
 			It->SendSplitJoin(Options);
 		}
 	}
+
+	NotifyConnectionUpdated();
 }
 
 void UChildConnection::HandleClientPlayer(APlayerController* PC, UNetConnection* NetConnection)
@@ -4027,6 +4075,8 @@ void UChildConnection::HandleClientPlayer(APlayerController* PC, UNetConnection*
 	UE_LOG(LogNet, Verbose, TEXT("%s setplayer %s"), *PC->GetName(), *NewPlayer->GetName());
 	PlayerController = PC;
 	OwningActor = PC;
+
+	NotifyConnectionUpdated();
 }
 
 #if DO_ENABLE_NET_TEST
@@ -4472,7 +4522,7 @@ void UNetConnection::ProcessJitter(uint32 PacketJitterClockTimeMS)
 
 void UNetConnection::SendChallengeControlMessage()
 {
-	if (State != USOCK_Invalid && State != USOCK_Closed && Driver)
+	if (GetConnectionState() != USOCK_Invalid && GetConnectionState() != USOCK_Closed && Driver)
 	{
 		Challenge = FString::Printf(TEXT("%08X"), FPlatformTime::Cycles());
 		SetExpectedClientLoginMsgType(NMT_Login);
@@ -4487,7 +4537,7 @@ void UNetConnection::SendChallengeControlMessage()
 
 void UNetConnection::SendChallengeControlMessage(const FEncryptionKeyResponse& Response)
 {
-	if (State != USOCK_Invalid && State != USOCK_Closed && Driver)
+	if (GetConnectionState() != USOCK_Invalid && GetConnectionState() != USOCK_Closed && Driver)
 	{
 		if (Response.Response == EEncryptionResponse::Success)
 		{
@@ -4560,9 +4610,11 @@ USimulatedClientNetConnection::USimulatedClientNetConnection( const FObjectIniti
 
 void USimulatedClientNetConnection::HandleClientPlayer( class APlayerController* PC, class UNetConnection* NetConnection )
 {
-	State = USOCK_Open;
+	SetConnectionState(EConnectionState::USOCK_Open);
 	PlayerController = PC;
 	OwningActor = PC;
+
+	NotifyConnectionUpdated();
 }
 
 // ----------------------------------------------------------------
diff --git a/Engine/Source/Runtime/Engine/Private/NetDriver.cpp b/Engine/Source/Runtime/Engine/Private/NetDriver.cpp
--- a/Engine/Source/Runtime/Engine/Private/NetDriver.cpp
+++ b/Engine/Source/Runtime/Engine/Private/NetDriver.cpp
@@ -489,6 +489,8 @@ void UNetDriver::PostInitProperties()
 
 		LoadChannelDefinitions();
 	}
+
+	NotifyGameInstanceUpdated();
 }
 
 void UNetDriver::PostReloadConfig(FProperty* PropertyToLoad)
@@ -538,6 +540,17 @@ void UNetDriver::LoadChannelDefinitions()
 	}
 }
 
+void UNetDriver::NotifyGameInstanceUpdated()
+{
+#if UE_NET_TRACE_ENABLED
+	if (GetWorld() && GetWorld()->WorldType)
+	{
+		FString InstanceName = FString::Printf(TEXT("%s (%s)"), *GetNameSafe(GetWorld()->GetGameInstance()), LexToString(GetWorld()->WorldType.GetValue()));
+		UE_NET_TRACE_UPDATE_INSTANCE(GetUniqueID(), IsServer(), *InstanceName);
+	}
+#endif
+}
+
 void UNetDriver::AssertValid()
 {
 }
@@ -1900,7 +1913,7 @@ void UNetDriver::TickDispatch( float DeltaTime )
 			{
 				UNetConnection* CurConn = ClientConnections[ConnIdx];
 
-				if (CurConn->State == USOCK_Closed)
+				if (CurConn->GetConnectionState() == USOCK_Closed)
 				{
 					CurConn->CleanUp();
 				}
@@ -2059,7 +2072,7 @@ void UNetDriver::InternalProcessRemoteFunctionPrivate(
 	}
 
 	// Prevent RPC calls to closed connections
-	if (Connection->State == USOCK_Closed)
+	if (Connection->GetConnectionState() == USOCK_Closed)
 	{
 		DEBUG_REMOTEFUNCTION(TEXT("Attempting to call RPC on a closed connection. Not calling %s::%s"), *GetNameSafe(Actor), *GetNameSafe(Function));
 		return;
@@ -2817,7 +2830,7 @@ bool UNetDriver::HandleNetFloodCommand( const TCHAR* Cmd, FOutputDevice& Ar )
 	{
 		Ar.Logf(TEXT("Flooding connection 0 with control messages"));
 
-		for (int32 i = 0; i < 256 && TestConn->State == USOCK_Open; i++)
+		for (int32 i = 0; i < 256 && TestConn->GetConnectionState() == USOCK_Open; i++)
 		{
 			FNetControlMessage<NMT_Netspeed>::Send(TestConn, TestConn->CurrentNetSpeed);
 			TestConn->FlushNet();
@@ -3633,7 +3646,7 @@ UChildConnection* UNetDriver::CreateChild(UNetConnection* Parent)
 	auto Child = NewObject<UChildConnection>();
 	Child->Driver = this;
 	Child->URL = FURL();
-	Child->State = Parent->State;
+	Child->SetConnectionState(Parent->GetConnectionState());
 	Child->URL.Host = Parent->URL.Host;
 	Child->Parent = Parent;
 	Child->PackageMap = Parent->PackageMap;
@@ -4128,14 +4141,14 @@ int32 UNetDriver::ServerReplicateActors_PrepConnections( const float DeltaSecond
 	{
 		UNetConnection* Connection = ClientConnections[ConnIdx];
 		check( Connection );
-		check( Connection->State == USOCK_Pending || Connection->State == USOCK_Open || Connection->State == USOCK_Closed );
+		check( Connection->GetConnectionState() == USOCK_Pending || Connection->GetConnectionState() == USOCK_Open || Connection->GetConnectionState() == USOCK_Closed );
 		checkSlow( Connection->GetUChildConnection() == NULL );
 
 		// Handle not ready channels.
 		//@note: we cannot add Connection->IsNetReady(0) here to check for saturation, as if that's the case we still want to figure out the list of relevant actors
 		//			to reset their NetUpdateTime so that they will get sent as soon as the connection is no longer saturated
 		AActor* OwningActor = Connection->OwningActor;
-		if ( OwningActor != NULL && Connection->State == USOCK_Open && ( Connection->Driver->GetElapsedTime() - Connection->LastReceiveTime < 1.5 ) )
+		if ( OwningActor != NULL && Connection->GetConnectionState() == USOCK_Open && ( Connection->Driver->GetElapsedTime() - Connection->LastReceiveTime < 1.5 ) )
 		{
 			check( World == OwningActor->GetWorld() );
 
@@ -5660,6 +5673,8 @@ void UNetDriver::SetWorld(class UWorld* InWorld)
 		RegisterTickEvents(InWorld);
 
 		GetNetworkObjectList().AddInitialObjects(InWorld, this);
+
+		NotifyGameInstanceUpdated();
 	}
 
 	if (ReplicationDriver)
@@ -5805,6 +5820,8 @@ void UNetDriver::SetReplicationDriver(UReplicationDriver* NewReplicationDriver)
 		ReplicationDriver->InitForNetDriver(this);
 		ReplicationDriver->InitializeActorsInWorld( GetWorld() );
 	}
+
+	NotifyGameInstanceUpdated();
 }
 
 UNetConnection* UNetDriver::GetConnectionById(uint32 ConnectionId) const
diff --git a/Engine/Source/Runtime/Engine/Private/NetworkObjectList.cpp b/Engine/Source/Runtime/Engine/Private/NetworkObjectList.cpp
--- a/Engine/Source/Runtime/Engine/Private/NetworkObjectList.cpp
+++ b/Engine/Source/Runtime/Engine/Private/NetworkObjectList.cpp
@@ -123,7 +123,7 @@ void FNetworkObjectList::Remove(AActor* const Actor)
 	{
 		UNetConnection* Connection = (*ConnectionIt).Get();
 
-		if (Connection == nullptr || Connection->State == USOCK_Closed)
+		if (Connection == nullptr || Connection->GetConnectionState() == USOCK_Closed)
 		{
 			ConnectionIt.RemoveCurrent();
 			continue;
@@ -178,7 +178,7 @@ void FNetworkObjectList::MarkDormant(AActor* const Actor, UNetConnection* const
 	// Clean up DormantConnections list (remove possible GC'd connections)
 	for (auto ConnectionIt = NetworkObjectInfo->DormantConnections.CreateIterator(); ConnectionIt; ++ConnectionIt)
 	{
-		if ((*ConnectionIt).Get() == nullptr || (*ConnectionIt).Get()->State == USOCK_Closed)
+		if ((*ConnectionIt).Get() == nullptr || (*ConnectionIt).Get()->GetConnectionState() == USOCK_Closed)
 		{
 			ConnectionIt.RemoveCurrent();
 		}
diff --git a/Engine/Source/Runtime/Engine/Private/PendingNetGame.cpp b/Engine/Source/Runtime/Engine/Private/PendingNetGame.cpp
--- a/Engine/Source/Runtime/Engine/Private/PendingNetGame.cpp
+++ b/Engine/Source/Runtime/Engine/Private/PendingNetGame.cpp
@@ -369,7 +369,7 @@ void UPendingNetGame::FinalizeEncryptedConnection(const FEncryptionKeyResponse&
 	UNetConnection* Connection = WeakConnection.Get();
 	if (Connection)
 	{
-		if (Connection->State != USOCK_Invalid && Connection->State != USOCK_Closed && Connection->Driver)
+		if (Connection->GetConnectionState() != USOCK_Invalid && Connection->GetConnectionState() != USOCK_Closed && Connection->Driver)
 		{
 			if (Response.Response == EEncryptionResponse::Success)
 			{
@@ -426,7 +426,7 @@ void UPendingNetGame::SetEncryptionKey(const FEncryptionKeyResponse& Response)
 		UNetConnection* const Connection = NetDriver->ServerConnection;
 		if (Connection)
 		{
-			if (Connection->State != USOCK_Invalid && Connection->State != USOCK_Closed && Connection->Driver)
+			if (Connection->GetConnectionState() != USOCK_Invalid && Connection->GetConnectionState() != USOCK_Closed && Connection->Driver)
 			{
 				if (Response.Response == EEncryptionResponse::Success)
 				{
@@ -474,7 +474,7 @@ void UPendingNetGame::Tick( float DeltaTime )
 	//-V:NetDriver<<:522
 
 	// Handle timed out or failed connection.
-	if (NetDriver->ServerConnection->State == USOCK_Closed && ConnectionError == TEXT(""))
+	if (NetDriver->ServerConnection->GetConnectionState() == USOCK_Closed && ConnectionError == TEXT(""))
 	{
 		ConnectionError = NSLOCTEXT("Engine", "ConnectionFailed", "Your connection to the host has been lost.").ToString();
 		return;
diff --git a/Engine/Source/Runtime/Engine/Private/PlayerController.cpp b/Engine/Source/Runtime/Engine/Private/PlayerController.cpp
--- a/Engine/Source/Runtime/Engine/Private/PlayerController.cpp
+++ b/Engine/Source/Runtime/Engine/Private/PlayerController.cpp
@@ -184,7 +184,7 @@ bool APlayerController::DestroyNetworkActorHandled()
 	UNetConnection* C = Cast<UNetConnection>(Player);
 	if (C)
 	{
-		if (C->Channels[0] && C->State != USOCK_Closed)
+		if (C->Channels[0] && C->GetConnectionState() != USOCK_Closed)
 		{
 			C->bPendingDestroy = true;
 			C->Channels[0]->Close(EChannelCloseReason::Destroyed);
diff --git a/Engine/Source/Runtime/Net/Core/Private/Net/Core/Trace/NetTraceInternal.cpp b/Engine/Source/Runtime/Net/Core/Private/Net/Core/Trace/NetTraceInternal.cpp
--- a/Engine/Source/Runtime/Net/Core/Private/Net/Core/Trace/NetTraceInternal.cpp
+++ b/Engine/Source/Runtime/Net/Core/Private/Net/Core/Trace/NetTraceInternal.cpp
@@ -74,6 +74,7 @@ void FNetTrace::SetTraceVerbosity(uint32 Verbosity)
 	if (!GetTraceVerbosity() && NewVerbosity)
 	{
 		UE::Trace::ToggleChannel(TEXT("NetChannel"), true);
+		UE::Trace::ToggleChannel(TEXT("FrameChannel"), true);
 
 		FNetTraceInternal::Reporter::ReportInitEvent(FNetTraceInternal::NetTraceVersion);
 	}
@@ -99,6 +100,14 @@ void FNetTrace::TraceEndSession(uint32 GameInstanceId)
 	}
 }
 
+void FNetTrace::TraceInstanceUpdated(uint32 GameInstanceId, bool bIsServer, const TCHAR* Name)
+{
+	if (GNetTraceRuntimeVerbosity)
+	{
+		FNetTraceInternal::Reporter::ReportInstanceUpdated(GameInstanceId, bIsServer, Name);
+	}
+}
+
 FNetTraceCollector* FNetTrace::CreateTraceCollector()
 {
 	if (!GNetTraceRuntimeVerbosity)
@@ -457,6 +466,22 @@ void FNetTrace::TraceConnectionCreated(uint32 GameInstanceId, uint32 ConnectionI
 	}
 }
 
+void FNetTrace::TraceConnectionStateUpdated(uint32 GameInstanceId, uint32 ConnectionId, uint8 ConnectionStateValue)
+{
+	if (GNetTraceRuntimeVerbosity)
+	{
+		FNetTraceInternal::Reporter::ReportConnectionStateUpdated(GameInstanceId, ConnectionId, ConnectionStateValue);
+	}
+}
+
+void FNetTrace::TraceConnectionUpdated(uint32 GameInstanceId, uint32 ConnectionId, const TCHAR* AddressString, const TCHAR* OwningActor)
+{
+	if (GNetTraceRuntimeVerbosity)
+	{
+		FNetTraceInternal::Reporter::ReportConnectionUpdated(GameInstanceId, ConnectionId, AddressString, OwningActor);
+	}
+}
+
 void FNetTrace::TraceConnectionClosed(uint32 GameInstanceId, uint32 ConnectionId)
 {
 	if (GNetTraceRuntimeVerbosity)
diff --git a/Engine/Source/Runtime/Net/Core/Private/Net/Core/Trace/Reporters/NetTraceReporter.cpp b/Engine/Source/Runtime/Net/Core/Private/Net/Core/Trace/Reporters/NetTraceReporter.cpp
--- a/Engine/Source/Runtime/Net/Core/Private/Net/Core/Trace/Reporters/NetTraceReporter.cpp
+++ b/Engine/Source/Runtime/Net/Core/Private/Net/Core/Trace/Reporters/NetTraceReporter.cpp
@@ -47,12 +47,33 @@ UE_TRACE_EVENT_BEGIN(NetTrace, ConnectionCreatedEvent)
 	UE_TRACE_EVENT_FIELD(uint8, GameInstanceId)
 UE_TRACE_EVENT_END()
 
+UE_TRACE_EVENT_BEGIN(NetTrace, ConnectionStateUpdatedEvent)
+	UE_TRACE_EVENT_FIELD(uint16, ConnectionId)
+	UE_TRACE_EVENT_FIELD(uint8, GameInstanceId)
+	UE_TRACE_EVENT_FIELD(uint8, ConnectionStateValue)
+UE_TRACE_EVENT_END()
+
+// Provides additional information about connection after it is created
+UE_TRACE_EVENT_BEGIN(NetTrace, ConnectionUpdatedEvent)
+	UE_TRACE_EVENT_FIELD(uint16, ConnectionId)
+	UE_TRACE_EVENT_FIELD(uint8, GameInstanceId)
+	UE_TRACE_EVENT_FIELD(UE::Trace::WideString, Name)
+	UE_TRACE_EVENT_FIELD(UE::Trace::WideString, Address)
+UE_TRACE_EVENT_END()
+
 // Add close reason?
 UE_TRACE_EVENT_BEGIN(NetTrace, ConnectionClosedEvent)
 	UE_TRACE_EVENT_FIELD(uint16, ConnectionId)
 	UE_TRACE_EVENT_FIELD(uint8, GameInstanceId)
 UE_TRACE_EVENT_END()
 
+// Provides additional information about game instance
+UE_TRACE_EVENT_BEGIN(NetTrace, InstanceUpdatedEvent)
+	UE_TRACE_EVENT_FIELD(uint8, GameInstanceId)
+	UE_TRACE_EVENT_FIELD(bool, bIsServer)
+	UE_TRACE_EVENT_FIELD(UE::Trace::WideString, Name)
+UE_TRACE_EVENT_END()
+
 // rename 
 UE_TRACE_EVENT_BEGIN(NetTrace, InstanceDestroyedEvent)
 	UE_TRACE_EVENT_FIELD(uint8, GameInstanceId)
@@ -92,6 +113,14 @@ void FNetTraceReporter::ReportInitEvent(uint32 NetTraceVersion)
 		<< InitEvent.NetTraceReporterVersion(NetTraceReporterVersion);
 }
 
+void FNetTraceReporter::ReportInstanceUpdated(uint32 GameInstanceId, bool bIsServer, const TCHAR* Name)
+{
+	UE_TRACE_LOG(NetTrace, InstanceUpdatedEvent, NetChannel)
+		<< InstanceUpdatedEvent.GameInstanceId(GameInstanceId)
+		<< InstanceUpdatedEvent.bIsServer(bIsServer)
+		<< InstanceUpdatedEvent.Name(Name);
+}
+
 void FNetTraceReporter::ReportInstanceDestroyed(uint32 GameInstanceId)
 {
 	UE_TRACE_LOG(NetTrace, InstanceDestroyedEvent, NetChannel)
@@ -234,6 +263,23 @@ void FNetTraceReporter::ReportConnectionCreated(uint32 GameInstanceId, uint32 Co
 		<< ConnectionCreatedEvent.GameInstanceId(GameInstanceId);
 }
 
+void FNetTraceReporter::ReportConnectionStateUpdated(uint32 GameInstanceId, uint32 ConnectionId, uint8 ConnectionStateValue)
+{
+	UE_TRACE_LOG(NetTrace, ConnectionStateUpdatedEvent, NetChannel)
+		<< ConnectionStateUpdatedEvent.ConnectionId(ConnectionId)
+		<< ConnectionStateUpdatedEvent.GameInstanceId(GameInstanceId)
+		<< ConnectionStateUpdatedEvent.ConnectionStateValue(ConnectionStateValue);
+}
+
+void FNetTraceReporter::ReportConnectionUpdated(uint32 GameInstanceId, uint32 ConnectionId, const TCHAR* AddressString, const TCHAR* OwningActor)
+{
+	UE_TRACE_LOG(NetTrace, ConnectionUpdatedEvent, NetChannel)
+		<< ConnectionUpdatedEvent.GameInstanceId(GameInstanceId)
+		<< ConnectionUpdatedEvent.ConnectionId(ConnectionId)
+		<< ConnectionUpdatedEvent.Name(OwningActor)
+		<< ConnectionUpdatedEvent.Address(AddressString);
+}
+
 void FNetTraceReporter::ReportConnectionClosed(uint32 GameInstanceId, uint32 ConnectionId)
 {
 	UE_TRACE_LOG(NetTrace, ConnectionClosedEvent, NetChannel)
diff --git a/Engine/Source/Runtime/Net/Core/Private/Net/Core/Trace/Reporters/NetTraceReporter.h b/Engine/Source/Runtime/Net/Core/Private/Net/Core/Trace/Reporters/NetTraceReporter.h
--- a/Engine/Source/Runtime/Net/Core/Private/Net/Core/Trace/Reporters/NetTraceReporter.h
+++ b/Engine/Source/Runtime/Net/Core/Private/Net/Core/Trace/Reporters/NetTraceReporter.h
@@ -18,7 +18,10 @@ struct FNetTraceReporter
 	static void ReportObjectCreated(uint32 GameInstanceId, uint32 NetObjectId, FNetDebugNameId NameId, uint64 TypeIdentifier, uint32 OwnerId);
 	static void ReportObjectDestroyed(uint32 GameInstanceId, uint32 NetObjectId);
 	static void ReportConnectionCreated(uint32 GameInstanceId, uint32 ConnectionId);
+	static void ReportConnectionStateUpdated(uint32 GameInstanceId, uint32 ConnectionId, uint8 ConnectionStateValue);
+	static void ReportConnectionUpdated(uint32 GameInstanceId, uint32 ConnectionId, const TCHAR* AddressString, const TCHAR* OwningActor);
 	static void ReportConnectionClosed(uint32 GameInstanceId, uint32 ConnectionId);
+	static void ReportInstanceUpdated(uint32 GameInstanceId, bool bIsServer, const TCHAR* Name);
 	static void ReportInstanceDestroyed(uint32 GameInstanceId);
 };
 
diff --git a/Engine/Source/Runtime/Net/Core/Public/Net/Core/Trace/NetTrace.h b/Engine/Source/Runtime/Net/Core/Public/Net/Core/Trace/NetTrace.h
--- a/Engine/Source/Runtime/Net/Core/Public/Net/Core/Trace/NetTrace.h
+++ b/Engine/Source/Runtime/Net/Core/Public/Net/Core/Trace/NetTrace.h
@@ -93,6 +93,12 @@
 /** Trace that a new connection has been created */
 #define UE_NET_TRACE_CONNECTION_CREATED(...) UE_NET_TRACE_INTERNAL_CONNECTION_CREATED(__VA_ARGS__)
 
+/** Trace that the state of this connection has been set */
+#define UE_NET_TRACE_CONNECTION_STATE_UPDATED(...) UE_NET_TRACE_INTERNAL_CONNECTION_STATE_UPDATED(__VA_ARGS__)
+
+/** Trace additional information about the given connection */
+#define UE_NET_TRACE_CONNECTION_UPDATED(...) UE_NET_TRACE_INTERNAL_CONNECTION_UPDATED(__VA_ARGS__)
+
 /** That that a connection has been closed */
 #define UE_NET_TRACE_CONNECTION_CLOSED(...) UE_NET_TRACE_INTERNAL_CONNECTION_CLOSED(__VA_ARGS__)
 
@@ -108,6 +114,9 @@
 /** Trace that the session has ended */
 #define UE_NET_TRACE_END_SESSION(GameInstanceId) UE_NET_TRACE_INTERNAL_END_SESSION(GameInstanceId)
 
+/** Trace additional information about the given game instance */
+#define UE_NET_TRACE_UPDATE_INSTANCE(...) UE_NET_TRACE_INTERNAL_UPDATE_INSTANCE(__VA_ARGS__)
+
 #else
 
 #define UE_NET_TRACE_SCOPE(...)
@@ -135,6 +144,8 @@
 #define UE_NET_TRACE_NETHANDLE_CREATED(...)
 #define UE_NET_TRACE_NETHANDLE_DESTROYED(...)
 #define UE_NET_TRACE_CONNECTION_CREATED(...)
+#define UE_NET_TRACE_CONNECTION_STATE_UPDATED(...)
+#define UE_NET_TRACE_CONNECTION_UPDATED(...)
 #define UE_NET_TRACE_CONNECTION_CLOSED(...)
 #define UE_NET_TRACE_PACKET_DROPPED(...)
 #define UE_NET_TRACE_NAMED_OBJECT_SCOPE(...)
@@ -146,4 +157,6 @@
 
 #define UE_NET_TRACE_END_SESSION(...)
 
+#define UE_NET_TRACE_UPDATE_INSTANCE(...)
+
 #endif // UE_NET_TRACE_ENABLED
diff --git a/Engine/Source/Runtime/Net/Core/Public/Net/Core/Trace/Private/NetTraceInternal.h b/Engine/Source/Runtime/Net/Core/Public/Net/Core/Trace/Private/NetTraceInternal.h
--- a/Engine/Source/Runtime/Net/Core/Public/Net/Core/Trace/Private/NetTraceInternal.h
+++ b/Engine/Source/Runtime/Net/Core/Public/Net/Core/Trace/Private/NetTraceInternal.h
@@ -148,6 +148,9 @@ struct FNetTrace
 	/** Trace end of session for a GameInstance */
 	NETCORE_API static void TraceEndSession(uint32 GameInstanceId);
 
+	/** Trace that information about the current GameInstance has been updated */
+	NETCORE_API static void TraceInstanceUpdated(uint32 GameInstanceId, bool bIsServer, const TCHAR* Name);
+
 	/** Get the current trace verbosity */
 	inline static uint32 GetTraceVerbosity() { return GNetTraceRuntimeVerbosity; }
 	
@@ -218,6 +221,12 @@ struct FNetTrace
 	/** Trace that we have added a new connection for the given GameInstanceId */
 	NETCORE_API static void TraceConnectionCreated(uint32 GameInstanceId, uint32 ConnectionId);
 
+	/** Trace that the Connection State the given connection has been set */
+	NETCORE_API static void TraceConnectionStateUpdated(uint32 GameInstanceId, uint32 ConnectionId, uint8 ConnectionStateValue);
+
+	/** Trace additional information related to the given connection */
+	NETCORE_API static void TraceConnectionUpdated(uint32 GameInstanceId, uint32 ConnectionId, const TCHAR* AddressString, const TCHAR* OwningActor);
+
 	/** Trace that we have removed a connection for the given GameInstanceId */
 	NETCORE_API static void TraceConnectionClosed(uint32 GameInstanceId, uint32 ConnectionId);
 
@@ -546,11 +555,14 @@ void FNetTraceCollector::Reset()
 #define UE_NET_TRACE_INTERNAL_NETHANDLE_CREATED(Handle, DebugName, ProtocolId, OwnerId) FNetTrace::TraceObjectCreated(Handle.GetReplicationSystemId(), Handle.GetIndex(), DebugName, ProtocolId, OwnerId)
 #define UE_NET_TRACE_INTERNAL_NETHANDLE_DESTROYED(Handle) FNetTrace::TraceObjectDestroyed(Handle.GetReplicationSystemId(), Handle.GetIndex())
 #define UE_NET_TRACE_INTERNAL_CONNECTION_CREATED(...) FNetTrace::TraceConnectionCreated(__VA_ARGS__)
+#define UE_NET_TRACE_INTERNAL_CONNECTION_STATE_UPDATED(...) FNetTrace::TraceConnectionStateUpdated(__VA_ARGS__)
+#define UE_NET_TRACE_INTERNAL_CONNECTION_UPDATED(...) FNetTrace::TraceConnectionUpdated(__VA_ARGS__)
 #define UE_NET_TRACE_INTERNAL_CONNECTION_CLOSED(...) FNetTrace::TraceConnectionClosed(__VA_ARGS__)
 #define UE_NET_TRACE_INTERNAL_PACKET_DROPPED(...) UE_NET_TRACE_DO_IF(GNetTraceRuntimeVerbosity, FNetTrace::TracePacketDropped(__VA_ARGS__))
 #define UE_NET_TRACE_INTERNAL_PACKET_SEND(...) UE_NET_TRACE_DO_IF(GNetTraceRuntimeVerbosity, FNetTrace::TracePacket(__VA_ARGS__, ENetTracePacketType::Outgoing))
 #define UE_NET_TRACE_INTERNAL_PACKET_RECV(...) UE_NET_TRACE_DO_IF(GNetTraceRuntimeVerbosity, FNetTrace::TracePacket(__VA_ARGS__, ENetTracePacketType::Incoming))
 
 #define UE_NET_TRACE_INTERNAL_END_SESSION(GameInstanceId) FNetTrace::TraceEndSession(GameInstanceId);
+#define UE_NET_TRACE_INTERNAL_UPDATE_INSTANCE(...) FNetTrace::TraceInstanceUpdated(__VA_ARGS__)
 					
 #endif // UE_NET_TRACE_ENABLED
