Additive layers are not played back when anim sequences are in the anim graph
#fix ensure that float curve data is incorporated into DataModel its hash (used for determining DDC key)

diff --git a/Engine/Source/Runtime/Engine/Private/Animation/AnimDataModel.cpp b/Engine/Source/Runtime/Engine/Private/Animation/AnimDataModel.cpp
--- a/Engine/Source/Runtime/Engine/Private/Animation/AnimDataModel.cpp
+++ b/Engine/Source/Runtime/Engine/Private/Animation/AnimDataModel.cpp
@@ -113,7 +113,7 @@ void UAnimDataModel::GetBoneTrackNames(TArray<FName>& OutNames) const
 {
 	Algo::Transform(BoneAnimationTracks, OutNames, [](FBoneAnimationTrack Track)
 	{
-		return Track.Name;
+		return Track.Name; 
 	});
 }
 
@@ -318,14 +318,36 @@ FGuid UAnimDataModel::GenerateGuid() const
 		UpdateSHAWithArray(Track.InternalTrackData.ScaleKeys);
 	}
 
+    auto UpdateWithFloatCurve = [&UpdateWithData, &UpdateSHAWithArray](const FRichCurve& Curve)
+    {
+    	UpdateWithData(Curve.DefaultValue);
+    	UpdateSHAWithArray(Curve.GetConstRefOfKeys());
+    	UpdateWithData(Curve.PreInfinityExtrap);
+    	UpdateWithData(Curve.PostInfinityExtrap);
+    };
+
 	for (const FFloatCurve& Curve : CurveData.FloatCurves)
 	{
 		UpdateWithData(Curve.Name.UID);
-		UpdateWithData(Curve.FloatCurve.DefaultValue);
-		UpdateSHAWithArray(Curve.FloatCurve.GetConstRefOfKeys());
-		UpdateWithData(Curve.FloatCurve.PreInfinityExtrap);
-		UpdateWithData(Curve.FloatCurve.PostInfinityExtrap);
+		UpdateWithFloatCurve(Curve.FloatCurve);
 	}
+
+    for (const FTransformCurve& Curve : CurveData.TransformCurves)
+    {
+    	UpdateWithData(Curve.Name.UID);
+
+    	auto UpdateWithComponent = [&Sha, &UpdateWithFloatCurve](const FVectorCurve& VectorCurve)
+    	{
+    		for (int32 ChannelIndex = 0; ChannelIndex < 3; ++ChannelIndex)
+    		{
+    			UpdateWithFloatCurve(VectorCurve.FloatCurves[ChannelIndex]);
+    		}
+    	};
+	
+    	UpdateWithComponent(Curve.TranslationCurve);
+    	UpdateWithComponent(Curve.RotationCurve);
+    	UpdateWithComponent(Curve.ScaleCurve);
+    }
 	
 	for (const FAnimatedBoneAttribute& Attribute : AnimatedBoneAttributes)
 	{
@@ -391,7 +413,7 @@ void UAnimDataModel::GenerateTransientData() const
 		RawAnimationTrackNames.Add(AnimTrack.Name);
 		RawAnimationTrackSkeletonMappings.Add(AnimTrack.BoneTreeIndex);
 	}
-
+	
 	RawCurveTracks.FloatCurves = CurveData.FloatCurves;
 #if WITH_EDITOR
 	RawCurveTracks.TransformCurves = CurveData.TransformCurves;
