Fixed issue where a missing LSA case can be detected and unnecessarily regenerated outside of editor usage on both client and server, resulting in a duplicate LSA replicated to the client and a failed ensure

#rnx

diff --git a/Engine/Source/Runtime/Engine/Private/Level.cpp b/Engine/Source/Runtime/Engine/Private/Level.cpp
--- a/Engine/Source/Runtime/Engine/Private/Level.cpp
+++ b/Engine/Source/Runtime/Engine/Private/Level.cpp
@@ -3050,10 +3050,11 @@ void ULevel::RepairLevelScript()
 
 	// Catch the edge case where we have a level blueprint but have never created the LevelScriptActor based on it.
 	//    This could happen if a new level is saved before the level blueprint is compiled.
-	if (!LevelScriptActor && LevelScriptBlueprint && !LevelScriptBlueprint->bIsRegeneratingOnLoad)
+	if (!LevelScriptActor && LevelScriptBlueprint && !LevelScriptBlueprint->bIsRegeneratingOnLoad &&
+		OwningWorld && OwningWorld->IsEditorWorld())
 	{
 		RegenerateLevelScriptActor();
-	}	
+	}
 }
 
 void ULevel::RegenerateLevelScriptActor()
