Changes to better accommodate various distributed shader build systems.

diff --git a/Engine/Plugins/XGEController/Source/Private/XGEControllerModule.cpp b/Engine/Plugins/XGEController/Source/Private/XGEControllerModule.cpp
--- a/Engine/Plugins/XGEController/Source/Private/XGEControllerModule.cpp
+++ b/Engine/Plugins/XGEController/Source/Private/XGEControllerModule.cpp
@@ -477,7 +477,17 @@ void FXGEControllerModule::WriteOutThreadProc()
 
 				Writer << Task->ID;
 				Writer << Task->CommandData.Command;
-				Writer << Task->CommandData.CommandArgs;
+
+				const FString InputFileName = FPaths::GetCleanFilename(Task->CommandData.InputFileName);
+				const FString OutputFileName = FPaths::GetCleanFilename(Task->CommandData.OutputFileName);
+				FString WorkerParameters = FString::Printf(TEXT("\"%s/\" %d 0 \"%s\" \"%s\" -xge_int %s"),
+					*Task->CommandData.WorkingDirectory,
+					Task->CommandData.DispatcherPID,
+					*InputFileName,
+					*OutputFileName,
+					*Task->CommandData.ExtraCommandArgs);
+
+				Writer << WorkerParameters;
 				*reinterpret_cast<uint32*>(WriteBuffer.GetData()) = WriteBuffer.Num() - sizeof(uint32);
 
 				// Move the tasks to the dispatched tasks map before launching it
diff --git a/Engine/Source/Developer/DistributedBuildInterface/Public/DistributedBuildControllerInterface.h b/Engine/Source/Developer/DistributedBuildInterface/Public/DistributedBuildControllerInterface.h
--- a/Engine/Source/Developer/DistributedBuildInterface/Public/DistributedBuildControllerInterface.h
+++ b/Engine/Source/Developer/DistributedBuildInterface/Public/DistributedBuildControllerInterface.h
@@ -15,9 +15,10 @@ struct FDistributedBuildTaskResult
 struct FTaskCommandData
 {	
 	FString Command;
-	FString CommandArgs;
+	FString WorkingDirectory;
 	FString InputFileName;
 	FString OutputFileName;
+	FString ExtraCommandArgs;
 	uint32 DispatcherPID = 0;
 	TArray<FString> Dependencies;
 };
diff --git a/Engine/Source/Runtime/Engine/Private/ShaderCompiler/ShaderCompiler.cpp b/Engine/Source/Runtime/Engine/Private/ShaderCompiler/ShaderCompiler.cpp
--- a/Engine/Source/Runtime/Engine/Private/ShaderCompiler/ShaderCompiler.cpp
+++ b/Engine/Source/Runtime/Engine/Private/ShaderCompiler/ShaderCompiler.cpp
@@ -1146,12 +1146,10 @@ bool FShaderCompileUtilities::DoWriteTasks(const TArray<FShaderCommonCompileJobP
 
 	// Convert all the source directory paths to absolute, since SCW might be in a different directory to the editor executable
 	TMap<FString, FString> ShaderSourceDirectoryMappings = AllShaderSourceDirectoryMappings();
-	if (!bUseRelativePaths)
+	for(TPair<FString, FString>& Pair : ShaderSourceDirectoryMappings)
 	{
-		for(TPair<FString, FString>& Pair : ShaderSourceDirectoryMappings)
-		{
-			Pair.Value = FPaths::ConvertRelativePathToFull(Pair.Value);
-		}
+		// enforce relative paths when bUseRelativePaths=true
+		Pair.Value = bUseRelativePaths ? FPaths::CreateStandardFilename(Pair.Value) : FPaths::ConvertRelativePathToFull(Pair.Value);
 	}
 	TransferFile << ShaderSourceDirectoryMappings;
 
diff --git a/Engine/Source/Runtime/Engine/Private/ShaderCompiler/ShaderCompilerDistributed.cpp b/Engine/Source/Runtime/Engine/Private/ShaderCompiler/ShaderCompilerDistributed.cpp
--- a/Engine/Source/Runtime/Engine/Private/ShaderCompiler/ShaderCompilerDistributed.cpp
+++ b/Engine/Source/Runtime/Engine/Private/ShaderCompiler/ShaderCompilerDistributed.cpp
@@ -73,17 +73,6 @@ void FShaderCompileDistributedThreadRunnable_Interface::DispatchShaderCompileJob
 	FString OutputFilePath = BaseFilePath + TEXT(".out");
 
 	const FString WorkingDirectory = FPaths::GetPath(InputFilePath);
-	const FString InputFileName = FPaths::GetCleanFilename(InputFilePath);
-	const FString OutputFileName = FPaths::GetCleanFilename(OutputFilePath);
-
-	const FString WorkerParameters = FString::Printf(TEXT("\"%s/\" %d 0 \"%s\" \"%s\" -xge_int %s%s"),
-		*WorkingDirectory,
-		Manager->ProcessId,
-		*InputFileName,
-		*OutputFileName,
-		*FCommandLine::GetSubprocessCommandline(),
-		GIsBuildMachine ? TEXT(" -buildmachine") : TEXT("")
-	);
 
 	// Serialize the jobs to the input file
 	FArchive* InputFileAr = IFileManager::Get().CreateFileWriter(*InputFilePath, FILEWRITE_EvenIfReadOnly | FILEWRITE_NoFail);
@@ -95,10 +84,11 @@ void FShaderCompileDistributedThreadRunnable_Interface::DispatchShaderCompileJob
 
 	FTaskCommandData TaskCommandData;
 	TaskCommandData.Command = Manager->ShaderCompileWorkerName;
-	TaskCommandData.CommandArgs = WorkerParameters;
+	TaskCommandData.WorkingDirectory = WorkingDirectory;
+	TaskCommandData.DispatcherPID = Manager->ProcessId;
 	TaskCommandData.InputFileName = InputFilePath;
 	TaskCommandData.OutputFileName = OutputFilePath;
-	TaskCommandData.DispatcherPID = Manager->ProcessId;
+	TaskCommandData.ExtraCommandArgs = FString::Printf(TEXT("%s%s"), *FCommandLine::GetSubprocessCommandline(), GIsBuildMachine ? TEXT(" -buildmachine") : TEXT(""));
 	TaskCommandData.Dependencies = GetDependencyFilesForJobs(JobsToSerialize);
 	
 	DispatchedTasks.Add(
