Fix in cooked build NumVertices is reset to 0 after morph target render data is initialised.

diff --git a/Engine/Source/Runtime/Engine/Classes/Animation/MorphTarget.h b/Engine/Source/Runtime/Engine/Classes/Animation/MorphTarget.h
--- a/Engine/Source/Runtime/Engine/Classes/Animation/MorphTarget.h
+++ b/Engine/Source/Runtime/Engine/Classes/Animation/MorphTarget.h
@@ -98,8 +98,11 @@ struct FMorphTargetLODModel
 			}
 		}
 
-		// Cache Vertices array size, no need to serialize CachedNumVertices value.
-		M.NumVertices = M.Vertices.Num();
+		// Cache Vertices array size, no need to serialize NumVertices value.
+		if (Ar.IsLoading())
+		{
+			M.NumVertices = M.Vertices.Num();
+		}
 
 		return Ar;
 	}
