Add texture cube array assets, support for sampling in materials

diff --git a/Engine/Config/BaseEditor.ini b/Engine/Config/BaseEditor.ini
--- a/Engine/Config/BaseEditor.ini
+++ b/Engine/Config/BaseEditor.ini
@@ -37,6 +37,7 @@ ColorFontAtlasSize=512
 +RenderableThumbnailTypes=(ClassNeedingThumbnailName="/Script/Engine.LightMapTexture2D",RendererClassName="/Script/UnrealEd.TextureThumbnailRenderer")
 +RenderableThumbnailTypes=(ClassNeedingThumbnailName="/Script/Engine.TextureCube",RendererClassName="/Script/UnrealEd.TextureCubeThumbnailRenderer")
 +RenderableThumbnailTypes=(ClassNeedingThumbnailName="/Script/Engine.Texture2DArray",RendererClassName="/Script/UnrealEd.Texture2DArrayThumbnailRenderer")
++RenderableThumbnailTypes=(ClassNeedingThumbnailName="/Script/Engine.TextureCubeArray",RendererClassName="/Script/UnrealEd.TextureCubeArrayThumbnailRenderer")
 +RenderableThumbnailTypes=(ClassNeedingThumbnailName="/Script/Engine.VolumeTexture",RendererClassName="/Script/UnrealEd.VolumeTextureThumbnailRenderer")
 +RenderableThumbnailTypes=(ClassNeedingThumbnailName="/Script/Engine.Material",RendererClassName="/Script/UnrealEd.MaterialInstanceThumbnailRenderer")
 +RenderableThumbnailTypes=(ClassNeedingThumbnailName="/Script/Engine.MaterialInterface",RendererClassName="/Script/UnrealEd.MaterialInstanceThumbnailRenderer")
diff --git a/Engine/Shaders/Private/Common.ush b/Engine/Shaders/Private/Common.ush
--- a/Engine/Shaders/Private/Common.ush
+++ b/Engine/Shaders/Private/Common.ush
@@ -408,11 +408,36 @@ float4 PseudoVolumeTexture(Texture2D Tex, SamplerState TexSampler, float3 inPos,
 }
 
 #if FEATURE_LEVEL >= FEATURE_LEVEL_SM5 // Cubemap arrays are not supported in SM4 feature level
-	MaterialFloat4 TextureCubeArraySampleLevel(TextureCubeArray Tex, SamplerState Sampler, float3 UV, float ArrayIndex, MaterialFloat Mip)
-	{
-		return Tex.SampleLevel(Sampler, float4(UV, ArrayIndex), Mip);
-	}
+MaterialFloat4 TextureCubeArraySample(TextureCubeArray Tex, SamplerState Sampler, float4 UV)
+{
+	return Tex.Sample(Sampler, UV);
+}
+
+MaterialFloat4 TextureCubeArraySampleLevel(TextureCubeArray Tex, SamplerState Sampler, float4 UV, MaterialFloat Mip)
+{
+	return Tex.SampleLevel(Sampler, UV, Mip);
+}
+
+MaterialFloat4 TextureCubeArraySampleBias(TextureCubeArray Tex, SamplerState Sampler, float4 UV, MaterialFloat MipBias)
+{
+#if USE_FORCE_TEXTURE_MIP
+	return Tex.SampleLevel(Sampler, UV, 0);
+#else
+	return Tex.SampleBias(Sampler, UV, MipBias);
 #endif
+}
+
+MaterialFloat4 TextureCubeArraySampleGrad(TextureCubeArray Tex, SamplerState Sampler, float4 UV, MaterialFloat3 DDX, MaterialFloat3 DDY)
+{
+	return Tex.SampleGrad(Sampler, UV, DDX, DDY);
+}
+
+// OVerload with explicit array parameter, used by some non-material shader code
+MaterialFloat4 TextureCubeArraySampleLevel(TextureCubeArray Tex, SamplerState Sampler, float3 UV, float ArrayIndex, MaterialFloat Mip)
+{
+	return TextureCubeArraySampleLevel(Tex, Sampler, float4(UV, ArrayIndex), Mip);
+}
+#endif // FEATURE_LEVEL >= FEATURE_LEVEL_SM5
 
 // TANGENTTOWORLD0 is the first row of the tangent to world matrix, w might be needed for padding and is not used yet.
 // TANGENTTOWORLD2 is the last row of the tangent to world matrix, determinant of tangent basis in w
diff --git a/Engine/Source/Developer/AssetTools/Private/AssetTools.cpp b/Engine/Source/Developer/AssetTools/Private/AssetTools.cpp
--- a/Engine/Source/Developer/AssetTools/Private/AssetTools.cpp
+++ b/Engine/Source/Developer/AssetTools/Private/AssetTools.cpp
@@ -104,6 +104,7 @@
 #include "AssetTypeActions/AssetTypeActions_Texture2D.h"
 #include "AssetTypeActions/AssetTypeActions_Texture2DArray.h"
 #include "AssetTypeActions/AssetTypeActions_TextureCube.h"
+#include "AssetTypeActions/AssetTypeActions_TextureCubeArray.h"
 #include "AssetTypeActions/AssetTypeActions_VolumeTexture.h"
 #include "AssetTypeActions/AssetTypeActions_TextureRenderTarget2DArray.h"
 #include "AssetTypeActions/AssetTypeActions_TextureRenderTargetCube.h"
@@ -298,6 +299,7 @@ UAssetToolsImpl::UAssetToolsImpl(const FObjectInitializer& ObjectInitializer)
 	RegisterAssetTypeActions(MakeShareable(new FAssetTypeActions_Texture2D));
 	RegisterAssetTypeActions(MakeShareable(new FAssetTypeActions_TextureCube));
 	RegisterAssetTypeActions(MakeShareable(new FAssetTypeActions_Texture2DArray));
+	RegisterAssetTypeActions(MakeShareable(new FAssetTypeActions_TextureCubeArray));
 	RegisterAssetTypeActions(MakeShareable(new FAssetTypeActions_VolumeTexture));
 	RegisterAssetTypeActions(MakeShareable(new FAssetTypeActions_TextureRenderTarget));
 	RegisterAssetTypeActions(MakeShareable(new FAssetTypeActions_TextureRenderTarget2D));
diff --git a/Engine/Source/Developer/AssetTools/Private/AssetTypeActions/AssetTypeActions_TextureCube.cpp b/Engine/Source/Developer/AssetTools/Private/AssetTypeActions/AssetTypeActions_TextureCube.cpp
new file mode 100644
--- /dev/null
+++ b/Engine/Source/Developer/AssetTools/Private/AssetTypeActions/AssetTypeActions_TextureCube.cpp
@@ -0,0 +1,58 @@
+// Copyright Epic Games, Inc. All Rights Reserved.
+
+#include "AssetTypeActions/AssetTypeActions_TextureCube.h"
+#include "ToolMenus.h"
+#include "Misc/PackageName.h"
+#include "EditorStyleSet.h"
+#include "Factories/SlateBrushAssetFactory.h"
+#include "Slate/SlateBrushAsset.h"
+#include "Factories/TextureCubeArrayFactory.h"
+#include "Engine/TextureCubeArray.h"
+#include "AssetTools.h"
+#include "IContentBrowserSingleton.h"
+#include "ContentBrowserModule.h"
+
+#define LOCTEXT_NAMESPACE "AssetTypeActions"
+
+void FAssetTypeActions_TextureCube::GetActions(const TArray<UObject*>& InObjects, FToolMenuSection& Section)
+{
+	FAssetTypeActions_Texture::GetActions(InObjects, Section);
+
+	auto Textures = GetTypedWeakObjectPtrs<UTextureCube>(InObjects);
+
+	if (Textures.Num() > 0)
+	{
+		Section.AddMenuEntry(
+			"Texture_TextureCuberray",
+			LOCTEXT("Texture_TextureCuberray", "Create Texture Cube Array"),
+			LOCTEXT("Texture_CreateTexture2DArrayTooltip", "Creates a new texture cube array."),
+			FSlateIcon(FEditorStyle::GetStyleSetName(), "ClassIcon.TextureCube"),
+			FUIAction(FExecuteAction::CreateSP(this, &FAssetTypeActions_TextureCube::ExecuteCreateTextureArray, Textures), FCanExecuteAction())
+		);
+	}
+}
+
+void FAssetTypeActions_TextureCube::ExecuteCreateTextureArray(TArray<TWeakObjectPtr<UTextureCube>> Objects)
+{
+	const FString DefaultSuffix = TEXT("_Array");
+	FString Name;
+	FString PackagePath;
+	CreateUniqueAssetName(Objects[0].Get()->GetOutermost()->GetName(), DefaultSuffix, PackagePath, Name);
+
+	// Create the factory used to generate the asset.
+	UTextureCubeArrayFactory* Factory = NewObject<UTextureCubeArrayFactory>();
+	Factory->InitialTextures.Empty();
+
+	// Give the selected textures to the factory.
+	for (int32 TextureIndex = 0; TextureIndex < Objects.Num(); ++TextureIndex)
+	{
+		Factory->InitialTextures.Add(Objects[TextureIndex].Get());
+	}
+
+	{
+		FContentBrowserModule& ContentBrowserModule = FModuleManager::LoadModuleChecked<FContentBrowserModule>("ContentBrowser");
+		ContentBrowserModule.Get().CreateNewAsset(Name, FPackageName::GetLongPackagePath(PackagePath), UTextureCubeArray::StaticClass(), Factory);
+	}
+}
+
+#undef LOCTEXT_NAMESPACE
diff --git a/Engine/Source/Developer/AssetTools/Private/AssetTypeActions/AssetTypeActions_TextureCube.h b/Engine/Source/Developer/AssetTools/Private/AssetTypeActions/AssetTypeActions_TextureCube.h
--- a/Engine/Source/Developer/AssetTools/Private/AssetTypeActions/AssetTypeActions_TextureCube.h
+++ b/Engine/Source/Developer/AssetTools/Private/AssetTypeActions/AssetTypeActions_TextureCube.h
@@ -14,4 +14,8 @@ public:
 	virtual FColor GetTypeColor() const override { return FColor(192,64,128); }
 	virtual UClass* GetSupportedClass() const override { return UTextureCube::StaticClass(); }
 	virtual bool CanFilter() override { return true; }
+	virtual void GetActions(const TArray<UObject*>& InObjects, struct FToolMenuSection& Section) override;
+
+private:
+	void ExecuteCreateTextureArray(TArray<TWeakObjectPtr<UTextureCube>> Objects);
 };
diff --git a/Engine/Source/Developer/AssetTools/Private/AssetTypeActions/AssetTypeActions_TextureCubeArray.cpp b/Engine/Source/Developer/AssetTools/Private/AssetTypeActions/AssetTypeActions_TextureCubeArray.cpp
new file mode 100644
--- /dev/null
+++ b/Engine/Source/Developer/AssetTools/Private/AssetTypeActions/AssetTypeActions_TextureCubeArray.cpp
@@ -0,0 +1,3 @@
+// Copyright Epic Games, Inc. All Rights Reserved.
+
+#include "AssetTypeActions_TextureCubeArray.h"
diff --git a/Engine/Source/Developer/AssetTools/Private/AssetTypeActions/AssetTypeActions_TextureCubeArray.h b/Engine/Source/Developer/AssetTools/Private/AssetTypeActions/AssetTypeActions_TextureCubeArray.h
new file mode 100644
--- /dev/null
+++ b/Engine/Source/Developer/AssetTools/Private/AssetTypeActions/AssetTypeActions_TextureCubeArray.h
@@ -0,0 +1,17 @@
+// Copyright Epic Games, Inc. All Rights Reserved.
+
+#pragma once
+
+#include "CoreMinimal.h"
+#include "AssetTypeActions/AssetTypeActions_Texture.h"
+#include "Engine/TextureCubeArray.h"
+
+class FAssetTypeActions_TextureCubeArray : public FAssetTypeActions_Texture
+{
+public:
+	// IAssetTypeActions Implementation
+	virtual FText GetName() const override { return NSLOCTEXT("AssetTypeActions", "AssetTypeActions_TextureCubeArray", "Texture Cube Array"); }
+	virtual FColor GetTypeColor() const override { return FColor(192, 192, 128); }
+	virtual UClass* GetSupportedClass() const override { return UTextureCubeArray::StaticClass(); }
+	virtual bool CanFilter() override { return true; }
+};
\ No newline at end of file
diff --git a/Engine/Source/Developer/TextureCompressor/Private/TextureCompressorModule.cpp b/Engine/Source/Developer/TextureCompressor/Private/TextureCompressorModule.cpp
--- a/Engine/Source/Developer/TextureCompressor/Private/TextureCompressorModule.cpp
+++ b/Engine/Source/Developer/TextureCompressor/Private/TextureCompressorModule.cpp
@@ -2223,7 +2223,7 @@ public:
 		TArray<FCompressedImage2D>& OutTextureMips,
 		uint32& OutNumMipsInTail,
 		uint32& OutExtData
-		)
+	)
 	{
 		const ITextureFormat* TextureFormat = nullptr;
 
@@ -2247,15 +2247,15 @@ public:
 		// we can't use the Ex version here because it needs an FImage, which needs BuildTextureMips to be called
 		const FTextureFormatCompressorCaps CompressorCaps = TextureFormat->GetFormatCapabilities();
 
-		if(!BuildTextureMips(SourceMips, BuildSettings, CompressorCaps, IntermediateMipChain))
+		if (!BuildTextureMips(SourceMips, BuildSettings, CompressorCaps, IntermediateMipChain))
 		{
 			return false;
 		}
 
 		// apply roughness adjustment depending on normal map variation
-		if(AssociatedNormalSourceMips.Num())
+		if (AssociatedNormalSourceMips.Num())
 		{
-//			check AssociatedNormalSourceMips.Format; 
+			//			check AssociatedNormalSourceMips.Format; 
 
 			TArray<FImage> IntermediateAssociatedNormalSourceMipChain;
 
@@ -2268,12 +2268,12 @@ public:
 			// important to make accurate computation with normal length
 			DefaultSettings.bRenormalizeTopMip = true;
 
-			if(!BuildTextureMips(AssociatedNormalSourceMips, DefaultSettings, CompressorCaps, IntermediateAssociatedNormalSourceMipChain))
+			if (!BuildTextureMips(AssociatedNormalSourceMips, DefaultSettings, CompressorCaps, IntermediateAssociatedNormalSourceMipChain))
 			{
 				UE_LOG(LogTextureCompressor, Warning, TEXT("Failed to generate texture mips for composite texture"));
 			}
 
-			if(!ApplyCompositeTexture(IntermediateMipChain, IntermediateAssociatedNormalSourceMipChain, BuildSettings.CompositeTextureMode, BuildSettings.CompositePower))
+			if (!ApplyCompositeTexture(IntermediateMipChain, IntermediateAssociatedNormalSourceMipChain, BuildSettings.CompositeTextureMode, BuildSettings.CompositePower))
 			{
 				UE_LOG(LogTextureCompressor, Warning, TEXT("Failed to apply composite texture"));
 			}
@@ -2284,7 +2284,21 @@ public:
 		BuildSettings.TopMipSize.X = IntermediateMipChain[0].SizeX;
 		BuildSettings.TopMipSize.Y = IntermediateMipChain[0].SizeY;
 		BuildSettings.VolumeSizeZ = BuildSettings.bVolume ? IntermediateMipChain[0].NumSlices : 1;
-		BuildSettings.ArraySlices = BuildSettings.bTextureArray ? IntermediateMipChain[0].NumSlices : 1;
+		if (BuildSettings.bTextureArray)
+		{
+			if (BuildSettings.bCubemap)
+			{
+				BuildSettings.ArraySlices = IntermediateMipChain[0].NumSlices / 6;
+			}
+			else
+			{
+				BuildSettings.ArraySlices = IntermediateMipChain[0].NumSlices;
+			}
+		}
+		else
+		{
+			BuildSettings.ArraySlices = 1;
+		}
 		
 		return CompressMipChain(TextureFormat, IntermediateMipChain, BuildSettings, OutTextureMips, OutNumMipsInTail, OutExtData);
 	}
@@ -2332,11 +2346,19 @@ private:
 		check(InSourceMips[0].SizeX > 0 && InSourceMips[0].SizeY > 0 && InSourceMips[0].NumSlices > 0);
 
 		// Identify long-lat cubemaps.
-		bool bLongLatCubemap = BuildSettings.bCubemap && InSourceMips[0].NumSlices == 1;
-
-		if (BuildSettings.bCubemap && InSourceMips[0].NumSlices != 6 && !bLongLatCubemap)
+		const bool bLongLatCubemap = BuildSettings.bCubemap && !BuildSettings.bTextureArray && InSourceMips[0].NumSlices == 1;
+		if (BuildSettings.bCubemap && !bLongLatCubemap)
 		{
-			return false;
+			if (BuildSettings.bTextureArray && (InSourceMips[0].NumSlices % 6) != 0)
+			{
+				// Cube array must have multiiple of 6 slices
+				return false;
+			}
+			if (!BuildSettings.bTextureArray && InSourceMips[0].NumSlices != 6)
+			{
+				// Non-array cube must have exactly 6 slices
+				return false;
+			}
 		}
 
 		// Determine the maximum possible mip counts for source and dest.
diff --git a/Engine/Source/Editor/MaterialEditor/Private/MaterialEditor.cpp b/Engine/Source/Editor/MaterialEditor/Private/MaterialEditor.cpp
--- a/Engine/Source/Editor/MaterialEditor/Private/MaterialEditor.cpp
+++ b/Engine/Source/Editor/MaterialEditor/Private/MaterialEditor.cpp
@@ -37,6 +37,7 @@
 #include "StaticParameterSet.h"
 #include "Engine/TextureCube.h"
 #include "Engine/Texture2DArray.h"
+#include "Engine/TextureCubeArray.h"
 #include "Dialogs/Dialogs.h"
 #include "UnrealEdGlobals.h"
 #include "Editor.h"
@@ -76,6 +77,7 @@
 #include "Materials/MaterialExpressionTextureSampleParameter2D.h"
 #include "Materials/MaterialExpressionTextureSampleParameterCube.h"
 #include "Materials/MaterialExpressionTextureSampleParameter2DArray.h"
+#include "Materials/MaterialExpressionTextureSampleParameterCubeArray.h"
 #include "Materials/MaterialExpressionTextureSampleParameterSubUV.h"
 #include "Materials/MaterialExpressionTransformPosition.h"
 #include "Materials/MaterialExpressionVectorParameter.h"
@@ -3586,6 +3588,10 @@ void FMaterialEditor::OnConvertObjects()
 				{
 					ClassToCreate = UMaterialExpressionTextureSampleParameter2DArray::StaticClass();
 				}
+				else if (TextureSampleExpression && TextureSampleExpression->Texture && TextureSampleExpression->Texture->IsA(UTextureCubeArray::StaticClass()))
+				{
+					ClassToCreate = UMaterialExpressionTextureSampleParameterCubeArray::StaticClass();
+				}
 				else if (TextureObjectExpression)
 				{
 					ClassToCreate = UMaterialExpressionTextureObjectParameter::StaticClass();
diff --git a/Engine/Source/Editor/TextureEditor/Private/Models/TextureEditorViewportClient.cpp b/Engine/Source/Editor/TextureEditor/Private/Models/TextureEditorViewportClient.cpp
--- a/Engine/Source/Editor/TextureEditor/Private/Models/TextureEditorViewportClient.cpp
+++ b/Engine/Source/Editor/TextureEditor/Private/Models/TextureEditorViewportClient.cpp
@@ -8,6 +8,7 @@
 #include "ThumbnailRendering/ThumbnailManager.h"
 #include "Engine/TextureCube.h"
 #include "Engine/Texture2DArray.h"
+#include "Engine/TextureCubeArray.h"
 #include "Engine/VolumeTexture.h"
 #include "Engine/TextureRenderTarget2D.h"
 #include "Engine/TextureRenderTarget2DArray.h"
@@ -82,6 +83,7 @@ void FTextureEditorViewportClient::Draw(FViewport* Viewport, FCanvas* Canvas)
 	UTexture2D* Texture2D = Cast<UTexture2D>(Texture);
 	UTextureCube* TextureCube = Cast<UTextureCube>(Texture);
 	UTexture2DArray* Texture2DArray = Cast<UTexture2DArray>(Texture);
+	UTextureCubeArray* TextureCubeArray = Cast<UTextureCubeArray>(Texture);
 	UVolumeTexture* VolumeTexture = Cast<UVolumeTexture>(Texture);
 	UTextureRenderTarget2D* TextureRT2D = Cast<UTextureRenderTarget2D>(Texture);
 	UTextureRenderTarget2DArray* TextureRT2DArray = Cast<UTextureRenderTarget2DArray>(Texture);
@@ -106,7 +108,7 @@ void FTextureEditorViewportClient::Draw(FViewport* Viewport, FCanvas* Canvas)
 
 	if (GMaxRHIFeatureLevel >= ERHIFeatureLevel::SM5)
 	{
-		if (TextureCube || RTTextureCube)
+		if (TextureCube || TextureCubeArray || RTTextureCube)
 		{
 			BatchedElementParameters = new FMipLevelBatchedElementParameters(MipLevel, false);
 		}
diff --git a/Engine/Source/Editor/TextureEditor/Private/TextureEditorToolkit.cpp b/Engine/Source/Editor/TextureEditor/Private/TextureEditorToolkit.cpp
--- a/Engine/Source/Editor/TextureEditor/Private/TextureEditorToolkit.cpp
+++ b/Engine/Source/Editor/TextureEditor/Private/TextureEditorToolkit.cpp
@@ -21,6 +21,7 @@
 #include "Engine/Texture2DDynamic.h"
 #include "Engine/TextureCube.h"
 #include "Engine/Texture2DArray.h"
+#include "Engine/TextureCubeArray.h"
 #include "Engine/VolumeTexture.h"
 #include "Engine/TextureRenderTarget.h"
 #include "Engine/TextureRenderTarget2D.h"
@@ -470,13 +471,14 @@ void FTextureEditorToolkit::PopulateQuickInfo( )
 	UTexture2D* Texture2D = Cast<UTexture2D>(Texture);
 
 	const bool bIsVolume = IsVolumeTexture();
-	const bool bIsArray2D = Is2DArrayTexture();
+	const bool bIsArray = IsArrayTexture();
 	const bool bIsCube = IsCubeTexture();
 
 	const uint32 SurfaceWidth = (uint32)Texture->GetSurfaceWidth();
 	const uint32 SurfaceHeight = (uint32)Texture->GetSurfaceHeight();
 	const uint32 SurfaceDepth = (uint32)Texture->GetSurfaceDepth();
-	const uint32 ArraySize = (uint32)Texture->GetSurfaceArraySize();
+	const uint32 NumSurfaces = (uint32)Texture->GetSurfaceArraySize();
+	const uint32 ArraySize = bIsCube ? (NumSurfaces / 6u) : NumSurfaces;
 
 	const uint32 ImportedWidth = FMath::Max<uint32>(SurfaceWidth, Texture->Source.GetSizeX());
 	const uint32 ImportedHeight =  FMath::Max<uint32>(SurfaceHeight, Texture->Source.GetSizeY());
@@ -514,6 +516,11 @@ void FTextureEditorToolkit::PopulateQuickInfo( )
 	FNumberFormattingOptions Options;
 	Options.UseGrouping = false;
 
+	FText CubemapAdd;
+	if (bIsCube)
+	{
+		CubemapAdd = NSLOCTEXT("TextureEditor", "QuickInfo_PerCubeSide", "*6 (CubeMap)");
+	}
 
 	if (bIsVolume)
 	{
@@ -535,20 +542,14 @@ void FTextureEditorToolkit::PopulateQuickInfo( )
 			PreviewEffectiveTextureHeight *= (uint32)NumTilesY;
 		}
 	}
-	else if (bIsArray2D)
+	else if (bIsArray)
 	{
 		ImportedText->SetText(FText::Format(NSLOCTEXT("TextureEditor", "QuickInfo_Imported_3x", "Imported: {0}x{1}*{2}"), FText::AsNumber(ImportedWidth, &Options), FText::AsNumber(ImportedHeight, &Options), FText::AsNumber(ArraySize, &Options)));
-		CurrentText->SetText(FText::Format(NSLOCTEXT("TextureEditor", "QuickInfo_Displayed_3x", "Displayed: {0}x{1}*{2}"), FText::AsNumber(PreviewEffectiveTextureWidth, &Options), FText::AsNumber(PreviewEffectiveTextureHeight, &Options), FText::AsNumber(ArraySize, &Options)));
-		MaxInGameText->SetText(FText::Format(NSLOCTEXT("TextureEditor", "QuickInfo_MaxInGame_3x", "Max In-Game: {0}x{1}*{2}"), FText::AsNumber(MaxInGameWidth, &Options), FText::AsNumber(MaxInGameHeight, &Options), FText::AsNumber(ArraySize, &Options)));
+		CurrentText->SetText(FText::Format(NSLOCTEXT("TextureEditor", "QuickInfo_Displayed_3x", "Displayed: {0}x{1}{2}*{3}"), FText::AsNumber(PreviewEffectiveTextureWidth, &Options), FText::AsNumber(PreviewEffectiveTextureHeight, &Options), CubemapAdd, FText::AsNumber(ArraySize, &Options)));
+		MaxInGameText->SetText(FText::Format(NSLOCTEXT("TextureEditor", "QuickInfo_MaxInGame_3x", "Max In-Game: {0}x{1}{2}*{3}"), FText::AsNumber(MaxInGameWidth, &Options), FText::AsNumber(MaxInGameHeight, &Options), CubemapAdd, FText::AsNumber(ArraySize, &Options)));
 	}
 	else
 	{
-	    FText CubemapAdd;
-	    if(bIsCube)
-	    {
-		    CubemapAdd = NSLOCTEXT("TextureEditor", "QuickInfo_PerCubeSide", "*6 (CubeMap)");
-	    }
-    
 	    ImportedText->SetText(FText::Format( NSLOCTEXT("TextureEditor", "QuickInfo_Imported_2x", "Imported: {0}x{1}"), FText::AsNumber(ImportedWidth, &Options), FText::AsNumber(ImportedHeight, &Options)));
 		CurrentText->SetText(FText::Format( NSLOCTEXT("TextureEditor", "QuickInfo_Displayed_2x", "Displayed: {0}x{1}{2}"), FText::AsNumber(PreviewEffectiveTextureWidth, &Options ), FText::AsNumber(PreviewEffectiveTextureHeight, &Options), CubemapAdd));
 		MaxInGameText->SetText(FText::Format( NSLOCTEXT("TextureEditor", "QuickInfo_MaxInGame_2x", "Max In-Game: {0}x{1}{2}"), FText::AsNumber(MaxInGameWidth, &Options), FText::AsNumber(MaxInGameHeight, &Options), CubemapAdd));
@@ -1057,6 +1058,10 @@ int32 FTextureEditorToolkit::GetNumMips( ) const
 	{
 		return TextureCube->GetNumMips();
 	}
+	else if (const UTextureCubeArray* TextureCubeArray = Cast<UTextureCubeArray>(Texture))
+	{
+		return TextureCubeArray->GetNumMips();
+	}
 	else if (const UTexture2DArray* Texture2DArray = Cast<UTexture2DArray>(Texture))
 	{
 		return Texture2DArray->GetNumMips();
@@ -1107,6 +1112,10 @@ EPixelFormat FTextureEditorToolkit::GetPixelFormat() const
 	{
 		return Texture2DArray->GetPixelFormat();
 	}
+	else if (const UTextureCubeArray* TextureCubeArray = Cast<UTextureCubeArray>(Texture))
+	{
+		return TextureCubeArray->GetPixelFormat();
+	}
 	else if (const UVolumeTexture* VolumeTexture = Cast<UVolumeTexture>(Texture))
 	{
 		return VolumeTexture->GetPixelFormat();
@@ -1146,7 +1155,7 @@ TOptional<int32> FTextureEditorToolkit::GetMaxLayer() const
 
 bool FTextureEditorToolkit::IsCubeTexture( ) const
 {
-	return (Texture->IsA(UTextureCube::StaticClass()) || Texture->IsA(UTextureRenderTargetCube::StaticClass()));
+	return (Texture->IsA(UTextureCube::StaticClass()) || Texture->IsA(UTextureCubeArray::StaticClass()) || Texture->IsA(UTextureRenderTargetCube::StaticClass()));
 }
 
 bool FTextureEditorToolkit::IsVolumeTexture() const
@@ -1159,6 +1168,11 @@ bool FTextureEditorToolkit::Is2DArrayTexture() const
 	return (Texture->IsA(UTexture2DArray::StaticClass()) || Texture->IsA(UTextureRenderTarget2DArray::StaticClass()));
 }
 
+bool FTextureEditorToolkit::IsArrayTexture() const
+{
+	return Is2DArrayTexture() || Texture->IsA(UTextureCubeArray::StaticClass());
+}
+
 TSharedRef<SWidget> FTextureEditorToolkit::OnGenerateMipMapLevelMenu()
 {
 	FMenuBuilder MenuBuilder(true, nullptr);
diff --git a/Engine/Source/Editor/TextureEditor/Private/TextureEditorToolkit.h b/Engine/Source/Editor/TextureEditor/Private/TextureEditorToolkit.h
--- a/Engine/Source/Editor/TextureEditor/Private/TextureEditorToolkit.h
+++ b/Engine/Source/Editor/TextureEditor/Private/TextureEditorToolkit.h
@@ -160,6 +160,11 @@ protected:
 	 */
 	bool Is2DArrayTexture() const;
 
+	/**
+	 * Checks whether the texture being edited is a texture array (Cube or 2D).
+	 */
+	bool IsArrayTexture() const;
+
 	TSharedRef<SWidget> OnGenerateMipMapLevelMenu();
 	TSharedRef<SWidget> OnGenerateSettingsMenu();
 private:
diff --git a/Engine/Source/Editor/UnrealEd/Classes/Factories/TextureCubeArrayFactory.h b/Engine/Source/Editor/UnrealEd/Classes/Factories/TextureCubeArrayFactory.h
new file mode 100644
--- /dev/null
+++ b/Engine/Source/Editor/UnrealEd/Classes/Factories/TextureCubeArrayFactory.h
@@ -0,0 +1,32 @@
+// Copyright Epic Games, Inc. All Rights Reserved.
+
+#pragma once
+
+#include "CoreMinimal.h"
+#include "UObject/ObjectMacros.h"
+#include "Factories/Factory.h"
+#include "TextureCubeArrayFactory.generated.h"
+
+class UTextureCube;
+
+/** Factory for creating volume texture */
+UCLASS(hidecategories = Object, MinimalAPI)
+class UTextureCubeArrayFactory : public UFactory
+{
+	GENERATED_UCLASS_BODY()
+
+	UPROPERTY()
+	TArray<TObjectPtr<UTextureCube>> InitialTextures;
+
+	//~ Begin UFactory Interface
+	virtual bool CanCreateNew() const override;
+	virtual FText GetDisplayName() const override;
+	virtual bool ConfigureProperties() override;
+	virtual UObject* FactoryCreateNew(UClass* Class, UObject* InParent, FName Name, EObjectFlags Flags, UObject* Context, FFeedbackContext* Warn) override;
+	//~ Begin UFactory Interface
+
+	bool CheckArrayTexturesCompatibility();
+};
+
+
+
diff --git a/Engine/Source/Editor/UnrealEd/Classes/Factories/TextureCubeArrayThumbnailRenderer.h b/Engine/Source/Editor/UnrealEd/Classes/Factories/TextureCubeArrayThumbnailRenderer.h
new file mode 100644
--- /dev/null
+++ b/Engine/Source/Editor/UnrealEd/Classes/Factories/TextureCubeArrayThumbnailRenderer.h
@@ -0,0 +1,24 @@
+// Copyright Epic Games, Inc. All Rights Reserved.
+
+/**
+ *
+ * This thumbnail renderer displays the texture cube array
+ */
+
+#pragma once
+
+#include "CoreMinimal.h"
+#include "UObject/ObjectMacros.h"
+#include "ThumbnailRendering/TextureThumbnailRenderer.h"
+#include "TextureCubeArrayThumbnailRenderer.generated.h"
+
+UCLASS()
+class UTextureCubeArrayThumbnailRenderer : public UTextureThumbnailRenderer
+{
+	GENERATED_UCLASS_BODY()
+
+	// Begin UThumbnailRenderer Interface
+	virtual void GetThumbnailSize(UObject* Object, float Zoom, uint32& OutWidth, uint32& OutHeight) const override;
+	virtual void Draw(UObject* Object, int32 X, int32 Y, uint32 Width, uint32 Height, FRenderTarget*, FCanvas* Canvas, bool bAdditionalViewFamily) override;
+	// End UThumbnailRenderer Interface
+};
diff --git a/Engine/Source/Editor/UnrealEd/Private/Factories/TextureCubeArrayFactory.cpp b/Engine/Source/Editor/UnrealEd/Private/Factories/TextureCubeArrayFactory.cpp
new file mode 100644
--- /dev/null
+++ b/Engine/Source/Editor/UnrealEd/Private/Factories/TextureCubeArrayFactory.cpp
@@ -0,0 +1,93 @@
+// Copyright Epic Games, Inc. All Rights Reserved.
+
+#include "Factories/TextureCubeArrayFactory.h"
+#include "Styling/SlateBrush.h"
+#include "Brushes/SlateDynamicImageBrush.h"
+#include "Engine/TextureCubeArray.h"
+
+
+#define LOCTEXT_NAMESPACE "TextureCubeArrayFactory"
+
+UTextureCubeArrayFactory::UTextureCubeArrayFactory(const FObjectInitializer& ObjectInitializer)
+	: Super(ObjectInitializer)
+{
+	bCreateNew = true;
+	bEditAfterNew = true;
+	SupportedClass = UTextureCubeArray::StaticClass();
+}
+
+FText UTextureCubeArrayFactory::GetDisplayName() const
+{
+	return LOCTEXT("TextureCubeArrayFactoryDescription", "Texture Cube Array");
+}
+
+bool UTextureCubeArrayFactory::ConfigureProperties()
+{
+	return true;
+}
+
+bool UTextureCubeArrayFactory::CanCreateNew() const
+{
+	return true;
+}
+
+UObject* UTextureCubeArrayFactory::FactoryCreateNew(UClass* Class, UObject* InParent, FName Name, EObjectFlags Flags, UObject* Context, FFeedbackContext* Warn)
+{
+	// Make sure the size of all textures is the same.
+	if (!CheckArrayTexturesCompatibility())
+	{
+		return nullptr;
+	}
+
+	UTextureCubeArray* NewTextureCubeArray = CastChecked<UTextureCubeArray>(CreateOrOverwriteAsset(UTextureCubeArray::StaticClass(), InParent, Name, Flags));
+
+	if (NewTextureCubeArray == 0)
+	{
+		UE_LOG(LogTexture, Warning, TEXT("TextureCubeArray creation failed.\n"));
+	}
+	else if (InitialTextures.Num() > 0)
+	{
+		// Upload the compatible textures to the texture array.
+		for (int32 TextureIndex = 0; TextureIndex < InitialTextures.Num(); ++TextureIndex)
+		{
+			NewTextureCubeArray->SourceTextures.Add(InitialTextures[TextureIndex]);
+		}
+
+		// Create the texture array resource and corresponding rhi resource.
+		NewTextureCubeArray->UpdateSourceFromSourceTextures();
+	}
+
+	return NewTextureCubeArray;
+}
+
+bool UTextureCubeArrayFactory::CheckArrayTexturesCompatibility()
+{
+	bool bError = false;
+	for (int32 TextureIndex = 0; TextureIndex < InitialTextures.Num(); ++TextureIndex)
+	{
+		FTextureSource& TextureSource = InitialTextures[TextureIndex]->Source;
+		const int32 FormatDataSize = TextureSource.GetBytesPerPixel();
+		const int32 SizeX = TextureSource.GetSizeX();
+		const int32 SizeY = TextureSource.GetSizeY();
+
+		for (int32 TextureCmpIndex = TextureIndex + 1; TextureCmpIndex < InitialTextures.Num(); ++TextureCmpIndex)
+		{
+			FTextureSource& TextureSourceCmp = InitialTextures[TextureCmpIndex]->Source;
+			if (TextureSourceCmp.GetSizeX() != SizeX || TextureSourceCmp.GetSizeY() != SizeY)
+			{
+				UE_LOG(LogTexture, Warning, TEXT("TextureCubeArray creation failed. Textures have different sizes."));
+				bError = true;
+			}
+
+			if (TextureSourceCmp.GetBytesPerPixel() != FormatDataSize)
+			{
+				UE_LOG(LogTexture, Warning, TEXT("TextureCubeArray creation failed. Textures have different pixel formats."));
+				bError = true;
+			}
+		}
+	}
+
+	return (!bError);
+}
+
+#undef LOCTEXT_NAMESPACE
diff --git a/Engine/Source/Editor/UnrealEd/Private/ThumbnailRendering/TextureCubeArrayThumbnailRenderer.cpp b/Engine/Source/Editor/UnrealEd/Private/ThumbnailRendering/TextureCubeArrayThumbnailRenderer.cpp
new file mode 100644
--- /dev/null
+++ b/Engine/Source/Editor/UnrealEd/Private/ThumbnailRendering/TextureCubeArrayThumbnailRenderer.cpp
@@ -0,0 +1,31 @@
+// Copyright Epic Games, Inc. All Rights Reserved.
+
+#include "Factories/TextureCubeArrayThumbnailRenderer.h"
+#include "Engine/TextureCubeArray.h"
+
+UTextureCubeArrayThumbnailRenderer::UTextureCubeArrayThumbnailRenderer(const FObjectInitializer& ObjectInitializer)
+	: Super(ObjectInitializer)
+{
+}
+
+void UTextureCubeArrayThumbnailRenderer::GetThumbnailSize(UObject* Object, float Zoom, uint32& OutWidth, uint32& OutHeight) const
+{
+	OutWidth = 0;
+	OutHeight = 0;
+
+	UTextureCubeArray* CubeMap = Cast<UTextureCubeArray>(Object);
+	if (CubeMap != nullptr)
+	{
+		// Let the base class get the size of a face
+		Super::GetThumbnailSize(CubeMap, Zoom, OutWidth, OutHeight);
+	}
+}
+
+void UTextureCubeArrayThumbnailRenderer::Draw(UObject* Object, int32 X, int32 Y, uint32 Width, uint32 Height, FRenderTarget*, FCanvas* Canvas, bool bAdditionalViewFamily)
+{
+	UTextureCubeArray* CubeMap = Cast<UTextureCubeArray>(Object);
+	if (CubeMap != nullptr)
+	{
+		Super::Draw(CubeMap, X, Y, Width, Height, nullptr, Canvas, bAdditionalViewFamily);
+	}
+}
diff --git a/Engine/Source/Editor/UnrealEd/Private/ThumbnailRendering/TextureThumbnailRenderer.cpp b/Engine/Source/Editor/UnrealEd/Private/ThumbnailRendering/TextureThumbnailRenderer.cpp
--- a/Engine/Source/Editor/UnrealEd/Private/ThumbnailRendering/TextureThumbnailRenderer.cpp
+++ b/Engine/Source/Editor/UnrealEd/Private/ThumbnailRendering/TextureThumbnailRenderer.cpp
@@ -10,6 +10,7 @@
 #include "EngineGlobals.h"
 #include "Engine/TextureCube.h"
 #include "Engine/Texture2DArray.h"
+#include "Engine/TextureCubeArray.h"
 #include "Texture2DPreview.h"
 #include "Engine/TextureRenderTargetCube.h"
 
@@ -61,13 +62,14 @@ void UTextureThumbnailRenderer::Draw(UObject* Object, int32 X, int32 Y, uint32 W
 
 		UTextureCube* TextureCube = Cast<UTextureCube>(Texture);
 		UTexture2DArray* Texture2DArray = Cast<UTexture2DArray>(Texture);
+		UTextureCubeArray* TextureCubeArray = Cast<UTextureCubeArray>(Texture);
 		UTextureRenderTargetCube* RTTextureCube = Cast<UTextureRenderTargetCube>(Texture);
 		UTextureLightProfile* TextureLightProfile = Cast<UTextureLightProfile>(Texture);
 		const bool bIsVirtualTexture = Texture->IsCurrentlyVirtualTextured();
 
 		TRefCountPtr<FBatchedElementParameters> BatchedElementParameters;
 
-		if(TextureCube || RTTextureCube)
+		if(TextureCube || TextureCubeArray || RTTextureCube)
 		{
 			// is released by the render thread when it was rendered
 			BatchedElementParameters = new FMipLevelBatchedElementParameters((float)0);
diff --git a/Engine/Source/Runtime/Engine/Classes/Engine/TextureCubeArray.h b/Engine/Source/Runtime/Engine/Classes/Engine/TextureCubeArray.h
new file mode 100644
--- /dev/null
+++ b/Engine/Source/Runtime/Engine/Classes/Engine/TextureCubeArray.h
@@ -0,0 +1,98 @@
+// Copyright Epic Games, Inc. All Rights Reserved.
+
+#pragma once
+
+#include "CoreMinimal.h"
+#include "UObject/ObjectMacros.h"
+#include "Engine/TextureCube.h"
+#include "TextureCubeArray.generated.h"
+
+UCLASS(HideCategories = Object, MinimalAPI, BlueprintType)
+class UTextureCubeArray : public UTexture
+{
+	GENERATED_UCLASS_BODY()
+
+public:
+	/** Platform data. */
+	FTexturePlatformData* PlatformData;
+	TMap<FString, FTexturePlatformData*> CookedPlatformData;
+
+	/*
+	* Initialize texture source from textures in SourceArray.
+	* @param bUpdateSourceSettings Set to false to prevent overriding current texture settings.
+	*/
+	/** Trivial accessors. */
+	FORCEINLINE int32 GetSizeX() const
+	{
+		return PlatformData ? PlatformData->SizeX : 0;
+	}
+	FORCEINLINE int32 GetSizeY() const
+	{
+		return PlatformData ? PlatformData->SizeY : 0;
+	}
+	FORCEINLINE int32 GetNumSlices() const
+	{
+		return PlatformData ? PlatformData->GetNumSlices() : 0;
+	}
+	FORCEINLINE int32 GetNumMips() const
+	{
+		return PlatformData ? PlatformData->Mips.Num() : 0;
+	}
+	FORCEINLINE EPixelFormat GetPixelFormat() const
+	{
+		return PlatformData ? PlatformData->PixelFormat : PF_Unknown;
+	}
+
+	//~ Begin UTexture Interface
+	virtual void Serialize(FArchive& Ar) override;
+	virtual void PostLoad() override;
+	virtual void GetAssetRegistryTags(TArray<FAssetRegistryTag>& OutTags) const override;
+	virtual FString GetDesc() override;
+	virtual void GetResourceSizeEx(FResourceSizeEx& CumulativeResourceSize) override;
+	virtual float GetSurfaceWidth() const override { return GetSizeX(); }
+	virtual float GetSurfaceHeight() const override { return GetSizeY(); }
+	virtual float GetSurfaceDepth() const override { return 0; }
+	virtual uint32 GetSurfaceArraySize() const override { return GetNumSlices(); }
+	virtual FTextureResource* CreateResource() override;
+#if WITH_EDITOR
+	ENGINE_API virtual void PostEditChangeProperty(FPropertyChangedEvent& PropertyChangedEvent) override;
+	ENGINE_API bool UpdateSourceFromSourceTextures(bool bCreatingNewTexture = true);
+	ENGINE_API void InvadiateTextureSource();
+	ENGINE_API bool CheckArrayTexturesCompatibility();
+#endif // WITH_EDITOR
+	virtual void UpdateResource() override;
+	virtual EMaterialValueType GetMaterialType() const override { return MCT_TextureCubeArray; }
+	virtual FTexturePlatformData** GetRunningPlatformData() override { return &PlatformData; }
+	virtual TMap<FString, FTexturePlatformData*>* GetCookedPlatformData() override { return &CookedPlatformData; }
+	//~ End UTexture Interface
+
+#if WITH_EDITORONLY_DATA
+	/** Add Textures*/
+	UPROPERTY(EditAnywhere, Category = SourceCube, meta = (DisplayName = "Source Textures"))
+	TArray<TObjectPtr<UTextureCube>> SourceTextures;
+#endif
+
+	/**
+	 * Calculates the size of this texture in bytes if it had MipCount mip-levels streamed in.
+	 *
+	 * @param	MipCount Number of mips to calculate size for, counting from the smallest 1x1 mip-level and up.
+	 * @return	Size of MipCount mips in bytes
+	 */
+	uint32 CalcTextureMemorySize(int32 MipCount) const;
+
+	/**
+	* Calculates the size of this texture if it had MipCount mip levels streamed in.
+	*
+	* @param	Enum which mips to calculate size for.
+	* @return	Total size of all specified mips, in bytes
+	*/
+	virtual uint32 CalcTextureMemorySizeEnum(ETextureMipCount Enum) const override;
+
+#if WITH_EDITOR
+	/**
+	* Return maximum dimension for this texture type.
+	*/
+	virtual uint32 GetMaximumDimension() const override;
+
+#endif
+};
\ No newline at end of file
diff --git a/Engine/Source/Runtime/Engine/Classes/Materials/MaterialExpressionTextureSampleParameterCubeArray.h b/Engine/Source/Runtime/Engine/Classes/Materials/MaterialExpressionTextureSampleParameterCubeArray.h
new file mode 100644
--- /dev/null
+++ b/Engine/Source/Runtime/Engine/Classes/Materials/MaterialExpressionTextureSampleParameterCubeArray.h
@@ -0,0 +1,31 @@
+// Copyright Epic Games, Inc. All Rights Reserved.
+
+
+#pragma once
+
+#include "CoreMinimal.h"
+#include "UObject/ObjectMacros.h"
+#include "Materials/MaterialExpressionTextureSampleParameter.h"
+#include "MaterialExpressionTextureSampleParameterCubeArray.generated.h"
+
+class UTexture;
+
+UCLASS(collapsecategories, hidecategories = Object, MinimalAPI)
+class UMaterialExpressionTextureSampleParameterCubeArray : public UMaterialExpressionTextureSampleParameter
+{
+	GENERATED_UCLASS_BODY()
+
+#if WITH_EDITOR
+	//~ Begin UMaterialExpression Interface
+	virtual void GetCaption(TArray<FString>& OutCaptions) const override;
+	virtual int32 Compile(class FMaterialCompiler* Compiler, int32 OutputIndex) override;
+	//~ End UMaterialExpression Interface
+#endif // WITH_EDITOR
+
+	//~ Begin UMaterialExpressionTextureSampleParameter Interface
+#if WITH_EDITOR
+	virtual bool TextureIsValid(UTexture* InTexture, FString& OutMessage) override;
+#endif // WITH_EDITOR
+	virtual const TCHAR* GetRequirements();
+	//~ End UMaterialExpressionTextureSampleParameter Interface
+};
diff --git a/Engine/Source/Runtime/Engine/Private/Materials/HLSLMaterialTranslator.cpp b/Engine/Source/Runtime/Engine/Private/Materials/HLSLMaterialTranslator.cpp
--- a/Engine/Source/Runtime/Engine/Private/Materials/HLSLMaterialTranslator.cpp
+++ b/Engine/Source/Runtime/Engine/Private/Materials/HLSLMaterialTranslator.cpp
@@ -2603,6 +2603,7 @@ const TCHAR* FHLSLMaterialTranslator::DescribeType(EMaterialValueType Type) cons
 	case MCT_Texture2D:				return TEXT("texture2D");
 	case MCT_TextureCube:			return TEXT("textureCube");
 	case MCT_Texture2DArray:		return TEXT("texture2DArray");
+	case MCT_TextureCubeArray:		return TEXT("textureCubeArray");
 	case MCT_VolumeTexture:			return TEXT("volumeTexture");
 	case MCT_StaticBool:			return TEXT("static bool");
 	case MCT_MaterialAttributes:	return TEXT("MaterialAttributes");
@@ -2628,6 +2629,7 @@ const TCHAR* FHLSLMaterialTranslator::HLSLTypeString(EMaterialValueType Type) co
 	case MCT_Texture2D:				return TEXT("texture2D");
 	case MCT_TextureCube:			return TEXT("textureCube");
 	case MCT_Texture2DArray:		return TEXT("texture2DArray");
+	case MCT_TextureCubeArray:		return TEXT("textureCubeArray");
 	case MCT_VolumeTexture:			return TEXT("volumeTexture");
 	case MCT_StaticBool:			return TEXT("static bool");
 	case MCT_MaterialAttributes:	return TEXT("FMaterialAttributes");
@@ -2653,6 +2655,7 @@ const TCHAR* FHLSLMaterialTranslator::HLSLTypeStringDeriv(EMaterialValueType Typ
 	case MCT_Texture2D:				return TEXT("texture2D");
 	case MCT_TextureCube:			return TEXT("textureCube");
 	case MCT_Texture2DArray:		return TEXT("texture2DArray");
+	case MCT_TextureCubeArray:		return TEXT("textureCubeArray");
 	case MCT_VolumeTexture:			return TEXT("volumeTexture");
 	case MCT_StaticBool:			return TEXT("static bool");
 	case MCT_MaterialAttributes:	return TEXT("FMaterialAttributes");
@@ -3113,6 +3116,10 @@ int32 FHLSLMaterialTranslator::AccessUniformExpression(int32 Index)
 			TextureInputIndex = UniformTextureExpressions[(uint32)EMaterialTextureParameterType::Array2D].AddUnique(TextureUniformExpression);
 			BaseName = TEXT("Texture2DArray");
 			break;
+		case MCT_TextureCubeArray:
+			TextureInputIndex = UniformTextureExpressions[(uint32)EMaterialTextureParameterType::ArrayCube].AddUnique(TextureUniformExpression);
+			BaseName = TEXT("TextureCubeArray");
+			break;
 		case MCT_VolumeTexture:
 			TextureInputIndex = UniformTextureExpressions[(uint32)EMaterialTextureParameterType::Volume].AddUnique(TextureUniformExpression);
 			BaseName = TEXT("VolumeTexture");
@@ -5525,6 +5532,10 @@ int32 FHLSLMaterialTranslator::TextureSample(
 	{
 		TextureName = CoerceParameter(TextureIndex, MCT_Texture2DArray);
 	}
+	else if (TextureType == MCT_TextureCubeArray)
+	{
+		TextureName = CoerceParameter(TextureIndex, MCT_TextureCubeArray);
+	}
 	else if (TextureType == MCT_VolumeTexture)
 		{
 		TextureName = CoerceParameter(TextureIndex, MCT_VolumeTexture);
@@ -5597,9 +5608,22 @@ int32 FHLSLMaterialTranslator::TextureSample(
 		}
 	}
 		
-	const EMaterialValueType UVsType = (TextureType == MCT_TextureCube || TextureType == MCT_Texture2DArray || TextureType == MCT_VolumeTexture) ? MCT_Float3 : MCT_Float2;
+	EMaterialValueType UVsType;
+	switch (TextureType)
+	{
+	case MCT_TextureCubeArray:
+		UVsType = MCT_Float4;
+		break;
+	case MCT_TextureCube:
+	case MCT_Texture2DArray:
+	case MCT_VolumeTexture:
+		UVsType = MCT_Float3;
+		break;
+	default:
+		UVsType = MCT_Float2;
+		break;
+	}
 	const EDerivativeStatus UvDerivativeStatus = GetDerivativeStatus(CoordinateIndex);
-	
 	if (RequiresManualViewMipBias)
 	{
 		if (MipValueMode == TMVM_Derivative)
@@ -5645,6 +5669,10 @@ int32 FHLSLMaterialTranslator::TextureSample(
 		{
 		TextureTypeName = TEXT("Texture2DArraySample");
 		}
+	else if (TextureType == MCT_TextureCubeArray)
+	{
+		TextureTypeName = TEXT("TextureCubeArraySample");
+	}
 	else if (TextureType == MCT_VolumeTexture)
 		{
 		TextureTypeName = TEXT("Texture3DSample");
@@ -9374,6 +9402,13 @@ int32 FHLSLMaterialTranslator::CustomExpression( class UMaterialExpressionCustom
 				InputParamDecl += InputNameStr;
 				InputParamDecl += TEXT("Sampler ");
 				break;
+			case MCT_TextureCubeArray:
+				InputParamDecl += TEXT("TextureCubeArray ");
+				InputParamDecl += InputNameStr;
+				InputParamDecl += TEXT(", SamplerState ");
+				InputParamDecl += InputNameStr;
+				InputParamDecl += TEXT("Sampler ");
+				break;
 			case MCT_TextureExternal:
 				InputParamDecl += TEXT("TextureExternal ");
 				InputParamDecl += InputNameStr;
@@ -9518,7 +9553,7 @@ int32 FHLSLMaterialTranslator::CustomExpression( class UMaterialExpressionCustom
 
 			CodeChunk += TEXT(",");
 			CodeChunk += *ParamCode;
-			if (ParamType == MCT_Texture2D || ParamType == MCT_TextureCube || ParamType == MCT_Texture2DArray || ParamType == MCT_TextureExternal || ParamType == MCT_VolumeTexture)
+			if (ParamType == MCT_Texture2D || ParamType == MCT_TextureCube || ParamType == MCT_TextureCubeArray || ParamType == MCT_Texture2DArray || ParamType == MCT_TextureExternal || ParamType == MCT_VolumeTexture)
 			{
 				CodeChunk += TEXT(",");
 				CodeChunk += *ParamCode;
diff --git a/Engine/Source/Runtime/Engine/Private/Materials/MaterialExpressions.cpp b/Engine/Source/Runtime/Engine/Private/Materials/MaterialExpressions.cpp
--- a/Engine/Source/Runtime/Engine/Private/Materials/MaterialExpressions.cpp
+++ b/Engine/Source/Runtime/Engine/Private/Materials/MaterialExpressions.cpp
@@ -37,6 +37,7 @@
 #include "Engine/Texture2DDynamic.h"
 #include "Engine/TextureCube.h"
 #include "Engine/Texture2DArray.h"
+#include "Engine/TextureCubeArray.h"
 #include "Engine/TextureRenderTargetCube.h"
 #include "Engine/VolumeTexture.h"
 #include "Styling/CoreStyle.h"
@@ -217,6 +218,7 @@
 #include "Materials/MaterialExpressionTextureSampleParameterSubUV.h"
 #include "Materials/MaterialExpressionTextureSampleParameterCube.h"
 #include "Materials/MaterialExpressionTextureSampleParameter2DArray.h"
+#include "Materials/MaterialExpressionTextureSampleParameterCubeArray.h"
 #include "Materials/MaterialExpressionTextureSampleParameterVolume.h"
 #include "Materials/MaterialExpressionTextureCoordinate.h"
 #include "Materials/MaterialExpressionTime.h"
@@ -395,6 +397,9 @@ void GetMaterialValueTypeDescriptions(const uint32 MaterialValueType, TArray<FTe
 			case MCT_Texture2DArray:
 				OutDescriptions.Add(LOCTEXT("Texture2DArray", "Texture 2D Array"));
 				break;
+			case MCT_TextureCubeArray:
+				OutDescriptions.Add(LOCTEXT("TextureCubeArray", "Texture Cube Array"));
+				break;
 			case MCT_VolumeTexture:
 				OutDescriptions.Add(LOCTEXT("VolumeTexture", "Volume Texture"));
 				break;
@@ -1975,6 +1980,10 @@ int32 UMaterialExpressionTextureSample::Compile(class FMaterialCompiler* Compile
 				{
 					return CompilerError(Compiler, TEXT("UVW input required for texturearray sample"));
 				}
+				else if (TextureType == MCT_TextureCubeArray && !Coordinates.GetTracedInput().Expression)
+				{
+					return CompilerError(Compiler, TEXT("UVWX input required for texturecubearray sample"));
+				}
 			}
 
 			int32 CoordinateIndex = Coordinates.GetTracedInput().Expression ? Coordinates.Compile(Compiler) : Compiler->TextureCoordinate(ConstCoordinate, false, false);
@@ -3222,6 +3231,10 @@ uint32 UMaterialExpressionTextureObject::GetOutputType(int32 OutputIndex)
 	{
 		return MCT_Texture2DArray;
 	}
+	else if (Cast<UTextureCubeArray>(Texture) != nullptr)
+	{
+		return MCT_TextureCubeArray;
+	}
 	else if (Cast<UVolumeTexture>(Texture) != nullptr)
 	{
 		return MCT_VolumeTexture;
@@ -3565,6 +3578,71 @@ const TCHAR* UMaterialExpressionTextureSampleParameter2DArray::GetRequirements()
 	return TEXT("Requires Texture2DArray");
 }
 
+//
+//  UMaterialExpressionTextureSampleParameterCubeArray
+//
+UMaterialExpressionTextureSampleParameterCubeArray::UMaterialExpressionTextureSampleParameterCubeArray(const FObjectInitializer& ObjectInitializer)
+	: Super(ObjectInitializer)
+{
+	// Structure to hold one-time initialization
+	struct FConstructorStatics
+	{
+		FText NAME_Texture;
+		FText NAME_Parameters;
+		FConstructorStatics()
+			: NAME_Texture(LOCTEXT("Texture", "Texture"))
+			, NAME_Parameters(LOCTEXT("Parameters", "Parameters"))
+		{
+		}
+	};
+	static FConstructorStatics ConstructorStatics;
+
+#if WITH_EDITORONLY_DATA
+	MenuCategories.Empty();
+	MenuCategories.Add(ConstructorStatics.NAME_Texture);
+	MenuCategories.Add(ConstructorStatics.NAME_Parameters);
+#endif
+}
+
+#if WITH_EDITOR
+int32 UMaterialExpressionTextureSampleParameterCubeArray::Compile(class FMaterialCompiler* Compiler, int32 OutputIndex)
+{
+	if (!Coordinates.GetTracedInput().Expression)
+	{
+		return CompilerError(Compiler, TEXT("Cube Array sample needs UV input"));
+	}
+
+	return UMaterialExpressionTextureSampleParameter::Compile(Compiler, OutputIndex);
+}
+
+void UMaterialExpressionTextureSampleParameterCubeArray::GetCaption(TArray<FString>& OutCaptions) const
+{
+	OutCaptions.Add(TEXT("ParamCubeArray"));
+	OutCaptions.Add(FString::Printf(TEXT("'%s'"), *ParameterName.ToString()));
+}
+
+bool UMaterialExpressionTextureSampleParameterCubeArray::TextureIsValid(UTexture* InTexture, FString& OutMessage)
+{
+	if (!InTexture)
+	{
+		OutMessage = TEXT("Found NULL, requires TextureCubeArray");
+		return false;
+	}
+	else if (!(InTexture->GetMaterialType() & MCT_TextureCubeArray))
+	{
+		OutMessage = FString::Printf(TEXT("Found %s, requires TextureCubeArray"), *InTexture->GetClass()->GetName());
+		return false;
+	}
+
+	return true;
+}
+#endif // WITH_EDITOR
+
+const TCHAR* UMaterialExpressionTextureSampleParameterCubeArray::GetRequirements()
+{
+	return TEXT("Requires TextureCubeArray");
+}
+
 //
 //  UMaterialExpressionTextureSampleParameterVolume
 //
diff --git a/Engine/Source/Runtime/Engine/Private/Materials/MaterialInstance.cpp b/Engine/Source/Runtime/Engine/Private/Materials/MaterialInstance.cpp
--- a/Engine/Source/Runtime/Engine/Private/Materials/MaterialInstance.cpp
+++ b/Engine/Source/Runtime/Engine/Private/Materials/MaterialInstance.cpp
@@ -1492,6 +1492,12 @@ void UMaterialInstance::ValidateTextureOverrides(ERHIFeatureLevel::Type InFeatur
 						UE_LOG(LogMaterial, Error, TEXT("MaterialInstance \"%s\" parameter '%s' assigned texture \"%s\" has invalid type, required texture array"), *MaterialName, *TextureInfo.GetParameterName().ToString(), *Texture->GetName());
 					}
 					break;
+				case EMaterialTextureParameterType::ArrayCube:
+					if (!(TextureType & MCT_TextureCubeArray))
+					{
+						UE_LOG(LogMaterial, Error, TEXT("MaterialInstance \"%s\" parameter '%s' assigned texture \"%s\" has invalid type, required texture cube array"), *MaterialName, *TextureInfo.GetParameterName().ToString(), *Texture->GetName());
+					}
+					break;
 				case EMaterialTextureParameterType::Volume:
 					if (!(TextureType & MCT_VolumeTexture))
 					{
diff --git a/Engine/Source/Runtime/Engine/Private/Materials/MaterialUniformExpressions.cpp b/Engine/Source/Runtime/Engine/Private/Materials/MaterialUniformExpressions.cpp
--- a/Engine/Source/Runtime/Engine/Private/Materials/MaterialUniformExpressions.cpp
+++ b/Engine/Source/Runtime/Engine/Private/Materials/MaterialUniformExpressions.cpp
@@ -347,12 +347,13 @@ bool FUniformExpressionSet::operator==(const FUniformExpressionSet& ReferenceSet
 
 FString FUniformExpressionSet::GetSummaryString() const
 {
-	return FString::Printf(TEXT("(%u vectors, %u scalars, %u 2d tex, %u cube tex, %u 2darray tex, %u 3d tex, %u virtual tex, %u external tex, %u VT stacks, %u collections)"),
+	return FString::Printf(TEXT("(%u vectors, %u scalars, %u 2d tex, %u cube tex, %u 2darray tex, %u cubearray tex, %u 3d tex, %u virtual tex, %u external tex, %u VT stacks, %u collections)"),
 		UniformVectorPreshaders.Num(), 
 		UniformScalarPreshaders.Num(),
 		UniformTextureParameters[(uint32)EMaterialTextureParameterType::Standard2D].Num(),
 		UniformTextureParameters[(uint32)EMaterialTextureParameterType::Cube].Num(),
 		UniformTextureParameters[(uint32)EMaterialTextureParameterType::Array2D].Num(),
+		UniformTextureParameters[(uint32)EMaterialTextureParameterType::ArrayCube].Num(),
 		UniformTextureParameters[(uint32)EMaterialTextureParameterType::Volume].Num(),
 		UniformTextureParameters[(uint32)EMaterialTextureParameterType::Virtual].Num(),
 		UniformExternalTextureParameters.Num(),
@@ -414,6 +415,8 @@ FShaderParametersMetadata* FUniformExpressionSet::CreateBufferStruct()
 	static FString TextureCubeSamplerNames[128];
 	static FString Texture2DArrayNames[128];
 	static FString Texture2DArraySamplerNames[128];
+	static FString TextureCubeArrayNames[128];
+	static FString TextureCubeArraySamplerNames[128];
 	static FString VolumeTextureNames[128];
 	static FString VolumeTextureSamplerNames[128];
 	static FString ExternalTextureNames[128];
@@ -435,6 +438,8 @@ FShaderParametersMetadata* FUniformExpressionSet::CreateBufferStruct()
 			TextureCubeSamplerNames[i] = FString::Printf(TEXT("TextureCube_%dSampler"), i);
 			Texture2DArrayNames[i] = FString::Printf(TEXT("Texture2DArray_%d"), i);
 			Texture2DArraySamplerNames[i] = FString::Printf(TEXT("Texture2DArray_%dSampler"), i);
+			TextureCubeArrayNames[i] = FString::Printf(TEXT("TextureCubeArray_%d"), i);
+			TextureCubeArraySamplerNames[i] = FString::Printf(TEXT("TextureCubeArray_%dSampler"), i);
 			VolumeTextureNames[i] = FString::Printf(TEXT("VolumeTexture_%d"), i);
 			VolumeTextureSamplerNames[i] = FString::Printf(TEXT("VolumeTexture_%dSampler"), i);
 			ExternalTextureNames[i] = FString::Printf(TEXT("ExternalTexture_%d"), i);
@@ -480,6 +485,15 @@ FShaderParametersMetadata* FUniformExpressionSet::CreateBufferStruct()
 		NextMemberOffset += SHADER_PARAMETER_POINTER_ALIGNMENT;
 	}
 
+	for (int32 i = 0; i < UniformTextureParameters[(uint32)EMaterialTextureParameterType::ArrayCube].Num(); ++i)
+	{
+		check((NextMemberOffset % SHADER_PARAMETER_POINTER_ALIGNMENT) == 0);
+		new(Members) FShaderParametersMetadata::FMember(*TextureCubeArrayNames[i], TEXT("TextureCubeArray"), __LINE__, NextMemberOffset, UBMT_TEXTURE, EShaderPrecisionModifier::Float, 1, 1, 0, NULL);
+		NextMemberOffset += SHADER_PARAMETER_POINTER_ALIGNMENT;
+		new(Members) FShaderParametersMetadata::FMember(*TextureCubeArraySamplerNames[i], TEXT("SamplerState"), __LINE__, NextMemberOffset, UBMT_SAMPLER, EShaderPrecisionModifier::Float, 1, 1, 0, NULL);
+		NextMemberOffset += SHADER_PARAMETER_POINTER_ALIGNMENT;
+	}
+
 	for (int32 i = 0; i < UniformTextureParameters[(uint32)EMaterialTextureParameterType::Volume].Num(); ++i)
 	{
 		check((NextMemberOffset % SHADER_PARAMETER_POINTER_ALIGNMENT) == 0);
@@ -821,6 +835,7 @@ void FUniformExpressionSet::FillUniformBuffer(const FMaterialRenderContext& Mate
 				UniformTextureParameters[(uint32)EMaterialTextureParameterType::Standard2D].Num() * 2
 				+ UniformTextureParameters[(uint32)EMaterialTextureParameterType::Cube].Num() * 2
 				+ UniformTextureParameters[(uint32)EMaterialTextureParameterType::Array2D].Num() * 2
+				+ UniformTextureParameters[(uint32)EMaterialTextureParameterType::ArrayCube].Num() * 2
 				+ UniformTextureParameters[(uint32)EMaterialTextureParameterType::Volume].Num() * 2
 				+ UniformExternalTextureParameters.Num() * 2
 				+ UniformTextureParameters[(uint32)EMaterialTextureParameterType::Virtual].Num() * 2
@@ -1005,6 +1020,45 @@ void FUniformExpressionSet::FillUniformBuffer(const FMaterialRenderContext& Mate
 			}
 		}
 
+		// Cache cube array texture uniform expressions.
+		for (int32 ExpressionIndex = 0; ExpressionIndex < GetNumTextures(EMaterialTextureParameterType::ArrayCube); ExpressionIndex++)
+		{
+			const FMaterialTextureParameterInfo& Parameter = GetTextureParameter(EMaterialTextureParameterType::ArrayCube, ExpressionIndex);
+
+			const UTexture* Value;
+			GetTextureValue(EMaterialTextureParameterType::ArrayCube, ExpressionIndex, MaterialRenderContext, MaterialRenderContext.Material, Value);
+
+			void** ResourceTableTexturePtr = (void**)((uint8*)BufferCursor + 0 * SHADER_PARAMETER_POINTER_ALIGNMENT);
+			void** ResourceTableSamplerPtr = (void**)((uint8*)BufferCursor + 1 * SHADER_PARAMETER_POINTER_ALIGNMENT);
+			BufferCursor = ((uint8*)BufferCursor) + (SHADER_PARAMETER_POINTER_ALIGNMENT * 2);
+
+			if (Value && Value->GetResource() && (Value->GetMaterialType() & MCT_TextureCubeArray) != 0u)
+			{
+				check(Value->TextureReference.TextureReferenceRHI);
+				*ResourceTableTexturePtr = Value->TextureReference.TextureReferenceRHI;
+				const FSamplerStateRHIRef* SamplerSource = &Value->GetResource()->SamplerStateRHI;
+				ESamplerSourceMode SourceMode = Parameter.SamplerSource;
+				if (SourceMode == SSM_Wrap_WorldGroupSettings)
+				{
+					SamplerSource = &Wrap_WorldGroupSettings->SamplerStateRHI;
+				}
+				else if (SourceMode == SSM_Clamp_WorldGroupSettings)
+				{
+					SamplerSource = &Clamp_WorldGroupSettings->SamplerStateRHI;
+				}
+
+				check(*SamplerSource);
+				*ResourceTableSamplerPtr = *SamplerSource;
+			}
+			else
+			{
+				check(GBlackArrayTexture->TextureRHI);
+				*ResourceTableTexturePtr = GBlackCubeArrayTexture->TextureRHI;
+				check(GBlackArrayTexture->SamplerStateRHI);
+				*ResourceTableSamplerPtr = GBlackCubeArrayTexture->SamplerStateRHI;
+			}
+		}
+
 		// Cache volume texture uniform expressions.
 		for (int32 ExpressionIndex = 0;ExpressionIndex < GetNumTextures(EMaterialTextureParameterType::Volume);ExpressionIndex++)
 		{
diff --git a/Engine/Source/Runtime/Engine/Private/ShaderDerivedDataVersion.h b/Engine/Source/Runtime/Engine/Private/ShaderDerivedDataVersion.h
--- a/Engine/Source/Runtime/Engine/Private/ShaderDerivedDataVersion.h
+++ b/Engine/Source/Runtime/Engine/Private/ShaderDerivedDataVersion.h
@@ -11,4 +11,4 @@
 // In case of merge conflicts with DDC versions, you MUST generate a new GUID and set this new
 // guid as version
 #define GLOBALSHADERMAP_DERIVEDDATA_VER			TEXT("60CD2243B26E480D898BAA7C66094C0A")
-#define MATERIALSHADERMAP_DERIVEDDATA_VER		TEXT("20297056A97D43148C19F0F2D16CFE78")
\ No newline at end of file
+#define MATERIALSHADERMAP_DERIVEDDATA_VER		TEXT("2831FE97589D4E34A316BE837CCA0039")
\ No newline at end of file
diff --git a/Engine/Source/Runtime/Engine/Private/TextureCubeArray.cpp b/Engine/Source/Runtime/Engine/Private/TextureCubeArray.cpp
new file mode 100644
--- /dev/null
+++ b/Engine/Source/Runtime/Engine/Private/TextureCubeArray.cpp
@@ -0,0 +1,604 @@
+// Copyright Epic Games, Inc. All Rights Reserved.
+
+/*=============================================================================
+	TextureCubeArray.cpp: UTextureCubeArray implementation.
+=============================================================================*/
+
+#include "Engine/TextureCubeArray.h"
+#include "RenderUtils.h"
+#include "TextureResource.h"
+#include "EngineUtils.h"
+#include "DeviceProfiles/DeviceProfile.h"
+#include "DeviceProfiles/DeviceProfileManager.h"
+#include "Containers/ResourceArray.h"
+
+class FTextureCubeArrayResource : public FTextureResource
+{
+	/** The FName of the LODGroup-specific stat	*/
+	FName					LODGroupStatName;
+
+public:
+	/**
+	 * Minimal initialization constructor.
+	 * @param InOwner - The UTextureCube which this FTextureCubeResource represents.
+	 */
+	FTextureCubeArrayResource(UTextureCubeArray* InOwner)
+		: Owner(InOwner)
+		, TextureSize(0)
+		, ProxiedResource(nullptr)
+	{
+		check(Owner->GetNumMips() > 0);
+
+		TIndirectArray<FTexture2DMipMap>& Mips = InOwner->PlatformData->Mips;
+		const int32 FirstMipTailIndex = Mips.Num() - FMath::Max(1, InOwner->PlatformData->GetNumMipsInTail());
+
+		NumSlices = Owner->GetNumSlices();
+		MipData.Empty(NumSlices * (FirstMipTailIndex + 1));
+
+		for (int32 MipIndex = 0; MipIndex <= FirstMipTailIndex; MipIndex++)
+		{
+			FTexture2DMipMap& Mip = Mips[MipIndex];
+			if (Mip.BulkData.GetBulkDataSize() <= 0)
+			{
+				UE_LOG(LogTexture, Error, TEXT("Corrupt texture [%s]! Missing bulk data for MipIndex=%d"), *InOwner->GetFullName(), MipIndex);
+			}
+			else
+			{
+				TextureSize += Mip.BulkData.GetBulkDataSize();
+				uint32 MipSize = Mip.BulkData.GetBulkDataSize() / NumSlices;
+
+				uint8* In = (uint8*)Mip.BulkData.Lock(LOCK_READ_ONLY);
+				for (uint32 Face = 0; Face < NumSlices; ++Face)
+				{
+					void* MipMemory = FMemory::Malloc(MipSize);
+					FMemory::Memcpy(MipMemory, In + MipSize * Face, MipSize);
+					MipData.Add(MipMemory);
+				}
+				Mip.BulkData.Unlock();
+			}
+		}
+		STAT(LODGroupStatName = TextureGroupStatFNames[InOwner->LODGroup]);
+	}
+
+	/**
+	 * Minimal initialization constructor.
+	 * @param InOwner           - The UTextureCube which this FTextureCubeResource represents.
+	 * @param InProxiedResource - The resource to proxy.
+	 */
+	FTextureCubeArrayResource(UTextureCubeArray* InOwner, const FTextureCubeArrayResource* InProxiedResource)
+		: Owner(InOwner)
+		, NumSlices(0u)
+		, TextureSize(0)
+		, ProxiedResource(InProxiedResource)
+	{
+	}
+
+	/**
+	 * Destructor, freeing MipData in the case of resource being destroyed without ever
+	 * having been initialized by the rendering thread via InitRHI.
+	 */
+	~FTextureCubeArrayResource()
+	{
+		// Make sure we're not leaking memory if InitRHI has never been called.
+		for (void* Mip : MipData)
+		{
+			if (Mip)
+			{
+				FMemory::Free(Mip);
+			}
+		}
+	}
+
+	/**
+	 * Called when the resource is initialized. This is only called by the rendering thread.
+	 */
+	virtual void InitRHI() override
+	{
+		if (ProxiedResource)
+		{
+			TextureCubeRHI = ProxiedResource->GetTextureCubeRHI();
+			TextureRHI = TextureCubeRHI;
+			SamplerStateRHI = ProxiedResource->SamplerStateRHI;
+			RHIUpdateTextureReference(Owner->TextureReference.TextureReferenceRHI, TextureRHI);
+			return;
+		}
+
+		INC_DWORD_STAT_BY(STAT_TextureMemory, TextureSize);
+		INC_DWORD_STAT_FNAME_BY(LODGroupStatName, TextureSize);
+
+		const uint32 ArraySize = NumSlices / 6u;
+
+		// Create the RHI texture.
+		ETextureCreateFlags TexCreateFlags = (Owner->SRGB ? TexCreate_SRGB : TexCreate_None) | (Owner->bNotOfflineProcessed ? TexCreate_None : TexCreate_OfflineProcessed);
+		FString Name = Owner->GetPathName();
+		FRHIResourceCreateInfo CreateInfo(*Name);
+		CreateInfo.ExtData = Owner->PlatformData ? Owner->PlatformData->GetExtData() : 0;
+		TextureCubeRHI = RHICreateTextureCubeArray(Owner->GetSizeX(), ArraySize, Owner->GetPixelFormat(), Owner->GetNumMips(), TexCreateFlags, CreateInfo);
+		TextureRHI = TextureCubeRHI;
+		TextureRHI->SetName(Owner->GetFName());
+		RHIBindDebugLabelName(TextureRHI, *Name);
+		RHIUpdateTextureReference(Owner->TextureReference.TextureReferenceRHI, TextureRHI);
+
+		// Read the mip-levels into the RHI texture.
+		const int32 NumMips = Owner->GetNumMips();
+		check(NumMips < MAX_TEXTURE_MIP_COUNT);
+		for (int32 MipIndex = 0; MipIndex < NumMips; MipIndex++)
+		{
+			for (uint32 ArrayIndex = 0u; ArrayIndex < ArraySize; ++ArrayIndex)
+			{
+				for (uint32 FaceIndex = 0u; FaceIndex < 6u; ++FaceIndex)
+				{
+					uint32 DestStride;
+					void* TheMipData = RHILockTextureCubeFace(TextureCubeRHI, FaceIndex, ArrayIndex, MipIndex, RLM_WriteOnly, DestStride, false);
+					GetData(ArrayIndex, FaceIndex, MipIndex, TheMipData, DestStride);
+					RHIUnlockTextureCubeFace(TextureCubeRHI, FaceIndex, ArrayIndex, MipIndex, false);
+				}
+			}
+		}
+
+		// Create the sampler state RHI resource.
+		FSamplerStateInitializerRHI SamplerStateInitializer
+		(
+			(ESamplerFilter)UDeviceProfileManager::Get().GetActiveProfile()->GetTextureLODSettings()->GetSamplerFilter(Owner),
+			AM_Clamp,
+			AM_Clamp,
+			AM_Clamp
+		);
+		SamplerStateRHI = GetOrCreateSamplerState(SamplerStateInitializer);
+
+		// Set the greyscale format flag appropriately.
+		EPixelFormat PixelFormat = Owner->GetPixelFormat();
+		bGreyScaleFormat = (PixelFormat == PF_G8) || (PixelFormat == PF_BC4);
+	}
+
+	virtual void ReleaseRHI() override
+	{
+		DEC_DWORD_STAT_BY(STAT_TextureMemory, TextureSize);
+		DEC_DWORD_STAT_FNAME_BY(LODGroupStatName, TextureSize);
+		RHIUpdateTextureReference(Owner->TextureReference.TextureReferenceRHI, nullptr);
+		TextureCubeRHI.SafeRelease();
+		FTextureResource::ReleaseRHI();
+	}
+
+	/** Returns the width of the texture in pixels. */
+	virtual uint32 GetSizeX() const override
+	{
+		if (ProxiedResource)
+		{
+			return ProxiedResource->GetSizeX();
+		}
+
+		return Owner->GetSizeX();
+	}
+
+	/** Returns the height of the texture in pixels. */
+	virtual uint32 GetSizeY() const override
+	{
+		if (ProxiedResource)
+		{
+			return ProxiedResource->GetSizeY();
+		}
+
+		return Owner->GetSizeY();
+	}
+
+	/** Returns the depth of the texture in pixels. */
+	virtual uint32 GetSizeZ() const
+	{
+		if (ProxiedResource)
+		{
+			return ProxiedResource->GetSizeZ();
+		}
+
+		return Owner->GetNumSlices();
+	}
+
+	FTextureCubeRHIRef GetTextureCubeRHI() const
+	{
+		return TextureCubeRHI;
+	}
+
+	virtual bool IsProxy() const override { return ProxiedResource != nullptr; }
+
+	const FTextureCubeArrayResource* GetProxiedResource() const { return ProxiedResource; }
+private:
+	/** A reference to the texture's RHI resource as a cube-map texture. */
+	FTextureCubeRHIRef TextureCubeRHI;
+
+	/** Local copy/ cache of mip data. Only valid between creation and first call to InitRHI */
+	TArray<void*> MipData;
+
+	/** The UTextureCube which this resource represents. */
+	const UTextureCubeArray* Owner;
+
+	/** Number of 2D faces per mip, equal to array size  6 */
+	uint32 NumSlices;
+
+	// Cached texture size for stats. */
+	int32	TextureSize;
+
+	const FTextureCubeArrayResource* const ProxiedResource;
+	/**
+	 * Writes the data for a single mip-level into a destination buffer.
+	 * @param FaceIndex		The index of the face of the mip-level to read.
+	 * @param MipIndex		The index of the mip-level to read.
+	 * @param Dest			The address of the destination buffer to receive the mip-level's data.
+	 * @param DestPitch		Number of bytes per row
+	 */
+	void GetData(uint32 SliceIndex, int32 FaceIndex, int32 MipIndex, void* Dest, uint32 DestPitch)
+	{
+		// Mips are stored sequentially
+		// For each mip, we store 6 faces per array index
+		const uint32 Index = MipIndex * NumSlices + SliceIndex * 6u + FaceIndex;
+		check(MipData[Index]);
+
+		// for platforms that returned 0 pitch from Lock, we need to just use the bulk data directly, never do 
+		// runtime block size checking, conversion, or the like
+		if (DestPitch == 0)
+		{
+			FMemory::Memcpy(Dest, MipData[Index], Owner->PlatformData->Mips[MipIndex].BulkData.GetBulkDataSize() / 6);
+		}
+		else
+		{
+			EPixelFormat PixelFormat = Owner->GetPixelFormat();
+			uint32 NumRows = 0;
+			uint32 SrcPitch = 0;
+			uint32 BlockSizeX = GPixelFormats[PixelFormat].BlockSizeX;	// Block width in pixels
+			uint32 BlockSizeY = GPixelFormats[PixelFormat].BlockSizeY;	// Block height in pixels
+			uint32 BlockBytes = GPixelFormats[PixelFormat].BlockBytes;
+
+			FIntPoint MipExtent = CalcMipMapExtent(Owner->GetSizeX(), Owner->GetSizeY(), PixelFormat, MipIndex);
+
+			uint32 NumColumns = (MipExtent.X + BlockSizeX - 1) / BlockSizeX;	// Num-of columns in the source data (in blocks)
+			NumRows = (MipExtent.Y + BlockSizeY - 1) / BlockSizeY;	// Num-of rows in the source data (in blocks)
+			SrcPitch = NumColumns * BlockBytes;		// Num-of bytes per row in the source data
+
+			SIZE_T MipSizeInBytes = CalcTextureMipMapSize(MipExtent.X, MipExtent.Y, PixelFormat, 0);
+
+			if (SrcPitch == DestPitch)
+			{
+				// Copy data, not taking into account stride!
+				FMemory::Memcpy(Dest, MipData[Index], MipSizeInBytes);
+			}
+			else
+			{
+				// Copy data, taking the stride into account!
+				uint8* Src = (uint8*)MipData[Index];
+				uint8* Dst = (uint8*)Dest;
+				for (uint32 Row = 0; Row < NumRows; ++Row)
+				{
+					FMemory::Memcpy(Dst, Src, SrcPitch);
+					Src += SrcPitch;
+					Dst += DestPitch;
+				}
+				check((PTRINT(Src) - PTRINT(MipData[Index])) == PTRINT(MipSizeInBytes));
+			}
+		}
+
+		FMemory::Free(MipData[Index]);
+		MipData[Index] = nullptr;
+	}
+};
+
+UTextureCubeArray::UTextureCubeArray(const FObjectInitializer& ObjectInitializer) : Super(ObjectInitializer)
+{
+#if WITH_EDITORONLY_DATA
+	SRGB = true;
+	MipGenSettings = TMGS_NoMipmaps;
+#endif
+}
+
+FTextureResource* UTextureCubeArray::CreateResource()
+{
+	const FPixelFormatInfo& FormatInfo = GPixelFormats[GetPixelFormat()];
+	if (GetNumMips() > 0 && FormatInfo.Supported)
+	{
+		return new FTextureCubeArrayResource(this);
+	}
+	else if (GetNumMips() == 0)
+	{
+		UE_LOG(LogTexture, Warning, TEXT("%s contains no miplevels! Please delete."), *GetFullName());
+	}
+	else if (!FormatInfo.Supported)
+	{
+		UE_LOG(LogTexture, Warning, TEXT("%s cannot be created, rhi does not support format %s."), *GetFullName(), FormatInfo.Name);
+	}
+	return nullptr;
+}
+
+void UTextureCubeArray::UpdateResource()
+{
+#if WITH_EDITOR
+	// Re-cache platform data if the source has changed.
+	CachePlatformData();
+#endif // #if WITH_EDITOR
+
+	Super::UpdateResource();
+}
+
+uint32 UTextureCubeArray::CalcTextureMemorySize(int32 MipCount) const
+{
+	uint32 Size = 0;
+	if (PlatformData)
+	{
+		int32 SizeX = GetSizeX();
+		int32 SizeY = GetSizeY();
+		const int32 ArraySize = GetNumSlices() / 6;
+		int32 NumMips = GetNumMips();
+		EPixelFormat Format = GetPixelFormat();
+
+		ensureMsgf(SizeX == SizeY, TEXT("Cubemap faces expected to be square.  Actual sizes are: %i, %i"), SizeX, SizeY);
+
+		// Figure out what the first mip to use is.
+		int32 FirstMip = FMath::Max(0, NumMips - MipCount);
+		FIntPoint MipExtents = CalcMipMapExtent(SizeX, SizeY, Format, FirstMip);
+
+		// TODO add RHICalcTextureCubeArrayPlatformSize
+		uint32 TextureAlign = 0;
+		uint64 TextureSize = RHICalcTextureCubePlatformSize(MipExtents.X, Format, FMath::Max(1, MipCount), TexCreate_None, FRHIResourceCreateInfo(PlatformData->GetExtData()), TextureAlign) * ArraySize;
+		Size = (uint32)TextureSize;
+	}
+	return Size;
+}
+
+uint32 UTextureCubeArray::CalcTextureMemorySizeEnum(ETextureMipCount Enum) const
+{
+	if (Enum == TMC_ResidentMips || Enum == TMC_AllMipsBiased)
+	{
+		return CalcTextureMemorySize(GetNumMips() - GetCachedLODBias());
+	}
+
+	return CalcTextureMemorySize(GetNumMips());
+}
+
+#if WITH_EDITOR
+
+ENGINE_API bool UTextureCubeArray::CheckArrayTexturesCompatibility()
+{
+	bool bError = false;
+
+	for (int32 TextureIndex = 0; TextureIndex < SourceTextures.Num(); ++TextureIndex)
+	{
+		// Do not create array till all texture slots are filled.
+		if (!SourceTextures[TextureIndex])
+		{
+			return false;
+		}
+
+		// Force the async texture build to complete
+		// TODO - better way to do this?
+		SourceTextures[TextureIndex]->GetPlatformData();
+
+		const FTextureSource& TextureSource = SourceTextures[TextureIndex]->Source;
+		// const int32 FormatDataSize = TextureSource.GetBytesPerPixel();
+		const ETextureSourceFormat SourceFormat = TextureSource.GetFormat();
+		const int32 SizeX = TextureSource.GetSizeX();
+		const int32 SizeY = TextureSource.GetSizeY();
+
+		for (int32 TextureCmpIndex = TextureIndex + 1; TextureCmpIndex < SourceTextures.Num(); ++TextureCmpIndex)
+		{
+			// Do not create array till all texture slots are filled.
+			if (!SourceTextures[TextureCmpIndex])
+			{
+				return false;
+			}
+
+			const FTextureSource& TextureSourceCmp = SourceTextures[TextureCmpIndex]->Source;
+			const FString TextureName = SourceTextures[TextureIndex]->GetFName().ToString();
+			const FString TextureNameCmp = SourceTextures[TextureCmpIndex]->GetFName().ToString();
+			const ETextureSourceFormat SourceFormatCmp = TextureSourceCmp.GetFormat();
+
+			if (TextureSourceCmp.GetSizeX() != SizeX || TextureSourceCmp.GetSizeY() != SizeY)
+			{
+				UE_LOG(LogTexture, Warning, TEXT("TextureCubeArray creation failed. Textures %s and %s have different sizes."), *TextureName, *TextureNameCmp);
+				bError = true;
+			}
+
+			if (SourceFormatCmp != SourceFormat)
+			{
+				UE_LOG(LogTexture, Warning, TEXT("TextureCubeArray creation failed. Textures %s and %s have incompatible pixel formats."), *TextureName, *TextureNameCmp);
+				bError = true;
+			}
+		}
+	}
+
+	return (!bError);
+}
+
+
+ENGINE_API bool UTextureCubeArray::UpdateSourceFromSourceTextures(bool bCreatingNewTexture)
+{
+	if (!CheckArrayTexturesCompatibility())
+	{
+		return false;
+	}
+
+	if (SourceTextures.Num() > 0)
+	{
+		FTextureSource& InitialSource = SourceTextures[0]->Source;
+		// Format and format size.
+		EPixelFormat PixelFormat = SourceTextures[0]->GetPixelFormat();
+		ETextureSourceFormat Format = InitialSource.GetFormat();
+		int32 FormatDataSize = InitialSource.GetBytesPerPixel();
+		// X,Y,Z size of the array.
+		int32 SizeX = InitialSource.GetSizeX();
+		int32 SizeY = InitialSource.GetSizeY();
+		uint32 ArraySize = SourceTextures.Num();
+		// Only copy the first mip from the source textures to array texture.
+		uint32 NumMips = 1;
+
+		// This should be false when texture is updated to avoid overriding user settings.
+		if (bCreatingNewTexture)
+		{
+			CompressionSettings = SourceTextures[0]->CompressionSettings;
+			MipGenSettings = TMGS_NoMipmaps;
+			PowerOfTwoMode = ETexturePowerOfTwoSetting::None;
+			LODGroup = SourceTextures[0]->LODGroup;
+			SRGB = SourceTextures[0]->SRGB;
+			NeverStream = true;
+		}
+
+		// Create the source texture for this UTexture.
+		const uint32 NumSlices = ArraySize * 6u;
+		Source.Init(SizeX, SizeY, NumSlices, NumMips, Format);
+
+		// We only copy the top level Mip map.
+		TArray<uint8*, TInlineAllocator<MAX_TEXTURE_MIP_COUNT> > DestMipData;
+		TArray<uint64, TInlineAllocator<MAX_TEXTURE_MIP_COUNT> > MipSizeBytes;
+		DestMipData.AddZeroed(NumMips);
+		MipSizeBytes.AddZeroed(NumMips);
+
+		for (uint32 MipIndex = 0; MipIndex < NumMips; ++MipIndex)
+		{
+			DestMipData[MipIndex] = Source.LockMip(MipIndex);
+			MipSizeBytes[MipIndex] = Source.CalcMipSize(MipIndex) / ArraySize;
+		}
+
+		TArray64<uint8> RefCubeData;
+		for (int32 SourceTexIndex = 0; SourceTexIndex < SourceTextures.Num(); ++SourceTexIndex)
+		{
+			for (uint32 MipIndex = 0; MipIndex < NumMips; ++MipIndex)
+			{
+				const uint64 Size = MipSizeBytes[MipIndex];
+				const uint64 CheckSize = SourceTextures[SourceTexIndex]->Source.CalcMipSize(MipIndex);
+				check(Size == CheckSize);
+
+				RefCubeData.Reset();
+				SourceTextures[SourceTexIndex]->Source.GetMipData(RefCubeData, MipIndex);
+				void* Dst = DestMipData[MipIndex] + Size * SourceTexIndex;
+				FMemory::Memcpy(Dst, RefCubeData.GetData(), Size);
+			}
+		}
+
+		for (uint32 MipIndex = 0; MipIndex < NumMips; ++MipIndex)
+		{
+			Source.UnlockMip(MipIndex);
+		}
+
+		SetLightingGuid();
+		UpdateResource();
+	}
+
+	return true;
+}
+
+ENGINE_API void UTextureCubeArray::InvadiateTextureSource()
+{
+	if (PlatformData)
+	{
+		delete PlatformData;
+		PlatformData = NULL;
+	}
+
+	Source.Init(0, 0, 0, 0, TSF_Invalid, nullptr);
+	UpdateResource();
+}
+#endif
+
+void UTextureCubeArray::Serialize(FArchive& Ar)
+{
+	DECLARE_SCOPE_CYCLE_COUNTER(TEXT("UTextureCubeArray::Serialize"), STAT_TextureCubeArray_Serialize, STATGROUP_LoadTime);
+
+	Super::Serialize(Ar);
+
+	FStripDataFlags StripFlags(Ar);
+	bool bCooked = Ar.IsCooking();
+	Ar << bCooked;
+
+	if (bCooked || Ar.IsCooking())
+	{
+		SerializeCookedPlatformData(Ar);
+	}
+
+#if WITH_EDITOR
+	if (Ar.IsLoading() && !Ar.IsTransacting() && !bCooked)
+	{
+		BeginCachePlatformData();
+	}
+#endif
+}
+
+void UTextureCubeArray::PostLoad()
+{
+#if WITH_EDITOR
+	FinishCachePlatformData();
+
+#endif // #if WITH_EDITOR
+	Super::PostLoad();
+};
+
+void UTextureCubeArray::GetAssetRegistryTags(TArray<FAssetRegistryTag>& OutTags) const
+{
+#if WITH_EDITOR
+	int32 SizeX = Source.GetSizeX();
+	int32 SizeY = Source.GetSizeY();
+	int32 ArraySize = Source.GetNumSlices() / 6;
+#else
+	int32 SizeX = 0;
+	int32 SizeY = 0;
+	int32 ArraySize = 0;
+#endif
+	const FString Dimensions = FString::Printf(TEXT("%dx%d*%d"), SizeX, SizeY, ArraySize);
+	OutTags.Add(FAssetRegistryTag("Dimensions", Dimensions, FAssetRegistryTag::TT_Dimensional));
+	OutTags.Add(FAssetRegistryTag("Format", GPixelFormats[GetPixelFormat()].Name, FAssetRegistryTag::TT_Alphabetical));
+
+	Super::GetAssetRegistryTags(OutTags);
+}
+
+FString UTextureCubeArray::GetDesc()
+{
+	return FString::Printf(TEXT("CubeArray: %dx%d*%d [%s]"),
+		GetSizeX(),
+		GetSizeY(),
+		GetNumSlices() / 6,
+		GPixelFormats[GetPixelFormat()].Name
+	);
+}
+
+void UTextureCubeArray::GetResourceSizeEx(FResourceSizeEx& CumulativeResourceSize)
+{
+	Super::GetResourceSizeEx(CumulativeResourceSize);
+	CumulativeResourceSize.AddUnknownMemoryBytes(CalcTextureMemorySizeEnum(TMC_ResidentMips));
+}
+
+#if WITH_EDITOR
+uint32 UTextureCubeArray::GetMaximumDimension() const
+{
+	return GetMaxCubeTextureDimension();
+
+}
+
+ENGINE_API void UTextureCubeArray::PostEditChangeProperty(FPropertyChangedEvent& PropertyChangedEvent)
+{
+	const FName PropertyName = PropertyChangedEvent.GetPropertyName();
+
+	if (PropertyName == GET_MEMBER_NAME_CHECKED(UTextureCubeArray, SourceTextures))
+	{
+		// Empty SourceTextures, remove any resources if present.
+		if (SourceTextures.Num() == 0)
+		{
+			InvadiateTextureSource();
+		}
+		// First entry into an empty texture array.
+		else if (SourceTextures.Num() == 1)
+		{
+			UpdateSourceFromSourceTextures(true);
+		}
+		// Couldn't add to non-empty array (Error msg logged).
+		else if (UpdateSourceFromSourceTextures(false) == false)
+		{
+			int32 ChangedIndex = PropertyChangedEvent.GetArrayIndex(PropertyName.ToString());
+			int32 LastIndex = SourceTextures.Num() - 1;
+
+			// But don't remove an empty texture, only an incompatible one.
+			if (SourceTextures[LastIndex] != nullptr && ChangedIndex == LastIndex)
+			{
+				SourceTextures.RemoveAt(LastIndex);
+			}
+		}
+	}
+
+	Super::PostEditChangeProperty(PropertyChangedEvent);
+}
+
+#endif // #if WITH_EDITOR
diff --git a/Engine/Source/Runtime/Engine/Private/TextureDerivedData.cpp b/Engine/Source/Runtime/Engine/Private/TextureDerivedData.cpp
--- a/Engine/Source/Runtime/Engine/Private/TextureDerivedData.cpp
+++ b/Engine/Source/Runtime/Engine/Private/TextureDerivedData.cpp
@@ -20,6 +20,7 @@
 #include "Engine/Texture2D.h"
 #include "Engine/TextureCube.h"
 #include "Engine/Texture2DArray.h"
+#include "Engine/TextureCubeArray.h"
 #include "DeviceProfiles/DeviceProfile.h"
 #include "DeviceProfiles/DeviceProfileManager.h"
 #include "TextureDerivedDataTask.h"
@@ -496,12 +497,17 @@ static void GetTextureBuildSettings(
 		OutBuildSettings.DiffuseConvolveMipLevel = GDiffuseConvolveMipLevel;
 		const UTextureCube* Cube = CastChecked<UTextureCube>(&Texture);
 		OutBuildSettings.bLongLatSource = (Cube->Source.GetNumSlices() == 1);
-		}
+	}
 	else if (Texture.IsA(UTexture2DArray::StaticClass()))
 	{
 		OutBuildSettings.bStreamable = GSupportsTexture2DArrayStreaming;
 		OutBuildSettings.bTextureArray = true;
 	}
+	else if (Texture.IsA(UTextureCubeArray::StaticClass()))
+	{
+		OutBuildSettings.bCubemap = true;
+		OutBuildSettings.bTextureArray = true;
+	}
 	else if (Texture.IsA(UVolumeTexture::StaticClass()))
 	{
 		OutBuildSettings.bStreamable = GSupportsVolumeTextureStreaming;
diff --git a/Engine/Source/Runtime/Engine/Private/TextureDerivedDataTask.cpp b/Engine/Source/Runtime/Engine/Private/TextureDerivedDataTask.cpp
--- a/Engine/Source/Runtime/Engine/Private/TextureDerivedDataTask.cpp
+++ b/Engine/Source/Runtime/Engine/Private/TextureDerivedDataTask.cpp
@@ -638,7 +638,18 @@ void FTextureCacheDerivedDataWorker::BuildTexture(bool bReplaceExistingDDC)
 						DerivedData->SizeX = CompressedImage.SizeX;
 						DerivedData->SizeY = CompressedImage.SizeY;
 						DerivedData->PixelFormat = (EPixelFormat)CompressedImage.PixelFormat;
-						DerivedData->SetNumSlices(BuildSettingsPerLayer[0].bCubemap ? 6 : (BuildSettingsPerLayer[0].bVolume || BuildSettingsPerLayer[0].bTextureArray) ? CompressedImage.SizeZ : 1);
+						if (BuildSettingsPerLayer[0].bVolume || BuildSettingsPerLayer[0].bTextureArray)
+						{
+							DerivedData->SetNumSlices(CompressedImage.SizeZ);
+						}
+						else if (BuildSettingsPerLayer[0].bCubemap)
+						{
+							DerivedData->SetNumSlices(6);
+						}
+						else
+						{
+							DerivedData->SetNumSlices(1);
+						}
 						DerivedData->SetIsCubemap(BuildSettingsPerLayer[0].bCubemap);
 					}
 					else
diff --git a/Engine/Source/Runtime/Engine/Public/MaterialShared.h b/Engine/Source/Runtime/Engine/Public/MaterialShared.h
--- a/Engine/Source/Runtime/Engine/Public/MaterialShared.h
+++ b/Engine/Source/Runtime/Engine/Public/MaterialShared.h
@@ -144,25 +144,26 @@ enum EMaterialValueType
 	MCT_Texture2D	          = 1 << 4,
 	MCT_TextureCube	          = 1 << 5,
 	MCT_Texture2DArray		  = 1 << 6,
-	MCT_VolumeTexture         = 1 << 7,
-	MCT_StaticBool            = 1 << 8,
-	MCT_Unknown               = 1 << 9,
-	MCT_MaterialAttributes	  = 1 << 10,
-	MCT_TextureExternal       = 1 << 11,
-	MCT_TextureVirtual        = 1 << 12,
-	MCT_Texture               = MCT_Texture2D | MCT_TextureCube | MCT_Texture2DArray |  MCT_VolumeTexture | MCT_TextureExternal | MCT_TextureVirtual,
+	MCT_TextureCubeArray      = 1 << 7,
+	MCT_VolumeTexture         = 1 << 8,
+	MCT_StaticBool            = 1 << 9,
+	MCT_Unknown               = 1 << 10,
+	MCT_MaterialAttributes	  = 1 << 11,
+	MCT_TextureExternal       = 1 << 12,
+	MCT_TextureVirtual        = 1 << 13,
+	MCT_Texture               = MCT_Texture2D | MCT_TextureCube | MCT_Texture2DArray | MCT_TextureCubeArray | MCT_VolumeTexture | MCT_TextureExternal | MCT_TextureVirtual,
 
 	/** Used internally when sampling from virtual textures */
-	MCT_VTPageTableResult     = 1 << 13,
+	MCT_VTPageTableResult     = 1 << 14,
 	
-	MCT_ShadingModel		  = 1 << 14,
+	MCT_ShadingModel		  = 1 << 15,
 
-	MCT_Strata				  = 1 << 15,
+	MCT_Strata				  = 1 << 16,
 
-	MCT_Execution             = 1 << 16,
+	MCT_Execution             = 1 << 17,
 
 	/** Used for code chunks that are statements with no value, rather than expressions */
-	MCT_VoidStatement         = 1 << 17,
+	MCT_VoidStatement         = 1 << 18,
 };
 
 /**
@@ -379,6 +380,7 @@ enum class EMaterialTextureParameterType : uint32
 	Standard2D,
 	Cube,
 	Array2D,
+	ArrayCube,
 	Volume,
 	Virtual,
 
