Support creating cube arrays from source cubemaps imported as long/lat maps

diff --git a/Engine/Source/Developer/TextureCompressor/Private/TextureCompressorModule.cpp b/Engine/Source/Developer/TextureCompressor/Private/TextureCompressorModule.cpp
--- a/Engine/Source/Developer/TextureCompressor/Private/TextureCompressorModule.cpp
+++ b/Engine/Source/Developer/TextureCompressor/Private/TextureCompressorModule.cpp
@@ -1073,11 +1073,11 @@ struct FImageViewLongLat
 	int32 SizeY;
 
 	/** Initialization constructor. */
-	explicit FImageViewLongLat(FImage& Image)
+	explicit FImageViewLongLat(FImage& Image, int32 SliceIndex)
 	{
-		ImageColors = (&Image.AsRGBA32F()[0]);
 		SizeX = Image.SizeX;
 		SizeY = Image.SizeY;
+		ImageColors = (&Image.AsRGBA32F()[0]) + SliceIndex * SizeY * SizeX;
 	}
 
 	/** Wraps X around W. */
@@ -1210,22 +1210,25 @@ void ITextureCompressorModule::GenerateBaseCubeMipFromLongitudeLatitude2D(FImage
 {
 	FImage LongLatImage;
 	SrcImage.CopyTo(LongLatImage, ERawImageFormat::RGBA32F, EGammaSpace::Linear);
-	FImageViewLongLat LongLatView(LongLatImage);
 
 	// TODO_TEXTURE: Expose target size to user.
 	uint32 Extent = ComputeLongLatCubemapExtents(LongLatImage, MaxCubemapTextureResolution);
 	float InvExtent = 1.0f / Extent;
-	OutMip->Init(Extent, Extent, 6, ERawImageFormat::RGBA32F, EGammaSpace::Linear);
+	OutMip->Init(Extent, Extent, SrcImage.NumSlices * 6, ERawImageFormat::RGBA32F, EGammaSpace::Linear);
 
-	for(uint32 Face = 0; Face < 6; ++Face)
+	for (int32 Slice = 0; Slice < SrcImage.NumSlices; ++Slice)
 	{
-		FImageView2D MipView(*OutMip, Face);
-		for(uint32 y = 0; y < Extent; ++y)
+		FImageViewLongLat LongLatView(LongLatImage, Slice);
+		for (uint32 Face = 0; Face < 6; ++Face)
 		{
-			for(uint32 x = 0; x < Extent; ++x)
+			FImageView2D MipView(*OutMip, Slice * 6 + Face);
+			for (uint32 y = 0; y < Extent; ++y)
 			{
-				FVector DirectionWS = ComputeWSCubeDirectionAtTexelCenter(Face, x, y, InvExtent);
-				MipView.Access(x, y) = LongLatView.LookupLongLat(DirectionWS);
+				for (uint32 x = 0; x < Extent; ++x)
+				{
+					FVector DirectionWS = ComputeWSCubeDirectionAtTexelCenter(Face, x, y, InvExtent);
+					MipView.Access(x, y) = LongLatView.LookupLongLat(DirectionWS);
+				}
 			}
 		}
 	}
@@ -2346,7 +2349,7 @@ private:
 		check(InSourceMips[0].SizeX > 0 && InSourceMips[0].SizeY > 0 && InSourceMips[0].NumSlices > 0);
 
 		// Identify long-lat cubemaps.
-		const bool bLongLatCubemap = BuildSettings.bCubemap && !BuildSettings.bTextureArray && InSourceMips[0].NumSlices == 1;
+		const bool bLongLatCubemap = BuildSettings.bLongLatSource;
 		if (BuildSettings.bCubemap && !bLongLatCubemap)
 		{
 			if (BuildSettings.bTextureArray && (InSourceMips[0].NumSlices % 6) != 0)
diff --git a/Engine/Source/Runtime/Engine/Classes/Engine/Texture.h b/Engine/Source/Runtime/Engine/Classes/Engine/Texture.h
--- a/Engine/Source/Runtime/Engine/Classes/Engine/Texture.h
+++ b/Engine/Source/Runtime/Engine/Classes/Engine/Texture.h
@@ -310,6 +310,7 @@ struct FTextureSource
 	FORCEINLINE int32 GetNumBlocks() const { return Blocks.Num() + 1; }
 	FORCEINLINE ETextureSourceFormat GetFormat(int32 LayerIndex = 0) const { return (LayerIndex == 0) ? Format : LayerFormat[LayerIndex]; }
 	FORCEINLINE bool IsPNGCompressed() const { return bPNGCompressed; }
+	FORCEINLINE bool IsLongLatCubemap() const { return bLongLatCubemap; }
 	FORCEINLINE int64 GetSizeOnDisk() const { return BulkData.GetPayloadSize(); }
 	inline bool HasPayloadData() const { return BulkData.HasPayloadData(); }
 	FORCEINLINE bool IsBulkDataLoaded() const { return BulkData.IsDataLoaded(); }
@@ -432,6 +433,7 @@ private:
 	friend class UTextureCube;
 	friend class UVolumeTexture;
 	friend class UTexture2DArray;
+	friend class UTextureCubeArray;
 
 #if WITH_EDITOR
 	/** Protects simultaneous access to BulkData */
@@ -536,6 +538,14 @@ private:
 	UPROPERTY(VisibleAnywhere, Category=TextureSource)
 	bool bPNGCompressed;
 
+	/**
+	 * Source represents a cubemap in long/lat format, will have only 1 slice per cube, rather than 6 slices.
+	 * Not needed for non-array cubemaps, since we can just look at NumSlices == 1 or 6
+	 * But for cube arrays, no way of determining whether NumSlices=6 means 1 cubemap, or 6 long/lat cubemaps
+	 */
+	UPROPERTY(VisibleAnywhere, Category = TextureSource)
+	bool bLongLatCubemap;
+
 	/** Compression format that source data is stored as. */
 	UPROPERTY(VisibleAnywhere, Category = TextureSource)
 	TEnumAsByte<enum ETextureSourceCompressionFormat> CompressionFormat;
diff --git a/Engine/Source/Runtime/Engine/Private/Texture.cpp b/Engine/Source/Runtime/Engine/Private/Texture.cpp
--- a/Engine/Source/Runtime/Engine/Private/Texture.cpp
+++ b/Engine/Source/Runtime/Engine/Private/Texture.cpp
@@ -1027,6 +1027,7 @@ FTextureSource::FTextureSource()
 	, NumMips(0)
 	, NumLayers(1) // Default to 1 so old data has the correct value
 	, bPNGCompressed(false)
+	, bLongLatCubemap(false)
 	, CompressionFormat(TSCF_None)
 	, bGuidIsHash(false)
 	, Format(TSF_Invalid)
diff --git a/Engine/Source/Runtime/Engine/Private/TextureCubeArray.cpp b/Engine/Source/Runtime/Engine/Private/TextureCubeArray.cpp
--- a/Engine/Source/Runtime/Engine/Private/TextureCubeArray.cpp
+++ b/Engine/Source/Runtime/Engine/Private/TextureCubeArray.cpp
@@ -355,7 +355,10 @@ uint32 UTextureCubeArray::CalcTextureMemorySizeEnum(ETextureMipCount Enum) const
 
 ENGINE_API bool UTextureCubeArray::CheckArrayTexturesCompatibility()
 {
-	bool bError = false;
+	if (SourceTextures.Num() == 0)
+	{
+		return false;
+	}
 
 	for (int32 TextureIndex = 0; TextureIndex < SourceTextures.Num(); ++TextureIndex)
 	{
@@ -368,37 +371,40 @@ ENGINE_API bool UTextureCubeArray::CheckArrayTexturesCompatibility()
 		// Force the async texture build to complete
 		// TODO - better way to do this?
 		SourceTextures[TextureIndex]->GetPlatformData();
+	}
 
-		const FTextureSource& TextureSource = SourceTextures[TextureIndex]->Source;
-		// const int32 FormatDataSize = TextureSource.GetBytesPerPixel();
-		const ETextureSourceFormat SourceFormat = TextureSource.GetFormat();
-		const int32 SizeX = TextureSource.GetSizeX();
-		const int32 SizeY = TextureSource.GetSizeY();
+	const FTextureSource& TextureSource = SourceTextures[0]->Source;
+	const FString TextureName = SourceTextures[0]->GetName();
+	const ETextureSourceFormat SourceFormat = TextureSource.GetFormat();
+	const int32 SizeX = TextureSource.GetSizeX();
+	const int32 SizeY = TextureSource.GetSizeY();
+	const int32 NumSlices = TextureSource.GetNumSlices();
 
-		for (int32 TextureCmpIndex = TextureIndex + 1; TextureCmpIndex < SourceTextures.Num(); ++TextureCmpIndex)
-		{
-			// Do not create array till all texture slots are filled.
-			if (!SourceTextures[TextureCmpIndex])
-			{
-				return false;
-			}
+	ensure(NumSlices == 1 || NumSlices == 6); // either cubemap, or lat/long map
+
+	bool bError = false;
+	for (int32 TextureCmpIndex =1; TextureCmpIndex < SourceTextures.Num(); ++TextureCmpIndex)
+	{
+		const FTextureSource& TextureSourceCmp = SourceTextures[TextureCmpIndex]->Source;
+		const FString TextureNameCmp = SourceTextures[TextureCmpIndex]->GetName();
+		const ETextureSourceFormat SourceFormatCmp = TextureSourceCmp.GetFormat();
 
-			const FTextureSource& TextureSourceCmp = SourceTextures[TextureCmpIndex]->Source;
-			const FString TextureName = SourceTextures[TextureIndex]->GetFName().ToString();
-			const FString TextureNameCmp = SourceTextures[TextureCmpIndex]->GetFName().ToString();
-			const ETextureSourceFormat SourceFormatCmp = TextureSourceCmp.GetFormat();
+		if (TextureSourceCmp.GetSizeX() != SizeX || TextureSourceCmp.GetSizeY() != SizeY)
+		{
+			UE_LOG(LogTexture, Warning, TEXT("TextureCubeArray creation failed. Textures %s and %s have different sizes."), *TextureName, *TextureNameCmp);
+			bError = true;
+		}
 
-			if (TextureSourceCmp.GetSizeX() != SizeX || TextureSourceCmp.GetSizeY() != SizeY)
-			{
-				UE_LOG(LogTexture, Warning, TEXT("TextureCubeArray creation failed. Textures %s and %s have different sizes."), *TextureName, *TextureNameCmp);
-				bError = true;
-			}
+		if (TextureSourceCmp.GetNumSlices() != NumSlices)
+		{
+			UE_LOG(LogTexture, Warning, TEXT("TextureCubeArray creation failed. Textures %s and %s have different number of faces (some are long/lat, some are not)."), *TextureName, *TextureNameCmp);
+			bError = true;
+		}
 
-			if (SourceFormatCmp != SourceFormat)
-			{
-				UE_LOG(LogTexture, Warning, TEXT("TextureCubeArray creation failed. Textures %s and %s have incompatible pixel formats."), *TextureName, *TextureNameCmp);
-				bError = true;
-			}
+		if (SourceFormatCmp != SourceFormat)
+		{
+			UE_LOG(LogTexture, Warning, TEXT("TextureCubeArray creation failed. Textures %s and %s have incompatible pixel formats."), *TextureName, *TextureNameCmp);
+			bError = true;
 		}
 	}
 
@@ -413,72 +419,71 @@ ENGINE_API bool UTextureCubeArray::UpdateSourceFromSourceTextures(bool bCreating
 		return false;
 	}
 
-	if (SourceTextures.Num() > 0)
-	{
-		FTextureSource& InitialSource = SourceTextures[0]->Source;
-		// Format and format size.
-		EPixelFormat PixelFormat = SourceTextures[0]->GetPixelFormat();
-		ETextureSourceFormat Format = InitialSource.GetFormat();
-		int32 FormatDataSize = InitialSource.GetBytesPerPixel();
-		// X,Y,Z size of the array.
-		int32 SizeX = InitialSource.GetSizeX();
-		int32 SizeY = InitialSource.GetSizeY();
-		uint32 ArraySize = SourceTextures.Num();
-		// Only copy the first mip from the source textures to array texture.
-		uint32 NumMips = 1;
-
-		// This should be false when texture is updated to avoid overriding user settings.
-		if (bCreatingNewTexture)
-		{
-			CompressionSettings = SourceTextures[0]->CompressionSettings;
-			MipGenSettings = TMGS_NoMipmaps;
-			PowerOfTwoMode = ETexturePowerOfTwoSetting::None;
-			LODGroup = SourceTextures[0]->LODGroup;
-			SRGB = SourceTextures[0]->SRGB;
-			NeverStream = true;
-		}
 
-		// Create the source texture for this UTexture.
-		const uint32 NumSlices = ArraySize * 6u;
-		Source.Init(SizeX, SizeY, NumSlices, NumMips, Format);
+	const FTextureSource& InitialSource = SourceTextures[0]->Source;
+	// Format and format size.
+	const EPixelFormat PixelFormat = SourceTextures[0]->GetPixelFormat();
+	const ETextureSourceFormat Format = InitialSource.GetFormat();
+	const int32 FormatDataSize = InitialSource.GetBytesPerPixel();
+	// X,Y,Z size of the array.
+	const int32 SizeX = InitialSource.GetSizeX();
+	const int32 SizeY = InitialSource.GetSizeY();
+	const int32 NumSlices = InitialSource.GetNumSlices();
+	const uint32 ArraySize = SourceTextures.Num();
+	// Only copy the first mip from the source textures to array texture.
+	const uint32 NumMips = 1;
+
+	// This should be false when texture is updated to avoid overriding user settings.
+	if (bCreatingNewTexture)
+	{
+		CompressionSettings = SourceTextures[0]->CompressionSettings;
+		MipGenSettings = TMGS_NoMipmaps;
+		PowerOfTwoMode = ETexturePowerOfTwoSetting::None;
+		LODGroup = SourceTextures[0]->LODGroup;
+		SRGB = SourceTextures[0]->SRGB;
+		NeverStream = true;
+	}
 
-		// We only copy the top level Mip map.
-		TArray<uint8*, TInlineAllocator<MAX_TEXTURE_MIP_COUNT> > DestMipData;
-		TArray<uint64, TInlineAllocator<MAX_TEXTURE_MIP_COUNT> > MipSizeBytes;
-		DestMipData.AddZeroed(NumMips);
-		MipSizeBytes.AddZeroed(NumMips);
+	// Create the source texture for this UTexture.
+	Source.Init(SizeX, SizeY, ArraySize * NumSlices, NumMips, Format);
+	Source.bLongLatCubemap = (NumSlices == 1);
 
-		for (uint32 MipIndex = 0; MipIndex < NumMips; ++MipIndex)
-		{
-			DestMipData[MipIndex] = Source.LockMip(MipIndex);
-			MipSizeBytes[MipIndex] = Source.CalcMipSize(MipIndex) / ArraySize;
-		}
+	// We only copy the top level Mip map.
+	TArray<uint8*, TInlineAllocator<MAX_TEXTURE_MIP_COUNT> > DestMipData;
+	TArray<uint64, TInlineAllocator<MAX_TEXTURE_MIP_COUNT> > MipSizeBytes;
+	DestMipData.AddZeroed(NumMips);
+	MipSizeBytes.AddZeroed(NumMips);
 
-		TArray64<uint8> RefCubeData;
-		for (int32 SourceTexIndex = 0; SourceTexIndex < SourceTextures.Num(); ++SourceTexIndex)
-		{
-			for (uint32 MipIndex = 0; MipIndex < NumMips; ++MipIndex)
-			{
-				const uint64 Size = MipSizeBytes[MipIndex];
-				const uint64 CheckSize = SourceTextures[SourceTexIndex]->Source.CalcMipSize(MipIndex);
-				check(Size == CheckSize);
-
-				RefCubeData.Reset();
-				SourceTextures[SourceTexIndex]->Source.GetMipData(RefCubeData, MipIndex);
-				void* Dst = DestMipData[MipIndex] + Size * SourceTexIndex;
-				FMemory::Memcpy(Dst, RefCubeData.GetData(), Size);
-			}
-		}
+	for (uint32 MipIndex = 0; MipIndex < NumMips; ++MipIndex)
+	{
+		DestMipData[MipIndex] = Source.LockMip(MipIndex);
+		MipSizeBytes[MipIndex] = Source.CalcMipSize(MipIndex) / ArraySize;
+	}
 
+	TArray64<uint8> RefCubeData;
+	for (int32 SourceTexIndex = 0; SourceTexIndex < SourceTextures.Num(); ++SourceTexIndex)
+	{
 		for (uint32 MipIndex = 0; MipIndex < NumMips; ++MipIndex)
 		{
-			Source.UnlockMip(MipIndex);
+			const uint64 Size = MipSizeBytes[MipIndex];
+			const uint64 CheckSize = SourceTextures[SourceTexIndex]->Source.CalcMipSize(MipIndex);
+			check(Size == CheckSize);
+
+			RefCubeData.Reset();
+			SourceTextures[SourceTexIndex]->Source.GetMipData(RefCubeData, MipIndex);
+			void* Dst = DestMipData[MipIndex] + Size * SourceTexIndex;
+			FMemory::Memcpy(Dst, RefCubeData.GetData(), Size);
 		}
+	}
 
-		SetLightingGuid();
-		UpdateResource();
+	for (uint32 MipIndex = 0; MipIndex < NumMips; ++MipIndex)
+	{
+		Source.UnlockMip(MipIndex);
 	}
 
+	SetLightingGuid();
+	UpdateResource();
+
 	return true;
 }
 
diff --git a/Engine/Source/Runtime/Engine/Private/TextureDerivedData.cpp b/Engine/Source/Runtime/Engine/Private/TextureDerivedData.cpp
--- a/Engine/Source/Runtime/Engine/Private/TextureDerivedData.cpp
+++ b/Engine/Source/Runtime/Engine/Private/TextureDerivedData.cpp
@@ -496,7 +496,7 @@ static void GetTextureBuildSettings(
 		OutBuildSettings.bCubemap = true;
 		OutBuildSettings.DiffuseConvolveMipLevel = GDiffuseConvolveMipLevel;
 		const UTextureCube* Cube = CastChecked<UTextureCube>(&Texture);
-		OutBuildSettings.bLongLatSource = (Cube->Source.GetNumSlices() == 1);
+		OutBuildSettings.bLongLatSource = (Cube->Source.GetNumSlices() == 1) || Cube->Source.IsLongLatCubemap();
 	}
 	else if (Texture.IsA(UTexture2DArray::StaticClass()))
 	{
@@ -507,6 +507,7 @@ static void GetTextureBuildSettings(
 	{
 		OutBuildSettings.bCubemap = true;
 		OutBuildSettings.bTextureArray = true;
+		OutBuildSettings.bLongLatSource = Texture.Source.IsLongLatCubemap();
 	}
 	else if (Texture.IsA(UVolumeTexture::StaticClass()))
 	{
