World Partition: make sure to call BlockTillLevelStreamingCompleted before settiing World.bMatchStarted to allow fully flushing level streaming requests in all situations.

diff --git a/Engine/Source/Runtime/Engine/Classes/Engine/World.h b/Engine/Source/Runtime/Engine/Classes/Engine/World.h
--- a/Engine/Source/Runtime/Engine/Classes/Engine/World.h
+++ b/Engine/Source/Runtime/Engine/Classes/Engine/World.h
@@ -2405,6 +2405,9 @@ public:
 	DECLARE_MULTICAST_DELEGATE(FOnWorldBeginPlay);
 	FOnWorldBeginPlay OnWorldBeginPlay;
 
+	DECLARE_MULTICAST_DELEGATE(FOnMatchStarting);
+	FOnMatchStarting OnWorldMatchStarting;
+
 	/** Returns true if gameplay has already started, false otherwise. */
 	bool HasBegunPlay() const;
 
diff --git a/Engine/Source/Runtime/Engine/Private/WorldPartition/WorldPartition.cpp b/Engine/Source/Runtime/Engine/Private/WorldPartition/WorldPartition.cpp
--- a/Engine/Source/Runtime/Engine/Private/WorldPartition/WorldPartition.cpp
+++ b/Engine/Source/Runtime/Engine/Private/WorldPartition/WorldPartition.cpp
@@ -640,7 +640,7 @@ void UWorldPartition::RegisterDelegates()
 
 	if (World->IsGameWorld())
 	{
-		World->OnWorldBeginPlay.AddUObject(this, &UWorldPartition::OnWorldBeginPlay);
+		World->OnWorldMatchStarting.AddUObject(this, &UWorldPartition::OnWorldMatchStarting);
 
 #if !UE_BUILD_SHIPPING
 		FCoreDelegates::OnGetOnScreenMessages.AddUObject(this, &UWorldPartition::GetOnScreenMessages);
@@ -669,7 +669,7 @@ void UWorldPartition::UnregisterDelegates()
 
 	if (World->IsGameWorld())
 	{
-		World->OnWorldBeginPlay.RemoveAll(this);
+		World->OnWorldMatchStarting.RemoveAll(this);
 
 #if !UE_BUILD_SHIPPING
 		FCoreDelegates::OnGetOnScreenMessages.RemoveAll(this);
@@ -687,7 +687,7 @@ void UWorldPartition::GetOnScreenMessages(FCoreDelegates::FSeverityMessageMap& O
 }
 #endif
 
-void UWorldPartition::OnWorldBeginPlay()
+void UWorldPartition::OnWorldMatchStarting()
 {
 	check(GetWorld()->IsGameWorld());
 	// Wait for any level streaming to complete
diff --git a/Engine/Source/Runtime/Engine/Private/WorldSettings.cpp b/Engine/Source/Runtime/Engine/Private/WorldSettings.cpp
--- a/Engine/Source/Runtime/Engine/Private/WorldSettings.cpp
+++ b/Engine/Source/Runtime/Engine/Private/WorldSettings.cpp
@@ -278,6 +278,7 @@ void AWorldSettings::NotifyBeginPlay()
 void AWorldSettings::NotifyMatchStarted()
 {
 	UWorld* World = GetWorld();
+	World->OnWorldMatchStarting.Broadcast();
 	World->bMatchStarted = true;
 }
 
diff --git a/Engine/Source/Runtime/Engine/Public/WorldPartition/WorldPartition.h b/Engine/Source/Runtime/Engine/Public/WorldPartition/WorldPartition.h
--- a/Engine/Source/Runtime/Engine/Public/WorldPartition/WorldPartition.h
+++ b/Engine/Source/Runtime/Engine/Public/WorldPartition/WorldPartition.h
@@ -225,7 +225,7 @@ private:
 #endif
 
 	bool IsMainWorldPartition() const;
-	void OnWorldBeginPlay();
+	void OnWorldMatchStarting();
 
 	// Delegates registration
 	void RegisterDelegates();
