Marked LodSyncComponent::ForcedLOD as interp so that it can be keyed in sequencer

diff --git a/Engine/Source/Runtime/Engine/Classes/Components/LODSyncComponent.h b/Engine/Source/Runtime/Engine/Classes/Components/LODSyncComponent.h
--- a/Engine/Source/Runtime/Engine/Classes/Components/LODSyncComponent.h
+++ b/Engine/Source/Runtime/Engine/Classes/Components/LODSyncComponent.h
@@ -70,7 +70,7 @@ class ENGINE_API ULODSyncComponent : public UActorComponent
 	int32 NumLODs = -1;
 
 	// if -1, it's automatically switching
-	UPROPERTY(EditAnywhere, BlueprintReadWrite, Category = LOD)
+	UPROPERTY(EditAnywhere, BlueprintReadWrite, interp, Category = LOD)
 	int32 ForcedLOD = -1;
 
 	// Minimum LOD to use when syncing components
