Rewrite of upload buffer allocation for streaming virtual texture.
Previous approach never freed from the high water mark (which can happen at start up loading locked mip for all loaded textures).

diff --git a/Engine/Source/Runtime/Engine/Private/VT/VirtualTextureChunkManager.cpp b/Engine/Source/Runtime/Engine/Private/VT/VirtualTextureChunkManager.cpp
--- a/Engine/Source/Runtime/Engine/Private/VT/VirtualTextureChunkManager.cpp
+++ b/Engine/Source/Runtime/Engine/Private/VT/VirtualTextureChunkManager.cpp
@@ -15,13 +15,6 @@
 #include "VirtualTextureChunkDDCCache.h"
 #endif
 
-static int32 NumTranscodeRequests = 32;
-static FAutoConsoleVariableRef CVarNumTranscodeRequests(
-	TEXT("r.VT.NumTranscodeRequests"),
-	NumTranscodeRequests,
-	TEXT("Number of transcode request that can be in flight. default 32\n"),
-	ECVF_Default);
-
 FVirtualTextureChunkStreamingManager::FVirtualTextureChunkStreamingManager()
 {
 #if WITH_EDITOR
@@ -98,7 +91,7 @@ FVTRequestPageResult FVirtualTextureChunkStreamingManager::RequestTile(FUploadin
 
 	// we limit the number of pending upload tiles in order to limit the memory required to store all the staging buffers
 	// Never throttle high priority requests
-	if (UploadCache.GetNumPendingTiles() >= (uint32)NumTranscodeRequests && Priority != EVTRequestPagePriority::High)
+	if (!UploadCache.IsInMemoryBudget() && Priority != EVTRequestPagePriority::High)
 	{
 		INC_DWORD_STAT(STAT_VTP_NumTranscodeDropped);
 		return EVTRequestPageStatus::Saturated;
diff --git a/Engine/Source/Runtime/Engine/Private/VT/VirtualTextureUploadCache.cpp b/Engine/Source/Runtime/Engine/Private/VT/VirtualTextureUploadCache.cpp
--- a/Engine/Source/Runtime/Engine/Private/VT/VirtualTextureUploadCache.cpp
+++ b/Engine/Source/Runtime/Engine/Private/VT/VirtualTextureUploadCache.cpp
@@ -1,12 +1,13 @@
 // Copyright Epic Games, Inc. All Rights Reserved.
 
 #include "VirtualTextureUploadCache.h"
-#include "VirtualTextureChunkManager.h"
-#include "RHI.h"
+
 #include "RenderGraphBuilder.h"
+#include "RHI.h"
+#include "VirtualTextureChunkManager.h"
 
-// allow uploading CPU buffer directly to GPU texture
-// this is slow under D3D11
+// Allow uploading CPU buffer directly to GPU texture
+// This is slow under D3D11
 // Should be pretty decent on D3D12X...UpdateTexture does make an extra copy of the data, but Lock/Unlock texture also buffers an extra copy of texture on this platform
 // Might also be worth enabling this path on PC D3D12, need to measure
 #if !defined(ALLOW_UPDATE_TEXTURE)
@@ -17,50 +18,219 @@ DECLARE_MEMORY_STAT_POOL(TEXT("Total GPU Upload Memory"), STAT_TotalGPUUploadSiz
 DECLARE_MEMORY_STAT(TEXT("Total CPU Upload Memory"), STAT_TotalCPUUploadSize, STATGROUP_VirtualTextureMemory);
 DECLARE_DWORD_ACCUMULATOR_STAT(TEXT("Num Upload Entries"), STAT_NumUploadEntries, STATGROUP_VirtualTextureMemory);
 
-FVirtualTextureUploadCache::FStagingBuffer::FStagingBuffer()
+static TAutoConsoleVariable<int32> CVarVTUploadMemoryPageSize(
+	TEXT("r.VT.UploadMemoryPageSize"),
+	4,
+	TEXT("Size in MB for a single page of virtual texture upload memory."),
+	ECVF_RenderThreadSafe
+);
+
+static TAutoConsoleVariable<int32> CVarVTMaxUploadMemory(
+	TEXT("r.VT.MaxUploadMemory"),
+	64,
+	TEXT("Maximum amount of upload memory to allocate in MB before throttling virtual texture streaming requests.\n")
+	TEXT("We never throttle high priority requests so allocation can peak above this value."),
+	ECVF_RenderThreadSafe
+);
+
+static TAutoConsoleVariable<int32> CVarMaxUploadRequests(
+	TEXT("r.VT.MaxUploadRequests"),
+	2000,
+	TEXT("Maximum number of virtual texture tile upload requests that can be in flight."),
+	ECVF_RenderThreadSafe);
+
+
+uint32 FVTUploadTileAllocator::Allocate(EPixelFormat InFormat, uint32 InTileSize)
 {
-}
+	// Find matching FormatBuffer.
+	const FPixelFormatInfo& FormatInfo = GPixelFormats[InFormat];
+	const uint32 TileWidthInBlocks = FMath::DivideAndRoundUp(InTileSize, (uint32)FormatInfo.BlockSizeX);
+	const uint32 TileHeightInBlocks = FMath::DivideAndRoundUp(InTileSize, (uint32)FormatInfo.BlockSizeY);
+
+	FSharedFormatDesc Desc;
+	Desc.BlockBytes = FormatInfo.BlockBytes;
+	Desc.Stride = TileWidthInBlocks * FormatInfo.BlockBytes;
+	Desc.MemorySize = TileWidthInBlocks * TileHeightInBlocks * FormatInfo.BlockBytes;
+
+	int32 FormatIndex = 0;
+	for (; FormatIndex < FormatDescs.Num(); ++FormatIndex)
+	{
+		if (Desc.BlockBytes == FormatDescs[FormatIndex].BlockBytes && Desc.Stride == FormatDescs[FormatIndex].Stride && Desc.MemorySize == FormatDescs[FormatIndex].MemorySize)
+		{
+			break;
+		}
+	}
 
-FVirtualTextureUploadCache::FStagingBuffer::~FStagingBuffer()
-{
-	if (RHIBuffer)
+	if (FormatIndex == FormatDescs.Num())
 	{
-		// Unmap and release the staging buffer if present
-		// In this case 'Memory' was the mapped pointer, so release it
-		Memory = nullptr;
-		RHIUnlockBuffer(RHIBuffer);
-		RHIBuffer.SafeRelease();
+		// Add newly found format.
+		FormatDescs.Add(Desc);
+		FormatBuffers.AddDefaulted();
 	}
 
-	if (Memory)
+	FSharedFormatBuffers& FormatBuffer = FormatBuffers[FormatIndex];
+
+	// Find available staging buffer.
+	int32 StagingBufferIndex = 0;
+	for (; StagingBufferIndex < FormatBuffer.StagingBuffers.Num(); ++StagingBufferIndex)
 	{
-		// If we still have 'Memory' pointer, it came from heap, release it now
-		FMemory::Free(Memory);
-		Memory = nullptr;
+		if (FormatBuffer.StagingBuffers[StagingBufferIndex].Memory == nullptr)
+		{
+			// Staging buffer was released so we can re-init and use it.
+			break;
+		}
+
+		if (FormatBuffer.StagingBuffers[StagingBufferIndex].TileFreeList.Num())
+		{
+			// Staging buffer has free tiles available.
+			break;
+		}
+	}
+
+	if (StagingBufferIndex == FormatBuffer.StagingBuffers.Num())
+	{
+		// Current staging buffers are full. Need to allocate a new staging buffer.
+		FormatBuffer.StagingBuffers.AddDefaulted();
 	}
+
+	FStagingBuffer& StagingBuffer = FormatBuffer.StagingBuffers[StagingBufferIndex];
+	if (StagingBuffer.Memory == nullptr)
+	{
+		// Staging buffer needs underlying buffer allocating.
+		StagingBuffer.Init(Desc.BlockBytes, Desc.MemorySize);
+		NumAllocatedBytes += StagingBuffer.NumTiles * StagingBuffer.TileSize;
+	}
+
+	// Pop a free tile and return handle.
+	int32 TileIndex = StagingBuffer.TileFreeList.Pop(false);
+
+	FHandle Handle;
+	Handle.FormatIndex = FormatIndex;
+	Handle.StagingBufferIndex = StagingBufferIndex;
+	Handle.TileIndex = TileIndex;
+	return Handle.PackedValue;
+}
+
+void FVTUploadTileAllocator::Free(uint32 InHandle)
+{
+	FHandle Handle;
+	Handle.PackedValue = InHandle;
+
+	// Push tile back onto free list.
+	FStagingBuffer& StagingBuffer = FormatBuffers[Handle.FormatIndex].StagingBuffers[Handle.StagingBufferIndex];
+	StagingBuffer.TileFreeList.Push(Handle.TileIndex);
+
+	if (StagingBuffer.NumTiles == StagingBuffer.TileFreeList.Num())
+	{
+		// All tiles are free, so release the underlying memory.
+		check(NumAllocatedBytes >= StagingBuffer.NumTiles * StagingBuffer.TileSize);
+		NumAllocatedBytes -= StagingBuffer.NumTiles * StagingBuffer.TileSize;
+
+		StagingBuffer.Release();
+	}
+}
+
+FVTUploadTileBuffer FVTUploadTileAllocator::GetBufferFromHandle(uint32 InHandle) const
+{
+	FHandle Handle;
+	Handle.PackedValue = InHandle;
+
+	FSharedFormatDesc const& FormatDesc = FormatDescs[Handle.FormatIndex];
+	FStagingBuffer const& StagingBuffer = FormatBuffers[Handle.FormatIndex].StagingBuffers[Handle.StagingBufferIndex];
+
+	FVTUploadTileBuffer Buffer;
+	Buffer.Memory = (uint8*)StagingBuffer.Memory + Handle.TileIndex * Align(StagingBuffer.TileSize, 128u);
+	Buffer.MemorySize = StagingBuffer.TileSize;
+	Buffer.Stride = FormatDesc.Stride;
+	return Buffer;
 }
 
-FVirtualTextureUploadCache::FTileEntry::FTileEntry()
+FVTUploadTileBufferExt FVTUploadTileAllocator::GetBufferFromHandleExt(uint32 InHandle) const
 {
+	FHandle Handle;
+	Handle.PackedValue = InHandle;
+
+	FSharedFormatDesc const& FormatDesc = FormatDescs[Handle.FormatIndex];
+	FStagingBuffer const& StagingBuffer = FormatBuffers[Handle.FormatIndex].StagingBuffers[Handle.StagingBufferIndex];
+
+	FVTUploadTileBufferExt Buffer;
+	Buffer.RHIBuffer = StagingBuffer.RHIBuffer;
+	Buffer.BufferMemory = StagingBuffer.Memory;
+	Buffer.BufferOffset = Handle.TileIndex * Align(StagingBuffer.TileSize, 128u);
+	Buffer.Stride = FormatDesc.Stride;
+	return Buffer;
 }
 
-FVirtualTextureUploadCache::FTileEntry::~FTileEntry()
+FVTUploadTileAllocator::FStagingBuffer::~FStagingBuffer()
 {
+	Release();
 }
 
-FVirtualTextureUploadCache::FVirtualTextureUploadCache()
+void FVTUploadTileAllocator::FStagingBuffer::Init(uint32 InBufferStrideBytes, uint32 InTileSizeBytes)
 {
-	StagingBuffers.AddDefaulted(); // empty buffer
-	Tiles.AddDefaulted(LIST_COUNT);
-	for (int i = 0; i < LIST_COUNT; ++i)
+	TileSize = InTileSizeBytes;
+	const uint32 BufferSize = CVarVTUploadMemoryPageSize.GetValueOnRenderThread() * 1024 * 1024;
+	const uint32 AlignedTileSizeBytes = Align(InTileSizeBytes, 128u);
+
+	NumTiles = FMath::DivideAndRoundUp(BufferSize, AlignedTileSizeBytes);
+	const uint32 AlignedBufferSize = NumTiles * AlignedTileSizeBytes;
+
+	check(TileFreeList.Num() == 0);
+	TileFreeList.AddUninitialized(NumTiles);
+	for (uint32 Index = 0; Index < NumTiles; ++Index)
+	{
+		TileFreeList[Index] = NumTiles - Index - 1;
+	}
+
+	if (GRHISupportsDirectGPUMemoryLock)
 	{
-		FTileEntry& Entry = Tiles[i];
-		Entry.NextIndex = Entry.PrevIndex = i;
+		// Allocate staging buffer directly in GPU memory.
+		checkSlow(IsInRenderingThread());
+		FRHICommandListImmediate& RHICmdList = FRHICommandListExecutor::GetImmediateCommandList();
+
+		FRHIResourceCreateInfo CreateInfo(TEXT("StagingBuffer"));
+		RHIBuffer = RHICreateStructuredBuffer(InBufferStrideBytes, AlignedBufferSize, BUF_ShaderResource | BUF_Static | BUF_KeepCPUAccessible, CreateInfo);
+
+		// Here we bypass 'normal' RHI operations in order to get a persistent pointer to GPU memory, on supported platforms
+		// This should be encapsulated into a proper RHI method at some point
+		Memory = RHICmdList.LockBuffer(RHIBuffer, 0u, AlignedBufferSize, RLM_WriteOnly_NoOverwrite);
+
+		INC_MEMORY_STAT_BY(STAT_TotalGPUUploadSize, AlignedBufferSize);
+	}
+	else
+	{
+		// Allocate staging buffer in CPU memory.
+		Memory = FMemory::Malloc(AlignedBufferSize, 128u);
+
+		INC_MEMORY_STAT_BY(STAT_TotalCPUUploadSize, AlignedBufferSize);
 	}
 }
 
-FVirtualTextureUploadCache::~FVirtualTextureUploadCache()
+void FVTUploadTileAllocator::FStagingBuffer::Release()
 {
+	if (RHIBuffer.IsValid())
+	{
+		// Unmap and release the GPU buffer if present.
+		RHIUnlockBuffer(RHIBuffer);
+		RHIBuffer.SafeRelease();
+		// In this case 'Memory' was the mapped pointer, so release it.
+		Memory = nullptr;
+
+		DEC_MEMORY_STAT_BY(STAT_TotalGPUUploadSize, TileSize * NumTiles);
+	}
+
+	if (Memory != nullptr)
+	{
+		// If we still have 'Memory' pointer, it is CPU, release it now.
+		FMemory::Free(Memory);
+		Memory = nullptr;
+
+		DEC_MEMORY_STAT_BY(STAT_TotalCPUUploadSize, TileSize * NumTiles);
+	}
+
+	TileSize = 0u;
+	NumTiles = 0u;
+	TileFreeList.Reset();
 }
 
 int32 FVirtualTextureUploadCache::GetOrCreatePoolIndex(EPixelFormat InFormat, uint32 InTileSize)
@@ -78,40 +248,36 @@ int32 FVirtualTextureUploadCache::GetOrCreatePoolIndex(EPixelFormat InFormat, ui
 	FPoolEntry& Entry = Pools[PoolIndex];
 	Entry.Format = InFormat;
 	Entry.TileSize = InTileSize;
-	Entry.FreeTileListHead = CreateTileEntry(PoolIndex);
-	Entry.SubmitTileListHead = CreateTileEntry(PoolIndex);
 
 	return PoolIndex;
 }
 
 void FVirtualTextureUploadCache::Finalize(FRDGBuilder& GraphBuilder)
 {
-	SCOPE_CYCLE_COUNTER(STAT_VTP_FlushUpload)
-
-	check(IsInRenderingThread());
-
 	FRHICommandListImmediate& RHICmdList = GraphBuilder.RHICmdList;
 	SCOPED_DRAW_EVENT(RHICmdList, FVirtualTextureUploadCache_Finalize);
+	SCOPE_CYCLE_COUNTER(STAT_VTP_FlushUpload)
 
-	// Multi-GPU support : May be ineffecient for AFR.
-	SCOPED_GPU_MASK(RHICmdList, FRHIGPUMask::All());
+		// Multi-GPU support : May be inefficient for AFR.
+		SCOPED_GPU_MASK(RHICmdList, FRHIGPUMask::All());
 
 	for (int PoolIndex = 0; PoolIndex < Pools.Num(); ++PoolIndex)
 	{
 		FPoolEntry& PoolEntry = Pools[PoolIndex];
-		const uint32 BatchCount = PoolEntry.BatchCount;
-		if (BatchCount == 0u)
+		const uint32 BatchCount = PoolEntry.PendingSubmit.Num();
+		if (BatchCount == 0)
 		{
 			continue;
 		}
 
 		const FPixelFormatInfo& FormatInfo = GPixelFormats[PoolEntry.Format];
 		const uint32 TileSize = PoolEntry.TileSize;
+		const uint32 BlockBytes = FormatInfo.BlockBytes;
 		const uint32 TileWidthInBlocks = FMath::DivideAndRoundUp(TileSize, (uint32)FormatInfo.BlockSizeX);
 		const uint32 TileHeightInBlocks = FMath::DivideAndRoundUp(TileSize, (uint32)FormatInfo.BlockSizeY);
 
 		const uint32 TextureIndex = PoolEntry.BatchTextureIndex;
-		PoolEntry.BatchTextureIndex = (PoolEntry.BatchTextureIndex + 1u) % NUM_STAGING_TEXTURES;
+		PoolEntry.BatchTextureIndex = (PoolEntry.BatchTextureIndex + 1u) % FPoolEntry::NUM_STAGING_TEXTURES;
 		FStagingTexture& StagingTexture = PoolEntry.StagingTexture[TextureIndex];
 
 		// On some platforms the staging texture create/lock behavior will depend on whether we are running with RHI threading
@@ -147,40 +313,37 @@ void FVirtualTextureUploadCache::Finalize(FRDGBuilder& GraphBuilder)
 		uint32 BatchStride = 0u;
 		void* BatchMemory = RHICmdList.LockTexture2D(StagingTexture.RHITexture, 0, RLM_WriteOnly, BatchStride, false, false);
 
-		// copy all tiles to the staging texture
-		const int32 SubmitListHead = PoolEntry.SubmitTileListHead;
-		int32 Index = Tiles[SubmitListHead].NextIndex;
-		while (Index != SubmitListHead)
+		// Copy all tiles to the staging texture
+		for (uint32 Index = 0u; Index < BatchCount; ++Index)
 		{
-			const FTileEntry& Entry = Tiles[Index];
-			const FStagingBuffer& StagingBuffer = StagingBuffers[Entry.BufferIndex];
-			const int32 NextIndex = Entry.NextIndex;
-			const uint32_t SrcTileX = Entry.SubmitBatchIndex % StagingTexture.WidthInTiles;
-			const uint32_t SrcTileY = Entry.SubmitBatchIndex / StagingTexture.WidthInTiles;
+			const FTileEntry& Entry = PoolEntry.PendingSubmit[Index];
+			const uint32_t SrcTileX = Index % StagingTexture.WidthInTiles;
+			const uint32_t SrcTileY = Index / StagingTexture.WidthInTiles;
+
+			const FVTUploadTileBufferExt UploadBuffer = TileAllocator.GetBufferFromHandleExt(Entry.TileHandle);
 
-			uint8* BatchDst = (uint8*)BatchMemory + TileHeightInBlocks * SrcTileY * BatchStride + TileWidthInBlocks * SrcTileX * FormatInfo.BlockBytes;
+			uint8* BatchDst = (uint8*)BatchMemory + TileHeightInBlocks * SrcTileY * BatchStride + TileWidthInBlocks * SrcTileX * BlockBytes;
 			for (uint32 y = 0u; y < TileHeightInBlocks; ++y)
 			{
 				FMemory::Memcpy(
 					BatchDst + y * BatchStride,
-					(uint8*)StagingBuffer.Memory + Entry.BufferOffset + y * Entry.Stride,
-					TileWidthInBlocks * FormatInfo.BlockBytes);
+					(uint8*)UploadBuffer.BufferMemory + UploadBuffer.BufferOffset + y * UploadBuffer.Stride,
+					TileWidthInBlocks * BlockBytes);
 			}
 
-			Index = NextIndex;
+			// Can release upload buffer.
+			TileAllocator.Free(Entry.TileHandle);
 		}
 
 		RHICmdList.UnlockTexture2D(StagingTexture.RHITexture, 0u, false, false);
 		RHICmdList.Transition(FRHITransitionInfo(StagingTexture.RHITexture, ERHIAccess::SRVMask, ERHIAccess::CopySrc));
 
-		// upload each tile from staging texture to physical texture
-		Index = Tiles[SubmitListHead].NextIndex;
-		while (Index != SubmitListHead)
+		// Upload each tile from staging texture to physical texture
+		for (uint32 Index = 0u; Index < BatchCount; ++Index)
 		{
-			FTileEntry& Entry = Tiles[Index];
-			const int32 NextIndex = Entry.NextIndex;
-			const uint32_t SrcTileX = Entry.SubmitBatchIndex % StagingTexture.WidthInTiles;
-			const uint32_t SrcTileY = Entry.SubmitBatchIndex / StagingTexture.WidthInTiles;
+			const FTileEntry& Entry = PoolEntry.PendingSubmit[Index];
+			const uint32_t SrcTileX = Index % StagingTexture.WidthInTiles;
+			const uint32_t SrcTileY = Index / StagingTexture.WidthInTiles;
 
 			const uint32 SkipBorderSize = Entry.SubmitSkipBorderSize;
 			const uint32 SubmitTileSize = TileSize - SkipBorderSize * 2;
@@ -198,21 +361,11 @@ void FVirtualTextureUploadCache::Finalize(FRDGBuilder& GraphBuilder)
 			CopyInfo.SourcePosition = SourceBoxStart;
 			CopyInfo.DestPosition = DestinationBoxStart;
 			RHICmdList.CopyTexture(StagingTexture.RHITexture, Entry.RHISubmitTexture, CopyInfo);
-
-			Entry.RHISubmitTexture = nullptr;
-			Entry.SubmitBatchIndex = 0u;
-			Entry.SubmitDestX = 0;
-			Entry.SubmitDestY = 0;
-			Entry.SubmitSkipBorderSize = 0;
-
-			RemoveFromList(Index);
-			AddToList(PoolEntry.FreeTileListHead, Index);
-			Index = NextIndex;
 		}
 
 		RHICmdList.Transition(FRHITransitionInfo(StagingTexture.RHITexture, ERHIAccess::CopySrc, ERHIAccess::SRVMask));
 
-		PoolEntry.BatchCount = 0u;
+		PoolEntry.PendingSubmit.Reset();
 	}
 
 	// Transition all updates textures back to SRV
@@ -229,103 +382,48 @@ void FVirtualTextureUploadCache::Finalize(FRDGBuilder& GraphBuilder)
 
 void FVirtualTextureUploadCache::ReleaseRHI()
 {
+	// Complete/Cancel all work will release allocated staging buffers.
+	UpdateFreeList(true);
+	for (TSparseArray<FTileEntry>::TIterator It(PendingUpload); It; ++It)
+	{
+		CancelTile(FVTUploadTileHandle(It.GetIndex()));
+	}
+	// Release staging textures.
 	Pools.Empty();
-	Tiles.Empty();
 }
 
 FVTUploadTileHandle FVirtualTextureUploadCache::PrepareTileForUpload(FVTUploadTileBuffer& OutBuffer, EPixelFormat InFormat, uint32 InTileSize)
 {
 	SCOPE_CYCLE_COUNTER(STAT_VTP_StageTile)
-
 	checkSlow(IsInRenderingThread());
-	FRHICommandListImmediate& RHICmdList = FRHICommandListExecutor::GetImmediateCommandList(); // only valid on the render thread
+
+	uint32 TileHandle = TileAllocator.Allocate(InFormat, InTileSize);
+	OutBuffer = TileAllocator.GetBufferFromHandle(TileHandle);
 
 	const int32 PoolIndex = GetOrCreatePoolIndex(InFormat, InTileSize);
 	const FPoolEntry& PoolEntry = Pools[PoolIndex];
-
-	int32 Index = Tiles[PoolEntry.FreeTileListHead].NextIndex;
-	if (Index == PoolEntry.FreeTileListHead)
-	{
-		const FPixelFormatInfo& FormatInfo = GPixelFormats[InFormat];
-		const uint32 TileWidthInBlocks = FMath::DivideAndRoundUp(InTileSize, (uint32)FormatInfo.BlockSizeX);
-		const uint32 TileHeightInBlocks = FMath::DivideAndRoundUp(InTileSize, (uint32)FormatInfo.BlockSizeY);
-		const uint32 Stride = TileWidthInBlocks * FormatInfo.BlockBytes;
-		const uint32 MemorySize = Stride * TileHeightInBlocks;
-
-		int32 StagingBufferIndex = StagingBuffers.Num() - 1;
-		uint32 StagingBufferOffset = Align(StagingBuffers[StagingBufferIndex].CurrentOffset, 128u);
-		if (StagingBufferOffset + MemorySize > StagingBuffers[StagingBufferIndex].Size)
-		{
-			StagingBufferIndex = StagingBuffers.AddDefaulted();
-			StagingBufferOffset = 0u;
-			FStagingBuffer& StagingBuffer = StagingBuffers[StagingBufferIndex];
-			StagingBuffer.Size = MemorySize * 1024u;
-
-			// We support several different methods for staging tile data to GPU textures
-			// On some platforms, CPU can write linear texture data to persist mapped buffer, then this can be uploaded directly to GPU...this is fastest method
-			// Otherwise, CPU writes texture data to temp buffer, then this is copied to GPU via a batched staging texture...this involves more copying, but is best method under default D3D11
-			// Can potentially write each tile to a separate staging texture, but this has too much lock/unlock overhead
-			
-			
-			// Stage to persist mapped GPU buffer then GPU copy into texture
-			// this is fast where supported
-			if (GRHISupportsDirectGPUMemoryLock)
-			{
-				FRHIResourceCreateInfo CreateInfo(TEXT("StagingBuffer"));
-				StagingBuffer.RHIBuffer = RHICreateStructuredBuffer(FormatInfo.BlockBytes, StagingBuffer.Size, BUF_ShaderResource | BUF_Static | BUF_KeepCPUAccessible, CreateInfo);
-
-				// Here we bypass 'normal' RHI operations in order to get a persistent pointer to GPU memory, on supported platforms
-				// This should be encapsulated into a proper RHI method at some point
-				StagingBuffer.Memory = RHICmdList.LockBuffer(StagingBuffer.RHIBuffer, 0u, StagingBuffer.Size, RLM_WriteOnly_NoOverwrite);
-
-				INC_MEMORY_STAT_BY(STAT_TotalGPUUploadSize, StagingBuffer.Size);
-			}
-			else
-			{
-				StagingBuffer.Memory = FMemory::Malloc(StagingBuffer.Size);
-				INC_MEMORY_STAT_BY(STAT_TotalCPUUploadSize, StagingBuffer.Size);
-			}
-		}
-
-		StagingBuffers[StagingBufferIndex].CurrentOffset = StagingBufferOffset + MemorySize;
-
-		Index = CreateTileEntry(PoolIndex);
-		FTileEntry& NewEntry = Tiles[Index];
-		NewEntry.BufferIndex = StagingBufferIndex;
-		NewEntry.BufferOffset = StagingBufferOffset;
-		NewEntry.Stride = Stride;
-		NewEntry.MemorySize = MemorySize;
-		INC_DWORD_STAT(STAT_NumUploadEntries);
-	}
-	else
-	{
-		RemoveFromList(Index);
-	}
-
-	FTileEntry& Entry = Tiles[Index];
-	++NumPendingTiles;
 	
-	OutBuffer.Memory = (uint8*)StagingBuffers[Entry.BufferIndex].Memory + Entry.BufferOffset;
-	OutBuffer.MemorySize = Entry.MemorySize;
-	OutBuffer.Stride = Entry.Stride;
+	FTileEntry Tile;
+	Tile.PoolIndex = PoolIndex;
+	Tile.TileHandle = TileHandle;
+
+	uint32 Index = PendingUpload.Emplace(Tile);
 	return FVTUploadTileHandle(Index);
 }
 
 void FVirtualTextureUploadCache::SubmitTile(FRHICommandListImmediate& RHICmdList, const FVTUploadTileHandle& InHandle, FRHITexture2D* InDestTexture, int InDestX, int InDestY, int InSkipBorderSize)
 {
 	checkSlow(IsInRenderingThread());
-	check(NumPendingTiles > 0u);
-	--NumPendingTiles;
 
-	const int32 Index = InHandle.Index;
-	FTileEntry& Entry = Tiles[Index];
+	check(PendingUpload.IsValidIndex(InHandle.Index));
+	FTileEntry& Entry = PendingUpload[InHandle.Index];
 	Entry.FrameSubmitted = GFrameNumberRenderThread;
 
 	FPoolEntry& PoolEntry = Pools[Entry.PoolIndex];
 	const uint32 TileSize = PoolEntry.TileSize - InSkipBorderSize * 2;
 
-	const FStagingBuffer& StagingBuffer = StagingBuffers[Entry.BufferIndex];
-	if (StagingBuffer.RHIBuffer)
+	const FVTUploadTileBufferExt UploadBuffer = TileAllocator.GetBufferFromHandleExt(Entry.TileHandle);
+	if (UploadBuffer.RHIBuffer != nullptr)
 	{
 		if (!UpdatedTextures.Contains(InDestTexture))
 		{
@@ -335,11 +433,11 @@ void FVirtualTextureUploadCache::SubmitTile(FRHICommandListImmediate& RHICmdList
 
 		check(GRHISupportsUpdateFromBufferTexture);
 		const FUpdateTextureRegion2D UpdateRegion(InDestX * TileSize, InDestY * TileSize, InSkipBorderSize, InSkipBorderSize, TileSize, TileSize);
-		RHICmdList.UpdateFromBufferTexture2D(InDestTexture, 0u, UpdateRegion, Entry.Stride, StagingBuffer.RHIBuffer, Entry.BufferOffset);
+		RHICmdList.UpdateFromBufferTexture2D(InDestTexture, 0u, UpdateRegion, UploadBuffer.Stride, UploadBuffer.RHIBuffer, UploadBuffer.BufferOffset);
 
-		// move to pending list, so we won't re-use this buffer until the GPU has finished the copy
-		// (we're using persist mapped buffer here, so this is the only synchronization method in place...without this delay we'd get corrupt textures)
-		AddToList(LIST_SUBMITTED, Index);
+		// Move to pending list, so we won't re-use this buffer until the GPU has finished the copy.
+		// We're using persist mapped buffer here, so this is the only synchronization method in place...without this delay we'd get corrupt textures.
+		PendingRelease.Emplace(Entry);
 	}
 	else if(ALLOW_UPDATE_TEXTURE)
 	{
@@ -350,59 +448,59 @@ void FVirtualTextureUploadCache::SubmitTile(FRHICommandListImmediate& RHICmdList
 		}
 
 		const FUpdateTextureRegion2D UpdateRegion(InDestX * TileSize, InDestY * TileSize, InSkipBorderSize, InSkipBorderSize, TileSize, TileSize);
-		RHICmdList.UpdateTexture2D(InDestTexture, 0u, UpdateRegion, Entry.Stride, (uint8*)StagingBuffer.Memory + Entry.BufferOffset);
+		RHICmdList.UpdateTexture2D(InDestTexture, 0u, UpdateRegion, UploadBuffer.Stride, (uint8*)UploadBuffer.BufferMemory + UploadBuffer.BufferOffset);
 
 		// UpdateTexture2D makes internal copy of data, no need to wait before re-using tile
-		AddToList(PoolEntry.FreeTileListHead, Index);
+		TileAllocator.Free(Entry.TileHandle);
 	}
 	else
 	{
+		// Move to list of batched updates for the current pool
 		Entry.RHISubmitTexture = InDestTexture;
 		Entry.SubmitDestX = InDestX;
 		Entry.SubmitDestY = InDestY;
 		Entry.SubmitSkipBorderSize = InSkipBorderSize;
-		Entry.SubmitBatchIndex = PoolEntry.BatchCount++;
-
-		// move to list of batched updates for the current pool
-		AddToList(PoolEntry.SubmitTileListHead, Index);
+		PoolEntry.PendingSubmit.Emplace(Entry);
 	}
+
+	// Remove from pending uploads.
+	PendingUpload.RemoveAt(InHandle.Index);
 }
 
 void FVirtualTextureUploadCache::CancelTile(const FVTUploadTileHandle& InHandle)
 {
 	checkSlow(IsInRenderingThread());
-	check(NumPendingTiles > 0u);
-	--NumPendingTiles;
 
-	const int32 Index = InHandle.Index;
-	FTileEntry& Entry = Tiles[Index];
-	FPoolEntry& PoolEntry = Pools[Entry.PoolIndex];
-
-	AddToList(PoolEntry.FreeTileListHead, Index);
+	check(PendingUpload.IsValidIndex(InHandle.Index));
+	FTileEntry& Entry = PendingUpload[InHandle.Index];
+	TileAllocator.Free(Entry.TileHandle);
+	PendingUpload.RemoveAt(InHandle.Index);
 }
 
-void FVirtualTextureUploadCache::UpdateFreeList()
+void FVirtualTextureUploadCache::UpdateFreeList(bool bForceFreeAll)
 {
-	check(IsInRenderingThread());
+	checkSlow(IsInRenderingThread());
+
 	const uint32 CurrentFrame = GFrameNumberRenderThread;
 
-	int32 Index = Tiles[LIST_SUBMITTED].NextIndex;
-	while (Index != LIST_SUBMITTED)
+	// Iterate tiles pending release and free them if they are old enough.
+	for (TSparseArray<FTileEntry>::TIterator It(PendingRelease); It; ++It)
 	{
-		FTileEntry& Entry = Tiles[Index];
-		const int32 NextIndex = Entry.NextIndex;
-
-		check(CurrentFrame >= Entry.FrameSubmitted);
-		const uint32 FramesSinceSubmitted = CurrentFrame - Entry.FrameSubmitted;
+		check(CurrentFrame >= It->FrameSubmitted);
+		const uint32 FramesSinceSubmitted = CurrentFrame - It->FrameSubmitted;
 		if (FramesSinceSubmitted < 2u)
 		{
 			break;
 		}
 
-		const FPoolEntry& PoolEntry = Pools[Entry.PoolIndex];
-		RemoveFromList(Index);
-		AddToList(PoolEntry.FreeTileListHead, Index);
-
-		Index = NextIndex;
+		TileAllocator.Free(It->TileHandle);
+		It.RemoveCurrent();
 	}
 }
+
+uint32 FVirtualTextureUploadCache::IsInMemoryBudget() const
+{
+	return 
+		PendingUpload.Num() + PendingRelease.Num() <= CVarMaxUploadRequests.GetValueOnRenderThread() &&
+		TileAllocator.TotalAllocatedBytes() <= (uint32)CVarVTMaxUploadMemory.GetValueOnRenderThread() * 1024u * 1024u;
+}
diff --git a/Engine/Source/Runtime/Engine/Private/VT/VirtualTextureUploadCache.h b/Engine/Source/Runtime/Engine/Private/VT/VirtualTextureUploadCache.h
--- a/Engine/Source/Runtime/Engine/Private/VT/VirtualTextureUploadCache.h
+++ b/Engine/Source/Runtime/Engine/Private/VT/VirtualTextureUploadCache.h
@@ -3,7 +3,6 @@
 #pragma once
 
 #include "CoreMinimal.h"
-#include "HAL/CriticalSection.h"
 #include "RendererInterface.h"
 #include "RenderResource.h"
 #include "PixelFormat.h"
@@ -11,6 +10,22 @@
 class FRHITexture2D;
 class FRHIBuffer;
 
+/** Opaque handle for referencing an upload tile returned by FVirtualTextureUploadCache::PrepareTileForUpload(). */
+struct FVTUploadTileHandle
+{
+	explicit FVTUploadTileHandle(uint32 InIndex = ~0u) 
+		: Index(InIndex) 
+	{}
+
+	inline bool IsValid() const { return Index != ~0u; }
+
+	uint32 Index;
+};
+
+/** 
+ * Memory buffer for uploading virtual texture data. 
+ * This is a simple view of the buffer memory for a single tile intended for use by the streaming systems.
+ */
 struct FVTUploadTileBuffer
 {
 	void* Memory = nullptr;
@@ -18,134 +33,168 @@ struct FVTUploadTileBuffer
 	uint32 Stride = 0u;
 };
 
-//
-struct FVTUploadTileHandle
+/** 
+ * Extended definition of the memory buffer for uploading virtual texture 
+ * This extended view of the buffer is used internally by FVirtualTextureUploadCache.
+ */
+struct FVTUploadTileBufferExt
 {
-	explicit FVTUploadTileHandle(uint32 InIndex = 0u) : Index(InIndex) {}
-
-	inline bool IsValid() const { return Index != 0u; }
-
-	uint32 Index;
+	TRefCountPtr<FRHIBuffer> RHIBuffer;
+	void* BufferMemory = nullptr;
+	uint32 BufferOffset = 0u;
+	uint32 Stride = 0u;
 };
 
-class FVirtualTextureUploadCache : public IVirtualTextureFinalizer, public FRenderResource
+/** 
+ * Handles allocation of staging buffer memory. 
+ */
+class FVTUploadTileAllocator
 {
 public:
-	FVirtualTextureUploadCache();
-	virtual ~FVirtualTextureUploadCache();
-
-	// IVirtualTextureFinalizer
-	virtual void Finalize(FRDGBuilder& GraphBuilder) override;
-
-	// FRenderResource
-	virtual void ReleaseRHI() override;
-
-	uint32 GetNumPendingTiles() const { return NumPendingTiles; }
-
-	FVTUploadTileHandle PrepareTileForUpload(FVTUploadTileBuffer& OutBuffer, EPixelFormat InFormat, uint32 InTileSize);
-	void SubmitTile(FRHICommandListImmediate& RHICmdList, const FVTUploadTileHandle& InHandle, FRHITexture2D* InDestTexture, int InDestX, int InDestY, int InSkipBorderSize);
-	void CancelTile(const FVTUploadTileHandle& InHandle);
-
-	void UpdateFreeList();
+	/** Allocate a tile. Sometimes does an allocation of the backing CPU/GPU block of memory. */
+	uint32 Allocate(EPixelFormat InFormat, uint32 InTileSize);
+	/** Free a tile. Sometimes does a free of the backing CPU/GPU block of memory. */
+	void Free(uint32 InHandle);
+
+	/** Get upload buffer description from handle. */
+	FVTUploadTileBuffer GetBufferFromHandle(uint32 InHandle) const;
+	/** Get upload buffer extended description from handle. */
+	FVTUploadTileBufferExt GetBufferFromHandleExt(uint32 InHandle) const;
+	
+	/** Get allocated memory in bytes. */
+	uint32 TotalAllocatedBytes() const { return NumAllocatedBytes; }
 
 private:
-	enum EListType
+	/** Handle for an allocated tile used by the allocation system. */
+	union FHandle
 	{
-		LIST_SUBMITTED,
-
-		LIST_COUNT,
+		uint32 PackedValue = 0u;
+		struct
+		{
+			uint32 FormatIndex : 8;
+			uint32 StagingBufferIndex : 8;
+			uint32 TileIndex : 16;
+		};
 	};
 
-	static const uint32 NUM_STAGING_TEXTURES = 3u;
-
+	/**
+	 * Backing memory for buffers used by the streaming/transcoding to write texture data.
+	 * The memory is split into equal sized tiles for multiple upload tasks.
+	 * Backing memory can be either CPU heap memory or a locked GPU memory buffer depending on the platform.
+	 */
 	struct FStagingBuffer
 	{
-		FStagingBuffer();
 		~FStagingBuffer();
 
+		void Init(uint32 InBufferStrideBytes, uint32 InTileSizeBytes);
+		void Release();
+
+		/** GPU buffer if used on platform. */
 		TRefCountPtr<FRHIBuffer> RHIBuffer;
+		/** Memory pointer to locked GPU buffer if used on platform, or to allocated CPU heap memory if not. */
 		void* Memory = nullptr;
-		uint32 Size = 0u;
-		uint32 CurrentOffset = 0u;
+		uint32 TileSize = 0u;
+		uint32 NumTiles = 0u;
+		/** List of tile indices that haven't been allocated. */
+		TArray<uint16> TileFreeList;
 	};
 
-	struct FStagingTexture
+	/** Container for multiple staging buffers. */
+	struct FSharedFormatBuffers
 	{
-		TRefCountPtr<FRHITexture2D> RHITexture;
-		uint32_t WidthInTiles = 0u;
-		uint32_t BatchCapacity = 0u;
-		bool bIsCPUWritable;
+		TArray<FStagingBuffer> StagingBuffers;
 	};
 
-	struct FPoolEntry
+	/**
+	 * Description of values that affect staging buffer creation.
+	 * Multiple virtual texture pools may map onto the same description and so can share staging buffer memory.
+	 */
+	struct FSharedFormatDesc
 	{
-		FStagingTexture StagingTexture[NUM_STAGING_TEXTURES];
-		EPixelFormat Format = PF_Unknown;
-		uint32 TileSize = 0u;
-		uint32 BatchTextureIndex = 0u;
-		uint32 BatchCount = 0u;
-		int32 FreeTileListHead = 0;
-		int32 SubmitTileListHead = 0;
+		uint32 BlockBytes = 0u;
+		uint32 Stride = 0u;
+		uint32 MemorySize = 0u;
 	};
 
+	/** Array of all discovered format descriptions. */
+	TArray<FSharedFormatDesc> FormatDescs;
+	/** Array of staging buffers. Kept in sync with associated formats from FormatDescs. */
+	TArray<FSharedFormatBuffers> FormatBuffers;
+	/** Allocated memory counter. */
+	uint32 NumAllocatedBytes = 0;
+};
+
+/** 
+ * Finalizer implementation for uploading virtual textures.
+ * Handles management of upload buffers and copying streamed data to the GPU physical texture.
+ */
+class FVirtualTextureUploadCache : public IVirtualTextureFinalizer, public FRenderResource
+{
+public:
+	//~ Begin FRenderResource Interface.
+	virtual void ReleaseRHI() override;
+	//~ End FRenderResource Interface.
+
+	/** Get a staging upload buffer for streaming texture data into. */
+	FVTUploadTileHandle PrepareTileForUpload(FVTUploadTileBuffer& OutBuffer, EPixelFormat InFormat, uint32 InTileSize);
+	/** 
+	 * Mark streamed upload data ready for upload to to the physical virtual texture.
+	 * Depending on the platform the upload might happen here, or be deferred to the Finalize() call.
+	 */
+	void SubmitTile(FRHICommandListImmediate& RHICmdList, const FVTUploadTileHandle& InHandle, FRHITexture2D* InDestTexture, int InDestX, int InDestY, int InSkipBorderSize);
+	/** Cancel a tile that was already in flight. */
+	void CancelTile(const FVTUploadTileHandle& InHandle);
+
+	//~ Begin IVirtualTextureFinalizer Interface.
+	virtual void Finalize(FRDGBuilder& GraphBuilder) override;
+	//~ End IVirtualTextureFinalizer Interface.
+
+	/** Call on a tick to recycle submitted staging buffers. */
+	void UpdateFreeList(bool bForceFreeAll = false);
+
+	/** Returns true if underlying allocator within the budget set by r.VT.MaxUploadMemory.  */
+	uint32 IsInMemoryBudget() const;
+
+private:
+	int32 GetOrCreatePoolIndex(EPixelFormat InFormat, uint32 InTileSize);
+
+	/** Description of a single allocated tile. Carries mutable state as tile moves from uploading to submitting to pending delete. */
 	struct FTileEntry
 	{
-		FTileEntry();
-		~FTileEntry();
-
-		FRHITexture2D* RHISubmitTexture = nullptr;
-		uint32 BufferIndex = 0u;
-		uint32 BufferOffset = 0u;
-		uint32 MemorySize = 0u;
-		uint32 Stride = 0u;
+		uint32 PoolIndex = 0;
+		uint32 TileHandle = 0u;
 		uint32 FrameSubmitted = 0u;
-		uint32 SubmitBatchIndex = 0u;
+		FRHITexture2D* RHISubmitTexture = nullptr;
 		int32 SubmitDestX = 0;
 		int32 SubmitDestY = 0;
 		int32 SubmitSkipBorderSize = 0;
-		int32 PoolIndex = 0;
-		int32 NextIndex = 0;
-		int32 PrevIndex = 0;
 	};
 
-	int32 GetOrCreatePoolIndex(EPixelFormat InFormat, uint32 InTileSize);
-
-	int32 CreateTileEntry(int32 PoolIndex)
-	{
-		const int32 Index = Tiles.AddDefaulted();
-		FTileEntry& Entry = Tiles[Index];
-		Entry.NextIndex = Entry.PrevIndex = Index;
-		Entry.PoolIndex = PoolIndex;
-		return Index;
-	}
-
-	void RemoveFromList(int32 Index)
+	/** Staging texture used for tile upload. Only used on platforms that don't have faster upload methods. */
+	struct FStagingTexture
 	{
-		FTileEntry& Entry = Tiles[Index];
-		checkSlow(Index >= LIST_COUNT); // if we're trying to remove a list head, something is corrupt
-		Tiles[Entry.PrevIndex].NextIndex = Entry.NextIndex;
-		Tiles[Entry.NextIndex].PrevIndex = Entry.PrevIndex;
-		Entry.NextIndex = Entry.PrevIndex = Index;
-	}
-
-	void AddToList(int32 HeadIndex, int32 Index)
+		TRefCountPtr<FRHITexture2D> RHITexture;
+		uint32 WidthInTiles = 0u;
+		uint32 BatchCapacity = 0u;
+		bool bIsCPUWritable;
+	};
+
+	/** State for a single pool. A pool covers all virtual textures of the same format and tile size. */
+	struct FPoolEntry
 	{
-		FTileEntry& Head = Tiles[HeadIndex];
-		FTileEntry& Entry = Tiles[Index];
+		EPixelFormat Format = PF_Unknown;
+		uint32 TileSize = 0u;
 
-		// make sure we're not currently in any list
-		check(Entry.NextIndex == Index);
-		check(Entry.PrevIndex == Index);
+		static const uint32 NUM_STAGING_TEXTURES = 3u;
+		FStagingTexture StagingTexture[NUM_STAGING_TEXTURES];
+		uint32 BatchTextureIndex = 0u;
 
-		Entry.NextIndex = HeadIndex;
-		Entry.PrevIndex = Head.PrevIndex;
-		Tiles[Head.PrevIndex].NextIndex = Index;
-		Head.PrevIndex = Index;
-	}
+		TArray<FTileEntry> PendingSubmit;
+	};
 
 	TArray<FPoolEntry> Pools;
-	TArray<FStagingBuffer> StagingBuffers;
-	TArray<FTileEntry> Tiles;
+	FVTUploadTileAllocator TileAllocator;
+	TSparseArray<FTileEntry> PendingUpload;
+	TSparseArray<FTileEntry> PendingRelease;
 	TArray<FRHITexture*> UpdatedTextures;
-	uint32 NumPendingTiles = 0u;
-};
\ No newline at end of file
+};
