World Partrition: only include cells which loaded state was changed by a user operations in the last loaded cells.

diff --git a/Engine/Source/Editor/DataLayerEditor/Private/DataLayer/DataLayerEditorSubsystem.cpp b/Engine/Source/Editor/DataLayerEditor/Private/DataLayer/DataLayerEditorSubsystem.cpp
--- a/Engine/Source/Editor/DataLayerEditor/Private/DataLayer/DataLayerEditorSubsystem.cpp
+++ b/Engine/Source/Editor/DataLayerEditor/Private/DataLayer/DataLayerEditorSubsystem.cpp
@@ -169,11 +169,11 @@ void UDataLayerEditorSubsystem::Deinitialize()
 	DataLayersBroadcast->Deinitialize();
 }
 
-bool UDataLayerEditorSubsystem::RefreshWorldPartitionEditorCells()
+bool UDataLayerEditorSubsystem::RefreshWorldPartitionEditorCells(bool bIsFromUserChange)
 {
 	if (UWorldPartition* WorldPartition = GetWorld() ? GetWorld()->GetWorldPartition() : nullptr)
 	{
-		if (!WorldPartition->RefreshLoadedEditorCells())
+		if (!WorldPartition->RefreshLoadedEditorCells(bIsFromUserChange))
 		{
 			return false;
 		}
@@ -987,51 +987,6 @@ void UDataLayerEditorSubsystem::MakeAllDataLayersVisible()
 	UpdateAllActorsVisibility(true, true);
 }
 
-bool UDataLayerEditorSubsystem::SetDataLayerIsDynamicallyLoadedInternal(UDataLayer* DataLayer, const bool bIsDynamicallyLoaded)
-{
-	check(DataLayer);
-	if (DataLayer->IsDynamicallyLoaded() != bIsDynamicallyLoaded)
-	{
-		DataLayer->Modify();
-		DataLayer->SetIsDynamicallyLoaded(bIsDynamicallyLoaded);
-		DataLayerChanged.Broadcast(EDataLayerAction::Modify, DataLayer, "bIsDynamicallyLoaded");
-		return true;
-	}
-	return false;
-}
-
-bool UDataLayerEditorSubsystem::SetDataLayerIsDynamicallyLoaded(UDataLayer* DataLayer, const bool bIsDynamicallyLoaded)
-{
-	bool bRefreshNeeded = SetDataLayerIsDynamicallyLoadedInternal(DataLayer, bIsDynamicallyLoaded);
-	return bRefreshNeeded ? RefreshWorldPartitionEditorCells() : true;
-}
-
-bool UDataLayerEditorSubsystem::SetDataLayersIsDynamicallyLoaded(const TArray<UDataLayer*>& DataLayers, const bool bIsDynamicallyLoaded)
-{
-	bool bRefreshNeeded = false;
-	for (UDataLayer* DataLayer : DataLayers)
-	{
-		bRefreshNeeded |= SetDataLayerIsDynamicallyLoadedInternal(DataLayer, bIsDynamicallyLoaded);
-	}
-	return bRefreshNeeded ? RefreshWorldPartitionEditorCells() : true;
-}
-
-bool UDataLayerEditorSubsystem::ToggleDataLayerIsDynamicallyLoaded(UDataLayer* DataLayer)
-{
-	check(DataLayer);
-	return SetDataLayerIsDynamicallyLoaded(DataLayer, !DataLayer->IsDynamicallyLoaded());
-}
-
-bool UDataLayerEditorSubsystem::ToggleDataLayersIsDynamicallyLoaded(const TArray<UDataLayer*>& DataLayers)
-{
-	bool bRefreshNeeded = false;
-	for (UDataLayer* DataLayer : DataLayers)
-	{
-		bRefreshNeeded |= SetDataLayerIsDynamicallyLoadedInternal(DataLayer, !DataLayer->IsDynamicallyLoaded());
-	}
-	return bRefreshNeeded ? RefreshWorldPartitionEditorCells() : true;
-}
-
 bool UDataLayerEditorSubsystem::SetDataLayerIsDynamicallyLoadedInEditorInternal(UDataLayer* DataLayer, const bool bIsDynamicallyLoadedInEditor, const bool bIsFromUserChange)
 {
 	check(DataLayer);
@@ -1048,7 +1003,7 @@ bool UDataLayerEditorSubsystem::SetDataLayerIsDynamicallyLoadedInEditorInternal(
 bool UDataLayerEditorSubsystem::SetDataLayerIsDynamicallyLoadedInEditor(UDataLayer* DataLayer, const bool bIsDynamicallyLoadedInEditor, const bool bIsFromUserChange)
 {
 	bool bRefreshNeeded = SetDataLayerIsDynamicallyLoadedInEditorInternal(DataLayer, bIsDynamicallyLoadedInEditor, bIsFromUserChange);
-	return bRefreshNeeded ? RefreshWorldPartitionEditorCells() : true;
+	return bRefreshNeeded ? RefreshWorldPartitionEditorCells(bIsFromUserChange) : true;
 }
 
 bool UDataLayerEditorSubsystem::SetDataLayersIsDynamicallyLoadedInEditor(const TArray<UDataLayer*>& DataLayers, const bool bIsDynamicallyLoadedInEditor, const bool bIsFromUserChange)
@@ -1058,7 +1013,7 @@ bool UDataLayerEditorSubsystem::SetDataLayersIsDynamicallyLoadedInEditor(const T
 	{
 		bRefreshNeeded |= SetDataLayerIsDynamicallyLoadedInEditorInternal(DataLayer, bIsDynamicallyLoadedInEditor, bIsFromUserChange);
 	}
-	return bRefreshNeeded ? RefreshWorldPartitionEditorCells() : true;
+	return bRefreshNeeded ? RefreshWorldPartitionEditorCells(bIsFromUserChange) : true;
 }
 
 bool UDataLayerEditorSubsystem::ToggleDataLayerIsDynamicallyLoadedInEditor(UDataLayer* DataLayer, const bool bIsFromUserChange)
@@ -1074,7 +1029,7 @@ bool UDataLayerEditorSubsystem::ToggleDataLayersIsDynamicallyLoadedInEditor(cons
 	{
 		bRefreshNeeded |= SetDataLayerIsDynamicallyLoadedInEditorInternal(DataLayer, !DataLayer->IsDynamicallyLoadedInEditor(), bIsFromUserChange);
 	}
-	return bRefreshNeeded ? RefreshWorldPartitionEditorCells() : true;
+	return bRefreshNeeded ? RefreshWorldPartitionEditorCells(bIsFromUserChange) : true;
 }
 
 bool UDataLayerEditorSubsystem::ResetUserSettings(const TArray<UDataLayer*>& DataLayers)
@@ -1084,7 +1039,7 @@ bool UDataLayerEditorSubsystem::ResetUserSettings(const TArray<UDataLayer*>& Dat
 	{
 		bRefreshNeeded |= SetDataLayerIsDynamicallyLoadedInEditorInternal(DataLayer, DataLayer->IsInitiallyLoadedInEditor(), true);
 	}
-	return bRefreshNeeded ? RefreshWorldPartitionEditorCells() : true;
+	return bRefreshNeeded ? RefreshWorldPartitionEditorCells(true) : true;
 }
 
 UDataLayer* UDataLayerEditorSubsystem::GetDataLayer(const FActorDataLayer& ActorDataLayer) const
diff --git a/Engine/Source/Editor/DataLayerEditor/Public/DataLayer/DataLayerEditorSubsystem.h b/Engine/Source/Editor/DataLayerEditor/Public/DataLayer/DataLayerEditorSubsystem.h
--- a/Engine/Source/Editor/DataLayerEditor/Public/DataLayer/DataLayerEditorSubsystem.h
+++ b/Engine/Source/Editor/DataLayerEditor/Public/DataLayer/DataLayerEditorSubsystem.h
@@ -437,40 +437,6 @@ public:
 	UFUNCTION(BlueprintCallable, Category = DataLayers)
 	void ToggleDataLayersVisibility(const TArray<UDataLayer*>& DataLayers);
 
-	/**
-	 * Changes the DataLayer's IsDynamicallyLoaded flag to the provided state
-	 *
-	 * @param	DataLayer				The DataLayer to affect.
-	 * @param	bIsDynamicallyLoaded	If true the DataLayer will affect runtime streaming and actors will only be loaded if the DataLayer is active.
-	 */
-	UFUNCTION(BlueprintCallable, Category = DataLayers)
-	bool SetDataLayerIsDynamicallyLoaded(UDataLayer* DataLayer, const bool bIsDynamicallyLoaded);
-
-	/**
-	 * Changes the IsDynamicallyLoaded flag of the DataLayers to the provided state
-	 *
-	 * @param	DataLayers				The DataLayers to affect
-	 * @param	bIsDynamicallyLoaded	If true the DataLayer will affect runtime streaming and actors will only be loaded if the DataLayer is active.
-	 */
-	UFUNCTION(BlueprintCallable, Category = DataLayers)
-	bool SetDataLayersIsDynamicallyLoaded(const TArray<UDataLayer*>& DataLayers, const bool bIsDynamicallyLoaded);
-
-	/**
-	 * Toggles the DataLayer's IsDynamicallyLoaded flag
-	 *
-	 * @param DataLayer	The DataLayer to affect
-	 */
-	UFUNCTION(BlueprintCallable, Category = DataLayers)
-	bool ToggleDataLayerIsDynamicallyLoaded(UDataLayer* DataLayer);
-
-	/**
-	 * Toggles the IsDynamicallyLoaded flag of all of the DataLayers
-	 *
-	 * @param	DataLayers	The DataLayers to affect
-	 */
-	UFUNCTION(BlueprintCallable, Category = DataLayers)
-	bool ToggleDataLayersIsDynamicallyLoaded(const TArray<UDataLayer*>& DataLayers);
-
 	/**
 	 * Changes the DataLayer's IsDynamicallyLoadedInEditor flag to the provided state
 	 *
@@ -648,9 +614,8 @@ private:
 	/** Delegate handler for FEditorDelegates::PostUndoRedo. It internally calls DataLayerChanged.Broadcast and UpdateAllActorsVisibility to refresh the actors of each DataLayer. */
 	void PostUndoRedo();
 
-	bool SetDataLayerIsDynamicallyLoadedInternal(UDataLayer* DataLayer, const bool bIsDynamicallyLoaded);
 	bool SetDataLayerIsDynamicallyLoadedInEditorInternal(UDataLayer* DataLayer, const bool bIsDynamicallyLoadedInEditor, const bool bIsFromUserChange);
-	bool RefreshWorldPartitionEditorCells();
+	bool RefreshWorldPartitionEditorCells(bool bIsFromUserChange);
 	void UpdateDataLayerEditorPerProjectUserSettings();
 
 	/** Fires whenever one or more DataLayer changes */
diff --git a/Engine/Source/Editor/FoliageEdit/Private/ProceduralFoliageComponentDetails.cpp b/Engine/Source/Editor/FoliageEdit/Private/ProceduralFoliageComponentDetails.cpp
--- a/Engine/Source/Editor/FoliageEdit/Private/ProceduralFoliageComponentDetails.cpp
+++ b/Engine/Source/Editor/FoliageEdit/Private/ProceduralFoliageComponentDetails.cpp
@@ -118,7 +118,7 @@ FReply FProceduralFoliageComponentDetails::OnLoadUnloadedAreas()
 				FVector Origin;
 				FVector Extent;
 				Component->GetOwner()->GetActorBounds(false, Origin, Extent);
-				WorldPartition->LoadEditorCells(FBox(Origin - Extent, Origin + Extent));
+				WorldPartition->LoadEditorCells(FBox(Origin - Extent, Origin + Extent), true);
 			}
 		}
 	}
diff --git a/Engine/Source/Editor/UnrealEd/Private/WorldPartition/WorldPartitionBuilder.cpp b/Engine/Source/Editor/UnrealEd/Private/WorldPartition/WorldPartitionBuilder.cpp
--- a/Engine/Source/Editor/UnrealEd/Private/WorldPartition/WorldPartitionBuilder.cpp
+++ b/Engine/Source/Editor/UnrealEd/Private/WorldPartition/WorldPartitionBuilder.cpp
@@ -186,7 +186,7 @@ bool UWorldPartitionBuilder::Run(UWorld* World, FPackageSourceControlHelper& Pac
 		if (bUpdateEditorCells)
 		{
 			UE_LOG(LogWorldPartitionBuilder, Display, TEXT("DataLayer load state changed refreshing editor cells"));
-			WorldPartition->RefreshLoadedEditorCells();
+			WorldPartition->RefreshLoadedEditorCells(false);
 		}
 	}
 		
@@ -228,7 +228,7 @@ bool UWorldPartitionBuilder::Run(UWorld* World, FPackageSourceControlHelper& Pac
 					BoundsToLoad.ExpandBy(IterativeCellOverlapSize);
 
 					UE_LOG(LogWorldPartitionBuilder, Verbose, TEXT("Loading Bounds: Min %s, Max %s"), *BoundsToLoad.Min.ToString(), *BoundsToLoad.Max.ToString());
-					WorldPartition->LoadEditorCells(BoundsToLoad);
+					WorldPartition->LoadEditorCells(BoundsToLoad, false);
 					LoadedBounds += BoundsToLoad;
 
 
@@ -239,7 +239,7 @@ bool UWorldPartitionBuilder::Run(UWorld* World, FPackageSourceControlHelper& Pac
 
 					if (FWorldPartitionHelpers::HasExceededMaxMemory())
 					{
-						WorldPartition->UnloadEditorCells(LoadedBounds);
+						WorldPartition->UnloadEditorCells(LoadedBounds, false);
 						// Reset Loaded Bounds
 						LoadedBounds.Init();
 
@@ -257,7 +257,7 @@ bool UWorldPartitionBuilder::Run(UWorld* World, FPackageSourceControlHelper& Pac
 		if (LoadingMode == ELoadingMode::EntireWorld)
 		{
 			BoundsToLoad += FBox(FVector(-WORLD_MAX, -WORLD_MAX, -WORLD_MAX), FVector(WORLD_MAX, WORLD_MAX, WORLD_MAX));
-			WorldPartition->LoadEditorCells(BoundsToLoad);
+			WorldPartition->LoadEditorCells(BoundsToLoad, false);
 		}
 
 		return RunInternal(World, BoundsToLoad, PackageHelper);
diff --git a/Engine/Source/Editor/WorldPartitionEditor/Private/WorldPartition/SWorldPartitionEditorGrid2D.cpp b/Engine/Source/Editor/WorldPartitionEditor/Private/WorldPartition/SWorldPartitionEditorGrid2D.cpp
--- a/Engine/Source/Editor/WorldPartitionEditor/Private/WorldPartition/SWorldPartitionEditorGrid2D.cpp
+++ b/Engine/Source/Editor/WorldPartitionEditor/Private/WorldPartition/SWorldPartitionEditorGrid2D.cpp
@@ -160,14 +160,14 @@ void SWorldPartitionEditorGrid2D::Construct(const FArguments& InArgs)
 
 void SWorldPartitionEditorGrid2D::LoadSelectedCells()
 {
-	WorldPartition->LoadEditorCells(SelectBox);
+	WorldPartition->LoadEditorCells(SelectBox, true);
 	GEditor->RedrawLevelEditingViewports();
 	Refresh();
 }
 
 void SWorldPartitionEditorGrid2D::UnloadSelectedCells()
 {
-	WorldPartition->UnloadEditorCells(SelectBox);
+	WorldPartition->UnloadEditorCells(SelectBox, true);
 	GEditor->RedrawLevelEditingViewports();
 	Refresh();
 }
diff --git a/Engine/Source/Editor/WorldPartitionEditor/Private/WorldPartitionEditorModule.cpp b/Engine/Source/Editor/WorldPartitionEditor/Private/WorldPartitionEditorModule.cpp
--- a/Engine/Source/Editor/WorldPartitionEditor/Private/WorldPartitionEditorModule.cpp
+++ b/Engine/Source/Editor/WorldPartitionEditor/Private/WorldPartitionEditorModule.cpp
@@ -35,7 +35,7 @@ static void OnLoadSelectedWorldPartitionVolumes(TArray<TWeakObjectPtr<AActor>> V
 	for (TWeakObjectPtr<AActor> Actor: Volumes)
 	{
 		AWorldPartitionVolume* WorldPartitionVolume = CastChecked<AWorldPartitionVolume>(Actor.Get());
-		WorldPartitionVolume->LoadIntersectingCells();
+		WorldPartitionVolume->LoadIntersectingCells(true);
 	}
 }
 
diff --git a/Engine/Source/Runtime/Engine/Private/WorldPartition/WorldPartition.cpp b/Engine/Source/Runtime/Engine/Private/WorldPartition/WorldPartition.cpp
--- a/Engine/Source/Runtime/Engine/Private/WorldPartition/WorldPartition.cpp
+++ b/Engine/Source/Runtime/Engine/Private/WorldPartition/WorldPartition.cpp
@@ -457,7 +457,7 @@ void UWorldPartition::Initialize(UWorld* InWorld, const FTransform& InTransform)
 		}
 
 		// Load the always loaded cell, don't call LoadCells to avoid creating a transaction
-		UpdateLoadingEditorCell(EditorHash->GetAlwaysLoadedCell(), true);
+		UpdateLoadingEditorCell(EditorHash->GetAlwaysLoadedCell(), true, false);
 
 		// Load more cells depending on the user's settings
 		// Skipped when running from a commandlet
@@ -474,7 +474,7 @@ void UWorldPartition::Initialize(UWorld* InWorld, const FTransform& InTransform)
 			{
 				EditorHash->ForEachCell([this](UWorldPartitionEditorCell* Cell)
 				{
-					UpdateLoadingEditorCell(Cell, true);
+					UpdateLoadingEditorCell(Cell, true, false);
 				});
 			}
 
@@ -487,7 +487,7 @@ void UWorldPartition::Initialize(UWorld* InWorld, const FTransform& InTransform)
 				{
 					if (UWorldPartitionEditorCell* Cell = FindObject<UWorldPartitionEditorCell>(EditorHash, *EditorGridLastLoadedCell.ToString()))
 					{
-						UpdateLoadingEditorCell(Cell, true);
+						UpdateLoadingEditorCell(Cell, true, true);
 					}
 				}
 			}
@@ -586,7 +586,7 @@ void UWorldPartition::Uninitialize()
 				// @todo_ow: Once Metadata is removed from external actor's package, this won't be necessary anymore.
 				EditorHash->ForEachCell([this](UWorldPartitionEditorCell* Cell)
 				{
-					UpdateLoadingEditorCell(Cell, /*bShouldBeLoaded*/false);
+					UpdateLoadingEditorCell(Cell, /*bShouldBeLoaded*/false, /*bIsFromUserChange*/false);
 				});
 			}
 		}
@@ -746,7 +746,7 @@ bool UWorldPartition::IsSimulating()
 }
 
 #if WITH_EDITOR
-void UWorldPartition::LoadEditorCells(const FBox& Box)
+void UWorldPartition::LoadEditorCells(const FBox& Box, bool bIsFromUserChange)
 {
 	TArray<UWorldPartitionEditorCell*> CellsToLoad;
 	if (EditorHash->GetIntersectingCells(Box, CellsToLoad))
@@ -763,18 +763,18 @@ void UWorldPartition::LoadEditorCells(const FBox& Box)
 		for (UWorldPartitionEditorCell* Cell : CellsToLoad)
 		{
 			SlowTask.EnterProgressFrame(Cell->Actors.Num() - Cell->LoadedActors.Num());
-			UpdateLoadingEditorCell(Cell, true);
+			UpdateLoadingEditorCell(Cell, true, bIsFromUserChange);
 		}
 	}
 }
 
-bool UWorldPartition::RefreshLoadedEditorCells()
+bool UWorldPartition::RefreshLoadedEditorCells(bool bIsFromUserChange)
 {
 	auto GetCellsToRefresh = [this](TArray<UWorldPartitionEditorCell*>& CellsToRefresh)
 	{
 		EditorHash->ForEachCell([&CellsToRefresh](UWorldPartitionEditorCell* Cell)
 		{
-			if (Cell->bLoaded)
+			if (Cell->IsLoaded())
 			{
 				CellsToRefresh.Add(Cell);
 			}
@@ -782,10 +782,10 @@ bool UWorldPartition::RefreshLoadedEditorCells()
 		return !CellsToRefresh.IsEmpty();
 	};
 
-	return UpdateEditorCells(GetCellsToRefresh, /*bIsCellShouldBeLoaded*/true);
+	return UpdateEditorCells(GetCellsToRefresh, /*bIsCellShouldBeLoaded*/true, bIsFromUserChange);
 }
 
-void UWorldPartition::UnloadEditorCells(const FBox& Box)
+void UWorldPartition::UnloadEditorCells(const FBox& Box, bool bIsFromUserChange)
 {
 	FWorldPartionCellUpdateContext CellUpdateContext(this);
 
@@ -794,7 +794,7 @@ void UWorldPartition::UnloadEditorCells(const FBox& Box)
 		return EditorHash->GetIntersectingCells(Box, CellsToUnload) > 0;
 	};
 
-	UpdateEditorCells(GetCellsToUnload, /*bIsCellShouldBeLoaded*/false);
+	UpdateEditorCells(GetCellsToUnload, /*bIsCellShouldBeLoaded*/false, bIsFromUserChange);
 }
 
 bool UWorldPartition::AreEditorCellsLoaded(const FBox& Box)
@@ -804,7 +804,7 @@ bool UWorldPartition::AreEditorCellsLoaded(const FBox& Box)
 	{
 		for (UWorldPartitionEditorCell* Cell : CellsToLoad)
 		{
-			if (!Cell->bLoaded)
+			if (!Cell->IsLoaded())
 			{
 				return false;
 			}
@@ -814,7 +814,7 @@ bool UWorldPartition::AreEditorCellsLoaded(const FBox& Box)
 	return true;
 }
 
-bool UWorldPartition::UpdateEditorCells(TFunctionRef<bool(TArray<UWorldPartitionEditorCell*>&)> GetCellsToProcess, bool bIsCellShouldBeLoaded)
+bool UWorldPartition::UpdateEditorCells(TFunctionRef<bool(TArray<UWorldPartitionEditorCell*>&)> GetCellsToProcess, bool bIsCellShouldBeLoaded, bool bIsFromUserChange)
 {
 	FWorldPartionCellUpdateContext CellUpdateContext(this);
 
@@ -900,7 +900,7 @@ bool UWorldPartition::UpdateEditorCells(TFunctionRef<bool(TArray<UWorldPartition
 	for (UWorldPartitionEditorCell* Cell : CellsToProcess)
 	{
 		SlowTask.EnterProgressFrame(1);
-		UpdateLoadingEditorCell(Cell, bIsCellShouldBeLoaded);
+		UpdateLoadingEditorCell(Cell, bIsCellShouldBeLoaded, bIsFromUserChange);
 	}
 
 	return true;
@@ -944,13 +944,13 @@ bool UWorldPartition::ShouldActorBeLoadedByEditorCells(const FWorldPartitionActo
 };
 
 // Loads actors in Editor cell
-void UWorldPartition::UpdateLoadingEditorCell(UWorldPartitionEditorCell* Cell, bool bShouldBeLoaded)
+void UWorldPartition::UpdateLoadingEditorCell(UWorldPartitionEditorCell* Cell, bool bShouldBeLoaded, bool bFromUserOperation)
 {
 	TRACE_CPUPROFILER_EVENT_SCOPE(UWorldPartition::UpdateLoadingEditorCell);
 	
 	FWorldPartionCellUpdateContext CellUpdateContext(this);
 
-	UE_LOG(LogWorldPartition, Verbose, TEXT("UWorldPartition::UpdateLoadingEditorCell 0x%08llx [%s]"), Cell, bShouldBeLoaded ? TEXT("Load") : TEXT("Unload"));
+	UE_LOG(LogWorldPartition, Verbose, TEXT("UWorldPartition::UpdateLoadingEditorCell 0x%08llx [%s/%s]"), Cell, bShouldBeLoaded ? TEXT("Load") : TEXT("Unload"), bFromUserOperation ? TEXT("User") : TEXT("Auto"));
 
 	Cell->Modify(false);
 
@@ -988,11 +988,11 @@ void UWorldPartition::UpdateLoadingEditorCell(UWorldPartitionEditorCell* Cell, b
 		}
 	}
 
-	if (Cell->bLoaded != bShouldBeLoaded)
+	if (Cell->IsLoaded() != bShouldBeLoaded)
 	{
-		Cell->bLoaded = bShouldBeLoaded;
+		Cell->SetLoaded(bShouldBeLoaded, bShouldBeLoaded && bFromUserOperation);
 
-		if (Cell->bLoaded)
+		if (Cell->IsLoaded())
 		{
 			EditorHash->OnCellLoaded(Cell);
 		}
@@ -1271,12 +1271,23 @@ void UWorldPartition::SavePerUserSettings()
 	if (GIsEditor && !World->IsGameWorld() && !IsRunningCommandlet() && !IsEngineExitRequested())
 	{
 		// Save last loaded cells settings
+		const TArray<FName>& LastEditorGridLoadedCells = GetMutableDefault<UWorldPartitionEditorPerProjectUserSettings>()->GetEditorGridLoadedCells(GetWorld());
+
 		TArray<FName> EditorGridLastLoadedCells;
-		EditorHash->ForEachCell([this, &EditorGridLastLoadedCells](UWorldPartitionEditorCell* Cell)
+		EditorHash->ForEachCell([this, &LastEditorGridLoadedCells, &EditorGridLastLoadedCells](UWorldPartitionEditorCell* Cell)
 		{
-			if ((Cell != EditorHash->GetAlwaysLoadedCell()) && Cell->bLoaded)
+			FName CellName = Cell->GetFName();
+
+			if (Cell != EditorHash->GetAlwaysLoadedCell())
 			{
-				EditorGridLastLoadedCells.Add(Cell->GetFName());
+				if (Cell->IsLoaded() && Cell->IsLoadedChangedByUserOperation())
+				{
+					EditorGridLastLoadedCells.Add(CellName);
+				}
+				else if (!Cell->IsLoaded() && !Cell->IsLoadedChangedByUserOperation() && LastEditorGridLoadedCells.Contains(CellName))
+				{
+					EditorGridLastLoadedCells.Add(CellName);
+				}
 			}
 		});
 
diff --git a/Engine/Source/Runtime/Engine/Private/WorldPartition/WorldPartitionEditorCell.cpp b/Engine/Source/Runtime/Engine/Private/WorldPartition/WorldPartitionEditorCell.cpp
--- a/Engine/Source/Runtime/Engine/Private/WorldPartition/WorldPartitionEditorCell.cpp
+++ b/Engine/Source/Runtime/Engine/Private/WorldPartition/WorldPartitionEditorCell.cpp
@@ -114,5 +114,6 @@ UWorldPartitionEditorCell::UWorldPartitionEditorCell(const FObjectInitializer& O
 #if WITH_EDITOR
 	, Bounds(ForceInitToZero)
 	, bLoaded(false)
+	, bLoadedChangedByUserOperation(false)
 #endif
 {}
diff --git a/Engine/Source/Runtime/Engine/Private/WorldPartition/WorldPartitionEditorSpatialHash.cpp b/Engine/Source/Runtime/Engine/Private/WorldPartition/WorldPartitionEditorSpatialHash.cpp
--- a/Engine/Source/Runtime/Engine/Private/WorldPartition/WorldPartitionEditorSpatialHash.cpp
+++ b/Engine/Source/Runtime/Engine/Private/WorldPartition/WorldPartitionEditorSpatialHash.cpp
@@ -87,7 +87,7 @@ void UWorldPartitionEditorSpatialHash::Tick(float DeltaSeconds)
 			UWorldPartitionEditorCell* EditorCell = nullptr;
 			if (UWorldPartitionEditorCell** EditorCellPtr = HashCells.Find(CellCoord))
 			{
-				if ((*EditorCellPtr)->bLoaded)
+				if ((*EditorCellPtr)->IsLoaded())
 				{
 					FCellCoord LevelCellCoord = CellCoord;
 					while (LevelCellCoord.Level < CurrentLevel)
@@ -222,7 +222,7 @@ void UWorldPartitionEditorSpatialHash::HashActor(FWorldPartitionHandle& InActorH
 					else
 					{
 						UWorldPartitionEditorCell* EditorCell = HashCells.FindChecked(CellCoord);
-						bSetLoadedMask = EditorCell->bLoaded;
+						bSetLoadedMask = EditorCell->IsLoaded();
 					}
 
 					FCellCoord LevelCellCoord = CellCoord;
@@ -285,7 +285,7 @@ void UWorldPartitionEditorSpatialHash::UnhashActor(FWorldPartitionHandle& InActo
 				verify(HashCells.Remove(CellCoord));
 
 				bool bClearChildMask = true;
-				bool bClearLoadedMask = EditorCell->bLoaded;
+				bool bClearLoadedMask = EditorCell->IsLoaded();
 
 				FCellCoord LevelCellCoord = CellCoord;
 				while (LevelCellCoord.Level < CurrentLevel)
@@ -339,7 +339,7 @@ void UWorldPartitionEditorSpatialHash::OnCellLoaded(const UWorldPartitionEditorC
 		return;
 	}
 
-	check(Cell->bLoaded);
+	check(Cell->IsLoaded());
 
 	const FCellCoord CellCoord = GetCellCoords(Cell->Bounds.GetCenter(), 0);
 	const int32 CurrentLevel = GetLevelForBox(Bounds);
@@ -371,7 +371,7 @@ void UWorldPartitionEditorSpatialHash::OnCellUnloaded(const UWorldPartitionEdito
 		return;
 	}
 
-	check(!Cell->bLoaded);
+	check(!Cell->IsLoaded());
 
 	const FCellCoord CellCoord = GetCellCoords(Cell->Bounds.GetCenter(), 0);
 	const int32 CurrentLevel = GetLevelForBox(Bounds);
@@ -524,7 +524,7 @@ int32 UWorldPartitionEditorSpatialHash::ForEachIntersectingUnloadedRegionInner(c
 		UWorldPartitionEditorCell* EditorCell = nullptr;
 		if (UWorldPartitionEditorCell** EditorCellPtr = HashCells.Find(CellCoord))
 		{
-			if (!(*EditorCellPtr)->bLoaded)
+			if (!(*EditorCellPtr)->IsLoaded())
 			{
 				InOperation(CellCoord);
 				NumIntersecting++;
diff --git a/Engine/Source/Runtime/Engine/Private/WorldPartition/WorldPartitionVolume.cpp b/Engine/Source/Runtime/Engine/Private/WorldPartition/WorldPartitionVolume.cpp
--- a/Engine/Source/Runtime/Engine/Private/WorldPartition/WorldPartitionVolume.cpp
+++ b/Engine/Source/Runtime/Engine/Private/WorldPartition/WorldPartitionVolume.cpp
@@ -14,29 +14,21 @@ AWorldPartitionVolume::AWorldPartitionVolume(const FObjectInitializer& ObjectIni
 }
 
 #if WITH_EDITOR
-void AWorldPartitionVolume::LoadIntersectingCells()
+void AWorldPartitionVolume::LoadIntersectingCells(bool bIsFromUserChange)
 {
 	if (UWorldPartition* WorldPartition = GetWorld()->GetWorldPartition())
 	{
-		WorldPartition->LoadEditorCells(GetIntersectingBounds());
+		WorldPartition->LoadEditorCells(GetStreamingBounds(), bIsFromUserChange);
 	}
 }
 
-void AWorldPartitionVolume::UnloadIntersectingCells()
+void AWorldPartitionVolume::UnloadIntersectingCells(bool bIsFromUserChange)
 {
 	if (UWorldPartition* WorldPartition = GetWorld()->GetWorldPartition())
 	{
-		WorldPartition->UnloadEditorCells(GetIntersectingBounds());
+		WorldPartition->UnloadEditorCells(GetStreamingBounds(), bIsFromUserChange);
 	}
 }
-
-FBox AWorldPartitionVolume::GetIntersectingBounds() const
-{
-	FVector Origin;
-	FVector Extent;
-	GetActorBounds(false, Origin, Extent);
-	return FBox(Origin - Extent, Origin + Extent);
-}
 #endif
 
 #undef LOCTEXT_NAMESPACE
\ No newline at end of file
diff --git a/Engine/Source/Runtime/Engine/Public/WorldPartition/WorldPartition.h b/Engine/Source/Runtime/Engine/Public/WorldPartition/WorldPartition.h
--- a/Engine/Source/Runtime/Engine/Public/WorldPartition/WorldPartition.h
+++ b/Engine/Source/Runtime/Engine/Public/WorldPartition/WorldPartition.h
@@ -129,10 +129,10 @@ public:
 #if WITH_EDITOR
 	FName GetWorldPartitionEditorName();
 
-	void LoadEditorCells(const FBox& Box);
-	void UnloadEditorCells(const FBox& Box);
+	void LoadEditorCells(const FBox& Box, bool bIsFromUserChange);
+	void UnloadEditorCells(const FBox& Box, bool bIsFromUserChange);
 	bool AreEditorCellsLoaded(const FBox& Box);
-	bool RefreshLoadedEditorCells();
+	bool RefreshLoadedEditorCells(bool bIsFromUserChange);
 
 	// PIE/Game/Cook Methods
 	bool GenerateStreaming(EWorldPartitionStreamingMode Mode, TArray<FString>* OutPackagesToGenerate = nullptr);
@@ -232,11 +232,11 @@ private:
 	void UnregisterDelegates();	
 
 #if WITH_EDITOR
-	void UpdateLoadingEditorCell(UWorldPartitionEditorCell* Cell, bool bShouldBeLoaded);
+	void UpdateLoadingEditorCell(UWorldPartitionEditorCell* Cell, bool bShouldBeLoaded, bool bFromUserOperation);
 	void HashActorDesc(FWorldPartitionActorDesc* ActorDesc);
 	void UnhashActorDesc(FWorldPartitionActorDesc* ActorDesc);
 	bool ShouldActorBeLoadedByEditorCells(const FWorldPartitionActorDesc* ActorDesc) const;
-	bool UpdateEditorCells(TFunctionRef<bool(TArray<UWorldPartitionEditorCell*>&)> GetCellsToProcess, bool bIsCellShouldBeLoaded);
+	bool UpdateEditorCells(TFunctionRef<bool(TArray<UWorldPartitionEditorCell*>&)> GetCellsToProcess, bool bIsCellShouldBeLoaded, bool bIsFromUserChange);
 
 	void ApplyActorTransform(AActor* Actor, const FTransform& InTransform);
 #endif
diff --git a/Engine/Source/Runtime/Engine/Public/WorldPartition/WorldPartitionEditorCell.h b/Engine/Source/Runtime/Engine/Public/WorldPartition/WorldPartitionEditorCell.h
--- a/Engine/Source/Runtime/Engine/Public/WorldPartition/WorldPartitionEditorCell.h
+++ b/Engine/Source/Runtime/Engine/Public/WorldPartition/WorldPartitionEditorCell.h
@@ -72,11 +72,33 @@ public:
 	void RemoveActor(const FWorldPartitionHandle& ActorHandle);
 	void RemoveActor(const FGuid& Source, const FWorldPartitionHandle& ActorHandle);
 
+	inline bool IsLoaded() const
+	{
+		return bLoaded;
+	}
+
+	inline void SetLoaded(bool bInLoaded, bool bInLoadedByUserOperation)
+	{
+		bLoaded = bInLoaded;
+		bLoadedChangedByUserOperation |= bInLoadedByUserOperation;
+	}
+
+	inline bool IsLoadedChangedByUserOperation() const
+	{
+		return bLoadedChangedByUserOperation;
+	}
+
 	FBox						Bounds;
 
-	/** Tells if the cell was manually loaded in the editor */
+private:
+	/** Tells if the cell is loaded in the editor */
 	bool						bLoaded : 1;
 
+	/** Tells if the cell loading state was changed by a user operation */
+	bool						bLoadedChangedByUserOperation : 1;
+
+public:
+
 	TSet<FActorHandle> Actors;
 	TSet<FActorReference> LoadedActors;
 #endif
diff --git a/Engine/Source/Runtime/Engine/Public/WorldPartition/WorldPartitionVolume.h b/Engine/Source/Runtime/Engine/Public/WorldPartition/WorldPartitionVolume.h
--- a/Engine/Source/Runtime/Engine/Public/WorldPartition/WorldPartitionVolume.h
+++ b/Engine/Source/Runtime/Engine/Public/WorldPartition/WorldPartitionVolume.h
@@ -21,11 +21,8 @@ public:
 	// End AActor
 
 #if WITH_EDITOR
-	void LoadIntersectingCells();
-	void UnloadIntersectingCells();
+	void LoadIntersectingCells(bool bIsFromUserChange);
+	void UnloadIntersectingCells(bool bIsFromUserChange);
 	virtual EActorGridPlacement GetDefaultGridPlacement() const override { return EActorGridPlacement::AlwaysLoaded; }
-
-private:
-	FBox GetIntersectingBounds() const;
 #endif
 };
\ No newline at end of file
diff --git a/Engine/Source/Runtime/Foliage/Private/ProceduralFoliageComponent.cpp b/Engine/Source/Runtime/Foliage/Private/ProceduralFoliageComponent.cpp
--- a/Engine/Source/Runtime/Foliage/Private/ProceduralFoliageComponent.cpp
+++ b/Engine/Source/Runtime/Foliage/Private/ProceduralFoliageComponent.cpp
@@ -127,7 +127,7 @@ bool UProceduralFoliageComponent::ExecuteSimulation(TArray<FDesiredFoliageInstan
 	{
 		FBox Bounds = GetBounds();
 		check(Bounds.IsValid);
-		WorldPartition->LoadEditorCells(Bounds);
+		WorldPartition->LoadEditorCells(Bounds, false);
 	}
 
 	FBodyInstance* BoundsBodyInstance = GetBoundsBodyInstance();
