Breaking up a level instance now propagates data layer information

diff --git a/Engine/Source/Runtime/Engine/Private/LevelInstance/LevelInstanceSubsystem.cpp b/Engine/Source/Runtime/Engine/Private/LevelInstance/LevelInstanceSubsystem.cpp
--- a/Engine/Source/Runtime/Engine/Private/LevelInstance/LevelInstanceSubsystem.cpp
+++ b/Engine/Source/Runtime/Engine/Private/LevelInstance/LevelInstanceSubsystem.cpp
@@ -937,8 +937,10 @@ void ULevelInstanceSubsystem::BreakLevelInstance_Impl(ALevelInstance* LevelInsta
 			return;
 		}
 
+		TArray<const UDataLayer*> LevelInstanceDataLayers = LevelInstanceActor->GetDataLayerObjects();
+
 		TSet<AActor*> ActorsToMove;
-		TFunction<bool(AActor*)> AddActorToMove = [this, &ActorsToMove, &AddActorToMove](AActor* Actor)
+		TFunction<bool(AActor*)> AddActorToMove = [this, &ActorsToMove, &AddActorToMove, &LevelInstanceDataLayers](AActor* Actor)
 		{
 			if (ActorsToMove.Contains(Actor))
 			{
@@ -961,6 +963,15 @@ void ULevelInstanceSubsystem::BreakLevelInstance_Impl(ALevelInstance* LevelInsta
 						}
 					}
 
+					// Apply the same data layer settings to the actors to move out
+					if (Actor->SupportsDataLayer() && Actor->IsValidForDataLayer())
+					{
+						for (const UDataLayer* DataLayer : LevelInstanceDataLayers)
+						{
+							Actor->AddDataLayer(DataLayer);
+						}
+					}
+
 					ActorsToMove.Add(Actor);
 					return true;
 				}
