Remove UWorldPartition::FlushStreaming from API.

diff --git a/Engine/Source/Runtime/Engine/Private/WorldPartition/WorldPartition.cpp b/Engine/Source/Runtime/Engine/Private/WorldPartition/WorldPartition.cpp
--- a/Engine/Source/Runtime/Engine/Private/WorldPartition/WorldPartition.cpp
+++ b/Engine/Source/Runtime/Engine/Private/WorldPartition/WorldPartition.cpp
@@ -325,7 +325,7 @@ void UWorldPartition::OnEndPlay()
 {
 	check(IsMainWorldPartition());
 
-	FlushStreaming();
+	RuntimeHash->FlushStreaming();
 	RuntimeHash->OnEndPlay();
 
 	StreamingPolicy = nullptr;
@@ -1295,11 +1295,6 @@ void UWorldPartition::SavePerUserSettings()
 	}
 }
 
-void UWorldPartition::FlushStreaming()
-{
-	RuntimeHash->FlushStreaming();
-}
-
 void UWorldPartition::GenerateHLOD(ISourceControlHelper* SourceControlHelper, bool bCreateActorsOnly)
 {
 	RuntimeHash->GenerateHLOD(SourceControlHelper, bCreateActorsOnly);
diff --git a/Engine/Source/Runtime/Engine/Public/WorldPartition/WorldPartition.h b/Engine/Source/Runtime/Engine/Public/WorldPartition/WorldPartition.h
--- a/Engine/Source/Runtime/Engine/Public/WorldPartition/WorldPartition.h
+++ b/Engine/Source/Runtime/Engine/Public/WorldPartition/WorldPartition.h
@@ -94,8 +94,6 @@ private:
 
 	void SavePerUserSettings();
 		
-	void FlushStreaming();
-
 	void OnGCPostReachabilityAnalysis();
 
 	// PIE/Game Methods
