AActor::IsValidForDataLayer now calls AActor::SupportsDataLayer, and AActor::SupportsDataLayer is now private.

diff --git a/Engine/Source/Editor/DataLayerEditor/Private/DataLayer/DataLayerMode.cpp b/Engine/Source/Editor/DataLayerEditor/Private/DataLayer/DataLayerMode.cpp
--- a/Engine/Source/Editor/DataLayerEditor/Private/DataLayer/DataLayerMode.cpp
+++ b/Engine/Source/Editor/DataLayerEditor/Private/DataLayer/DataLayerMode.cpp
@@ -245,7 +245,7 @@ FSceneOutlinerDragValidationInfo FDataLayerMode::ValidateDrop(const ISceneOutlin
 	{
 		for (AActor* Actor : PayloadActors)
 		{
-			if (!Actor->SupportsDataLayer())
+			if (!Actor->IsValidForDataLayer())
 			{
 				return FSceneOutlinerDragValidationInfo(ESceneOutlinerDropCompatibility::IncompatibleGeneric, LOCTEXT("ActorCantBeAssignedToDataLayer", "Can't assign actors to Data Layer"));
 			}
diff --git a/Engine/Source/Runtime/Engine/Classes/Components/VolumetricCloudComponent.h b/Engine/Source/Runtime/Engine/Classes/Components/VolumetricCloudComponent.h
--- a/Engine/Source/Runtime/Engine/Classes/Components/VolumetricCloudComponent.h
+++ b/Engine/Source/Runtime/Engine/Classes/Components/VolumetricCloudComponent.h
@@ -176,13 +176,10 @@ class AVolumetricCloud : public AInfo
 	GENERATED_UCLASS_BODY()
 
 private:
-
 	UPROPERTY(BlueprintReadOnly, VisibleAnywhere, Category = Atmosphere, meta = (AllowPrivateAccess = "true"))
 	TObjectPtr<class UVolumetricCloudComponent> VolumetricCloudComponent;
 
-public:
 #if WITH_EDITOR
 	virtual bool SupportsDataLayer() const override { return true; }
 #endif
-
 };
diff --git a/Engine/Source/Runtime/Engine/Classes/GameFramework/Actor.h b/Engine/Source/Runtime/Engine/Classes/GameFramework/Actor.h
--- a/Engine/Source/Runtime/Engine/Classes/GameFramework/Actor.h
+++ b/Engine/Source/Runtime/Engine/Classes/GameFramework/Actor.h
@@ -1115,7 +1115,6 @@ public:
 	bool RemoveDataLayer(const UDataLayer* DataLayer);
 	bool RemoveAllDataLayers();
 	bool ContainsDataLayer(const UDataLayer* DataLayer) const;
-	virtual bool SupportsDataLayer() const { return true; }
 	bool HasDataLayers() const;
 	bool HasValidDataLayers() const;
 	bool HasAllDataLayers(const TArray<const UDataLayer*>& DataLayers) const;
@@ -1126,6 +1125,11 @@ public:
 	bool IsPropertyChangedAffectingDataLayers(FPropertyChangedEvent& PropertyChangedEvent) const;
 	bool IsValidForDataLayer() const;
 	void FixupDataLayers();
+
+private:
+	virtual bool SupportsDataLayer() const { return true; }
+
+public:
 #endif
 
 	//~=============================================================================
diff --git a/Engine/Source/Runtime/Engine/Classes/GameFramework/Info.h b/Engine/Source/Runtime/Engine/Classes/GameFramework/Info.h
--- a/Engine/Source/Runtime/Engine/Classes/GameFramework/Info.h
+++ b/Engine/Source/Runtime/Engine/Classes/GameFramework/Info.h
@@ -34,6 +34,8 @@ public:
 
 #if WITH_EDITOR
 	virtual EActorGridPlacement GetDefaultGridPlacement() const override { return EActorGridPlacement::AlwaysLoaded; }
+
+private:
 	virtual bool SupportsDataLayer() const override { return false; }
 #endif
 
diff --git a/Engine/Source/Runtime/Engine/Private/ActorEditor.cpp b/Engine/Source/Runtime/Engine/Private/ActorEditor.cpp
--- a/Engine/Source/Runtime/Engine/Private/ActorEditor.cpp
+++ b/Engine/Source/Runtime/Engine/Private/ActorEditor.cpp
@@ -81,7 +81,7 @@ bool AActor::CanEditChange(const FProperty* PropertyThatWillChange) const
 		return false;
 	}
 
-	if (bIsDataLayers && !SupportsDataLayer())
+	if (bIsDataLayers && !IsValidForDataLayer())
 	{
 		return false;
 	}
@@ -1095,7 +1095,7 @@ EDataValidationResult AActor::IsDataValid(TArray<FText>& ValidationErrors)
 bool AActor::AddDataLayer(const UDataLayer* DataLayer)
 {
 	bool bActorWasModified = false;
-	if (SupportsDataLayer() && DataLayer && !ContainsDataLayer(DataLayer))
+	if (IsValidForDataLayer() && DataLayer && !ContainsDataLayer(DataLayer))
 	{
 		if (!bActorWasModified)
 		{
@@ -1209,7 +1209,7 @@ void AActor::FixupDataLayers()
 {
 	if (!GetPackage()->HasAnyPackageFlags(PKG_PlayInEditor))
 	{
-		if (!SupportsDataLayer())
+		if (!IsValidForDataLayer())
 		{
 			DataLayers.Empty();
 		}
@@ -1267,19 +1267,32 @@ bool AActor::IsPropertyChangedAffectingDataLayers(FPropertyChangedEvent& Propert
 
 bool AActor::IsValidForDataLayer() const
 {
-	UWorld* World = GetWorld();
-	if (!World)
+	if (UWorld* World = GetWorld())
 	{
-		return false;
-	}
+		if (World->WorldType != EWorldType::Editor)
+		{
+			return false;
+		}
+
+		if (!UWorld::HasSubsystem<UWorldPartitionSubsystem>(World))
+		{
+			return false;
+		}
 
-	const bool bIsPartitionedActor = UWorld::HasSubsystem<UWorldPartitionSubsystem>(World);
-	const bool bIsInEditorWorld = World->WorldType == EWorldType::Editor;
-	const bool bIsBuilderBrush = FActorEditorUtils::IsABuilderBrush(this);
-	const bool bIsHidden = GetClass()->GetDefaultObject<AActor>()->bHiddenEd;
-	const bool bIsValid = !bIsHidden && !bIsBuilderBrush && bIsInEditorWorld && bIsPartitionedActor;
+		if (GetClass()->GetDefaultObject<AActor>()->bHiddenEd)
+		{
+			return false;
+		}
+
+		if (FActorEditorUtils::IsABuilderBrush(this))
+		{
+			return false;
+		}
 
-	return bIsValid;
+		return SupportsDataLayer();
+	}
+
+	return false;
 }
 
 // DataLayers (end)
diff --git a/Engine/Source/Runtime/Engine/Private/LevelInstance/LevelInstanceSubsystem.cpp b/Engine/Source/Runtime/Engine/Private/LevelInstance/LevelInstanceSubsystem.cpp
--- a/Engine/Source/Runtime/Engine/Private/LevelInstance/LevelInstanceSubsystem.cpp
+++ b/Engine/Source/Runtime/Engine/Private/LevelInstance/LevelInstanceSubsystem.cpp
@@ -964,7 +964,7 @@ void ULevelInstanceSubsystem::BreakLevelInstance_Impl(ALevelInstance* LevelInsta
 					}
 
 					// Apply the same data layer settings to the actors to move out
-					if (Actor->SupportsDataLayer() && Actor->IsValidForDataLayer())
+					if (Actor->IsValidForDataLayer())
 					{
 						for (const UDataLayer* DataLayer : LevelInstanceDataLayers)
 						{
diff --git a/Engine/Source/Runtime/Engine/Private/LightWeightInstanceSubsystem.cpp b/Engine/Source/Runtime/Engine/Private/LightWeightInstanceSubsystem.cpp
--- a/Engine/Source/Runtime/Engine/Private/LightWeightInstanceSubsystem.cpp
+++ b/Engine/Source/Runtime/Engine/Private/LightWeightInstanceSubsystem.cpp
@@ -108,7 +108,7 @@ ALightWeightInstanceManager* FLightWeightInstanceSubsystem::FindLightWeightInsta
 		if (Instance->GetRepresentedClass() == ActorClass)
 		{
 #if WITH_EDITOR
-			if (!Layer || (Instance->SupportsDataLayer() && Instance->ContainsDataLayer(Layer)))
+			if (!Layer || (Instance->IsValidForDataLayer() && Instance->ContainsDataLayer(Layer)))
 #endif // WITH_EDITOR
 			{
 				return Instance;
@@ -132,7 +132,7 @@ ALightWeightInstanceManager* FLightWeightInstanceSubsystem::FindOrAddLightWeight
 		if (InstanceManager->GetRepresentedClass() == ActorClass)
 		{
 #if WITH_EDITOR
-			if (!Layer || (InstanceManager->SupportsDataLayer() && InstanceManager->ContainsDataLayer(Layer)))
+			if (!Layer || (InstanceManager->IsValidForDataLayer() && InstanceManager->ContainsDataLayer(Layer)))
 #endif // WITH_EDITOR
 			{
 				return InstanceManager;
