Fixing subuv cutouts for HDR textures

diff --git a/Engine/Plugins/FX/Niagara/Source/Niagara/Private/NiagaraSpriteRendererProperties.cpp b/Engine/Plugins/FX/Niagara/Source/Niagara/Private/NiagaraSpriteRendererProperties.cpp
--- a/Engine/Plugins/FX/Niagara/Source/Niagara/Private/NiagaraSpriteRendererProperties.cpp
+++ b/Engine/Plugins/FX/Niagara/Source/Niagara/Private/NiagaraSpriteRendererProperties.cpp
@@ -546,6 +546,11 @@ void UNiagaraSpriteRendererProperties::GetRendererFeedback(const UNiagaraEmitter
 			}
 		}
 	}
+
+	if (CutoutTexture)
+	{
+		DerivedData.GetFeedback(CutoutTexture, (int32)SubImageSize.X, (int32)SubImageSize.Y, BoundingMode, AlphaThreshold, OpacitySourceMode, OutErrors, OutWarnings, OutInfo);
+	}
 }
 
 void UNiagaraSpriteRendererProperties::GetRendererTooltipWidgets(const FNiagaraEmitterInstance* InEmitter, TArray<TSharedPtr<SWidget>>& OutWidgets, TSharedPtr<FAssetThumbnailPool> InThumbnailPool) const
diff --git a/Engine/Source/Runtime/Engine/Classes/Particles/SubUVAnimation.h b/Engine/Source/Runtime/Engine/Classes/Particles/SubUVAnimation.h
--- a/Engine/Source/Runtime/Engine/Classes/Particles/SubUVAnimation.h
+++ b/Engine/Source/Runtime/Engine/Classes/Particles/SubUVAnimation.h
@@ -58,6 +58,8 @@ public:
 	TArray<FVector2D> BoundingGeometry;
 
 	static FString GetDDCKeyString(const FGuid& StateId, int32 SizeX, int32 SizeY, int32 Mode, float AlphaThreshold, int32 OpacitySourceMode);
+	static void GetFeedback(UTexture2D* SubUVTexture, int32 SubImages_Horizontal, int32 SubImages_Vertical, ESubUVBoundingVertexCount BoundingMode, float AlphaThreshold, EOpacitySourceMode OpacitySourceMode,
+		TArray<FText>& OutErrors, TArray<FText>& OutWarnings, TArray<FText>& OutInfo);
 	void Serialize(FStructuredArchive::FSlot Slot);
 	void Build(UTexture2D* SubUVTexture, int32 SubImages_Horizontal, int32 SubImages_Vertical, ESubUVBoundingVertexCount BoundingMode, float AlphaThreshold, EOpacitySourceMode OpacitySourceMode);
 };
diff --git a/Engine/Source/Runtime/Engine/Private/Particles/SubUVAnimation.cpp b/Engine/Source/Runtime/Engine/Private/Particles/SubUVAnimation.cpp
--- a/Engine/Source/Runtime/Engine/Private/Particles/SubUVAnimation.cpp
+++ b/Engine/Source/Runtime/Engine/Private/Particles/SubUVAnimation.cpp
@@ -9,6 +9,8 @@
 #include "DerivedDataCacheInterface.h"
 #include "ComponentReregisterContext.h"
 
+#define LOCTEXT_NAMESPACE "FSubUVDerivedData"
+
 #if ENABLE_COOK_STATS
 FCookStats::FDDCResourceUsageStats SubUVAnimationCookStats::UsageStats;
 FCookStatsManager::FAutoRegisterCallback SubUVAnimationCookStats::RegisterCookStats([](FCookStatsManager::AddStatFuncRef AddStat)
@@ -26,10 +28,13 @@ FString FSubUVDerivedData::GetDDCKeyString(const FGuid& StateId, int32 SizeX, in
 		KeyString += FString::Printf(TEXT("_%u"), OpacitySourceMode);
 	}
 	// adding v2 to the key after fixing color channel offsets
-	KeyString += TEXT("_V2");
+	// adding v3 to the key after allowing other formats
+	KeyString += TEXT("_V3");
 	return FDerivedDataCacheInterface::BuildCacheKey(TEXT("SUBUV_"), SUBUV_DERIVEDDATA_VER, *KeyString);
 }
 
+
+
 void FSubUVDerivedData::Serialize(FStructuredArchive::FSlot Slot)
 {
 	Slot << BoundingGeometry;
@@ -444,36 +449,156 @@ FIntPoint GNeighbors[] =
 	FIntPoint(-1, 0)
 };
 
-uint8 ComputeOpacityValue(const uint8* BGRA, EOpacitySourceMode OpacitySourceMode)
+
+bool IsSupportedFormat(ETextureSourceFormat SrcFormat)
 {
-	if (OpacitySourceMode == OSM_Alpha)
+	return SrcFormat == TSF_BGRA8 || SrcFormat == TSF_RGBA16 || SrcFormat == TSF_RGBA16F;
+}
+
+uint32 GetByteSizePerPixel(ETextureSourceFormat SrcFormat)
+{
+	if (SrcFormat == TSF_BGRA8)
 	{
-		return *(BGRA + 3);
+		return 4;
 	}
-	else if (OpacitySourceMode == OSM_RedChannel)
+	else if (SrcFormat == TSF_RGBA16)
 	{
-		return *(BGRA + 2);
+		return 8;
 	}
-	else if (OpacitySourceMode == OSM_GreenChannel)
+	else if (SrcFormat == TSF_RGBA16F)
 	{
-		return *(BGRA + 1);
+		return sizeof(FFloat16) * 4;
 	}
-	else if (OpacitySourceMode == OSM_BlueChannel)
+	return 0;
+}
+
+bool ComputeOpacityValue(const uint8* BGRA, int32 x, int32 y, int32 TextureSizeX,  EOpacitySourceMode OpacitySourceMode, ETextureSourceFormat SrcFormat, float AlphaThresholdF, uint8 AlphaThreshold8, uint16 AlphaThreshold16)
+{
+	int32 Offset = (y * TextureSizeX + x) * GetByteSizePerPixel(SrcFormat);
+
+	if (SrcFormat == TSF_BGRA8)
 	{
-		return *(BGRA + 0);
+		const uint8* BGRA8 = (const uint8*)(BGRA + Offset);
+		uint32 Opacity = 255;
+		if (OpacitySourceMode == OSM_Alpha)
+		{
+			Opacity = *(BGRA8 + 3);
+		}
+		else if (OpacitySourceMode == OSM_RedChannel)
+		{
+			Opacity = *(BGRA8 + 2);
+		}
+		else if (OpacitySourceMode == OSM_GreenChannel)
+		{
+			Opacity = *(BGRA8 + 1);
+		}
+		else if (OpacitySourceMode == OSM_BlueChannel)
+		{
+			Opacity = *(BGRA8 + 0);
+		}
+		else
+		{
+			uint32 R = *(BGRA8 + 0);
+			uint32 G = *(BGRA8 + 1);
+			uint32 B = *(BGRA8 + 2);
+
+			Opacity = ((R + G + B) / 3);
+		}
+		return Opacity > AlphaThreshold8;
 	}
-	else
+	else if (SrcFormat == TSF_RGBA16)
 	{
-		uint32 R = *BGRA;
-		uint32 G = *(BGRA + 1);
-		uint32 B = *(BGRA + 2);
+		const uint16* RGBA16 = (const uint16*)(BGRA + Offset);
+
+		uint32 Opacity = 65535;
+
+		if (OpacitySourceMode == OSM_Alpha)
+		{
+			Opacity = *(RGBA16 + 3) ;
+		}
+		else if (OpacitySourceMode == OSM_RedChannel)
+		{
+			Opacity = *(RGBA16 + 0) ;
+		}
+		else if (OpacitySourceMode == OSM_GreenChannel)
+		{
+			Opacity = *(RGBA16 + 1) ;
+		}
+		else if (OpacitySourceMode == OSM_BlueChannel)
+		{
+			Opacity = *(RGBA16 + 2);
+		}
+		else
+		{
+			uint32 R = *(RGBA16 + 0);
+			uint32 G = *(RGBA16 + 1);
+			uint32 B = *(RGBA16 + 2);
 
-		return (uint8)((R + G + B) / 3);
+			Opacity = ((R + G + B) / 3) ;
+		}
+
+		return Opacity > AlphaThreshold16;
+	}
+	else if (SrcFormat == TSF_RGBA16F)
+	{
+		const FFloat16* RGBAf = (const FFloat16 *)( BGRA + Offset);
+
+		float Opacity = 1.0f;
+		if (OpacitySourceMode == OSM_Alpha)
+		{
+			Opacity = (*(RGBAf + 3)).GetFloat();
+		}
+		else if (OpacitySourceMode == OSM_RedChannel)
+		{
+			Opacity = (*(RGBAf + 0)).GetFloat();
+		}
+		else if (OpacitySourceMode == OSM_GreenChannel)
+		{
+			Opacity = (*(RGBAf + 1)).GetFloat();
+		}
+		else if (OpacitySourceMode == OSM_BlueChannel)
+		{
+			Opacity = (*(RGBAf + 2)).GetFloat();
+		}
+		else
+		{
+			float R = (*RGBAf).GetFloat();
+			float G = (*(RGBAf + 1)).GetFloat();
+			float B = (*(RGBAf + 2)).GetFloat();
+
+			Opacity = (R + G + B) / 3.0f;
+		}
+
+		return Opacity > AlphaThresholdF;
 	}
+
+	return true;
+}
+
+
+void FSubUVDerivedData::GetFeedback(UTexture2D* SubUVTexture, int32 SubImages_Horizontal, int32 SubImages_Vertical, ESubUVBoundingVertexCount BoundingMode, float AlphaThreshold, EOpacitySourceMode OpacitySourceMode,
+	TArray<FText>& OutErrors, TArray<FText>& OutWarnings, TArray<FText>& OutInfo)
+{
+
+#if WITH_EDITORONLY_DATA
+	if (SubUVTexture)
+	{
+		ETextureSourceFormat SourceFormat = SubUVTexture->Source.GetFormat();
+
+		if (!IsSupportedFormat(SourceFormat))
+		{
+			OutErrors.Add(LOCTEXT("CutoutNotSupported", "Image is not in a supported format for cutouts (BGRA8, RGBA16, RGBA16F). It will be ignored."));
+		} 
+		if (SubUVTexture->Source.GetNumMips() == 0)
+		{
+			OutErrors.Add(LOCTEXT("CutoutNoMipNotSupported", "Image does not have a 0th Mip Level."));
+		}
+	}
+#endif
 }
 
 /** Counts how many neighbors have non-zero alpha. */
-int32 ComputeNeighborCount(int32 X, int32 Y, int32 GlobalX, int32 GlobalY, int32 SubImageSizeX, int32 SubImageSizeY, int32 TextureSizeX, const TArray64<uint8>& MipData, uint8 AlphaThresholdByte, EOpacitySourceMode OpacitySourceMode)
+int32 ComputeNeighborCount(int32 X, int32 Y, int32 GlobalX, int32 GlobalY, int32 SubImageSizeX, int32 SubImageSizeY, int32 TextureSizeX, const TArray64<uint8>& MipData,  EOpacitySourceMode OpacitySourceMode, ETextureSourceFormat SourceFormat, float AlphaThresholdF, uint8 AlphaThreshold8, uint16 AlphaThreshold16)
 {
 	int32 NeighborCount = 0;
 
@@ -485,9 +610,9 @@ int32 ComputeNeighborCount(int32 X, int32 Y, int32 GlobalX, int32 GlobalY, int32
 		if (NeighborX >= 0 && NeighborX < SubImageSizeX 
 			&& NeighborY >= 0 && NeighborY < SubImageSizeY)
 		{
-			uint8 NeighborAlphaValue = ComputeOpacityValue(&MipData[((GlobalY + GNeighbors[NeighborIndex].Y) * TextureSizeX + GlobalX + GNeighbors[NeighborIndex].X) * 4], OpacitySourceMode);
+			bool bNeighborPasses = ComputeOpacityValue(MipData.GetData(), GlobalX + GNeighbors[NeighborIndex].X, GlobalY + GNeighbors[NeighborIndex].Y, TextureSizeX, OpacitySourceMode, SourceFormat, AlphaThresholdF, AlphaThreshold8, AlphaThreshold16);
 
-			if (NeighborAlphaValue > AlphaThresholdByte)
+			if (bNeighborPasses)
 			{
 				NeighborCount++;
 			}
@@ -530,7 +655,8 @@ void FSubUVDerivedData::Build(UTexture2D* SubUVTexture, int32 SubImages_Horizont
 	if (SubUVTexture)
 	{
 		TArray64<uint8> MipData;
-		bool bSuccess = SubUVTexture->Source.GetFormat() == TSF_BGRA8 && SubUVTexture->Source.GetMipData(MipData, 0);
+		ETextureSourceFormat SourceFormat = SubUVTexture->Source.GetFormat();
+		bool bSuccess = IsSupportedFormat(SourceFormat) && SubUVTexture->Source.GetMipData(MipData, 0);
 
 		const int32 TextureSizeX = SubUVTexture->Source.GetSizeX();
 		const int32 TextureSizeY = SubUVTexture->Source.GetSizeY();
@@ -540,9 +666,10 @@ void FSubUVDerivedData::Build(UTexture2D* SubUVTexture, int32 SubImages_Horizont
 		const int32 SubImageSizeYFloat = TextureSizeY / (float)SubImages_Vertical;
 
 		const int32 NumSubImages = SubImages_Horizontal * SubImages_Vertical;
-		const uint8 AlphaThresholdByte = FMath::Clamp(FMath::TruncToInt(AlphaThreshold * 255.0f), 0, 255);
+		const uint8 AlphaThreshold8 = FMath::Clamp(FMath::TruncToInt(AlphaThreshold * 255.0f), 0, 255);
+		const uint16 AlphaThreshold16 = FMath::Clamp(FMath::TruncToInt(AlphaThreshold * 65535.0f), 0, 65535.0f);
 
-		check(!bSuccess || MipData.Num() == TextureSizeX * TextureSizeY * 4);
+		check(!bSuccess || MipData.Num() == TextureSizeX * TextureSizeY * GetByteSizePerPixel(SourceFormat));
 
 		const int32 TargetNumBoundingVertices = BoundingMode == BVC_FourVertices ? 4 : 8;
 
@@ -573,13 +700,13 @@ void FSubUVDerivedData::Build(UTexture2D* SubUVTexture, int32 SubImages_Horizont
 						{
 							int32 GlobalX = FMath::RoundToInt(SubImageX * SubImageSizeXFloat) + X;
 							int32 NextGlobalX = FMath::RoundToInt(NextSubImageX * SubImageSizeXFloat) + X;
-							uint8 AlphaValue = ComputeOpacityValue(&MipData[(GlobalY * TextureSizeX + GlobalX) * 4], OpacitySourceMode);
-							uint8 NextAlphaValue = ComputeOpacityValue(&MipData[(NextGlobalY * TextureSizeX + NextGlobalX) * 4], OpacitySourceMode);
+							bool AlphaValuePasses = ComputeOpacityValue(MipData.GetData(), GlobalX, GlobalY, TextureSizeX, OpacitySourceMode, SourceFormat, AlphaThreshold, AlphaThreshold8, AlphaThreshold16);
+							bool NextAlphaValuePasses = ComputeOpacityValue(MipData.GetData(), NextGlobalX, NextGlobalY, TextureSizeX, OpacitySourceMode, SourceFormat, AlphaThreshold, AlphaThreshold8, AlphaThreshold16);
 
-							if (AlphaValue > AlphaThresholdByte || NextAlphaValue > AlphaThresholdByte)
+							if (AlphaValuePasses || NextAlphaValuePasses)
 							{
-								int32 NeighborCount = AlphaValue > AlphaThresholdByte ? ComputeNeighborCount(X, Y, GlobalX, GlobalY, SubImageSizeX, SubImageSizeY, TextureSizeX, MipData, AlphaThresholdByte, OpacitySourceMode) : 8;
-								int32 NextNeighborCount = NextAlphaValue > AlphaThresholdByte ? ComputeNeighborCount(X, Y, NextGlobalX, NextGlobalY, SubImageSizeX, SubImageSizeY, TextureSizeX, MipData, AlphaThresholdByte, OpacitySourceMode) : 8;
+								int32 NeighborCount = AlphaValuePasses ? ComputeNeighborCount(X, Y, GlobalX, GlobalY, SubImageSizeX, SubImageSizeY, TextureSizeX, MipData, OpacitySourceMode, SourceFormat, AlphaThreshold, AlphaThreshold8, AlphaThreshold16) : 8;
+								int32 NextNeighborCount = NextAlphaValuePasses ? ComputeNeighborCount(X, Y, NextGlobalX, NextGlobalY, SubImageSizeX, SubImageSizeY, TextureSizeX, MipData, OpacitySourceMode, SourceFormat, AlphaThreshold, AlphaThreshold8, AlphaThreshold16) : 8;
 
 								// Points with non-zero alpha that have 5 or more filled in neighbors must be in the solid interior
 								if (NeighborCount < 5 || NextNeighborCount < 5)
@@ -663,3 +790,4 @@ void FSubUVDerivedData::Build(UTexture2D* SubUVTexture, int32 SubImages_Horizont
 	}
 #endif
 }
+#undef LOCTEXT_NAMESPACE
