Add breakdown memory stats to skin cache

diff --git a/Engine/Source/Runtime/Engine/Private/GPUSkinCache.cpp b/Engine/Source/Runtime/Engine/Private/GPUSkinCache.cpp
--- a/Engine/Source/Runtime/Engine/Private/GPUSkinCache.cpp
+++ b/Engine/Source/Runtime/Engine/Private/GPUSkinCache.cpp
@@ -2389,9 +2389,36 @@ void FGPUSkinCache::PrintMemorySummary() const
 		FGPUSkinCacheEntry* Entry = Entries[i];
 		if (Entry)
 		{
+			FString RecomputeTangentSections = TEXT("");
+			for (int32 DispatchIdx = 0; DispatchIdx < Entry->DispatchData.Num(); ++DispatchIdx)
+			{
+				const FGPUSkinCacheEntry::FSectionDispatchData& Data = Entry->DispatchData[DispatchIdx];
+				if (Data.IndexBuffer)
+				{
+					if (RecomputeTangentSections.IsEmpty())
+					{
+						RecomputeTangentSections = TEXT("[Section]") + FString::FromInt(Data.SectionIndex);
+					}
+					else
+					{
+						RecomputeTangentSections = RecomputeTangentSections + TEXT("/") + FString::FromInt(Data.SectionIndex);
+					}
+				}
+			}
+			if (RecomputeTangentSections.IsEmpty())
+			{
+				RecomputeTangentSections = TEXT("Off");
+			}
+
 			const FString RayTracingTag = (Entry->Mode == EGPUSkinCacheEntryMode::RayTracing ? TEXT("[RT]") : TEXT(""));
 			uint64 MemInBytes = Entry->PositionAllocation ? Entry->PositionAllocation->GetNumBytes() : 0;
-			UE_LOG(LogSkinCache, Display, TEXT("   SkinCacheEntry_%d: %sMesh=%s, LOD=%d, Mem=%.3fKB"), i, *RayTracingTag, *GetSkeletalMeshObjectName(Entry->GPUSkin), Entry->LOD, MemInBytes / 1024.f);
+			uint64 TangentsInBytes = (Entry->PositionAllocation && Entry->PositionAllocation->GetTangentBuffer()) ? Entry->PositionAllocation->GetTangentBuffer()->Buffer.NumBytes : 0;
+			uint64 IntermediateTangentsInBytes = (Entry->PositionAllocation && Entry->PositionAllocation->GetIntermediateTangentBuffer()) ? Entry->PositionAllocation->GetIntermediateTangentBuffer()->Buffer.NumBytes : 0;
+			uint64 IntermediateAccumulatedTangentsInBytes = (Entry->PositionAllocation && Entry->PositionAllocation->GetIntermediateAccumulatedTangentBuffer()) ? Entry->PositionAllocation->GetIntermediateAccumulatedTangentBuffer()->Buffer.NumBytes : 0;
+
+			UE_LOG(LogSkinCache, Display, TEXT("   SkinCacheEntry_%d: %sMesh=%s, LOD=%d, RecomputeTangent=%s, Mem=%.3fKB (Tangents=%.3fKB, InterTangents=%.3fKB, InterAccumTangents=%.3fKB)"), 
+					i, *RayTracingTag, *GetSkeletalMeshObjectName(Entry->GPUSkin), Entry->LOD, *RecomputeTangentSections, 
+					MemInBytes / 1024.f, TangentsInBytes / 1024.f, IntermediateTangentsInBytes / 1024.f, IntermediateAccumulatedTangentsInBytes / 1024.f);
 			TotalMemInBytes += MemInBytes;
 		}
 	}
