Print portion of total time per shader type.

- Also fix stats not being dumped on cooker exit.

diff --git a/Engine/Source/Runtime/Engine/Private/ShaderCompiler/ShaderCompiler.cpp b/Engine/Source/Runtime/Engine/Private/ShaderCompiler/ShaderCompiler.cpp
--- a/Engine/Source/Runtime/Engine/Private/ShaderCompiler/ShaderCompiler.cpp
+++ b/Engine/Source/Runtime/Engine/Private/ShaderCompiler/ShaderCompiler.cpp
@@ -2752,14 +2752,13 @@ void FShaderCompilerStats::WriteStatSummary()
 	if (ShaderTimings.Num())
 	{
 		// calculate effective parallelization (total time needed to compile all shaders divided by actual wall clock time spent processing at least 1 shader)
+		double TotalTimeForAllShaders = 0;
+		for (TMap<FString, FShaderTimings>::TConstIterator Iter(ShaderTimings); Iter; ++Iter)
+		{
+			TotalTimeForAllShaders += Iter.Value().TotalCompileTime;
+		}
 		if (TotalTimeAtLeastOneJobWasInFlight > 0.0)
 		{
-			double TotalTimeForAllShaders = 0;
-			for (TMap<FString, FShaderTimings>::TConstIterator Iter(ShaderTimings); Iter; ++Iter)
-			{
-				TotalTimeForAllShaders += Iter.Value().TotalCompileTime;
-			}
-
 			double EffectiveParallelization = TotalTimeForAllShaders / TotalTimeAtLeastOneJobWasInFlight;
 			UE_LOG(LogShaderCompilers, Display, TEXT("Effective parallelization: %.2f (times faster than compiling all shaders on one thread)."), EffectiveParallelization);
 		}
@@ -2793,7 +2792,8 @@ void FShaderCompilerStats::WriteStatSummary()
 		{
 			const FShaderTimings& Timings = Iter.Value();
 
-			UE_LOG(LogShaderCompilers, Display, TEXT("%60s (compiled %4d times, average %4.2f sec, max %4.2f sec, min %4.2f sec)"), *Iter.Key(), Timings.NumCompiled, Timings.AverageCompileTime, Timings.MaxCompileTime, Timings.MinCompileTime);
+			UE_LOG(LogShaderCompilers, Display, TEXT("%60s - %.2f%% of total time (compiled %4d times, average %4.2f sec, max %4.2f sec, min %4.2f sec)"), 
+				*Iter.Key(), 100.0 * Timings.TotalCompileTime / TotalTimeForAllShaders, Timings.NumCompiled, Timings.AverageCompileTime, Timings.MaxCompileTime, Timings.MinCompileTime);
 			if (++Idx >= MaxShadersToPrint)
 			{
 				break;
@@ -3318,7 +3318,7 @@ FShaderCompilingManager::~FShaderCompilingManager()
 		return;
 	}
 
-	PrintStats();
+	PrintStats(true);
 
 	for (const auto& Thread : Threads)
 	{
