Maybe fix texture cube array on platforms using packed mip tails

diff --git a/Engine/Source/Runtime/Engine/Private/TextureCubeArray.cpp b/Engine/Source/Runtime/Engine/Private/TextureCubeArray.cpp
--- a/Engine/Source/Runtime/Engine/Private/TextureCubeArray.cpp
+++ b/Engine/Source/Runtime/Engine/Private/TextureCubeArray.cpp
@@ -120,9 +120,8 @@ public:
 		RHIUpdateTextureReference(Owner->TextureReference.TextureReferenceRHI, TextureRHI);
 
 		// Read the mip-levels into the RHI texture.
-		const int32 NumMips = Owner->GetNumMips();
-		check(NumMips < MAX_TEXTURE_MIP_COUNT);
-		for (int32 MipIndex = 0; MipIndex < NumMips; MipIndex++)
+		const int32 FirstMipTailIndex = Owner->GetNumMips() - FMath::Max(1, Owner->PlatformData->GetNumMipsInTail());
+		for (int32 MipIndex = 0; MipIndex <= FirstMipTailIndex; MipIndex++)
 		{
 			for (uint32 ArrayIndex = 0u; ArrayIndex < ArraySize; ++ArrayIndex)
 			{
