Fixed crash in FTextureCubeArrayResource::GetData() when the destination pitch is 0

diff --git a/Engine/Source/Runtime/Engine/Private/TextureCubeArray.cpp b/Engine/Source/Runtime/Engine/Private/TextureCubeArray.cpp
--- a/Engine/Source/Runtime/Engine/Private/TextureCubeArray.cpp
+++ b/Engine/Source/Runtime/Engine/Private/TextureCubeArray.cpp
@@ -235,7 +235,7 @@ private:
 		// runtime block size checking, conversion, or the like
 		if (DestPitch == 0)
 		{
-			FMemory::Memcpy(Dest, MipData[Index], Owner->PlatformData->Mips[MipIndex].BulkData.GetBulkDataSize() / 6);
+			FMemory::Memcpy(Dest, MipData[Index], Owner->PlatformData->Mips[MipIndex].BulkData.GetBulkDataSize() / Owner->GetNumSlices());
 		}
 		else
 		{
