WorldPartitionVolume: revert back to using GetActorBounds to get the volume loading bounds.

diff --git a/Engine/Source/Runtime/Engine/Private/WorldPartition/WorldPartitionVolume.cpp b/Engine/Source/Runtime/Engine/Private/WorldPartition/WorldPartitionVolume.cpp
--- a/Engine/Source/Runtime/Engine/Private/WorldPartition/WorldPartitionVolume.cpp
+++ b/Engine/Source/Runtime/Engine/Private/WorldPartition/WorldPartitionVolume.cpp
@@ -18,7 +18,7 @@ void AWorldPartitionVolume::LoadIntersectingCells(bool bIsFromUserChange)
 {
 	if (UWorldPartition* WorldPartition = GetWorld()->GetWorldPartition())
 	{
-		WorldPartition->LoadEditorCells(GetStreamingBounds(), bIsFromUserChange);
+		WorldPartition->LoadEditorCells(GetIntersectingBounds(), bIsFromUserChange);
 	}
 }
 
@@ -26,9 +26,17 @@ void AWorldPartitionVolume::UnloadIntersectingCells(bool bIsFromUserChange)
 {
 	if (UWorldPartition* WorldPartition = GetWorld()->GetWorldPartition())
 	{
-		WorldPartition->UnloadEditorCells(GetStreamingBounds(), bIsFromUserChange);
+		WorldPartition->UnloadEditorCells(GetIntersectingBounds(), bIsFromUserChange);
 	}
 }
+
+FBox AWorldPartitionVolume::GetIntersectingBounds() const
+{
+	FVector Origin;
+	FVector Extent;
+	GetActorBounds(false, Origin, Extent);
+	return FBox(Origin - Extent, Origin + Extent);
+}
 #endif
 
 #undef LOCTEXT_NAMESPACE
\ No newline at end of file
diff --git a/Engine/Source/Runtime/Engine/Public/WorldPartition/WorldPartitionVolume.h b/Engine/Source/Runtime/Engine/Public/WorldPartition/WorldPartitionVolume.h
--- a/Engine/Source/Runtime/Engine/Public/WorldPartition/WorldPartitionVolume.h
+++ b/Engine/Source/Runtime/Engine/Public/WorldPartition/WorldPartitionVolume.h
@@ -24,5 +24,8 @@ public:
 	void LoadIntersectingCells(bool bIsFromUserChange);
 	void UnloadIntersectingCells(bool bIsFromUserChange);
 	virtual EActorGridPlacement GetDefaultGridPlacement() const override { return EActorGridPlacement::AlwaysLoaded; }
+
+private:
+	FBox GetIntersectingBounds() const;
 #endif
 };
\ No newline at end of file
