Add DistanceFieldSize tag to StaticMesh

diff --git a/Engine/Source/Runtime/Engine/Private/StaticMesh.cpp b/Engine/Source/Runtime/Engine/Private/StaticMesh.cpp
--- a/Engine/Source/Runtime/Engine/Private/StaticMesh.cpp
+++ b/Engine/Source/Runtime/Engine/Private/StaticMesh.cpp
@@ -4032,6 +4032,16 @@ void UStaticMesh::GetAssetRegistryTags(TArray<FAssetRegistryTag>& OutTags) const
 	int32 NumNaniteTriangles = GetNumNaniteTriangles();
 	int32 NumNaniteVertices = GetNumNaniteVertices();
 
+	int32 DistanceFieldSize = 0;
+
+	if (GetRenderData() && GetRenderData()->LODResources.Num() > 0 && GetRenderData()->LODResources[0].DistanceFieldData != nullptr)
+	{
+		const FDistanceFieldVolumeData& VolumeData = *(GetRenderData()->LODResources[0].DistanceFieldData);
+
+		DistanceFieldSize += VolumeData.GetResourceSizeBytes();
+		DistanceFieldSize += VolumeData.StreamableMips.GetBulkDataSize();
+	}
+
 	OutTags.Add(FAssetRegistryTag("NaniteTriangles", FString::FromInt(NumNaniteTriangles), FAssetRegistryTag::TT_Numerical));
 	OutTags.Add(FAssetRegistryTag("NaniteVertices", FString::FromInt(NumNaniteVertices), FAssetRegistryTag::TT_Numerical));
 	OutTags.Add(FAssetRegistryTag("Triangles", FString::FromInt(NumTriangles), FAssetRegistryTag::TT_Numerical) );
@@ -4045,6 +4055,7 @@ void UStaticMesh::GetAssetRegistryTags(TArray<FAssetRegistryTag>& OutTags) const
 	OutTags.Add(FAssetRegistryTag("SectionsWithCollision", FString::FromInt(NumSectionsWithCollision), FAssetRegistryTag::TT_Numerical));
 	OutTags.Add(FAssetRegistryTag("DefaultCollision", DefaultCollisionName.ToString(), FAssetRegistryTag::TT_Alphabetical));
 	OutTags.Add(FAssetRegistryTag("CollisionComplexity", ComplexityString, FAssetRegistryTag::TT_Alphabetical));
+	OutTags.Add(FAssetRegistryTag("DistanceFieldSize", FString::FromInt(DistanceFieldSize), FAssetRegistryTag::TT_Numerical, FAssetRegistryTag::TD_Memory));
 
 	Super::GetAssetRegistryTags(OutTags);
 }
