Renamed SceneComponent.bComputeBoundsOnceDuringCook to bComputeBoundsOnceForGame, and implemented it for all game builds.
With this change, all HLOD scene components will compute their bounds once during cook or once for PIE/Game builds.

diff --git a/Engine/Source/Runtime/CoreUObject/Private/UObject/CoreRedirects.cpp b/Engine/Source/Runtime/CoreUObject/Private/UObject/CoreRedirects.cpp
--- a/Engine/Source/Runtime/CoreUObject/Private/UObject/CoreRedirects.cpp
+++ b/Engine/Source/Runtime/CoreUObject/Private/UObject/CoreRedirects.cpp
@@ -1452,6 +1452,7 @@ static void RegisterNativeRedirects40(TArray<FCoreRedirect>& Redirects)
 	PROPERTY_REDIRECT("SceneComponent.ModifyFrequency", "SceneComponent.Mobility");
 	PROPERTY_REDIRECT("SceneComponent.RelativeTranslation", "SceneComponent.RelativeLocation");
 	PROPERTY_REDIRECT("SceneComponent.bAbsoluteTranslation", "SceneComponent.bAbsoluteLocation");
+	PROPERTY_REDIRECT("SceneComponent.bComputeBoundsOnceDuringCook", "SceneComponent.bComputeBoundsOnceForGame");
 	PROPERTY_REDIRECT("SkeletalMeshComponent.AnimationBlueprint", "SkeletalMeshComponent.AnimBlueprintGeneratedClass");
 	PROPERTY_REDIRECT("SlateBrush.TextureName", "SlateBrush.ResourceName");
 	PROPERTY_REDIRECT("SlateBrush.TextureObject", "SlateBrush.ResourceObject");
diff --git a/Engine/Source/Runtime/Engine/Classes/Components/SceneComponent.h b/Engine/Source/Runtime/Engine/Classes/Components/SceneComponent.h
--- a/Engine/Source/Runtime/Engine/Classes/Components/SceneComponent.h
+++ b/Engine/Source/Runtime/Engine/Classes/Components/SceneComponent.h
@@ -234,9 +234,13 @@ public:
 	UPROPERTY()
 	uint8 bComputeFastLocalBounds : 1;
 
-	/** If true, this component will cache its bounds during cooking and never recompute it for cooked builds. This is for components that are known to be static. */
+	/** If true, this component will cache its bounds during cooking or in PIE and never recompute it again. This is for components that are known to be static. */
 	UPROPERTY()
-	uint8 bComputeBoundsOnceDuringCook : 1;
+	uint8 bComputeBoundsOnceForGame : 1;
+
+	/** If true, this component has already cached its bounds during cooking or in PIE and will never recompute it again.  */
+	UPROPERTY()
+	uint8 bComputedBoundsOnceForGame : 1;
 
 	/** Get the current local bounds of the component */
 	FBoxSphereBounds GetLocalBounds() const;
diff --git a/Engine/Source/Runtime/Engine/Private/Components/SceneComponent.cpp b/Engine/Source/Runtime/Engine/Private/Components/SceneComponent.cpp
--- a/Engine/Source/Runtime/Engine/Private/Components/SceneComponent.cpp
+++ b/Engine/Source/Runtime/Engine/Private/Components/SceneComponent.cpp
@@ -1216,10 +1216,13 @@ void USceneComponent::UpdateBounds()
 	else
 	{
 		// Calculate new bounds
-		if (!bComputeBoundsOnceDuringCook || !FPlatformProperties::RequiresCookedData())
+		const UWorld* const World = GetWorld();
+		const bool bIsGameWorld = World && World->IsGameWorld();
+		if (!bComputeBoundsOnceForGame || !bIsGameWorld || !bComputedBoundsOnceForGame)
 		{
 			SCOPE_CYCLE_COUNTER(STAT_ComponentCalcBounds);
 			Bounds = CalcBounds(GetComponentTransform());
+			bComputedBoundsOnceForGame = (bIsGameWorld || IsRunningCookCommandlet()) && bComputeBoundsOnceForGame;
 		}
 	}
 
@@ -3324,11 +3327,11 @@ void USceneComponent::Serialize(FArchive& Ar)
 
 	Super::Serialize(Ar);
 
-	if (bComputeBoundsOnceDuringCook)
+	if (bComputeBoundsOnceForGame)
 	{
 		if(Ar.CustomVer(FUE5PrivateFrostyStreamObjectVersion::GUID) >= FUE5PrivateFrostyStreamObjectVersion::SerializeSceneComponentStaticBounds)
 		{
-			bool bIsCooked = Ar.IsCooking();
+			bool bIsCooked = bComputedBoundsOnceForGame && Ar.IsCooking();
 			Ar << bIsCooked;
 
 			if (bIsCooked)
diff --git a/Engine/Source/Runtime/Engine/Private/WorldPartition/HLOD/HLODBuilder.cpp b/Engine/Source/Runtime/Engine/Private/WorldPartition/HLOD/HLODBuilder.cpp
--- a/Engine/Source/Runtime/Engine/Private/WorldPartition/HLOD/HLODBuilder.cpp
+++ b/Engine/Source/Runtime/Engine/Private/WorldPartition/HLOD/HLODBuilder.cpp
@@ -64,7 +64,7 @@ void UHLODBuilder::Build(AWorldPartitionHLOD* InHLODActor, const UHLODLayer* InH
 
 		// Enable optimizations
 		HLODPrimitive->bComputeFastLocalBounds = true;
-		HLODPrimitive->bComputeBoundsOnceDuringCook = true;
+		HLODPrimitive->bComputeBoundsOnceForGame = true;
 
 		if (InHLODLayer->GetLayerType() != EHLODLayerType::Instancing)
 		{
