LevelInstance: Fix reset WorldAsset not resetting

diff --git a/Engine/Source/Runtime/Engine/Private/LevelInstance/LevelInstanceActor.cpp b/Engine/Source/Runtime/Engine/Private/LevelInstance/LevelInstanceActor.cpp
--- a/Engine/Source/Runtime/Engine/Private/LevelInstance/LevelInstanceActor.cpp
+++ b/Engine/Source/Runtime/Engine/Private/LevelInstance/LevelInstanceActor.cpp
@@ -420,6 +420,12 @@ bool ALevelInstance::CheckForLoop(TSoftObjectPtr<UWorld> InLevelInstance, TArray
 
 bool ALevelInstance::CanSetValue(TSoftObjectPtr<UWorld> InLevelInstance, FString* Reason) const
 {
+	// Set to null is valid
+	if (InLevelInstance.IsNull())
+	{
+		return true;
+	}
+
 	FString PackageName;
 	if (!FPackageName::DoesPackageExist(InLevelInstance.GetLongPackageName(), nullptr, nullptr))
 	{
