World Partition
- Rely on GetIsInBlockTillLevelStreamingCompleted instead of bMatchStarted to get the maximum number of cells to load.
- Add objects unhashing and streaming performance status in world partition debug display.
- Always simulate good streaming performances when the app is using a fixed timestep.

diff --git a/Engine/Source/Runtime/Engine/Private/WorldPartition/WorldPartition.cpp b/Engine/Source/Runtime/Engine/Private/WorldPartition/WorldPartition.cpp
--- a/Engine/Source/Runtime/Engine/Private/WorldPartition/WorldPartition.cpp
+++ b/Engine/Source/Runtime/Engine/Private/WorldPartition/WorldPartition.cpp
@@ -1228,6 +1228,11 @@ void UWorldPartition::DrawStreamingStatusLegend(class UCanvas* Canvas, FVector2D
 	StreamingPolicy->DrawStreamingStatusLegend(Canvas, Offset);
 }
 
+EWorldPartitionStreamingPerformance UWorldPartition::GetStreamingPerformance() const
+{
+	return StreamingPolicy->GetStreamingPerformance();
+}
+
 #if WITH_EDITOR
 void UWorldPartition::DrawRuntimeHashPreview()
 {
diff --git a/Engine/Source/Runtime/Engine/Private/WorldPartition/WorldPartitionRuntimeHash.cpp b/Engine/Source/Runtime/Engine/Private/WorldPartition/WorldPartitionRuntimeHash.cpp
--- a/Engine/Source/Runtime/Engine/Private/WorldPartition/WorldPartitionRuntimeHash.cpp
+++ b/Engine/Source/Runtime/Engine/Private/WorldPartition/WorldPartitionRuntimeHash.cpp
@@ -282,9 +282,15 @@ void UWorldPartitionRuntimeHash::SortStreamingCellsByImportance(const TSet<const
 
 EWorldPartitionStreamingPerformance UWorldPartitionRuntimeHash::GetStreamingPerformance(const TSet<const UWorldPartitionRuntimeCell*>& CellsToActivate) const
 {
+	bool bForceGoodStreamingPerformance = !GetWorld()->bMatchStarted;
+
+#if WITH_EDITOR
+	bForceGoodStreamingPerformance |= FApp::UseFixedTimeStep();
+#endif
+
 	EWorldPartitionStreamingPerformance StreamingPerformance = EWorldPartitionStreamingPerformance::Good;
 
-	if (!CellsToActivate.IsEmpty() && GetWorld()->bMatchStarted)
+	if (!CellsToActivate.IsEmpty() && !bForceGoodStreamingPerformance)
 	{
 		UWorld* World = GetWorld();
 
diff --git a/Engine/Source/Runtime/Engine/Private/WorldPartition/WorldPartitionStreamingPolicy.cpp b/Engine/Source/Runtime/Engine/Private/WorldPartition/WorldPartitionStreamingPolicy.cpp
--- a/Engine/Source/Runtime/Engine/Private/WorldPartition/WorldPartitionStreamingPolicy.cpp
+++ b/Engine/Source/Runtime/Engine/Private/WorldPartition/WorldPartitionStreamingPolicy.cpp
@@ -389,7 +389,7 @@ int32 UWorldPartitionStreamingPolicy::GetMaxCellsToLoad() const
 	const ENetMode NetMode = World->GetNetMode();
 	if (NetMode == NM_Standalone || NetMode == NM_Client)
 	{
-		return World->bMatchStarted ? (GMaxLoadingStreamingCells - GetCellLoadingCount()) : MAX_int32;
+		return !World->GetIsInBlockTillLevelStreamingCompleted() ? (GMaxLoadingStreamingCells - GetCellLoadingCount()) : MAX_int32;
 	}
 
 	// Always allow max on server to make sure StreamingLevels are added before clients update the visibility
diff --git a/Engine/Source/Runtime/Engine/Private/WorldPartition/WorldPartitionSubsystem.cpp b/Engine/Source/Runtime/Engine/Private/WorldPartition/WorldPartitionSubsystem.cpp
--- a/Engine/Source/Runtime/Engine/Private/WorldPartition/WorldPartitionSubsystem.cpp
+++ b/Engine/Source/Runtime/Engine/Private/WorldPartition/WorldPartitionSubsystem.cpp
@@ -273,12 +273,38 @@ void UWorldPartitionSubsystem::Draw(UCanvas* Canvas, class APlayerController* PC
 			CurrentOffset.X = CanvasBottomRightPadding.X;
 		}
 
-		FString StatusText;
-		if (IsIncrementalPurgePending()) { StatusText += TEXT("(Purging) "); }
-		if (IsAsyncLoading()) { StatusText += TEXT("(AsyncLoading) "); }
-		if (StatusText.IsEmpty()) { StatusText = TEXT("(Idle)"); }
-		const FString Text = FString::Printf(TEXT("Streaming Status: %s"), *StatusText);
-		FWorldPartitionDebugHelper::DrawText(Canvas, Text, GEngine->GetSmallFont(), FColor::White, CurrentOffset);
+		{
+			FString StatusText;
+			if (IsIncrementalPurgePending()) { StatusText += TEXT("(Purging) "); }
+			if (IsIncrementalUnhashPending()) { StatusText += TEXT("(Unhashing) "); }
+			if (IsAsyncLoading()) { StatusText += TEXT("(AsyncLoading) "); }
+			if (StatusText.IsEmpty()) { StatusText = TEXT("(Idle)"); }
+			const FString Text = FString::Printf(TEXT("Streaming Status: %s"), *StatusText);
+			FWorldPartitionDebugHelper::DrawText(Canvas, Text, GEngine->GetSmallFont(), FColor::White, CurrentOffset);
+		}
+
+		{
+			FString StatusText;
+			const UWorldPartition* WorldPartition = GetMainWorldPartition();
+			EWorldPartitionStreamingPerformance StreamingPerformance = WorldPartition->GetStreamingPerformance();
+			switch (StreamingPerformance)
+			{
+			case EWorldPartitionStreamingPerformance::Good:
+				StatusText = TEXT("Good");
+				break;
+			case EWorldPartitionStreamingPerformance::Slow:
+				StatusText = TEXT("Slow");
+				break;
+			case EWorldPartitionStreamingPerformance::Critical:
+				StatusText = TEXT("Critical");
+				break;
+			default:
+				StatusText = TEXT("Unknown");
+				break;
+			}
+			const FString Text = FString::Printf(TEXT("Streaming Performance: %s"), *StatusText);
+			FWorldPartitionDebugHelper::DrawText(Canvas, Text, GEngine->GetSmallFont(), FColor::White, CurrentOffset);
+		}
 	}
 
 	if (GDrawStreamingSources || GDrawRuntimeHash2D)
diff --git a/Engine/Source/Runtime/Engine/Public/WorldPartition/WorldPartition.h b/Engine/Source/Runtime/Engine/Public/WorldPartition/WorldPartition.h
--- a/Engine/Source/Runtime/Engine/Public/WorldPartition/WorldPartition.h
+++ b/Engine/Source/Runtime/Engine/Public/WorldPartition/WorldPartition.h
@@ -31,6 +31,7 @@ class UCanvas;
 struct IWorldPartitionStreamingSourceProvider;
 
 enum class EWorldPartitionRuntimeCellState : uint8;
+enum class EWorldPartitionStreamingPerformance : uint8;
 
 DECLARE_LOG_CATEGORY_EXTERN(LogWorldPartition, Log, All);
 
@@ -183,6 +184,8 @@ public:
 	void DrawRuntimeCellsDetails(UCanvas* Canvas, FVector2D& Offset);
 	void DrawStreamingStatusLegend(UCanvas* Canvas, FVector2D& Offset);
 
+	EWorldPartitionStreamingPerformance GetStreamingPerformance() const;
+
 #if WITH_EDITORONLY_DATA
 	UPROPERTY(DuplicateTransient)
 	TObjectPtr<UWorldPartitionEditorHash> EditorHash;
diff --git a/Engine/Source/Runtime/Engine/Public/WorldPartition/WorldPartitionRuntimeHash.h b/Engine/Source/Runtime/Engine/Public/WorldPartition/WorldPartitionRuntimeHash.h
--- a/Engine/Source/Runtime/Engine/Public/WorldPartition/WorldPartitionRuntimeHash.h
+++ b/Engine/Source/Runtime/Engine/Public/WorldPartition/WorldPartitionRuntimeHash.h
@@ -16,7 +16,7 @@
 #include "WorldPartitionRuntimeHash.generated.h"
 
 UENUM()
-enum class EWorldPartitionStreamingPerformance
+enum class EWorldPartitionStreamingPerformance : uint8
 {
 	Good,
 	Slow,
diff --git a/Engine/Source/Runtime/Engine/Public/WorldPartition/WorldPartitionStreamingPolicy.h b/Engine/Source/Runtime/Engine/Public/WorldPartition/WorldPartitionStreamingPolicy.h
--- a/Engine/Source/Runtime/Engine/Public/WorldPartition/WorldPartitionStreamingPolicy.h
+++ b/Engine/Source/Runtime/Engine/Public/WorldPartition/WorldPartitionStreamingPolicy.h
@@ -70,6 +70,8 @@ public:
 
 	const TArray<FWorldPartitionStreamingSource>& GetStreamingSources() const { return StreamingSources; }
 
+	EWorldPartitionStreamingPerformance GetStreamingPerformance() const { return StreamingPerformance; }
+
 protected:
 	virtual void SetCellsStateToLoaded(const TSet<const UWorldPartitionRuntimeCell*>& ToLoadCells);
 	virtual void SetCellsStateToActivated(const TSet<const UWorldPartitionRuntimeCell*>& ToActivateCells);
