Make ray tracing shares the same skin cache entry with raster if they use the same LOD level

diff --git a/Engine/Source/Runtime/Engine/Private/GPUSkinCache.cpp b/Engine/Source/Runtime/Engine/Private/GPUSkinCache.cpp
--- a/Engine/Source/Runtime/Engine/Private/GPUSkinCache.cpp
+++ b/Engine/Source/Runtime/Engine/Private/GPUSkinCache.cpp
@@ -1852,7 +1852,7 @@ void FGPUSkinCache::EndBatchDispatch(FRHICommandListImmediate& RHICmdList)
 
 			FGPUSkinCacheEntry* SkinCacheEntry = DispatchItem.SkinCacheEntry;
 
-			if (SkinCacheEntry->Mode != EGPUSkinCacheEntryMode::RayTracing)
+			if (SkinCacheEntry->GPUSkin->ShouldUseSeparateSkinCacheEntryForRayTracing() && SkinCacheEntry->Mode != EGPUSkinCacheEntryMode::RayTracing)
 			{
 				continue;
 			}
diff --git a/Engine/Source/Runtime/Engine/Private/SkeletalMesh.cpp b/Engine/Source/Runtime/Engine/Private/SkeletalMesh.cpp
--- a/Engine/Source/Runtime/Engine/Private/SkeletalMesh.cpp
+++ b/Engine/Source/Runtime/Engine/Private/SkeletalMesh.cpp
@@ -6225,7 +6225,7 @@ void FSkeletalMeshSceneProxy::CreateBaseMeshBatch(const FSceneView* View, const
 	BatchElement.IndexBuffer = LODData.MultiSizeIndexContainer.GetIndexBuffer();
 	BatchElement.MinVertexIndex = LODData.RenderSections[SectionIndex].GetVertexBufferIndex();
 	BatchElement.MaxVertexIndex = LODData.RenderSections[SectionIndex].GetVertexBufferIndex() + LODData.RenderSections[SectionIndex].GetNumVertices() - 1;
-	BatchElement.VertexFactoryUserData = FGPUSkinCache::GetFactoryUserData(VFMode == ESkinVertexFactoryMode::RayTracing ? MeshObject->SkinCacheEntryForRayTracing : MeshObject->SkinCacheEntry, SectionIndex);
+	BatchElement.VertexFactoryUserData = FGPUSkinCache::GetFactoryUserData(VFMode == ESkinVertexFactoryMode::RayTracing ? MeshObject->GetSkinCacheEntryForRayTracing() : MeshObject->SkinCacheEntry, SectionIndex);
 	BatchElement.PrimitiveUniformBuffer = GetUniformBuffer();
 	BatchElement.NumPrimitives = LODData.RenderSections[SectionIndex].NumTriangles;
 }
@@ -6374,7 +6374,7 @@ void FSkeletalMeshSceneProxy::GetDynamicRayTracingInstances(FRayTracingMaterialG
 	}
 
 	// GetRayTracingGeometry()->IsInitialized() is checked as a workaround for UE-92634. FSkeletalMeshSceneProxy's resources may have already been released, but proxy has not removed yet)
-	if (MeshObject->GetRayTracingGeometry() && MeshObject->GetRayTracingGeometry()->IsInitialized() && MeshObject->SkinCacheEntryForRayTracing)
+	if (MeshObject->GetRayTracingGeometry() && MeshObject->GetRayTracingGeometry()->IsInitialized() && MeshObject->GetSkinCacheEntryForRayTracing())
 	{
 		// #dxr: the only case where RayTracingGeometryRHI is invalid is the very first frame - if that's not the case we have a bug somewhere else
 		if (MeshObject->GetRayTracingGeometry()->RayTracingGeometryRHI.IsValid())
diff --git a/Engine/Source/Runtime/Engine/Private/SkeletalRenderGPUSkin.cpp b/Engine/Source/Runtime/Engine/Private/SkeletalRenderGPUSkin.cpp
--- a/Engine/Source/Runtime/Engine/Private/SkeletalRenderGPUSkin.cpp
+++ b/Engine/Source/Runtime/Engine/Private/SkeletalRenderGPUSkin.cpp
@@ -398,7 +398,7 @@ void FSkeletalMeshObjectGPUSkin::UpdateDynamicData_RenderThread(FGPUSkinCache* G
 	bRequireRecreatingRayTracingGeometry = (DynamicData == nullptr || RayTracingGeometry.Initializer.Segments.Num() == 0 || // Newly created
 		(DynamicData != nullptr && DynamicData->RayTracingLODIndex != InDynamicData->RayTracingLODIndex) || // LOD level changed
 		bHiddenMaterialVisibilityDirtyForRayTracing ||
-		SkinCacheEntryForRayTracing == nullptr);
+		GetSkinCacheEntryForRayTracing() == nullptr);
 
 	bHiddenMaterialVisibilityDirtyForRayTracing = false;
 	
@@ -434,7 +434,7 @@ void FSkeletalMeshObjectGPUSkin::UpdateDynamicData_RenderThread(FGPUSkinCache* G
 		ProcessUpdatedDynamicData(EGPUSkinCacheEntryMode::Raster, GPUSkinCache, RHICmdList, FrameNumberToPrepare, RevisionNumber, bMorphNeedsUpdate, DynamicData->LODIndex);
 
 	#if RHI_RAYTRACING
-		if (FGPUSkinCache::IsGPUSkinCacheRayTracingSupported() && GPUSkinCache && SkeletalMeshRenderData->bSupportRayTracing)
+		if (ShouldUseSeparateSkinCacheEntryForRayTracing() && FGPUSkinCache::IsGPUSkinCacheRayTracingSupported() && GPUSkinCache && SkeletalMeshRenderData->bSupportRayTracing)
 		{
 			// When in game, updates not sent in UWorld::SendAllEndOfFrameUpdates() (implied by IsBatchingDispatch()) are ignored.
 			if (!bIsGameWorld || GPUSkinCache->IsBatchingDispatch())
@@ -442,6 +442,11 @@ void FSkeletalMeshObjectGPUSkin::UpdateDynamicData_RenderThread(FGPUSkinCache* G
 				ProcessUpdatedDynamicData(EGPUSkinCacheEntryMode::RayTracing, GPUSkinCache, RHICmdList, FrameNumberToPrepare, RevisionNumber, bMorphNeedsUpdate, DynamicData->RayTracingLODIndex);
 			}
 		}
+		else
+		{
+			// Immediately release any stale entry if we decide to share with the raster path
+			FGPUSkinCache::Release(SkinCacheEntryForRayTracing);
+		}
 	#endif
 	}
 
@@ -449,10 +454,10 @@ void FSkeletalMeshObjectGPUSkin::UpdateDynamicData_RenderThread(FGPUSkinCache* G
 	// Immediate update path. Only used in the editor as the world can be not ticking. When in game, updates not sent in UWorld::SendAllEndOfFrameUpdates() are ignored.
 	if (!bIsGameWorld && FGPUSkinCache::IsGPUSkinCacheRayTracingSupported() && GPUSkinCache && SkeletalMeshRenderData->bSupportRayTracing && !GPUSkinCache->IsBatchingDispatch())
 	{
-		if (SkinCacheEntryForRayTracing)
+		if (GetSkinCacheEntryForRayTracing())
 		{
 			FSkeletalMeshLODRenderData& LODModel = this->SkeletalMeshRenderData->LODRenderData[DynamicData->RayTracingLODIndex];
-			GPUSkinCache->ProcessRayTracingGeometryToUpdate(RHICmdList, SkinCacheEntryForRayTracing, LODModel, bRequireRecreatingRayTracingGeometry, DynamicData->bAnySegmentUsesWorldPositionOffset);
+			GPUSkinCache->ProcessRayTracingGeometryToUpdate(RHICmdList, GetSkinCacheEntryForRayTracing(), LODModel, bRequireRecreatingRayTracingGeometry, DynamicData->bAnySegmentUsesWorldPositionOffset);
 		}
 		else
 		{
@@ -507,6 +512,13 @@ void FSkeletalMeshObjectGPUSkin::ProcessUpdatedDynamicData(EGPUSkinCacheEntryMod
 
 	bool bGPUSkinCacheEnabled = GPUSkinCache && GEnableGPUSkinCache && (FeatureLevel >= ERHIFeatureLevel::SM5) && (DynamicData->bIsSkinCacheAllowed || Mode == EGPUSkinCacheEntryMode::RayTracing);
 
+	// Immediately release any stale entry if we've recently switched to a LOD level that disallows skin cache
+	// This saves memory and avoids confusing ShouldUseSeparateSkinCacheEntryForRayTracing() which checks SkinCacheEntry == nullptr
+	if (!bGPUSkinCacheEnabled && SkinCacheEntry)
+	{
+		FGPUSkinCache::Release(SkinCacheEntry);
+	}
+
 	if (DynamicData->PreSkinningOffsets.Num() > 0)
 	{
 		FPositionVertexBuffer& PositionBuffer = LOD.VertexOffsetVertexBuffers.PreSkinningOffsetsVertexBuffer;
@@ -1160,8 +1172,8 @@ const FVertexFactory* FSkeletalMeshObjectGPUSkin::GetSkinVertexFactory(const FSc
 	// Return the passthrough vertex factory if it is requested (by ray tracing)
 	if (VFMode == ESkinVertexFactoryMode::RayTracing)
 	{
-		check(SkinCacheEntryForRayTracing);
-		check(FGPUSkinCache::IsEntryValid(SkinCacheEntryForRayTracing, ChunkIdx));
+		check(GetSkinCacheEntryForRayTracing());
+		check(FGPUSkinCache::IsEntryValid(GetSkinCacheEntryForRayTracing(), ChunkIdx));
 
 		return LOD.GPUSkinVertexFactories.PassthroughVertexFactories[ChunkIdx].Get();
 	}
diff --git a/Engine/Source/Runtime/Engine/Public/SkeletalRenderPublic.h b/Engine/Source/Runtime/Engine/Public/SkeletalRenderPublic.h
--- a/Engine/Source/Runtime/Engine/Public/SkeletalRenderPublic.h
+++ b/Engine/Source/Runtime/Engine/Public/SkeletalRenderPublic.h
@@ -207,6 +207,9 @@ public:
 	virtual int32 GetRayTracingLOD() const { return GetLOD(); }
 #endif // RHI_RAYTRACING
 
+	virtual bool ShouldUseSeparateSkinCacheEntryForRayTracing() const { return GetLOD() != GetRayTracingLOD() || SkinCacheEntry == nullptr; }
+	virtual class FGPUSkinCacheEntry* GetSkinCacheEntryForRayTracing() const { return ShouldUseSeparateSkinCacheEntryForRayTracing() ? SkinCacheEntryForRayTracing : SkinCacheEntry; }
+
 	/** Called to notify clothing data that component transform has changed */
 	virtual void RefreshClothingTransforms(const FMatrix& InNewLocalToWorld, uint32 FrameNumber) {};
 
@@ -287,7 +290,7 @@ protected:
 	/** Per-LOD info. */
 	TArray<FSkeletalMeshLODInfo> SkeletalMeshLODInfo;
 
-	class FGPUSkinCacheEntry* SkinCacheEntry;
+	FGPUSkinCacheEntry* SkinCacheEntry;
 	FGPUSkinCacheEntry* SkinCacheEntryForRayTracing;
 
 	/** Used to keep track of the first call to UpdateMinDesiredLODLevel each frame. from ViewFamily.FrameNumber */
