CIS fixes for non RHI_RAYTRACING platforms

diff --git a/Engine/Source/Runtime/Engine/Private/SkeletalMesh.cpp b/Engine/Source/Runtime/Engine/Private/SkeletalMesh.cpp
--- a/Engine/Source/Runtime/Engine/Private/SkeletalMesh.cpp
+++ b/Engine/Source/Runtime/Engine/Private/SkeletalMesh.cpp
@@ -6225,7 +6225,11 @@ void FSkeletalMeshSceneProxy::CreateBaseMeshBatch(const FSceneView* View, const
 	BatchElement.IndexBuffer = LODData.MultiSizeIndexContainer.GetIndexBuffer();
 	BatchElement.MinVertexIndex = LODData.RenderSections[SectionIndex].GetVertexBufferIndex();
 	BatchElement.MaxVertexIndex = LODData.RenderSections[SectionIndex].GetVertexBufferIndex() + LODData.RenderSections[SectionIndex].GetNumVertices() - 1;
+#if RHI_RAYTRACING
 	BatchElement.VertexFactoryUserData = FGPUSkinCache::GetFactoryUserData(VFMode == ESkinVertexFactoryMode::RayTracing ? MeshObject->GetSkinCacheEntryForRayTracing() : MeshObject->SkinCacheEntry, SectionIndex);
+#else
+	BatchElement.VertexFactoryUserData = FGPUSkinCache::GetFactoryUserData(MeshObject->SkinCacheEntry, SectionIndex);
+#endif
 	BatchElement.PrimitiveUniformBuffer = GetUniformBuffer();
 	BatchElement.NumPrimitives = LODData.RenderSections[SectionIndex].NumTriangles;
 }
diff --git a/Engine/Source/Runtime/Engine/Private/SkeletalRenderGPUSkin.cpp b/Engine/Source/Runtime/Engine/Private/SkeletalRenderGPUSkin.cpp
--- a/Engine/Source/Runtime/Engine/Private/SkeletalRenderGPUSkin.cpp
+++ b/Engine/Source/Runtime/Engine/Private/SkeletalRenderGPUSkin.cpp
@@ -1169,6 +1169,7 @@ const FVertexFactory* FSkeletalMeshObjectGPUSkin::GetSkinVertexFactory(const FSc
 	const FSkelMeshObjectLODInfo& MeshLODInfo = LODInfo[LODIndex];
 	const FSkeletalMeshObjectLOD& LOD = LODs[LODIndex];
 
+#if RHI_RAYTRACING
 	// Return the passthrough vertex factory if it is requested (by ray tracing)
 	if (VFMode == ESkinVertexFactoryMode::RayTracing)
 	{
@@ -1177,6 +1178,7 @@ const FVertexFactory* FSkeletalMeshObjectGPUSkin::GetSkinVertexFactory(const FSc
 
 		return LOD.GPUSkinVertexFactories.PassthroughVertexFactories[ChunkIdx].Get();
 	}
+#endif
 
 	// If the GPU skinning cache was used, return the passthrough vertex factory
 	if (SkinCacheEntry && FGPUSkinCache::IsEntryValid(SkinCacheEntry, ChunkIdx) && DynamicData->bIsSkinCacheAllowed)
diff --git a/Engine/Source/Runtime/Engine/Public/SkeletalRenderPublic.h b/Engine/Source/Runtime/Engine/Public/SkeletalRenderPublic.h
--- a/Engine/Source/Runtime/Engine/Public/SkeletalRenderPublic.h
+++ b/Engine/Source/Runtime/Engine/Public/SkeletalRenderPublic.h
@@ -18,6 +18,7 @@ class FPrimitiveDrawInterface;
 class FVertexFactory;
 class UMorphTarget;
 struct FSkelMeshRenderSection;
+class FGPUSkinCacheEntry;
 
 /** data for a single skinned skeletal mesh vertex */
 struct FFinalSkinVertex
@@ -205,10 +206,10 @@ public:
 	virtual const FRayTracingGeometry* GetRayTracingGeometry() const { return nullptr; }
 	virtual FRWBuffer* GetRayTracingDynamicVertexBuffer() { return nullptr; }
 	virtual int32 GetRayTracingLOD() const { return GetLOD(); }
-#endif // RHI_RAYTRACING
 
 	virtual bool ShouldUseSeparateSkinCacheEntryForRayTracing() const { return GetLOD() != GetRayTracingLOD() || SkinCacheEntry == nullptr; }
-	virtual class FGPUSkinCacheEntry* GetSkinCacheEntryForRayTracing() const { return ShouldUseSeparateSkinCacheEntryForRayTracing() ? SkinCacheEntryForRayTracing : SkinCacheEntry; }
+	virtual FGPUSkinCacheEntry* GetSkinCacheEntryForRayTracing() const { return ShouldUseSeparateSkinCacheEntryForRayTracing() ? SkinCacheEntryForRayTracing : SkinCacheEntry; }
+#endif // RHI_RAYTRACING
 
 	/** Called to notify clothing data that component transform has changed */
 	virtual void RefreshClothingTransforms(const FMatrix& InNewLocalToWorld, uint32 FrameNumber) {};
