Make sure WorldPartitionVolume don't support Data Layers.

diff --git a/Engine/Source/Runtime/Engine/Public/WorldPartition/WorldPartitionVolume.h b/Engine/Source/Runtime/Engine/Public/WorldPartition/WorldPartitionVolume.h
--- a/Engine/Source/Runtime/Engine/Public/WorldPartition/WorldPartitionVolume.h
+++ b/Engine/Source/Runtime/Engine/Public/WorldPartition/WorldPartitionVolume.h
@@ -24,5 +24,8 @@ public:
 	void LoadIntersectingCells(bool bIsFromUserChange);
 	void UnloadIntersectingCells(bool bIsFromUserChange);
 	virtual EActorGridPlacement GetDefaultGridPlacement() const override { return EActorGridPlacement::AlwaysLoaded; }
+
+private:
+	virtual bool SupportsDataLayer() const override { return false; }
 #endif
 };
\ No newline at end of file
