PR #8323: Added virtual keyword to USkeletalMeshComponent::ShouldBlendPhysicsBones (Contributed by [REDACTED_PERSON])

diff --git a/Engine/Source/Runtime/Engine/Classes/Components/SkeletalMeshComponent.h b/Engine/Source/Runtime/Engine/Classes/Components/SkeletalMeshComponent.h
--- a/Engine/Source/Runtime/Engine/Classes/Components/SkeletalMeshComponent.h
+++ b/Engine/Source/Runtime/Engine/Classes/Components/SkeletalMeshComponent.h
@@ -2250,7 +2250,7 @@ protected:
 	bool NeedToSpawnAnimScriptInstance() const;
 	bool NeedToSpawnPostPhysicsInstance(bool bForceReinit) const;
 
-	bool ShouldBlendPhysicsBones() const;
+	virtual bool ShouldBlendPhysicsBones() const;
 
 	/** Extract collisions for cloth from this component (given a component we want to apply the data to) */
 	static void ExtractCollisionsForCloth(USkeletalMeshComponent* SourceComponent,  UPhysicsAsset* PhysicsAsset, USkeletalMeshComponent* DestClothComponent, FClothCollisionData& OutCollisions, FClothCollisionSource& ClothCollisionSource);
