Change textures to allow in-editor streaming of mips produced by DDC2 as well as DDC1.
Does not include Virtual Textures yet.

diff --git a/Engine/Source/Runtime/Engine/Classes/Engine/Texture.h b/Engine/Source/Runtime/Engine/Classes/Engine/Texture.h
--- a/Engine/Source/Runtime/Engine/Classes/Engine/Texture.h
+++ b/Engine/Source/Runtime/Engine/Classes/Engine/Texture.h
@@ -20,6 +20,10 @@
 #include "PerPlatformProperties.h"
 #include "Misc/FieldAccessor.h"
 #include "Virtualization/VirtualizedBulkData.h"
+#if WITH_EDITORONLY_DATA
+#include "Misc/TVariant.h"
+#include "DerivedDataCacheKeyProxy.h"
+#endif
 
 #if WITH_EDITOR
 #include "Templates/DontCopy.h"
@@ -631,17 +635,11 @@ struct FTexturePlatformData
 
 #if WITH_EDITORONLY_DATA
 	/** The key associated with this derived data. */
-	FString DerivedDataKey;
-	/** 
-	* The key associated with the derived data generated for the shipping build.
-	* In order to save developer's time when running a build in the Editor,
-	* cooking a texture for the Editor build normally involves less complicated processing
-	* than cooking a texture for the shipping build. However, in many cases a fully cooked
-	* texture is already available in the cache when the project is opened on the developer's machine.
-	* This key allows to load the texture generated for the shipping build from the cache, if available,
-	* instead of building a new texture locally using fast cooking.
-	*/
-	FString ShippingDerivedDataKey;
+	TVariant<FString, UE::DerivedData::FCacheKeyProxy> DerivedDataKey;
+
+	/** The key used when comparing this derived data key to another reference key. */
+	TVariant<FString, UE::DerivedData::FCacheKeyProxy> ComparisonDerivedDataKey;
+
 	/** Async cache task if one is outstanding. */
 	struct FTextureAsyncCacheDerivedDataTask* AsyncTask;
 #endif
@@ -677,6 +675,9 @@ public:
 
 #if WITH_EDITORONLY_DATA
 	void SerializeWithConditionalBulkData(FArchive& Ar, class UTexture* Owner);
+
+	FString GetDerivedDataMipKeyString(int32 MipIndex, const FTexture2DMipMap& Mip) const;
+	UE::DerivedData::FCachePayloadKeyProxy GetDerivedDataMipKeyProxy(int32 MipIndex, const FTexture2DMipMap& Mip) const;
 #endif // WITH_EDITORONLY_DATA
 
 	/** 
@@ -733,6 +734,7 @@ public:
 	}
 
 #if WITH_EDITOR
+	static bool IsUsingNewDerivedData();
 	bool IsAsyncWorkComplete() const;
 	void Cache(
 		class UTexture& InTexture,
diff --git a/Engine/Source/Runtime/Engine/Private/DerivedDataCacheKeyProxy.cpp b/Engine/Source/Runtime/Engine/Private/DerivedDataCacheKeyProxy.cpp
new file mode 100644
--- /dev/null
+++ b/Engine/Source/Runtime/Engine/Private/DerivedDataCacheKeyProxy.cpp
@@ -0,0 +1,39 @@
+// Copyright Epic Games, Inc. All Rights Reserved.
+
+#include "DerivedDataCacheKeyProxy.h"
+
+#if WITH_EDITORONLY_DATA
+#include "DerivedDataCacheKey.h"
+
+namespace UE::DerivedData
+{
+
+static_assert(sizeof(FCacheKeyProxy) == sizeof(FCacheKeyProxy));
+static_assert(alignof(FCacheKeyProxy) == alignof(FCacheKeyProxy));
+
+static_assert(sizeof(FCachePayloadKeyProxy) == sizeof(FCachePayloadKey));
+static_assert(alignof(FCachePayloadKeyProxy) == alignof(FCachePayloadKey));
+
+FCacheKeyProxy::FCacheKeyProxy(const FCacheKey& InKey)
+{
+	new(AsCacheKey()) FCacheKey(InKey);
+}
+
+FCacheKeyProxy::~FCacheKeyProxy()
+{
+	AsCacheKey()->~FCacheKey();
+}
+
+FCachePayloadKeyProxy::FCachePayloadKeyProxy(const FCachePayloadKey& InKey)
+{
+	new(AsCachePayloadKey()) FCachePayloadKey(InKey);
+}
+
+FCachePayloadKeyProxy::~FCachePayloadKeyProxy()
+{
+	AsCachePayloadKey()->~FCachePayloadKey();
+}
+
+}
+
+#endif // WITH_EDITORONLY_DATA
\ No newline at end of file
diff --git a/Engine/Source/Runtime/Engine/Private/HLODProxy.cpp b/Engine/Source/Runtime/Engine/Private/HLODProxy.cpp
--- a/Engine/Source/Runtime/Engine/Private/HLODProxy.cpp
+++ b/Engine/Source/Runtime/Engine/Private/HLODProxy.cpp
@@ -6,6 +6,7 @@
 #include "GameFramework/WorldSettings.h"
 
 #if WITH_EDITOR
+#include "DerivedDataCacheKey.h"
 #include "Misc/ConfigCacheIni.h"
 #include "Serialization/ArchiveCrc32.h"
 #include "HierarchicalLOD.h"
@@ -357,7 +358,19 @@ uint32 UHLODProxy::GetCRC(UTexture* InTexture, uint32 InCRC)
          FTexturePlatformData* PlatformData = *InTexture->GetRunningPlatformData();
          if (PlatformData != nullptr)
          {
-             return FCrc::StrCrc32(*PlatformData->DerivedDataKey, InCRC);
+			 if (const FString* KeyString = PlatformData->DerivedDataKey.TryGet<FString>())
+			 {
+				 return FCrc::StrCrc32(**KeyString, InCRC);
+			 }
+			 else if (const UE::DerivedData::FCacheKeyProxy* KeyProxy = PlatformData->DerivedDataKey.TryGet<UE::DerivedData::FCacheKeyProxy>())
+			 {
+				 const UE::DerivedData::FCacheKey* Key = KeyProxy->AsCacheKey();
+				 FAnsiStringView BucketStringView = Key->Bucket.ToString();
+				 FArchiveCrc32 Ar(InCRC);
+				 Ar.Serialize(const_cast<ANSICHAR*>(BucketStringView.GetData()), BucketStringView.Len());
+				 Ar.Serialize(const_cast<FIoHash*>(&Key->Hash), sizeof(FIoHash));
+				 return Ar.GetCrc();
+			 }
          }
      }
  
diff --git a/Engine/Source/Runtime/Engine/Private/Streaming/Texture2DMipDataProvider_DDC.cpp b/Engine/Source/Runtime/Engine/Private/Streaming/Texture2DMipDataProvider_DDC.cpp
--- a/Engine/Source/Runtime/Engine/Private/Streaming/Texture2DMipDataProvider_DDC.cpp
+++ b/Engine/Source/Runtime/Engine/Private/Streaming/Texture2DMipDataProvider_DDC.cpp
@@ -13,8 +13,13 @@ Texture2DMipDataProvider_DDC.cpp : Implementation of FTextureMipDataProvider usi
 
 #if WITH_EDITORONLY_DATA
 
+#include "DerivedDataCache.h"
+#include "DerivedDataCacheInterface.h"
+#include "DerivedDataCacheKey.h"
+
 FTexture2DMipDataProvider_DDC::FTexture2DMipDataProvider_DDC(const UTexture* Texture)
 	: FTextureMipDataProvider(Texture, ETickState::Init, ETickThread::Async)
+	, DDCRequestOwner(UE::DerivedData::EPriority::Normal)
 {
 }
 
@@ -28,15 +33,72 @@ void FTexture2DMipDataProvider_DDC::Init(const FTextureUpdateContext& Context, c
 	if (!DDCHandles.Num())
 	{
 		DDCHandles.AddZeroed(CurrentFirstLODIdx);
+		DDCBuffers.AddZeroed(CurrentFirstLODIdx);
 
-		for (int32 MipIndex = PendingFirstLODIdx; MipIndex < CurrentFirstLODIdx; ++MipIndex)
+		FTexturePlatformData*const * PtrPlatformData = const_cast<UTexture*>(Context.Texture)->GetRunningPlatformData();
+		if (PtrPlatformData && *PtrPlatformData)
 		{
-			const FTexture2DMipMap& OwnerMip = *Context.MipsView[MipIndex];
-			if (!OwnerMip.DerivedDataKey.IsEmpty())
+			const FTexturePlatformData* PlatformData = *PtrPlatformData;
+			const int32 LODBias = static_cast<int32>(Context.MipsView.GetData() - PlatformData->Mips.GetData());
+
+			if (PlatformData->DerivedDataKey.IsType<FString>())
+			{
+				for (int32 MipIndex = PendingFirstLODIdx; MipIndex < CurrentFirstLODIdx; ++MipIndex)
+				{
+					const FTexture2DMipMap& OwnerMip = *Context.MipsView[MipIndex];
+					if (OwnerMip.IsPagedToDerivedData())
+					{
+						DDCHandles[MipIndex] = GetDerivedDataCacheRef().GetAsynchronous(*PlatformData->GetDerivedDataMipKeyString(MipIndex + LODBias, OwnerMip), Context.Texture->GetPathName());
+					}
+				}
+			}
+			else if (PlatformData->DerivedDataKey.IsType<UE::DerivedData::FCacheKeyProxy>())
+			{
+				using namespace UE::DerivedData;
+				TArray<FCachePayloadKey> MipKeys;
+
+				for (int32 MipIndex = PendingFirstLODIdx; MipIndex < CurrentFirstLODIdx; ++MipIndex)
+				{
+					const FTexture2DMipMap& MipMap = *Context.MipsView[MipIndex];
+					if (MipMap.IsPagedToDerivedData())
+					{
+						MipKeys.Emplace(*PlatformData->GetDerivedDataMipKeyProxy(MipIndex + LODBias, MipMap).AsCachePayloadKey());
+					}
+				}
+
+				if (MipKeys.Num())
+				{
+					GetCache().GetPayload(MipKeys, Context.Texture->GetPathName(), ECachePolicy::Default, DDCRequestOwner,
+						[this, &Context, LODBias](FCacheGetPayloadCompleteParams&& Params)
+						{
+							if (Params.Status == EStatus::Ok)
+							{
+								for (int32 MipIndex = PendingFirstLODIdx; MipIndex < CurrentFirstLODIdx; ++MipIndex)
+								{
+									TStringBuilder<16> MipName;
+									MipName << TEXT("Mip") << (MipIndex + LODBias);
+									if (Params.Payload.GetId() == FPayloadId::FromName(MipName))
+									{
+										check(!DDCBuffers[MipIndex]);
+										DDCBuffers[MipIndex] = Params.Payload.GetData().Decompress();
+										break;
+									}
+								}
+							}
+						});
+				}
+			}
+			else
 			{
-				DDCHandles[MipIndex] = GetDerivedDataCacheRef().GetAsynchronous(*OwnerMip.DerivedDataKey, Context.Texture->GetPathName());
+				UE_LOG(LogTexture, Error, TEXT("Attempting to stream in mips for texture that has not generated a supported derived data key format."));
 			}
 		}
+		else
+		{
+			UE_LOG(LogTexture, Error, TEXT("Attempting to stream in mips for texture that has not generated derived data yet."));
+		}
+
+
 		*SyncOptions.bSnooze = true;
 	}
 	else // The DDC request have been issued, only check whether they are ready (since no good sync option is used).
@@ -49,72 +111,127 @@ void FTexture2DMipDataProvider_DDC::Init(const FTextureUpdateContext& Context, c
 				return;
 			}
 		}
+
+		if (!DDCRequestOwner.Poll())
+		{
+			*SyncOptions.bSnooze = true;
+			return;
+		}
+
 		AdvanceTo(ETickState::GetMips, ETickThread::Async);
 	}
 }
 
+bool FTexture2DMipDataProvider_DDC::SerializeMipInfo(const FTextureUpdateContext& Context, FArchive& Ar, int32 MipIndex, const FTextureMipInfo& OutMipInfo)
+{
+	int32 MipSize = 0;
+	Ar << MipSize;
+
+	const uint32 DepthOrArraySize = FMath::Max<uint32>(OutMipInfo.ArraySize, OutMipInfo.SizeZ);
+	if (MipSize == OutMipInfo.DataSize)
+	{
+		Ar.Serialize(OutMipInfo.DestData, MipSize);
+		return true;
+	}
+	else if (MipSize < OutMipInfo.DataSize && DepthOrArraySize > 1 && OutMipInfo.DataSize % DepthOrArraySize == 0 && MipSize % DepthOrArraySize == 0)
+	{
+		UE_LOG(LogTexture, Verbose, TEXT("DDC mip size smaller than streaming buffer size. (%s, Mip %d): %d KB / %d KB."), *Context.Resource->GetTextureName().ToString(), ResourceState.MaxNumLODs - MipIndex, OutMipInfo.DataSize / 1024, MipSize / 1024);
+
+		const uint64 SourceSubSize = MipSize / DepthOrArraySize;
+		const uint64 DestSubSize = OutMipInfo.DataSize / DepthOrArraySize;
+		const uint64 PaddingSubSize = DestSubSize - SourceSubSize;
+
+		uint8* DestData = (uint8*)OutMipInfo.DestData;
+		for (uint32 SubIdx = 0; SubIdx < DepthOrArraySize; ++SubIdx)
+		{
+			Ar.Serialize(DestData, SourceSubSize);
+			DestData += SourceSubSize;
+			FMemory::Memzero(DestData, PaddingSubSize);
+			DestData += PaddingSubSize;
+		}
+		return true;
+	}
+
+	UE_LOG(LogTexture, Warning, TEXT("Mismatch between DDC mip size and streaming buffer size. (%s, Mip %d): %d KB / %d KB."), *Context.Resource->GetTextureName().ToString(), ResourceState.MaxNumLODs - MipIndex, OutMipInfo.DataSize / 1024, MipSize / 1024);
+	FMemory::Memzero(OutMipInfo.DestData, OutMipInfo.DataSize);
+	return false;
+}
+
 int32 FTexture2DMipDataProvider_DDC::GetMips(
 	const FTextureUpdateContext& Context,
 	int32 StartingMipIndex,
 	const FTextureMipInfoArray& MipInfos, 
 	const FTextureUpdateSyncOptions& SyncOptions)
 {
-	for (int32 MipIndex = StartingMipIndex; MipIndex < CurrentFirstLODIdx; ++MipIndex)
+	FTexturePlatformData* const* PtrPlatformData = const_cast<UTexture*>(Context.Texture)->GetRunningPlatformData();
+	if (PtrPlatformData && *PtrPlatformData)
 	{
-		const uint32 Handle = DDCHandles[MipIndex];
-		bool bSuccess = false;
-		if (Handle)
-		{
-			DDCHandles[MipIndex] = 0; // Clear the handle.
+		const FTexturePlatformData* PlatformData = *PtrPlatformData;
+		const int32 LODBias = static_cast<int32>(Context.MipsView.GetData() - PlatformData->Mips.GetData());
 
-			TArray<uint8> DerivedMipData;
-			if (GetDerivedDataCacheRef().GetAsynchronousResults(Handle, DerivedMipData))
+		if (PlatformData->DerivedDataKey.IsType<FString>())
+		{
+			for (int32 MipIndex = StartingMipIndex; MipIndex < CurrentFirstLODIdx; ++MipIndex)
 			{
-				const FTextureMipInfo& MipInfo = MipInfos[MipIndex];
+				const uint32 Handle = DDCHandles[MipIndex];
+				bool bSuccess = false;
+				if (Handle)
+				{
+					DDCHandles[MipIndex] = 0; // Clear the handle.
 
-				// The result must be read from a memory reader!
-				FMemoryReader Ar(DerivedMipData, true);
-				int32 MipSize = 0;
-				Ar << MipSize;
+					TArray<uint8> DerivedMipData;
+					if (GetDerivedDataCacheRef().GetAsynchronousResults(Handle, DerivedMipData))
+					{
+						const FTextureMipInfo& MipInfo = MipInfos[MipIndex];
 
-				const uint32 DepthOrArraySize = FMath::Max<uint32>(MipInfo.ArraySize, MipInfo.SizeZ);
-				if (MipSize == MipInfo.DataSize)
+						// The result must be read from a memory reader!
+						FMemoryReader Ar(DerivedMipData, true);
+						if (SerializeMipInfo(Context, Ar, MipIndex, MipInfo))
+						{
+							bSuccess = true;
+						}
+					}
+				}
+
+				if (!bSuccess)
 				{
-					Ar.Serialize(MipInfo.DestData, MipSize);
-					bSuccess = true;
+					AdvanceTo(ETickState::CleanUp, ETickThread::Async);
+					return MipIndex; // We failed at getting this mip. Cancel will be called.
 				}
-				else if (MipSize < MipInfo.DataSize && DepthOrArraySize > 1 && MipInfo.DataSize % DepthOrArraySize == 0 && MipSize % DepthOrArraySize == 0)
+			}
+		}
+		else if (PlatformData->DerivedDataKey.IsType<UE::DerivedData::FCacheKeyProxy>())
+		{
+			for (int32 MipIndex = StartingMipIndex; MipIndex < CurrentFirstLODIdx; ++MipIndex)
+			{
+				bool bSuccess = false;
+				if (DDCBuffers[MipIndex])
 				{
-					UE_LOG(LogTexture, Verbose, TEXT("DDC mip size smaller than streaming buffer size. (%s, Mip %d): %d KB / %d KB."), *Context.Resource->GetTextureName().ToString(), ResourceState.MaxNumLODs - MipIndex, MipInfo.DataSize / 1024, MipSize / 1024);
-
-					const uint64 SourceSubSize = MipSize / DepthOrArraySize;
-					const uint64 DestSubSize = MipInfo.DataSize / DepthOrArraySize;
-					const uint64 PaddingSubSize = DestSubSize - SourceSubSize;
-
-					uint8* DestData = (uint8*)MipInfo.DestData;
-					for (uint32 SubIdx = 0; SubIdx < DepthOrArraySize; ++SubIdx)
+					// The result must be read from a memory reader!
+					FMemoryReaderView Ar(DDCBuffers[MipIndex], true);
+					if (SerializeMipInfo(Context, Ar, MipIndex, MipInfos[MipIndex]))
 					{
-						Ar.Serialize(DestData, SourceSubSize);
-						DestData += SourceSubSize;
-						FMemory::Memzero(DestData, PaddingSubSize);
-						DestData += PaddingSubSize;
+						bSuccess = true;
 					}
-					bSuccess = true;
+					DDCBuffers[MipIndex].Reset();
 				}
-				else
+
+				if (!bSuccess)
 				{
-					UE_LOG(LogTexture, Warning, TEXT("Mismatch between DDC mip size and streaming buffer size. (%s, Mip %d): %d KB / %d KB."), *Context.Resource->GetTextureName().ToString(), ResourceState.MaxNumLODs - MipIndex, MipInfo.DataSize / 1024, MipSize / 1024);
-					FMemory::Memzero(MipInfo.DestData, MipInfo.DataSize);
+					AdvanceTo(ETickState::CleanUp, ETickThread::Async);
+					return MipIndex; // We failed at getting this mip. Cancel will be called.
 				}
 			}
 		}
-
-		if (!bSuccess)
+		else
 		{
-			AdvanceTo(ETickState::CleanUp, ETickThread::Async);
-			return MipIndex; // We failed at getting this mip. Cancel will be called.
+			UE_LOG(LogTexture, Error, TEXT("Attempting to stream in mips for texture that has not generated a supported derived data key format."));
 		}
 	}
+	else
+	{
+		UE_LOG(LogTexture, Error, TEXT("Attempting to stream in mips for texture that has not generated derived data yet."));
+	}
 
 	AdvanceTo(ETickState::CleanUp, ETickThread::Async);
 	return CurrentFirstLODIdx;
@@ -128,18 +245,18 @@ bool FTexture2DMipDataProvider_DDC::PollMips(const FTextureUpdateSyncOptions& Sy
 
 void FTexture2DMipDataProvider_DDC::CleanUp(const FTextureUpdateSyncOptions& SyncOptions)
 {
-	ReleaseDDCHandles();
+	ReleaseDDCResources();
 	AdvanceTo(ETickState::Done, ETickThread::None);
 }
 
 void FTexture2DMipDataProvider_DDC::Cancel(const FTextureUpdateSyncOptions& SyncOptions)
 {
-	ReleaseDDCHandles();
+	ReleaseDDCResources();
 }
 
 FTextureMipDataProvider::ETickThread FTexture2DMipDataProvider_DDC::GetCancelThread() const
 {
-	if (DDCHandles.Num())
+	if (DDCHandles.Num() || DDCBuffers.Num())
 	{
 		return ETickThread::Async;
 	}
@@ -149,7 +266,7 @@ FTextureMipDataProvider::ETickThread FTexture2DMipDataProvider_DDC::GetCancelThr
 	}
 }
 
-void FTexture2DMipDataProvider_DDC::ReleaseDDCHandles()
+void FTexture2DMipDataProvider_DDC::ReleaseDDCResources()
 {
 	for (uint32& Handle : DDCHandles)
 	{
@@ -160,6 +277,8 @@ void FTexture2DMipDataProvider_DDC::ReleaseDDCHandles()
 		}
 	}
 	DDCHandles.Empty();
+	DDCRequestOwner.Cancel();
+	DDCBuffers.Empty();
 }
 
 #endif //WITH_EDITORONLY_DATA
diff --git a/Engine/Source/Runtime/Engine/Private/Streaming/Texture2DMipDataProvider_DDC.h b/Engine/Source/Runtime/Engine/Private/Streaming/Texture2DMipDataProvider_DDC.h
--- a/Engine/Source/Runtime/Engine/Private/Streaming/Texture2DMipDataProvider_DDC.h
+++ b/Engine/Source/Runtime/Engine/Private/Streaming/Texture2DMipDataProvider_DDC.h
@@ -7,9 +7,11 @@ Texture2DMipDataProvider_DDC.h : Implementation of FTextureMipDataProvider using
 #pragma once
 
 #include "CoreMinimal.h"
+#include "Memory/SharedBuffer.h"
 #include "Streaming/TextureMipDataProvider.h"
 
 #if WITH_EDITORONLY_DATA
+#include "DerivedDataRequestOwner.h"
 
 /**
 * FTexture2DMipDataProvider_DDC implements FTextureMipAllocator using DDC requests from FDerivedDataCacheInterface.
@@ -36,12 +38,16 @@ public:
 protected:
 
 	// Helper to route incomplete DDC requests to the FAbandonedDDCHandleManager.
-	void ReleaseDDCHandles();
+	void ReleaseDDCResources();
+
+	bool SerializeMipInfo(const FTextureUpdateContext& Context, FArchive& Ar, int32 MipIndex, const FTextureMipInfo& OutMipInfo);
 
 private:
 
 	// The list of DDC handle for each mips (0 based on the highest mip).
 	TArray<uint32, TInlineAllocator<MAX_TEXTURE_MIP_COUNT>> DDCHandles;
+	TArray<FSharedBuffer, TInlineAllocator<MAX_TEXTURE_MIP_COUNT> > DDCBuffers;
+	UE::DerivedData::FRequestOwner DDCRequestOwner;
 };
 
 #endif //WITH_EDITORONLY_DATA
diff --git a/Engine/Source/Runtime/Engine/Private/Streaming/Texture2DStreamIn_DDC.cpp b/Engine/Source/Runtime/Engine/Private/Streaming/Texture2DStreamIn_DDC.cpp
--- a/Engine/Source/Runtime/Engine/Private/Streaming/Texture2DStreamIn_DDC.cpp
+++ b/Engine/Source/Runtime/Engine/Private/Streaming/Texture2DStreamIn_DDC.cpp
@@ -7,11 +7,15 @@ Texture2DStreamIn_DDC.cpp: Stream in helper for 2D textures loading DDC files.
 #include "Streaming/Texture2DStreamIn_DDC.h"
 #include "Streaming/TextureStreamingHelpers.h"
 #include "RenderUtils.h"
-#include "DerivedDataCacheInterface.h"
 #include "Serialization/MemoryReader.h"
 
 #if WITH_EDITORONLY_DATA
 
+#include "DerivedDataCache.h"
+#include "DerivedDataCacheInterface.h"
+#include "DerivedDataCacheKey.h"
+#include "DerivedDataRequestOwner.h"
+
 int32 GStreamingUseAsyncRequestsForDDC = 1;
 static FAutoConsoleVariableRef CVarStreamingDDCPendingSleep(
 	TEXT("r.Streaming.UseAsyncRequestsForDDC"),
@@ -95,8 +99,10 @@ void PurgeAbandonedDDCHandles()
 
 FTexture2DStreamIn_DDC::FTexture2DStreamIn_DDC(UTexture2D* InTexture)
 	: FTexture2DStreamIn(InTexture)
+	, DDCRequestOwner(UE::DerivedData::EPriority::Normal)
 {
 	DDCHandles.AddZeroed(ResourceState.MaxNumLODs);
+	DDCBuffers.AddZeroed(ResourceState.MaxNumLODs);
 }
 
 FTexture2DStreamIn_DDC::~FTexture2DStreamIn_DDC()
@@ -115,30 +121,86 @@ FTexture2DStreamIn_DDC::~FTexture2DStreamIn_DDC()
 
 void FTexture2DStreamIn_DDC::DoCreateAsyncDDCRequests(const FContext& Context)
 {
-	if (Context.Texture && Context.Resource)
+	if (!Context.Texture || !Context.Resource)
+	{
+		return;
+	}
+
+	if (const FTexturePlatformData* PlatformData = Context.Texture->GetPlatformData())
 	{
-		for (int32 MipIndex = PendingFirstLODIdx; MipIndex < CurrentFirstLODIdx && !IsCancelled(); ++MipIndex)
+		const int32 LODBias = static_cast<int32>(Context.MipsView.GetData() - PlatformData->Mips.GetData());
+
+		if (PlatformData->DerivedDataKey.IsType<FString>())
 		{
-			const FTexture2DMipMap& MipMap = *Context.MipsView[MipIndex];
-			if (!MipMap.DerivedDataKey.IsEmpty())
+			for (int32 MipIndex = PendingFirstLODIdx; MipIndex < CurrentFirstLODIdx && !IsCancelled(); ++MipIndex)
 			{
-				check(!DDCHandles[MipIndex]);
-				DDCHandles[MipIndex] = GetDerivedDataCacheRef().GetAsynchronous(*MipMap.DerivedDataKey, Context.Texture->GetPathName());
+				const FTexture2DMipMap& MipMap = *Context.MipsView[MipIndex];
+				if (MipMap.IsPagedToDerivedData())
+				{
+					check(!DDCHandles[MipIndex]);
+					DDCHandles[MipIndex] = GetDerivedDataCacheRef().GetAsynchronous(*PlatformData->GetDerivedDataMipKeyString(MipIndex + LODBias, MipMap), Context.Texture->GetPathName());
 
 #if !UE_BUILD_SHIPPING
-				// On some platforms the IO is too fast to test cancelation requests timing issues.
-				if (FRenderAssetStreamingSettings::ExtraIOLatency > 0 && TaskSynchronization.GetValue() == 0)
+					// On some platforms the IO is too fast to test cancelation requests timing issues.
+					if (FRenderAssetStreamingSettings::ExtraIOLatency > 0 && TaskSynchronization.GetValue() == 0)
+					{
+						FPlatformProcess::Sleep(FRenderAssetStreamingSettings::ExtraIOLatency * .001f); // Slow down the streaming.
+					}
+#endif
+				}
+				else
 				{
-					FPlatformProcess::Sleep(FRenderAssetStreamingSettings::ExtraIOLatency * .001f); // Slow down the streaming.
+					UE_LOG(LogTexture, Error, TEXT("Attempting to stream in a mip that is already present."));
+					MarkAsCancelled();
 				}
-#endif
 			}
-			else
+		}
+		else if (PlatformData->DerivedDataKey.IsType<UE::DerivedData::FCacheKeyProxy>())
+		{
+			using namespace UE::DerivedData;
+			TArray<FCachePayloadKey> MipKeys;
+
+			for (int32 MipIndex = PendingFirstLODIdx; MipIndex < CurrentFirstLODIdx && !IsCancelled(); ++MipIndex)
+			{
+				const FTexture2DMipMap& MipMap = *Context.MipsView[MipIndex];
+				if (MipMap.IsPagedToDerivedData())
+				{
+					MipKeys.Emplace(*PlatformData->GetDerivedDataMipKeyProxy(MipIndex + LODBias, MipMap).AsCachePayloadKey());
+				}
+			}
+
+			if (MipKeys.Num())
 			{
-				UE_LOG(LogTexture, Error, TEXT("DDC key missing."));
-				MarkAsCancelled();
+				GetCache().GetPayload(MipKeys, Context.Texture->GetPathName(), ECachePolicy::Default, DDCRequestOwner,
+					[this, &Context, LODBias](FCacheGetPayloadCompleteParams&& Params)
+					{
+						if (Params.Status == EStatus::Ok)
+						{
+							for (int32 MipIndex = PendingFirstLODIdx; MipIndex < CurrentFirstLODIdx; ++MipIndex)
+							{
+								TStringBuilder<16> MipName;
+								MipName << TEXT("Mip") << (MipIndex + LODBias);
+								if (Params.Payload.GetId() == FPayloadId::FromName(MipName))
+								{
+									check(!DDCBuffers[MipIndex]);
+									DDCBuffers[MipIndex] = Params.Payload.GetData().Decompress();
+									break;
+								}
+							}
+						}
+					});
 			}
 		}
+		else
+		{
+			UE_LOG(LogTexture, Error, TEXT("Attempting to stream in mips for texture that has not generated a supported derived data key format."));
+			MarkAsCancelled();
+		}
+	}
+	else
+	{
+		UE_LOG(LogTexture, Error, TEXT("Attempting to stream in mips for texture that has not generated derived data yet."));
+		MarkAsCancelled();
 	}
 }
 
@@ -152,67 +214,117 @@ bool FTexture2DStreamIn_DDC::DoPoolDDCRequests(const FContext& Context)
 			return false;
 		}
 	}
-	return true;
+	return DDCRequestOwner.Poll();
 }
 
 void FTexture2DStreamIn_DDC::DoLoadNewMipsFromDDC(const FContext& Context)
 {
-	if (Context.Texture && Context.Resource)
+	if (!Context.Texture || !Context.Resource)
 	{
-		for (int32 MipIndex = PendingFirstLODIdx; MipIndex < CurrentFirstLODIdx && !IsCancelled(); ++MipIndex)
-		{
-			const FTexture2DMipMap& MipMap = *Context.MipsView[MipIndex];
-			check(MipData[MipIndex]);
+		return;
+	}
+
+	if (const FTexturePlatformData* PlatformData = Context.Texture->GetPlatformData())
+	{
+		const int32 LODBias = static_cast<int32>(Context.MipsView.GetData() - PlatformData->Mips.GetData());
 
-			if (!MipMap.DerivedDataKey.IsEmpty())
+		if (PlatformData->DerivedDataKey.IsType<FString>())
+		{
+			for (int32 MipIndex = PendingFirstLODIdx; MipIndex < CurrentFirstLODIdx && !IsCancelled(); ++MipIndex)
 			{
-				// The overhead of doing 2 copy of each mip data (from GetSynchronous() and FMemoryReader) in hidden by other texture DDC ops happening at the same time.
-				TArray<uint8> DerivedMipData;
-				bool bDDCValid = true;
+				const FTexture2DMipMap& MipMap = *Context.MipsView[MipIndex];
+				check(MipData[MipIndex]);
 
-				const uint32 Handle = DDCHandles[MipIndex];
-				if (Handle)
+				if (MipMap.IsPagedToDerivedData())
 				{
-					bDDCValid = GetDerivedDataCacheRef().GetAsynchronousResults(Handle, DerivedMipData);
-					DDCHandles[MipIndex] = 0;
+					// The overhead of doing 2 copy of each mip data (from GetSynchronous() and FMemoryReader) in hidden by other texture DDC ops happening at the same time.
+					TArray<uint8> DerivedMipData;
+					bool bDDCValid = true;
+
+					const uint32 Handle = DDCHandles[MipIndex];
+					if (Handle)
+					{
+						bDDCValid = GetDerivedDataCacheRef().GetAsynchronousResults(Handle, DerivedMipData);
+						DDCHandles[MipIndex] = 0;
+					}
+					else
+					{
+						bDDCValid = GetDerivedDataCacheRef().GetSynchronous(*PlatformData->GetDerivedDataMipKeyString(MipIndex + LODBias, MipMap), DerivedMipData, Context.Texture->GetPathName());
+					}
+
+					if (bDDCValid)
+					{
+						const int32 ExpectedMipSize = CalcTextureMipMapSize(MipMap.SizeX, MipMap.SizeY, Context.Resource->GetPixelFormat(), 0);
+						FMemoryReader Ar(DerivedMipData, true);
+
+						int32 MipSize = 0;
+						Ar << MipSize;
+
+						if (MipSize == ExpectedMipSize)
+						{
+							Ar.Serialize(MipData[MipIndex], MipSize);
+						}
+						else
+						{
+							UE_LOG(LogTexture, Error, TEXT("DDC mip size (%d) not as expected (%d)."), MipIndex, ExpectedMipSize);
+							MarkAsCancelled();
+						}
+					}
+					else
+					{
+						// UE_LOG(LogTexture, Warning, TEXT("Failed to stream mip data from the derived data cache for %s. Streaming mips will be recached."), Context.Texture->GetPathName() );
+						MarkAsCancelled();
+					}
 				}
 				else
 				{
-					bDDCValid = GetDerivedDataCacheRef().GetSynchronous(*MipMap.DerivedDataKey, DerivedMipData, Context.Texture->GetPathName());
+					UE_LOG(LogTexture, Error, TEXT("Attempting to stream in a mip that is already present."));
+					MarkAsCancelled();
 				}
+			}
+		}
+		else if (PlatformData->DerivedDataKey.IsType<UE::DerivedData::FCacheKeyProxy>())
+		{
+			for (int32 MipIndex = PendingFirstLODIdx; MipIndex < CurrentFirstLODIdx && !IsCancelled(); ++MipIndex)
+			{
+				const FTexture2DMipMap& MipMap = *Context.MipsView[MipIndex];
+				check(MipData[MipIndex]);
 
-				if (bDDCValid)				
+				if (MipMap.IsPagedToDerivedData())
 				{
-					const int32 ExpectedMipSize = CalcTextureMipMapSize(MipMap.SizeX, MipMap.SizeY, Context.Resource->GetPixelFormat(), 0);
-					FMemoryReader Ar(DerivedMipData, true);
-
-					int32 MipSize = 0;
-					Ar << MipSize;
+					check(DDCBuffers[MipIndex]);
 
-					if (MipSize == ExpectedMipSize)
+					const int32 ExpectedMipSize = CalcTextureMipMapSize(MipMap.SizeX, MipMap.SizeY, Context.Resource->GetPixelFormat(), 0);
+					if (DDCBuffers[MipIndex].GetSize() == ExpectedMipSize)
 					{
-						Ar.Serialize(MipData[MipIndex], MipSize);
+						FMemory::Memcpy(MipData[MipIndex], DDCBuffers[MipIndex].GetData(), DDCBuffers[MipIndex].GetSize());
 					}
 					else
 					{
-						UE_LOG(LogTexture, Error, TEXT("DDC mip size (%d) not as expected (%d)."), MipIndex, ExpectedMipSize);
+						UE_LOG(LogTexture, Error, TEXT("DDC mip size (%d) not as expected (%d) for mip %d of %s."), static_cast<int32>(DDCBuffers[MipIndex].GetSize()), ExpectedMipSize, MipIndex, *Context.Texture->GetPathName());
 						MarkAsCancelled();
 					}
+					DDCBuffers[MipIndex].Reset();
 				}
 				else
 				{
-					// UE_LOG(LogTexture, Warning, TEXT("Failed to stream mip data from the derived data cache for %s. Streaming mips will be recached."), Context.Texture->GetPathName() );
+					UE_LOG(LogTexture, Error, TEXT("Attempting to stream in a mip that is already present."));
 					MarkAsCancelled();
 				}
 			}
-			else
-			{
-				UE_LOG(LogTexture, Error, TEXT("DDC key missing."));
-				MarkAsCancelled();
-			}
+		}
+		else
+		{
+			UE_LOG(LogTexture, Error, TEXT("Attempting to stream in mips for texture that has not generated a supported derived data key format."));
+			MarkAsCancelled();
 		}
 		FPlatformMisc::MemoryBarrier();
 	}
+	else
+	{
+		UE_LOG(LogTexture, Error, TEXT("Attempting to stream in mips for texture that has not generated derived data yet."));
+		MarkAsCancelled();
+	}
 }
 
 #endif
\ No newline at end of file
diff --git a/Engine/Source/Runtime/Engine/Private/Streaming/Texture2DStreamIn_DDC.h b/Engine/Source/Runtime/Engine/Private/Streaming/Texture2DStreamIn_DDC.h
--- a/Engine/Source/Runtime/Engine/Private/Streaming/Texture2DStreamIn_DDC.h
+++ b/Engine/Source/Runtime/Engine/Private/Streaming/Texture2DStreamIn_DDC.h
@@ -7,9 +7,11 @@ Texture2DStreamIn_DDC.h: Stream in helper for 2D textures loading DDC files.
 #pragma once
 
 #include "CoreMinimal.h"
+#include "Memory/SharedBuffer.h"
 #include "Texture2DStreamIn.h"
 
 #if WITH_EDITORONLY_DATA
+#include "DerivedDataRequestOwner.h"
 
 extern int32 GStreamingUseAsyncRequestsForDDC;
 
@@ -24,6 +26,8 @@ protected:
 
 	// StreamIn_Default : Locked mips of the intermediate textures, used as disk load destination.
 	TArray<uint32, TInlineAllocator<MAX_TEXTURE_MIP_COUNT> > DDCHandles;
+	TArray<FSharedBuffer, TInlineAllocator<MAX_TEXTURE_MIP_COUNT> > DDCBuffers;
+	UE::DerivedData::FRequestOwner DDCRequestOwner;
 
 	// ****************************
 	// ********* Helpers **********
diff --git a/Engine/Source/Runtime/Engine/Private/Texture2D.cpp b/Engine/Source/Runtime/Engine/Private/Texture2D.cpp
--- a/Engine/Source/Runtime/Engine/Private/Texture2D.cpp
+++ b/Engine/Source/Runtime/Engine/Private/Texture2D.cpp
@@ -54,6 +54,11 @@
 #include "UObject/ObjectSaveContext.h"
 #include "UObject/StrongObjectPtr.h"
 
+#if WITH_EDITORONLY_DATA
+#include "Misc/TVariant.h"
+#include "DerivedDataCacheKey.h"
+#endif
+
 #define LOCTEXT_NAMESPACE "UTexture2D"
 
 UTexture2D::UTexture2D(const FObjectInitializer& ObjectInitializer)
@@ -139,7 +144,7 @@ void FTexture2DMipMap::SerializeCommon(FArchive& Ar, UObject* Owner, int32 MipId
 	Ar << bCooked;
 
 #if WITH_EDITORONLY_DATA
-	if (bIgnoreBulkDataForStreamingMips && !DerivedDataKey.IsEmpty())
+	if (bIgnoreBulkDataForStreamingMips && IsPagedToDerivedData())
 	{
 		FByteBulkData EmptyData;
 		EmptyData.Serialize(Ar, Owner, MipIdx, false, FileRegionType);
@@ -160,7 +165,17 @@ void FTexture2DMipMap::SerializeCommon(FArchive& Ar, UObject* Owner, int32 MipId
 	if (!bCooked)
 	{
 		Ar << FileRegionType;
-		Ar << DerivedDataKey;
+
+		FString LegacyDDCKeyString;
+		if (Ar.IsSaving())
+		{
+			if (IsPagedToDerivedData())
+			{
+				LegacyDDCKeyString = TEXT("DummyLegacyKey");
+			}
+		}
+
+		Ar << LegacyDDCKeyString;
 	}
 #endif // #if WITH_EDITORONLY_DATA
 }
@@ -191,7 +206,6 @@ uint32 FTexture2DMipMap::StoreInDerivedDataCache(const FString& InDerivedDataKey
 	}
 	const uint32 Result = DerivedData.Num();
 	GetDerivedDataCacheRef().Put(*InDerivedDataKey, DerivedData, TextureName, bReplaceExistingDDC);
-	DerivedDataKey = InDerivedDataKey;
 	BulkData.RemoveBulkData();
 	return Result;
 }
diff --git a/Engine/Source/Runtime/Engine/Private/TextureDerivedData.cpp b/Engine/Source/Runtime/Engine/Private/TextureDerivedData.cpp
--- a/Engine/Source/Runtime/Engine/Private/TextureDerivedData.cpp
+++ b/Engine/Source/Runtime/Engine/Private/TextureDerivedData.cpp
@@ -32,7 +32,11 @@
 
 #if WITH_EDITOR
 
+#include "DerivedDataCache.h"
 #include "DerivedDataCacheInterface.h"
+#include "DerivedDataCacheKey.h"
+#include "DerivedDataRequestOwner.h"
+#include "DerivedDataRequestTypes.h"
 #include "Interfaces/ITargetPlatform.h"
 #include "Interfaces/ITargetPlatformManagerModule.h"
 #include "Interfaces/ITextureFormat.h"
@@ -830,8 +834,7 @@ void FTexturePlatformData::Cache(
 
 	ETextureCacheFlags Flags = ETextureCacheFlags(InFlags);
 
-	static const bool bUseNewBuildPath = FParse::Param(FCommandLine::Get(), TEXT("DDC2AsyncTextureBuilds"));
-	if (bUseNewBuildPath && InTexture.Source.GetNumLayers() == 1 && !InSettingsPerLayer->bVirtualStreamable)
+	if (IsUsingNewDerivedData() && InTexture.Source.GetNumLayers() == 1 && !InSettingsPerLayer->bVirtualStreamable)
 	{
 		COOK_STAT(auto Timer = TextureCookStats::UsageStats.TimeSyncWork());
 		COOK_STAT(Timer.TrackCyclesOnly());
@@ -851,24 +854,6 @@ void FTexturePlatformData::Cache(
 
 	bool bForceRebuild = EnumHasAnyFlags(Flags, ETextureCacheFlags::ForceRebuild);
 	bool bAsync = EnumHasAnyFlags(Flags, ETextureCacheFlags::Async);
-	GetTextureDerivedDataKey(InTexture, InSettingsPerLayer, DerivedDataKey);
-
-	if (!bForceRebuild && InSettingsPerLayer[0].FastTextureEncode == ETextureFastEncode::TryOffEncodeFast)
-	{
-		// If we are running in the Editor, generate a shipping derived data key, which can be used to load
-		// a texture cooked for the shipping build from the cache, if available.
-		// Specifically, FTextureBuildSettings::FastTextureEncode controls the usage of Oodle RDO when building textures.
-		// In cases when an RDO-processed texture is present in the cache, we can use it instead of building or loading a non-RDO texture.
-		const int32 NumLayers = InTexture.Source.GetNumLayers();
-		TArray<FTextureBuildSettings> ShippingBuildSettingsPerLayer;
-		ShippingBuildSettingsPerLayer.SetNum(NumLayers);
-		for (int32 LayerIndex = 0; LayerIndex < NumLayers; ++LayerIndex)
-		{
-			ShippingBuildSettingsPerLayer[LayerIndex] = InSettingsPerLayer[LayerIndex];
-			ShippingBuildSettingsPerLayer[LayerIndex].FastTextureEncode = ETextureFastEncode::Off;
-		}
-		GetTextureDerivedDataKey(InTexture, ShippingBuildSettingsPerLayer.GetData(), ShippingDerivedDataKey);
-	}
 
 	if (!Compressor)
 	{
@@ -924,6 +909,12 @@ void FTexturePlatformData::CancelCache()
 	}
 }
 
+bool FTexturePlatformData::IsUsingNewDerivedData()
+{
+	static const bool bUseNewDerivedData = FParse::Param(FCommandLine::Get(), TEXT("DDC2AsyncTextureBuilds"));
+	return bUseNewDerivedData;
+}
+
 bool FTexturePlatformData::IsAsyncWorkComplete() const
 {
 	return !AsyncTask || AsyncTask->Poll();
@@ -956,18 +947,75 @@ typedef TArray<uint32> FAsyncVTChunkHandles;
  * @param FirstMipToLoad - Index of the first mip to retrieve.
  * @param OutHandles - Handles to the asynchronous DDC gets.
  */
-static void BeginLoadDerivedMips(TIndirectArray<FTexture2DMipMap>& Mips, int32 FirstMipToLoad, UTexture* Texture, FAsyncMipHandles& OutHandles)
+static bool BeginLoadDerivedMips(FTexturePlatformData& PlatformData, int32 FirstMipToLoad, UTexture* Texture, FAsyncMipHandles& OutHandles, TUniqueFunction<void (int32 MipIndex, FSharedBuffer MipData)> Callback)
 {
-	FDerivedDataCacheInterface& DDC = GetDerivedDataCacheRef();
-	OutHandles.AddZeroed(Mips.Num());
-	for (int32 MipIndex = FirstMipToLoad; MipIndex < Mips.Num(); ++MipIndex)
+	using namespace UE::DerivedData;
+
+	TIndirectArray<FTexture2DMipMap>& Mips = PlatformData.Mips;
+	if (PlatformData.DerivedDataKey.IsType<FString>())
 	{
-		const FTexture2DMipMap& Mip = Mips[MipIndex];
-		if (Mip.DerivedDataKey.IsEmpty() == false)
+		FDerivedDataCacheInterface& DDC = GetDerivedDataCacheRef();
+		OutHandles.AddZeroed(Mips.Num());
+		for (int32 MipIndex = FirstMipToLoad; MipIndex < Mips.Num(); ++MipIndex)
 		{
-			OutHandles[MipIndex] = DDC.GetAsynchronous(*Mip.DerivedDataKey, Texture ? Texture->GetPathName() : TEXT("Unknown Texture"_SV));
+			const FTexture2DMipMap& Mip = Mips[MipIndex];
+			if (Mip.IsPagedToDerivedData())
+			{
+				OutHandles[MipIndex] = DDC.GetAsynchronous(*PlatformData.GetDerivedDataMipKeyString(MipIndex, Mip), Texture ? Texture->GetPathName() : TEXT("Unknown Texture"_SV));
+			}
 		}
 	}
+	else if (PlatformData.DerivedDataKey.IsType<FCacheKeyProxy>())
+	{
+		TArray<FCachePayloadKey> MipKeys;
+		for (int32 MipIndex = FirstMipToLoad; MipIndex < Mips.Num(); ++MipIndex)
+		{
+			const FTexture2DMipMap& Mip = Mips[MipIndex];
+			if (Mip.IsPagedToDerivedData())
+			{
+				MipKeys.Emplace(*PlatformData.GetDerivedDataMipKeyProxy(MipIndex, Mip).AsCachePayloadKey());
+			}
+		}
+
+		if (!MipKeys.IsEmpty())
+		{
+			check(Callback);
+			bool bMiss = false;
+			FRequestOwner RequestOwner(EPriority::Blocking);
+			GetCache().GetPayload(MipKeys, Texture ? Texture->GetPathName() : TEXT("Unknown Texture"_SV), ECachePolicy::Default, RequestOwner,
+				[&Mips, Callback = MoveTemp(Callback), FirstMipToLoad, &bMiss](FCacheGetPayloadCompleteParams&& Params)
+			{
+				if (Params.Status == EStatus::Ok)
+				{
+					for (int32 MipIndex = FirstMipToLoad; MipIndex < Mips.Num(); ++MipIndex)
+					{
+						TStringBuilder<16> MipName;
+						MipName << TEXT("Mip") << MipIndex;
+						if (Params.Payload.GetId() == FPayloadId::FromName(MipName))
+						{
+							Callback(MipIndex, Params.Payload.GetData().Decompress());
+							break;
+						}
+					}
+				}
+				else
+				{
+					bMiss = true;
+				}
+			});
+			RequestOwner.Wait();
+
+			if (bMiss)
+			{
+				return false;
+			}
+		}
+	}
+	else
+	{
+		UE_LOG(LogTexture, Error, TEXT("Attempting to stream in mips for texture that has not generated a supported derived data key format."));
+	}
+	return true;
 }
 
 static void BeginLoadDerivedVTChunks(const TArray<FVirtualTextureDataChunk>& Chunks, UTexture* Texture, FAsyncVTChunkHandles& OutHandles)
@@ -1016,7 +1064,20 @@ bool FTexturePlatformData::TryInlineMipData(int32 FirstMipToLoad, UTexture* Text
 	TArray<uint8> TempData;
 	FDerivedDataCacheInterface& DDC = GetDerivedDataCacheRef();
 
-	BeginLoadDerivedMips(Mips, FirstMipToLoad, Texture, AsyncHandles);
+
+	if (!BeginLoadDerivedMips(*this, FirstMipToLoad, Texture, AsyncHandles,
+		[this](int32 MipIndex, FSharedBuffer MipBuffer)
+		{
+			FTexture2DMipMap& Mip = Mips[MipIndex];
+			Mip.BulkData.Lock(LOCK_READ_WRITE);
+			void* MipData = Mip.BulkData.Realloc(int64(MipBuffer.GetSize()));
+			FMemory::Memcpy(MipData, MipBuffer.GetData(), MipBuffer.GetSize());
+			Mip.BulkData.Unlock();
+		}))
+	{
+		return false;
+	}
+
 	if (VTData != nullptr)
 	{
 		BeginLoadDerivedVTChunks(VTData->Chunks, Texture, AsyncVTHandles);
@@ -1026,7 +1087,7 @@ bool FTexturePlatformData::TryInlineMipData(int32 FirstMipToLoad, UTexture* Text
 	for (int32 MipIndex = FirstMipToLoad; MipIndex < Mips.Num(); ++MipIndex)
 	{
 		FTexture2DMipMap& Mip = Mips[MipIndex];
-		if (Mip.DerivedDataKey.IsEmpty() == false)
+		if (Mip.IsPagedToDerivedData())
 		{
 			uint32 AsyncHandle = AsyncHandles[MipIndex];
 			bool bLoadedFromDDC = false;
@@ -1046,7 +1107,6 @@ bool FTexturePlatformData::TryInlineMipData(int32 FirstMipToLoad, UTexture* Text
 				void* MipData = Mip.BulkData.Realloc(MipSize);
 				Ar.Serialize(MipData, MipSize);
 				Mip.BulkData.Unlock();
-				Mip.DerivedDataKey.Empty();
 			}
 			else
 			{
@@ -1151,7 +1211,24 @@ bool FTexturePlatformData::TryLoadMips(int32 FirstMipToLoad, void** OutMipData,
 	TArray<uint8> TempData;
 	FAsyncMipHandles AsyncHandles;
 	FDerivedDataCacheInterface& DDC = GetDerivedDataCacheRef();
-	BeginLoadDerivedMips(Mips, FirstMipToLoad, Texture, AsyncHandles);
+	if (!BeginLoadDerivedMips(*this, FirstMipToLoad, Texture, AsyncHandles,
+		[this, OutMipData, FirstMipToLoad, &NumMipsCached](int32 MipIndex, FSharedBuffer MipBuffer)
+		{
+			FTexture2DMipMap& Mip = Mips[MipIndex];
+
+			const int32 MipSize = static_cast<int32>(MipBuffer.GetSize());
+			CheckMipSize(Mip, PixelFormat, MipSize);
+			NumMipsCached++;
+
+			if (OutMipData)
+			{
+				OutMipData[MipIndex - FirstMipToLoad] = FMemory::Malloc(MipSize);
+				FMemory::Memcpy(OutMipData[MipIndex - FirstMipToLoad], MipBuffer.GetData(), MipSize);
+			}
+		}))
+	{
+		return false;
+	}
 #endif // #if WITH_EDITOR
 
 	// Handle the case where we inlined more mips than we intend to keep resident
@@ -1193,35 +1270,38 @@ bool FTexturePlatformData::TryLoadMips(int32 FirstMipToLoad, void** OutMipData,
 	}
 
 #if WITH_EDITOR
-	// Wait for async DDC gets.
-	for (int32 MipIndex = FirstMipToLoad; MipIndex < LoadableMips; ++MipIndex)
+	if (DerivedDataKey.IsType<FString>())
 	{
-		FTexture2DMipMap& Mip = Mips[MipIndex];
-		if (Mip.DerivedDataKey.IsEmpty() == false)
+		// Wait for async DDC gets.
+		for (int32 MipIndex = FirstMipToLoad; MipIndex < LoadableMips; ++MipIndex)
 		{
-			uint32 AsyncHandle = AsyncHandles[MipIndex];
-			DDC.WaitAsynchronousCompletion(AsyncHandle);
-			if (DDC.GetAsynchronousResults(AsyncHandle, TempData))
+			FTexture2DMipMap& Mip = Mips[MipIndex];
+			if (Mip.IsPagedToDerivedData())
 			{
-				int32 MipSize = 0;
-				FMemoryReader Ar(TempData, /*bIsPersistent=*/ true);
-				Ar << MipSize;
-				CheckMipSize(Mip, PixelFormat, MipSize);
-				NumMipsCached++;
+				uint32 AsyncHandle = AsyncHandles[MipIndex];
+				DDC.WaitAsynchronousCompletion(AsyncHandle);
+				if (DDC.GetAsynchronousResults(AsyncHandle, TempData))
+				{
+					int32 MipSize = 0;
+					FMemoryReader Ar(TempData, /*bIsPersistent=*/ true);
+					Ar << MipSize;
+					CheckMipSize(Mip, PixelFormat, MipSize);
+					NumMipsCached++;
 
-				if (OutMipData)
+					if (OutMipData)
+					{
+						OutMipData[MipIndex - FirstMipToLoad] = FMemory::Malloc(MipSize);
+						Ar.Serialize(OutMipData[MipIndex - FirstMipToLoad], MipSize);
+					}
+				}
+				else
 				{
-					OutMipData[MipIndex - FirstMipToLoad] = FMemory::Malloc(MipSize);
-					Ar.Serialize(OutMipData[MipIndex - FirstMipToLoad], MipSize);
+					UE_LOG(LogTexture, Verbose, TEXT("DDC.GetAsynchronousResults() failed for %s, MipIndex: %d"),
+						Texture ? *Texture->GetPathName() : TEXT("nullptr"),
+						MipIndex);
 				}
+				TempData.Reset();
 			}
-			else
-			{
-				UE_LOG(LogTexture, Verbose, TEXT("DDC.GetAsynchronousResults() failed for %s, MipIndex: %d"),
-					Texture ? *Texture->GetPathName() : TEXT("nullptr"),
-					MipIndex);
-			}
-			TempData.Reset();
 		}
 	}
 #endif // #if WITH_EDITOR
@@ -1346,7 +1426,7 @@ bool FTexturePlatformData::CanBeLoaded() const
 	for (const FTexture2DMipMap& Mip : Mips)
 	{
 #if WITH_EDITORONLY_DATA
-		if (!Mip.DerivedDataKey.IsEmpty())
+		if (Mip.IsPagedToDerivedData())
 		{
 			return true;
 		}
@@ -1380,28 +1460,43 @@ EPixelFormat FTexturePlatformData::GetLayerPixelFormat(uint32 LayerIndex) const
 #if WITH_EDITOR
 bool FTexturePlatformData::AreDerivedMipsAvailable(FStringView Context) const
 {
-	TArray<FString, TInlineAllocator<16>> MipKeys;
-	for (const FTexture2DMipMap& Mip : Mips)
+	TArray<int32, TInlineAllocator<16>> PagedMipIndices;
+	for (int32 MipIndex = 0; MipIndex < Mips.Num(); ++MipIndex)
 	{
-		if (!Mip.DerivedDataKey.IsEmpty())
+		if (Mips[MipIndex].IsPagedToDerivedData())
 		{
-			MipKeys.Add(Mip.DerivedDataKey);
+			PagedMipIndices.Add(MipIndex);
 		}
 	}
 
-	if (MipKeys.IsEmpty())
+	if (PagedMipIndices.IsEmpty())
 	{
 		return true;
 	}
 
-	const bool bAreDerivedMipsAvailable = GetDerivedDataCacheRef().AllCachedDataProbablyExists(MipKeys);
+	bool bAreDerivedMipsAvailable = false;
 
-	// When using a shared DDC and performing async loading, 
-	// prefetch the lowest mip to avoid high prio request stalls from the render thread
-	if (bAreDerivedMipsAvailable && !IsInGameThread())
+	if (DerivedDataKey.IsType<FString>())
 	{
-		TRACE_CPUPROFILER_EVENT_SCOPE(PrefetchSmallestMip);
-		GetDerivedDataCacheRef().TryToPrefetch({ MipKeys.Last() }, Context);
+		TArray<FString, TInlineAllocator<16>> MipKeys;
+		for (int32 PagedMipIndex : PagedMipIndices)
+		{
+			MipKeys.Add(GetDerivedDataMipKeyString(PagedMipIndex, Mips[PagedMipIndex]));
+		}
+
+		bAreDerivedMipsAvailable = GetDerivedDataCacheRef().AllCachedDataProbablyExists(MipKeys);
+
+		// When using a shared DDC and performing async loading, 
+		// prefetch the lowest mip to avoid high prio request stalls from the render thread
+		if (bAreDerivedMipsAvailable && !IsInGameThread())
+		{
+			TRACE_CPUPROFILER_EVENT_SCOPE(PrefetchSmallestMip);
+			GetDerivedDataCacheRef().TryToPrefetch({MipKeys.Last()}, Context);
+		}
+	}
+	else if (DerivedDataKey.IsType<UE::DerivedData::FCachePayloadKeyProxy>())
+	{
+		bAreDerivedMipsAvailable = true;
 	}
 
 	return bAreDerivedMipsAvailable;
@@ -1686,6 +1781,24 @@ void FTexturePlatformData::SerializeWithConditionalBulkData(FArchive& Ar, UTextu
 {
 	SerializePlatformData(Ar, this, Owner, EPlatformDataSerializationFlags::IgnoreBulkDataForStreamingMips);
 }
+
+FString FTexturePlatformData::GetDerivedDataMipKeyString(int32 MipIndex, const FTexture2DMipMap& Mip) const
+{
+	const FString& KeyString = DerivedDataKey.Get<FString>();
+	return FDerivedDataCacheInterface::SanitizeCacheKey(*FString::Printf(TEXT("%s_MIP%u_%dx%d"), *KeyString, MipIndex, Mip.SizeX, Mip.SizeY));
+}
+
+UE::DerivedData::FCachePayloadKeyProxy FTexturePlatformData::GetDerivedDataMipKeyProxy(int32 MipIndex, const FTexture2DMipMap& Mip) const
+{
+	using namespace UE::DerivedData;
+	const FCacheKeyProxy& KeyProxy = DerivedDataKey.Get<FCacheKeyProxy>();
+
+	TStringBuilder<16> MipName;
+	MipName << TEXT("Mip") << MipIndex;
+	FCachePayloadKey PayloadKey {*KeyProxy.AsCacheKey(), FPayloadId::FromName(MipName)};
+	return FCachePayloadKeyProxy(PayloadKey);
+}
+
 #endif // WITH_EDITORONLY_DATA
 
 void FTexturePlatformData::SerializeCooked(FArchive& Ar, UTexture* Owner, bool bStreamable)
@@ -1778,13 +1891,57 @@ void UTexture::CachePlatformData(bool bAsyncCache, bool bAllowAsyncBuild, bool b
 		FTexturePlatformData*& PlatformDataLink = *PlatformDataLinkPtr;
 		if (Source.IsValid() && FApp::CanEverRender())
 		{
-			FString DerivedDataKey;
+			bool bPerformCache = false;
+
+			ETextureCacheFlags CacheFlags =
+				(bAsyncCache ? ETextureCacheFlags::Async : ETextureCacheFlags::None) |
+				(bAllowAsyncBuild? ETextureCacheFlags::AllowAsyncBuild : ETextureCacheFlags::None) |
+				(bAllowAsyncLoading? ETextureCacheFlags::AllowAsyncLoading : ETextureCacheFlags::None);
+
 			TArray<FTextureBuildSettings> BuildSettings;
 			GetBuildSettingsForRunningPlatform(*this, BuildSettings);
 			check(BuildSettings.Num() == Source.GetNumLayers());
-			GetTextureDerivedDataKey(*this, BuildSettings.GetData(), DerivedDataKey);
 
-			if (PlatformDataLink == NULL || PlatformDataLink->DerivedDataKey != DerivedDataKey)
+			if (PlatformDataLink == NULL)
+			{
+				bPerformCache = true;
+			}
+			else
+			{
+				if (FTexturePlatformData::IsUsingNewDerivedData())
+				{
+					using namespace UE::DerivedData;
+					if (const FCacheKeyProxy* ExistingDerivedDataKey = PlatformDataLink->ComparisonDerivedDataKey.TryGet<FCacheKeyProxy>())
+					{
+						if (*ExistingDerivedDataKey->AsCacheKey() != *CreateTextureCacheKeyProxy(*this, CacheFlags, BuildSettings[0]).AsCacheKey())
+						{
+							bPerformCache = true;
+						}
+					}
+					else
+					{
+						bPerformCache = true;
+					}
+				}
+				else
+				{
+					if (const FString* ExistingDerivedDataKey = PlatformDataLink->ComparisonDerivedDataKey.TryGet<FString>())
+					{
+						FString DerivedDataKey;
+						GetTextureDerivedDataKey(*this, BuildSettings.GetData(), DerivedDataKey);
+						if (*ExistingDerivedDataKey != DerivedDataKey)
+						{
+							bPerformCache = true;
+						}
+					}
+					else
+					{
+						bPerformCache = true;
+					}
+				}
+			}
+
+			if (bPerformCache)
 			{
 				// Release our resource if there is existing derived data.
 				if (PlatformDataLink)
@@ -1799,10 +1956,6 @@ void UTexture::CachePlatformData(bool bAsyncCache, bool bAllowAsyncBuild, bool b
 				{
 					PlatformDataLink = new FTexturePlatformData();
 				}
-				ETextureCacheFlags CacheFlags = 
-					(bAsyncCache ? ETextureCacheFlags::Async : ETextureCacheFlags::None) |
-					(bAllowAsyncBuild? ETextureCacheFlags::AllowAsyncBuild : ETextureCacheFlags::None) |
-					(bAllowAsyncLoading? ETextureCacheFlags::AllowAsyncLoading : ETextureCacheFlags::None);
 
 				PlatformDataLink->Cache(*this, BuildSettings.GetData(), uint32(CacheFlags), Compressor);
 			}
@@ -2088,19 +2241,6 @@ void UTexture::FinishCachePlatformData()
 				// make sure async requests are finished
 				RunningPlatformData->FinishCache();
 			}
-
-#if DO_CHECK
-			if (!GetOutermost()->HasAnyPackageFlags(PKG_FilterEditorOnly))
-			{
-				FString DerivedDataKey;
-				TArray<FTextureBuildSettings> BuildSettings;
-				GetBuildSettingsForRunningPlatform(*this, BuildSettings);
-				check(BuildSettings.Num() == Source.GetNumLayers());
-				GetTextureDerivedDataKey(*this, BuildSettings.GetData(), DerivedDataKey);
-
-				check(!RunningPlatformData || RunningPlatformData->DerivedDataKey == DerivedDataKey);
-			}
-#endif
 		}
 	}
 
@@ -2140,15 +2280,18 @@ void UTexture::MarkPlatformDataTransient()
 		FTexturePlatformData* PlatformData = *RunningPlatformData;
 		if (PlatformData)
 		{
-			for (int32 MipIndex = 0; MipIndex < PlatformData->Mips.Num(); ++MipIndex)
+			if (const FString* KeyString = PlatformData->DerivedDataKey.TryGet<FString>())
 			{
-				FTexture2DMipMap& Mip = PlatformData->Mips[MipIndex];
-				if (Mip.DerivedDataKey.Len() > 0)
+				for (int32 MipIndex = 0; MipIndex < PlatformData->Mips.Num(); ++MipIndex)
 				{
-					DDC.MarkTransient(*Mip.DerivedDataKey);
+					FTexture2DMipMap& Mip = PlatformData->Mips[MipIndex];
+					if (Mip.IsPagedToDerivedData())
+					{
+						DDC.MarkTransient(*PlatformData->GetDerivedDataMipKeyString(MipIndex, Mip));
+					}
 				}
+				DDC.MarkTransient(**KeyString);
 			}
-			DDC.MarkTransient(*PlatformData->DerivedDataKey);
 		}
 	}
 
@@ -2158,15 +2301,18 @@ void UTexture::MarkPlatformDataTransient()
 		for ( auto It : *CookedPlatformData )
 		{
 			FTexturePlatformData* PlatformData = It.Value;
-			for (int32 MipIndex = 0; MipIndex < PlatformData->Mips.Num(); ++MipIndex)
+			if (const FString* KeyString = PlatformData->DerivedDataKey.TryGet<FString>())
 			{
-				FTexture2DMipMap& Mip = PlatformData->Mips[MipIndex];
-				if (Mip.DerivedDataKey.Len() > 0)
+				for (int32 MipIndex = 0; MipIndex < PlatformData->Mips.Num(); ++MipIndex)
 				{
-					DDC.MarkTransient(*Mip.DerivedDataKey);
+					FTexture2DMipMap& Mip = PlatformData->Mips[MipIndex];
+					if (Mip.IsPagedToDerivedData())
+					{
+						DDC.MarkTransient(*PlatformData->GetDerivedDataMipKeyString(MipIndex, Mip));
+					}
 				}
+				DDC.MarkTransient(**KeyString);
 			}
-			DDC.MarkTransient(*PlatformData->DerivedDataKey);
 		}
 	}
 }
diff --git a/Engine/Source/Runtime/Engine/Private/TextureDerivedDataTask.cpp b/Engine/Source/Runtime/Engine/Private/TextureDerivedDataTask.cpp
--- a/Engine/Source/Runtime/Engine/Private/TextureDerivedDataTask.cpp
+++ b/Engine/Source/Runtime/Engine/Private/TextureDerivedDataTask.cpp
@@ -31,6 +31,7 @@
 #include "DerivedDataBuildOutput.h"
 #include "DerivedDataBuildSession.h"
 #include "DerivedDataCacheInterface.h"
+#include "DerivedDataCacheKey.h"
 #include "DerivedDataPayload.h"
 #include "DerivedDataRequestOwner.h"
 #include "Engine/TextureCube.h"
@@ -414,7 +415,7 @@
 		for (int32 MipIndex = 0; MipIndex < DerivedData->Mips.Num(); ++MipIndex)
 		{
 			FTexture2DMipMap& Mip = DerivedData->Mips[MipIndex];
-			if (Mip.DerivedDataKey.IsEmpty())
+			if (!Mip.IsPagedToDerivedData())
 			{
 				break;
 			}
@@ -435,7 +436,7 @@
 				int32 MipSize = TCheckValueCast<int32>(MipData.GetSize());
 				MipWriter << MipSize;
 				MipWriter.Serialize(const_cast<void*>(MipData.GetData()), MipSize);
-				GetDerivedDataCacheRef().Put(*Mip.DerivedDataKey, MipDataCopy, TexturePath, bReplaceExistingDDC);
+				GetDerivedDataCacheRef().Put(*DerivedData->GetDerivedDataMipKeyString(MipIndex, Mip), MipDataCopy, TexturePath, bReplaceExistingDDC);
 			}
 
 			if (bInlineMips && (MipIndex >= (int32)BuildSettingsPerLayer[0].LODBiasWithCinematicMips))
@@ -446,14 +447,11 @@
 				void* MipAllocData = Mip.BulkData.Realloc(int64(MipSize));
 				MakeMemoryView(MipAllocData, MipSize).CopyFrom(MipData);
 				Mip.BulkData.Unlock();
-				Mip.DerivedDataKey.Empty();
 			}
 		}
 
-		FString DerivedDataKey;
-		GetTextureDerivedDataKeyFromSuffix(KeySuffix, DerivedDataKey);
 		TArrayView<const uint8> PrimaryDataView((const uint8*)PrimaryData.GetData(), PrimaryData.GetSize());
-		GetDerivedDataCacheRef().Put(*DerivedDataKey, PrimaryDataView, TexturePath, bReplaceExistingDDC);
+		GetDerivedDataCacheRef().Put(*DerivedData->DerivedDataKey.Get<FString>(), PrimaryDataView, TexturePath, bReplaceExistingDDC);
 
 		bSucceeded = true;
 	}
@@ -723,7 +721,6 @@
 		DerivedData->VTData = nullptr;
 	}
 	UTexture::GetPixelFormatEnum();
-	GetTextureDerivedDataKeySuffix(Texture, InSettingsPerLayer, KeySuffix);
 		
 	const bool bAllowAsyncBuild = EnumHasAnyFlags(CacheFlags, ETextureCacheFlags::AllowAsyncBuild);
 	const bool bAllowAsyncLoading = EnumHasAnyFlags(CacheFlags, ETextureCacheFlags::AllowAsyncLoading);
@@ -795,24 +792,53 @@
 	const bool bAllowAsyncBuild = EnumHasAnyFlags(CacheFlags, ETextureCacheFlags::AllowAsyncBuild);
 	const bool bAllowAsyncLoading = EnumHasAnyFlags(CacheFlags, ETextureCacheFlags::AllowAsyncLoading);
 	const bool bForVirtualTextureStreamingBuild = EnumHasAnyFlags(CacheFlags, ETextureCacheFlags::ForVirtualTextureStreamingBuild);
+	const bool bValidateVirtualTextureCompression = CVarVTValidateCompressionOnLoad.GetValueOnAnyThread() != 0;
 	bool bInvalidVirtualTextureCompression = false;
 
 	TArray<uint8> RawDerivedData;
 
+	FString LocalDerivedDataKeySuffix;
+	FString LocalDerivedDataKey;
+	FString LocalComparisonDerivedDataKey;
+	GetTextureDerivedDataKeySuffix(Texture, BuildSettingsPerLayer.GetData(), LocalDerivedDataKeySuffix);
+	GetTextureDerivedDataKeyFromSuffix(LocalDerivedDataKeySuffix, LocalDerivedDataKey);
+	LocalComparisonDerivedDataKey = LocalDerivedDataKey;
 	if (!bForceRebuild)
 	{
 		// First try to load a texture generated for the shipping build from the cache.
 		// FTexturePlatformData::ShippingDerivedDataKey is set when we are running a build in the Editor.
 		// This allows to preview how the texture will look in the final build and avoid rebuilding texture locally using fast cooking.
-		if (!DerivedData->ShippingDerivedDataKey.IsEmpty() && DerivedData->ShippingDerivedDataKey != DerivedData->DerivedDataKey)
+		if (BuildSettingsPerLayer[0].FastTextureEncode == ETextureFastEncode::TryOffEncodeFast)
 		{
-			bLoadedFromDDC = GetDerivedDataCacheRef().GetSynchronous(*DerivedData->ShippingDerivedDataKey, RawDerivedData, Texture.GetPathName());
+			const int32 NumLayers = Texture.Source.GetNumLayers();
+			TArray<FTextureBuildSettings> ShippingBuildSettingsPerLayer;
+			ShippingBuildSettingsPerLayer.SetNum(NumLayers);
+			for (int32 LayerIndex = 0; LayerIndex < NumLayers; ++LayerIndex)
+			{
+				ShippingBuildSettingsPerLayer[LayerIndex] = BuildSettingsPerLayer[LayerIndex];
+				ShippingBuildSettingsPerLayer[LayerIndex].FastTextureEncode = ETextureFastEncode::Off;
+			}
+			FString ShippingDerivedDataKeySuffix;
+			FString ShippingDerivedDataKey;
+			GetTextureDerivedDataKeySuffix(Texture, ShippingBuildSettingsPerLayer.GetData(), ShippingDerivedDataKeySuffix);
+			GetTextureDerivedDataKeyFromSuffix(ShippingDerivedDataKeySuffix, ShippingDerivedDataKey);
+
+			bLoadedFromDDC = GetDerivedDataCacheRef().GetSynchronous(*ShippingDerivedDataKey, RawDerivedData, Texture.GetPathName());
+			if (bLoadedFromDDC)
+			{
+				LocalDerivedDataKeySuffix = ShippingDerivedDataKeySuffix;
+				LocalDerivedDataKey = ShippingDerivedDataKey;
+			}
 		}
+
 		if (!bLoadedFromDDC)
 		{
-			bLoadedFromDDC = GetDerivedDataCacheRef().GetSynchronous(*DerivedData->DerivedDataKey, RawDerivedData, Texture.GetPathName());
+			bLoadedFromDDC = GetDerivedDataCacheRef().GetSynchronous(*LocalDerivedDataKey, RawDerivedData, Texture.GetPathName());
 		}
 	}
+	KeySuffix = LocalDerivedDataKeySuffix;
+	DerivedData->DerivedDataKey.Emplace<FString>(LocalDerivedDataKey);
+	DerivedData->ComparisonDerivedDataKey.Emplace<FString>(LocalComparisonDerivedDataKey);
 
 	if (bLoadedFromDDC)
 	{
@@ -1062,7 +1088,7 @@
 		}
 
 		FBuildDefinition Definition = CreateDefinition(Build, Texture, TexturePath, FunctionName, KeySuffix, Settings, bUseCompositeTexture);
-
+		
 		if (!EnumHasAnyFlags(Flags, ETextureCacheFlags::ForceRebuild) && Settings.FastTextureEncode == ETextureFastEncode::TryOffEncodeFast)
 		{
 			FTextureBuildSettings ShippingSettings = Settings;
@@ -1075,7 +1101,12 @@
 					{
 					default:
 					case EStatus::Ok:
-						return EndBuild(MoveTemp(Params.Output), Params.BuildStatus);
+						BuildSession.Get().Build(Definition, EBuildPolicy::None, *Owner,
+							[this](FBuildCompleteParams&& Params)
+							{
+								DerivedData.ComparisonDerivedDataKey.Emplace<FCacheKeyProxy>(Params.CacheKey);
+							});
+						return EndBuild(Params.CacheKey, MoveTemp(Params.Output), Params.BuildStatus);
 					case EStatus::Error:
 						return BeginBuild(Definition, Flags);
 					}
@@ -1116,12 +1147,17 @@
 			BuildPolicy &= ~EBuildPolicy::CacheQuery;
 		}
 		BuildSession.Get().Build(Definition, BuildPolicy, *Owner,
-			[this](FBuildCompleteParams&& Params) { EndBuild(MoveTemp(Params.Output), Params.BuildStatus); });
+			[this](FBuildCompleteParams&& Params)
+			{
+				DerivedData.ComparisonDerivedDataKey.Emplace<FCacheKeyProxy>(Params.CacheKey);
+				EndBuild(Params.CacheKey, MoveTemp(Params.Output), Params.BuildStatus);
+			});
 	}
 
-	void EndBuild(UE::DerivedData::FBuildOutput&& Output, UE::DerivedData::EBuildStatus Status)
+	void EndBuild(const UE::DerivedData::FCacheKey& CacheKey, UE::DerivedData::FBuildOutput&& Output, UE::DerivedData::EBuildStatus Status)
 	{
 		using namespace UE::DerivedData;
+		DerivedData.DerivedDataKey.Emplace<FCacheKeyProxy>(CacheKey);
 		bCacheHit = EnumHasAnyFlags(Status, EBuildStatus::CacheQueryHit);
 		BuildOutputSize = Algo::TransformAccumulate(Output.GetPayloads(),
 			[](const FPayload& Payload) { return Payload.GetData().GetRawSize(); }, uint64(0));
@@ -1181,114 +1217,6 @@
 		return Owner->Poll();
 	}
 
-private:
-	void WriteDerivedData(UE::DerivedData::FBuildOutput&& Output)
-	{
-		using namespace UE::DerivedData;
-
-		Output.IterateDiagnostics([](const FBuildDiagnostic& Diagnostic)
-		{
-			if (Diagnostic.Level == EBuildDiagnosticLevel::Error)
-			{
-				UE_LOG(LogTexture, Warning, TEXT("[Build Error] %.*s: %.*s"),
-					Diagnostic.Category.Len(), Diagnostic.Category.GetData(),
-					Diagnostic.Message.Len(), Diagnostic.Message.GetData());
-			}
-			else
-			{
-				UE_LOG(LogTexture, Warning, TEXT("[Build Warning] %.*s: %.*s"),
-					Diagnostic.Category.Len(), Diagnostic.Category.GetData(),
-					Diagnostic.Message.Len(), Diagnostic.Message.GetData());
-			}
-		});
-
-		if (Output.HasError())
-		{
-			UE_LOG(LogTexture, Warning, TEXT("Failed to build derived data for build of '%s' by %s."),
-				*WriteToString<128>(Output.GetName()), *WriteToString<32>(Output.GetFunction()));
-			return;
-		}
-
-		if (const FPayload& Payload = Output.GetPayload(FPayloadId::FromName(TEXT("Texture"))))
-		{
-			FSharedBuffer Data = Payload.GetData().Decompress();
-			FMemoryReaderView Ar(Data, /*bIsPersistent*/ true);
-			DerivedData.Serialize(Ar, nullptr);
-		}
-		else
-		{
-			UE_LOG(LogTexture, Warning, TEXT("Missing output for build of '%s' by %s."),
-				*WriteToString<128>(Output.GetName()), *WriteToString<32>(Output.GetFunction()));
-			return;
-		}
-
-		for (int32 MipIndex = 0, MipCount = DerivedData.Mips.Num(); MipIndex < MipCount; ++MipIndex)
-		{
-			FTexture2DMipMap& Mip = DerivedData.Mips[MipIndex];
-			if (Mip.DerivedDataKey.IsEmpty())
-			{
-				break;
-			}
-
-			if (const FPayload& MipPayload = Output.GetPayload(FPayloadId::FromName(WriteToString<8>(TEXT("Mip"), MipIndex))))
-			{
-				if (bInlineMips && MipIndex >= FirstMipToLoad)
-				{
-					FSharedBuffer MipData = MipPayload.GetData().Decompress();
-					const uint64 MipSize = MipData.GetSize();
-
-					Mip.BulkData.Lock(LOCK_READ_WRITE);
-					void* MipAllocData = Mip.BulkData.Realloc(int64(MipSize));
-					MakeMemoryView(MipAllocData, MipSize).CopyFrom(MipData);
-					Mip.BulkData.Unlock();
-					Mip.DerivedDataKey.Empty();
-				}
-			}
-			else
-			{
-				UE_LOG(LogTexture, Warning, TEXT("Missing output for mip %d for build of '%s' by %s."),
-					MipIndex, *WriteToString<128>(Output.GetName()), *WriteToString<32>(Output.GetFunction()));
-				return;
-			}
-		}
-	}
-
-	static UE::DerivedData::EPriority ConvertPriority(EQueuedWorkPriority SourcePriority)
-	{
-		using namespace UE::DerivedData;
-		switch (SourcePriority)
-		{
-		case EQueuedWorkPriority::Lowest:  return EPriority::Lowest;
-		case EQueuedWorkPriority::Low:     return EPriority::Low;
-		case EQueuedWorkPriority::Normal:  return EPriority::Normal;
-		case EQueuedWorkPriority::High:    return EPriority::High;
-		case EQueuedWorkPriority::Highest: return EPriority::Highest;
-		default:                           return EPriority::Normal;
-		}
-	}
-
-	static EQueuedWorkPriority ConvertPriority(UE::DerivedData::EPriority SourcePriority)
-	{
-		using namespace UE::DerivedData;
-		switch (SourcePriority)
-		{
-		case EPriority::Lowest:   return EQueuedWorkPriority::Lowest;
-		case EPriority::Low:      return EQueuedWorkPriority::Low;
-		case EPriority::Normal:   return EQueuedWorkPriority::Normal;
-		case EPriority::High:     return EQueuedWorkPriority::High;
-		case EPriority::Highest:  return EQueuedWorkPriority::Highest;
-		case EPriority::Blocking: return EQueuedWorkPriority::Highest;
-		default:                  return EQueuedWorkPriority::Normal;
-		}
-	}
-
-	static bool LoadModules()
-	{
-		FModuleManager::LoadModuleChecked<IImageWrapperModule>(FName("ImageWrapper"));
-		FModuleManager::LoadModuleChecked<ITextureCompressorModule>(TEXTURE_COMPRESSOR_MODULENAME);
-		return true;
-	}
-
 	static bool IsTextureValidForBuilding(UTexture& Texture, ETextureCacheFlags Flags, bool& bOutUseCompositeTexture)
 	{
 		const int32 NumBlocks = Texture.Source.GetNumBlocks();
@@ -1309,7 +1237,7 @@
 			case TSF_BGRE8:
 			case TSF_RGBA16:
 			case TSF_RGBA16F:
-			break;
+				break;
 			default:
 				UE_LOG(LogTexture, Fatal, TEXT("Texture %s has source art in an invalid format."), *Texture.GetPathName());
 				return false;
@@ -1399,6 +1327,113 @@
 		return true;
 	}
 
+private:
+	void WriteDerivedData(UE::DerivedData::FBuildOutput&& Output)
+	{
+		using namespace UE::DerivedData;
+
+		Output.IterateDiagnostics([](const FBuildDiagnostic& Diagnostic)
+		{
+			if (Diagnostic.Level == EBuildDiagnosticLevel::Error)
+			{
+				UE_LOG(LogTexture, Warning, TEXT("[Build Error] %.*s: %.*s"),
+					Diagnostic.Category.Len(), Diagnostic.Category.GetData(),
+					Diagnostic.Message.Len(), Diagnostic.Message.GetData());
+			}
+			else
+			{
+				UE_LOG(LogTexture, Warning, TEXT("[Build Warning] %.*s: %.*s"),
+					Diagnostic.Category.Len(), Diagnostic.Category.GetData(),
+					Diagnostic.Message.Len(), Diagnostic.Message.GetData());
+			}
+		});
+
+		if (Output.HasError())
+		{
+			UE_LOG(LogTexture, Warning, TEXT("Failed to build derived data for build of '%s' by %s."),
+				*WriteToString<128>(Output.GetName()), *WriteToString<32>(Output.GetFunction()));
+			return;
+		}
+
+		if (const FPayload& Payload = Output.GetPayload(FPayloadId::FromName(TEXT("Texture"))))
+		{
+			FSharedBuffer Data = Payload.GetData().Decompress();
+			FMemoryReaderView Ar(Data, /*bIsPersistent*/ true);
+			DerivedData.Serialize(Ar, nullptr);
+		}
+		else
+		{
+			UE_LOG(LogTexture, Warning, TEXT("Missing output for build of '%s' by %s."),
+				*WriteToString<128>(Output.GetName()), *WriteToString<32>(Output.GetFunction()));
+			return;
+		}
+
+		for (int32 MipIndex = 0, MipCount = DerivedData.Mips.Num(); MipIndex < MipCount; ++MipIndex)
+		{
+			FTexture2DMipMap& Mip = DerivedData.Mips[MipIndex];
+			if (!Mip.IsPagedToDerivedData())
+			{
+				break;
+			}
+
+			if (const FPayload& MipPayload = Output.GetPayload(FPayloadId::FromName(WriteToString<8>(TEXT("Mip"), MipIndex))))
+			{
+				if (bInlineMips && MipIndex >= FirstMipToLoad)
+				{
+					FSharedBuffer MipData = MipPayload.GetData().Decompress();
+					const uint64 MipSize = MipData.GetSize();
+
+					Mip.BulkData.Lock(LOCK_READ_WRITE);
+					void* MipAllocData = Mip.BulkData.Realloc(int64(MipSize));
+					MakeMemoryView(MipAllocData, MipSize).CopyFrom(MipData);
+					Mip.BulkData.Unlock();
+				}
+			}
+			else
+			{
+				UE_LOG(LogTexture, Warning, TEXT("Missing output for mip %d for build of '%s' by %s."),
+					MipIndex, *WriteToString<128>(Output.GetName()), *WriteToString<32>(Output.GetFunction()));
+				return;
+			}
+		}
+	}
+
+	static UE::DerivedData::EPriority ConvertPriority(EQueuedWorkPriority SourcePriority)
+	{
+		using namespace UE::DerivedData;
+		switch (SourcePriority)
+		{
+		case EQueuedWorkPriority::Lowest:  return EPriority::Lowest;
+		case EQueuedWorkPriority::Low:     return EPriority::Low;
+		case EQueuedWorkPriority::Normal:  return EPriority::Normal;
+		case EQueuedWorkPriority::High:    return EPriority::High;
+		case EQueuedWorkPriority::Highest: return EPriority::Highest;
+		default:                           return EPriority::Normal;
+		}
+	}
+
+	static EQueuedWorkPriority ConvertPriority(UE::DerivedData::EPriority SourcePriority)
+	{
+		using namespace UE::DerivedData;
+		switch (SourcePriority)
+		{
+		case EPriority::Lowest:   return EQueuedWorkPriority::Lowest;
+		case EPriority::Low:      return EQueuedWorkPriority::Low;
+		case EPriority::Normal:   return EQueuedWorkPriority::Normal;
+		case EPriority::High:     return EQueuedWorkPriority::High;
+		case EPriority::Highest:  return EQueuedWorkPriority::Highest;
+		case EPriority::Blocking: return EQueuedWorkPriority::Highest;
+		default:                  return EQueuedWorkPriority::Normal;
+		}
+	}
+
+	static bool LoadModules()
+	{
+		FModuleManager::LoadModuleChecked<IImageWrapperModule>(FName("ImageWrapper"));
+		FModuleManager::LoadModuleChecked<ITextureCompressorModule>(TEXTURE_COMPRESSOR_MODULENAME);
+		return true;
+	}
+
 	FTexturePlatformData& DerivedData;
 	TOptional<UE::DerivedData::FRequestOwner> Owner;
 	UE::DerivedData::FOptionalBuildSession BuildSession;
@@ -1427,4 +1462,47 @@
 	return nullptr;
 }
 
+UE::DerivedData::FCacheKeyProxy CreateTextureCacheKeyProxy(
+	UTexture& Texture,
+	ETextureCacheFlags CacheFlags,
+	const FTextureBuildSettings& Settings)
+{
+	using namespace UE::DerivedData;
+
+	TStringBuilder<64> FunctionName;
+	if (TryFindTextureBuildFunction(FunctionName, Settings))
+	{
+		IBuild& Build = GetBuild();
+
+		TStringBuilder<256> TexturePath;
+		Texture.GetPathName(nullptr, TexturePath);
+
+		FString KeySuffix;
+		GetTextureDerivedDataKeySuffix(Texture, &Settings, KeySuffix);
+
+		bool bUseCompositeTexture = false;
+		if (FTextureBuildTask::IsTextureValidForBuilding(Texture, CacheFlags, bUseCompositeTexture))
+		{
+			TOptional<UE::TextureDerivedData::FTextureBuildInputResolver> PlaceholderResolver;
+			UE::DerivedData::IBuildInputResolver* InputResolver = GetGlobalBuildInputResolver();
+			if (!InputResolver)
+			{
+				PlaceholderResolver.Emplace(Texture);
+				InputResolver = &PlaceholderResolver.GetValue();
+			}
+			FBuildSession Session = Build.CreateSession(TexturePath, InputResolver);
+			FBuildDefinition Definition = FTextureBuildTask::CreateDefinition(Build, Texture, TexturePath, FunctionName, KeySuffix, Settings, bUseCompositeTexture);
+			
+			FCacheKey GeneratedKey;
+			FRequestOwner CacheKeyGenerationOwner(EPriority::Blocking);
+			Session.Build(Definition, EBuildPolicy::None, CacheKeyGenerationOwner,
+				[&GeneratedKey](FBuildCompleteParams&& Params) { GeneratedKey = Params.CacheKey; });
+			CacheKeyGenerationOwner.Wait();
+
+			return GeneratedKey;
+		}
+	}
+	return FCacheKey::Empty;
+}
+
 #endif // WITH_EDITOR
diff --git a/Engine/Source/Runtime/Engine/Private/TextureDerivedDataTask.h b/Engine/Source/Runtime/Engine/Private/TextureDerivedDataTask.h
--- a/Engine/Source/Runtime/Engine/Private/TextureDerivedDataTask.h
+++ b/Engine/Source/Runtime/Engine/Private/TextureDerivedDataTask.h
@@ -11,6 +11,7 @@
 #if WITH_EDITOR
 
 #include "Async/AsyncWork.h"
+#include "DerivedDataCacheKeyProxy.h"
 #include "Engine/Texture2D.h"
 #include "IImageWrapperModule.h"
 #include "ImageCore.h"
@@ -141,7 +142,6 @@ class FTextureCacheDerivedDataWorker : public FNonAbandonableTask
 	void BuildTexture(bool bReplaceExistingDDC = false);
 
 	void ConsumeBuildFunctionOutput(const UE::DerivedData::FBuildOutput& BuildOutput, const FString& TexturePath, bool bReplaceExistingDDC);
-
 public:
 
 	/** Initialization constructor. */
@@ -264,4 +264,9 @@ FTextureAsyncCacheDerivedDataTask* CreateTextureBuildTask(
 	EQueuedWorkPriority Priority,
 	ETextureCacheFlags Flags);
 
+UE::DerivedData::FCacheKeyProxy CreateTextureCacheKeyProxy(
+	UTexture& Texture,
+	ETextureCacheFlags CacheFlags,
+	const FTextureBuildSettings& Settings);
+
 #endif // WITH_EDITOR
diff --git a/Engine/Source/Runtime/Engine/Public/DerivedDataCacheKeyProxy.h b/Engine/Source/Runtime/Engine/Public/DerivedDataCacheKeyProxy.h
new file mode 100644
--- /dev/null
+++ b/Engine/Source/Runtime/Engine/Public/DerivedDataCacheKeyProxy.h
@@ -0,0 +1,46 @@
+// Copyright Epic Games, Inc. All Rights Reserved.
+
+#pragma once
+
+#include "CoreTypes.h"
+#include "IO/IoHash.h"
+#include "Templates/TypeCompatibleBytes.h"
+
+namespace UE::DerivedData
+{
+
+struct FCacheKey;
+struct FCachePayloadKey;
+
+namespace Private
+{
+	struct FCacheKeyDummy
+	{
+		const ANSICHAR* BucketNamePtrDummy = nullptr;
+		FIoHash HashDummy;
+	};
+
+	struct FCachePayloadKeyDummy
+	{
+		FCacheKeyDummy CacheKeyDummy;
+		alignas(uint32) uint8 PayloadBytesDummy[12];
+	};
+} // Private
+
+struct FCacheKeyProxy : private TAlignedBytes<sizeof(UE::DerivedData::Private::FCacheKeyDummy), alignof(UE::DerivedData::Private::FCacheKeyDummy)>
+{
+	FCacheKeyProxy(const FCacheKey& InKey);
+	~FCacheKeyProxy();
+	FCacheKey* AsCacheKey() { return (FCacheKey*)this;  }
+	const FCacheKey* AsCacheKey() const { return (const FCacheKey*)this; }
+};
+
+struct FCachePayloadKeyProxy : private TAlignedBytes<sizeof(UE::DerivedData::Private::FCachePayloadKeyDummy), alignof(UE::DerivedData::Private::FCachePayloadKeyDummy)>
+{
+	FCachePayloadKeyProxy(const FCachePayloadKey& InKey);
+	~FCachePayloadKeyProxy();
+	FCachePayloadKey* AsCachePayloadKey() { return (FCachePayloadKey*)this;  }
+	const FCachePayloadKey* AsCachePayloadKey() const { return (const FCachePayloadKey*)this; }
+};
+
+} // UE::DerivedData
diff --git a/Engine/Source/Runtime/Engine/Public/TextureResource.h b/Engine/Source/Runtime/Engine/Public/TextureResource.h
--- a/Engine/Source/Runtime/Engine/Public/TextureResource.h
+++ b/Engine/Source/Runtime/Engine/Public/TextureResource.h
@@ -64,19 +64,20 @@ struct FTexture2DMipMap
 #if WITH_EDITORONLY_DATA
 	void SerializeWithConditionalBulkData(FArchive& Ar, UObject* Owner, int32 MipIndex);
 
-	/** Key if stored in the derived data cache. */
-	FString DerivedDataKey;
-
 	/** The file region type appropriate for this mip's pixel format. */
 	EFileRegionType FileRegionType = EFileRegionType::None;
 
+	bool IsPagedToDerivedData() const { return BulkData.GetBulkDataSize() == 0; }
+
 	/**
 	 * Place mip-map data in the derived data cache associated with the provided
 	 * key.
 	 */
 	uint32 StoreInDerivedDataCache(const FString& InDerivedDataKey, const FStringView& TextureName, bool bReplaceExistingDDC);
+
 #endif // #if WITH_EDITORONLY_DATA
 private:
+
 	void SerializeCommon(FArchive& Ar, UObject* Owner, int32 MipIndex, bool bBulkDataConditionalOnKey);
 };
 
