Cloth - Don't run the cloth simulation when bDisableClothSimulation is set.

diff --git a/Engine/Source/Runtime/Engine/Private/SkeletalMeshComponentPhysics.cpp b/Engine/Source/Runtime/Engine/Private/SkeletalMeshComponentPhysics.cpp
--- a/Engine/Source/Runtime/Engine/Private/SkeletalMeshComponentPhysics.cpp
+++ b/Engine/Source/Runtime/Engine/Private/SkeletalMeshComponentPhysics.cpp
@@ -3594,7 +3594,7 @@ const TMap<int32, FClothSimulData>& USkeletalMeshComponent::GetCurrentClothingDa
 void USkeletalMeshComponent::UpdateClothStateAndSimulate(float DeltaTime, FTickFunction& ThisTickFunction)
 {
 	// If disabled or no simulation
-	if (CVarEnableClothPhysics.GetValueOnGameThread() == 0 || !ClothingSimulation)
+	if (CVarEnableClothPhysics.GetValueOnGameThread() == 0 || !ClothingSimulation || bDisableClothSimulation)
 	{
 		return;
 	}
