Fix recompute tangent intermediate accumulated tangent buffer size being allocated more than needed
1) only allocate enough to accommodate sections doing recompute tangent
2) allocate for vertex count instead of index count (i.e. triangle * 3)

diff --git a/Engine/Source/Runtime/Engine/Private/GPUSkinCache.cpp b/Engine/Source/Runtime/Engine/Private/GPUSkinCache.cpp
--- a/Engine/Source/Runtime/Engine/Private/GPUSkinCache.cpp
+++ b/Engine/Source/Runtime/Engine/Private/GPUSkinCache.cpp
@@ -308,6 +308,7 @@ public:
 		FGPUSkinCache::FSkinCacheRWBuffer* TangentBuffer = nullptr;
 		FGPUSkinCache::FSkinCacheRWBuffer* IntermediateTangentBuffer = nullptr;
 		FGPUSkinCache::FSkinCacheRWBuffer* IntermediateAccumulatedTangentBuffer = nullptr;
+		uint32 IntermediateAccumulatedTangentBufferOffset = 0;
 		FGPUSkinCache::FSkinCacheRWBuffer* PositionBuffer = nullptr;
 		FGPUSkinCache::FSkinCacheRWBuffer* PreviousPositionBuffer = nullptr;
 
@@ -421,7 +422,7 @@ public:
 	}
 
 	void SetupSection(int32 SectionIndex, FGPUSkinCache::FRWBuffersAllocation* InPositionAllocation, FSkelMeshRenderSection* Section, const FMorphVertexBuffer* MorphVertexBuffer, const FSkeletalMeshVertexClothBuffer* ClothVertexBuffer,
-		uint32 NumVertices, uint32 InputStreamStart, FGPUBaseSkinVertexFactory* InSourceVertexFactory, FGPUSkinPassthroughVertexFactory* InTargetVertexFactory)
+		uint32 NumVertices, uint32 InputStreamStart, FGPUBaseSkinVertexFactory* InSourceVertexFactory, FGPUSkinPassthroughVertexFactory* InTargetVertexFactory, uint32 InIntermediateAccumulatedTangentBufferOffset)
 	{
 		//UE_LOG(LogSkinCache, Warning, TEXT("*** SetupSection E %p Alloc %p Sec %d(%p) LOD %d"), this, InAllocation, SectionIndex, Section, LOD);
 		FSectionDispatchData& Data = DispatchData[SectionIndex];
@@ -487,6 +488,7 @@ public:
 				{
 					Data.NumTriangles = Section->NumTriangles;
 					Data.IndexBufferOffsetValue = Section->BaseIndex;
+					Data.IntermediateAccumulatedTangentBufferOffset = InIntermediateAccumulatedTangentBufferOffset;
 				}
 			}
 		}
@@ -975,7 +977,7 @@ public:
 
 		// UAV
 		SetUAVParameter(RHICmdList, ShaderRHI, IntermediateAccumBufferUAV, StagingBuffer.UAV);
-		SetShaderValue(RHICmdList, ShaderRHI, IntermediateAccumBufferOffset, GRecomputeTangentsParallelDispatch * DispatchData.IndexBufferOffsetValue);
+		SetShaderValue(RHICmdList, ShaderRHI, IntermediateAccumBufferOffset, GRecomputeTangentsParallelDispatch * DispatchData.IntermediateAccumulatedTangentBufferOffset);
 
         if (!GAllowDupedVertsForRecomputeTangents)
         {
@@ -1095,7 +1097,7 @@ public:
 
 		// UAVs
 		SetUAVParameter(RHICmdList, ShaderRHI, IntermediateAccumBufferUAV, StagingBuffer.UAV);
-		SetShaderValue(RHICmdList, ShaderRHI, IntermediateAccumBufferOffset, GRecomputeTangentsParallelDispatch * DispatchData.IndexBufferOffsetValue);
+		SetShaderValue(RHICmdList, ShaderRHI, IntermediateAccumBufferOffset, GRecomputeTangentsParallelDispatch * DispatchData.IntermediateAccumulatedTangentBufferOffset);
 		SetUAVParameter(RHICmdList, ShaderRHI, TangentBufferUAV, DispatchData.GetTangentRWBuffer()->Buffer.UAV);
 
 		SetSRVParameter(RHICmdList, ShaderRHI, TangentInputBuffer, DispatchData.IntermediateTangentBuffer ? DispatchData.IntermediateTangentBuffer->Buffer.SRV : nullptr);
@@ -1164,7 +1166,7 @@ void FGPUSkinCache::DispatchUpdateSkinTangents(FRHICommandListImmediate& RHICmdL
 			}
 
 			// no need to clear the staging buffer because we create it cleared and clear it after each usage in the per vertex pass
-			uint32 NumIntsPerBuffer = DispatchData.NumTriangles * 3 * FGPUSkinCache::IntermediateAccumBufferNumInts;
+			uint32 NumIntsPerBuffer = DispatchData.NumVertices * FGPUSkinCache::IntermediateAccumBufferNumInts;
 			CurrentStagingBufferIndex = (CurrentStagingBufferIndex + 1) % StagingBuffers.Num();
 			StagingBuffer = &StagingBuffers[CurrentStagingBufferIndex];
 			if (StagingBuffer->Buffer.NumBytes < NumIntsPerBuffer * sizeof(uint32))
@@ -1514,6 +1516,29 @@ bool FGPUSkinCache::ProcessEntry(
 		InvalidateAllEntries();
 	}
 
+	int32 RecomputeTangentsMode = GForceRecomputeTangents > 0 ? 1 : GSkinCacheRecomputeTangents;
+	bool bShouldRecomputeTangent = false;
+
+	// IntermediateAccumulatedTangents buffer is needed if mesh has at least one section needing recomputing tangents.
+	uint32 InterAccumTangentBufferSize = 0;
+	uint32 CurrInterAccumTangentBufferOffset = 0;
+	if (RecomputeTangentsMode > 0)
+	{
+		for (int32 i = 0; i < LodData.RenderSections.Num(); ++i)
+		{			
+			const FSkelMeshRenderSection& RenderSection = LodData.RenderSections[i];
+			if (RecomputeTangentsMode == 1 || RenderSection.bRecomputeTangent)
+			{
+				bShouldRecomputeTangent = true;
+				InterAccumTangentBufferSize += RenderSection.GetNumVertices();
+				if (i < Section)
+				{
+					CurrInterAccumTangentBufferOffset += RenderSection.GetNumVertices();
+				}
+			}
+		}
+	}
+
 	if (InOutEntry)
 	{
 		// If the LOD changed, the entry has to be invalidated
@@ -1527,35 +1552,18 @@ bool FGPUSkinCache::ProcessEntry(
 			if (!InOutEntry->IsSectionValid(Section) || !InOutEntry->IsSourceFactoryValid(Section, VertexFactory))
 			{
 				// This section might not be valid yet, so set it up
-				InOutEntry->SetupSection(Section, InOutEntry->PositionAllocation, &LodData.RenderSections[Section], MorphVertexBuffer, ClothVertexBuffer, NumVertices, InputStreamStart, VertexFactory, TargetVertexFactory);
+				InOutEntry->SetupSection(Section, InOutEntry->PositionAllocation, &LodData.RenderSections[Section], MorphVertexBuffer, ClothVertexBuffer, NumVertices, InputStreamStart, 
+											VertexFactory, TargetVertexFactory, CurrInterAccumTangentBufferOffset);
 			}
 		}
 	}
 
-	int32 RecomputeTangentsMode = GForceRecomputeTangents > 0 ? 1 : GSkinCacheRecomputeTangents;
 	// Try to allocate a new entry
 	if (!InOutEntry)
 	{
 		bool WithTangents = RecomputeTangentsMode > 0;
 		int32 TotalNumVertices = VertexFactory->GetNumVertices();
 		
-		bool bShouldRecomputeTangent = false;
-		if (RecomputeTangentsMode > 0)
-		{
-			bShouldRecomputeTangent = (RecomputeTangentsMode == 1);
-			if (!bShouldRecomputeTangent)
-			{
-				for (const FSkelMeshRenderSection& RenderSection : LodData.RenderSections)
-				{
-					if (RenderSection.bRecomputeTangent)
-					{
-						bShouldRecomputeTangent = true;
-						break;
-					}
-				}
-			}
-		}
-
 		// IntermediateTangents buffer is needed if mesh has at least one section using vertex color as recompute tangents blending mask
 		bool bEntryUseIntermediateTangents = false;
 		if (bShouldRecomputeTangent)
@@ -1570,17 +1578,13 @@ bool FGPUSkinCache::ProcessEntry(
 			}
 		}
 
-		// IntermediateAccumulatedTangents buffer is needed if mesh has at least one section needing recomputing tangents.
-		// TotalNumTriangles > 0 signals creation of IntermediateAccumulatedTangents buffer.
-		const uint32 TotalNumTriangles = (GRecomputeTangentsParallelDispatch && bShouldRecomputeTangent) ? LodData.GetTotalFaces() : 0;
-
-		FRWBuffersAllocation* NewPositionAllocation = TryAllocBuffer(TotalNumVertices, WithTangents, bEntryUseIntermediateTangents, TotalNumTriangles, RHICmdList);
+		FRWBuffersAllocation* NewPositionAllocation = TryAllocBuffer(TotalNumVertices, WithTangents, bEntryUseIntermediateTangents, InterAccumTangentBufferSize, RHICmdList);
 		if (!NewPositionAllocation)
 		{
 			if (GSkinCachePrintMemorySummary > 0)
 			{
 				const FString RayTracingTag = (Mode == EGPUSkinCacheEntryMode::RayTracing ? TEXT("[RT]") : TEXT(""));
-				uint64 RequiredMemInBytes = FRWBuffersAllocation::CalculateRequiredMemory(TotalNumVertices, WithTangents, bEntryUseIntermediateTangents, TotalNumTriangles);
+				uint64 RequiredMemInBytes = FRWBuffersAllocation::CalculateRequiredMemory(TotalNumVertices, WithTangents, bEntryUseIntermediateTangents, InterAccumTangentBufferSize);
 				UE_LOG(LogSkinCache, Warning, TEXT("FGPUSkinCache::ProcessEntry%s failed to allocate %.3fMB for mesh %s LOD%d, extra required memory increased to %.3fMB"),
 					*RayTracingTag, RequiredMemInBytes / MBSize, *GetSkeletalMeshObjectName(Skin), LODIndex, ExtraRequiredMemory / MBSize);
 			}
@@ -1592,7 +1596,8 @@ bool FGPUSkinCache::ProcessEntry(
 		InOutEntry = new FGPUSkinCacheEntry(this, Skin, NewPositionAllocation, LODIndex, Mode);
 		InOutEntry->GPUSkin = Skin;
 
-		InOutEntry->SetupSection(Section, NewPositionAllocation, &LodData.RenderSections[Section], MorphVertexBuffer, ClothVertexBuffer, NumVertices, InputStreamStart, VertexFactory, TargetVertexFactory);
+		InOutEntry->SetupSection(Section, NewPositionAllocation, &LodData.RenderSections[Section], MorphVertexBuffer, ClothVertexBuffer, NumVertices, InputStreamStart, 
+									VertexFactory, TargetVertexFactory, CurrInterAccumTangentBufferOffset);
 		Entries.Add(InOutEntry);
 	}
 
diff --git a/Engine/Source/Runtime/Engine/Public/GPUSkinCache.h b/Engine/Source/Runtime/Engine/Public/GPUSkinCache.h
--- a/Engine/Source/Runtime/Engine/Public/GPUSkinCache.h
+++ b/Engine/Source/Runtime/Engine/Public/GPUSkinCache.h
@@ -241,8 +241,8 @@ public:
 	{
 		friend struct FRWBufferTracker;
 
-		FRWBuffersAllocation(uint32 InNumVertices, bool InWithTangents, bool InUseIntermediateTangents, uint32 InNumTriangles, FRHICommandListImmediate& RHICmdList)
-			: NumVertices(InNumVertices), WithTangents(InWithTangents), UseIntermediateTangents(InUseIntermediateTangents), NumTriangles(InNumTriangles)
+		FRWBuffersAllocation(uint32 InNumVertices, bool InWithTangents, bool InUseIntermediateTangents, uint32 InIntermediateAccumulatedTangentsSize, FRHICommandListImmediate& RHICmdList)
+			: NumVertices(InNumVertices), WithTangents(InWithTangents), UseIntermediateTangents(InUseIntermediateTangents), IntermediateAccumulatedTangentsSize(InIntermediateAccumulatedTangentsSize)
 		{
 			for (int32 Index = 0; Index < NUM_BUFFERS; ++Index)
 			{
@@ -259,9 +259,9 @@ public:
 					IntermediateTangents.AccessState = ERHIAccess::Unknown;
 				}
 			}
-			if (NumTriangles > 0)
+			if (IntermediateAccumulatedTangentsSize > 0)
 			{
-				IntermediateAccumulatedTangents.Buffer.Initialize(TEXT("SkinCacheIntermediateAccumulatedTangents"), sizeof(int32), NumTriangles * 3 * FGPUSkinCache::IntermediateAccumBufferNumInts, PF_R32_SINT, BUF_UnorderedAccess);
+				IntermediateAccumulatedTangents.Buffer.Initialize(TEXT("SkinCacheIntermediateAccumulatedTangents"), sizeof(int32), IntermediateAccumulatedTangentsSize * FGPUSkinCache::IntermediateAccumBufferNumInts, PF_R32_SINT, BUF_UnorderedAccess);
 				IntermediateAccumulatedTangents.AccessState = ERHIAccess::Unknown;
 				// The UAV must be zero-filled. We leave it zeroed after each round (see RecomputeTangentsPerVertexPass.usf), so this is only needed on when the buffer is first created.
 				RHICmdList.ClearUAVUint(IntermediateAccumulatedTangents.Buffer.UAV, FUintVector4(0, 0, 0, 0));
@@ -279,13 +279,13 @@ public:
 				Tangents.Release();
 				IntermediateTangents.Release();
 			}
-			if (NumTriangles > 0)
+			if (IntermediateAccumulatedTangentsSize > 0)
 			{
 				IntermediateAccumulatedTangents.Release();
 			}
 		}
 
-		static uint64 CalculateRequiredMemory(uint32 InNumVertices, bool InWithTangents, bool InUseIntermediateTangents, uint32 InNumTriangles)
+		static uint64 CalculateRequiredMemory(uint32 InNumVertices, bool InWithTangents, bool InUseIntermediateTangents, uint32 InIntermediateAccumulatedTangentsSize)
 		{
 			uint64 PositionBufferSize = PosBufferBytesPerElement * InNumVertices * 3 * NUM_BUFFERS;
 			uint64 TangentBufferSize = InWithTangents ? TangentBufferBytesPerElement * InNumVertices * 2 : 0;
@@ -294,13 +294,13 @@ public:
 			{
 				IntermediateTangentBufferSize = InWithTangents ? TangentBufferBytesPerElement * InNumVertices * 2 : 0;
 			}
-			uint64 AccumulatedTangentBufferSize = InNumTriangles * 3 * FGPUSkinCache::IntermediateAccumBufferNumInts * sizeof(int32);
+			uint64 AccumulatedTangentBufferSize = InIntermediateAccumulatedTangentsSize * FGPUSkinCache::IntermediateAccumBufferNumInts * sizeof(int32);
 			return TangentBufferSize + IntermediateTangentBufferSize + PositionBufferSize + AccumulatedTangentBufferSize;
 		}
 
 		uint64 GetNumBytes() const
 		{
-			return CalculateRequiredMemory(NumVertices, WithTangents, UseIntermediateTangents, NumTriangles);
+			return CalculateRequiredMemory(NumVertices, WithTangents, UseIntermediateTangents, IntermediateAccumulatedTangentsSize);
 		}
 
 		FSkinCacheRWBuffer* GetTangentBuffer()
@@ -315,7 +315,7 @@ public:
 
 		FSkinCacheRWBuffer* GetIntermediateAccumulatedTangentBuffer()
 		{
-			return NumTriangles > 0 ? &IntermediateAccumulatedTangents : nullptr;
+			return IntermediateAccumulatedTangentsSize > 0 ? &IntermediateAccumulatedTangents : nullptr;
 		}
 
 		void RemoveAllFromTransitionArray(TSet<FSkinCacheRWBuffer*>& BuffersToTransition);
@@ -331,7 +331,7 @@ public:
 		const uint32 NumVertices;
 		const bool WithTangents;
 		const bool UseIntermediateTangents;
-		const uint32 NumTriangles;
+		const uint32 IntermediateAccumulatedTangentsSize;
 
 		static const uint32 PosBufferBytesPerElement = 4;
 		static const uint32 TangentBufferBytesPerElement = 8;
