PR #7950: Eliminate an O(N^2) loop to update bone parents.

diff --git a/Engine/Source/Runtime/Engine/Private/Animation/AnimationUtils.cpp b/Engine/Source/Runtime/Engine/Private/Animation/AnimationUtils.cpp
--- a/Engine/Source/Runtime/Engine/Private/Animation/AnimationUtils.cpp
+++ b/Engine/Source/Runtime/Engine/Private/Animation/AnimationUtils.cpp
@@ -83,19 +83,15 @@ void FAnimationUtils::BuildSkeletonMetaData(USkeleton* Skeleton, TArray<FBoneDat
 	}
 
 	// Enumerate children (bones that refer to this bone as parent).
-	for ( int32 BoneIndex = 0 ; BoneIndex < OutBoneData.Num() ; ++BoneIndex )
+	for(int32 BoneIndex = 1; BoneIndex < OutBoneData.Num(); ++BoneIndex)
 	{
-		FBoneData& BoneData = OutBoneData[BoneIndex];
-		// Exclude the root bone as it is the child of nothing.
-		for ( int32 BoneIndex2 = 1 ; BoneIndex2 < OutBoneData.Num() ; ++BoneIndex2 )
+		const int32 ParentIndex = OutBoneData[BoneIndex].GetParent();
+		if (OutBoneData.IsValidIndex(ParentIndex))
 		{
-			if ( OutBoneData[BoneIndex2].GetParent() == BoneIndex )
-			{
-				BoneData.Children.Add(BoneIndex2);
-			}
+			OutBoneData[ParentIndex].Children.Add(BoneIndex);
 		}
 	}
-
+	
 	// Enumerate end effectors.  For each end effector, propagate its index up to all ancestors.
 	if( bEnableLogging )
 	{
