Fix non-unity compile error in CL 17296224

diff --git a/Engine/Source/Runtime/Engine/Public/StereoRendering.h b/Engine/Source/Runtime/Engine/Public/StereoRendering.h
--- a/Engine/Source/Runtime/Engine/Public/StereoRendering.h
+++ b/Engine/Source/Runtime/Engine/Public/StereoRendering.h
@@ -195,7 +195,7 @@ public:
 	/**
 	* Provides the final view rect that the renderer will render into.
 	*/
-	virtual void SetFinalViewRect(FRHICommandListImmediate& RHICmdList, const enum EStereoscopicPass StereoPass, const FIntRect& FinalViewRect) {}
+	virtual void SetFinalViewRect(class FRHICommandListImmediate& RHICmdList, const enum EStereoscopicPass StereoPass, const FIntRect& FinalViewRect) {}
 
 	/**
 	 * Gets the percentage bounds of the safe region to draw in.  This allows things like stat rendering to appear within the readable portion of the stereo view.
