#7871 : PR : Support NetSerialization of FInterfaceProperty

diff --git a/Engine/Source/Runtime/Core/Public/Misc/NetworkVersion.h b/Engine/Source/Runtime/Core/Public/Misc/NetworkVersion.h
--- a/Engine/Source/Runtime/Core/Public/Misc/NetworkVersion.h
+++ b/Engine/Source/Runtime/Core/Public/Misc/NetworkVersion.h
@@ -51,6 +51,7 @@ enum EEngineNetworkVersionHistory
 	HISTORY_ENUM_SERIALIZATION_COMPAT = 17,			// Bump version to include enum bits required for serialization into compat checksums, as well as unify enum and byte property enum serialization
 	HISTORY_SUBOBJECT_OUTER_CHAIN = 18,				// Bump version to support subobject outer chains matching on client and server
 	HISTORY_HITRESULT_INSTANCEHANDLE = 19,			// Bump version to support FHitResult change of Actor to HitObjectHandle. This change was made in CL 14369221 but a net version wasn't added at the time.
+	HISTORY_INTERFACE_PROPERTY_SERIALIZATION = 20,	// Bump version to support net serialization of FInterfaceProperty
 	// New history items go above here.
 
 	HISTORY_ENGINENETVERSION_PLUS_ONE,
diff --git a/Engine/Source/Runtime/CoreUObject/Private/UObject/PropertyInterface.cpp b/Engine/Source/Runtime/CoreUObject/Private/UObject/PropertyInterface.cpp
--- a/Engine/Source/Runtime/CoreUObject/Private/UObject/PropertyInterface.cpp
+++ b/Engine/Source/Runtime/CoreUObject/Private/UObject/PropertyInterface.cpp
@@ -10,6 +10,8 @@
 #include "UObject/UnrealTypePrivate.h"
 #include "UObject/LinkerPlaceholderClass.h"
 #include "Hash/Blake3.h"
+#include "UObject/CoreNet.h"
+#include "Misc/NetworkVersion.h"
 
 // WARNING: This should always be the last include in any file that needs it (except .generated.h)
 #include "UObject/UndefineUPropertyMacros.h"
@@ -153,7 +155,23 @@ void FInterfaceProperty::SerializeItem( FStructuredArchive::FSlot Slot, void* Va
 
 bool FInterfaceProperty::NetSerializeItem( FArchive& Ar, UPackageMap* Map, void* Data, TArray<uint8> * MetaData ) const
 {
-	//@todo
+	if (Ar.EngineNetVer() >= HISTORY_INTERFACE_PROPERTY_SERIALIZATION)
+	{ 
+		FScriptInterface* InterfaceValue = (FScriptInterface*)Data;
+		bool Result = Map->SerializeObject(Ar, InterfaceClass, InterfaceValue->GetObjectRef());
+		if (Ar.IsLoading())
+		{
+			if (InterfaceValue->GetObject() != nullptr)
+			{
+				InterfaceValue->SetInterface(InterfaceValue->GetObject()->GetInterfaceAddress(InterfaceClass));
+			}
+			else
+			{
+				InterfaceValue->SetInterface(nullptr);
+			}
+		}
+		return Result;
+	}
 	return false;
 }
 
diff --git a/Engine/Source/Runtime/CoreUObject/Public/UObject/UnrealType.h b/Engine/Source/Runtime/CoreUObject/Public/UObject/UnrealType.h
--- a/Engine/Source/Runtime/CoreUObject/Public/UObject/UnrealType.h
+++ b/Engine/Source/Runtime/CoreUObject/Public/UObject/UnrealType.h
@@ -2812,6 +2812,7 @@ public:
 	virtual bool Identical( const void* A, const void* B, uint32 PortFlags ) const override;
 	virtual void SerializeItem( FStructuredArchive::FSlot Slot, void* Value, void const* Defaults) const override;
 	virtual bool NetSerializeItem( FArchive& Ar, UPackageMap* Map, void* Data, TArray<uint8> * MetaData = NULL ) const override;
+	virtual bool SupportsNetSharedSerialization() const override { return false; }
 	virtual void ExportTextItem( FString& ValueStr, const void* PropertyValue, const void* DefaultValue, UObject* Parent, int32 PortFlags, UObject* ExportRootScope ) const override;
 	virtual const TCHAR* ImportText_Internal( const TCHAR* Buffer, void* Data, int32 PortFlags, UObject* OwnerObject, FOutputDevice* ErrorText ) const override;
 	virtual bool ContainsObjectReference(TArray<const FStructProperty*>& EncounteredStructProps, EPropertyObjectReferenceType InReferenceType = EPropertyObjectReferenceType::Strong) const override;
diff --git a/Engine/Source/Runtime/Engine/Private/PackageMapClient.cpp b/Engine/Source/Runtime/Engine/Private/PackageMapClient.cpp
--- a/Engine/Source/Runtime/Engine/Private/PackageMapClient.cpp
+++ b/Engine/Source/Runtime/Engine/Private/PackageMapClient.cpp
@@ -296,7 +296,7 @@ bool UPackageMapClient::SerializeObject( FArchive& Ar, UClass* Class, UObject*&
 			}
 
 			// Check that we got the right class
-			if (Object && !Object->IsA(Class))
+			if (Object && !(Class->HasAnyClassFlags(CLASS_Interface) ? Object->GetClass()->ImplementsInterface(Class) : Object->IsA(Class)))
 			{
 				UE_LOG(LogNetPackageMap, Warning, TEXT("Forged object: got %s, expecting %s"), *Object->GetFullName(), *Class->GetFullName());
 				Object = NULL;
diff --git a/Engine/Source/Runtime/Engine/Private/RepLayout.cpp b/Engine/Source/Runtime/Engine/Private/RepLayout.cpp
--- a/Engine/Source/Runtime/Engine/Private/RepLayout.cpp
+++ b/Engine/Source/Runtime/Engine/Private/RepLayout.cpp
@@ -589,6 +589,14 @@ static FORCEINLINE bool CompareWeakObject(
 	return ObjectA.HasSameIndexAndSerialNumber(ObjectB);
 }
 
+static FORCEINLINE bool CompareInterface(
+	const FRepLayoutCmd& Cmd,
+	const void* A,
+	const void* B)
+{
+	return Cmd.Property->Identical(A, B);
+}
+
 static FORCEINLINE bool CompareNetSerializeStructWithObjectProperties(
 	const TArray<FRepLayoutCmd>& Cmds,
 	const TMap<FRepLayoutCmd*, TArray<FRepLayoutCmd>>& NetSerializeLayouts,
@@ -644,6 +652,9 @@ static FORCEINLINE bool PropertiesAreIdenticalNative(
 		case ERepLayoutCmdType::PropertyWeakObject:
 			return CompareWeakObject(Cmd, A, B);
 
+		case ERepLayoutCmdType::PropertyInterface:
+			return CompareInterface(Cmd, A, B);
+
 		case ERepLayoutCmdType::PropertyUInt32:
 			return CompareValue<uint32>(A, B);
 
@@ -5062,7 +5073,7 @@ static bool DiffStableProperties_r(FDiffStablePropertiesSharedParams& Params, TD
 					Cmd.Type == ERepLayoutCmdType::PropertyWeakObject ||
 					Cmd.Type == ERepLayoutCmdType::PropertySoftObject)
 				{
-				if (FObjectPropertyBase* ObjProperty = CastFieldChecked<FObjectPropertyBase>(Cmd.Property))
+					if (FObjectPropertyBase* ObjProperty = CastFieldChecked<FObjectPropertyBase>(Cmd.Property))
 					{
 						if (ObjProperty->PropertyClass && (ObjProperty->PropertyClass->IsChildOf(AActor::StaticClass()) || ObjProperty->PropertyClass->IsChildOf(UActorComponent::StaticClass())))
 						{
@@ -5086,6 +5097,32 @@ static bool DiffStableProperties_r(FDiffStablePropertiesSharedParams& Params, TD
 						}
 					}
 				}
+				else if (Cmd.Type == ERepLayoutCmdType::PropertyInterface)
+				{
+					if (FInterfaceProperty* InterfaceProperty = CastFieldChecked<FInterfaceProperty>(Cmd.Property))
+					{
+						if (UObject* InterfaceObjValue = InterfaceProperty->GetPropertyValue(StackParams.Source + Cmd).GetObject())
+						{
+							if (InterfaceObjValue->GetClass() && (InterfaceObjValue->GetClass()->IsChildOf(AActor::StaticClass()) || InterfaceObjValue->GetClass()->IsChildOf(UActorComponent::StaticClass())))
+							{
+								// skip actor and component references
+								continue;
+							}
+
+							const bool bStableForNetworking = (InterfaceObjValue->HasAnyFlags(RF_WasLoaded | RF_DefaultSubObject) || InterfaceObjValue->IsNative() || InterfaceObjValue->IsDefaultSubobject());
+							if (!bStableForNetworking)
+							{
+								// skip object references without a stable name
+								continue;
+							}
+
+							if (Params.ObjReferences)
+							{
+								Params.ObjReferences->AddUnique(InterfaceObjValue);
+							}
+						}
+					}
+				}
 
 				StoreProperty(Cmd, StackParams.Destination + Cmd, StackParams.Source + Cmd);
 
@@ -5268,6 +5305,10 @@ static uint32 AddPropertyCmd(
 			Cmd.Type = ERepLayoutCmdType::PropertyObject;
 		}
 	}
+	else if (UnderlyingProperty->IsA(FInterfaceProperty::StaticClass()))
+	{
+		Cmd.Type = ERepLayoutCmdType::PropertyInterface;
+	}
 	else if (UnderlyingProperty->IsA(FNameProperty::StaticClass()))
 	{
 		Cmd.Type = ERepLayoutCmdType::PropertyName;
diff --git a/Engine/Source/Runtime/Engine/Public/Net/RepLayout.h b/Engine/Source/Runtime/Engine/Public/Net/RepLayout.h
--- a/Engine/Source/Runtime/Engine/Public/Net/RepLayout.h
+++ b/Engine/Source/Runtime/Engine/Public/Net/RepLayout.h
@@ -771,7 +771,8 @@ enum class ERepLayoutCmdType : uint8
 	PropertyNativeBool		= 21,
 	PropertySoftObject		= 22,
 	PropertyWeakObject		= 23,
-	NetSerializeStructWithObjectReferences = 24,
+	PropertyInterface		= 24,
+	NetSerializeStructWithObjectReferences = 25,
 };
 
 /** Various flags that describe how a Top Level Property should be handled. */
