Added missing include to fix NonUnity build #review-17312781 trivial #fyi [REDACTED_HANDLE]

diff --git a/Engine/Source/Runtime/Engine/Private/Animation/AnimNotifyState.cpp b/Engine/Source/Runtime/Engine/Private/Animation/AnimNotifyState.cpp
--- a/Engine/Source/Runtime/Engine/Private/Animation/AnimNotifyState.cpp
+++ b/Engine/Source/Runtime/Engine/Private/Animation/AnimNotifyState.cpp
@@ -4,6 +4,7 @@
 #include "Animation/AnimTypes.h"
 #include "Animation/AnimNotifies/AnimNotify.h"
 #include "UObject/ObjectSaveContext.h"
+#include "Animation/AnimSequenceBase.h"
 
 /////////////////////////////////////////////////////
 // UAnimNotifyState
