Changed world partition HLODs source cells mapping from FSotObjectPath to FName.
This reduce the time spent serializing HLOD cells in PIE.

diff --git a/Engine/Source/Developer/WorldPartitionHLODUtilities/Private/WorldPartitionHLODUtilities.cpp b/Engine/Source/Developer/WorldPartitionHLODUtilities/Private/WorldPartitionHLODUtilities.cpp
--- a/Engine/Source/Developer/WorldPartitionHLODUtilities/Private/WorldPartitionHLODUtilities.cpp
+++ b/Engine/Source/Developer/WorldPartitionHLODUtilities/Private/WorldPartitionHLODUtilities.cpp
@@ -213,13 +213,9 @@ TArray<AWorldPartitionHLOD*> FWorldPartitionHLODUtilities::CreateHLODActors(FHLO
 			SpawnParams.NameMode = FActorSpawnParameters::ESpawnActorNameMode::Required_Fatal;
 			HLODActor = InCreationParams.WorldPartition->GetWorld()->SpawnActor<AWorldPartitionHLOD>(SpawnParams);
 
-			FSoftObjectPath WorldPartitionPath(InCreationParams.WorldPartition);
-			FString CellObjectPath = FString::Printf(TEXT("%s.%s"), *WorldPartitionPath.ToString(), *InCreationParams.CellName.ToString());
-			TSoftObjectPtr<UWorldPartitionRuntimeCell> RuntimeCell(CellObjectPath);
-
 			HLODActor->SetActorLabel(FString::Printf(TEXT("HLOD%d_%s"), InCreationParams.HLODLevel, *InCreationParams.CellName.ToString()));
 			HLODActor->SetFolderPath(*FString::Printf(TEXT("HLOD/HLOD%d"), InCreationParams.HLODLevel));
-			HLODActor->SetSourceCell(RuntimeCell);
+			HLODActor->SetSourceCellName(InCreationParams.CellName);
 			HLODActor->SetSubActorsHLODLayer(HLODLayer);
 			HLODActor->SetGridIndices(InCreationParams.GridIndexX, InCreationParams.GridIndexY, InCreationParams.GridIndexZ);
 
diff --git a/Engine/Source/Runtime/Core/Public/UObject/UE5PrivateFrostyStreamObjectVersion.h b/Engine/Source/Runtime/Core/Public/UObject/UE5PrivateFrostyStreamObjectVersion.h
--- a/Engine/Source/Runtime/Core/Public/UObject/UE5PrivateFrostyStreamObjectVersion.h
+++ b/Engine/Source/Runtime/Core/Public/UObject/UE5PrivateFrostyStreamObjectVersion.h
@@ -24,6 +24,9 @@ struct CORE_API FUE5PrivateFrostyStreamObjectVersion
 		// Always serialize the actor label in cooked builds
 		SerializeActorLabelInCookedBuilds,
 
+		// Changed world partition HLODs cells from FSotObjectPath to FName
+		ConvertWorldPartitionHLODsCellsToName,
+
 		// -----<new versions can be added above this line>-------------------------------------------------
 		VersionPlusOne,
 		LatestVersion = VersionPlusOne - 1
diff --git a/Engine/Source/Runtime/Engine/Private/WorldPartition/HLOD/HLODActor.cpp b/Engine/Source/Runtime/Engine/Private/WorldPartition/HLOD/HLODActor.cpp
--- a/Engine/Source/Runtime/Engine/Private/WorldPartition/HLOD/HLODActor.cpp
+++ b/Engine/Source/Runtime/Engine/Private/WorldPartition/HLOD/HLODActor.cpp
@@ -4,6 +4,7 @@
 #include "WorldPartition/HLOD/HLODSubsystem.h"
 #include "Components/PrimitiveComponent.h"
 #include "UObject/UE5MainStreamObjectVersion.h"
+#include "UObject/UE5PrivateFrostyStreamObjectVersion.h"
 #include "UObject/ObjectSaveContext.h"
 
 #if WITH_EDITOR
@@ -58,12 +59,26 @@ void AWorldPartitionHLOD::EndPlay(const EEndPlayReason::Type EndPlayReason)
 void AWorldPartitionHLOD::Serialize(FArchive& Ar)
 {
 	Ar.UsingCustomVersion(FUE5MainStreamObjectVersion::GUID);
+	Ar.UsingCustomVersion(FUE5PrivateFrostyStreamObjectVersion::GUID);
+
 	Super::Serialize(Ar);
 
 #if WITH_EDITOR
 	if(Ar.IsLoading() && Ar.CustomVer(FUE5MainStreamObjectVersion::GUID) < FUE5MainStreamObjectVersion::WorldPartitionStreamingCellsNamingShortened)
 	{
-		SourceCell = SourceCell.ToString().Replace(TEXT("WPRT_"), TEXT(""), ESearchCase::CaseSensitive).Replace(TEXT("Cell_"), TEXT(""), ESearchCase::CaseSensitive);
+		SourceCell_DEPRECATED = SourceCell_DEPRECATED.ToString().Replace(TEXT("WPRT_"), TEXT(""), ESearchCase::CaseSensitive).Replace(TEXT("Cell_"), TEXT(""), ESearchCase::CaseSensitive);
+	}
+
+	if(Ar.IsLoading() && Ar.CustomVer(FUE5PrivateFrostyStreamObjectVersion::GUID) < FUE5PrivateFrostyStreamObjectVersion::ConvertWorldPartitionHLODsCellsToName)
+	{
+		FString CellName;
+		FString CellContext;
+		const FString CellPath = FPackageName::GetShortName(SourceCell_DEPRECATED.ToSoftObjectPath().GetSubPathString());
+		if (!CellPath.Split(TEXT("."), &CellContext, &CellName, ESearchCase::CaseSensitive, ESearchDir::FromEnd))
+		{
+			CellName = CellPath;
+		}
+		SourceCellName = *CellName;
 	}
 #endif
 }
@@ -138,9 +153,9 @@ const TArray<FGuid>& AWorldPartitionHLOD::GetSubActors() const
 	return SubActors;
 }
 
-void AWorldPartitionHLOD::SetSourceCell(const TSoftObjectPtr<UWorldPartitionRuntimeCell>& InSourceCell)
+void AWorldPartitionHLOD::SetSourceCellName(FName InSourceCellName)
 {
-	SourceCell = InSourceCell;
+	SourceCellName = InSourceCellName;
 }
 
 const FBox& AWorldPartitionHLOD::GetHLODBounds() const
diff --git a/Engine/Source/Runtime/Engine/Private/WorldPartition/HLOD/HLODSubsystem.cpp b/Engine/Source/Runtime/Engine/Private/WorldPartition/HLOD/HLODSubsystem.cpp
--- a/Engine/Source/Runtime/Engine/Private/WorldPartition/HLOD/HLODSubsystem.cpp
+++ b/Engine/Source/Runtime/Engine/Private/WorldPartition/HLOD/HLODSubsystem.cpp
@@ -113,19 +113,11 @@ void UHLODSubsystem::RegisterHLODActor(AWorldPartitionHLOD* InWorldPartitionHLOD
 {
 	TRACE_CPUPROFILER_EVENT_SCOPE(UHLODSubsystem::RegisterHLODActor);
 
-	const TSoftObjectPtr<UWorldPartitionRuntimeCell>& RuntimeCell = InWorldPartitionHLOD->GetSourceCell();
-	const FString CellPath = FPackageName::GetShortName(RuntimeCell.ToSoftObjectPath().GetSubPathString());
-	FString CellName;
-	FString CellContext;
-	if (!CellPath.Split(TEXT("."), &CellContext, &CellName, ESearchCase::CaseSensitive, ESearchDir::FromEnd))
-	{
-		CellName = CellPath;
-	}
-
-	FCellHLODMapping* CellHLODs = CellsHLODMapping.Find(FName(CellName));
+	const FName CellName = InWorldPartitionHLOD->GetSourceCellName();
+	FCellHLODMapping* CellHLODs = CellsHLODMapping.Find(CellName);
 
 #if WITH_EDITOR
-	UE_LOG(LogHLODSubsystem, Verbose, TEXT("Registering HLOD %s (%s) for cell %s"), *InWorldPartitionHLOD->GetActorLabel(), *InWorldPartitionHLOD->GetActorGuid().ToString(), *RuntimeCell.ToString());
+	UE_LOG(LogHLODSubsystem, Verbose, TEXT("Registering HLOD %s (%s) for cell %s"), *InWorldPartitionHLOD->GetActorLabel(), *InWorldPartitionHLOD->GetActorGuid().ToString(), *CellName.ToString());
 #endif
 
 	if (CellHLODs)
@@ -135,7 +127,7 @@ void UHLODSubsystem::RegisterHLODActor(AWorldPartitionHLOD* InWorldPartitionHLOD
 	}
 	else
 	{
-		UE_LOG(LogHLODSubsystem, Verbose, TEXT("Found HLOD referencing nonexistent cell '%s'"), *RuntimeCell.ToString());
+		UE_LOG(LogHLODSubsystem, Verbose, TEXT("Found HLOD referencing nonexistent cell '%s'"), *CellName.ToString());
 		InWorldPartitionHLOD->SetVisibility(false);
 	}
 }
@@ -144,19 +136,11 @@ void UHLODSubsystem::UnregisterHLODActor(AWorldPartitionHLOD* InWorldPartitionHL
 {
 	TRACE_CPUPROFILER_EVENT_SCOPE(UHLODSubsystem::UnregisterHLODActor);
 
-	const TSoftObjectPtr<UWorldPartitionRuntimeCell>& RuntimeCell = InWorldPartitionHLOD->GetSourceCell();
-	const FString CellPath = FPackageName::GetShortName(RuntimeCell.ToSoftObjectPath().GetSubPathString());
-	FString CellName;
-	FString CellContext;
-	if (!CellPath.Split(TEXT("."), &CellContext, &CellName, ESearchCase::CaseSensitive, ESearchDir::FromEnd))
-	{
-		CellName = CellPath;
-	}
-
-	FCellHLODMapping* CellHLODs = CellsHLODMapping.Find(FName(CellName));
+	const FName CellName = InWorldPartitionHLOD->GetSourceCellName();
+	FCellHLODMapping* CellHLODs = CellsHLODMapping.Find(CellName);
 
 #if WITH_EDITOR
-	UE_LOG(LogHLODSubsystem, Verbose, TEXT("Unregistering HLOD %s (%s) for cell %s"), *InWorldPartitionHLOD->GetActorLabel(), *InWorldPartitionHLOD->GetActorGuid().ToString(), *RuntimeCell.ToString());
+	UE_LOG(LogHLODSubsystem, Verbose, TEXT("Unregistering HLOD %s (%s) for cell %s"), *InWorldPartitionHLOD->GetActorLabel(), *InWorldPartitionHLOD->GetActorGuid().ToString(), *CellName.ToString());
 #endif
 
 	if (CellHLODs)
diff --git a/Engine/Source/Runtime/Engine/Public/WorldPartition/HLOD/HLODActor.h b/Engine/Source/Runtime/Engine/Public/WorldPartition/HLOD/HLODActor.h
--- a/Engine/Source/Runtime/Engine/Public/WorldPartition/HLOD/HLODActor.h
+++ b/Engine/Source/Runtime/Engine/Public/WorldPartition/HLOD/HLODActor.h
@@ -19,7 +19,7 @@ class ENGINE_API AWorldPartitionHLOD : public AActor
 public:
 	void SetVisibility(bool bInVisible);
 
-	inline const TSoftObjectPtr<UWorldPartitionRuntimeCell>& GetSourceCell() const { return SourceCell; }
+	inline FName GetSourceCellName() const { return SourceCellName; }
 	inline uint32 GetLODLevel() const { return LODLevel; }
 
 	virtual bool IsHLODRelevant() const override { return true; }
@@ -47,7 +47,7 @@ public:
 		OutGridIndexZ = GridIndexZ;
 	}
 
-	void SetSourceCell(const TSoftObjectPtr<UWorldPartitionRuntimeCell>& InSourceCell);
+	void SetSourceCellName(FName InSourceCellName);
 	inline void SetLODLevel(uint32 InLODLevel) { LODLevel = InLODLevel; }
 
 	const FBox& GetHLODBounds() const;
@@ -108,7 +108,10 @@ private:
 	uint32 LODLevel;
 
 	UPROPERTY()
-	TSoftObjectPtr<UWorldPartitionRuntimeCell> SourceCell;
+	TSoftObjectPtr<UWorldPartitionRuntimeCell> SourceCell_DEPRECATED;
+
+	UPROPERTY()
+	FName SourceCellName;
 };
 
 DEFINE_ACTORDESC_TYPE(AWorldPartitionHLOD, FHLODActorDesc);
