World Partition: filter out invalid instanced static mesh components when computing streaming bounds.

diff --git a/Engine/Source/Runtime/Engine/Classes/Components/InstancedStaticMeshComponent.h b/Engine/Source/Runtime/Engine/Classes/Components/InstancedStaticMeshComponent.h
--- a/Engine/Source/Runtime/Engine/Classes/Components/InstancedStaticMeshComponent.h
+++ b/Engine/Source/Runtime/Engine/Classes/Components/InstancedStaticMeshComponent.h
@@ -350,6 +350,7 @@ public:
 	virtual bool SupportsStaticLighting() const override { return true; }
 #if WITH_EDITOR
 	virtual void GetStaticLightingInfo(FStaticLightingPrimitiveInfo& OutPrimitiveInfo,const TArray<ULightComponent*>& InRelevantLights,const FLightingBuildOptions& Options) override;
+	virtual FBox GetStreamingBounds() const override;
 #endif
 	virtual void GetLightAndShadowMapMemoryUsage( int32& LightMapMemoryUsage, int32& ShadowMapMemoryUsage ) const override;
 
diff --git a/Engine/Source/Runtime/Engine/Private/Actor.cpp b/Engine/Source/Runtime/Engine/Private/Actor.cpp
--- a/Engine/Source/Runtime/Engine/Private/Actor.cpp
+++ b/Engine/Source/Runtime/Engine/Private/Actor.cpp
@@ -984,44 +984,44 @@ FBox AActor::GetStreamingBounds() const
 {
 	FBox StreamingBounds(ForceInit);
 
-	auto IsComponentRelevant = [this](const UActorComponent* Component)
+	auto HandleComponent = [this, &StreamingBounds](const UActorComponent* Component)
 	{
 		check(Component);
 
 		if (!Component->IsRegistered())
 		{
-			return false;
+			return;
 		}
 
 		if (Component->HasAnyFlags(RF_Transient))
 		{
-			return false;
+			return;
 		}
 
 		if (!IsEditorOnly() && Component->IsEditorOnly())
 		{
-			return false;
+			return;
 		}
 
-		return true;
+		const FBox ComponentStreamingBound = Component->GetStreamingBounds();
+		if (!ComponentStreamingBound.IsValid)
+		{
+			return;
+		}
+
+		StreamingBounds += ComponentStreamingBound;
 	};
 
-	ForEachComponent<UPrimitiveComponent>(true, [this, &StreamingBounds, &IsComponentRelevant](UActorComponent* Component)
+	ForEachComponent<UPrimitiveComponent>(true, [this, &HandleComponent](UActorComponent* Component)
 	{
-		if (IsComponentRelevant(Component))
-		{
-			StreamingBounds += Component->GetStreamingBounds();
-		}
+		HandleComponent(Component);
 	});
 
 	if (!StreamingBounds.IsValid)
 	{
-		ForEachComponent<UActorComponent>(false, [this, &StreamingBounds, &IsComponentRelevant](UActorComponent* Component)
+		ForEachComponent<UActorComponent>(false, [this, &HandleComponent](UActorComponent* Component)
 		{
-			if (IsComponentRelevant(Component))
-			{
-				StreamingBounds += Component->GetStreamingBounds();
-			}
+			HandleComponent(Component);
 		});
 	}
 
diff --git a/Engine/Source/Runtime/Engine/Private/InstancedStaticMesh.cpp b/Engine/Source/Runtime/Engine/Private/InstancedStaticMesh.cpp
--- a/Engine/Source/Runtime/Engine/Private/InstancedStaticMesh.cpp
+++ b/Engine/Source/Runtime/Engine/Private/InstancedStaticMesh.cpp
@@ -2555,6 +2555,11 @@ void UInstancedStaticMeshComponent::ApplyLightMapping(FStaticLightingTextureMapp
 		MarkRenderStateDirty();
 	}
 }
+
+FBox UInstancedStaticMeshComponent::GetStreamingBounds() const
+{
+	return (GetStaticMesh() && PerInstanceSMData.Num()) ? Super::GetStreamingBounds() : FBox(ForceInit);
+}
 #endif
 
 void UInstancedStaticMeshComponent::ReleasePerInstanceRenderData()
