Update physics collision toggling for capsule on back of vehicle. Don't regenerate physics bodies when changing shape sim - it overwrites the collision override settings.

diff --git a/Engine/Source/Runtime/Engine/Private/PhysicsEngine/BodyInstance.cpp b/Engine/Source/Runtime/Engine/Private/PhysicsEngine/BodyInstance.cpp
--- a/Engine/Source/Runtime/Engine/Private/PhysicsEngine/BodyInstance.cpp
+++ b/Engine/Source/Runtime/Engine/Private/PhysicsEngine/BodyInstance.cpp
@@ -687,14 +687,6 @@ void FBodyInstance::SetShapeCollisionEnabled(const int32 ShapeIndex, ECollisionE
 			{
 				UpdatePhysicsFilterData();
 			}
-
-			if (CollisionEnabledHasPhysics(OldType) != CollisionEnabledHasPhysics(NewType))
-			{
-				if(UPrimitiveComponent* PrimComponent = OwnerComponent.Get())
-				{
-					PrimComponent->RecreatePhysicsState();
-				}
-			}
 		}
 	}
 }
