[NO NON-PII COMMIT MESSAGE]

diff --git a/Engine/Source/Runtime/AudioPlatformConfiguration/Public/AudioCompressionSettings.h b/Engine/Source/Runtime/AudioPlatformConfiguration/Public/AudioCompressionSettings.h
--- a/Engine/Source/Runtime/AudioPlatformConfiguration/Public/AudioCompressionSettings.h
+++ b/Engine/Source/Runtime/AudioPlatformConfiguration/Public/AudioCompressionSettings.h
@@ -53,7 +53,7 @@ struct FPlatformAudioCookOverrides
 	// Increment this to force a recook on all Stream Caching assets.
 	// For testing, it's useful to set this to either a negative number or
 	// absurdly large number, to ensure you do not pollute the DDC.
-	static const int32 StreamCachingVersion = 5021;
+	static const int32 StreamCachingVersion = 5022;
 
 	bool bResampleForDevice;
 
diff --git a/Engine/Source/Runtime/Engine/Private/SoundWave.cpp b/Engine/Source/Runtime/Engine/Private/SoundWave.cpp
--- a/Engine/Source/Runtime/Engine/Private/SoundWave.cpp
+++ b/Engine/Source/Runtime/Engine/Private/SoundWave.cpp
@@ -387,7 +387,7 @@ void USoundWave::Serialize( FArchive& Ar )
 				{
 					// for now we only support one format per wav
 					FName Format = CookingTarget->GetWaveFormat(this);
-					const FPlatformAudioCookOverrides* CompressionOverrides = FPlatformCompressionUtilities::GetCookOverrides();
+					const FPlatformAudioCookOverrides* CompressionOverrides = FPlatformCompressionUtilities::GetCookOverrides(*CookingTarget->IniPlatformName());
 
 					GetCompressedData(Format, CompressionOverrides); // Get the data from the DDC or build it
 					if (CompressionOverrides)
@@ -551,7 +551,7 @@ bool USoundWave::HasCompressedData(FName Format, ITargetPlatform* TargetPlatform
 	{
 		if (TargetPlatform)
 		{
-			CompressionOverrides = FPlatformCompressionUtilities::GetCookOverrides();
+			CompressionOverrides = FPlatformCompressionUtilities::GetCookOverrides(*TargetPlatform->IniPlatformName());
 		}
 	}
 	else
@@ -870,7 +870,7 @@ void USoundWave::PostLoad()
 		{
 			if (!Platform->IsServerOnly())
 			{
-				BeginGetCompressedData(Platform->GetWaveFormat(this), FPlatformCompressionUtilities::GetCookOverrides());
+				BeginGetCompressedData(Platform->GetWaveFormat(this), FPlatformCompressionUtilities::GetCookOverrides(*Platform->IniPlatformName()));
 			}
 		}
 	}
@@ -1134,7 +1134,7 @@ void USoundWave::InvalidateSoundWaveIfNeccessary()
 
 float USoundWave::GetSampleRateForTargetPlatform(const ITargetPlatform* TargetPlatform)
 {
-	const FPlatformAudioCookOverrides* Overrides = FPlatformCompressionUtilities::GetCookOverrides();
+	const FPlatformAudioCookOverrides* Overrides = FPlatformCompressionUtilities::GetCookOverrides(*TargetPlatform->IniPlatformName());
 	if (Overrides)
 	{
 		return GetSampleRateForCompressionOverrides(Overrides);
@@ -2036,7 +2036,7 @@ bool USoundWave::IsStreaming(const TCHAR* PlatformName/* = nullptr */) const
 		return false;
 	}
 
-	return IsStreaming(*FPlatformCompressionUtilities::GetCookOverrides());
+	return IsStreaming(*FPlatformCompressionUtilities::GetCookOverrides(PlatformName));
 }
 
 bool USoundWave::IsStreaming(const FPlatformAudioCookOverrides& Overrides) const
