Multiple LSAs map corruption: adding detection at map load time in editor, with a notification call to go through a repair process

diff --git a/Engine/Source/Runtime/Engine/Classes/Engine/Level.h b/Engine/Source/Runtime/Engine/Classes/Engine/Level.h
--- a/Engine/Source/Runtime/Engine/Classes/Engine/Level.h
+++ b/Engine/Source/Runtime/Engine/Classes/Engine/Level.h
@@ -36,6 +36,7 @@ struct FLevelCollection;
 class ULevelActorContainer;
 class FLevelPartitionOperationScope;
 class FRegisterComponentContext;
+class SNotificationItem;
 
 UINTERFACE()
 class ULevelPartitionInterface : public UInterface
@@ -1201,11 +1202,21 @@ private:
 	bool IncrementalRunConstructionScripts(bool bProcessAllActors);
 
 #if WITH_EDITOR
+private:
 	/** Attempts to detect and fix any issues with the level script blueprint and associated objects */
 	void RepairLevelScript();
 
 	/** Replace the existing LSA (if set) by spawning a new one based on this level's script blueprint */
 	void RegenerateLevelScriptActor();
+
+	/** Find and destroy any extra LSAs, as they will cause duplicated level script operations */
+	void RemoveExtraLevelScriptActors();
+
+	/** Notification popup used to guide the user to repair multiple LSAs detected upon loading in the editor */
+	TWeakPtr<SNotificationItem> MultipleLSAsNotification;
+	void OnMultipleLSAsPopupClicked();
+	void OnMultipleLSAsPopupDismissed();
+
 #endif // WITH_EDITOR
 };
 
diff --git a/Engine/Source/Runtime/Engine/Classes/Engine/LevelScriptActor.h b/Engine/Source/Runtime/Engine/Classes/Engine/LevelScriptActor.h
--- a/Engine/Source/Runtime/Engine/Classes/Engine/LevelScriptActor.h
+++ b/Engine/Source/Runtime/Engine/Classes/Engine/LevelScriptActor.h
@@ -68,6 +68,9 @@ class ENGINE_API ALevelScriptActor : public AActor
 
 #if WITH_EDITOR
 	virtual bool SupportsExternalPackaging() const override { return false; }
+
+	/** Utility function for finding a map corruption issue that can cause duplicate level scripting events */
+	TArray<ALevelScriptActor*> FindSiblingLevelScriptActors() const;
 #endif
 	//~ End AActor Interface
 
diff --git a/Engine/Source/Runtime/Engine/Private/Level.cpp b/Engine/Source/Runtime/Engine/Private/Level.cpp
--- a/Engine/Source/Runtime/Engine/Private/Level.cpp
+++ b/Engine/Source/Runtime/Engine/Private/Level.cpp
@@ -60,6 +60,10 @@ Level.cpp: Level-related functions
 #include "IAssetRegistry.h"
 #include "AssetData.h"
 #include "PieFixupSerializer.h"
+#include "Editor.h"
+#include "Subsystems/EditorActorSubsystem.h"
+#include "Widgets/Notifications/SNotificationList.h"
+#include "Framework/Notifications/NotificationManager.h"
 #endif
 #include "WorldPartition/WorldPartition.h"
 #include "Engine/LevelStreaming.h"
@@ -73,6 +77,7 @@ Level.cpp: Level-related functions
 #include "HAL/LowLevelMemTracker.h"
 #include "ObjectTrace.h"
 
+#define LOCTEXT_NAMESPACE "ULevel"
 DEFINE_LOG_CATEGORY(LogLevel);
 
 int32 GActorClusteringEnabled = 1;
@@ -3048,6 +3053,29 @@ void ULevel::RepairLevelScript()
 		}
 	}
 
+	// Show a toast notification to the user if there are multiple LSAs in this level (only for levels being actively edited)
+	if (LevelScriptActor && !MultipleLSAsNotification.IsValid() &&
+		OwningWorld && OwningWorld->WorldType == EWorldType::Editor)
+	{
+		TArray<ALevelScriptActor*> SiblingLSAs = LevelScriptActor->FindSiblingLevelScriptActors();
+
+		if (!SiblingLSAs.IsEmpty())
+		{
+			UE_LOG(LogLevel, Error, TEXT("Detected more than one LevelScriptActor in map '%s'. This can lead to duplicate level blueprint operations during play."), *GetPathName());
+
+			FNotificationInfo Info(LOCTEXT("MultipleLSAsPopupTitle", "Map Corruption: Multiple Level Script Actors"));
+			Info.ExpireDuration = 10.0f;
+			Info.bFireAndForget = true;
+			Info.Image = FCoreStyle::Get().GetBrush(TEXT("MessageLog.Error"));
+
+			Info.ButtonDetails.Add(FNotificationButtonInfo(LOCTEXT("MultipleLSAsPopupBeginRepair", "Repair Map"), LOCTEXT("MultipleLSAsBeginRepairTT", "Repairing the map by deleting extra LevelScriptActors"), FSimpleDelegate::CreateUObject(this, &ULevel::OnMultipleLSAsPopupClicked)));
+			Info.ButtonDetails.Add(FNotificationButtonInfo(LOCTEXT("MultipleLSAsPopupDismiss", "Dismiss"), LOCTEXT("MultipleLSAsPopupDismissTT", "Dismiss this notification"), FSimpleDelegate::CreateUObject(this, &ULevel::OnMultipleLSAsPopupDismissed)));
+
+			MultipleLSAsNotification = FSlateNotificationManager::Get().AddNotification(Info);
+			MultipleLSAsNotification.Pin()->SetCompletionState(SNotificationItem::CS_Pending);
+		}
+	}
+
 	// Catch the edge case where we have a level blueprint but have never created the LevelScriptActor based on it.
 	//    This could happen if a new level is saved before the level blueprint is compiled.
 	if (!LevelScriptActor && LevelScriptBlueprint && !LevelScriptBlueprint->bIsRegeneratingOnLoad &&
@@ -3106,4 +3134,48 @@ void ULevel::RegenerateLevelScriptActor()
 		UE_LOG(LogLevel, Error, TEXT("Skipped regeneration of LevelScriptActor due to blueprint '%s' having no spawnable class. A probable cause is that the blueprint is deriving from an invalid class, and may not function."), *LevelScriptBlueprint->GetName());
 	}
 }
+
+void ULevel::RemoveExtraLevelScriptActors()
+{
+	if (LevelScriptActor)
+	{
+		TArray<ALevelScriptActor*> ExtraLSAs = LevelScriptActor->FindSiblingLevelScriptActors();
+
+		if (ExtraLSAs.Num() > 0)
+		{
+			if (UEditorActorSubsystem* EditorActorSubsystem = GEditor->GetEditorSubsystem<UEditorActorSubsystem>())
+			{
+				for (ALevelScriptActor* LSA : ExtraLSAs)
+				{
+					UE_LOG(LogLevel, Log, TEXT("Deleting extra LevelScriptActor: %s"), *LSA->GetPathName());
+					EditorActorSubsystem->DestroyActor(LSA);
+				}
+			}
+		}
+	}
+}
+
+void ULevel::OnMultipleLSAsPopupClicked()
+{
+	RemoveExtraLevelScriptActors();
+	OnMultipleLSAsPopupDismissed();
+}
+
+void ULevel::OnMultipleLSAsPopupDismissed()
+{
+	if (MultipleLSAsNotification.IsValid())
+	{
+		TSharedPtr<SNotificationItem> NotifPopup = MultipleLSAsNotification.Pin();
+
+		NotifPopup->SetCompletionState(SNotificationItem::CS_None);
+		NotifPopup->SetExpireDuration(0.0f);
+		NotifPopup->SetFadeOutDuration(0.0f);
+		NotifPopup->ExpireAndFadeout();
+
+		MultipleLSAsNotification.Reset();
+	}
+}
+
 #endif // WITH_EDITOR
+
+#undef LOCTEXT_NAMESPACE
diff --git a/Engine/Source/Runtime/Engine/Private/LevelScriptActor.cpp b/Engine/Source/Runtime/Engine/Private/LevelScriptActor.cpp
--- a/Engine/Source/Runtime/Engine/Private/LevelScriptActor.cpp
+++ b/Engine/Source/Runtime/Engine/Private/LevelScriptActor.cpp
@@ -33,22 +33,13 @@ ALevelScriptActor::ALevelScriptActor(const FObjectInitializer& ObjectInitializer
 
 #if WITH_EDITOR
 	// this is intended as an early detection of cases where more than one LevelScriptActor is introduced into a single map
-	const UObject* ThisOuter = GetOuter();
-	if (!Cast<UPackage>(ThisOuter))
-	{
-		TArray<UObject*> AllSiblingObjects;
-		GetObjectsWithOuter(ThisOuter, AllSiblingObjects, false, RF_NoFlags, EInternalObjectFlags::PendingKill);
+	TArray<ALevelScriptActor*> Siblings = FindSiblingLevelScriptActors();
 
-		for (const UObject* Sibling : AllSiblingObjects)
-		{
-			bool bIsNotAnLSA           = !Cast<ALevelScriptActor>(Sibling);
-			bool bIsTheSameObject      = Sibling == this;
-			bool bHasNewerClassVersion = Sibling->GetClass()->HasAnyClassFlags(CLASS_NewerVersionExists);
-
-			ensureMsgf((bIsNotAnLSA || bIsTheSameObject || bHasNewerClassVersion),
-				TEXT("Detected the creation of more than one LevelScriptActor (%s, %s) within the same outer (%s) in package (%s). This can lead to duplicate level blueprint operations during play."),
-				*GetName(), *Sibling->GetName(), *ThisOuter->GetName(), *ThisOuter->GetPackage()->GetName());
-		}
+	for (const ALevelScriptActor* Sibling : Siblings)
+	{
+		ensureMsgf(false,
+			TEXT("Detected the creation of more than one LevelScriptActor (%s, %s) within the same outer. This can lead to duplicate level blueprint operations during play."),
+			*GetPathName(), *Sibling->GetPathName());
 	}
 #endif // WITH_EDITOR
 }
@@ -122,3 +113,34 @@ void ALevelScriptActor::DisableInput(class APlayerController* PlayerController)
 	}
 	bInputEnabled = false;
 }
+
+
+#if WITH_EDITOR
+
+TArray<ALevelScriptActor*> ALevelScriptActor::FindSiblingLevelScriptActors() const
+{
+	TArray<ALevelScriptActor*> SiblingLSAs;
+
+	const UObject* ThisOuter = GetOuter();
+	if (!ThisOuter->IsA<UPackage>())
+	{
+		TArray<UObject*> AllSiblingObjects;
+		GetObjectsWithOuter(ThisOuter, AllSiblingObjects, false, RF_NoFlags, EInternalObjectFlags::PendingKill);
+
+		for (UObject* Sibling : AllSiblingObjects)
+		{
+			ALevelScriptActor* SiblingAsLSA = Cast<ALevelScriptActor>(Sibling);
+			const bool bIsNotAnLSA = SiblingAsLSA == nullptr;
+			const bool bIsTheSameObject = Sibling == this;
+			const bool bHasNewerClassVersion = Sibling->GetClass()->HasAnyClassFlags(CLASS_NewerVersionExists);
+
+			if (!bIsNotAnLSA && !bIsTheSameObject && !bHasNewerClassVersion)
+			{
+				SiblingLSAs.Add(SiblingAsLSA);
+			}
+		}
+	}
+
+	return SiblingLSAs;
+}
+#endif // WITH_EDITOR
\ No newline at end of file
