Fix for HRTF rendering

#pf trivial

diff --git a/Engine/Source/Runtime/Engine/Private/SoundWave.cpp b/Engine/Source/Runtime/Engine/Private/SoundWave.cpp
--- a/Engine/Source/Runtime/Engine/Private/SoundWave.cpp
+++ b/Engine/Source/Runtime/Engine/Private/SoundWave.cpp
@@ -2263,7 +2263,23 @@ void USoundWave::Parse(FAudioDevice* AudioDevice, const UPTRINT NodeWaveInstance
 	{
 		//Parse the parameters of the wave instance
 		WaveInstance->bEnableBusSends = ParseParams.bEnableBusSends;
-		WaveInstance->bEnableBaseSubmix = ParseParams.bEnableBaseSubmix;
+
+		// HRTF rendering doesn't render their output on the base submix
+		if (WaveInstance->SpatializationMethod != SPATIALIZATION_HRTF)
+		{
+			if (ActiveSound.bHasActiveMainSubmixOutputOverride)
+			{
+				WaveInstance->bEnableBaseSubmix = ActiveSound.bEnableMainSubmixOutputOverride;
+			}
+			else
+			{
+				WaveInstance->bEnableBaseSubmix = ParseParams.bEnableBaseSubmix;
+			}
+		}
+		else
+		{
+			WaveInstance->bEnableBaseSubmix = false;
+		}
 		WaveInstance->bEnableSubmixSends = ParseParams.bEnableSubmixSends;
 
 		// Active sounds can override their enablement behavior via audio components
@@ -2272,11 +2288,6 @@ void USoundWave::Parse(FAudioDevice* AudioDevice, const UPTRINT NodeWaveInstance
 			WaveInstance->bEnableBusSends = ActiveSound.bEnableBusSendRoutingOverride;
 		}
 
-		if (ActiveSound.bHasActiveMainSubmixOutputOverride)
-		{
-			WaveInstance->bEnableBaseSubmix = ActiveSound.bEnableMainSubmixOutputOverride;
-		}
-
 		if (ActiveSound.bHasActiveSubmixSendRoutingOverride)
 		{
 			WaveInstance->bEnableSubmixSends = ActiveSound.bEnableSubmixSendRoutingOverride;
