ISM - Added support for additional random seeds ranges
Used by the HLOD instancing batching to preserve random values from the source actors.

diff --git a/Engine/Source/Developer/WorldPartitionHLODUtilities/Private/HLODBuilderInstancing.cpp b/Engine/Source/Developer/WorldPartitionHLODUtilities/Private/HLODBuilderInstancing.cpp
--- a/Engine/Source/Developer/WorldPartitionHLODUtilities/Private/HLODBuilderInstancing.cpp
+++ b/Engine/Source/Developer/WorldPartitionHLODUtilities/Private/HLODBuilderInstancing.cpp
@@ -35,6 +35,8 @@ namespace
 
 		int32							NumCustomDataFloats = 0;
 		TArray<float>					InstancesCustomData;
+
+		TArray<FInstancedStaticMeshRandomSeed> RandomSeeds;
 	};
 }
 
@@ -61,6 +63,7 @@ TArray<UPrimitiveComponent*> UHLODBuilderInstancing::CreateComponents(AWorldPart
 			if (UInstancedStaticMeshComponent* ISMC = Cast<UInstancedStaticMeshComponent>(SMC))
 			{
 				InstancingData.NumCustomDataFloats = FMath::Max(InstancingData.NumCustomDataFloats, ISMC->NumCustomDataFloats);
+				InstancingData.RandomSeeds.Add( { InstancingData.NumInstances, ISMC->InstancingRandomSeed } );
 				InstancingData.NumInstances += ISMC->GetInstanceCount();
 			}
 			else
@@ -124,6 +127,16 @@ TArray<UPrimitiveComponent*> UHLODBuilderInstancing::CreateComponents(AWorldPart
 		Component->AddInstances(EntryInstancingData.InstancesTransforms, /*bShouldReturnIndices*/false, /*bWorldSpace*/true);
 		Component->PerInstanceSMCustomData = MoveTemp(EntryInstancingData.InstancesCustomData);
 
+		if (!EntryInstancingData.RandomSeeds.IsEmpty())
+		{
+			Component->InstancingRandomSeed = EntryInstancingData.RandomSeeds[0].RandomSeed;
+		}
+
+		if (EntryInstancingData.RandomSeeds.Num() > 1)
+		{
+			Component->AdditionalRandomSeeds = TArrayView<FInstancedStaticMeshRandomSeed>(&EntryInstancingData.RandomSeeds[1], EntryInstancingData.RandomSeeds.Num() - 1);
+		}
+
 		Components.Add(Component);
 	};
 
diff --git a/Engine/Source/Developer/WorldPartitionHLODUtilities/Private/WorldPartitionHLODUtilities.cpp b/Engine/Source/Developer/WorldPartitionHLODUtilities/Private/WorldPartitionHLODUtilities.cpp
--- a/Engine/Source/Developer/WorldPartitionHLODUtilities/Private/WorldPartitionHLODUtilities.cpp
+++ b/Engine/Source/Developer/WorldPartitionHLODUtilities/Private/WorldPartitionHLODUtilities.cpp
@@ -93,7 +93,7 @@ static uint32 ComputeHLODHash(AWorldPartitionHLOD* InHLODActor, const TArray<FWo
 	FArchiveCrc32 Ar;
 
 	// Base key, changing this will force a rebuild of all HLODs
-	FString HLODBaseKey = "29766A60F4A74BF9B10B77095F4B3B16";
+	FString HLODBaseKey = "4CE9431F9C6842F996786C980641B63A";
 	Ar << HLODBaseKey;
 
 	// HLOD Layer
diff --git a/Engine/Source/Runtime/Engine/Classes/Components/InstancedStaticMeshComponent.h b/Engine/Source/Runtime/Engine/Classes/Components/InstancedStaticMeshComponent.h
--- a/Engine/Source/Runtime/Engine/Classes/Components/InstancedStaticMeshComponent.h
+++ b/Engine/Source/Runtime/Engine/Classes/Components/InstancedStaticMeshComponent.h
@@ -119,6 +119,18 @@ struct FInstancedStaticMeshMappingInfo
 	}
 };
 
+USTRUCT()
+struct FInstancedStaticMeshRandomSeed
+{
+	GENERATED_USTRUCT_BODY()
+
+	UPROPERTY()
+	int32 StartInstanceIndex;
+
+	UPROPERTY()
+	int32 RandomSeed;
+};
+
 /** A component that efficiently renders multiple instances of the same StaticMesh. */
 UCLASS(ClassGroup = Rendering, meta = (BlueprintSpawnableComponent), Blueprintable)
 class ENGINE_API UInstancedStaticMeshComponent : public UStaticMeshComponent, public ISMInstanceManager
@@ -157,6 +169,10 @@ class ENGINE_API UInstancedStaticMeshComponent : public UStaticMeshComponent, pu
 	UPROPERTY(EditAnywhere, BlueprintReadWrite, Category=InstancedStaticMeshComponent)
 	int32 InstancingRandomSeed=0;
 
+	/** Additional random seeds ranges. Each seed entry will be applied from AdditionalRandomSeeds[i].StartInstanceIndex to AdditionalRandomSeeds[i+1].StartInstanceIndex -1 */
+	UPROPERTY()
+	TArray<FInstancedStaticMeshRandomSeed> AdditionalRandomSeeds;
+
 	/** Distance from camera at which each instance begins to fade out. */
 	UPROPERTY(EditAnywhere, BlueprintReadOnly, Category=Culling)
 	int32 InstanceStartCullDistance;
@@ -562,4 +578,8 @@ public:
 	/* The cached random seed */
 	UPROPERTY()
 	int32 InstancingRandomSeed = 0;
+
+	/* Additional random seeds */
+	UPROPERTY()
+	TArray<FInstancedStaticMeshRandomSeed> AdditionalRandomSeeds;
 };
diff --git a/Engine/Source/Runtime/Engine/Private/HLODProxy.cpp b/Engine/Source/Runtime/Engine/Private/HLODProxy.cpp
--- a/Engine/Source/Runtime/Engine/Private/HLODProxy.cpp
+++ b/Engine/Source/Runtime/Engine/Private/HLODProxy.cpp
@@ -474,6 +474,12 @@ uint32 UHLODProxy::GetCRC(UStaticMeshComponent* InComponent, uint32 InCRC, const
 
 		Ar << ISMC->PerInstanceSMCustomData;
 		Ar << ISMC->InstancingRandomSeed;
+
+		for (FInstancedStaticMeshRandomSeed& ISMCRandomSeed : ISMC->AdditionalRandomSeeds)
+		{
+			Ar << ISMCRandomSeed.StartInstanceIndex;
+			Ar << ISMCRandomSeed.RandomSeed;
+		}
 	}
 
 	return Ar.GetCrc();
diff --git a/Engine/Source/Runtime/Engine/Private/InstancedStaticMesh.cpp b/Engine/Source/Runtime/Engine/Private/InstancedStaticMesh.cpp
--- a/Engine/Source/Runtime/Engine/Private/InstancedStaticMesh.cpp
+++ b/Engine/Source/Runtime/Engine/Private/InstancedStaticMesh.cpp
@@ -1925,6 +1925,7 @@ TStructOnScope<FActorComponentInstanceData> UInstancedStaticMeshComponent::GetCo
 
 	// Back up random seed
 	StaticMeshInstanceData->InstancingRandomSeed = InstancingRandomSeed;
+	StaticMeshInstanceData->AdditionalRandomSeeds = AdditionalRandomSeeds;
 #endif
 	return InstanceData;
 }
@@ -1987,6 +1988,7 @@ void UInstancedStaticMeshComponent::ApplyComponentInstanceData(FInstancedStaticM
 	SelectedInstances = InstancedMeshData->SelectedInstances;
 
 	InstancingRandomSeed = InstancedMeshData->InstancingRandomSeed;
+	AdditionalRandomSeeds = InstancedMeshData->AdditionalRandomSeeds;
 
 	// Force recreation of the render data
 	InstanceUpdateCmdBuffer.Edit();
@@ -2123,6 +2125,9 @@ void UInstancedStaticMeshComponent::BuildRenderData(FStaticMeshInstanceData& Out
 	
 	check(InstancingRandomSeed != 0);
 	FRandomStream RandomStream = FRandomStream(InstancingRandomSeed);
+
+	auto AdditionanlRandomSeedsIt = AdditionalRandomSeeds.CreateIterator();
+	int32 SeedResetIndex = AdditionanlRandomSeedsIt ? AdditionanlRandomSeedsIt->StartInstanceIndex : INDEX_NONE;
 	
 	for (int32 Index = 0; Index < NumInstances; ++Index)
 	{
@@ -2132,7 +2137,15 @@ void UInstancedStaticMeshComponent::BuildRenderData(FStaticMeshInstanceData& Out
 			// could be skipped by density settings
 			continue;
 		}
-			
+
+		// Reset the random stream if necessary
+		if (Index == SeedResetIndex)
+		{
+			RandomStream = FRandomStream(AdditionanlRandomSeedsIt->RandomSeed);
+			AdditionanlRandomSeedsIt++;
+			SeedResetIndex = AdditionanlRandomSeedsIt ? AdditionanlRandomSeedsIt->StartInstanceIndex : INDEX_NONE;
+		}
+
 		const FInstancedStaticMeshInstanceData& InstanceData = PerInstanceSMData[Index];
 		FVector2D LightmapUVBias = FVector2D(-1.0f, -1.0f);
 		FVector2D ShadowmapUVBias = FVector2D(-1.0f, -1.0f);
