Tentative crash fix when reducing skin cache memory pool.

diff --git a/Engine/Source/Runtime/Engine/Private/GPUSkinCache.cpp b/Engine/Source/Runtime/Engine/Private/GPUSkinCache.cpp
--- a/Engine/Source/Runtime/Engine/Private/GPUSkinCache.cpp
+++ b/Engine/Source/Runtime/Engine/Private/GPUSkinCache.cpp
@@ -2283,7 +2283,7 @@ FCachedGeometry FGPUSkinCache::GetCachedGeometry(uint32 ComponentId) const
 	FCachedGeometry Out;
 	for (FGPUSkinCacheEntry* Entry : Entries)
 	{
-		if (Entry && Entry->GPUSkin && Entry->GPUSkin->GetComponentId() == ComponentId)
+		if (Entry && Entry->GPUSkin && Entry->GPUSkin->GetComponentId() == ComponentId && Entry->GPUSkin->HaveValidDynamicData())
 		{
 			const uint32 LODIndex = Entry->LOD;
 			const FSkeletalMeshRenderData& RenderData = Entry->GPUSkin->GetSkeletalMeshRenderData();
