World Partition: skip HLOD levels when remapping soft objects paths for streamed cells in PIE.

diff --git a/Engine/Source/Runtime/Engine/Private/WorldPartition/HLOD/HLODActor.cpp b/Engine/Source/Runtime/Engine/Private/WorldPartition/HLOD/HLODActor.cpp
--- a/Engine/Source/Runtime/Engine/Private/WorldPartition/HLOD/HLODActor.cpp
+++ b/Engine/Source/Runtime/Engine/Private/WorldPartition/HLOD/HLODActor.cpp
@@ -203,11 +203,4 @@ void AWorldPartitionHLOD::BuildHLOD(bool bForceBuild)
 	BuildActorTextureStreamingData(this, EMaterialQualityLevel::High, GMaxRHIFeatureLevel);
 }
 
-#endif // WITH_EDITOR
-
-
-UWorldPartitionRuntimeHLODCellData::UWorldPartitionRuntimeHLODCellData(const FObjectInitializer& ObjectInitializer)
-	: Super(ObjectInitializer)
-{
-}
-
+#endif // WITH_EDITOR
\ No newline at end of file
diff --git a/Engine/Source/Runtime/Engine/Private/WorldPartition/WorldPartitionLevelStreamingDynamic.cpp b/Engine/Source/Runtime/Engine/Private/WorldPartition/WorldPartitionLevelStreamingDynamic.cpp
--- a/Engine/Source/Runtime/Engine/Private/WorldPartition/WorldPartitionLevelStreamingDynamic.cpp
+++ b/Engine/Source/Runtime/Engine/Private/WorldPartition/WorldPartitionLevelStreamingDynamic.cpp
@@ -322,25 +322,27 @@ void UWorldPartitionLevelStreamingDynamic::FinalizeRuntimeLevel()
 	UPackage* RuntimePackage = RuntimeLevel->GetPackage();
 	RuntimePackage->MarkAsFullyLoaded();
 
-	if (OuterWorld->IsPlayInEditor())
+	if (!StreamingCell->GetIsHLOD())
 	{
-		int32 PIEInstanceID = GetPackage()->PIEInstanceID;
-		check(PIEInstanceID != INDEX_NONE);
+		TRACE_CPUPROFILER_EVENT_SCOPE(FixupIDs);
 
+		if (OuterWorld->IsPlayInEditor())
 		{
-			TRACE_CPUPROFILER_EVENT_SCOPE(FixupIDsForPIE);
+			int32 PIEInstanceID = GetPackage()->PIEInstanceID;
+			check(PIEInstanceID != INDEX_NONE);
+
 			RuntimeLevel->FixupForPIE(PIEInstanceID, [&](int32 InPIEInstanceID, FSoftObjectPath& ObjectPath)
 			{
 				// Remap Runtime Level's SoftObjectPath before each PIE Fixup to avoid doing 2 passes of serialization
 				OuterWorldPartition->RemapSoftObjectPath(ObjectPath);
 			});
 		}
-	}
-	else if (OuterWorld->IsGameWorld())
-	{
-		check(IsRunningGame());
-		// Remap Runtime Level's SoftObjectPaths
-		FWorldPartitionLevelHelper::RemapLevelSoftObjectPaths(RuntimeLevel, OuterWorldPartition.Get());
+		else if (OuterWorld->IsGameWorld())
+		{
+			check(IsRunningGame());
+			// Remap Runtime Level's SoftObjectPaths
+			FWorldPartitionLevelHelper::RemapLevelSoftObjectPaths(RuntimeLevel, OuterWorldPartition.Get());
+		}
 	}
 
 	SetLoadedLevel(RuntimeLevel);
diff --git a/Engine/Source/Runtime/Engine/Private/WorldPartition/WorldPartitionRuntimeCell.cpp b/Engine/Source/Runtime/Engine/Private/WorldPartition/WorldPartitionRuntimeCell.cpp
--- a/Engine/Source/Runtime/Engine/Private/WorldPartition/WorldPartitionRuntimeCell.cpp
+++ b/Engine/Source/Runtime/Engine/Private/WorldPartition/WorldPartitionRuntimeCell.cpp
@@ -20,12 +20,6 @@ UWorldPartitionRuntimeCell::UWorldPartitionRuntimeCell(const FObjectInitializer&
 {}
 
 #if WITH_EDITOR
-
-void UWorldPartitionRuntimeCell::AddCellData(const UWorldPartitionRuntimeCellData* InCellData)
-{
-	CellDataMap.Add(InCellData->GetClass(), InCellData);
-}
-
 void UWorldPartitionRuntimeCell::SetDataLayers(const TArray<const UDataLayer*>& InDataLayers)
 {
 	check(DataLayers.IsEmpty());
@@ -101,13 +95,4 @@ bool UWorldPartitionRuntimeCell::IsDebugShown() const
 		   FWorldPartitionDebugHelper::IsDebugStreamingStatusShown(GetStreamingStatus()) &&
 		   FWorldPartitionDebugHelper::AreDebugDataLayersShown(DataLayers) &&
 		   FWorldPartitionDebugHelper::IsDebugCellNameShow(DebugName);
-}
-
-const UWorldPartitionRuntimeCellData* UWorldPartitionRuntimeCell::GetCellData(const TSubclassOf<UWorldPartitionRuntimeCellData> InCellDataClass) const
-{
-	return CellDataMap.FindRef(InCellDataClass);
-}
-
-UWorldPartitionRuntimeCellData::UWorldPartitionRuntimeCellData(const FObjectInitializer& ObjectInitializer)
-: Super(ObjectInitializer)
-{}
\ No newline at end of file
+}
\ No newline at end of file
diff --git a/Engine/Source/Runtime/Engine/Private/WorldPartition/WorldPartitionRuntimeSpatialHash.cpp b/Engine/Source/Runtime/Engine/Private/WorldPartition/WorldPartitionRuntimeSpatialHash.cpp
--- a/Engine/Source/Runtime/Engine/Private/WorldPartition/WorldPartitionRuntimeSpatialHash.cpp
+++ b/Engine/Source/Runtime/Engine/Private/WorldPartition/WorldPartitionRuntimeSpatialHash.cpp
@@ -952,6 +952,7 @@ bool UWorldPartitionRuntimeSpatialHash::CreateStreamingGrid(const FSpatialHashRu
 				StreamingCell->SetDebugInfo(CellGlobalCoords, CurrentStreamingGrid.GridName);
 				StreamingCell->SetClientOnlyVisible(CurrentStreamingGrid.bClientOnlyVisible);
 				StreamingCell->SetBlockOnSlowLoading(CurrentStreamingGrid.bBlockOnSlowStreaming);
+				StreamingCell->SetIsHLOD(RuntimeGrid.HLODLayer ? true : false);
 
 				UE_LOG(LogWorldPartition, Verbose, TEXT("Cell%s %s Actors = %d Bounds (%s)"), bIsCellAlwaysLoaded ? TEXT(" (AlwaysLoaded)") : TEXT(""), *StreamingCell->GetName(), FilteredActors.Num(), *Bounds.ToString());
 
@@ -970,9 +971,6 @@ bool UWorldPartitionRuntimeSpatialHash::CreateStreamingGrid(const FSpatialHashRu
 					}
 				}
 
-				// Keep track of all AWorldPartitionHLOD actors referenced by this cell
-				int32 NumHLODActors = 0;
-
 				for (const FActorInstance& ActorInstance : FilteredActors)
 				{
 					const FWorldPartitionActorDescView& ActorDescView = ActorInstance.GetActorDescView();
@@ -985,22 +983,9 @@ bool UWorldPartitionRuntimeSpatialHash::CreateStreamingGrid(const FSpatialHashRu
 							StreamingCell->ActorContainer->Actors.Add(Actor->GetFName(), Actor);
 						}
 					}
-
-					if (ActorDescView.GetActorClass()->IsChildOf<AWorldPartitionHLOD>())
-					{
-						NumHLODActors++;
-					}
-
 					UE_LOG(LogWorldPartition, Verbose, TEXT("  Actor : %s (%s) (Container %08x) Origin(%s)"), *(ActorDescView.GetActorPath().ToString()), *ActorDescView.GetGuid().ToString(EGuidFormats::UniqueObjectGuid), ActorInstance.ContainerInstance->ID, *FVector2D(ActorInstance.GetOrigin()).ToString());
 				}
 
-				if (NumHLODActors == FilteredActors.Num())
-				{
-					// Mark this cell as containing HLOD data
-					UWorldPartitionRuntimeHLODCellData* HLODCellData = NewObject<UWorldPartitionRuntimeHLODCellData>(StreamingCell);
-					StreamingCell->AddCellData(HLODCellData);
-				}
-
 				if (Mode == EWorldPartitionStreamingMode::Cook)
 				{
 					UE_LOG(LogWorldPartition, Log, TEXT("Creating runtime streaming cells %s."), *StreamingCell->GetName());
diff --git a/Engine/Source/Runtime/Engine/Private/WorldPartition/WorldPartitionStreamingPolicy.cpp b/Engine/Source/Runtime/Engine/Private/WorldPartition/WorldPartitionStreamingPolicy.cpp
--- a/Engine/Source/Runtime/Engine/Private/WorldPartition/WorldPartitionStreamingPolicy.cpp
+++ b/Engine/Source/Runtime/Engine/Private/WorldPartition/WorldPartitionStreamingPolicy.cpp
@@ -259,7 +259,7 @@ void UWorldPartitionStreamingPolicy::UpdateStreamingState()
 		{
 			for (auto It = Cells.CreateIterator(); It; ++It)
 			{
-				if ((*It)->HasCellData<UWorldPartitionRuntimeHLODCellData>())
+				if ((*It)->GetIsHLOD())
 				{
 					It.RemoveCurrent();
 				}
@@ -559,7 +559,7 @@ bool UWorldPartitionStreamingPolicy::IsStreamingCompleted(EWorldPartitionRuntime
 				// Don't consider HLOD cells if HLODs are disabled.
 				if (!bIsHLODEnabled)
 				{
-					bSkipCell = Cell->HasCellData<UWorldPartitionRuntimeHLODCellData>();
+					bSkipCell = Cell->GetIsHLOD();
 				}
 
 				// If we are querying for Unloaded/Loaded but a Cell is part of a data layer outside of the query that is activated do not consider it
diff --git a/Engine/Source/Runtime/Engine/Public/WorldPartition/HLOD/HLODActor.h b/Engine/Source/Runtime/Engine/Public/WorldPartition/HLOD/HLODActor.h
--- a/Engine/Source/Runtime/Engine/Public/WorldPartition/HLOD/HLODActor.h
+++ b/Engine/Source/Runtime/Engine/Public/WorldPartition/HLOD/HLODActor.h
@@ -114,10 +114,4 @@ private:
 	FName SourceCellName;
 };
 
-DEFINE_ACTORDESC_TYPE(AWorldPartitionHLOD, FHLODActorDesc);
-
-UCLASS()
-class UWorldPartitionRuntimeHLODCellData : public UWorldPartitionRuntimeCellData
-{
-	GENERATED_UCLASS_BODY()
-};
+DEFINE_ACTORDESC_TYPE(AWorldPartitionHLOD, FHLODActorDesc);
\ No newline at end of file
diff --git a/Engine/Source/Runtime/Engine/Public/WorldPartition/WorldPartitionRuntimeCell.h b/Engine/Source/Runtime/Engine/Public/WorldPartition/WorldPartitionRuntimeCell.h
--- a/Engine/Source/Runtime/Engine/Public/WorldPartition/WorldPartitionRuntimeCell.h
+++ b/Engine/Source/Runtime/Engine/Public/WorldPartition/WorldPartitionRuntimeCell.h
@@ -81,15 +81,6 @@ struct FWorldPartitionRuntimeCellObjectMapping
 #endif
 };
 
-/**
- * Serve as a generic container of type specific data that can be assigned to each runtime cell
- */
-UCLASS(Abstract)
-class UWorldPartitionRuntimeCellData : public UObject
-{
-	GENERATED_UCLASS_BODY()
-};
-
 class UActorDescContainer;
 
 /**
@@ -162,7 +153,6 @@ class UWorldPartitionRuntimeCell : public UObject
 
 	void SetDataLayers(const TArray<const UDataLayer*>& InDataLayers);
 	void SetDebugInfo(FIntVector InCoords, FName InGridName);
-	void AddCellData(const UWorldPartitionRuntimeCellData* InCellData);
 	virtual void AddActorToCell(const FWorldPartitionActorDescView& ActorDescView, uint64 InContainerID, const FTransform& InContainerTransform, const UActorDescContainer* InContainer) PURE_VIRTUAL(UWorldPartitionRuntimeCell::AddActorToCell,);
 	virtual int32 GetActorCount() const PURE_VIRTUAL(UWorldPartitionRuntimeCell::GetActorCount, return 0;);
 
@@ -171,12 +161,12 @@ class UWorldPartitionRuntimeCell : public UObject
 	virtual bool PopulateGeneratedPackageForCook(UPackage* InPackage) PURE_VIRTUAL(UWorldPartitionRuntimeCell::PopulateGeneratedPackageForCook, return false;);
 	virtual void MoveAlwaysLoadedContentToPersistentLevel() PURE_VIRTUAL(UWorldPartitionRuntimeCell::MoveAlwaysLoadedContentToPersistentLevel);
 	virtual FString GetPackageNameToCreate() const PURE_VIRTUAL(UWorldPartitionRuntimeCell::GetPackageNameToCreate, return FString(""););
-#endif
 
-	const UWorldPartitionRuntimeCellData* GetCellData(const TSubclassOf<UWorldPartitionRuntimeCellData> InCellDataClass) const;
-	template <class T> inline const T* GetCellData() const { return Cast<const T>(GetCellData(T::StaticClass())); }
-	template <class T> inline bool HasCellData() const { return GetCellData<T>() != nullptr; }
+	void SetIsHLOD(bool bInIsHLOD) { bIsHLOD = bInIsHLOD; }
+#endif
 	
+	bool GetIsHLOD() const { return bIsHLOD; }
+
 protected:
 #if WITH_EDITOR
 	void UpdateDebugName();
@@ -186,9 +176,6 @@ protected:
 	bool bIsAlwaysLoaded;
 
 private:
-	UPROPERTY()
-	TMap<const TSubclassOf<UWorldPartitionRuntimeCellData>, TObjectPtr<const UWorldPartitionRuntimeCellData>> CellDataMap;
-
 	UPROPERTY()
 	TArray<FName> DataLayers;
 
@@ -209,6 +196,9 @@ private:
 	UPROPERTY()
 	bool bClientOnlyVisible;
 
+	UPROPERTY()
+	bool bIsHLOD;
+
 	UPROPERTY()
 	bool bBlockOnSlowLoading;
 
