Allow access to CustomDepth and Stencil in decals for mobile

diff --git a/Engine/Source/Runtime/Engine/Private/Materials/HLSLMaterialTranslator.cpp b/Engine/Source/Runtime/Engine/Private/Materials/HLSLMaterialTranslator.cpp
--- a/Engine/Source/Runtime/Engine/Private/Materials/HLSLMaterialTranslator.cpp
+++ b/Engine/Source/Runtime/Engine/Private/Materials/HLSLMaterialTranslator.cpp
@@ -6222,7 +6222,7 @@ void FHLSLMaterialTranslator::UseSceneTextureId(ESceneTextureId SceneTextureId,
 			Errorf(TEXT("Decals can only access SceneDepth, CustomDepth, CustomStencil, and WorldNormal."));
 		}
 
-		const bool bSceneTextureRequiresSM5 = SceneTextureId == PPI_WorldNormal || SceneTextureId == PPI_CustomDepth || SceneTextureId == PPI_CustomStencil;
+		const bool bSceneTextureRequiresSM5 = SceneTextureId == PPI_WorldNormal;
 		if (bSceneTextureRequiresSM5)
 		{
 			ErrorUnlessFeatureLevelSupported(ERHIFeatureLevel::SM5);
