Scene Outliner: Show Data Layers for unloaded actors

diff --git a/Engine/Source/Editor/SceneOutliner/Private/SceneOutlinerModule.cpp b/Engine/Source/Editor/SceneOutliner/Private/SceneOutlinerModule.cpp
--- a/Engine/Source/Editor/SceneOutliner/Private/SceneOutlinerModule.cpp
+++ b/Engine/Source/Editor/SceneOutliner/Private/SceneOutlinerModule.cpp
@@ -229,23 +229,34 @@ void FSceneOutlinerModule::CreateActorInfoColumns(FSceneOutlinerInitializationOp
 	{
 		TStringBuilder<128> Builder;
 
+		TArray<const UDataLayer*> DataLayerObjects;
+
 		if (const FActorTreeItem* ActorItem = Item.CastTo<FActorTreeItem>())
 		{
 			AActor* Actor = ActorItem->Actor.Get();
 
-			if (!Actor)
+			if (Actor)
 			{
-				return FString();
+				DataLayerObjects = Actor->GetDataLayerObjects();
+			}
+		}
+		else if (const FActorDescTreeItem* ActorDescItem = Item.CastTo<FActorDescTreeItem>())
+		{
+			const FWorldPartitionActorDesc* ActorDesc = ActorDescItem->ActorDescHandle.GetActorDesc();
+
+			if (ActorDesc)
+			{
+				DataLayerObjects = ActorDesc->GetDataLayerObjects();
 			}
+		}
 
-			for (const UDataLayer* DataLayer : Actor->GetDataLayerObjects())
+		for (const UDataLayer* DataLayer : DataLayerObjects)
+		{
+			if (Builder.Len())
 			{
-				if (Builder.Len())
-				{
-					Builder += TEXT(", ");
-				}
-				Builder += DataLayer->GetDataLayerLabel().ToString();
+				Builder += TEXT(", ");
 			}
+			Builder += DataLayer->GetDataLayerLabel().ToString();
 		}
 		
 		return Builder.ToString();
diff --git a/Engine/Source/Runtime/Engine/Private/WorldPartition/WorldPartitionActorDesc.cpp b/Engine/Source/Runtime/Engine/Private/WorldPartition/WorldPartitionActorDesc.cpp
--- a/Engine/Source/Runtime/Engine/Private/WorldPartition/WorldPartitionActorDesc.cpp
+++ b/Engine/Source/Runtime/Engine/Private/WorldPartition/WorldPartitionActorDesc.cpp
@@ -16,6 +16,7 @@
 #include "WorldPartition/WorldPartition.h"
 #include "WorldPartition/ActorDescContainer.h"
 #include "WorldPartition/DataLayer/DataLayer.h"
+#include "WorldPartition/DataLayer/WorldDataLayers.h"
 #include "WorldPartition/HLOD/HLODLayer.h"
 #include "Engine/Public/ActorReferencesUtils.h"
 #endif
@@ -215,6 +216,29 @@ FName FWorldPartitionActorDesc::GetActorName() const
 	return *ActorName;
 }
 
+TArray<const UDataLayer*> FWorldPartitionActorDesc::GetDataLayerObjects() const
+{
+	if (Container)
+	{
+		if (AWorldDataLayers* WorldDataLayers = Container->GetWorld()->GetWorldDataLayers())
+		{
+			TArray<const UDataLayer*> DataLayerObjects;
+			DataLayerObjects.Reserve(DataLayers.Num());
+			for (const FName& DataLayerName : DataLayers)
+			{
+				if (const UDataLayer* DataLayer = WorldDataLayers->GetDataLayerFromName(DataLayerName))
+				{
+					DataLayerObjects.Add(DataLayer);
+				}
+			}
+
+			return DataLayerObjects;
+		}
+	}
+
+	return TArray<const UDataLayer*>();
+}
+
 UHLODLayer* FWorldPartitionActorDesc::GetHLODLayer() const
 {
 	return HLODLayer.IsNone() ? nullptr : Cast<UHLODLayer>(FSoftObjectPath(HLODLayer).TryLoad());
diff --git a/Engine/Source/Runtime/Engine/Public/WorldPartition/WorldPartitionActorDesc.h b/Engine/Source/Runtime/Engine/Public/WorldPartition/WorldPartitionActorDesc.h
--- a/Engine/Source/Runtime/Engine/Public/WorldPartition/WorldPartitionActorDesc.h
+++ b/Engine/Source/Runtime/Engine/Public/WorldPartition/WorldPartitionActorDesc.h
@@ -67,6 +67,7 @@ public:
 	inline bool GetActorIsHLODRelevant() const { return bActorIsHLODRelevant; }
 	class UHLODLayer* GetHLODLayer() const;
 	inline const TArray<FName>& GetDataLayers() const { return DataLayers; }
+	TArray<const UDataLayer*> GetDataLayerObjects() const;
 	inline FName GetActorPackage() const { return ActorPackage; }
 	inline FName GetActorPath() const { return ActorPath; }
 	inline FName GetActorLabel() const { return ActorLabel; }
