LevelInstance: Pivot fixes

diff --git a/Engine/Source/Runtime/Engine/Private/LevelInstance/LevelInstanceEditorLevelStreaming.cpp b/Engine/Source/Runtime/Engine/Private/LevelInstance/LevelInstanceEditorLevelStreaming.cpp
--- a/Engine/Source/Runtime/Engine/Private/LevelInstance/LevelInstanceEditorLevelStreaming.cpp
+++ b/Engine/Source/Runtime/Engine/Private/LevelInstance/LevelInstanceEditorLevelStreaming.cpp
@@ -88,7 +88,7 @@ void ULevelStreamingLevelInstanceEditor::SetLoadedLevel(ULevel* Level)
 		check(!NewLoadedLevel->bAlreadyMovedActors);
 		if (AWorldSettings* WorldSettings = NewLoadedLevel->GetWorldSettings())
 		{
-			LevelTransform.AddToTranslation(WorldSettings->LevelInstancePivotOffset);
+			LevelTransform = FTransform(WorldSettings->LevelInstancePivotOffset) * LevelTransform;
 		}
 	}
 }
diff --git a/Engine/Source/Runtime/Engine/Private/LevelInstance/LevelInstanceEditorPivotActor.cpp b/Engine/Source/Runtime/Engine/Private/LevelInstance/LevelInstanceEditorPivotActor.cpp
--- a/Engine/Source/Runtime/Engine/Private/LevelInstance/LevelInstanceEditorPivotActor.cpp
+++ b/Engine/Source/Runtime/Engine/Private/LevelInstance/LevelInstanceEditorPivotActor.cpp
@@ -35,8 +35,8 @@ ALevelInstancePivot* ALevelInstancePivot::Create(ALevelInstance* LevelInstanceAc
 
 	
 	// Keep Spawn World Transform in case Level Instance transform changes.
-	PivotActor->SpawnOffset = PivotActor->GetActorLocation();
-	PivotActor->OriginalPivotOffset = LevelInstanceActor->GetActorLocation() - WorldSettings->LevelInstancePivotOffset;
+	PivotActor->SpawnTransform = PivotActor->GetActorTransform();
+	PivotActor->OriginalPivotOffset = WorldSettings->LevelInstancePivotOffset;
 	PivotActor->SetLevelInstanceID(LevelInstanceActor->GetLevelInstanceID());
 
 	return PivotActor;
@@ -108,18 +108,17 @@ void ALevelInstancePivot::SetPivot(ELevelInstancePivotType PivotType, AActor* Pi
 
 void ALevelInstancePivot::UpdateOffset()
 {
-	AWorldSettings* WorldSettings = GetLevel()->GetWorldSettings();
-
-	// Offset Change is the relative transform to the spawn transform
-	FVector LocalToPivot = GetActorLocation() - SpawnOffset;
+	// Offset Change is the relative transform of the pivot actor to its spawn transform (we only care about relative translation as we don't support rotation on pivot)
+	FVector LocalToPivot = GetActorTransform().GetRelativeTransform(SpawnTransform).GetTranslation();
 
-	// We then apply that offset to the original pivot world transform
-	FVector NewPivotOffset = LocalToPivot + OriginalPivotOffset;
+	// We then apply that offset to the original pivot
+	FVector NewPivotOffset = OriginalPivotOffset - LocalToPivot;
 
+	AWorldSettings* WorldSettings = GetLevel()->GetWorldSettings();
 	if (!NewPivotOffset.Equals(WorldSettings->LevelInstancePivotOffset))
 	{
 		WorldSettings->Modify();
-		WorldSettings->LevelInstancePivotOffset = SpawnOffset - NewPivotOffset;
+		WorldSettings->LevelInstancePivotOffset = NewPivotOffset;
 	}
 }
 
diff --git a/Engine/Source/Runtime/Engine/Private/LevelInstance/LevelInstanceLevelStreaming.cpp b/Engine/Source/Runtime/Engine/Private/LevelInstance/LevelInstanceLevelStreaming.cpp
--- a/Engine/Source/Runtime/Engine/Private/LevelInstance/LevelInstanceLevelStreaming.cpp
+++ b/Engine/Source/Runtime/Engine/Private/LevelInstance/LevelInstanceLevelStreaming.cpp
@@ -162,7 +162,7 @@ void ULevelStreamingLevelInstance::SetLoadedLevel(ULevel* Level)
 		{
 			AWorldSettings* WorldSettings = NewLoadedLevel->GetWorldSettings();
 			check(WorldSettings);
-			LevelTransform.AddToTranslation(WorldSettings->LevelInstancePivotOffset);
+			LevelTransform = FTransform(WorldSettings->LevelInstancePivotOffset) * LevelTransform;
 		}
 	}
-}
\ No newline at end of file
+}
diff --git a/Engine/Source/Runtime/Engine/Private/LevelInstance/Packed/LevelInstanceISMPacker.cpp b/Engine/Source/Runtime/Engine/Private/LevelInstance/Packed/LevelInstanceISMPacker.cpp
--- a/Engine/Source/Runtime/Engine/Private/LevelInstance/Packed/LevelInstanceISMPacker.cpp
+++ b/Engine/Source/Runtime/Engine/Private/LevelInstance/Packed/LevelInstanceISMPacker.cpp
@@ -41,10 +41,7 @@ void FLevelInstanceISMPacker::PackActors(FPackedLevelInstanceBuilderContext& InC
 	UInstancedStaticMeshComponent* PackComponent = InPackingActor->AddPackedComponent<UInstancedStaticMeshComponent>();
 	
 	FTransform ActorTransform = InPackingActor->GetActorTransform();
-	FTransform CurrentPivotOffsetInverse = ActorTransform.GetRelativeTransform(InContext.GetLevelTransform());
-
-	PackComponent->SetComponentToWorld(ActorTransform);
-	PackComponent->AttachToComponent(InPackingActor->GetRootComponent(), FAttachmentTransformRules::KeepWorldTransform);
+	PackComponent->AttachToComponent(InPackingActor->GetRootComponent(), FAttachmentTransformRules::SnapToTargetIncludingScale);
 	
 	FLevelInstanceISMPackerCluster* ISMCluster = (FLevelInstanceISMPackerCluster*)InClusterID.GetData();
 	check(ISMCluster);
@@ -63,21 +60,20 @@ void FLevelInstanceISMPacker::PackActors(FPackedLevelInstanceBuilderContext& InC
 
 				if (ensure(ISMComponent->GetInstanceTransform(InstanceIndex, InstanceTransform, /*bWorldSpace=*/ true)))
 				{
-					InstanceTransforms.Add(InstanceTransform);
+					// WorldSpace -> ActorSpace -> Apply pivot change
+					InstanceTransforms.Add(InstanceTransform.GetRelativeTransform(ActorTransform) * InContext.GetRelativePivotTransform());
 				}
 			}
 		}
 		else // other subclasses are processed like regular UStaticMeshComponent
 		{
 			UStaticMeshComponent* StaticMeshComponent = CastChecked<UStaticMeshComponent>(Component);
-			InstanceTransforms.Add(StaticMeshComponent->GetComponentTransform());
+
+			// WorldSpace -> ActorSpace -> Apply pivot change
+			InstanceTransforms.Add(StaticMeshComponent->GetComponentTransform().GetRelativeTransform(ActorTransform) * InContext.GetRelativePivotTransform());
 		}
 	}
-	PackComponent->AddInstances(InstanceTransforms, /*bShouldReturnIndices*/false, /*bWorldSpace*/true);
-
-	FTransform NewWorldTransform = ActorTransform * CurrentPivotOffsetInverse * FTransform(InContext.GetPivotOffset());
-
-	PackComponent->SetWorldTransform(NewWorldTransform);
+	PackComponent->AddInstances(InstanceTransforms, /*bShouldReturnIndices*/false, /*bWorldSpace*/false);
 	PackComponent->RegisterComponent();
 }
 
diff --git a/Engine/Source/Runtime/Engine/Private/LevelInstance/Packed/PackedLevelInstanceBuilder.cpp b/Engine/Source/Runtime/Engine/Private/LevelInstance/Packed/PackedLevelInstanceBuilder.cpp
--- a/Engine/Source/Runtime/Engine/Private/LevelInstance/Packed/PackedLevelInstanceBuilder.cpp
+++ b/Engine/Source/Runtime/Engine/Private/LevelInstance/Packed/PackedLevelInstanceBuilder.cpp
@@ -186,8 +186,12 @@ bool FPackedLevelInstanceBuilder::PackActor(APackedLevelInstance* InPackedLevelI
 	AWorldSettings* WorldSettings = SourceLevel->GetWorldSettings();
 	Context.DiscardActor(WorldSettings);
 
-	Context.SetLevelTransform(SourceLevelStreaming->LevelTransform);
-	Context.SetPivotOffset(WorldSettings->LevelInstancePivotOffset);
+	// Build relative transform without rotation because pivots don't support rotation
+	FTransform CurrentPivotTransform(SourceLevelStreaming->LevelTransform.GetRelativeTransform(InPackedLevelInstance->GetActorTransform()).GetTranslation());
+	FTransform NewPivotTransform(WorldSettings->LevelInstancePivotOffset);
+	FTransform RelativePivotTransform(NewPivotTransform.GetRelativeTransform(CurrentPivotTransform));
+		
+	Context.SetRelativePivotTransform(RelativePivotTransform);
 	
 	if (AActor* DefaultBrush = SourceLevel->GetDefaultBrush())
 	{
diff --git a/Engine/Source/Runtime/Engine/Public/LevelInstance/LevelInstanceEditorPivotActor.h b/Engine/Source/Runtime/Engine/Public/LevelInstance/LevelInstanceEditorPivotActor.h
--- a/Engine/Source/Runtime/Engine/Public/LevelInstance/LevelInstanceEditorPivotActor.h
+++ b/Engine/Source/Runtime/Engine/Public/LevelInstance/LevelInstanceEditorPivotActor.h
@@ -38,7 +38,7 @@ private:
 private:
 #if WITH_EDITORONLY_DATA
 	FLevelInstanceID LevelInstanceID;
-	FVector SpawnOffset;
+	FTransform SpawnTransform;
 	FVector OriginalPivotOffset;
 #endif
 };
\ No newline at end of file
diff --git a/Engine/Source/Runtime/Engine/Public/LevelInstance/Packed/PackedLevelInstanceBuilder.h b/Engine/Source/Runtime/Engine/Public/LevelInstance/Packed/PackedLevelInstanceBuilder.h
--- a/Engine/Source/Runtime/Engine/Public/LevelInstance/Packed/PackedLevelInstanceBuilder.h
+++ b/Engine/Source/Runtime/Engine/Public/LevelInstance/Packed/PackedLevelInstanceBuilder.h
@@ -73,7 +73,7 @@ class FPackedLevelInstanceBuilderContext
 {
 public:
 	FPackedLevelInstanceBuilderContext(const FPackedLevelInstanceBuilder& InBuilder, APackedLevelInstance* InPackedLevelInstance) 
-		: Packers(InBuilder.Packers), ClassDiscards(InBuilder.ClassDiscards), PackedLevelInstance(InPackedLevelInstance), LevelTransform(FVector::ZeroVector), PivotOffset(FVector::ZeroVector) {}
+		: Packers(InBuilder.Packers), ClassDiscards(InBuilder.ClassDiscards), PackedLevelInstance(InPackedLevelInstance), RelativePivotTransform(FTransform::Identity) {}
 
 	/* Interface for ILevelInstancePacker's to use */
 	void ClusterLevelActor(AActor* InLevelActor);
@@ -82,11 +82,9 @@ public:
 	void Report(FMessageLog& LevelInstanceLog) const;
 
 	const TMap<FLevelInstancePackerClusterID, TArray<UActorComponent*>>& GetClusters() const { return Clusters; }
-
-	void SetLevelTransform(const FTransform& InLevelTransform) { LevelTransform = InLevelTransform; }
-	void SetPivotOffset(const FVector& InPivotOffset) { PivotOffset = InPivotOffset; }
-	const FTransform& GetLevelTransform() const { return LevelTransform; }
-	const FVector GetPivotOffset() const { return PivotOffset; }
+	
+	void SetRelativePivotTransform(const FTransform& InRelativePivotTransform) { RelativePivotTransform = InRelativePivotTransform; }
+	const FTransform GetRelativePivotTransform() const { return RelativePivotTransform; }
 
 	bool ShouldPackComponent(UActorComponent* InActorComponent) const;
 private:
@@ -100,8 +98,7 @@ private:
 	TMap<AActor*, TSet<UActorComponent*>> PerActorClusteredComponents;
 	TSet<AActor*> ActorDiscards;
 
-	FTransform LevelTransform;
-	FVector PivotOffset;
+	FTransform RelativePivotTransform;
 };
 
 #endif
\ No newline at end of file
