Marked ENetMode as UENUM

diff --git a/Engine/Source/Runtime/Engine/Classes/Engine/EngineBaseTypes.h b/Engine/Source/Runtime/Engine/Classes/Engine/EngineBaseTypes.h
--- a/Engine/Source/Runtime/Engine/Classes/Engine/EngineBaseTypes.h
+++ b/Engine/Source/Runtime/Engine/Classes/Engine/EngineBaseTypes.h
@@ -885,6 +885,7 @@ struct ENGINE_API FURL
  * The network mode the game is currently running.
  * @see https://docs.unrealengine.com/latest/INT/Gameplay/Networking/Overview/
  */
+UENUM()
 enum ENetMode
 {
 	/** Standalone: a game without networking, with one or more local players. Still considered a server because it has all server functionality. */
