FInstancedStaticMeshRandomSeed intialization fixes

diff --git a/Engine/Source/Runtime/Engine/Classes/Components/InstancedStaticMeshComponent.h b/Engine/Source/Runtime/Engine/Classes/Components/InstancedStaticMeshComponent.h
--- a/Engine/Source/Runtime/Engine/Classes/Components/InstancedStaticMeshComponent.h
+++ b/Engine/Source/Runtime/Engine/Classes/Components/InstancedStaticMeshComponent.h
@@ -125,10 +125,10 @@ struct FInstancedStaticMeshRandomSeed
 	GENERATED_USTRUCT_BODY()
 
 	UPROPERTY()
-	int32 StartInstanceIndex;
+	int32 StartInstanceIndex = 0;
 
 	UPROPERTY()
-	int32 RandomSeed;
+	int32 RandomSeed = 0;
 };
 
 /** A component that efficiently renders multiple instances of the same StaticMesh. */
