New compression scheme for morph target compression.
-Instead of storing deltas individually, they are now grouped into batches of 64 that are encoded using the minimum number of bits required for the ranges spanned by the individual vertex components of the batch.
-Tangents are now only stored when tangents are not being recomputed at runtime.
-Quality is artist controlable through the per-LOD Morph Target Position Error Tolerance setting.
-Initial experiments suggest ~10x compression is possible with acceptable quality
-Fixed performance issue where the old code would launch threads for morphs with zero weight!
-Overall morph blending performance improved by ~4x
-Added SkeletalMeshReport from Josie

diff --git a/Engine/Shaders/Private/MorphTargets.usf b/Engine/Shaders/Private/MorphTargets.usf
--- a/Engine/Shaders/Private/MorphTargets.usf
+++ b/Engine/Shaders/Private/MorphTargets.usf
@@ -6,37 +6,50 @@
 
 
 #include "Common.ush"
+#include "BitPacking.ush"
+#include "ComputeShaderUtils.ush"
 
 ////////////////////////////////////////////////////////////////
 
-#define GMorphTargetDispatchBatchSize 128
+#define GMorphTargetDispatchBatchSize	128
+#define BATCH_SIZE						64
+#define NUM_BATCH_HEADER_DWORDS			10
+#define INDEX_MAX_BITS					31
+#define POSITION_MAX_BITS				28
+#define TANGENTZ_MAX_BITS				16
 
 RWBuffer<uint> MorphVertexBuffer;
 RWBuffer<int> MorphNormalizationBuffer;
+ByteAddressBuffer MorphDataBuffer;
 
-uint4 ThreadOffsets[GMorphTargetDispatchBatchSize / 4];
-float4 MorphTargetWeight[GMorphTargetDispatchBatchSize / 4];
+//TODO: use a buffer instead of packing into constant buffer?
+uint4 MorphTargetBatchOffsets[GMorphTargetDispatchBatchSize / 4];
+uint4 MorphTargetGroupOffsets[GMorphTargetDispatchBatchSize / 4];
+float4 MorphTargetWeights[GMorphTargetDispatchBatchSize / 4];
 
 float4 PositionScale;
 float WeightScale;
-uint GlobalDispatchOffset;
+float2 Precision;
 
-Buffer<uint> VertexIndicies;
-Buffer<float> MorphDeltas;
+groupshared uint LocalGroupOffsets[GMorphTargetDispatchBatchSize];
 
-groupshared float LocalWeights[GMorphTargetDispatchBatchSize];
-groupshared uint LocalThreadOffsets[GMorphTargetDispatchBatchSize];
-
-#if 0
-// This is how MorphVertexBuffer is interpreted as (6 floats)
-struct FMorphGPUSkinVertex
+void LoadLocalData(uint GroupIndex)
 {
-	FFloat16 PosDelta[3];
-	FFloat16 TangentDelta[3];
-};
-#endif
+	if (GroupIndex < GMorphTargetDispatchBatchSize / 4)
+	{
+		uint GroupSharedIndex = GroupIndex * 4;
+		uint4 Offset = MorphTargetGroupOffsets[GroupIndex];
 
-float FindWeight(in uint DispatchIndex, out uint2 Offsets)
+		for (int i = 0; i < 4; ++i)
+		{
+			LocalGroupOffsets[GroupSharedIndex + i] = Offset[i];
+		}
+	}
+
+	GroupMemoryBarrierWithGroupSync();
+}
+
+uint GroupIndexToMorphTargetLocalIndex(uint GroupThreadIndex)
 {
 	uint Idx = 0;
 	uint Width = GMorphTargetDispatchBatchSize >> 1;
@@ -44,84 +57,127 @@ float FindWeight(in uint DispatchIndex, out uint2 Offsets)
 	UNROLL
 	while (Width > 0)
 	{
-		Idx += (DispatchIndex < LocalThreadOffsets[Idx + Width]) ? 0 : Width;
+		Idx += (GroupThreadIndex >= LocalGroupOffsets[Idx + Width]) ? Width : 0;
 		Width = Width >> 1;
 	}
 
-	Offsets = GlobalDispatchOffset + float2(0, LocalThreadOffsets[Idx + 1]);
-	return LocalWeights[Idx];
+	return Idx;
 }
 
-void LoadLocalData(uint GroupIndex)
+struct FBatchHeader
 {
-	if (GroupIndex < GMorphTargetDispatchBatchSize / 4)
-	{
-		uint GroupSharedIndex = GroupIndex * 4;
-		float4 Weight = MorphTargetWeight[GroupIndex];
-		uint4 Offset = ThreadOffsets[GroupIndex];
+	uint	DataOffset;
+	uint	NumElements;
+	bool	bTangents;
 
-		for (int i = 0; i < 4; ++i)
-		{
-			LocalWeights[GroupSharedIndex + i] = Weight[i];
-			LocalThreadOffsets[GroupSharedIndex + i] = Offset[i];
-		}
-	}
+	uint	IndexBits;
+	uint3	PositionBits;
+	uint3	TangentZBits;
 
-	GroupMemoryBarrierWithGroupSync();
-}
+	uint	IndexMin;
+	int3	PositionMin;
+	int3	TangentZMin;
+};
 
-[numthreads(2, 32, 1)]
-void GPUMorphUpdateCS(uint2 DispatchId : SV_DispatchThreadID, uint2 GroupId : SV_GroupThreadID, uint GroupIndex : SV_GroupIndex)
+FBatchHeader GetBatchHeader(uint BatchIndex)
 {
-	LoadLocalData(GroupIndex);
-
-	uint2 Offsets;
-	float Weight = FindWeight(DispatchId.y, Offsets);
+	const uint BaseAddress = BatchIndex * (NUM_BATCH_HEADER_DWORDS * 4);
+	uint4 Data0 = MorphDataBuffer.Load4(BaseAddress);
+	uint2 Data1 = MorphDataBuffer.Load2(BaseAddress + 16);
+
+	FBatchHeader Header		= (FBatchHeader)0;
+	Header.DataOffset		= Data0.x;
+	Header.IndexBits		= BitFieldExtractU32(Data0.y, 5, 0);
+	Header.PositionBits.x	= BitFieldExtractU32(Data0.y, 5, 5);
+	Header.PositionBits.y	= BitFieldExtractU32(Data0.y, 5, 10);
+	Header.PositionBits.z	= BitFieldExtractU32(Data0.y, 5, 15);
+	Header.bTangents		= BitFieldExtractU32(Data0.y, 1, 20);
+	Header.NumElements		= Data0.y >> 21;
+	
+	Header.IndexMin			= Data0.z;
+	Header.PositionMin		= int3(Data0.w, Data1);
+
+	if (Header.bTangents)
+	{
+		// TODO: Don't even store this. Need to handle two different batch header sizes.
+		uint4 Data2				= MorphDataBuffer.Load4(BaseAddress + 24);
+		Header.TangentZBits.x	= BitFieldExtractU32(Data2.x, 5, 0);
+		Header.TangentZBits.y	= BitFieldExtractU32(Data2.x, 5, 5);
+		Header.TangentZBits.z	= BitFieldExtractU32(Data2.x, 5, 10);
+		Header.TangentZMin		= int3(Data2.yzw);
+	}
+	
+	return Header;
+}
 
-	const uint WorkItemIndex = DispatchId.y + Offsets.x;
-	if (WorkItemIndex < Offsets.y && (Weight != 0.0f))
-	{		
-		for (uint i = 0; i < 3; i++)
-		{
-			const uint ElementIndex = 3 * GroupId.x + i;
-			const uint MorphVertexBufferStride = 3u + 3u;
-			float MorphDelta = MorphDeltas[MorphVertexBufferStride * WorkItemIndex + ElementIndex];
-			float MorphOffset = Weight * MorphDelta;
+[numthreads(BATCH_SIZE, 1, 1)]
+void GPUMorphUpdateCS(uint3 GroupId : SV_GroupID, uint GroupThreadIndex : SV_GroupIndex)
+{
+	LoadLocalData(GroupThreadIndex);
 
-			float Scale = (GroupId.x == 0) ? PositionScale[i] : PositionScale[3];
-			int QuantizedOffset = round(MorphOffset * Scale);
+	const uint GroupIndex = GetUnWrappedDispatchGroupId(GroupId);
+	const uint MorphTargetLocalIndex = GroupIndexToMorphTargetLocalIndex(GroupIndex);
 
-			const uint DestVertexIndex = MorphVertexBufferStride * VertexIndicies[WorkItemIndex] + ElementIndex;
-			InterlockedAdd(MorphVertexBuffer[DestVertexIndex], asuint(QuantizedOffset));
-		}
+	const uint BatchStartIndex = MorphTargetBatchOffsets[MorphTargetLocalIndex >> 2][MorphTargetLocalIndex & 3u];			//TODO: use a buffer instead of packing into constant buffer?
+	const uint GroupOffset = GroupIndex - LocalGroupOffsets[MorphTargetLocalIndex];
+	const FBatchHeader BatchHeader = GetBatchHeader(BatchStartIndex + GroupIndex - LocalGroupOffsets[MorphTargetLocalIndex]);
 
-		if (GroupId.x == 0)
+	if (GroupThreadIndex < BatchHeader.NumElements)
+	{
+		const float Weight			= MorphTargetWeights[MorphTargetLocalIndex >> 2][MorphTargetLocalIndex & 3u];			//TODO: use a buffer instead of packing into constant buffer?
+		const uint NumPositionBits	= BatchHeader.PositionBits.x + BatchHeader.PositionBits.y + BatchHeader.PositionBits.z;
+		const uint NumTangentBits	= BatchHeader.TangentZBits.x + BatchHeader.TangentZBits.y + BatchHeader.TangentZBits.z;
+		const uint StrideInBits		= BatchHeader.IndexBits + NumPositionBits + (BatchHeader.bTangents ? NumTangentBits : 0);
+
+		const uint CompileTimeMaxBits = INDEX_MAX_BITS + 3 * POSITION_MAX_BITS + 3 * TANGENTZ_MAX_BITS;
+		FBitStreamReaderStateRO InputStream = BitStreamReader_Create_Aligned_RO(MorphDataBuffer, BatchHeader.DataOffset, GroupThreadIndex * StrideInBits, CompileTimeMaxBits);
+		const uint DestVertexIndex = BitStreamReader_Read(InputStream, BatchHeader.IndexBits, INDEX_MAX_BITS) + BatchHeader.IndexMin + GroupThreadIndex;
+		
+		const int3 LocalPosition = BitStreamReader_Read3(InputStream, BatchHeader.PositionBits, POSITION_MAX_BITS);
+		const float3 PositionDelta = float3(LocalPosition + BatchHeader.PositionMin) * Precision.x;
+		const int3 QuantizedPositionOffset = (int3)round(PositionDelta * Weight * PositionScale.xyz);
+
+		const uint DestVertexAddress = DestVertexIndex * 6;
+		InterlockedAdd(MorphVertexBuffer[DestVertexAddress + 0], asuint(QuantizedPositionOffset.x));
+		InterlockedAdd(MorphVertexBuffer[DestVertexAddress + 1], asuint(QuantizedPositionOffset.y));
+		InterlockedAdd(MorphVertexBuffer[DestVertexAddress + 2], asuint(QuantizedPositionOffset.z));
+
+		if (BatchHeader.bTangents)
 		{
-			int QuantizedWeight = round(Weight * WeightScale);
-			InterlockedAdd(MorphNormalizationBuffer[WorkItemIndex], QuantizedWeight);
+			const int3 LocalTangent = BitStreamReader_Read3(InputStream, BatchHeader.TangentZBits, TANGENTZ_MAX_BITS);
+			const float3 TangentZDelta = float3(LocalTangent + BatchHeader.TangentZMin) * Precision.y;
+			const int3 QuantizedTangentZOffset = (int3)round(TangentZDelta * Weight * PositionScale.w);
+			InterlockedAdd(MorphVertexBuffer[DestVertexAddress + 3], asuint(QuantizedTangentZOffset.x));
+			InterlockedAdd(MorphVertexBuffer[DestVertexAddress + 4], asuint(QuantizedTangentZOffset.y));
+			InterlockedAdd(MorphVertexBuffer[DestVertexAddress + 5], asuint(QuantizedTangentZOffset.z));
 		}
+
+		// Update weight sum
+		int QuantizedWeight = (int)round(Weight * WeightScale);
+		InterlockedAdd(MorphNormalizationBuffer[DestVertexIndex], QuantizedWeight);	//TODO: Only used to normalize tangent weights. Are we sure we even want to do this?
 	}
 }
 
-[numthreads(2, 32, 1)]
-void GPUMorphNormalizeCS(uint2 DispatchId : SV_DispatchThreadID, uint2 GroupId : SV_GroupThreadID, uint GroupIndex : SV_GroupIndex)
+[numthreads(64, 1, 1)]
+void GPUMorphNormalizeCS(uint3 GroupId : SV_GroupID, uint GroupThreadIndex : SV_GroupIndex)
 {
-	const uint WorkItemIndex = DispatchId.y;
-	for (uint i = 0; i < 3; i++)
+	const uint DestVertexIndex = GetUnWrappedDispatchThreadId(GroupId, GroupThreadIndex, 64);
+	const uint DestVertexAddress = DestVertexIndex * 6;
+
+	const float3 Position = float3(int3(MorphVertexBuffer[DestVertexAddress + 0], MorphVertexBuffer[DestVertexAddress + 1], MorphVertexBuffer[DestVertexAddress + 2])) * PositionScale.xyz;
+
+	MorphVertexBuffer[DestVertexAddress + 0] = asuint(Position.x);
+	MorphVertexBuffer[DestVertexAddress + 1] = asuint(Position.y);
+	MorphVertexBuffer[DestVertexAddress + 2] = asuint(Position.z);
+	
+	// TODO: Does this tangent normalization code make sense?
+	const float Weight = float(MorphNormalizationBuffer[DestVertexIndex]) * WeightScale;
+	float3 TangentZ = float3(int3(MorphVertexBuffer[DestVertexAddress + 3], MorphVertexBuffer[DestVertexAddress + 4], MorphVertexBuffer[DestVertexAddress + 5])) * PositionScale.w;
+	if (Weight > 1.0f)
 	{
-		const uint ElementIndex = 3 * GroupId.x + i;
-		const uint MorphVertexBufferStride = 3u + 3u;
-		const uint DestVertexIndex = MorphVertexBufferStride * WorkItemIndex + ElementIndex;
-
-		float Scale = (GroupId.x == 0) ? PositionScale[i] : PositionScale[3];
-		float NomalizedResult = float(asint(MorphVertexBuffer[DestVertexIndex])) * Scale;
-		//normalize directions by 1 / accumulated absolute weights
-
-		float Weight = float(MorphNormalizationBuffer[WorkItemIndex]) * WeightScale;
-		if(GroupId.x && Weight > 1.0f)
-		{
-			NomalizedResult /= Weight;
-		}
-		MorphVertexBuffer[DestVertexIndex] = asuint(NomalizedResult);
+		TangentZ /= Weight;
 	}
+	MorphVertexBuffer[DestVertexAddress + 3] = asuint(TangentZ.x);
+	MorphVertexBuffer[DestVertexAddress + 4] = asuint(TangentZ.y);
+	MorphVertexBuffer[DestVertexAddress + 5] = asuint(TangentZ.z);
 }
diff --git a/Engine/Source/Editor/Persona/Private/PersonaMeshDetails.cpp b/Engine/Source/Editor/Persona/Private/PersonaMeshDetails.cpp
--- a/Engine/Source/Editor/Persona/Private/PersonaMeshDetails.cpp
+++ b/Engine/Source/Editor/Persona/Private/PersonaMeshDetails.cpp
@@ -2436,7 +2436,8 @@ void FPersonaMeshDetails::CustomizeLODInfoSetingsDetails(IDetailLayoutBuilder& D
 			LODInfoGroup.AddPropertyRow(LODInfoChildHandle);
 
 			//Some property has to call post edit change, since they are not part of the DDC key, but they are use by directly by the skinned mesh component.
-			bool bShouldPreventPostEditChange = GET_MEMBER_NAME_CHECKED(FSkeletalMeshLODInfo, SkinCacheUsage) != LODInfoChildHandle->GetProperty()->GetFName();
+			bool bShouldPreventPostEditChange = GET_MEMBER_NAME_CHECKED(FSkeletalMeshLODInfo, SkinCacheUsage) != LODInfoChildHandle->GetProperty()->GetFName() &&
+												GET_MEMBER_NAME_CHECKED(FSkeletalMeshLODInfo, MorphTargetPositionErrorTolerance) != LODInfoChildHandle->GetProperty()->GetFName();
 			if(bShouldPreventPostEditChange)
 			{
 				PreventAttributePostEditChange(LODInfoChildHandle, LODIndex, LODInfoChildHandle->GetProperty()->GetFName(), bForceComponentRefreshFalse);
diff --git a/Engine/Source/Editor/UnrealEd/Private/ImportUtils/SkeletalMeshImportUtils.cpp b/Engine/Source/Editor/UnrealEd/Private/ImportUtils/SkeletalMeshImportUtils.cpp
--- a/Engine/Source/Editor/UnrealEd/Private/ImportUtils/SkeletalMeshImportUtils.cpp
+++ b/Engine/Source/Editor/UnrealEd/Private/ImportUtils/SkeletalMeshImportUtils.cpp
@@ -530,6 +530,7 @@ void SkeletalMesUtilsImpl::RestoreLODInfo(const TSharedPtr<const FExistingSkelMe
 	ImportedLODInfo.BakePoseOverride = ExistingLODInfo.BakePoseOverride;
 	ImportedLODInfo.SourceImportFilename = ExistingLODInfo.SourceImportFilename;
 	ImportedLODInfo.SkinCacheUsage = ExistingLODInfo.SkinCacheUsage;
+	ImportedLODInfo.MorphTargetPositionErrorTolerance = ExistingLODInfo.MorphTargetPositionErrorTolerance;
 	ImportedLODInfo.bAllowCPUAccess = ExistingLODInfo.bAllowCPUAccess;
 	ImportedLODInfo.bSupportUniformlyDistributedSampling = ExistingLODInfo.bSupportUniformlyDistributedSampling;
 }
diff --git a/Engine/Source/Runtime/Engine/Classes/Engine/Engine.h b/Engine/Source/Runtime/Engine/Classes/Engine/Engine.h
--- a/Engine/Source/Runtime/Engine/Classes/Engine/Engine.h
+++ b/Engine/Source/Runtime/Engine/Classes/Engine/Engine.h
@@ -2069,6 +2069,7 @@ public:
 	bool HandleLogoutStatLevelsCommand( const TCHAR* Cmd, FOutputDevice& Ar, UWorld* InWorld );
 	bool HandleMemReportCommand( const TCHAR* Cmd, FOutputDevice& Ar, UWorld* InWorld );
 	bool HandleMemReportDeferredCommand( const TCHAR* Cmd, FOutputDevice& Ar, UWorld* InWorld );
+	bool HandleSkeletalMeshReportCommand(const TCHAR* Cmd, FOutputDevice& Ar, UWorld* InWorld);
 	bool HandleParticleMeshUsageCommand( const TCHAR* Cmd, FOutputDevice& Ar );
 	bool HandleDumpParticleCountsCommand( const TCHAR* Cmd, FOutputDevice& Ar );
 	bool HandleListLoadedPackagesCommand( const TCHAR* Cmd, FOutputDevice& Ar );
diff --git a/Engine/Source/Runtime/Engine/Classes/Engine/SkeletalMesh.h b/Engine/Source/Runtime/Engine/Classes/Engine/SkeletalMesh.h
--- a/Engine/Source/Runtime/Engine/Classes/Engine/SkeletalMesh.h
+++ b/Engine/Source/Runtime/Engine/Classes/Engine/SkeletalMesh.h
@@ -429,6 +429,10 @@ struct FSkeletalMeshLODInfo
 	UPROPERTY(EditAnywhere, Category = SkeletalMeshLODInfo)
 	ESkinCacheUsage SkinCacheUsage = ESkinCacheUsage::Auto;
 
+	/** The Morph target position error tolerance in microns. Larger values result in better compression and lower memory footprint, but also lower quality. */
+	UPROPERTY(EditAnywhere, Category = SkeletalMeshLODInfo, meta = (UIMin = "0.01", ClampMin = "0.01", UIMax = "10000.0", ClampMax = "10000.0"))
+	float MorphTargetPositionErrorTolerance = 20.0f;
+
 	/** Whether to disable morph targets for this LOD. */
 	UPROPERTY()
 	uint8 bHasBeenSimplified:1;
diff --git a/Engine/Source/Runtime/Engine/Private/GPUSkinCache.cpp b/Engine/Source/Runtime/Engine/Private/GPUSkinCache.cpp
--- a/Engine/Source/Runtime/Engine/Private/GPUSkinCache.cpp
+++ b/Engine/Source/Runtime/Engine/Private/GPUSkinCache.cpp
@@ -76,7 +76,7 @@ TAutoConsoleVariable<int32> CVarGPUSkinCacheRecomputeTangents(
 	ECVF_RenderThreadSafe
 );
 
-static int32 GForceRecomputeTangents = 0;
+int32 GForceRecomputeTangents = 0;
 FAutoConsoleVariableRef CVarGPUSkinCacheForceRecomputeTangents(
 	TEXT("r.SkinCache.ForceRecomputeTangents"),
 	GForceRecomputeTangents,
diff --git a/Engine/Source/Runtime/Engine/Private/Rendering/MorphTargetVertexInfoBuffers.cpp b/Engine/Source/Runtime/Engine/Private/Rendering/MorphTargetVertexInfoBuffers.cpp
--- a/Engine/Source/Runtime/Engine/Private/Rendering/MorphTargetVertexInfoBuffers.cpp
+++ b/Engine/Source/Runtime/Engine/Private/Rendering/MorphTargetVertexInfoBuffers.cpp
@@ -7,46 +7,22 @@ void FMorphTargetVertexInfoBuffers::InitRHI()
 {
 	SCOPED_LOADTIMER(FFMorphTargetVertexInfoBuffers_InitRHI);
 
-	check(NumTotalWorkItems > 0);
+	check(NumTotalBatches > 0);
 
-	{
-		FRHIResourceCreateInfo CreateInfo(TEXT("VertexIndicesVB"));
-		if (bMaxVertexIndex32Bits)
-		{
-			VertexIndicesVB = RHICreateBuffer(VertexIndices.Num() * sizeof(uint32), BUF_Static | BUF_VertexBuffer | BUF_ShaderResource, 0, ERHIAccess::VertexOrIndexBuffer | ERHIAccess::SRVMask, CreateInfo);
-			void* VertexIndicesVBData = RHILockBuffer(VertexIndicesVB, 0, VertexIndicesVB->GetSize(), RLM_WriteOnly);
-			FMemory::ParallelMemcpy(VertexIndicesVBData, VertexIndices.GetData(), VertexIndices.Num() * sizeof(uint32), EMemcpyCachePolicy::StoreUncached);
-			RHIUnlockBuffer(VertexIndicesVB);
-			VertexIndicesSRV = RHICreateShaderResourceView(VertexIndicesVB, sizeof(uint32), PF_R32_UINT);
-		}
-		else
-		{
-			VertexIndicesVB = RHICreateBuffer(VertexIndices.Num() * sizeof(uint16), BUF_Static | BUF_VertexBuffer | BUF_ShaderResource, 0, ERHIAccess::VertexOrIndexBuffer | ERHIAccess::SRVMask, CreateInfo);
-			uint16* VertexIndicesVBData = (uint16*)RHILockBuffer(VertexIndicesVB, 0, VertexIndicesVB->GetSize(), RLM_WriteOnly);
-			for (int32 i = 0; i < VertexIndices.Num(); ++i)
-			{
-				VertexIndicesVBData[i] = (uint16)VertexIndices[i];
-			}
-			RHIUnlockBuffer(VertexIndicesVB);
-			VertexIndicesSRV = RHICreateShaderResourceView(VertexIndicesVB, sizeof(uint16), PF_R16_UINT);
-		}
-	}
-	{
-		FRHIResourceCreateInfo CreateInfo(TEXT("MorphDeltasVB"));
-		MorphDeltasVB = RHICreateBuffer(MorphDeltas.Num() * MorphDeltas.GetTypeSize(), BUF_Static | BUF_VertexBuffer | BUF_ShaderResource, 0, ERHIAccess::VertexOrIndexBuffer | ERHIAccess::SRVMask, CreateInfo);
-		void* MorphDeltasVBData = RHILockBuffer(MorphDeltasVB, 0, MorphDeltasVB->GetSize(), RLM_WriteOnly);
-		FMemory::ParallelMemcpy(MorphDeltasVBData, MorphDeltas.GetData(), MorphDeltas.Num() * MorphDeltas.GetTypeSize(), EMemcpyCachePolicy::StoreUncached);
-		RHIUnlockBuffer(MorphDeltasVB);
-		MorphDeltasSRV = RHICreateShaderResourceView(MorphDeltasVB, 2, PF_R16F);
-	}
-	VertexIndices.Empty();
-	MorphDeltas.Empty();
+	const uint32 BufferSize = MorphData.Num() * sizeof(uint32);
+	FRHIResourceCreateInfo CreateInfo(TEXT("MorphData"));
+	MorphDataBuffer = RHICreateStructuredBuffer(sizeof(uint32), BufferSize, BUF_Static | BUF_ByteAddressBuffer | BUF_ShaderResource, ERHIAccess::SRVMask, CreateInfo);
+	
+	void* BufferPtr = RHILockBuffer(MorphDataBuffer, 0, BufferSize, RLM_WriteOnly);
+	FMemory::ParallelMemcpy(BufferPtr, MorphData.GetData(), BufferSize, EMemcpyCachePolicy::StoreUncached);
+	RHIUnlockBuffer(MorphDataBuffer);
+	MorphDataSRV = RHICreateShaderResourceView(MorphDataBuffer);
+
+	MorphData.Empty();
 }
 
 void FMorphTargetVertexInfoBuffers::ReleaseRHI()
 {
-	VertexIndicesVB.SafeRelease();
-	VertexIndicesSRV.SafeRelease();
-	MorphDeltasVB.SafeRelease();
-	MorphDeltasSRV.SafeRelease();
+	MorphDataBuffer.SafeRelease();
+	MorphDataSRV.SafeRelease();
 }
\ No newline at end of file
diff --git a/Engine/Source/Runtime/Engine/Private/SkeletalMeshLODRenderData.cpp b/Engine/Source/Runtime/Engine/Private/SkeletalMeshLODRenderData.cpp
--- a/Engine/Source/Runtime/Engine/Private/SkeletalMeshLODRenderData.cpp
+++ b/Engine/Source/Runtime/Engine/Private/SkeletalMeshLODRenderData.cpp
@@ -28,6 +28,9 @@ static FAutoConsoleVariableRef CVarStripSkeletalMeshLodsBelowMinLod(
 	TEXT("If set will strip skeletal mesh LODs under the minimum renderable LOD for the target platform during cooking.")
 );
 
+extern int32 GForceRecomputeTangents;
+extern int32 GSkinCacheRecomputeTangents;
+
 namespace
 {
 	struct FReverseOrderBitArraysBySetBits
@@ -126,6 +129,44 @@ FArchive& operator<<(FArchive& Ar, FSkelMeshRenderSection& S)
 	return Ar;
 }
 
+class FDwordBitWriter
+{
+public:
+	FDwordBitWriter(TArray<uint32>& Buffer) :
+		Buffer(Buffer),
+		PendingBits(0ull),
+		NumPendingBits(0)
+	{
+	}
+
+	void PutBits(uint32 Bits, uint32 NumBits)
+	{
+		check((uint64)Bits < (1ull << NumBits));
+		PendingBits |= (uint64)Bits << NumPendingBits;
+		NumPendingBits += NumBits;
+
+		while (NumPendingBits >= 32)
+		{
+			Buffer.Add((uint32)PendingBits);
+			PendingBits >>= 32;
+			NumPendingBits -= 32;
+		}
+	}
+
+	void Flush()
+	{
+		if (NumPendingBits > 0)
+			Buffer.Add((uint32)PendingBits);
+		PendingBits = 0;
+		NumPendingBits = 0;
+	}
+
+private:
+	TArray<uint32>&	Buffer;
+	uint64 			PendingBits;
+	int32 			NumPendingBits;
+};
+
 
 void FSkeletalMeshLODRenderData::InitResources(bool bNeedsVertexColors, int32 LODIndex, TArray<UMorphTarget*>& InMorphTargets, USkeletalMesh* Owner)
 {
@@ -177,22 +218,93 @@ void FSkeletalMeshLODRenderData::InitResources(bool bNeedsVertexColors, int32 LO
 	// UseGPUMorphTargets() can be toggled only on SM5 atm
 	if (IsFeatureLevelSupported(GMaxRHIShaderPlatform, ERHIFeatureLevel::SM5) && InMorphTargets.Num() > 0)
 	{
-		MorphTargetVertexInfoBuffers.VertexIndices.Empty();
-		MorphTargetVertexInfoBuffers.MorphDeltas.Empty();
-		MorphTargetVertexInfoBuffers.NumTotalWorkItems = 0;
+		// Simple Morph compression 0.1
+		// Instead of storing vertex deltas individually they are organized into batches of 64.
+		// Each batch has a header that describes how many bits are allocated to each of the vertex components.
+		// Batches also store an explicit offset to its associated data. This makes it trivial to decode batches
+		// in parallel, and because deltas are fixed-width inside a batch, deltas can also be decoded in parallel.
+		// The result is a semi-adaptive encoding that functions as a crude substitute for entropy coding, that is
+		// fast to decode on parallel hardware.
+		
+		// Quantization still happens globally to avoid issues with cracks at duplicate vertices.
+		// The quantization is artist controlled on a per LOD basis. Higher error tolerance results in smaller deltas
+		// and a smaller compressed size.
+
+		const FSkeletalMeshLODInfo* SkeletalMeshLODInfo = Owner->GetLODInfo(LODIndex);
+		
+		const float UnrealUnitPerMeter		= 100.0f;
+		const float PositionPrecision		= SkeletalMeshLODInfo->MorphTargetPositionErrorTolerance * 2.0f * 1e-6f * UnrealUnitPerMeter;	// * 2.0 because correct rounding guarantees error is at most half of the cell size.
+		const float RcpPositionPrecision	= 1.0f / PositionPrecision;
+
+		const float TangentZPrecision		= 1.0f / 2048.0f;				// Object scale irrelevant here. Let's assume ~12bits per component is plenty.
+		const float RcpTangentZPrecision	= 1.0f / TangentZPrecision;
+
+		const uint32 BatchSize				= 64;
+		const uint32 NumBatchHeaderDwords	= 10u;
+		
+		const uint32 IndexMaxBits			= 31u;
+
+		const uint32 PositionMaxBits		= 28u;				// Probably more than we need, but let's just allow it to go this high to be safe for now.
+																// For larger deltas this can even be more precision than what was in the float input data!
+																// Maybe consider float-like or exponential encoding of large values?
+		const float PositionMinValue		= -134217728.0f;	// -2^(MaxBits-1)
+		const float PositionMaxValue		=  134217720.0f;	// Largest float smaller than 2^(MaxBits-1)-1
+																// Using 134217727.0f would NOT work as it would be rounded up to 134217728.0f, which is
+																// outside the range.
+
+		const uint32 TangentZMaxBits		= 16u;
+		const float TangentZMinValue		= -32768.0f;		// -2^(MaxBits-1)
+		const float TangentZMaxValue		=  32767.0f;		//  2^(MaxBits-1)-1
+
+		struct FBatchHeader
+		{
+			uint32		DataOffset;
+			uint32		NumElements;
+			bool		bTangents;
+
+			uint32		IndexBits;
+			FIntVector	PositionBits;
+			FIntVector	TangentZBits;
+
+			uint32		IndexMin;
+			FIntVector	PositionMin;
+			FIntVector	TangentZMin;
+		};
 
-		MorphTargetVertexInfoBuffers.StartOffsetPerMorph.Empty(InMorphTargets.Num());
-		MorphTargetVertexInfoBuffers.WorkItemsPerMorph.Empty(InMorphTargets.Num());
+		// uint32 StartTime = FPlatformTime::Cycles();
+
+		MorphTargetVertexInfoBuffers.MorphData.Empty();
+		MorphTargetVertexInfoBuffers.NumTotalBatches = 0;
+		MorphTargetVertexInfoBuffers.PositionPrecision = PositionPrecision;
+		MorphTargetVertexInfoBuffers.TangentZPrecision = TangentZPrecision;
+
+		MorphTargetVertexInfoBuffers.BatchStartOffsetPerMorph.Empty(InMorphTargets.Num());
+		MorphTargetVertexInfoBuffers.BatchesPerMorph.Empty(InMorphTargets.Num());
 		MorphTargetVertexInfoBuffers.MaximumValuePerMorph.Empty(InMorphTargets.Num());
 		MorphTargetVertexInfoBuffers.MinimumValuePerMorph.Empty(InMorphTargets.Num());
-		MorphTargetVertexInfoBuffers.NumSplitsPerMorph.Empty(InMorphTargets.Num());
 
-		uint32 MaxVertexIndex = 0;
+		// Mark vertices that are in a section that doesn't recompute tangents as needing tangents
+		const int32 RecomputeTangentsMode = GForceRecomputeTangents > 0 ? 1 : GSkinCacheRecomputeTangents;
+		TBitArray<> VertexNeedsTangents;
+		VertexNeedsTangents.Init(false, StaticVertexBuffers.PositionVertexBuffer.GetNumVertices());
+		for (const FSkelMeshRenderSection& RenderSection : RenderSections)
+		{
+			const bool bRecomputeTangents = RecomputeTangentsMode > 0 && (RenderSection.bRecomputeTangent || RecomputeTangentsMode == 1);
+			if (!bRecomputeTangents)
+			{
+				for (uint32 i = 0; i < RenderSection.NumVertices; i++)
+				{
+					VertexNeedsTangents[RenderSection.BaseVertexIndex + i] = true;
+				}
+			}
+		}
+
 		// Populate the arrays to be filled in later in the render thread
+		TArray<FBatchHeader> BatchHeaders;
+		TArray<uint32> BitstreamData;
 		for (int32 AnimIdx = 0; AnimIdx < InMorphTargets.Num(); ++AnimIdx)
 		{
-			uint32 StartOffset = MorphTargetVertexInfoBuffers.NumTotalWorkItems;
-			MorphTargetVertexInfoBuffers.NumSplitsPerMorph.Add(0);
+			uint32 BatchStartOffset = MorphTargetVertexInfoBuffers.NumTotalBatches;
 
 			float MaximumValues[4] = { -FLT_MAX, -FLT_MAX, -FLT_MAX, -FLT_MAX };
 			float MinimumValues[4] = { +FLT_MAX, +FLT_MAX, +FLT_MAX, +FLT_MAX };
@@ -214,63 +326,209 @@ void FSkeletalMeshLODRenderData::InitResources(bool bNeedsVertexColors, int32 LO
 			}
 			else
 			{
+				struct FQuantizedDelta
+				{
+					FIntVector	Position;
+					FIntVector	TangentZ;
+					uint32		Index;
+				};
+				TArray<FQuantizedDelta> QuantizedDeltas;
+				QuantizedDeltas.Reserve(NumSrcDeltas);
+
+				bool bVertexIndicesSorted = true;
+
+				int32 PrevVertexIndex = -1;
 				for (int32 DeltaIndex = 0; DeltaIndex < NumSrcDeltas; DeltaIndex++)
 				{
 					const auto& MorphDelta = MorphDeltas[DeltaIndex];
+					const FVector3f TangentZDelta = VertexNeedsTangents[MorphDelta.SourceIdx] ? MorphDelta.TangentZDelta : FVector3f::ZeroVector;
+
 					// when import, we do check threshold, and also when adding weight, we do have threshold for how smaller weight can fit in
 					// so no reason to check here another threshold
 					MaximumValues[0] = FMath::Max(MaximumValues[0], MorphDelta.PositionDelta.X);
 					MaximumValues[1] = FMath::Max(MaximumValues[1], MorphDelta.PositionDelta.Y);
 					MaximumValues[2] = FMath::Max(MaximumValues[2], MorphDelta.PositionDelta.Z);
-					MaximumValues[3] = FMath::Max(MaximumValues[3], FMath::Max(MorphDelta.TangentZDelta.X, FMath::Max(MorphDelta.TangentZDelta.Y, MorphDelta.TangentZDelta.Z)));
+					MaximumValues[3] = FMath::Max(MaximumValues[3], FMath::Max(TangentZDelta.X, FMath::Max(TangentZDelta.Y, TangentZDelta.Z)));
 
 					MinimumValues[0] = FMath::Min(MinimumValues[0], MorphDelta.PositionDelta.X);
 					MinimumValues[1] = FMath::Min(MinimumValues[1], MorphDelta.PositionDelta.Y);
 					MinimumValues[2] = FMath::Min(MinimumValues[2], MorphDelta.PositionDelta.Z);
-					MinimumValues[3] = FMath::Min(MinimumValues[3], FMath::Min(MorphDelta.TangentZDelta.X, FMath::Min(MorphDelta.TangentZDelta.Y, MorphDelta.TangentZDelta.Z)));
+					MinimumValues[3] = FMath::Min(MinimumValues[3], FMath::Min(TangentZDelta.X, FMath::Min(TangentZDelta.Y, TangentZDelta.Z)));
+
+					// Check if input is sorted. It usually is, but it might not be.
+					if ((int32)MorphDelta.SourceIdx < PrevVertexIndex)
+						bVertexIndicesSorted = false;
+					PrevVertexIndex = (int32)MorphDelta.SourceIdx;
+
+					// Quantize delta
+					FQuantizedDelta QuantizedDelta;
+					const FVector3f& PositionDelta	= MorphDelta.PositionDelta;
+					QuantizedDelta.Position.X		= FMath::RoundToInt(FMath::Clamp(PositionDelta.X * RcpPositionPrecision, PositionMinValue, PositionMaxValue));
+					QuantizedDelta.Position.Y		= FMath::RoundToInt(FMath::Clamp(PositionDelta.Y * RcpPositionPrecision, PositionMinValue, PositionMaxValue));
+					QuantizedDelta.Position.Z		= FMath::RoundToInt(FMath::Clamp(PositionDelta.Z * RcpPositionPrecision, PositionMinValue, PositionMaxValue));
+					QuantizedDelta.TangentZ.X		= FMath::RoundToInt(FMath::Clamp(TangentZDelta.X * RcpTangentZPrecision, TangentZMinValue, TangentZMaxValue));
+					QuantizedDelta.TangentZ.Y		= FMath::RoundToInt(FMath::Clamp(TangentZDelta.Y * RcpTangentZPrecision, TangentZMinValue, TangentZMaxValue));
+					QuantizedDelta.TangentZ.Z		= FMath::RoundToInt(FMath::Clamp(TangentZDelta.Z * RcpTangentZPrecision, TangentZMinValue, TangentZMaxValue));
+					QuantizedDelta.Index			= MorphDelta.SourceIdx;
+
+					if(QuantizedDelta.Position != FIntVector::ZeroValue || QuantizedDelta.TangentZ != FIntVector::ZeroValue)
+					{
+						// Only add delta if it is non-zero
+						QuantizedDeltas.Add(QuantizedDelta);
+					}
+				}
 
-					MaxVertexIndex = FMath::Max(MorphDelta.SourceIdx, MaxVertexIndex);
-					MorphTargetVertexInfoBuffers.VertexIndices.Add(MorphDelta.SourceIdx);
-					MorphTargetVertexInfoBuffers.MorphDeltas.Emplace(MorphDelta.PositionDelta, MorphDelta.TangentZDelta);
-					MorphTargetVertexInfoBuffers.NumTotalWorkItems++;
+				// Sort deltas if the source wasn't already sorted
+				if (!bVertexIndicesSorted)
+				{
+					Algo::Sort(QuantizedDeltas, [](const FQuantizedDelta& A, const FQuantizedDelta& B) { return A.Index < B.Index; });
 				}
-			}
 
-			uint32 MorphTargetSize = MorphTargetVertexInfoBuffers.NumTotalWorkItems - StartOffset;
-			if (MorphTargetSize > 0)
-			{
-				ensureMsgf(MaximumValues[0] < +32752.0f && MaximumValues[1] < +32752.0f && MaximumValues[2] < +32752.0f && MaximumValues[3] < +32752.0f, TEXT("Huge MorphTarget Delta found in %s at index %i, might break down because we use half float storage"), *MorphTarget->GetName(), AnimIdx);
-				ensureMsgf(MinimumValues[0] > -32752.0f && MinimumValues[1] > -32752.0f && MinimumValues[2] > -32752.0f && MaximumValues[3] > -32752.0f, TEXT("Huge MorphTarget Delta found in %s at index %i, might break down because we use half float storage"), *MorphTarget->GetName(), AnimIdx);
+				// Encode batch deltas
+				const uint32 MorphNumBatches = (QuantizedDeltas.Num() + BatchSize - 1u) / BatchSize;
+				for (uint32 BatchIndex = 0; BatchIndex < MorphNumBatches; BatchIndex++)
+				{
+					const uint32 BatchFirstElementIndex = BatchIndex * BatchSize;
+					const uint32 NumElements = FMath::Min(BatchSize, QuantizedDeltas.Num() - BatchFirstElementIndex);
+
+					// Calculate batch min/max bounds
+					uint32 IndexMin = MAX_uint32;
+					uint32 IndexMax = MIN_uint32;
+					FIntVector PositionMin = FIntVector(MAX_int32);
+					FIntVector PositionMax = FIntVector(MIN_int32);
+					FIntVector TangentZMin = FIntVector(MAX_int32);
+					FIntVector TangentZMax = FIntVector(MIN_int32);
+
+					for (uint32 LocalElementIndex = 0; LocalElementIndex < NumElements; LocalElementIndex++)
+					{
+						const FQuantizedDelta& Delta = QuantizedDeltas[BatchFirstElementIndex + LocalElementIndex];
+
+						// Trick: Deltas are sorted by index, so the index increase by at least one per delta.
+						//		  Naively this would mean that a batch always spans at least 64 index values and
+						//		  indices would have to use at least 6 bits per index.
+						//		  If instead of storing the raw index, we store the index relative to its position in the batch,
+						//		  then the spanned range becomes 63 smaller.
+						//		  For a consecutive range this even gets us down to 0 bits per index!
+						check(Delta.Index >= LocalElementIndex);
+						const uint32 AdjustedIndex = Delta.Index - LocalElementIndex;
+						IndexMin = FMath::Min(IndexMin, AdjustedIndex);
+						IndexMax = FMath::Max(IndexMax, AdjustedIndex);
+					
+						PositionMin.X = FMath::Min(PositionMin.X, Delta.Position.X);
+						PositionMin.Y = FMath::Min(PositionMin.Y, Delta.Position.Y);
+						PositionMin.Z = FMath::Min(PositionMin.Z, Delta.Position.Z);
+
+						PositionMax.X = FMath::Max(PositionMax.X, Delta.Position.X);
+						PositionMax.Y = FMath::Max(PositionMax.Y, Delta.Position.Y);
+						PositionMax.Z = FMath::Max(PositionMax.Z, Delta.Position.Z);
+						
+						TangentZMin.X = FMath::Min(TangentZMin.X, Delta.TangentZ.X);
+						TangentZMin.Y = FMath::Min(TangentZMin.Y, Delta.TangentZ.Y);
+						TangentZMin.Z = FMath::Min(TangentZMin.Z, Delta.TangentZ.Z);
+
+						TangentZMax.X = FMath::Max(TangentZMax.X, Delta.TangentZ.X);
+						TangentZMax.Y = FMath::Max(TangentZMax.Y, Delta.TangentZ.Y);
+						TangentZMax.Z = FMath::Max(TangentZMax.Z, Delta.TangentZ.Z);
+					}
+
+					const uint32 IndexDelta			= IndexMax - IndexMin;
+					const FIntVector PositionDelta	= PositionMax - PositionMin;
+					const FIntVector TangentZDelta	= TangentZMax - TangentZMin;
+					const bool bBatchHasTangents	= TangentZMin != FIntVector::ZeroValue || TangentZMax != FIntVector::ZeroValue;
+					
+					FBatchHeader BatchHeader;
+					BatchHeader.DataOffset			= BitstreamData.Num() * sizeof(uint32);
+					BatchHeader.bTangents			= bBatchHasTangents;
+					BatchHeader.NumElements			= NumElements;
+					BatchHeader.IndexBits			= FMath::CeilLogTwo(IndexDelta + 1);
+					BatchHeader.PositionBits.X		= FMath::CeilLogTwo(uint32(PositionDelta.X) + 1);
+					BatchHeader.PositionBits.Y		= FMath::CeilLogTwo(uint32(PositionDelta.Y) + 1);
+					BatchHeader.PositionBits.Z		= FMath::CeilLogTwo(uint32(PositionDelta.Z) + 1);
+					BatchHeader.TangentZBits.X		= FMath::CeilLogTwo(uint32(TangentZDelta.X) + 1);
+					BatchHeader.TangentZBits.Y		= FMath::CeilLogTwo(uint32(TangentZDelta.Y) + 1);
+					BatchHeader.TangentZBits.Z		= FMath::CeilLogTwo(uint32(TangentZDelta.Z) + 1);
+					check(BatchHeader.IndexBits <= IndexMaxBits);
+					check(BatchHeader.PositionBits.X <= PositionMaxBits);
+					check(BatchHeader.PositionBits.Y <= PositionMaxBits);
+					check(BatchHeader.PositionBits.Z <= PositionMaxBits);
+					check(BatchHeader.TangentZBits.X <= TangentZMaxBits);
+					check(BatchHeader.TangentZBits.Y <= TangentZMaxBits);
+					check(BatchHeader.TangentZBits.Z <= TangentZMaxBits);
+					BatchHeader.IndexMin			= IndexMin;
+					BatchHeader.PositionMin			= PositionMin;
+					BatchHeader.TangentZMin			= TangentZMin;
+					
+					// Write quantized bits
+					FDwordBitWriter BitWriter(BitstreamData);
+					for (uint32 LocalElementIndex = 0; LocalElementIndex < NumElements; LocalElementIndex++)
+					{
+						const FQuantizedDelta& Delta = QuantizedDeltas[BatchFirstElementIndex + LocalElementIndex];
+						const uint32 AdjustedIndex = Delta.Index - LocalElementIndex;
+						BitWriter.PutBits(AdjustedIndex - IndexMin,					BatchHeader.IndexBits);
+						BitWriter.PutBits(uint32(Delta.Position.X - PositionMin.X), BatchHeader.PositionBits.X);
+						BitWriter.PutBits(uint32(Delta.Position.Y - PositionMin.Y), BatchHeader.PositionBits.Y);
+						BitWriter.PutBits(uint32(Delta.Position.Z - PositionMin.Z), BatchHeader.PositionBits.Z);
+						if (bBatchHasTangents)
+						{
+							BitWriter.PutBits(uint32(Delta.TangentZ.X - TangentZMin.X), BatchHeader.TangentZBits.X);
+							BitWriter.PutBits(uint32(Delta.TangentZ.Y - TangentZMin.Y), BatchHeader.TangentZBits.Y);
+							BitWriter.PutBits(uint32(Delta.TangentZ.Z - TangentZMin.Z), BatchHeader.TangentZBits.Z);
+						}
+					}
+					BitWriter.Flush();
+
+					BatchHeaders.Add(BatchHeader);
+				}
+				MorphTargetVertexInfoBuffers.NumTotalBatches += MorphNumBatches;
 			}
 
-			do 
-			{
-				MorphTargetVertexInfoBuffers.StartOffsetPerMorph.Add(StartOffset);
-				MorphTargetVertexInfoBuffers.WorkItemsPerMorph.Add((MorphTargetSize <= FMorphTargetVertexInfoBuffers::GetMaximumThreadGroupSize()) ? MorphTargetSize : FMorphTargetVertexInfoBuffers::GetMaximumThreadGroupSize());
-				MorphTargetVertexInfoBuffers.MaximumValuePerMorph.Add(FVector4(MaximumValues[0], MaximumValues[1], MaximumValues[2], MaximumValues[3]));
-				MorphTargetVertexInfoBuffers.MinimumValuePerMorph.Add(FVector4(MinimumValues[0], MinimumValues[1], MinimumValues[2], MinimumValues[3]));
-				MorphTargetVertexInfoBuffers.NumSplitsPerMorph[AnimIdx]++;
-
-				MorphTargetSize = (MorphTargetSize > FMorphTargetVertexInfoBuffers::GetMaximumThreadGroupSize()) ? MorphTargetSize - FMorphTargetVertexInfoBuffers::GetMaximumThreadGroupSize() : 0;
-				StartOffset += FMorphTargetVertexInfoBuffers::GetMaximumThreadGroupSize();
-			} while (MorphTargetSize > 0);
+			const uint32 MorphNumBatches = MorphTargetVertexInfoBuffers.NumTotalBatches - BatchStartOffset;
+			MorphTargetVertexInfoBuffers.BatchStartOffsetPerMorph.Add(BatchStartOffset);
+			MorphTargetVertexInfoBuffers.BatchesPerMorph.Add(MorphNumBatches);
+			MorphTargetVertexInfoBuffers.MaximumValuePerMorph.Add(FVector4(MaximumValues[0], MaximumValues[1], MaximumValues[2], MaximumValues[3]));
+			MorphTargetVertexInfoBuffers.MinimumValuePerMorph.Add(FVector4(MinimumValues[0], MinimumValues[1], MinimumValues[2], MinimumValues[3]));
 
 #if !WITH_EDITOR
 			if (NumSrcDeltas > 0)
 			{
-				// A CPU copy of the morph deltas has been made so it is safe to 
+				// A CPU copy of the morph deltas has beenA  made so it is safe to 
 				// discard the original data.  Keep CPU buffers when in the editor.
 				MorphTarget->DiscardVertexData();
 			}
 #endif
 		}
 
-		MorphTargetVertexInfoBuffers.bMaxVertexIndex32Bits = MaxVertexIndex > 65535;
 
-		check(MorphTargetVertexInfoBuffers.WorkItemsPerMorph.Num() == MorphTargetVertexInfoBuffers.StartOffsetPerMorph.Num());
-		check(MorphTargetVertexInfoBuffers.WorkItemsPerMorph.Num() == MorphTargetVertexInfoBuffers.MaximumValuePerMorph.Num());
-		check(MorphTargetVertexInfoBuffers.WorkItemsPerMorph.Num() == MorphTargetVertexInfoBuffers.MinimumValuePerMorph.Num());
-		if (MorphTargetVertexInfoBuffers.NumTotalWorkItems > 0)
+		// Write packed batch headers
+		for (const FBatchHeader& BatchHeader : BatchHeaders)
+		{
+			const uint32 DataOffset = BatchHeader.DataOffset + BatchHeaders.Num() * NumBatchHeaderDwords * sizeof(uint32);
+
+			MorphTargetVertexInfoBuffers.MorphData.Add(DataOffset);
+			MorphTargetVertexInfoBuffers.MorphData.Add(	BatchHeader.IndexBits | 
+														(BatchHeader.PositionBits.X << 5) | (BatchHeader.PositionBits.Y << 10) | (BatchHeader.PositionBits.Z << 15) |
+														(BatchHeader.bTangents ? (1u << 20) : 0u) | 
+														(BatchHeader.NumElements << 21));
+			MorphTargetVertexInfoBuffers.MorphData.Add(BatchHeader.IndexMin);
+			MorphTargetVertexInfoBuffers.MorphData.Add(BatchHeader.PositionMin.X);
+			MorphTargetVertexInfoBuffers.MorphData.Add(BatchHeader.PositionMin.Y);
+			MorphTargetVertexInfoBuffers.MorphData.Add(BatchHeader.PositionMin.Z);
+
+			MorphTargetVertexInfoBuffers.MorphData.Add(BatchHeader.TangentZBits.X | (BatchHeader.TangentZBits.Y << 5) | (BatchHeader.TangentZBits.Z << 10));
+			MorphTargetVertexInfoBuffers.MorphData.Add(BatchHeader.TangentZMin.X);
+			MorphTargetVertexInfoBuffers.MorphData.Add(BatchHeader.TangentZMin.Y);
+			MorphTargetVertexInfoBuffers.MorphData.Add(BatchHeader.TangentZMin.Z);
+		}
+
+		// Append bitstream data
+		MorphTargetVertexInfoBuffers.MorphData.Append(BitstreamData);
+
+		// UE_LOG(LogStaticMesh, Log, TEXT("Morph compression time: [%.2fs]"), FPlatformTime::ToMilliseconds(FPlatformTime::Cycles() - StartTime) / 1000.0f);
+
+		check(MorphTargetVertexInfoBuffers.BatchesPerMorph.Num() == MorphTargetVertexInfoBuffers.BatchStartOffsetPerMorph.Num());
+		check(MorphTargetVertexInfoBuffers.BatchesPerMorph.Num() == MorphTargetVertexInfoBuffers.MaximumValuePerMorph.Num());
+		check(MorphTargetVertexInfoBuffers.BatchesPerMorph.Num() == MorphTargetVertexInfoBuffers.MinimumValuePerMorph.Num());
+		if (MorphTargetVertexInfoBuffers.NumTotalBatches > 0)
 		{
 			BeginInitResource(&MorphTargetVertexInfoBuffers);
 		}
diff --git a/Engine/Source/Runtime/Engine/Private/SkeletalRenderGPUSkin.cpp b/Engine/Source/Runtime/Engine/Private/SkeletalRenderGPUSkin.cpp
--- a/Engine/Source/Runtime/Engine/Private/SkeletalRenderGPUSkin.cpp
+++ b/Engine/Source/Runtime/Engine/Private/SkeletalRenderGPUSkin.cpp
@@ -14,6 +14,7 @@
 #include "ShaderParameterUtils.h"
 #include "SkeletalMeshTypes.h"
 #include "HAL/LowLevelMemTracker.h"
+#include "RenderGraphUtils.h"
 
 DEFINE_LOG_CATEGORY_STATIC(LogSkeletalGPUSkinMesh, Warning, All);
 
@@ -583,18 +584,7 @@ void FSkeletalMeshObjectGPUSkin::ProcessUpdatedDynamicData(EGPUSkinCacheEntryMod
 			{
 				// sometimes this goes out of bound, we'll ensure here
 				ensureAlways(DynamicData->MorphTargetWeights.Num() == LODData.MorphTargetVertexInfoBuffers.GetNumMorphs());
-
-				// update the morph data for the lod (before SkinCache)
-				TArray<float> MorphTargetWeights;
-				MorphTargetWeights.Reserve(LODData.MorphTargetVertexInfoBuffers.GetNumMorphs());
-				for (int i = 0; i < DynamicData->MorphTargetWeights.Num(); i++)
-				{
-					for (uint32 j = 0; j < LODData.MorphTargetVertexInfoBuffers.GetNumSplitsPerMorph(i); j++)
-					{
-						MorphTargetWeights.Add(DynamicData->MorphTargetWeights[i]);
-					}
-				}
-				LOD.UpdateMorphVertexBufferGPU(RHICmdList, MorphTargetWeights, LODData.MorphTargetVertexInfoBuffers, DynamicData->SectionIdsUseByActiveMorphTargets, GetDebugName(), Mode);
+				LOD.UpdateMorphVertexBufferGPU(RHICmdList, DynamicData->MorphTargetWeights, LODData.MorphTargetVertexInfoBuffers, DynamicData->SectionIdsUseByActiveMorphTargets, GetDebugName(), Mode);
 			}
 			else
 			{
@@ -746,49 +736,25 @@ void FGPUMorphUpdateCS::SetParameters(FRHICommandList& RHICmdList, const FVector
 
 	SetShaderValue(RHICmdList, CS, PositionScaleParameter, LocalScale);
 	SetShaderValue(RHICmdList, CS, WeightScaleParameter, WeightScale);
+	FVector2D Precision = { MorphTargetVertexInfoBuffers.GetPositionPrecision(), MorphTargetVertexInfoBuffers.GetTangentZPrecision() };
+	SetShaderValue(RHICmdList, CS, PrecisionParameter, Precision);
 
-	SetSRVParameter(RHICmdList, CS, VertexIndicesParameter, MorphTargetVertexInfoBuffers.VertexIndicesSRV);
-	SetSRVParameter(RHICmdList, CS, MorphDeltasParameter, MorphTargetVertexInfoBuffers.MorphDeltasSRV);
+	SetSRVParameter(RHICmdList, CS, MorphDataBufferParameter, MorphTargetVertexInfoBuffers.MorphDataSRV);
 }
 
-static const uint32 GMorphTargetDispatchBatchSize = 128;
-
-void FGPUMorphUpdateCS::SetOffsetAndSize(FRHICommandList& RHICmdList, uint32 StartIndex, uint32 EndIndexPlusOne, const FMorphTargetVertexInfoBuffers& MorphTargetVertexInfoBuffers, const TArray<float>& MorphTargetWeights)
+void FGPUMorphUpdateCS::SetMorphOffsetsAndWeights(FRHICommandList& RHICmdList, uint32 BatchOffsets[MorphTargetDispatchBatchSize], uint32 GroupOffsets[MorphTargetDispatchBatchSize], float Weights[MorphTargetDispatchBatchSize])
 {
 	FRHIComputeShader* CS = RHICmdList.GetBoundComputeShader();
-
-	uint32 ThreadOffsets[GMorphTargetDispatchBatchSize];
-	float Weights[GMorphTargetDispatchBatchSize];
-	
-	uint32 BaseOffset = MorphTargetVertexInfoBuffers.GetStartOffset(StartIndex);
-	check(EndIndexPlusOne <= MorphTargetVertexInfoBuffers.GetNumMorphs());
-	uint32 ThreadOffset = 0u;
-	for (uint32 i = 0; i < GMorphTargetDispatchBatchSize; i++)
-	{
-		if (StartIndex + i < EndIndexPlusOne)
-		{
-			Weights[i] = MorphTargetWeights[StartIndex + i];
-			ThreadOffsets[i] = ThreadOffset;
-			ThreadOffset += MorphTargetVertexInfoBuffers.GetNumWorkItems(StartIndex + i);
-		}
-		else
-		{
-			uint32 LastStart = MorphTargetVertexInfoBuffers.GetStartOffset(EndIndexPlusOne - 1);
-			uint32 LastSize = MorphTargetVertexInfoBuffers.GetNumWorkItems(EndIndexPlusOne - 1);
-			Weights[i] = 0.0f;
-			ThreadOffsets[i] = ThreadOffset;
-		}
-	}
-
-	SetShaderValue(RHICmdList, CS, GlobalDispatchOffsetParameter, BaseOffset);
-	SetShaderValue(RHICmdList, CS, ThreadOffsetsParameter, ThreadOffsets);
-	SetShaderValue(RHICmdList, CS, MorphTargetWeightParameter, Weights);
+	SetShaderValue(RHICmdList, CS, MorphTargetBatchOffsetsParameter,*(uint32(*)[MorphTargetDispatchBatchSize]) BatchOffsets);
+	SetShaderValue(RHICmdList, CS, MorphTargetGroupOffsetsParameter,*(uint32(*)[MorphTargetDispatchBatchSize]) GroupOffsets);
+	SetShaderValue(RHICmdList, CS, MorphTargetWeightsParameter,		*(float(*)[MorphTargetDispatchBatchSize]) Weights);
 }
 
 
 void FGPUMorphUpdateCS::Dispatch(FRHICommandList& RHICmdList, uint32 Size)
 {
-	RHICmdList.DispatchComputeShader(1, (Size + 31) / 32, 1);	
+	const FIntVector DispatchSize = FComputeShaderUtils::GetGroupCountWrapped(Size);
+	RHICmdList.DispatchComputeShader(DispatchSize.X, DispatchSize.Y, DispatchSize.Z);
 }
 
 void FGPUMorphUpdateCS::EndAllDispatches(FRHICommandList& RHICmdList)
@@ -808,9 +774,10 @@ void FGPUMorphNormalizeCS::SetParameters(FRHICommandList& RHICmdList, const FVec
 	SetShaderValue(RHICmdList, CS, WeightScaleParameter, InvWeightScale);
 }
 
-void FGPUMorphNormalizeCS::Dispatch(FRHICommandList& RHICmdList, uint32 NumVerticies)
+void FGPUMorphNormalizeCS::Dispatch(FRHICommandList& RHICmdList, uint32 NumVertices)
 {
-	RHICmdList.DispatchComputeShader(1, (NumVerticies + 31) / 32, 1);
+	FIntVector DispatchSize = FComputeShaderUtils::GetGroupCountWrapped(NumVertices, 64);
+	RHICmdList.DispatchComputeShader(DispatchSize.X, DispatchSize.Y, DispatchSize.Z);
 }
 
 void FGPUMorphNormalizeCS::EndAllDispatches(FRHICommandList& RHICmdList)
@@ -892,9 +859,9 @@ void FSkeletalMeshObjectGPUSkin::FSkeletalMeshObjectLOD::UpdateMorphVertexBuffer
 		const FString RayTracingTag = (Mode == EGPUSkinCacheEntryMode::RayTracing ? TEXT("[RT]") : TEXT(""));
 		const FString LODName = OwnerName.ToString() + TEXT("_LOD") + FString::FromInt(LODIndex);
 		SCOPED_DRAW_EVENTF(RHICmdList, MorphUpdate,
-			TEXT("MorphUpdate%s_%s LodVertices=%d Threads=%d"), *RayTracingTag, *LODName,
+			TEXT("MorphUpdate%s_%s LodVertices=%d Batches=%d"), *RayTracingTag, *LODName,
 			LodData.GetNumVertices(),
-			MorphTargetVertexInfoBuffers.GetNumWorkItems());
+			MorphTargetVertexInfoBuffers.GetNumBatches());
 
 		RHICmdList.Transition(
 		{
@@ -928,35 +895,43 @@ void FSkeletalMeshObjectGPUSkin::FSkeletalMeshObjectLOD::UpdateMorphVertexBuffer
 				//multiple morph targets can be batched by a single shader where the shader will rely on
 				//binary search to find the correct target weight within the batch.
 				TShaderMapRef<FGPUMorphUpdateCS> GPUMorphUpdateCS(GetGlobalShaderMap(FeatureLevel));
-				for (uint32 i = 0; i < MorphTargetVertexInfoBuffers.GetNumMorphs();)
+
+				uint32 InputMorphStartIndex = 0;
+				while (InputMorphStartIndex < MorphTargetVertexInfoBuffers.GetNumMorphs())
 				{
-					uint32 NumMorphDeltas = 0;
-					uint32 j = 0;
-					for (; j < GMorphTargetDispatchBatchSize - 1; j++)
+					uint32 BatchOffsets[FGPUMorphUpdateCS::MorphTargetDispatchBatchSize];
+					uint32 GroupOffsets[FGPUMorphUpdateCS::MorphTargetDispatchBatchSize];
+					float Weights[FGPUMorphUpdateCS::MorphTargetDispatchBatchSize];
+
+					uint32 NumBatches = 0;
+					uint32 NumOutputMorphs = 0;
+					while (InputMorphStartIndex < MorphTargetVertexInfoBuffers.GetNumMorphs() && NumOutputMorphs < FGPUMorphUpdateCS::MorphTargetDispatchBatchSize)
 					{
-						if (i + j < MorphTargetVertexInfoBuffers.GetNumMorphs())
+						if (MorphTargetWeights[InputMorphStartIndex] != 0.0f) 	// Omit morphs with zero weight
 						{
-							if (NumMorphDeltas + MorphTargetVertexInfoBuffers.GetNumWorkItems(i + j) <= FMorphTargetVertexInfoBuffers::GetMaximumThreadGroupSize())
-							{
-								NumMorphDeltas += MorphTargetVertexInfoBuffers.GetNumWorkItems(i + j);
-								continue;
-							}
+							BatchOffsets[NumOutputMorphs] = MorphTargetVertexInfoBuffers.GetBatchStartOffset(InputMorphStartIndex);
+							GroupOffsets[NumOutputMorphs] = NumBatches;
+							Weights[NumOutputMorphs] = MorphTargetWeights[InputMorphStartIndex];
+							NumOutputMorphs++;
+
+							NumBatches += MorphTargetVertexInfoBuffers.GetNumBatches(InputMorphStartIndex);
 						}
-						break;
+						InputMorphStartIndex++;
 					}
-					check(j > 0);
 
-					if (NumMorphDeltas > 0)
+					for (uint32 i = NumOutputMorphs; i < FGPUMorphUpdateCS::MorphTargetDispatchBatchSize; i++)
 					{
-						RHICmdList.SetComputeShader(GPUMorphUpdateCS.GetComputeShader());
-						GPUMorphUpdateCS->SetParameters(RHICmdList, MorphScale, WeightScale, MorphTargetVertexInfoBuffers, MorphVertexBuffer);
-						GPUMorphUpdateCS->SetOffsetAndSize(RHICmdList, i, i + j, MorphTargetVertexInfoBuffers, MorphTargetWeights);
-						check(NumMorphDeltas <= FMorphTargetVertexInfoBuffers::GetMaximumThreadGroupSize());
-						GPUMorphUpdateCS->Dispatch(RHICmdList, NumMorphDeltas);
+						BatchOffsets[i] = 0;
+						GroupOffsets[i] = NumBatches;
+						Weights[i] = 0.0f;
 					}
-					i += j;
-				}
 
+					RHICmdList.SetComputeShader(GPUMorphUpdateCS.GetComputeShader());
+					GPUMorphUpdateCS->SetParameters(RHICmdList, MorphScale, WeightScale, MorphTargetVertexInfoBuffers, MorphVertexBuffer);
+					GPUMorphUpdateCS->SetMorphOffsetsAndWeights(RHICmdList, BatchOffsets, GroupOffsets, Weights);
+					GPUMorphUpdateCS->Dispatch(RHICmdList, NumBatches);
+				}
+				
 				GPUMorphUpdateCS->EndAllDispatches(RHICmdList);
 				RHICmdList.EndUAVOverlap({ MorphVertexBuffer.GetUAV(), MorphVertexBuffer.GetNormalizationUAV() });
 			}
diff --git a/Engine/Source/Runtime/Engine/Private/SkeletalRenderGPUSkin.h b/Engine/Source/Runtime/Engine/Private/SkeletalRenderGPUSkin.h
--- a/Engine/Source/Runtime/Engine/Private/SkeletalRenderGPUSkin.h
+++ b/Engine/Source/Runtime/Engine/Private/SkeletalRenderGPUSkin.h
@@ -698,18 +698,20 @@ public:
 		MorphVertexBufferParameter.Bind(Initializer.ParameterMap, TEXT("MorphVertexBuffer"));
 		MorphNormalizationBufferParameter.Bind(Initializer.ParameterMap, TEXT("MorphNormalizationBuffer"));
 
-		MorphTargetWeightParameter.Bind(Initializer.ParameterMap, TEXT("MorphTargetWeight"));
-		ThreadOffsetsParameter.Bind(Initializer.ParameterMap, TEXT("ThreadOffsets"));
-		GlobalDispatchOffsetParameter.Bind(Initializer.ParameterMap, TEXT("GlobalDispatchOffset"));
+		MorphTargetWeightsParameter.Bind(Initializer.ParameterMap, TEXT("MorphTargetWeights"));
+		MorphTargetBatchOffsetsParameter.Bind(Initializer.ParameterMap, TEXT("MorphTargetBatchOffsets"));
+		MorphTargetGroupOffsetsParameter.Bind(Initializer.ParameterMap, TEXT("MorphTargetGroupOffsets"));
 		PositionScaleParameter.Bind(Initializer.ParameterMap, TEXT("PositionScale"));
 		WeightScaleParameter.Bind(Initializer.ParameterMap, TEXT("WeightScale"));
+		PrecisionParameter.Bind(Initializer.ParameterMap, TEXT("Precision"));
 
-		VertexIndicesParameter.Bind(Initializer.ParameterMap, TEXT("VertexIndicies"));
-		MorphDeltasParameter.Bind(Initializer.ParameterMap, TEXT("MorphDeltas"));
+		MorphDataBufferParameter.Bind(Initializer.ParameterMap, TEXT("MorphDataBuffer"));
 	}
 
+	static const uint32 MorphTargetDispatchBatchSize = 128;
+
 	void SetParameters(FRHICommandList& RHICmdList, const FVector4& LocalScale, float WeightScale, const FMorphTargetVertexInfoBuffers& MorphTargetVertexInfoBuffers, FMorphVertexBuffer& MorphVertexBuffer);
-	void SetOffsetAndSize(FRHICommandList& RHICmdList, uint32 StartIndex, uint32 EndIndexPlusOne, const FMorphTargetVertexInfoBuffers& MorphTargetVertexInfoBuffers, const TArray<float>& MorphTargetWeights);
+	void SetMorphOffsetsAndWeights(FRHICommandList& RHICmdList, uint32 BatchOffsets[MorphTargetDispatchBatchSize], uint32 GroupOffsets[MorphTargetDispatchBatchSize], float Weights[MorphTargetDispatchBatchSize]);
 
 	void Dispatch(FRHICommandList& RHICmdList, uint32 Size);
 	void EndAllDispatches(FRHICommandList& RHICmdList);
@@ -723,15 +725,15 @@ protected:
 	LAYOUT_FIELD(FShaderResourceParameter, MorphVertexBufferParameter);
 	LAYOUT_FIELD(FShaderResourceParameter, MorphNormalizationBufferParameter);
 
-	LAYOUT_FIELD(FShaderParameter, MorphTargetWeightParameter);
+	LAYOUT_FIELD(FShaderParameter, MorphTargetWeightsParameter);
 	LAYOUT_FIELD(FShaderParameter, OffsetAndSizeParameter);
-	LAYOUT_FIELD(FShaderParameter, ThreadOffsetsParameter);
-	LAYOUT_FIELD(FShaderParameter, GlobalDispatchOffsetParameter);
+	LAYOUT_FIELD(FShaderParameter, MorphTargetBatchOffsetsParameter);
+	LAYOUT_FIELD(FShaderParameter, MorphTargetGroupOffsetsParameter);
 	LAYOUT_FIELD(FShaderParameter, PositionScaleParameter);
 	LAYOUT_FIELD(FShaderParameter, WeightScaleParameter);
+	LAYOUT_FIELD(FShaderParameter, PrecisionParameter);
 
-	LAYOUT_FIELD(FShaderResourceParameter, VertexIndicesParameter);
-	LAYOUT_FIELD(FShaderResourceParameter, MorphDeltasParameter);
+	LAYOUT_FIELD(FShaderResourceParameter, MorphDataBufferParameter);
 };
 
 class FGPUMorphNormalizeCS : public FGlobalShader
diff --git a/Engine/Source/Runtime/Engine/Private/UnrealEngine.cpp b/Engine/Source/Runtime/Engine/Private/UnrealEngine.cpp
--- a/Engine/Source/Runtime/Engine/Private/UnrealEngine.cpp
+++ b/Engine/Source/Runtime/Engine/Private/UnrealEngine.cpp
@@ -135,7 +135,7 @@ UnrealEngine.cpp: Implements the UEngine class and helpers.
 #include "HAL/PlatformApplicationMisc.h"
 #include "DynamicResolutionState.h"
 #include "Engine/ViewportStatsSubsystem.h"
-
+#include "GPUSkinCache.h"
 #include "Chaos/TriangleMeshImplicitObject.h"
 
 #include "Particles/Spawn/ParticleModuleSpawn.h"
@@ -4578,6 +4578,10 @@ bool UEngine::Exec( UWorld* InWorld, const TCHAR* Cmd, FOutputDevice& Ar )
 	{
 		return HandleMemReportDeferredCommand( Cmd, Ar, InWorld );
 	}
+	else if ( FParse::Command( &Cmd, TEXT("SkeletalMeshReport") ))
+	{
+		return HandleSkeletalMeshReportCommand(Cmd, Ar, InWorld);
+	}
 	else if( FParse::Command( &Cmd, TEXT("PARTICLEMESHUSAGE") ) )
 	{
 		return HandleParticleMeshUsageCommand( Cmd, Ar );
@@ -6823,6 +6827,97 @@ bool UEngine::HandleMemReportDeferredCommand( const TCHAR* Cmd, FOutputDevice& A
 	return true;
 }
 
+bool UEngine::HandleSkeletalMeshReportCommand(const TCHAR* Cmd, FOutputDevice& Ar, UWorld* InWorld)
+{
+	Ar.Logf(TEXT("+++ Skeletal Mesh Report +++"));
+
+	auto BytesToKB = [&](uint32 Bytes) -> uint32
+	{
+		return (Bytes + 512) / 1024;
+	};
+
+	auto GetResourceSizeKB = [&](const FRHIResource* Resource) -> uint32
+	{
+		int32 SizeInKB = 0;
+		FRHIResourceInfo ResourceInfo;
+		if (Resource && Resource->GetResourceInfo(ResourceInfo))
+		{
+			SizeInKB = BytesToKB(ResourceInfo.VRamAllocation.AllocationSize);
+		}
+		return SizeInKB;
+	};
+
+	static const IConsoleVariable* CVarDefaultGPUSkinCacheBehavior = IConsoleManager::Get().FindConsoleVariable(TEXT("r.SkinCache.DefaultBehavior"));
+	bool bSkinCacheGlobalDefault = CVarDefaultGPUSkinCacheBehavior && ESkinCacheDefaultBehavior(CVarDefaultGPUSkinCacheBehavior->GetInt()) == ESkinCacheDefaultBehavior::Inclusive;
+
+	int32 NumMeshes = 0;
+	FRHIResourceInfo ResourceInfo;
+	for (TObjectIterator<USkeletalMesh> It; It; ++It)
+	{
+		USkeletalMesh* Mesh = *It;
+		FString Name = Mesh->GetFullName();
+
+		if (Mesh->GetResourceForRendering() && Mesh->GetResourceForRendering()->LODRenderData.Num() > 0)
+		{
+			Ar.Logf(TEXT("%s"), *Name);
+
+			TIndirectArray<FSkeletalMeshLODRenderData>& LodRenderDatas = Mesh->GetResourceForRendering()->LODRenderData;
+			for (int32 LODIdx = 0; LODIdx < LodRenderDatas.Num(); LODIdx++)
+			{
+				const FSkeletalMeshLODRenderData& RenderData = LodRenderDatas[LODIdx];
+	
+				FSkeletalMeshLODInfo* LodInfo = Mesh->GetLODInfo(LODIdx);
+				check(LodInfo);
+				bool bSkinCachedLODEnabled = GEnableGPUSkinCache && (LodInfo->SkinCacheUsage == ESkinCacheUsage::Auto ? bSkinCacheGlobalDefault : LodInfo->SkinCacheUsage == ESkinCacheUsage::Enabled);
+
+				int32 MorphDataKB = GetResourceSizeKB(RenderData.MorphTargetVertexInfoBuffers.MorphDataBuffer);
+				int32 ClothVertexBufferKB = GetResourceSizeKB(RenderData.ClothVertexBuffer.VertexBufferRHI);
+
+				FString RecomputeTangentSections = TEXT("");
+				for (int32 SectionIdx = 0; SectionIdx < RenderData.RenderSections.Num(); SectionIdx++)
+				{
+					const FSkelMeshRenderSection& Section = RenderData.RenderSections[SectionIdx];
+					if (Section.bRecomputeTangent)
+					{
+						if (RecomputeTangentSections.IsEmpty())
+						{
+							RecomputeTangentSections = TEXT("[Section]") + FString::FromInt(SectionIdx);
+						}
+						else
+						{
+							RecomputeTangentSections = RecomputeTangentSections + TEXT("/") + FString::FromInt(SectionIdx);
+						}
+					}
+					
+				}
+				if (RecomputeTangentSections.IsEmpty())
+				{
+					RecomputeTangentSections = TEXT("Off");
+				}
+
+				Ar.Logf(TEXT("  [LOD%d]: RefSkeletonBones=%3d, RequiredBones=%3d, ActiveBones=%3d, MaxBoneInfluences=%2d, NumMorphTargets=%7d, MorphData=%4dKB, ClothVertexBuffer=%4dKB, SkinCache=%s, RecomputeTangent=%s, SupportRayTracing=%s, RayTracingMinLOD=%d"), LODIdx,
+					Mesh->GetRefSkeleton().GetNum(),
+					RenderData.RequiredBones.Num(),
+					RenderData.ActiveBoneIndices.Num(),
+					RenderData.SkinWeightVertexBuffer.GetMaxBoneInfluences(),
+					Mesh->GetMorphTargets().Num(),
+					MorphDataKB,
+					ClothVertexBufferKB,
+					bSkinCachedLODEnabled ? TEXT("Y") : TEXT("N"),
+					*RecomputeTangentSections,
+					Mesh->bSupportRayTracing ? TEXT("Y") : TEXT("N"),
+					Mesh->RayTracingMinLOD);
+			}
+		}
+
+		NumMeshes++;
+	}
+
+	Ar.Logf(TEXT("Total %d skeletal meshes."), NumMeshes);
+	Ar.Logf(TEXT("--- Skeletal Mesh Report ---"));
+
+	return true;
+}
 
 bool UEngine::HandleParticleMeshUsageCommand( const TCHAR* Cmd, FOutputDevice& Ar )
 {
diff --git a/Engine/Source/Runtime/Engine/Public/Rendering/MorphTargetVertexInfoBuffers.h b/Engine/Source/Runtime/Engine/Public/Rendering/MorphTargetVertexInfoBuffers.h
--- a/Engine/Source/Runtime/Engine/Public/Rendering/MorphTargetVertexInfoBuffers.h
+++ b/Engine/Source/Runtime/Engine/Public/Rendering/MorphTargetVertexInfoBuffers.h
@@ -8,7 +8,7 @@
 class FMorphTargetVertexInfoBuffers : public FRenderResource
 {
 public:
-	FMorphTargetVertexInfoBuffers() : NumTotalWorkItems(0)
+	FMorphTargetVertexInfoBuffers() : NumTotalBatches(0)
 	{
 	}
 
@@ -22,21 +22,21 @@ public:
 		return uint32(FMath::Min<uint64>(MaximumThreadGroupSize, UINT32_MAX));
 	}
 
-	uint32 GetNumWorkItems(uint32 index = UINT_MAX) const
+	uint32 GetNumBatches(uint32 index = UINT_MAX) const
 	{
-		check(index == UINT_MAX || index < (uint32)WorkItemsPerMorph.Num());
-		return index != UINT_MAX ? WorkItemsPerMorph[index] : NumTotalWorkItems;
+		check(index == UINT_MAX || index < (uint32)BatchesPerMorph.Num());
+		return index != UINT_MAX ? BatchesPerMorph[index] : NumTotalBatches;
 	}
 
 	uint32 GetNumMorphs() const
 	{
-		return WorkItemsPerMorph.Num();
+		return BatchesPerMorph.Num();
 	}
 
-	uint32 GetStartOffset(uint32 Index) const
+	uint32 GetBatchStartOffset(uint32 Index) const
 	{
-		check(Index < (uint32)StartOffsetPerMorph.Num());
-		return StartOffsetPerMorph[Index];
+		check(Index < (uint32)BatchStartOffsetPerMorph.Num());
+		return BatchStartOffsetPerMorph[Index];
 	}
 
 	const FVector4& GetMaximumMorphScale(uint32 Index) const
@@ -51,67 +51,45 @@ public:
 		return MinimumValuePerMorph[Index];
 	}
 
-	uint32 GetNumSplitsPerMorph(uint32 Index) const
+	const float GetPositionPrecision() const
 	{
-		check(Index < (uint32)NumSplitsPerMorph.Num());
-		return NumSplitsPerMorph[Index];
+		return PositionPrecision;
 	}
 
-	FBufferRHIRef VertexIndicesVB;
-	FShaderResourceViewRHIRef VertexIndicesSRV;
-
-	FBufferRHIRef MorphDeltasVB;
-	FShaderResourceViewRHIRef MorphDeltasSRV;
-
-	// Changes to this struct must be reflected in MorphTargets.usf
-	struct FMorphDelta
+	const float GetTangentZPrecision() const
 	{
-		FMorphDelta(FVector InPosDelta, FVector InTangentDelta)
-		{
-			PosDelta[0] = FFloat16(InPosDelta.X);
-			PosDelta[1] = FFloat16(InPosDelta.Y);
-			PosDelta[2] = FFloat16(InPosDelta.Z);
-
-			TangentDelta[0] = FFloat16(InTangentDelta.X);
-			TangentDelta[1] = FFloat16(InTangentDelta.Y);
-			TangentDelta[2] = FFloat16(InTangentDelta.Z);
-		}
+		return TangentZPrecision;
+	}
 
-		FFloat16 PosDelta[3];
-		FFloat16 TangentDelta[3];
-	};
+	FBufferRHIRef MorphDataBuffer;
+	FShaderResourceViewRHIRef MorphDataSRV;
 
 	void Reset()
 	{
-		VertexIndices.Empty();
-		bMaxVertexIndex32Bits = false;
-		MorphDeltas.Empty();
+		MorphData.Empty();
 		MaximumValuePerMorph.Empty();
 		MinimumValuePerMorph.Empty();
-		StartOffsetPerMorph.Empty();
-		WorkItemsPerMorph.Empty();
-		NumSplitsPerMorph.Empty();
-		TempStoreSize = 0;
-		NumTotalWorkItems = 0;
+		BatchStartOffsetPerMorph.Empty();
+		BatchesPerMorph.Empty();		
+		NumTotalBatches = 0;
+		PositionPrecision = 0.0f;
+		TangentZPrecision = 0.0f;
 	}
 
 protected:
 
-	// Transient data used while creating the vertex buffers, gets deleted as soon as VB gets initialized.
-	TArray<uint32> VertexIndices;
-	bool bMaxVertexIndex32Bits = false;
-	// Transient data used while creating the vertex buffers, gets deleted as soon as VB gets initialized.
-	TArray<FMorphDelta> MorphDeltas;
+	// Transient data. Gets deleted as soon as the GPU resource has been initialized.
+	TArray<uint32> MorphData;
 
 	//x,y,y separate for position and shared w for tangent
 	TArray<FVector4> MaximumValuePerMorph;
 	TArray<FVector4> MinimumValuePerMorph;
-	TArray<uint32> StartOffsetPerMorph;
-	TArray<uint32> WorkItemsPerMorph;
-	TArray<uint32> NumSplitsPerMorph; //splits due to too large dispatch size
-	uint32 TempStoreSize;
-
-	uint32 NumTotalWorkItems;
+	TArray<uint32> BatchStartOffsetPerMorph;
+	TArray<uint32> BatchesPerMorph;
+	
+	uint32 NumTotalBatches;
+	float PositionPrecision;
+	float TangentZPrecision;
 
 	friend class FSkeletalMeshLODRenderData;
 };
diff --git a/Engine/Source/Runtime/Engine/Public/Rendering/SkeletalMeshLODRenderData.h b/Engine/Source/Runtime/Engine/Public/Rendering/SkeletalMeshLODRenderData.h
--- a/Engine/Source/Runtime/Engine/Public/Rendering/SkeletalMeshLODRenderData.h
+++ b/Engine/Source/Runtime/Engine/Public/Rendering/SkeletalMeshLODRenderData.h
@@ -53,7 +53,7 @@ struct FSkelMeshRenderSection
 	/** The bones which are used by the vertices of this section. Indices of bones in the USkeletalMesh::RefSkeleton array */
 	TArray<FBoneIndexType> BoneMap;
 
-	/** The number of triangles in this section. */
+	/** The number of vertices in this section. */
 	uint32 NumVertices;
 
 	/** max # of bones used to skin the vertices in this section */
