Make logging of load/save of material shader maps from the DDC log to `Display`.

- Move the log of loading a shader map from the DDC into `LoadFromDerivedDataCache` to match `SaveToDerivedDataCache`
- Change the output text of `SaveToDerivedDataCache` to match the format of `LoadFromDerivedDataCache` equivalent.

diff --git a/Engine/Source/Runtime/Engine/Private/Materials/MaterialShader.cpp b/Engine/Source/Runtime/Engine/Private/Materials/MaterialShader.cpp
--- a/Engine/Source/Runtime/Engine/Private/Materials/MaterialShader.cpp
+++ b/Engine/Source/Runtime/Engine/Private/Materials/MaterialShader.cpp
@@ -1270,6 +1270,11 @@ void FMaterialShaderMap::LoadFromDerivedDataCache(const FMaterial* Material, con
 				COOK_STAT(Timer.TrackCyclesOnly());
 				InOutShaderMap = nullptr;
 			}
+
+			if (InOutShaderMap)
+			{
+				UE_LOG(LogMaterial, Display, TEXT("Loaded shaders for %s from DDC (key hash: %s)"), *Material->GetAssetName(), *OutDDCKeyDesc);
+			}
 		}
 		INC_FLOAT_STAT_BY(STAT_ShaderCompiling_DDCLoading,(float)MaterialDDCTime);
 	}
@@ -1284,7 +1289,7 @@ void FMaterialShaderMap::SaveToDerivedDataCache()
 	Serialize(Ar);
 
 	FString DataKey = GetMaterialShaderMapKeyString(ShaderMapId, GetShaderPlatform());
-	UE_LOG(LogMaterial, Verbose, TEXT("Saving material DDC for key hash %s"), *FSHA1_HashString(DataKey));
+	UE_LOG(LogMaterial, Display, TEXT("Saved shaders for %s from DDC (key hash: %s)"), GetMaterialPath(), *FSHA1_HashString(DataKey));
 	GetDerivedDataCacheRef().Put(*DataKey, SaveData, FStringView(GetFriendlyName()));
 	COOK_STAT(Timer.AddMiss(SaveData.Num()));
 }
diff --git a/Engine/Source/Runtime/Engine/Private/Materials/MaterialShared.cpp b/Engine/Source/Runtime/Engine/Private/Materials/MaterialShared.cpp
--- a/Engine/Source/Runtime/Engine/Private/Materials/MaterialShared.cpp
+++ b/Engine/Source/Runtime/Engine/Private/Materials/MaterialShared.cpp
@@ -2197,10 +2197,6 @@ bool FMaterial::CacheShaders(const FMaterialShaderMapId& ShaderMapId, EShaderPla
 				TRefCountPtr<FMaterialShaderMap> LoadedShaderMap;
 				FMaterialShaderMap::LoadFromDerivedDataCache(this, ShaderMapId, Platform, TargetPlatform, LoadedShaderMap, DDCKeyHash);
 				ShaderMap = LoadedShaderMap;
-				if (LoadedShaderMap)
-				{
-					UE_LOG(LogMaterial, Verbose, TEXT("Loaded shaders for %s from DDC (key hash: %s)"), *GetAssetName(), *DDCKeyHash);
-				}
 			}
 
 			check(!ShaderMap || ShaderMap->GetFrozenContentSize() > 0u);
