Added ifdef around to GetResourceInfo to fix compilation issues on platforms that don't support it.

diff --git a/Engine/Source/Runtime/Engine/Private/UnrealEngine.cpp b/Engine/Source/Runtime/Engine/Private/UnrealEngine.cpp
--- a/Engine/Source/Runtime/Engine/Private/UnrealEngine.cpp
+++ b/Engine/Source/Runtime/Engine/Private/UnrealEngine.cpp
@@ -6839,11 +6839,13 @@ bool UEngine::HandleSkeletalMeshReportCommand(const TCHAR* Cmd, FOutputDevice& A
 	auto GetResourceSizeKB = [&](const FRHIResource* Resource) -> uint32
 	{
 		int32 SizeInKB = 0;
+#if RHI_WANT_RESOURCE_INFO
 		FRHIResourceInfo ResourceInfo;
 		if (Resource && Resource->GetResourceInfo(ResourceInfo))
 		{
 			SizeInKB = BytesToKB(ResourceInfo.VRamAllocation.AllocationSize);
 		}
+#endif
 		return SizeInKB;
 	};
 
