World Partition: moved streaming policy from world settings to world partition and disallow setting a null streaming policy.

diff --git a/Engine/Source/Runtime/Engine/Classes/GameFramework/WorldSettings.h b/Engine/Source/Runtime/Engine/Classes/GameFramework/WorldSettings.h
--- a/Engine/Source/Runtime/Engine/Classes/GameFramework/WorldSettings.h
+++ b/Engine/Source/Runtime/Engine/Classes/GameFramework/WorldSettings.h
@@ -18,7 +18,6 @@ class UWorldPartition;
 class UAssetUserData;
 class UNetConnection;
 class UNavigationSystemConfig;
-class UWorldPartitionStreamingPolicy;
 class UAISystemBase;
 
 UENUM()
@@ -480,10 +479,6 @@ public:
 	UPROPERTY(BlueprintReadOnly, Category=WorldPartitionSetup)
 	uint8 bEnableWorldPartition:1;
 
-	/** Class of WorldPartitionStreamingPolicy to be used to manage world partition streaming. */
-	UPROPERTY(EditAnywhere, Category=WorldPartitionSetup, AdvancedDisplay, meta=(editcondition = "bEnableWorldPartition"))
-	TSubclassOf<UWorldPartitionStreamingPolicy> WorldPartitionStreamingPolicyClass;
-
 	/** Additional transform applied when applying LevelStreaming Transform to LevelInstance */
 	UPROPERTY()
 	FVector LevelInstancePivotOffset;
@@ -545,7 +540,7 @@ protected:
 	UPROPERTY(Transient)
 	TObjectPtr<UNavigationSystemConfig> NavigationSystemConfigOverride;
 
-	UPROPERTY(EditAnywhere, Category=WorldPartitionSetup, NoClear, meta=(NoResetToDefault), Instanced)
+	UPROPERTY(EditAnywhere, Category=WorldPartitionSetup, NoClear, meta=(NoResetToDefault, HideEditConditionToggle, EditCondition = "bEnableWorldPartition"), Instanced)
 	TObjectPtr<UWorldPartition> WorldPartition;
 
 public:
diff --git a/Engine/Source/Runtime/Engine/Private/WorldPartition/WorldPartition.cpp b/Engine/Source/Runtime/Engine/Private/WorldPartition/WorldPartition.cpp
--- a/Engine/Source/Runtime/Engine/Private/WorldPartition/WorldPartition.cpp
+++ b/Engine/Source/Runtime/Engine/Private/WorldPartition/WorldPartition.cpp
@@ -53,10 +53,10 @@ TMap<FName, FString> GetDataLayersDumpString(const UWorldPartition* WorldPartiti
 	if (const AWorldDataLayers* WorldDataLayers = WorldPartition->GetWorld()->GetWorldDataLayers())
 	{
 		WorldDataLayers->ForEachDataLayer([&DataLayersDumpString](const UDataLayer* DataLayer)
-			{
-				DataLayersDumpString.FindOrAdd(DataLayer->GetFName()) = FString::Format(TEXT("{0}({1})"), { DataLayer->GetDataLayerLabel().ToString(), DataLayer->GetFName().ToString() });
-				return true;
-			});
+		{
+			DataLayersDumpString.FindOrAdd(DataLayer->GetFName()) = FString::Format(TEXT("{0}({1})"), { DataLayer->GetDataLayerLabel().ToString(), DataLayer->GetFName().ToString() });
+			return true;
+		});
 	}
 	return DataLayersDumpString;
 }
@@ -249,6 +249,10 @@ UWorldPartition::UWorldPartition(const FObjectInitializer& ObjectInitializer)
 		FWorldDelegates::LevelRemovedFromWorld.AddStatic(&UWorldPartition::WorldPartitionOnLevelRemovedFromWorld);
 		bRegisteredDelegate = true;
 	}
+
+#if WITH_EDITOR
+	WorldPartitionStreamingPolicyClass = UWorldPartitionLevelStreamingPolicy::StaticClass();
+#endif
 }
 
 #if WITH_EDITOR
@@ -264,7 +268,6 @@ void UWorldPartition::OnGCPostReachabilityAnalysis()
 			return;
 		}
 	}
-	
 
 	for (FRawObjectIterator It; It; ++It)
 	{
@@ -1242,7 +1245,7 @@ void UWorldPartition::DrawRuntimeHashPreview()
 bool UWorldPartition::GenerateStreaming(EWorldPartitionStreamingMode Mode, TArray<FString>* OutPackagesToGenerate)
 {
 	check(!StreamingPolicy);
-	StreamingPolicy = NewObject<UWorldPartitionStreamingPolicy>(const_cast<UWorldPartition*>(this), GetWorld()->GetWorldSettings()->WorldPartitionStreamingPolicyClass.Get());
+	StreamingPolicy = NewObject<UWorldPartitionStreamingPolicy>(const_cast<UWorldPartition*>(this), WorldPartitionStreamingPolicyClass.Get());
 
 	check(RuntimeHash);
 	bool Result = RuntimeHash->GenerateRuntimeStreaming(Mode, StreamingPolicy, OutPackagesToGenerate);
diff --git a/Engine/Source/Runtime/Engine/Private/WorldSettings.cpp b/Engine/Source/Runtime/Engine/Private/WorldSettings.cpp
--- a/Engine/Source/Runtime/Engine/Private/WorldSettings.cpp
+++ b/Engine/Source/Runtime/Engine/Private/WorldSettings.cpp
@@ -14,7 +14,6 @@
 #include "Engine/WorldComposition.h"
 #include "WorldPartition/WorldPartition.h"
 #include "WorldPartition/WorldPartitionSubsystem.h"
-#include "WorldPartition/WorldPartitionLevelStreamingPolicy.h"
 #include "Net/UnrealNetwork.h"
 #include "GameFramework/GameNetworkManager.h"
 #include "AudioDevice.h"
@@ -124,9 +123,6 @@ PRAGMA_ENABLE_DEPRECATION_WARNINGS
 	DefaultBookmarkClass = UBookMark::StaticClass();
 	LastBookmarkClass = DefaultBookmarkClass;
 
-	// Default to UWorldPartitionLevelStreamingPolicy
-	WorldPartitionStreamingPolicyClass = UWorldPartitionLevelStreamingPolicy::StaticClass();
-
 	LevelInstancePivotOffset = FVector::ZeroVector;
 }
 
diff --git a/Engine/Source/Runtime/Engine/Public/WorldPartition/WorldPartition.h b/Engine/Source/Runtime/Engine/Public/WorldPartition/WorldPartition.h
--- a/Engine/Source/Runtime/Engine/Public/WorldPartition/WorldPartition.h
+++ b/Engine/Source/Runtime/Engine/Public/WorldPartition/WorldPartition.h
@@ -191,6 +191,10 @@ public:
 	TObjectPtr<UWorldPartitionEditorHash> EditorHash;
 
 	IWorldPartitionEditor* WorldPartitionEditor;
+
+	/** Class of WorldPartitionStreamingPolicy to be used to manage world partition streaming. */
+	UPROPERTY(EditAnywhere, Category=RuntimeHash, NoClear, meta = (NoResetToDefault))
+	TSubclassOf<UWorldPartitionStreamingPolicy> WorldPartitionStreamingPolicyClass;
 #endif
 
 	UPROPERTY(EditAnywhere, Category=RuntimeHash, NoClear, meta = (NoResetToDefault), Instanced)
