Avoid attempting to fetch VTChunk data for chunks with empty keys.
Fix typo (DCC->DDC) in VTChunk code.

diff --git a/Engine/Source/Runtime/Engine/Private/TextureDerivedDataTask.cpp b/Engine/Source/Runtime/Engine/Private/TextureDerivedDataTask.cpp
--- a/Engine/Source/Runtime/Engine/Private/TextureDerivedDataTask.cpp
+++ b/Engine/Source/Runtime/Engine/Private/TextureDerivedDataTask.cpp
@@ -1015,7 +1015,7 @@ void FTextureCacheDerivedDataWorker::DoWork()
 		CompositeTextureData.ReleaseMemory();
 
 		// Populate the VT DDC Cache now if we're asynchronously loading to avoid too many high prio/synchronous request on the render thread
-		if (!IsInGameThread() && DerivedData->VTData)
+		if (!IsInGameThread() && DerivedData->VTData && !DerivedData->VTData->Chunks.Last().DerivedDataKey.IsEmpty())
 		{
 			GetVirtualTextureChunkDDCCache()->MakeChunkAvailable_Concurrent(&DerivedData->VTData->Chunks.Last());
 		}
diff --git a/Engine/Source/Runtime/Engine/Private/VT/UploadingVirtualTexture.cpp b/Engine/Source/Runtime/Engine/Private/VT/UploadingVirtualTexture.cpp
--- a/Engine/Source/Runtime/Engine/Private/VT/UploadingVirtualTexture.cpp
+++ b/Engine/Source/Runtime/Engine/Private/VT/UploadingVirtualTexture.cpp
@@ -393,19 +393,19 @@ FVTDataAndStatus FUploadingVirtualTexture::ReadData(FGraphEventArray& OutComplet
 		{
 			SCOPE_CYCLE_COUNTER(STAT_VTP_MakeChunkAvailable);
 			check(Chunk.DerivedDataKey.IsEmpty() == false);
-			FString ChunkFileNameDCC;
+			FString ChunkFileNameDDC;
 
-			// If request is flagged as high priority, we will block here until DCC cache is populated
+			// If request is flagged as high priority, we will block here until DDC cache is populated
 			// This way we can service these high priority tasks immediately
-			// Would be better to have DCC cache return a task event handle, which could be used to chain a subsequent read operation,
+			// Would be better to have DDC cache return a task event handle, which could be used to chain a subsequent read operation,
 			// but that would be more complicated, and this should generally not be a critical runtime path
-			const bool bAsyncDCC = (Priority == EVTRequestPagePriority::Normal);
-			const bool Available = GetVirtualTextureChunkDDCCache()->MakeChunkAvailable(&Data->Chunks[ChunkIndex], bAsyncDCC, ChunkFileNameDCC, ChunkOffsetInFile);
+			const bool bAsyncDDC = (Priority == EVTRequestPagePriority::Normal);
+			const bool Available = GetVirtualTextureChunkDDCCache()->MakeChunkAvailable(&Data->Chunks[ChunkIndex], bAsyncDDC, ChunkFileNameDDC, ChunkOffsetInFile);
 			if (!Available)
 			{
 				return EVTRequestPageStatus::Saturated;
 			}
-			ChunkFileName = ChunkFileNameDCC;
+			ChunkFileName = ChunkFileNameDDC;
 			ChunkSource = EChunkSource::File;
 		}
 #else // WITH_EDITOR
