Don't warn about composite texture mismatches for UDIMs, this has never worked correctly

diff --git a/Engine/Source/Runtime/Engine/Private/TextureDerivedDataTask.cpp b/Engine/Source/Runtime/Engine/Private/TextureDerivedDataTask.cpp
--- a/Engine/Source/Runtime/Engine/Private/TextureDerivedDataTask.cpp
+++ b/Engine/Source/Runtime/Engine/Private/TextureDerivedDataTask.cpp
@@ -766,7 +766,13 @@ FTextureCacheDerivedDataWorker::FTextureCacheDerivedDataWorker(
 
 		if (!bMatchingBlocks)
 		{
-			UE_LOG(LogTexture, Warning, TEXT("Issue while building %s : Composite texture resolution/UDIMs do not match. Composite texture will be ignored"), *Texture.GetPathName());
+			// Only report the warning for textures with a single block
+			// In the future, we should support composite textures if matching blocks are in a different order
+			// Once that's working, then this warning should be reported in all cases
+			if (Texture.Source.GetNumBlocks() == 1)
+			{
+				UE_LOG(LogTexture, Warning, TEXT("Issue while building %s : Composite texture resolution/UDIMs do not match. Composite texture will be ignored"), *Texture.GetPathName());
+			}
 		}
 		else if (!bOnlyPowerOfTwoSize)
 		{
