Fix crash when invalid material expression is piped into RVT world position pin.

diff --git a/Engine/Source/Runtime/Engine/Private/Materials/MaterialExpressions.cpp b/Engine/Source/Runtime/Engine/Private/Materials/MaterialExpressions.cpp
--- a/Engine/Source/Runtime/Engine/Private/Materials/MaterialExpressions.cpp
+++ b/Engine/Source/Runtime/Engine/Private/Materials/MaterialExpressions.cpp
@@ -2570,10 +2570,13 @@ int32 UMaterialExpressionRuntimeVirtualTextureSample::Compile(class FMaterialCom
 			WorldPositionIndex = Compiler->WorldPosition(WPT_Default);
 		}
 		
-		const int32 P0 = Uniforms[ERuntimeVirtualTextureShaderUniform_WorldToUVTransform0];
-		const int32 P1 = Uniforms[ERuntimeVirtualTextureShaderUniform_WorldToUVTransform1];
-		const int32 P2 = Uniforms[ERuntimeVirtualTextureShaderUniform_WorldToUVTransform2];
-		CoordinateIndex = Compiler->VirtualTextureWorldToUV(WorldPositionIndex, P0, P1, P2);
+		if (WorldPositionIndex != INDEX_NONE)
+		{
+			const int32 P0 = Uniforms[ERuntimeVirtualTextureShaderUniform_WorldToUVTransform0];
+			const int32 P1 = Uniforms[ERuntimeVirtualTextureShaderUniform_WorldToUVTransform1];
+			const int32 P2 = Uniforms[ERuntimeVirtualTextureShaderUniform_WorldToUVTransform2];
+			CoordinateIndex = Compiler->VirtualTextureWorldToUV(WorldPositionIndex, P0, P1, P2);
+		}
 	}
 	
 	// Compile the mip level for the current mip value mode
