Fix incorrect calculation of virtual texture upload cache stats.
Adding mem to stat used correct aligned size.
Removing mem from stat used incorrect smaller unaligned size.

diff --git a/Engine/Source/Runtime/Engine/Private/VT/VirtualTextureUploadCache.cpp b/Engine/Source/Runtime/Engine/Private/VT/VirtualTextureUploadCache.cpp
--- a/Engine/Source/Runtime/Engine/Private/VT/VirtualTextureUploadCache.cpp
+++ b/Engine/Source/Runtime/Engine/Private/VT/VirtualTextureUploadCache.cpp
@@ -98,7 +98,7 @@ uint32 FVTUploadTileAllocator::Allocate(EPixelFormat InFormat, uint32 InTileSize
 	{
 		// Staging buffer needs underlying buffer allocating.
 		StagingBuffer.Init(Desc.BlockBytes, Desc.MemorySize);
-		NumAllocatedBytes += StagingBuffer.NumTiles * StagingBuffer.TileSize;
+		NumAllocatedBytes += StagingBuffer.TileSizeAligned * StagingBuffer.NumTiles;
 	}
 
 	// Pop a free tile and return handle.
@@ -123,8 +123,8 @@ void FVTUploadTileAllocator::Free(uint32 InHandle)
 	if (StagingBuffer.NumTiles == StagingBuffer.TileFreeList.Num())
 	{
 		// All tiles are free, so release the underlying memory.
-		check(NumAllocatedBytes >= StagingBuffer.NumTiles * StagingBuffer.TileSize);
-		NumAllocatedBytes -= StagingBuffer.NumTiles * StagingBuffer.TileSize;
+		check(NumAllocatedBytes >= StagingBuffer.TileSizeAligned * StagingBuffer.NumTiles);
+		NumAllocatedBytes -= StagingBuffer.TileSizeAligned * StagingBuffer.NumTiles;
 
 		StagingBuffer.Release();
 	}
@@ -139,7 +139,7 @@ FVTUploadTileBuffer FVTUploadTileAllocator::GetBufferFromHandle(uint32 InHandle)
 	FStagingBuffer const& StagingBuffer = FormatBuffers[Handle.FormatIndex].StagingBuffers[Handle.StagingBufferIndex];
 
 	FVTUploadTileBuffer Buffer;
-	Buffer.Memory = (uint8*)StagingBuffer.Memory + Handle.TileIndex * Align(StagingBuffer.TileSize, 128u);
+	Buffer.Memory = (uint8*)StagingBuffer.Memory + StagingBuffer.TileSizeAligned * Handle.TileIndex;
 	Buffer.MemorySize = StagingBuffer.TileSize;
 	Buffer.Stride = FormatDesc.Stride;
 	return Buffer;
@@ -156,7 +156,7 @@ FVTUploadTileBufferExt FVTUploadTileAllocator::GetBufferFromHandleExt(uint32 InH
 	FVTUploadTileBufferExt Buffer;
 	Buffer.RHIBuffer = StagingBuffer.RHIBuffer;
 	Buffer.BufferMemory = StagingBuffer.Memory;
-	Buffer.BufferOffset = Handle.TileIndex * Align(StagingBuffer.TileSize, 128u);
+	Buffer.BufferOffset = StagingBuffer.TileSizeAligned * Handle.TileIndex;
 	Buffer.Stride = FormatDesc.Stride;
 	return Buffer;
 }
@@ -169,11 +169,11 @@ FVTUploadTileAllocator::FStagingBuffer::~FStagingBuffer()
 void FVTUploadTileAllocator::FStagingBuffer::Init(uint32 InBufferStrideBytes, uint32 InTileSizeBytes)
 {
 	TileSize = InTileSizeBytes;
-	const uint32 BufferSize = CVarVTUploadMemoryPageSize.GetValueOnRenderThread() * 1024 * 1024;
-	const uint32 AlignedTileSizeBytes = Align(InTileSizeBytes, 128u);
+	TileSizeAligned = Align(InTileSizeBytes, 128u);
 
-	NumTiles = FMath::DivideAndRoundUp(BufferSize, AlignedTileSizeBytes);
-	const uint32 AlignedBufferSize = NumTiles * AlignedTileSizeBytes;
+	const uint32 RequestedBufferSize = CVarVTUploadMemoryPageSize.GetValueOnRenderThread() * 1024 * 1024;
+	NumTiles = FMath::DivideAndRoundUp(RequestedBufferSize, TileSizeAligned);
+	const uint32 BufferSize = TileSizeAligned * NumTiles;
 
 	check(TileFreeList.Num() == 0);
 	TileFreeList.AddUninitialized(NumTiles);
@@ -189,25 +189,27 @@ void FVTUploadTileAllocator::FStagingBuffer::Init(uint32 InBufferStrideBytes, ui
 		FRHICommandListImmediate& RHICmdList = FRHICommandListExecutor::GetImmediateCommandList();
 
 		FRHIResourceCreateInfo CreateInfo(TEXT("StagingBuffer"));
-		RHIBuffer = RHICreateStructuredBuffer(InBufferStrideBytes, AlignedBufferSize, BUF_ShaderResource | BUF_Static | BUF_KeepCPUAccessible, CreateInfo);
+		RHIBuffer = RHICreateStructuredBuffer(InBufferStrideBytes, BufferSize, BUF_ShaderResource | BUF_Static | BUF_KeepCPUAccessible, CreateInfo);
 
 		// Here we bypass 'normal' RHI operations in order to get a persistent pointer to GPU memory, on supported platforms
 		// This should be encapsulated into a proper RHI method at some point
-		Memory = RHICmdList.LockBuffer(RHIBuffer, 0u, AlignedBufferSize, RLM_WriteOnly_NoOverwrite);
+		Memory = RHICmdList.LockBuffer(RHIBuffer, 0u, BufferSize, RLM_WriteOnly_NoOverwrite);
 
-		INC_MEMORY_STAT_BY(STAT_TotalGPUUploadSize, AlignedBufferSize);
+		INC_MEMORY_STAT_BY(STAT_TotalGPUUploadSize, BufferSize);
 	}
 	else
 	{
 		// Allocate staging buffer in CPU memory.
-		Memory = FMemory::Malloc(AlignedBufferSize, 128u);
+		Memory = FMemory::Malloc(BufferSize, 128u);
 
-		INC_MEMORY_STAT_BY(STAT_TotalCPUUploadSize, AlignedBufferSize);
+		INC_MEMORY_STAT_BY(STAT_TotalCPUUploadSize, BufferSize);
 	}
 }
 
 void FVTUploadTileAllocator::FStagingBuffer::Release()
 {
+	const uint32 BufferSize = TileSizeAligned * NumTiles;
+
 	if (RHIBuffer.IsValid())
 	{
 		// Unmap and release the GPU buffer if present.
@@ -216,7 +218,7 @@ void FVTUploadTileAllocator::FStagingBuffer::Release()
 		// In this case 'Memory' was the mapped pointer, so release it.
 		Memory = nullptr;
 
-		DEC_MEMORY_STAT_BY(STAT_TotalGPUUploadSize, TileSize * NumTiles);
+		DEC_MEMORY_STAT_BY(STAT_TotalGPUUploadSize, BufferSize);
 	}
 
 	if (Memory != nullptr)
@@ -225,7 +227,7 @@ void FVTUploadTileAllocator::FStagingBuffer::Release()
 		FMemory::Free(Memory);
 		Memory = nullptr;
 
-		DEC_MEMORY_STAT_BY(STAT_TotalCPUUploadSize, TileSize * NumTiles);
+		DEC_MEMORY_STAT_BY(STAT_TotalCPUUploadSize, BufferSize);
 	}
 
 	TileSize = 0u;
diff --git a/Engine/Source/Runtime/Engine/Private/VT/VirtualTextureUploadCache.h b/Engine/Source/Runtime/Engine/Private/VT/VirtualTextureUploadCache.h
--- a/Engine/Source/Runtime/Engine/Private/VT/VirtualTextureUploadCache.h
+++ b/Engine/Source/Runtime/Engine/Private/VT/VirtualTextureUploadCache.h
@@ -94,6 +94,7 @@ private:
 		/** Memory pointer to locked GPU buffer if used on platform, or to allocated CPU heap memory if not. */
 		void* Memory = nullptr;
 		uint32 TileSize = 0u;
+		uint32 TileSizeAligned = 0u;
 		uint32 NumTiles = 0u;
 		/** List of tile indices that haven't been allocated. */
 		TArray<uint16> TileFreeList;
