Fixed local declation of CVarDefaultGPUSkinCacheBehavior shadowing global declaration in Unity builds.

diff --git a/Engine/Source/Runtime/Engine/Private/UnrealEngine.cpp b/Engine/Source/Runtime/Engine/Private/UnrealEngine.cpp
--- a/Engine/Source/Runtime/Engine/Private/UnrealEngine.cpp
+++ b/Engine/Source/Runtime/Engine/Private/UnrealEngine.cpp
@@ -6849,8 +6849,8 @@ bool UEngine::HandleSkeletalMeshReportCommand(const TCHAR* Cmd, FOutputDevice& A
 		return SizeInKB;
 	};
 
-	static const IConsoleVariable* CVarDefaultGPUSkinCacheBehavior = IConsoleManager::Get().FindConsoleVariable(TEXT("r.SkinCache.DefaultBehavior"));
-	bool bSkinCacheGlobalDefault = CVarDefaultGPUSkinCacheBehavior && ESkinCacheDefaultBehavior(CVarDefaultGPUSkinCacheBehavior->GetInt()) == ESkinCacheDefaultBehavior::Inclusive;
+	static const IConsoleVariable* DefaultGPUSkinCacheBehavior = IConsoleManager::Get().FindConsoleVariable(TEXT("r.SkinCache.DefaultBehavior"));
+	bool bSkinCacheGlobalDefault = DefaultGPUSkinCacheBehavior && ESkinCacheDefaultBehavior(DefaultGPUSkinCacheBehavior->GetInt()) == ESkinCacheDefaultBehavior::Inclusive;
 
 	int32 NumMeshes = 0;
 	FRHIResourceInfo ResourceInfo;
