Replaced overzealous check by ensure when HISM component transform or view proj matrix is invalid

diff --git a/Engine/Source/Runtime/Engine/Private/HierarchicalInstancedStaticMesh.cpp b/Engine/Source/Runtime/Engine/Private/HierarchicalInstancedStaticMesh.cpp
--- a/Engine/Source/Runtime/Engine/Private/HierarchicalInstancedStaticMesh.cpp
+++ b/Engine/Source/Runtime/Engine/Private/HierarchicalInstancedStaticMesh.cpp
@@ -1579,7 +1579,6 @@ void FHierarchicalStaticMeshSceneProxy::GetDynamicMeshElements(const TArray<cons
 			// Render built instances
 			if (ClusterTree.Num())
 			{
-
 				FFoliageCullInstanceParams InstanceParams(bSingleSections, bMultipleSections, bOverestimate, ClusterTree);
 				InstanceParams.LODs = RenderData->LODResources.Num();
 
@@ -1608,8 +1607,9 @@ void FHierarchicalStaticMeshSceneProxy::GetDynamicMeshElements(const TArray<cons
 				else
 				{
 					// build view frustum with no near plane / no far plane in frustum (far plane culling is done later in the function) : 
-					const bool bViewFrustumUsesNearPlane = false;
-					const bool bViewFrustumUsesFarPlane = false;
+					static constexpr bool bViewFrustumUsesNearPlane = false;
+					static constexpr bool bViewFrustumUsesFarPlane = false;
+					const bool bIsPerspectiveProjection = View->ViewMatrices.IsPerspectiveProjection();
 
 					// Instanced stereo needs to use the right plane from the right eye when constructing the frustum bounds to cull against.
 					// Otherwise we'll cull objects visible in the right eye, but not the left.
@@ -1627,7 +1627,7 @@ void FHierarchicalStaticMeshSceneProxy::GetDynamicMeshElements(const TArray<cons
 						// Invalid bounds retrieved, so skip render of this frame :
 						if (LeftEyeBounds.Planes.Num() != 4 || RightEyeBounds.Planes.Num() != 4)
 						{
-							checkf(false, TEXT("Invalid frustum, skipping render"));
+							ensureMsgf(false, TEXT("Invalid frustum, skipping render of HISM : culling view projection matrix:\n- Left eye: %s\n-Right eye: %s"), *LeftEyeLocalViewProjForCulling.ToString(), *RightEyeLocalViewProjForCulling.ToString());
 							continue;
 						}
 						
@@ -1642,16 +1642,19 @@ void FHierarchicalStaticMeshSceneProxy::GetDynamicMeshElements(const TArray<cons
 					{
 						const FMatrix LocalViewProjForCulling = GetLocalToWorld() * View->ViewMatrices.GetViewProjectionMatrix();
 						GetViewFrustumBounds(InstanceParams.ViewFrustumLocal, LocalViewProjForCulling, bViewFrustumUsesNearPlane, bViewFrustumUsesFarPlane);
-					}
 
-					if (View->ViewMatrices.IsPerspectiveProjection())
-					{
 						// Invalid bounds retrieved, so skip render of this frame :
-						if (InstanceParams.ViewFrustumLocal.Planes.Num() != 4)
+						if (bIsPerspectiveProjection && (InstanceParams.ViewFrustumLocal.Planes.Num() != 4))
 						{
-							checkf(false, TEXT("Invalid frustum, skipping render"));
+							// Report the error as a warning (instead of an ensure or a check) as the problem can come from improper user data (invalid transform or view-proj matrix) : 
+							ensureMsgf(false, TEXT("Invalid frustum, skipping render of HISM : culling view projection matrix:%s"), *LocalViewProjForCulling.ToString());
 							continue;
 						}
+					}
+
+					if (bIsPerspectiveProjection)
+					{
+						check(InstanceParams.ViewFrustumLocal.Planes.Num() == 4);
 
 						FMatrix ThreePlanes;
 						ThreePlanes.SetIdentity();
