Fix for instanced static mesh per-instance selection when the component is generated from a Blueprint construction script.

diff --git a/Engine/Source/Runtime/Engine/Classes/Components/InstancedStaticMeshComponent.h b/Engine/Source/Runtime/Engine/Classes/Components/InstancedStaticMeshComponent.h
--- a/Engine/Source/Runtime/Engine/Classes/Components/InstancedStaticMeshComponent.h
+++ b/Engine/Source/Runtime/Engine/Classes/Components/InstancedStaticMeshComponent.h
@@ -582,4 +582,7 @@ public:
 	/* Additional random seeds */
 	UPROPERTY()
 	TArray<FInstancedStaticMeshRandomSeed> AdditionalRandomSeeds;
+
+	UPROPERTY()
+	bool bHasPerInstanceHitProxies;
 };
diff --git a/Engine/Source/Runtime/Engine/Private/InstancedStaticMesh.cpp b/Engine/Source/Runtime/Engine/Private/InstancedStaticMesh.cpp
--- a/Engine/Source/Runtime/Engine/Private/InstancedStaticMesh.cpp
+++ b/Engine/Source/Runtime/Engine/Private/InstancedStaticMesh.cpp
@@ -1926,6 +1926,9 @@ TStructOnScope<FActorComponentInstanceData> UInstancedStaticMeshComponent::GetCo
 	// Back up random seed
 	StaticMeshInstanceData->InstancingRandomSeed = InstancingRandomSeed;
 	StaticMeshInstanceData->AdditionalRandomSeeds = AdditionalRandomSeeds;
+
+	// Back up per-instance hit proxies
+	StaticMeshInstanceData->bHasPerInstanceHitProxies = bHasPerInstanceHitProxies;
 #endif
 	return InstanceData;
 }
@@ -1990,6 +1993,8 @@ void UInstancedStaticMeshComponent::ApplyComponentInstanceData(FInstancedStaticM
 	InstancingRandomSeed = InstancedMeshData->InstancingRandomSeed;
 	AdditionalRandomSeeds = InstancedMeshData->AdditionalRandomSeeds;
 
+	bHasPerInstanceHitProxies = InstancedMeshData->bHasPerInstanceHitProxies;
+
 	// Force recreation of the render data
 	InstanceUpdateCmdBuffer.Edit();
 	MarkRenderStateDirty();
