Default initialize new component instance data member.

diff --git a/Engine/Source/Runtime/Engine/Classes/Components/InstancedStaticMeshComponent.h b/Engine/Source/Runtime/Engine/Classes/Components/InstancedStaticMeshComponent.h
--- a/Engine/Source/Runtime/Engine/Classes/Components/InstancedStaticMeshComponent.h
+++ b/Engine/Source/Runtime/Engine/Classes/Components/InstancedStaticMeshComponent.h
@@ -584,5 +584,5 @@ public:
 	TArray<FInstancedStaticMeshRandomSeed> AdditionalRandomSeeds;
 
 	UPROPERTY()
-	bool bHasPerInstanceHitProxies;
+	bool bHasPerInstanceHitProxies = false;
 };
