Continuation of previous material parameter refactor.  Replace FMaterialRenderProxy Scalar/Vector/Texture parameter query functions with generic version

diff --git a/Engine/Plugins/2D/Paper2D/Source/Paper2D/Private/PaperRenderSceneProxy.cpp b/Engine/Plugins/2D/Paper2D/Source/Paper2D/Private/PaperRenderSceneProxy.cpp
--- a/Engine/Plugins/2D/Paper2D/Source/Paper2D/Private/PaperRenderSceneProxy.cpp
+++ b/Engine/Plugins/2D/Paper2D/Source/Paper2D/Private/PaperRenderSceneProxy.cpp
@@ -106,41 +106,29 @@ public:
 		return Parent->GetFallback(InFeatureLevel);
 	}
 
-	virtual bool GetVectorValue(const FHashedMaterialParameterInfo& ParameterInfo, FLinearColor* OutValue, const FMaterialRenderContext& Context) const override
+	virtual bool GetParameterValue(EMaterialParameterType Type, const FHashedMaterialParameterInfo& ParameterInfo, FMaterialParameterValue& OutValue, const FMaterialRenderContext& Context) const
 	{
-		return Parent->GetVectorValue(ParameterInfo, OutValue, Context);
-	}
-
-	virtual bool GetScalarValue(const FHashedMaterialParameterInfo& ParameterInfo, float* OutValue, const FMaterialRenderContext& Context) const override
-	{
-		return Parent->GetScalarValue(ParameterInfo, OutValue, Context);
-	}
-
-	virtual bool GetTextureValue(const FHashedMaterialParameterInfo& ParameterInfo, const UTexture** OutValue, const FMaterialRenderContext& Context) const override
-	{
-		if (ParameterInfo.Name == TextureParameterName)
-		{
-			*OutValue = ApplyEditorOverrides(BaseTexture);
-			return true;
-		}
-
-		for (int32 AdditionalSlotIndex = 0; AdditionalSlotIndex < AdditionalTextures.Num(); ++AdditionalSlotIndex)
+		if (Type == EMaterialParameterType::Texture)
 		{
-			FName AdditionalSlotName = AdditionalTextureParameterRootName;
-			AdditionalSlotName.SetNumber(AdditionalSlotIndex + 1);
-			if (ParameterInfo.Name == AdditionalSlotName)
+			if (ParameterInfo.Name == TextureParameterName)
 			{
-				*OutValue = ApplyEditorOverrides(AdditionalTextures[AdditionalSlotIndex]);
+				OutValue = ApplyEditorOverrides(BaseTexture);
 				return true;
 			}
-		}
 
-		return Parent->GetTextureValue(ParameterInfo, OutValue, Context);
-	}
+			for (int32 AdditionalSlotIndex = 0; AdditionalSlotIndex < AdditionalTextures.Num(); ++AdditionalSlotIndex)
+			{
+				FName AdditionalSlotName = AdditionalTextureParameterRootName;
+				AdditionalSlotName.SetNumber(AdditionalSlotIndex + 1);
+				if (ParameterInfo.Name == AdditionalSlotName)
+				{
+					OutValue = ApplyEditorOverrides(AdditionalTextures[AdditionalSlotIndex]);
+					return true;
+				}
+			}
+		}
 
-	virtual bool GetTextureValue(const FHashedMaterialParameterInfo& ParameterInfo, const URuntimeVirtualTexture** OutValue, const FMaterialRenderContext& Context) const override
-	{
-		return Parent->GetTextureValue(ParameterInfo, OutValue, Context);
+		return Parent->GetParameterValue(Type, ParameterInfo, OutValue, Context);
 	}
 
 #if WITH_EDITOR
diff --git a/Engine/Plugins/Runtime/HairStrands/Source/HairStrandsCore/Private/GroomComponent.cpp b/Engine/Plugins/Runtime/HairStrands/Source/HairStrandsCore/Private/GroomComponent.cpp
--- a/Engine/Plugins/Runtime/HairStrands/Source/HairStrandsCore/Private/GroomComponent.cpp
+++ b/Engine/Plugins/Runtime/HairStrands/Source/HairStrandsCore/Private/GroomComponent.cpp
@@ -169,44 +169,38 @@ public:
 		return Parent->GetFallback(InFeatureLevel);
 	}
 
-	virtual bool GetVectorValue(const FHashedMaterialParameterInfo& ParameterInfo, FLinearColor* OutValue, const FMaterialRenderContext& Context) const override
+	virtual bool GetParameterValue(EMaterialParameterType Type, const FHashedMaterialParameterInfo& ParameterInfo, FMaterialParameterValue& OutValue, const FMaterialRenderContext& Context) const
 	{
-		if (ParameterInfo.Name == HairGroupHairColorParamName)
+		switch (Type)
 		{
-			*OutValue = HairGroupColor;
-			return true;
-		}
-		return Parent->GetVectorValue(ParameterInfo, OutValue, Context);
-	}
-
-	virtual bool GetScalarValue(const FHashedMaterialParameterInfo& ParameterInfo, float* OutValue, const FMaterialRenderContext& Context) const override
-	{
-		if (ParameterInfo.Name == DebugModeParamName)
-		{
-			*OutValue = DebugMode;
-			return true;
-		}
-		else if (ParameterInfo.Name == MinHairRadiusParamName)
-		{
-			*OutValue = HairMinRadius;
-			return true;
-		}
-		else if (ParameterInfo.Name == MaxHairRadiusParamName)
-		{
-			*OutValue = HairMaxRadius;
-			return true;
+		case EMaterialParameterType::Vector:
+			if (ParameterInfo.Name == HairGroupHairColorParamName)
+			{
+				OutValue = FVector3f(HairGroupColor);
+				return true;
+			}
+			break;
+		case EMaterialParameterType::Scalar:
+			if (ParameterInfo.Name == DebugModeParamName)
+			{
+				OutValue = DebugMode;
+				return true;
+			}
+			else if (ParameterInfo.Name == MinHairRadiusParamName)
+			{
+				OutValue = HairMinRadius;
+				return true;
+			}
+			else if (ParameterInfo.Name == MaxHairRadiusParamName)
+			{
+				OutValue = HairMaxRadius;
+				return true;
+			}
+			break;
+		default:
+			break;
 		}
-		return Parent->GetScalarValue(ParameterInfo, OutValue, Context);
-	}
-
-	virtual bool GetTextureValue(const FHashedMaterialParameterInfo& ParameterInfo, const URuntimeVirtualTexture** OutValue, const FMaterialRenderContext& Context) const override
-	{
-		return Parent->GetTextureValue(ParameterInfo, OutValue, Context);
-	}
-
-	virtual bool GetTextureValue(const FHashedMaterialParameterInfo& ParameterInfo, const UTexture** OutValue, const FMaterialRenderContext& Context) const override
-	{
-		return Parent->GetTextureValue(ParameterInfo, OutValue, Context);
+		return Parent->GetParameterValue(Type, ParameterInfo, OutValue, Context);
 	}
 };
 
diff --git a/Engine/Source/Developer/MaterialBaking/Private/ExportMaterialProxy.h b/Engine/Source/Developer/MaterialBaking/Private/ExportMaterialProxy.h
--- a/Engine/Source/Developer/MaterialBaking/Private/ExportMaterialProxy.h
+++ b/Engine/Source/Developer/MaterialBaking/Private/ExportMaterialProxy.h
@@ -344,24 +344,9 @@ public:
 		return UMaterial::GetDefaultMaterial(MD_Surface)->GetRenderProxy();
 	}
 
-	virtual bool GetVectorValue(const FHashedMaterialParameterInfo& ParameterInfo, FLinearColor* OutValue, const FMaterialRenderContext& Context) const override
+	virtual bool GetParameterValue(EMaterialParameterType Type, const FHashedMaterialParameterInfo& ParameterInfo, FMaterialParameterValue& OutValue, const FMaterialRenderContext& Context) const override
 	{
-		return MaterialInterface->GetRenderProxy()->GetVectorValue(ParameterInfo, OutValue, Context);
-	}
-
-	virtual bool GetScalarValue(const FHashedMaterialParameterInfo& ParameterInfo, float* OutValue, const FMaterialRenderContext& Context) const override
-	{
-		return MaterialInterface->GetRenderProxy()->GetScalarValue(ParameterInfo, OutValue, Context);
-	}
-
-	virtual bool GetTextureValue(const FHashedMaterialParameterInfo& ParameterInfo, const UTexture** OutValue, const FMaterialRenderContext& Context) const override
-	{
-		return MaterialInterface->GetRenderProxy()->GetTextureValue(ParameterInfo, OutValue, Context);
-	}
-
-	virtual bool GetTextureValue(const FHashedMaterialParameterInfo& ParameterInfo, const URuntimeVirtualTexture** OutValue, const FMaterialRenderContext& Context) const override
-	{
-		return MaterialInterface->GetRenderProxy()->GetTextureValue(ParameterInfo, OutValue, Context);
+		return MaterialInterface->GetRenderProxy()->GetParameterValue(Type, ParameterInfo, OutValue, Context);
 	}
 
 	// Material properties.
diff --git a/Engine/Source/Developer/MaterialUtilities/Private/MaterialUtilities.cpp b/Engine/Source/Developer/MaterialUtilities/Private/MaterialUtilities.cpp
--- a/Engine/Source/Developer/MaterialUtilities/Private/MaterialUtilities.cpp
+++ b/Engine/Source/Developer/MaterialUtilities/Private/MaterialUtilities.cpp
@@ -505,24 +505,9 @@ public:
 		return UMaterial::GetDefaultMaterial(MD_Surface)->GetRenderProxy();
 	}
 
-	virtual bool GetVectorValue(const FHashedMaterialParameterInfo& ParameterInfo, FLinearColor* OutValue, const FMaterialRenderContext& Context) const override
+	virtual bool GetParameterValue(EMaterialParameterType Type, const FHashedMaterialParameterInfo& ParameterInfo, FMaterialParameterValue& OutValue, const FMaterialRenderContext& Context) const override
 	{
-		return MaterialInterface->GetRenderProxy()->GetVectorValue(ParameterInfo, OutValue, Context);
-	}
-
-	virtual bool GetScalarValue(const FHashedMaterialParameterInfo& ParameterInfo, float* OutValue, const FMaterialRenderContext& Context) const override
-	{
-		return MaterialInterface->GetRenderProxy()->GetScalarValue(ParameterInfo, OutValue, Context);
-	}
-
-	virtual bool GetTextureValue(const FHashedMaterialParameterInfo& ParameterInfo,const UTexture** OutValue, const FMaterialRenderContext& Context) const override
-	{
-		return MaterialInterface->GetRenderProxy()->GetTextureValue(ParameterInfo,OutValue,Context);
-	}
-
-	virtual bool GetTextureValue(const FHashedMaterialParameterInfo& ParameterInfo, const URuntimeVirtualTexture** OutValue, const FMaterialRenderContext& Context) const override
-	{
-		return MaterialInterface->GetRenderProxy()->GetTextureValue(ParameterInfo, OutValue, Context);
+		return MaterialInterface->GetRenderProxy()->GetParameterValue(Type, ParameterInfo, OutValue, Context);
 	}
 
 	// Material properties.
diff --git a/Engine/Source/Editor/MaterialEditor/Private/MaterialEditor.cpp b/Engine/Source/Editor/MaterialEditor/Private/MaterialEditor.cpp
--- a/Engine/Source/Editor/MaterialEditor/Private/MaterialEditor.cpp
+++ b/Engine/Source/Editor/MaterialEditor/Private/MaterialEditor.cpp
@@ -464,8 +464,7 @@ void FMaterialEditor::InitMaterialEditor( const EToolkitMode::Type Mode, const T
 	FMaterialEditorSpawnNodeCommands::Register();
 
 	FEditorSupportDelegates::MaterialUsageFlagsChanged.AddRaw(this, &FMaterialEditor::OnMaterialUsageFlagsChanged);
-	FEditorSupportDelegates::VectorParameterDefaultChanged.AddRaw(this, &FMaterialEditor::OnVectorParameterDefaultChanged);
-	FEditorSupportDelegates::ScalarParameterDefaultChanged.AddRaw(this, &FMaterialEditor::OnScalarParameterDefaultChanged);
+	FEditorSupportDelegates::NumericParameterDefaultChanged.AddRaw(this, &FMaterialEditor::OnNumericParameterDefaultChanged);
 
 	FAssetRegistryModule& AssetRegistryModule = FModuleManager::GetModuleChecked<FAssetRegistryModule>(TEXT("AssetRegistry"));
 	
@@ -776,20 +775,14 @@ FMaterialEditor::~FMaterialEditor()
 	// Broadcast that this editor is going down to all listeners
 	OnMaterialEditorClosed().Broadcast();
 
-	for (int32 ParameterIndex = 0; ParameterIndex < OverriddenVectorParametersToRevert.Num(); ParameterIndex++)
+	for (const auto& It : OverriddenNumericParametersToRevert)
 	{
-		SetVectorParameterDefaultOnDependentMaterials(OverriddenVectorParametersToRevert[ParameterIndex], FLinearColor::Black, false);
-	}
-
-	for (int32 ParameterIndex = 0; ParameterIndex < OverriddenScalarParametersToRevert.Num(); ParameterIndex++)
-	{
-		SetScalarParameterDefaultOnDependentMaterials(OverriddenScalarParametersToRevert[ParameterIndex], 0, false);
+		SetNumericParameterDefaultOnDependentMaterials(It.Key, It.Value, UE::Shader::FValue(), false);
 	}
 
 	// Unregister this delegate
 	FEditorSupportDelegates::MaterialUsageFlagsChanged.RemoveAll(this);
-	FEditorSupportDelegates::VectorParameterDefaultChanged.RemoveAll(this);
-	FEditorSupportDelegates::ScalarParameterDefaultChanged.RemoveAll(this);
+	FEditorSupportDelegates::NumericParameterDefaultChanged.RemoveAll(this);
 
 	// Null out the expression preview material so they can be GC'ed
 	ExpressionPreviewMaterial = NULL;
@@ -4445,85 +4438,7 @@ void FMaterialEditor::OnMaterialUsageFlagsChanged(UMaterial* MaterialThatChanged
 	}
 }
 
-void FMaterialEditor::SetVectorParameterDefaultOnDependentMaterials(FName ParameterName, const FLinearColor& Value, bool bOverride)
-{
-	TArray<UMaterial*> MaterialsToOverride;
-
-	if (MaterialFunction)
-	{
-		// Find all materials that reference this function
-		for (TObjectIterator<UMaterial> It; It; ++It)
-		{
-			UMaterial* CurrentMaterial = *It;
-
-			if (CurrentMaterial != Material)
-			{
-				bool bUpdate = false;
-
-				for (int32 FunctionIndex = 0; FunctionIndex < CurrentMaterial->GetCachedExpressionData().FunctionInfos.Num(); FunctionIndex++)
-				{
-					if (CurrentMaterial->GetCachedExpressionData().FunctionInfos[FunctionIndex].Function == MaterialFunction->ParentFunction)
-					{
-						bUpdate = true;
-						break;
-					}
-				}
-
-				if (bUpdate)
-				{
-					MaterialsToOverride.Add(CurrentMaterial);
-				}
-			}
-		}
-	}
-	else
-	{
-		MaterialsToOverride.Add(OriginalMaterial);
-	}
-
-	const ERHIFeatureLevel::Type FeatureLevel = GEditor->GetEditorWorldContext().World()->FeatureLevel;
-
-	for (int32 MaterialIndex = 0; MaterialIndex < MaterialsToOverride.Num(); MaterialIndex++)
-	{
-		UMaterial* CurrentMaterial = MaterialsToOverride[MaterialIndex];
-
-		CurrentMaterial->OverrideVectorParameterDefault(ParameterName, Value, bOverride, FeatureLevel);
-	}
-
-	// Update MI's that reference any of the materials affected
-	for (TObjectIterator<UMaterialInstance> It; It; ++It)
-	{
-		UMaterialInstance* CurrentMaterialInstance = *It;
-
-		// Only care about MI's with static parameters, because we are overriding parameter defaults, 
-		// And only MI's with static parameters contain uniform expressions, which contain parameter defaults
-		if (CurrentMaterialInstance->bHasStaticPermutationResource)
-		{
-			UMaterial* BaseMaterial = CurrentMaterialInstance->GetMaterial();
-
-			if (MaterialsToOverride.Contains(BaseMaterial))
-			{
-				CurrentMaterialInstance->OverrideVectorParameterDefault(ParameterName, Value, bOverride, FeatureLevel);
-			}
-		}
-	}
-}
-
-void FMaterialEditor::OnVectorParameterDefaultChanged(class UMaterialExpression* Expression, FName ParameterName, const FLinearColor& Value)
-{
-	check(Expression);
-
-	if (Expression->Material == Material && OriginalMaterial)
-	{
-		SetVectorParameterDefaultOnDependentMaterials(ParameterName, Value, true);
-
-		OverriddenVectorParametersToRevert.AddUnique(ParameterName);
-	}
-
-	OnParameterDefaultChanged();
-}
-
-void FMaterialEditor::SetScalarParameterDefaultOnDependentMaterials(FName ParameterName, float Value, bool bOverride)
+void FMaterialEditor::SetNumericParameterDefaultOnDependentMaterials(EMaterialParameterType Type, FName ParameterName, const UE::Shader::FValue& Value, bool bOverride)
 {
 	TArray<UMaterial*> MaterialsToOverride;
 
@@ -4565,7 +4480,7 @@ void FMaterialEditor::SetScalarParameterDefaultOnDependentMaterials(FName Parame
 	{
 		UMaterial* CurrentMaterial = MaterialsToOverride[MaterialIndex];
 
-		CurrentMaterial->OverrideScalarParameterDefault(ParameterName, Value, bOverride, FeatureLevel);
+		CurrentMaterial->OverrideNumericParameterDefault(Type, ParameterName, Value, bOverride, FeatureLevel);
 	}
 
 	// Update MI's that reference any of the materials affected
@@ -4581,21 +4496,20 @@ void FMaterialEditor::SetScalarParameterDefaultOnDependentMaterials(FName Parame
 
 			if (MaterialsToOverride.Contains(BaseMaterial))
 			{
-				CurrentMaterialInstance->OverrideScalarParameterDefault(ParameterName, Value, bOverride, FeatureLevel);
+				CurrentMaterialInstance->OverrideNumericParameterDefault(Type, ParameterName, Value, bOverride, FeatureLevel);
 			}
 		}
 	}
 }
 
-void FMaterialEditor::OnScalarParameterDefaultChanged(class UMaterialExpression* Expression, FName ParameterName, float Value)
+void FMaterialEditor::OnNumericParameterDefaultChanged(class UMaterialExpression* Expression, EMaterialParameterType Type, FName ParameterName, const UE::Shader::FValue& Value)
 {
 	check(Expression);
 
 	if (Expression->Material == Material && OriginalMaterial)
 	{
-		SetScalarParameterDefaultOnDependentMaterials(ParameterName, Value, true);
-
-		OverriddenScalarParametersToRevert.AddUnique(ParameterName);
+		SetNumericParameterDefaultOnDependentMaterials(Type, ParameterName, Value, true);
+		OverriddenNumericParametersToRevert.Add(MakeTuple(Type, ParameterName));
 	}
 
 	OnParameterDefaultChanged();
diff --git a/Engine/Source/Editor/MaterialEditor/Private/MaterialEditor.h b/Engine/Source/Editor/MaterialEditor/Private/MaterialEditor.h
--- a/Engine/Source/Editor/MaterialEditor/Private/MaterialEditor.h
+++ b/Engine/Source/Editor/MaterialEditor/Private/MaterialEditor.h
@@ -127,38 +127,11 @@ public:
 		return UMaterial::GetDefaultMaterial(MD_Surface)->GetRenderProxy();
 	}
 
-	virtual bool GetVectorValue(const FHashedMaterialParameterInfo& ParameterInfo, FLinearColor* OutValue, const FMaterialRenderContext& Context) const override
+	virtual bool GetParameterValue(EMaterialParameterType Type, const FHashedMaterialParameterInfo& ParameterInfo, FMaterialParameterValue& OutValue, const FMaterialRenderContext& Context) const override
 	{
 		if (Expression.IsValid() && Expression->Material)
 		{
-			return Expression->Material->GetRenderProxy()->GetVectorValue(ParameterInfo, OutValue, Context);
-		}
-		return false;
-	}
-
-	virtual bool GetScalarValue(const FHashedMaterialParameterInfo& ParameterInfo, float* OutValue, const FMaterialRenderContext& Context) const override
-	{
-		if (Expression.IsValid() && Expression->Material)
-		{
-			return Expression->Material->GetRenderProxy()->GetScalarValue(ParameterInfo, OutValue, Context);
-		}
-		return false;
-	}
-
-	virtual bool GetTextureValue(const FHashedMaterialParameterInfo& ParameterInfo,const UTexture** OutValue, const FMaterialRenderContext& Context) const override
-	{
-		if (Expression.IsValid() && Expression->Material)
-		{
-			return Expression->Material->GetRenderProxy()->GetTextureValue(ParameterInfo, OutValue, Context);
-		}
-		return false;
-	}
-
-	virtual bool GetTextureValue(const FHashedMaterialParameterInfo& ParameterInfo, const URuntimeVirtualTexture** OutValue, const FMaterialRenderContext& Context) const override
-	{
-		if (Expression.IsValid() && Expression->Material)
-		{
-			return Expression->Material->GetRenderProxy()->GetTextureValue(ParameterInfo, OutValue, Context);
+			return Expression->Material->GetRenderProxy()->GetParameterValue(Type, ParameterInfo, OutValue, Context);
 		}
 		return false;
 	}
@@ -833,12 +806,10 @@ private:
 	/** Callback when an asset is imported */
 	void OnAssetPostImport(UFactory* InFactory, UObject* InObject);
 
-	void OnVectorParameterDefaultChanged(class UMaterialExpression*, FName ParameterName, const FLinearColor& Value);
-	void OnScalarParameterDefaultChanged(class UMaterialExpression*, FName ParameterName, float Value);
+	void OnNumericParameterDefaultChanged(class UMaterialExpression*, EMaterialParameterType Type, FName ParameterName, const UE::Shader::FValue& Value);
 	void OnParameterDefaultChanged();
 
-	void SetVectorParameterDefaultOnDependentMaterials(FName ParameterName, const FLinearColor& Value, bool bOverride);
-	void SetScalarParameterDefaultOnDependentMaterials(FName ParameterName, float, bool bOverride);
+	void SetNumericParameterDefaultOnDependentMaterials(EMaterialParameterType Type, FName ParameterName, const UE::Shader::FValue& Value, bool bOverride);
 
 	// FEditorUndoClient Interface
 	virtual void PostUndo(bool bSuccess) override;
@@ -980,8 +951,7 @@ private:
 	/** Used to store material errors */
 	TArray<TSharedPtr<FMaterialInfo>> MaterialInfoList;
 
-	TArray<FName> OverriddenVectorParametersToRevert;
-	TArray<FName> OverriddenScalarParametersToRevert;
+	TSet<TTuple<EMaterialParameterType, FName>> OverriddenNumericParametersToRevert;
 
 	/** If true, don't render connectors that are not connected to anything. */
 	bool bHideUnusedConnectors;
diff --git a/Engine/Source/Editor/UnrealEd/Private/Lightmass/LightmassRender.cpp b/Engine/Source/Editor/UnrealEd/Private/Lightmass/LightmassRender.cpp
--- a/Engine/Source/Editor/UnrealEd/Private/Lightmass/LightmassRender.cpp
+++ b/Engine/Source/Editor/UnrealEd/Private/Lightmass/LightmassRender.cpp
@@ -330,24 +330,9 @@ public:
 		return UMaterial::GetDefaultMaterial(MD_Surface)->GetRenderProxy();
 	}
 
-	virtual bool GetVectorValue(const FHashedMaterialParameterInfo& ParameterInfo, FLinearColor* OutValue, const FMaterialRenderContext& Context) const override
+	virtual bool GetParameterValue(EMaterialParameterType Type, const FHashedMaterialParameterInfo& ParameterInfo, FMaterialParameterValue& OutValue, const FMaterialRenderContext& Context) const override
 	{
-		return MaterialInterface->GetRenderProxy()->GetVectorValue(ParameterInfo, OutValue, Context);
-	}
-
-	virtual bool GetScalarValue(const FHashedMaterialParameterInfo& ParameterInfo, float* OutValue, const FMaterialRenderContext& Context) const override
-	{
-		return MaterialInterface->GetRenderProxy()->GetScalarValue(ParameterInfo, OutValue, Context);
-	}
-
-	virtual bool GetTextureValue(const FHashedMaterialParameterInfo& ParameterInfo,const UTexture** OutValue, const FMaterialRenderContext& Context) const override
-	{
-		return MaterialInterface->GetRenderProxy()->GetTextureValue(ParameterInfo,OutValue,Context);
-	}
-
-	virtual bool GetTextureValue(const FHashedMaterialParameterInfo& ParameterInfo, const URuntimeVirtualTexture** OutValue, const FMaterialRenderContext& Context) const override
-	{
-		return MaterialInterface->GetRenderProxy()->GetTextureValue(ParameterInfo, OutValue, Context);
+		return MaterialInterface->GetRenderProxy()->GetParameterValue(Type, ParameterInfo, OutValue, Context);
 	}
 
 	// Material properties.
diff --git a/Engine/Source/Runtime/Engine/Classes/Materials/Material.h b/Engine/Source/Runtime/Engine/Classes/Materials/Material.h
--- a/Engine/Source/Runtime/Engine/Classes/Materials/Material.h
+++ b/Engine/Source/Runtime/Engine/Classes/Materials/Material.h
@@ -1047,8 +1047,7 @@ public:
 	ENGINE_API virtual void GetUsedTextures(TArray<UTexture*>& OutTextures, EMaterialQualityLevel::Type QualityLevel, bool bAllQualityLevels, ERHIFeatureLevel::Type FeatureLevel, bool bAllFeatureLevels) const override;
 	ENGINE_API virtual void GetUsedTexturesAndIndices(TArray<UTexture*>& OutTextures, TArray< TArray<int32> >& OutIndices, EMaterialQualityLevel::Type QualityLevel, ERHIFeatureLevel::Type FeatureLevel) const override;
 	ENGINE_API virtual void OverrideTexture(const UTexture* InTextureToOverride, UTexture* OverrideTexture, ERHIFeatureLevel::Type InFeatureLevel) override;
-	ENGINE_API virtual void OverrideVectorParameterDefault(const FHashedMaterialParameterInfo& ParameterInfo, const FLinearColor& Value, bool bOverride, ERHIFeatureLevel::Type FeatureLevel) override;
-	ENGINE_API virtual void OverrideScalarParameterDefault(const FHashedMaterialParameterInfo& ParameterInfo, float Value, bool bOverride, ERHIFeatureLevel::Type FeatureLevel) override;
+	ENGINE_API virtual void OverrideNumericParameterDefault(EMaterialParameterType Type, const FHashedMaterialParameterInfo& ParameterInfo, const UE::Shader::FValue& Value, bool bOverride, ERHIFeatureLevel::Type FeatureLevel) override;
 	ENGINE_API virtual bool CheckMaterialUsage(const EMaterialUsage Usage) override;
 	ENGINE_API virtual bool CheckMaterialUsage_Concurrent(const EMaterialUsage Usage) const override;
 	ENGINE_API virtual FMaterialResource* AllocateResource();
diff --git a/Engine/Source/Runtime/Engine/Classes/Materials/MaterialInstance.h b/Engine/Source/Runtime/Engine/Classes/Materials/MaterialInstance.h
--- a/Engine/Source/Runtime/Engine/Classes/Materials/MaterialInstance.h
+++ b/Engine/Source/Runtime/Engine/Classes/Materials/MaterialInstance.h
@@ -560,13 +560,13 @@ public:
 	virtual ENGINE_API FMaterialResource* GetMaterialResource(ERHIFeatureLevel::Type InFeatureLevel, EMaterialQualityLevel::Type QualityLevel = EMaterialQualityLevel::Num) override;
 	virtual ENGINE_API const FMaterialResource* GetMaterialResource(ERHIFeatureLevel::Type InFeatureLevel, EMaterialQualityLevel::Type QualityLevel = EMaterialQualityLevel::Num) const override;
 
+	ENGINE_API bool GetParameterOverrideValue(EMaterialParameterType Type, const FMemoryImageMaterialParameterInfo& ParameterInfo, FMaterialParameterMetadata& OutValue, EMaterialGetParameterValueFlags Flags = EMaterialGetParameterValueFlags::None) const;
 	virtual ENGINE_API bool GetParameterValue(EMaterialParameterType Type, const FMemoryImageMaterialParameterInfo& ParameterInfo, FMaterialParameterMetadata& OutValue, EMaterialGetParameterValueFlags Flags) const override;
 
 	virtual ENGINE_API void GetUsedTextures(TArray<UTexture*>& OutTextures, EMaterialQualityLevel::Type QualityLevel, bool bAllQualityLevels, ERHIFeatureLevel::Type FeatureLevel, bool bAllFeatureLevels) const override;
 	virtual ENGINE_API void GetUsedTexturesAndIndices(TArray<UTexture*>& OutTextures, TArray< TArray<int32> >& OutIndices, EMaterialQualityLevel::Type QualityLevel, ERHIFeatureLevel::Type FeatureLevel) const;
 	virtual ENGINE_API void OverrideTexture(const UTexture* InTextureToOverride, UTexture* OverrideTexture, ERHIFeatureLevel::Type InFeatureLevel) override;
-	virtual ENGINE_API void OverrideVectorParameterDefault(const FHashedMaterialParameterInfo& ParameterInfo, const FLinearColor& Value, bool bOverride, ERHIFeatureLevel::Type FeatureLevel) override;
-	virtual ENGINE_API void OverrideScalarParameterDefault(const FHashedMaterialParameterInfo& ParameterInfo, float Value, bool bOverride, ERHIFeatureLevel::Type FeatureLevel) override;
+	virtual ENGINE_API void OverrideNumericParameterDefault(EMaterialParameterType Type, const FHashedMaterialParameterInfo& ParameterInfo, const UE::Shader::FValue& Value, bool bOverride, ERHIFeatureLevel::Type FeatureLevel) override;
 	virtual ENGINE_API bool CheckMaterialUsage(const EMaterialUsage Usage) override;
 	virtual ENGINE_API bool CheckMaterialUsage_Concurrent(const EMaterialUsage Usage) const override;
 #if WITH_EDITORONLY_DATA
@@ -816,6 +816,7 @@ protected:
 	/**
 	 * Internal interface for setting / updating values for material instances.
 	 */
+	void SetParameterValueInternal(const FMaterialParameterInfo& ParameterInfo, const FMaterialParameterMetadata& Meta, EMaterialSetParameterValueFlags Flags = EMaterialSetParameterValueFlags::None);
 	void SetVectorParameterValueInternal(const FMaterialParameterInfo& ParameterInfo, FLinearColor Value);
 	bool SetVectorParameterByIndexInternal(int32 ParameterIndex, FLinearColor Value);
 	bool SetScalarParameterByIndexInternal(int32 ParameterIndex, float Value);
diff --git a/Engine/Source/Runtime/Engine/Classes/Materials/MaterialInterface.h b/Engine/Source/Runtime/Engine/Classes/Materials/MaterialInterface.h
--- a/Engine/Source/Runtime/Engine/Classes/Materials/MaterialInterface.h
+++ b/Engine/Source/Runtime/Engine/Classes/Materials/MaterialInterface.h
@@ -400,9 +400,7 @@ public:
 	 * This is used to implement realtime previewing of parameter defaults. 
 	 * Handles updating dependent MI's and cached uniform expressions.
 	 */
-	virtual void OverrideVectorParameterDefault(const FHashedMaterialParameterInfo& ParameterInfo, const FLinearColor& Value, bool bOverride, ERHIFeatureLevel::Type FeatureLevel) PURE_VIRTUAL(UMaterialInterface::OverrideTexture, return;);
-	virtual void OverrideScalarParameterDefault(const FHashedMaterialParameterInfo& ParameterInfo, float Value, bool bOverride, ERHIFeatureLevel::Type FeatureLevel) PURE_VIRTUAL(UMaterialInterface::OverrideTexture, return;);
-
+	virtual void OverrideNumericParameterDefault(EMaterialParameterType Type, const FHashedMaterialParameterInfo& ParameterInfo, const UE::Shader::FValue& Value, bool bOverride, ERHIFeatureLevel::Type FeatureLevel) PURE_VIRTUAL(UMaterialInterface::OverrideNumericParameterDefault, return;);
 	/**
 	 * DEPRECATED: Returns default value of the given parameter
 	 */
diff --git a/Engine/Source/Runtime/Engine/Private/EngineGlobals.cpp b/Engine/Source/Runtime/Engine/Private/EngineGlobals.cpp
--- a/Engine/Source/Runtime/Engine/Private/EngineGlobals.cpp
+++ b/Engine/Source/Runtime/Engine/Private/EngineGlobals.cpp
@@ -59,8 +59,7 @@ FEditorSupportDelegates::FOnWindowsMessage FEditorSupportDelegates::PreWindowsMe
 FEditorSupportDelegates::FOnWindowsMessage FEditorSupportDelegates::PostWindowsMessage;
 /** Sent after the usages flags on a material have changed*/
 FEditorSupportDelegates::FOnMaterialUsageFlagsChanged FEditorSupportDelegates::MaterialUsageFlagsChanged;
-FEditorSupportDelegates::FOnVectorParameterDefaultChanged FEditorSupportDelegates::VectorParameterDefaultChanged;
-FEditorSupportDelegates::FOnScalarParameterDefaultChanged FEditorSupportDelegates::ScalarParameterDefaultChanged;
+FEditorSupportDelegates::FOnNumericParameterDefaultChanged FEditorSupportDelegates::NumericParameterDefaultChanged;
 
 #endif // WITH_EDITOR
 
diff --git a/Engine/Source/Runtime/Engine/Private/HLSLTree/HLSLTree.cpp b/Engine/Source/Runtime/Engine/Private/HLSLTree/HLSLTree.cpp
--- a/Engine/Source/Runtime/Engine/Private/HLSLTree/HLSLTree.cpp
+++ b/Engine/Source/Runtime/Engine/Private/HLSLTree/HLSLTree.cpp
@@ -100,7 +100,7 @@ void UE::HLSLTree::FCodeWriter::Append(const FCodeWriter& InWriter)
 
 void UE::HLSLTree::FCodeWriter::WriteConstant(const Shader::FValue& Value)
 {
-	auto ToString = [](bool v)
+	auto ToString = [](uint8 v)
 	{
 		return v ? TEXT("true") : TEXT("false");
 	};
diff --git a/Engine/Source/Runtime/Engine/Private/HLSLTree/HLSLTreeCommon.cpp b/Engine/Source/Runtime/Engine/Private/HLSLTree/HLSLTreeCommon.cpp
--- a/Engine/Source/Runtime/Engine/Private/HLSLTree/HLSLTreeCommon.cpp
+++ b/Engine/Source/Runtime/Engine/Private/HLSLTree/HLSLTreeCommon.cpp
@@ -89,7 +89,7 @@ bool UE::HLSLTree::FExpressionParameter::EmitCode(FEmitContext& Context, FExpres
 	if (OutResult.Type == Shader::EValueType::Bool1)
 	{
 		const FMaterialParameterInfo ParameterInfo(Declaration->Name);
-		bool bValue = Declaration->DefaultValue.Component[0].Bool;
+		bool bValue = Declaration->DefaultValue.Component[0].AsBool();
 		for (const FStaticSwitchParameter& Parameter : Context.StaticParameters->StaticSwitchParameters)
 		{
 			if (Parameter.ParameterInfo == ParameterInfo)
@@ -103,17 +103,19 @@ bool UE::HLSLTree::FExpressionParameter::EmitCode(FEmitContext& Context, FExpres
 	}
 	else if (OutResult.Type == Shader::EValueType::Float1)
 	{
-		const int32 ParameterIndex = Context.MaterialCompilationOutput->UniformExpressionSet.FindOrAddScalarParameter(Declaration->Name, Declaration->DefaultValue.Component[0].Float);
+		/*const int32 ParameterIndex = Context.MaterialCompilationOutput->UniformExpressionSet.FindOrAddScalarParameter(Declaration->Name, Declaration->DefaultValue.Component[0].Float);
 		check(ParameterIndex >= 0 && ParameterIndex <= 0xffff);
 		OutResult.EvaluationType = EExpressionEvaluationType::Preshader;
-		OutResult.Preshader.WriteOpcode(Shader::EPreshaderOpcode::ScalarParameter).Write((uint16)ParameterIndex);
+		OutResult.Preshader.WriteOpcode(Shader::EPreshaderOpcode::ScalarParameter).Write((uint16)ParameterIndex);*/
+		check(false);
 	}
 	else
 	{
-		const int32 ParameterIndex = Context.MaterialCompilationOutput->UniformExpressionSet.FindOrAddVectorParameter(Declaration->Name, Declaration->DefaultValue.AsLinearColor());
+		/*const int32 ParameterIndex = Context.MaterialCompilationOutput->UniformExpressionSet.FindOrAddVectorParameter(Declaration->Name, Declaration->DefaultValue.AsLinearColor());
 		check(ParameterIndex >= 0 && ParameterIndex <= 0xffff);
 		OutResult.EvaluationType = EExpressionEvaluationType::Preshader;
-		OutResult.Preshader.WriteOpcode(Shader::EPreshaderOpcode::VectorParameter).Write((uint16)ParameterIndex);
+		OutResult.Preshader.WriteOpcode(Shader::EPreshaderOpcode::VectorParameter).Write((uint16)ParameterIndex);*/
+		check(false);
 	}
 	return true;
 }
diff --git a/Engine/Source/Runtime/Engine/Private/Materials/HLSLMaterialTranslator.cpp b/Engine/Source/Runtime/Engine/Private/Materials/HLSLMaterialTranslator.cpp
--- a/Engine/Source/Runtime/Engine/Private/Materials/HLSLMaterialTranslator.cpp
+++ b/Engine/Source/Runtime/Engine/Private/Materials/HLSLMaterialTranslator.cpp
@@ -65,6 +65,42 @@ static inline EMaterialValueType GetVectorType(uint32 NumComponents)
 	};
 }
 
+static inline EMaterialValueType GetMaterialValueType(EMaterialParameterType Type)
+{
+	switch (Type)
+	{
+	case EMaterialParameterType::Scalar: return MCT_Float;
+	case EMaterialParameterType::Vector: return MCT_Float4;
+	default: checkNoEntry(); return MCT_Unknown;
+	}
+}
+
+static inline EMaterialValueType GetMaterialValueType(UE::Shader::EValueType Type)
+{
+	switch (Type)
+	{
+	case UE::Shader::EValueType::Void: return MCT_VoidStatement;
+
+	case UE::Shader::EValueType::Float1: return MCT_Float;
+	case UE::Shader::EValueType::Float2: return MCT_Float2;
+	case UE::Shader::EValueType::Float3: return MCT_Float3;
+	case UE::Shader::EValueType::Float4: return MCT_Float4;
+
+	case UE::Shader::EValueType::Int1: return MCT_Float;
+	case UE::Shader::EValueType::Int2: return MCT_Float2;
+	case UE::Shader::EValueType::Int3: return MCT_Float3;
+	case UE::Shader::EValueType::Int4: return MCT_Float4;
+
+	case UE::Shader::EValueType::Bool1: return MCT_Float;
+	case UE::Shader::EValueType::Bool2: return MCT_Float2;
+	case UE::Shader::EValueType::Bool3: return MCT_Float3;
+	case UE::Shader::EValueType::Bool4: return MCT_Float4;
+
+	case UE::Shader::EValueType::MaterialAttributes: return MCT_MaterialAttributes;
+	default: checkNoEntry(); return MCT_Unknown;
+	}
+}
+
 static inline int32 SwizzleComponentToIndex(TCHAR Component)
 {
 	switch (Component)
@@ -3976,41 +4012,33 @@ int32 FHLSLMaterialTranslator::AccessCollectionParameter(UMaterialParameterColle
 		ComponentIndex == -1 ? true : ComponentIndex % 4 == 3);
 }
 
-int32 FHLSLMaterialTranslator::ScalarParameter(FName ParameterName, float DefaultValue)
+int32 FHLSLMaterialTranslator::GenericParameter(EMaterialParameterType ParameterType, FName ParameterName, const UE::Shader::FValue& DefaultValue)
 {
-	FMaterialParameterInfo ParameterInfo = GetParameterAssociationInfo();
-	ParameterInfo.Name = ParameterName;
-	int32 ParameterIndex = INDEX_NONE;
-	for (int32 i = 0; i < MaterialCompilationOutput.UniformExpressionSet.UniformScalarParameters.Num(); ++i)
+	const UE::Shader::EValueType ValueType = GetShaderValueType(ParameterType);
+	check(DefaultValue.GetType() == ValueType);
+
+	const uint32* PrevDefaultOffset = DefaultUniformValues.Find(DefaultValue);
+	uint32 DefaultOffset;
+	if (PrevDefaultOffset)
 	{
-		const FMaterialScalarParameterInfo& Parameter = MaterialCompilationOutput.UniformExpressionSet.UniformScalarParameters[i];
-		if (Parameter.ParameterInfo == ParameterInfo)
-		{
-			ParameterIndex = i;
-			break;
-		}
+		DefaultOffset = *PrevDefaultOffset;
 	}
-	if (ParameterIndex == INDEX_NONE)
+	else
 	{
-		ParameterIndex = MaterialCompilationOutput.UniformExpressionSet.UniformScalarParameters.Num();
-		FMaterialScalarParameterInfo& Parameter = MaterialCompilationOutput.UniformExpressionSet.UniformScalarParameters.AddDefaulted_GetRef();
-		Parameter.ParameterInfo = ParameterInfo;
-		Parameter.DefaultValue = DefaultValue;
+		DefaultOffset = MaterialCompilationOutput.UniformExpressionSet.DefaultValues.Num();
+		DefaultUniformValues.Add(DefaultValue, DefaultOffset);
+		UE::Shader::FMemoryImageValue DefaultValueMemory = DefaultValue.AsMemoryImage();
+		MaterialCompilationOutput.UniformExpressionSet.DefaultValues.Append(DefaultValueMemory.Bytes, DefaultValueMemory.Size);
 	}
 
-	return AddUniformExpression(new FMaterialUniformExpressionScalarParameter(ParameterInfo, ParameterIndex), MCT_Float, TEXT(""));
-}
-
-int32 FHLSLMaterialTranslator::VectorParameter(FName ParameterName, const FLinearColor& DefaultValue)
-{
 	FMaterialParameterInfo ParameterInfo = GetParameterAssociationInfo();
 	ParameterInfo.Name = ParameterName;
 
 	int32 ParameterIndex = INDEX_NONE;
-	for (int32 i = 0; i < MaterialCompilationOutput.UniformExpressionSet.UniformVectorParameters.Num(); ++i)
+	for (int32 i = 0; i < MaterialCompilationOutput.UniformExpressionSet.UniformNumericParameters.Num(); ++i)
 	{
-		const FMaterialVectorParameterInfo& Parameter = MaterialCompilationOutput.UniformExpressionSet.UniformVectorParameters[i];
-		if (Parameter.ParameterInfo == ParameterInfo)
+		const FMaterialNumericParameterInfo& Parameter = MaterialCompilationOutput.UniformExpressionSet.UniformNumericParameters[i];
+		if (Parameter.ParameterType == ParameterType && Parameter.ParameterInfo == ParameterInfo)
 		{
 			ParameterIndex = i;
 			break;
@@ -4018,13 +4046,14 @@ int32 FHLSLMaterialTranslator::VectorParameter(FName ParameterName, const FLinea
 	}
 	if (ParameterIndex == INDEX_NONE)
 	{
-		ParameterIndex = MaterialCompilationOutput.UniformExpressionSet.UniformVectorParameters.Num();
-		FMaterialVectorParameterInfo& Parameter = MaterialCompilationOutput.UniformExpressionSet.UniformVectorParameters.AddDefaulted_GetRef();
+		ParameterIndex = MaterialCompilationOutput.UniformExpressionSet.UniformNumericParameters.Num();
+		FMaterialNumericParameterInfo& Parameter = MaterialCompilationOutput.UniformExpressionSet.UniformNumericParameters.AddDefaulted_GetRef();
 		Parameter.ParameterInfo = ParameterInfo;
-		Parameter.DefaultValue = DefaultValue;
+		Parameter.ParameterType = ParameterType;
+		Parameter.DefaultValueOffset = DefaultOffset;
 	}
 
-	return AddUniformExpression(new FMaterialUniformExpressionVectorParameter(ParameterInfo, ParameterIndex), MCT_Float4, TEXT(""));
+	return AddUniformExpression(new FMaterialUniformExpressionGenericParameter(ParameterInfo, ParameterIndex), GetMaterialValueType(ParameterType), TEXT(""));
 }
 
 int32 FHLSLMaterialTranslator::Constant(float X)
diff --git a/Engine/Source/Runtime/Engine/Private/Materials/HLSLMaterialTranslator.h b/Engine/Source/Runtime/Engine/Private/Materials/HLSLMaterialTranslator.h
--- a/Engine/Source/Runtime/Engine/Private/Materials/HLSLMaterialTranslator.h
+++ b/Engine/Source/Runtime/Engine/Private/Materials/HLSLMaterialTranslator.h
@@ -30,6 +30,7 @@
 #include "VT/RuntimeVirtualTexture.h"
 #include "Field/FieldSystemTypes.h"
 #include "Containers/Map.h"
+#include "Shader/ShaderTypes.h"
 
 #if WITH_EDITORONLY_DATA
 #include "Materials/MaterialExpressionSceneTexture.h"
@@ -277,6 +278,7 @@ protected:
 	TArray<TRefCountPtr<FMaterialUniformExpression> > UniformScalarExpressions;
 	TArray<TRefCountPtr<FMaterialUniformExpressionTexture> > UniformTextureExpressions[NumMaterialTextureParameterTypes];
 	TArray<TRefCountPtr<FMaterialUniformExpressionExternalTexture>> UniformExternalTextureExpressions;
+	TMap<UE::Shader::FValue, uint32> DefaultUniformValues;
 
 	/** Parameter collections referenced by this material.  The position in this array is used as an index on the shader parameter. */
 	TArray<UMaterialParameterCollection*> ParameterCollections;
@@ -674,10 +676,7 @@ protected:
 
 	virtual int32 AccessCollectionParameter(UMaterialParameterCollection* ParameterCollection, int32 ParameterIndex, int32 ComponentIndex) override;
 
-	virtual int32 ScalarParameter(FName ParameterName, float DefaultValue) override;
-
-	virtual int32 VectorParameter(FName ParameterName, const FLinearColor& DefaultValue) override;
-
+	virtual int32 GenericParameter(EMaterialParameterType ParameterType, FName ParameterName, const UE::Shader::FValue& DefaultValue) override;
 	virtual int32 Constant(float X) override;
 	virtual int32 Constant2(float X, float Y) override;
 	virtual int32 Constant3(float X, float Y, float Z) override;
diff --git a/Engine/Source/Runtime/Engine/Private/Materials/Material.cpp b/Engine/Source/Runtime/Engine/Private/Materials/Material.cpp
--- a/Engine/Source/Runtime/Engine/Private/Materials/Material.cpp
+++ b/Engine/Source/Runtime/Engine/Private/Materials/Material.cpp
@@ -408,30 +408,18 @@ public:
 		return Material;
 	}
 
-	virtual bool GetVectorValue(const FHashedMaterialParameterInfo& ParameterInfo, FLinearColor* OutValue, const FMaterialRenderContext& Context) const
+	virtual bool GetParameterValue(EMaterialParameterType Type, const FHashedMaterialParameterInfo& ParameterInfo, FMaterialParameterValue& OutValue, const FMaterialRenderContext& Context) const override
 	{
 		const FMaterialResource* MaterialResource = Material->GetMaterialResource(Context.Material.GetFeatureLevel());
-		if(MaterialResource && MaterialResource->GetRenderingThreadShaderMap())
-		{
-			return false;
-		}
-		else
-		{
-			return GetFallbackRenderProxy().GetVectorValue(ParameterInfo, OutValue, Context);
-		}
-	}
-	virtual bool GetScalarValue(const FHashedMaterialParameterInfo& ParameterInfo, float* OutValue, const FMaterialRenderContext& Context) const
-	{
-		const FMaterialResource* MaterialResource = Material->GetMaterialResource(Context.Material.GetFeatureLevel());
-		if(MaterialResource && MaterialResource->GetRenderingThreadShaderMap())
+		if (MaterialResource && MaterialResource->GetRenderingThreadShaderMap())
 		{
 			static FName NameSubsurfaceProfile(TEXT("__SubsurfaceProfile"));
-			if (ParameterInfo.Name == NameSubsurfaceProfile)
+			if (Type == EMaterialParameterType::Scalar && ParameterInfo.Name == NameSubsurfaceProfile)
 			{
 				const USubsurfaceProfile* MySubsurfaceProfileRT = GetSubsurfaceProfileRT();
 
 				int32 AllocationId = 0;
-				if(MySubsurfaceProfileRT)
+				if (MySubsurfaceProfileRT)
 				{
 					// can be optimized (cached)
 					AllocationId = GSubsurfaceProfileTextureObject.FindAllocationId(MySubsurfaceProfileRT);
@@ -442,8 +430,7 @@ public:
 					AllocationId = 0;
 				}
 
-				*OutValue = AllocationId / 255.0f;
-
+				OutValue = AllocationId / 255.0f;
 				return true;
 			}
 
@@ -451,32 +438,7 @@ public:
 		}
 		else
 		{
-			return GetFallbackRenderProxy().GetScalarValue(ParameterInfo, OutValue, Context);
-		}
-	}
-	virtual bool GetTextureValue(const FHashedMaterialParameterInfo& ParameterInfo,const UTexture** OutValue, const FMaterialRenderContext& Context) const
-	{
-		const FMaterialResource* MaterialResource = Material->GetMaterialResource(Context.Material.GetFeatureLevel());
-		if(MaterialResource && MaterialResource->GetRenderingThreadShaderMap())
-		{
-			return false;
-		}
-		else
-		{
-			return GetFallbackRenderProxy().GetTextureValue(ParameterInfo,OutValue,Context);
-		}
-	}
-
-	virtual bool GetTextureValue(const FHashedMaterialParameterInfo& ParameterInfo, const URuntimeVirtualTexture** OutValue, const FMaterialRenderContext& Context) const
-	{
-		const FMaterialResource* MaterialResource = Material->GetMaterialResource(Context.Material.GetFeatureLevel());
-		if (MaterialResource && MaterialResource->GetRenderingThreadShaderMap())
-		{
-			return false;
-		}
-		else
-		{
-			return GetFallbackRenderProxy().GetTextureValue(ParameterInfo, OutValue, Context);
+			return GetFallbackRenderProxy().GetParameterValue(Type, ParameterInfo, OutValue, Context);
 		}
 	}
 
@@ -1103,23 +1065,26 @@ void UMaterial::GetUsedTextures(TArray<UTexture*>& OutTextures, EMaterialQuality
 			if (MaterialInstance)
 			{
 				// Also look for any scalar parameters that are acting as lookups for an atlas texture, and store the atlas texture
-				const TArrayView<const FMaterialScalarParameterInfo> AtlasExpressions[1] =
+				const TArrayView<const FMaterialNumericParameterInfo> AtlasExpressions[1] =
 				{
-					CurrentResource->GetUniformScalarParameterExpressions()
+					CurrentResource->GetUniformNumericParameterExpressions()
 				};
+
 				for (int32 TypeIndex = 0; TypeIndex < UE_ARRAY_COUNT(AtlasExpressions); TypeIndex++)
 				{
 					// Iterate over each of the material's texture expressions.
-					for (const FMaterialScalarParameterInfo& Parameter : AtlasExpressions[TypeIndex])
+					for (const FMaterialNumericParameterInfo& Parameter : AtlasExpressions[TypeIndex])
 					{
-						bool bIsUsedAsAtlasPosition;
-						TSoftObjectPtr<class UCurveLinearColor> Curve;
-						TSoftObjectPtr<class UCurveLinearColorAtlas> Atlas;
-						MaterialInstance->IsScalarParameterUsedAsAtlasPosition(Parameter.ParameterInfo, bIsUsedAsAtlasPosition, Curve, Atlas);
-
-						if (Atlas)
+						if (Parameter.ParameterType == EMaterialParameterType::Scalar)
 						{
-							OutTextures.AddUnique(Atlas.Get());
+							bool bIsUsedAsAtlasPosition;
+							TSoftObjectPtr<class UCurveLinearColor> Curve;
+							TSoftObjectPtr<class UCurveLinearColorAtlas> Atlas;
+							MaterialInstance->IsScalarParameterUsedAsAtlasPosition(Parameter.ParameterInfo, bIsUsedAsAtlasPosition, Curve, Atlas);
+							if (Atlas)
+							{
+								OutTextures.AddUnique(Atlas.Get());
+							}
 						}
 					}
 				}
@@ -1275,48 +1240,20 @@ void UMaterial::OverrideTexture(const UTexture* InTextureToOverride, UTexture* O
 #endif // WITH_EDITOR
 }
 
-void UMaterial::OverrideVectorParameterDefault(const FHashedMaterialParameterInfo& ParameterInfo, const FLinearColor& Value, bool bOverride, ERHIFeatureLevel::Type InFeatureLevel)
-{
-#if WITH_EDITOR
-	FMaterialResource* Resource = GetMaterialResource(InFeatureLevel);
-	if (Resource)
-	{
-		Resource->TransientOverrides.SetVectorOverride(ParameterInfo, Value, bOverride);
-
-		const TArrayView<const FMaterialVectorParameterInfo> Parameters = Resource->GetUniformVectorParameterExpressions();
-		bool bShouldRecacheMaterialExpressions = false;
-		// Iterate over each of the material's vector expressions.
-		for (int32 i = 0; i < Parameters.Num(); ++i)
-		{
-			const FMaterialVectorParameterInfo& Parameter = Parameters[i];
-			if (Parameter.ParameterInfo == ParameterInfo)
-			{
-				bShouldRecacheMaterialExpressions = true;
-			}
-		}
-
-		if (bShouldRecacheMaterialExpressions)
-		{
-			RecacheUniformExpressions(false);
-		}
-	}
-#endif // #if WITH_EDITOR
-}
-
-void UMaterial::OverrideScalarParameterDefault(const FHashedMaterialParameterInfo& ParameterInfo, float Value, bool bOverride, ERHIFeatureLevel::Type InFeatureLevel)
+void UMaterial::OverrideNumericParameterDefault(EMaterialParameterType Type, const FHashedMaterialParameterInfo& ParameterInfo, const UE::Shader::FValue& Value, bool bOverride, ERHIFeatureLevel::Type InFeatureLevel)
 {
 #if WITH_EDITOR
 	FMaterialResource* Resource = GetMaterialResource(InFeatureLevel);
 	if (Resource)
 	{
-		Resource->TransientOverrides.SetScalarOverride(ParameterInfo, Value, bOverride);
+		Resource->TransientOverrides.SetNumericOverride(Type, ParameterInfo, Value, bOverride);
 
-		const TArrayView<const FMaterialScalarParameterInfo> Parameters = Resource->GetUniformScalarParameterExpressions();
+		const TArrayView<const FMaterialNumericParameterInfo> Parameters = Resource->GetUniformNumericParameterExpressions();
 		bool bShouldRecacheMaterialExpressions = false;
 		// Iterate over each of the material's vector expressions.
 		for (int32 i = 0; i < Parameters.Num(); ++i)
 		{
-			const FMaterialScalarParameterInfo& Parameter = Parameters[i];
+			const FMaterialNumericParameterInfo& Parameter = Parameters[i];
 			if (Parameter.ParameterInfo == ParameterInfo)
 			{
 				bShouldRecacheMaterialExpressions = true;
diff --git a/Engine/Source/Runtime/Engine/Private/Materials/MaterialExpressions.cpp b/Engine/Source/Runtime/Engine/Private/Materials/MaterialExpressions.cpp
--- a/Engine/Source/Runtime/Engine/Private/Materials/MaterialExpressions.cpp
+++ b/Engine/Source/Runtime/Engine/Private/Materials/MaterialExpressions.cpp
@@ -7877,7 +7877,7 @@ void UMaterialExpressionVectorParameter::PostEditChangeProperty(FPropertyChanged
 	if (PropertyName == GET_MEMBER_NAME_STRING_CHECKED(UMaterialExpressionVectorParameter, DefaultValue))
 	{
 		// Callback into the editor
-		FEditorSupportDelegates::VectorParameterDefaultChanged.Broadcast(this, ParameterName, DefaultValue);
+		FEditorSupportDelegates::NumericParameterDefaultChanged.Broadcast(this, EMaterialParameterType::Vector, ParameterName, DefaultValue);
 	}
 	else if (PropertyName == GET_MEMBER_NAME_STRING_CHECKED(UMaterialExpressionVectorParameter, PrimitiveDataIndex))
 	{
@@ -8048,7 +8048,7 @@ void UMaterialExpressionChannelMaskParameter::PostEditChangeProperty(FPropertyCh
 			DefaultValue = FLinearColor(0.0f, 0.0f, 0.0f, 1.0f); break;
 		}
 
-		FEditorSupportDelegates::VectorParameterDefaultChanged.Broadcast(this, ParameterName, DefaultValue);
+		FEditorSupportDelegates::NumericParameterDefaultChanged.Broadcast(this, EMaterialParameterType::Vector, ParameterName, DefaultValue);
 	}
 	else if (PropertyName == GET_MEMBER_NAME_STRING_CHECKED(UMaterialExpressionVectorParameter, DefaultValue))
 	{
@@ -8199,7 +8199,7 @@ void UMaterialExpressionScalarParameter::PostEditChangeProperty(FPropertyChanged
 	if (PropertyName == GET_MEMBER_NAME_STRING_CHECKED(UMaterialExpressionScalarParameter, DefaultValue))
 	{
 		// Callback into the editor
-		FEditorSupportDelegates::ScalarParameterDefaultChanged.Broadcast(this, ParameterName, DefaultValue);
+		FEditorSupportDelegates::NumericParameterDefaultChanged.Broadcast(this, EMaterialParameterType::Scalar, ParameterName, DefaultValue);
 	}
 	else if (PropertyName == GET_MEMBER_NAME_STRING_CHECKED(UMaterialExpressionScalarParameter, PrimitiveDataIndex))
 	{
diff --git a/Engine/Source/Runtime/Engine/Private/Materials/MaterialInstance.cpp b/Engine/Source/Runtime/Engine/Private/Materials/MaterialInstance.cpp
--- a/Engine/Source/Runtime/Engine/Private/Materials/MaterialInstance.cpp
+++ b/Engine/Source/Runtime/Engine/Private/Materials/MaterialInstance.cpp
@@ -49,8 +49,6 @@
 #include "UObject/UE5MainStreamObjectVersion.h"
 #include "ShaderCompilerCore.h"
 
-PRAGMA_DISABLE_OPTIMIZATION
-
 #if ENABLE_COOK_STATS
 #include "ProfilingDebugging/ScopedTimers.h"
 namespace MaterialInstanceCookStats
@@ -237,16 +235,15 @@ UMaterialInterface* FMaterialInstanceResource::GetMaterialInterface() const
 	return Owner;
 }
 
-bool FMaterialInstanceResource::GetScalarValue(
-	const FHashedMaterialParameterInfo& ParameterInfo,
-	float* OutValue,
-	const FMaterialRenderContext& Context
-	) const
+bool FMaterialInstanceResource::GetParameterValue(EMaterialParameterType Type, const FHashedMaterialParameterInfo& ParameterInfo, FMaterialParameterValue& OutValue, const FMaterialRenderContext& Context) const
 {
 	checkSlow(IsInParallelRenderingThread());
 
 	static FName NameSubsurfaceProfile(TEXT("__SubsurfaceProfile"));
-	if (ParameterInfo.Name == NameSubsurfaceProfile)
+	bool bResult = false;
+
+	// Check for hard-coded parameters
+	if (Type == EMaterialParameterType::Scalar && ParameterInfo.Name == NameSubsurfaceProfile)
 	{
 		check(ParameterInfo.Association == EMaterialParameterAssociation::GlobalParameter);
 		const USubsurfaceProfile* MySubsurfaceProfileRT = GetSubsurfaceProfileRT();
@@ -262,94 +259,30 @@ bool FMaterialInstanceResource::GetScalarValue(
 			// no profile specified means we use the default one stored at [0] which is human skin
 			AllocationId = 0;
 		}
-		*OutValue = AllocationId / 255.0f;
-
-		return true;
-	}
-
-	const float* Value = RenderThread_FindParameterByName<float>(ParameterInfo);
-	if(Value)
-	{
-		*OutValue = *Value;
-		return true;
-	}
-	else if (Parent)
-	{
-		return Parent->GetRenderProxy()->GetScalarValue(ParameterInfo, OutValue, Context);
+		OutValue = AllocationId / 255.0f;
+		bResult = true;
 	}
-	else
-	{
-		return false;
-	}
-}
 
-bool FMaterialInstanceResource::GetVectorValue(
-	const FHashedMaterialParameterInfo& ParameterInfo,
-	FLinearColor* OutValue,
-	const FMaterialRenderContext& Context
-	) const
-{
-	checkSlow(IsInParallelRenderingThread());
-	const FLinearColor* Value = RenderThread_FindParameterByName<FLinearColor>(ParameterInfo);
-	if(Value)
+	if (!bResult)
 	{
-		*OutValue = *Value;
-		return true;
-	}
-	else if(Parent)
-	{
-		return Parent->GetRenderProxy()->GetVectorValue(ParameterInfo, OutValue, Context);
-	}
-	else
-	{
-		return false;
+		// Check for instances overrides
+		switch (Type)
+		{
+		case EMaterialParameterType::Scalar: bResult = RenderThread_GetParameterValue<float>(ParameterInfo, OutValue); break;
+		case EMaterialParameterType::Vector: bResult = RenderThread_GetParameterValue<FLinearColor>(ParameterInfo, OutValue); break;
+		case EMaterialParameterType::Texture: bResult = RenderThread_GetParameterValue<const UTexture*>(ParameterInfo, OutValue); break;
+		case EMaterialParameterType::RuntimeVirtualTexture: bResult = RenderThread_GetParameterValue<const URuntimeVirtualTexture*>(ParameterInfo, OutValue); break;
+		default: ensure(false); break; // other parameter types are not expected on the render thread
+		}
 	}
-}
 
-bool FMaterialInstanceResource::GetTextureValue(
-	const FHashedMaterialParameterInfo& ParameterInfo,
-	const UTexture** OutValue,
-	const FMaterialRenderContext& Context
-	) const
-{
-	checkSlow(IsInParallelRenderingThread());
-	const UTexture* const * Value = RenderThread_FindParameterByName<const UTexture*>(ParameterInfo);
-	if(Value && *Value)
-	{
-		*OutValue = *Value;
-		return true;
-	}
-	else if(Parent)
+	if (!bResult && Parent)
 	{
-		return Parent->GetRenderProxy()->GetTextureValue(ParameterInfo, OutValue, Context);
-	}
-	else
-	{
-		return false;
+		// Check parent
+		bResult = Parent->GetRenderProxy()->GetParameterValue(Type, ParameterInfo, OutValue, Context);
 	}
-}
 
-bool FMaterialInstanceResource::GetTextureValue(
-	const FHashedMaterialParameterInfo& ParameterInfo,
-	const URuntimeVirtualTexture** OutValue,
-	const FMaterialRenderContext& Context
-) const
-{
-	checkSlow(IsInParallelRenderingThread());
-	const URuntimeVirtualTexture* const * Value = RenderThread_FindParameterByName<const URuntimeVirtualTexture*>(ParameterInfo);
-	if (Value && *Value)
-	{
-		*OutValue = *Value;
-		return true;
-	}
-	else if (Parent)
-	{
-		return Parent->GetRenderProxy()->GetTextureValue(ParameterInfo, OutValue, Context);
-	}
-	else
-	{
-		return false;
-	}
+	return bResult;
 }
 
 void UMaterialInstance::PropagateDataToMaterialProxy()
@@ -835,6 +768,27 @@ UMaterial* UMaterialInstance::GetMaterial()
 	}
 }
 
+bool UMaterialInstance::GetParameterOverrideValue(EMaterialParameterType Type, const FMemoryImageMaterialParameterInfo& ParameterInfo, FMaterialParameterMetadata& OutResult, EMaterialGetParameterValueFlags Flags) const
+{
+	// If we're not a parent, set the flag to signify this parameter was a MI override
+	const bool bSetOverride = !EnumHasAnyFlags(Flags, EMaterialGetParameterValueFlags::IsParent);
+	bool bResult = false;
+	switch (Type)
+	{
+	case EMaterialParameterType::Scalar: bResult = GameThread_GetParameterValue(ScalarParameterValues, ParameterInfo, bSetOverride, OutResult); break;
+	case EMaterialParameterType::Vector: bResult = GameThread_GetParameterValue(VectorParameterValues, ParameterInfo, bSetOverride, OutResult); break;
+	case EMaterialParameterType::Texture: bResult = GameThread_GetParameterValue(TextureParameterValues, ParameterInfo, bSetOverride, OutResult); break;
+	case EMaterialParameterType::RuntimeVirtualTexture: bResult = GameThread_GetParameterValue(RuntimeVirtualTextureParameterValues, ParameterInfo, bSetOverride, OutResult); break;
+	case EMaterialParameterType::Font: bResult = GameThread_GetParameterValue(FontParameterValues, ParameterInfo, bSetOverride, OutResult); break;
+#if WITH_EDITORONLY_DATA
+	case EMaterialParameterType::StaticSwitch: bResult = GameThread_GetParameterValue(StaticParameters.StaticSwitchParameters, ParameterInfo, bSetOverride, OutResult); break;
+	case EMaterialParameterType::StaticComponentMask: bResult = GameThread_GetParameterValue(StaticParameters.StaticComponentMaskParameters, ParameterInfo, bSetOverride, OutResult); break;
+#endif // WITH_EDITORONLY_DATA
+	default: checkNoEntry(); break;
+	}
+	return bResult;
+}
+
 bool UMaterialInstance::GetParameterValue(EMaterialParameterType Type, const FMemoryImageMaterialParameterInfo& ParameterInfo, FMaterialParameterMetadata& OutResult, EMaterialGetParameterValueFlags Flags) const
 {
 	bool bFoundAValue = false;
@@ -847,23 +801,7 @@ bool UMaterialInstance::GetParameterValue(EMaterialParameterType Type, const FMe
 	// Instance override
 	if (EnumHasAnyFlags(Flags, EMaterialGetParameterValueFlags::CheckInstanceOverrides))
 	{
-		// If we're not a parent, set the flag to signify this parameter was a MI override
-		const bool bSetOverride = !EnumHasAnyFlags(Flags, EMaterialGetParameterValueFlags::IsParent);
-		bool bResult = false;
-		switch (Type)
-		{
-		case EMaterialParameterType::Scalar: bResult = GameThread_GetParameterValue(ScalarParameterValues, ParameterInfo, bSetOverride, OutResult); break;
-		case EMaterialParameterType::Vector: bResult = GameThread_GetParameterValue(VectorParameterValues, ParameterInfo, bSetOverride, OutResult); break;
-		case EMaterialParameterType::Texture: bResult = GameThread_GetParameterValue(TextureParameterValues, ParameterInfo, bSetOverride, OutResult); break;
-		case EMaterialParameterType::RuntimeVirtualTexture: bResult = GameThread_GetParameterValue(RuntimeVirtualTextureParameterValues, ParameterInfo, bSetOverride, OutResult); break;
-		case EMaterialParameterType::Font: bResult = GameThread_GetParameterValue(FontParameterValues, ParameterInfo, bSetOverride, OutResult); break;
-#if WITH_EDITORONLY_DATA
-		case EMaterialParameterType::StaticSwitch: bResult = GameThread_GetParameterValue(StaticParameters.StaticSwitchParameters, ParameterInfo, bSetOverride, OutResult); break;
-		case EMaterialParameterType::StaticComponentMask: bResult = GameThread_GetParameterValue(StaticParameters.StaticComponentMaskParameters, ParameterInfo, bSetOverride, OutResult); break;
-#endif // WITH_EDITORONLY_DATA
-		default: checkNoEntry(); break;
-		}
-		if (bResult)
+		if (GetParameterOverrideValue(Type, ParameterInfo, OutResult, Flags))
 		{
 			return true;
 		}
@@ -1317,7 +1255,7 @@ void UMaterialInstance::OverrideTexture(const UTexture* InTextureToOverride, UTe
 #endif // #if WITH_EDITOR
 }
 
-void UMaterialInstance::OverrideVectorParameterDefault(const FHashedMaterialParameterInfo& ParameterInfo, const FLinearColor& Value, bool bOverride, ERHIFeatureLevel::Type InFeatureLevel)
+void UMaterialInstance::OverrideNumericParameterDefault(EMaterialParameterType Type, const FHashedMaterialParameterInfo& ParameterInfo, const UE::Shader::FValue& Value, bool bOverride, ERHIFeatureLevel::Type InFeatureLevel)
 {
 #if WITH_EDITOR
 	bool bShouldRecacheMaterialExpressions = false;
@@ -1326,42 +1264,12 @@ void UMaterialInstance::OverrideVectorParameterDefault(const FHashedMaterialPara
 		FMaterialResource* SourceMaterialResource = GetMaterialResource(InFeatureLevel);
 		if (SourceMaterialResource)
 		{
-			SourceMaterialResource->TransientOverrides.SetVectorOverride(ParameterInfo, Value, bOverride);
+			SourceMaterialResource->TransientOverrides.SetNumericOverride(Type, ParameterInfo, Value, bOverride);
 
-			const TArrayView<const FMaterialVectorParameterInfo> Parameters = SourceMaterialResource->GetUniformVectorParameterExpressions();
+			const TArrayView<const FMaterialNumericParameterInfo> Parameters = SourceMaterialResource->GetUniformNumericParameterExpressions();
 			for (int32 i = 0; i < Parameters.Num(); ++i)
 			{
-				const FMaterialVectorParameterInfo& Parameter = Parameters[i];
-				if (Parameter.ParameterInfo == ParameterInfo)
-				{
-					bShouldRecacheMaterialExpressions = true;
-				}
-			}
-		}
-	}
-
-	if (bShouldRecacheMaterialExpressions)
-	{
-		RecacheUniformExpressions(false);
-	}
-#endif // #if WITH_EDITOR
-}
-
-void UMaterialInstance::OverrideScalarParameterDefault(const FHashedMaterialParameterInfo& ParameterInfo, float Value, bool bOverride, ERHIFeatureLevel::Type InFeatureLevel)
-{
-#if WITH_EDITOR
-	bool bShouldRecacheMaterialExpressions = false;
-	if (bHasStaticPermutationResource)
-	{
-		FMaterialResource* SourceMaterialResource = GetMaterialResource(InFeatureLevel);
-		if (SourceMaterialResource)
-		{
-			SourceMaterialResource->TransientOverrides.SetScalarOverride(ParameterInfo, Value, bOverride);
-
-			const TArrayView<const FMaterialScalarParameterInfo> Parameters = SourceMaterialResource->GetUniformScalarParameterExpressions();
-			for (int32 i = 0; i < Parameters.Num(); ++i)
-			{
-				const FMaterialScalarParameterInfo& Parameter = Parameters[i];
+				const FMaterialNumericParameterInfo& Parameter = Parameters[i];
 				if (Parameter.ParameterInfo == ParameterInfo)
 				{
 					bShouldRecacheMaterialExpressions = true;
@@ -1519,53 +1427,44 @@ void UMaterialInstanceDynamic::CopyScalarAndVectorParameters(const UMaterialInte
 	// because static (bool) parameters can cause some parameters to be hidden
 	FMaterialResource* MaterialResource = GetMaterialResource(FeatureLevel);
 
-	if(MaterialResource)
+	if (MaterialResource)
 	{
 		// first, clear out all the parameter values
 		ClearParameterValuesInternal(false);
 
-		// scalar
+		TArrayView<const FMaterialNumericParameterInfo> Array = MaterialResource->GetUniformNumericParameterExpressions();
+		for (int32 i = 0, Count = Array.Num(); i < Count; ++i)
 		{
-			TArrayView<const FMaterialScalarParameterInfo> Array = MaterialResource->GetUniformScalarParameterExpressions();
+			const FMaterialNumericParameterInfo& Parameter = Array[i];
 
-			for (int32 i = 0, Count = Array.Num(); i < Count; ++i)
+			const UMaterialInterface* CheckMaterial = &SourceMaterialToCopyFrom;
+			FMaterialParameterMetadata ParameterValue;
+			bool bFoundValue = false;
+			while (CheckMaterial)
 			{
-				const FMaterialScalarParameterInfo& Parameter = Array[i];
-
-				float Value;
-				Parameter.GetGameThreadNumberValue(&SourceMaterialToCopyFrom, Value);
-
-				FScalarParameterValue* ParameterValue = GameThread_FindParameterByName(ScalarParameterValues, Parameter.ParameterInfo);
-				if (!ParameterValue)
+				const UMaterialInstance* CheckMaterialInstance = Cast<UMaterialInstance>(CheckMaterial);
+				if (CheckMaterialInstance)
 				{
-					ParameterValue = new(ScalarParameterValues)FScalarParameterValue;
-					ParameterValue->ParameterInfo = FMaterialParameterInfo(Parameter.ParameterInfo);
+					if (CheckMaterialInstance->GetParameterOverrideValue(Parameter.ParameterType, Parameter.ParameterInfo, ParameterValue))
+					{
+						bFoundValue = true;
+						break;
+					}
+					CheckMaterial = CheckMaterialInstance->Parent;
 				}
-
-				ParameterValue->ParameterValue = Value;
-			}
-		}
-
-		// vector
-		{
-			TArrayView<const FMaterialVectorParameterInfo> Array = MaterialResource->GetUniformVectorParameterExpressions();
-
-			for (int32 i = 0, Count = Array.Num(); i < Count; ++i)
-			{
-				const FMaterialVectorParameterInfo& Parameter = Array[i];
-
-				FLinearColor Value;
-				Parameter.GetGameThreadNumberValue(&SourceMaterialToCopyFrom, Value);
-
-				FVectorParameterValue* ParameterValue = GameThread_FindParameterByName(VectorParameterValues, Parameter.ParameterInfo);
-				if (!ParameterValue)
+				else
 				{
-					ParameterValue = new(VectorParameterValues)FVectorParameterValue;
-					ParameterValue->ParameterInfo = FMaterialParameterInfo(Parameter.ParameterInfo);
+					break;
 				}
+			}
 
-				ParameterValue->ParameterValue = Value;
+			if (!bFoundValue)
+			{
+				const UE::Shader::FValue DefaultValue = MaterialResource->GetUniformExpressions().GetDefaultParameterValue(Parameter.ParameterType, Parameter.DefaultValueOffset);
+				ParameterValue.Value = FMaterialParameterValue(Parameter.ParameterType, DefaultValue);
 			}
+
+			SetParameterValueInternal(FMaterialParameterInfo(Parameter.ParameterInfo), ParameterValue);
 		}
 
 		// now, init the resources
@@ -3079,34 +2978,14 @@ FMaterialInstanceParameterUpdateContext::~FMaterialInstanceParameterUpdateContex
 
 void FMaterialInstanceParameterUpdateContext::SetParameterValueEditorOnly(const FMaterialParameterInfo& ParameterInfo, const FMaterialParameterMetadata& Meta, EMaterialSetParameterValueFlags Flags)
 {
-	const FMaterialParameterValue& Value = Meta.Value;
-	if (IsStaticMaterialParameter(Value.Type))
+	if (IsStaticMaterialParameter(Meta.Value.Type))
 	{
 		// Route static parameters to the static parameter set
 		StaticParameters.SetParameterValue(ParameterInfo, Meta, Flags);
 	}
 	else
 	{
-		const bool bUseAtlas = EnumHasAnyFlags(Flags, EMaterialSetParameterValueFlags::SetCurveAtlas);
-		FScalarParameterAtlasInstanceData AtlasData;
-
-		switch (Value.Type)
-		{
-		case EMaterialParameterType::Scalar:
-			if (bUseAtlas)
-			{
-				AtlasData.bIsUsedAsAtlasPosition = Meta.bUsedAsAtlasPosition;
-				AtlasData.Atlas = Meta.ScalarAtlas;
-				AtlasData.Curve = Meta.ScalarCurve;
-			}
-			Instance->SetScalarParameterValueInternal(ParameterInfo, Value.AsScalar(), bUseAtlas, AtlasData);
-			break;
-		case EMaterialParameterType::Vector: Instance->SetVectorParameterValueInternal(ParameterInfo, Value.AsLinearColor()); break;
-		case EMaterialParameterType::Texture: Instance->SetTextureParameterValueInternal(ParameterInfo, Value.Texture); break;
-		case EMaterialParameterType::Font: Instance->SetFontParameterValueInternal(ParameterInfo, Value.Font.Value, Value.Font.Page); break;
-		case EMaterialParameterType::RuntimeVirtualTexture: Instance->SetRuntimeVirtualTextureParameterValueInternal(ParameterInfo, Value.RuntimeVirtualTexture); break;
-		default: checkNoEntry();
-		}
+		Instance->SetParameterValueInternal(ParameterInfo, Meta, Flags);
 	}
 }
 
@@ -3121,6 +3000,32 @@ void FMaterialInstanceParameterUpdateContext::SetBasePropertyOverrides(const FMa
 }
 #endif // WITH_EDITORONLY_DATA
 
+void UMaterialInstance::SetParameterValueInternal(const FMaterialParameterInfo& ParameterInfo, const FMaterialParameterMetadata& Meta, EMaterialSetParameterValueFlags Flags)
+{
+	const bool bUseAtlas = EnumHasAnyFlags(Flags, EMaterialSetParameterValueFlags::SetCurveAtlas);
+	const FMaterialParameterValue& Value = Meta.Value;
+	FScalarParameterAtlasInstanceData AtlasData;
+	switch (Value.Type)
+	{
+	case EMaterialParameterType::Scalar:
+#if WITH_EDITORONLY_DATA
+		if (bUseAtlas)
+		{
+			AtlasData.bIsUsedAsAtlasPosition = Meta.bUsedAsAtlasPosition;
+			AtlasData.Atlas = Meta.ScalarAtlas;
+			AtlasData.Curve = Meta.ScalarCurve;
+		}
+#endif // WITH_EDITORONLY_DATA
+		SetScalarParameterValueInternal(ParameterInfo, Value.AsScalar(), bUseAtlas, AtlasData);
+		break;
+	case EMaterialParameterType::Vector: SetVectorParameterValueInternal(ParameterInfo, Value.AsLinearColor()); break;
+	case EMaterialParameterType::Texture: SetTextureParameterValueInternal(ParameterInfo, Value.Texture); break;
+	case EMaterialParameterType::Font: SetFontParameterValueInternal(ParameterInfo, Value.Font.Value, Value.Font.Page); break;
+	case EMaterialParameterType::RuntimeVirtualTexture: SetRuntimeVirtualTextureParameterValueInternal(ParameterInfo, Value.RuntimeVirtualTexture); break;
+	default: checkNoEntry();
+	}
+}
+
 void UMaterialInstance::SetVectorParameterValueInternal(const FMaterialParameterInfo& ParameterInfo, FLinearColor Value)
 {
 	LLM_SCOPE(ELLMTag::MaterialInstance);
@@ -4392,22 +4297,26 @@ void UMaterialInstance::CopyMaterialUniformParametersInternal(UMaterialInterface
 
 				if (MaterialResource)
 				{
-					// Scalars
-					TArrayView<const FMaterialScalarParameterInfo> ScalarExpressions = MaterialResource->GetUniformScalarParameterExpressions();
-					for (const FMaterialScalarParameterInfo& Parameter : ScalarExpressions)
+					// Numeric
+					for (const FMaterialNumericParameterInfo& Parameter : MaterialResource->GetUniformNumericParameterExpressions())
 					{
-						FScalarParameterValue* ParameterValue = new(ScalarParameterValues) FScalarParameterValue;
-						ParameterValue->ParameterInfo.Name = Parameter.ParameterInfo.GetName();
-						Parameter.GetDefaultValue(ParameterValue->ParameterValue);
-					}
-
-					// Vectors
-					TArrayView<const FMaterialVectorParameterInfo> VectorExpressions = MaterialResource->GetUniformVectorParameterExpressions();
-					for (const FMaterialVectorParameterInfo& Parameter : VectorExpressions)
-					{
-						FVectorParameterValue* ParameterValue = new(VectorParameterValues) FVectorParameterValue;
-						ParameterValue->ParameterInfo.Name = Parameter.ParameterInfo.GetName();
-						Parameter.GetDefaultValue(ParameterValue->ParameterValue);
+						const UE::Shader::FValue DefaultValue = MaterialResource->GetUniformExpressions().GetDefaultParameterValue(Parameter.ParameterType, Parameter.DefaultValueOffset);
+						if (Parameter.ParameterType == EMaterialParameterType::Scalar)
+						{
+							FScalarParameterValue* ParameterValue = new(ScalarParameterValues) FScalarParameterValue;
+							ParameterValue->ParameterInfo.Name = Parameter.ParameterInfo.GetName();
+							ParameterValue->ParameterValue = DefaultValue.AsFloatScalar();
+						}
+						else if (Parameter.ParameterType == EMaterialParameterType::Vector)
+						{
+							FVectorParameterValue* ParameterValue = new(VectorParameterValues) FVectorParameterValue;
+							ParameterValue->ParameterInfo.Name = Parameter.ParameterInfo.GetName();
+							ParameterValue->ParameterValue = DefaultValue.AsLinearColor();
+						}
+						else
+						{
+							ensure(false);
+						}
 					}
 
 					// Textures
@@ -4524,5 +4433,3 @@ void UMaterialInstance::AppendReferencedParameterCollectionIdsTo(TArray<FGuid>&
 UMaterialInstance::FCustomStaticParametersGetterDelegate UMaterialInstance::CustomStaticParametersGetters;
 TArray<UMaterialInstance::FCustomParameterSetUpdaterDelegate> UMaterialInstance::CustomParameterSetUpdaters;
 #endif // WITH_EDITORONLY_DATA
-
-PRAGMA_ENABLE_OPTIMIZATION
\ No newline at end of file
diff --git a/Engine/Source/Runtime/Engine/Private/Materials/MaterialInstanceSupport.h b/Engine/Source/Runtime/Engine/Private/Materials/MaterialInstanceSupport.h
--- a/Engine/Source/Runtime/Engine/Private/Materials/MaterialInstanceSupport.h
+++ b/Engine/Source/Runtime/Engine/Private/Materials/MaterialInstanceSupport.h
@@ -98,10 +98,7 @@ public:
 	virtual const FMaterialRenderProxy* GetFallback(ERHIFeatureLevel::Type InFeatureLevel) const override;
 	virtual UMaterialInterface* GetMaterialInterface() const override;
 	
-	virtual bool GetVectorValue(const FHashedMaterialParameterInfo& ParameterInfo, FLinearColor* OutValue, const FMaterialRenderContext& Context) const override;
-	virtual bool GetScalarValue(const FHashedMaterialParameterInfo& ParameterInfo, float* OutValue, const FMaterialRenderContext& Context) const override;
-	virtual bool GetTextureValue(const FHashedMaterialParameterInfo& ParameterInfo, const UTexture** OutValue, const FMaterialRenderContext& Context) const override;
-	virtual bool GetTextureValue(const FHashedMaterialParameterInfo& ParameterInfo, const URuntimeVirtualTexture** OutValue, const FMaterialRenderContext& Context) const override;
+	virtual bool GetParameterValue(EMaterialParameterType Type, const FHashedMaterialParameterInfo& ParameterInfo, FMaterialParameterValue& OutValue, const FMaterialRenderContext& Context) const override;
 
 	void GameThread_SetParent(UMaterialInterface* ParentMaterialInterface);
 
@@ -151,21 +148,20 @@ public:
 	 * Retrieves a parameter by name.
 	 */
 	template <typename ValueType>
-	const ValueType* RenderThread_FindParameterByName(const FHashedMaterialParameterInfo& ParameterInfo) const
+	bool RenderThread_GetParameterValue(const FHashedMaterialParameterInfo& ParameterInfo, FMaterialParameterValue& OutValue) const
 	{
-		const TArray<TNamedParameter<ValueType> >& ValueArray = GetValueArray<ValueType>();
-		const int32 ParameterCount = ValueArray.Num();
-		for (int32 ParameterIndex = 0; ParameterIndex < ParameterCount; ++ParameterIndex)
+		const TArray<TNamedParameter<ValueType>>& ValueArray = GetValueArray<ValueType>();
+		for (const TNamedParameter<ValueType>& Parameter : ValueArray)
 		{
-			const TNamedParameter<ValueType>& Parameter = ValueArray[ParameterIndex];
 			if (Parameter.Info == ParameterInfo)
 			{
-				return &Parameter.Value;
+				OutValue = Parameter.Value;
+				return true;
 			}
 		}
-		return NULL;
+		return false;
 	}
-	
+
 private:
 	/**
 	 * Retrieves the array of values for a given type.
diff --git a/Engine/Source/Runtime/Engine/Private/Materials/MaterialShared.cpp b/Engine/Source/Runtime/Engine/Private/Materials/MaterialShared.cpp
--- a/Engine/Source/Runtime/Engine/Private/Materials/MaterialShared.cpp
+++ b/Engine/Source/Runtime/Engine/Private/Materials/MaterialShared.cpp
@@ -69,8 +69,7 @@ IMPLEMENT_TYPE_LAYOUT(FMeshMaterialShaderMap);
 IMPLEMENT_TYPE_LAYOUT(FMaterialProcessedSource);
 IMPLEMENT_TYPE_LAYOUT(FMaterialShaderMapContent);
 IMPLEMENT_TYPE_LAYOUT(FMaterialUniformPreshaderHeader);
-IMPLEMENT_TYPE_LAYOUT(FMaterialScalarParameterInfo);
-IMPLEMENT_TYPE_LAYOUT(FMaterialVectorParameterInfo);
+IMPLEMENT_TYPE_LAYOUT(FMaterialNumericParameterInfo);
 IMPLEMENT_TYPE_LAYOUT(FMaterialTextureParameterInfo);
 IMPLEMENT_TYPE_LAYOUT(FMaterialExternalTextureParameterInfo);
 IMPLEMENT_TYPE_LAYOUT(FMaterialVirtualTextureStack);
@@ -230,6 +229,16 @@ int32 FMaterialCompiler::Errorf(const TCHAR* Format,...)
 	return Error(ErrorText);
 }
 
+int32 FMaterialCompiler::ScalarParameter(FName ParameterName, float DefaultValue)
+{
+	return GenericParameter(EMaterialParameterType::Scalar, ParameterName, DefaultValue);
+}
+
+int32 FMaterialCompiler::VectorParameter(FName ParameterName, const FLinearColor& DefaultValue)
+{
+	return GenericParameter(EMaterialParameterType::Vector, ParameterName, DefaultValue);
+}
+
 IMPLEMENT_STRUCT(ExpressionInput);
 IMPLEMENT_STRUCT(ColorMaterialInput);
 IMPLEMENT_STRUCT(ScalarMaterialInput);
@@ -785,14 +794,9 @@ TArrayView<const FMaterialTextureParameterInfo> FMaterial::GetUniformTextureExpr
 	return GetUniformExpressions().UniformTextureParameters[(uint32)Type];
 }
 
-TArrayView<const FMaterialVectorParameterInfo> FMaterial::GetUniformVectorParameterExpressions() const
+TArrayView<const FMaterialNumericParameterInfo> FMaterial::GetUniformNumericParameterExpressions() const
 { 
-	return GetUniformExpressions().UniformVectorParameters;
-}
-
-TArrayView<const FMaterialScalarParameterInfo> FMaterial::GetUniformScalarParameterExpressions() const
-{ 
-	return GetUniformExpressions().UniformScalarParameters;
+	return GetUniformExpressions().UniformNumericParameters;
 }
 
 bool FMaterial::RequiresSceneColorCopy_GameThread() const
@@ -2970,6 +2974,50 @@ void FMaterialVirtualTextureStack::Serialize(FArchive& Ar)
 	FMaterialRenderProxy
 -----------------------------------------------------------------------------*/
 
+bool FMaterialRenderProxy::GetVectorValue(const FHashedMaterialParameterInfo& ParameterInfo, FLinearColor* OutValue, const FMaterialRenderContext& Context) const
+{
+	FMaterialParameterValue Value;
+	if (GetParameterValue(EMaterialParameterType::Vector, ParameterInfo, Value, Context))
+	{
+		*OutValue = Value.AsLinearColor();
+		return true;
+	}
+	return false;
+}
+
+bool FMaterialRenderProxy::GetScalarValue(const FHashedMaterialParameterInfo& ParameterInfo, float* OutValue, const FMaterialRenderContext& Context) const
+{
+	FMaterialParameterValue Value;
+	if (GetParameterValue(EMaterialParameterType::Scalar, ParameterInfo, Value, Context))
+	{
+		*OutValue = Value.AsScalar();
+		return true;
+	}
+	return false;
+}
+
+bool FMaterialRenderProxy::GetTextureValue(const FHashedMaterialParameterInfo& ParameterInfo, const UTexture** OutValue, const FMaterialRenderContext& Context) const
+{
+	FMaterialParameterValue Value;
+	if (GetParameterValue(EMaterialParameterType::Texture, ParameterInfo, Value, Context))
+	{
+		*OutValue = Value.Texture;
+		return true;
+	}
+	return false;
+}
+
+bool FMaterialRenderProxy::GetTextureValue(const FHashedMaterialParameterInfo& ParameterInfo, const URuntimeVirtualTexture** OutValue, const FMaterialRenderContext& Context) const
+{
+	FMaterialParameterValue Value;
+	if (GetParameterValue(EMaterialParameterType::RuntimeVirtualTexture, ParameterInfo, Value, Context))
+	{
+		*OutValue = Value.RuntimeVirtualTexture;
+		return true;
+	}
+	return false;
+}
+
 static void OnVirtualTextureDestroyedCB(const FVirtualTextureProducerHandle& InHandle, void* Baton)
 {
 	FMaterialRenderProxy* MaterialProxy = static_cast<FMaterialRenderProxy*>(Baton);
@@ -3407,48 +3455,33 @@ TSet<FMaterialRenderProxy*> FMaterialRenderProxy::DeferredUniformExpressionCache
 	FColoredMaterialRenderProxy
 -----------------------------------------------------------------------------*/
 
-bool FColoredMaterialRenderProxy::GetVectorValue(const FHashedMaterialParameterInfo& ParameterInfo, FLinearColor* OutValue, const FMaterialRenderContext& Context) const
+bool FColoredMaterialRenderProxy::GetParameterValue(EMaterialParameterType Type, const FHashedMaterialParameterInfo& ParameterInfo, FMaterialParameterValue& OutValue, const FMaterialRenderContext& Context) const
 {
-	if(ParameterInfo.Name == ColorParamName)
+	if (Type == EMaterialParameterType::Vector && ParameterInfo.Name == ColorParamName)
 	{
-		*OutValue = Color;
+		OutValue = Color;
 		return true;
 	}
 	else
 	{
-		return Parent->GetVectorValue(ParameterInfo, OutValue, Context);
+		return Parent->GetParameterValue(Type, ParameterInfo, OutValue, Context);
 	}
 }
 
-bool FColoredMaterialRenderProxy::GetScalarValue(const FHashedMaterialParameterInfo& ParameterInfo, float* OutValue, const FMaterialRenderContext& Context) const
-{
-	return Parent->GetScalarValue(ParameterInfo, OutValue, Context);
-}
-
-bool FColoredMaterialRenderProxy::GetTextureValue(const FHashedMaterialParameterInfo& ParameterInfo,const UTexture** OutValue, const FMaterialRenderContext& Context) const
-{
-	return Parent->GetTextureValue(ParameterInfo,OutValue,Context);
-}
-
-bool FColoredMaterialRenderProxy::GetTextureValue(const FHashedMaterialParameterInfo& ParameterInfo, const URuntimeVirtualTexture** OutValue, const FMaterialRenderContext& Context) const
-{
-	return Parent->GetTextureValue(ParameterInfo, OutValue, Context);
-}
-
 /*-----------------------------------------------------------------------------
 	FColoredTexturedMaterialRenderProxy
 -----------------------------------------------------------------------------*/
 
-bool FColoredTexturedMaterialRenderProxy::GetTextureValue(const FHashedMaterialParameterInfo& ParameterInfo, const UTexture** OutValue, const FMaterialRenderContext& Context) const
+bool FColoredTexturedMaterialRenderProxy::GetParameterValue(EMaterialParameterType Type, const FHashedMaterialParameterInfo& ParameterInfo, FMaterialParameterValue& OutValue, const FMaterialRenderContext& Context) const
 {
-	if (ParameterInfo.Name == TextureParamName)
+	if (Type == EMaterialParameterType::Texture && ParameterInfo.Name == TextureParamName)
 	{
-		*OutValue = Texture;
+		OutValue = Texture;
 		return true;
 	}
 	else
 	{
-		return Parent->GetTextureValue(ParameterInfo, OutValue, Context);
+		return Parent->GetParameterValue(Type, ParameterInfo, OutValue, Context);
 	}
 }
 
@@ -3465,47 +3498,32 @@ const FMaterialRenderProxy* FOverrideSelectionColorMaterialRenderProxy::GetFallb
 	return Parent->GetFallback(InFeatureLevel);
 }
 
-bool FOverrideSelectionColorMaterialRenderProxy::GetVectorValue(const FHashedMaterialParameterInfo& ParameterInfo, FLinearColor* OutValue, const FMaterialRenderContext& Context) const
+bool FOverrideSelectionColorMaterialRenderProxy::GetParameterValue(EMaterialParameterType Type, const FHashedMaterialParameterInfo& ParameterInfo, FMaterialParameterValue& OutValue, const FMaterialRenderContext& Context) const
 {
-	if (ParameterInfo.Name == FName(NAME_SelectionColor))
+	if (Type == EMaterialParameterType::Vector && ParameterInfo.Name == FName(NAME_SelectionColor))
 	{
-		*OutValue = SelectionColor;
+		OutValue = SelectionColor;
 		return true;
 	}
 	else
 	{
-		return Parent->GetVectorValue(ParameterInfo, OutValue, Context);
+		return Parent->GetParameterValue(Type, ParameterInfo, OutValue, Context);
 	}
 }
 
-bool FOverrideSelectionColorMaterialRenderProxy::GetScalarValue(const FHashedMaterialParameterInfo& ParameterInfo, float* OutValue, const FMaterialRenderContext& Context) const
-{
-	return Parent->GetScalarValue(ParameterInfo, OutValue, Context);
-}
-
-bool FOverrideSelectionColorMaterialRenderProxy::GetTextureValue(const FHashedMaterialParameterInfo& ParameterInfo, const UTexture** OutValue, const FMaterialRenderContext& Context) const
-{
-	return Parent->GetTextureValue(ParameterInfo, OutValue, Context);
-}
-
-bool FOverrideSelectionColorMaterialRenderProxy::GetTextureValue(const FHashedMaterialParameterInfo& ParameterInfo, const URuntimeVirtualTexture** OutValue, const FMaterialRenderContext& Context) const
-{
-	return Parent->GetTextureValue(ParameterInfo, OutValue, Context);
-}
-
 /*-----------------------------------------------------------------------------
 	FLightingDensityMaterialRenderProxy
 -----------------------------------------------------------------------------*/
 static FName NAME_LightmapRes = FName(TEXT("LightmapRes"));
 
-bool FLightingDensityMaterialRenderProxy::GetVectorValue(const FHashedMaterialParameterInfo& ParameterInfo, FLinearColor* OutValue, const FMaterialRenderContext& Context) const
+bool FLightingDensityMaterialRenderProxy::GetParameterValue(EMaterialParameterType Type, const FHashedMaterialParameterInfo& ParameterInfo, FMaterialParameterValue& OutValue, const FMaterialRenderContext& Context) const
 {
-	if (ParameterInfo.Name == NAME_LightmapRes)
+	if (Type == EMaterialParameterType::Vector && ParameterInfo.Name == NAME_LightmapRes)
 	{
-		*OutValue = FLinearColor(LightmapResolution.X, LightmapResolution.Y, 0.0f, 0.0f);
+		OutValue = FLinearColor(LightmapResolution.X, LightmapResolution.Y, 0.0f, 0.0f);
 		return true;
 	}
-	return FColoredMaterialRenderProxy::GetVectorValue(ParameterInfo, OutValue, Context);
+	return FColoredMaterialRenderProxy::GetParameterValue(Type, ParameterInfo, OutValue, Context);
 }
 
 #if WITH_EDITOR
@@ -4958,4 +4976,60 @@ void FMaterial::SetShaderMapsOnMaterialResources(const TMap<TRefCountPtr<FMateri
 		SetShaderMapsOnMaterialResources_RenderThread(RHICmdList, InMaterialsToUpdate);
 	});
 }
+
+FMaterialParameterValue::FMaterialParameterValue(EMaterialParameterType InType, const UE::Shader::FValue& InValue)
+{
+	switch (InType)
+	{
+	case EMaterialParameterType::Scalar: *this = InValue.AsFloatScalar(); break;
+	case EMaterialParameterType::Vector: *this = InValue.AsLinearColor(); break;
+	case EMaterialParameterType::StaticSwitch: *this = InValue.AsBoolScalar(); break;
+	case EMaterialParameterType::StaticComponentMask:
+	{
+		const UE::Shader::FBoolValue BoolValue = InValue.AsBool();
+		*this = FMaterialParameterValue(BoolValue[0], BoolValue[1], BoolValue[2], BoolValue[3]);
+		break;
+	}
+	default: ensure(false); Type = EMaterialParameterType::None; break;
+	}
+}
+
+UE::Shader::FValue FMaterialParameterValue::AsShaderValue() const
+{
+	switch (Type)
+	{
+	case EMaterialParameterType::Scalar: return Float[0];
+	case EMaterialParameterType::Vector: return UE::Shader::FValue(Float[0], Float[1], Float[2], Float[3]);
+	case EMaterialParameterType::StaticSwitch: return Bool[0];
+	case EMaterialParameterType::StaticComponentMask: return UE::Shader::FValue(Bool[0], Bool[1], Bool[2], Bool[3]);
+	case EMaterialParameterType::Texture:
+	case EMaterialParameterType::Font:
+	case EMaterialParameterType::RuntimeVirtualTexture:
+		// Non-numeric types, can't represent as shader values
+		return UE::Shader::FValue();
+	default:
+		checkNoEntry();
+		return UE::Shader::FValue();
+	}
+}
+
+UE::Shader::EValueType GetShaderValueType(EMaterialParameterType Type)
+{
+	switch (Type)
+	{
+	case EMaterialParameterType::Scalar: return UE::Shader::EValueType::Float1;
+	case EMaterialParameterType::Vector: return UE::Shader::EValueType::Float4;
+	case EMaterialParameterType::StaticSwitch: return UE::Shader::EValueType::Bool1;
+	case EMaterialParameterType::StaticComponentMask: return UE::Shader::EValueType::Bool4;
+	case EMaterialParameterType::Texture:
+	case EMaterialParameterType::Font:
+	case EMaterialParameterType::RuntimeVirtualTexture:
+		// Non-numeric types, can't represent as shader values
+		return UE::Shader::EValueType::Void;
+	default:
+		checkNoEntry();
+		return UE::Shader::EValueType::Void;
+	}
+}
+
 #undef LOCTEXT_NAMESPACE
\ No newline at end of file
diff --git a/Engine/Source/Runtime/Engine/Private/Materials/MaterialUniformExpressions.cpp b/Engine/Source/Runtime/Engine/Private/Materials/MaterialUniformExpressions.cpp
--- a/Engine/Source/Runtime/Engine/Private/Materials/MaterialUniformExpressions.cpp
+++ b/Engine/Source/Runtime/Engine/Private/Materials/MaterialUniformExpressions.cpp
@@ -140,44 +140,25 @@ void FMaterialUniformExpression::WriteNumberOpcodes(UE::Shader::FPreshaderData&
 	OutData.WriteOpcode(UE::Shader::EPreshaderOpcode::ConstantZero);
 }
 
-void FUniformParameterOverrides::SetScalarOverride(const FHashedMaterialParameterInfo& ParameterInfo, float Value, bool bOverride)
+void FUniformParameterOverrides::SetNumericOverride(EMaterialParameterType Type, const FHashedMaterialParameterInfo& ParameterInfo, const UE::Shader::FValue& Value, bool bOverride)
 {
+	const FNumericParameterKey Key{ ParameterInfo, Type };
 	if (bOverride)
 	{
-		ScalarOverrides.FindOrAdd(ParameterInfo) = Value;
+		const UE::Shader::EValueType ShaderType = GetShaderValueType(Type);
+		check(ShaderType == Value.GetType());
+		NumericOverrides.FindOrAdd(Key) = Value;
 	}
 	else
 	{
-		ScalarOverrides.Remove(ParameterInfo);
+		NumericOverrides.Remove(Key);
 	}
 }
 
-void FUniformParameterOverrides::SetVectorOverride(const FHashedMaterialParameterInfo& ParameterInfo, const FLinearColor& Value, bool bOverride)
+bool FUniformParameterOverrides::GetNumericOverride(EMaterialParameterType Type, const FHashedMaterialParameterInfo& ParameterInfo, UE::Shader::FValue& OutValue) const
 {
-	if (bOverride)
-	{
-		VectorOverrides.FindOrAdd(ParameterInfo) = Value;
-	}
-	else
-	{
-		VectorOverrides.Remove(ParameterInfo);
-	}
-}
-
-bool FUniformParameterOverrides::GetScalarOverride(const FHashedMaterialParameterInfo& ParameterInfo, float& OutValue) const
-{
-	const float* Result = ScalarOverrides.Find(ParameterInfo);
-	if (Result)
-	{
-		OutValue = *Result;
-		return true;
-	}
-	return false;
-}
-
-bool FUniformParameterOverrides::GetVectorOverride(const FHashedMaterialParameterInfo& ParameterInfo, FLinearColor& OutValue) const
-{
-	const FLinearColor* Result = VectorOverrides.Find(ParameterInfo);
+	const FNumericParameterKey Key{ ParameterInfo, Type };
+	const UE::Shader::FValue* Result = NumericOverrides.Find(Key);
 	if (Result)
 	{
 		OutValue = *Result;
@@ -240,8 +221,7 @@ bool FUniformExpressionSet::IsEmpty() const
 		}
 	}
 
-	return UniformVectorParameters.Num() == 0
-		&& UniformScalarParameters.Num() == 0
+	return UniformNumericParameters.Num() == 0
 		&& UniformVectorPreshaders.Num() == 0
 		&& UniformScalarPreshaders.Num() == 0
 		&& UniformExternalTextureParameters.Num() == 0
@@ -259,8 +239,7 @@ bool FUniformExpressionSet::operator==(const FUniformExpressionSet& ReferenceSet
 		}
 	}
 
-	if (UniformScalarParameters.Num() != ReferenceSet.UniformScalarParameters.Num()
-		|| UniformVectorParameters.Num() != ReferenceSet.UniformVectorParameters.Num()
+	if (UniformNumericParameters.Num() != ReferenceSet.UniformNumericParameters.Num()
 		|| UniformScalarPreshaders.Num() != ReferenceSet.UniformScalarPreshaders.Num()
 		|| UniformVectorPreshaders.Num() != ReferenceSet.UniformVectorPreshaders.Num()
 		|| UniformExternalTextureParameters.Num() != ReferenceSet.UniformExternalTextureParameters.Num()
@@ -270,17 +249,9 @@ bool FUniformExpressionSet::operator==(const FUniformExpressionSet& ReferenceSet
 		return false;
 	}
 
-	for (int32 i = 0; i < UniformScalarParameters.Num(); i++)
+	for (int32 i = 0; i < UniformNumericParameters.Num(); i++)
 	{
-		if (UniformScalarParameters[i] != ReferenceSet.UniformScalarParameters[i])
-		{
-			return false;
-		}
-	}
-
-	for (int32 i = 0; i < UniformVectorParameters.Num(); i++)
-	{
-		if (UniformVectorParameters[i] != ReferenceSet.UniformVectorParameters[i])
+		if (UniformNumericParameters[i] != ReferenceSet.UniformNumericParameters[i])
 		{
 			return false;
 		}
@@ -590,64 +561,6 @@ void FMaterialUniformExpression::GetNumberValue(const struct FMaterialRenderCont
 	OutValue = Value.AsLinearColor();
 }
 
-const FMaterialVectorParameterInfo* FUniformExpressionSet::FindVectorParameter(const FHashedMaterialParameterInfo& ParameterInfo) const
-{
-	for (const FMaterialVectorParameterInfo& Parameter : UniformVectorParameters)
-	{
-		if (Parameter.ParameterInfo == ParameterInfo)
-		{
-			return &Parameter;
-		}
-	}
-	return nullptr;
-}
-
-const FMaterialScalarParameterInfo* FUniformExpressionSet::FindScalarParameter(const FHashedMaterialParameterInfo& ParameterInfo) const
-{
-	for (const FMaterialScalarParameterInfo& Parameter : UniformScalarParameters)
-	{
-		if (Parameter.ParameterInfo == ParameterInfo)
-		{
-			return &Parameter;
-		}
-	}
-	return nullptr;
-}
-
-int32 FUniformExpressionSet::FindOrAddScalarParameter(const FHashedMaterialParameterInfo& ParameterInfo, float DefaultValue)
-{
-	for (int32 i = 0; i < UniformScalarParameters.Num(); ++i)
-	{
-		if (UniformScalarParameters[i].ParameterInfo == ParameterInfo)
-		{
-			return i;
-		}
-	}
-
-	const int32 Index = UniformScalarParameters.Num();
-	FMaterialScalarParameterInfo& Parameter = UniformScalarParameters.AddDefaulted_GetRef();
-	Parameter.ParameterInfo = ParameterInfo;
-	Parameter.DefaultValue = DefaultValue;
-	return Index;
-}
-
-int32 FUniformExpressionSet::FindOrAddVectorParameter(const FHashedMaterialParameterInfo& ParameterInfo, const FLinearColor& DefaultValue)
-{
-	for (int32 i = 0; i < UniformVectorParameters.Num(); ++i)
-	{
-		if (UniformVectorParameters[i].ParameterInfo == ParameterInfo)
-		{
-			return i;
-		}
-	}
-
-	const int32 Index = UniformVectorParameters.Num();
-	FMaterialVectorParameterInfo& Parameter = UniformVectorParameters.AddDefaulted_GetRef();
-	Parameter.ParameterInfo = ParameterInfo;
-	Parameter.DefaultValue = DefaultValue;
-	return Index;
-}
-
 int32 FUniformExpressionSet::FindOrAddTextureParameter(EMaterialTextureParameterType Type, const FMaterialTextureParameterInfo& Info)
 {
 	for (int32 i = 0; i < UniformTextureParameters[(int32)Type].Num(); ++i)
@@ -1414,72 +1327,6 @@ bool FMaterialUniformExpressionExternalTextureParameter::IsIdentical(const FMate
 	return ParameterName == Other->ParameterName && Super::IsIdentical(OtherExpression);
 }
 
-void FMaterialScalarParameterInfo::GetGameThreadNumberValue(const UMaterialInterface* SourceMaterialToCopyFrom, float& OutValue) const
-{
-	check(IsInGameThread());
-	checkSlow(SourceMaterialToCopyFrom);
-
-	const UMaterialInterface* It = SourceMaterialToCopyFrom;
-
-	for (;;)
-	{
-		const UMaterialInstance* MatInst = Cast<UMaterialInstance>(It);
-
-		if (MatInst)
-		{
-			const FScalarParameterValue* ParameterValue = GameThread_FindParameterByName(MatInst->ScalarParameterValues, ParameterInfo);
-			if (ParameterValue)
-			{
-				OutValue = ParameterValue->ParameterValue;
-				break;
-			}
-
-			// go up the hierarchy
-			It = MatInst->Parent;
-		}
-		else
-		{
-			// we reached the base material
-			// get the copy form the base material
-			GetDefaultValue(OutValue);
-			break;
-		}
-	}
-}
-
-void FMaterialVectorParameterInfo::GetGameThreadNumberValue(const UMaterialInterface* SourceMaterialToCopyFrom, FLinearColor& OutValue) const
-{
-	check(IsInGameThread());
-	checkSlow(SourceMaterialToCopyFrom);
-
-	const UMaterialInterface* It = SourceMaterialToCopyFrom;
-
-	for (;;)
-	{
-		const UMaterialInstance* MatInst = Cast<UMaterialInstance>(It);
-
-		if (MatInst)
-		{
-			const FVectorParameterValue* ParameterValue = GameThread_FindParameterByName(MatInst->VectorParameterValues, ParameterInfo);
-			if (ParameterValue)
-			{
-				OutValue = ParameterValue->ParameterValue;
-				break;
-			}
-
-			// go up the hierarchy
-			It = MatInst->Parent;
-		}
-		else
-		{
-			// we reached the base material
-			// get the copy form the base material
-			GetDefaultValue(OutValue);
-			break;
-		}
-	}
-}
-
 void FMaterialTextureParameterInfo::GetGameThreadTextureValue(const UMaterialInterface* MaterialInterface, const FMaterial& Material, UTexture*& OutValue) const
 {
 	if (!ParameterInfo.Name.IsNone())
@@ -1587,8 +1434,7 @@ class FMaterialUniformExpressionRuntimeVirtualTextureParameter_DEPRECATED : publ
 
 IMPLEMENT_MATERIALUNIFORMEXPRESSION_TYPE(FMaterialUniformExpressionTexture);
 IMPLEMENT_MATERIALUNIFORMEXPRESSION_TYPE(FMaterialUniformExpressionConstant);
-IMPLEMENT_MATERIALUNIFORMEXPRESSION_TYPE(FMaterialUniformExpressionVectorParameter);
-IMPLEMENT_MATERIALUNIFORMEXPRESSION_TYPE(FMaterialUniformExpressionScalarParameter);
+IMPLEMENT_MATERIALUNIFORMEXPRESSION_TYPE(FMaterialUniformExpressionGenericParameter);
 IMPLEMENT_MATERIALUNIFORMEXPRESSION_TYPE(FMaterialUniformExpressionTextureParameter);
 IMPLEMENT_MATERIALUNIFORMEXPRESSION_TYPE(FMaterialUniformExpressionExternalTextureBase);
 IMPLEMENT_MATERIALUNIFORMEXPRESSION_TYPE(FMaterialUniformExpressionExternalTexture);
diff --git a/Engine/Source/Runtime/Engine/Private/Materials/MaterialUniformExpressions.h b/Engine/Source/Runtime/Engine/Private/Materials/MaterialUniformExpressions.h
--- a/Engine/Source/Runtime/Engine/Private/Materials/MaterialUniformExpressions.h
+++ b/Engine/Source/Runtime/Engine/Private/Materials/MaterialUniformExpressions.h
@@ -218,13 +218,13 @@ private:
 
 /**
  */
-class FMaterialUniformExpressionVectorParameter: public FMaterialUniformExpression
+class FMaterialUniformExpressionGenericParameter : public FMaterialUniformExpression
 {
-	DECLARE_MATERIALUNIFORMEXPRESSION_TYPE(FMaterialUniformExpressionVectorParameter);
+	DECLARE_MATERIALUNIFORMEXPRESSION_TYPE(FMaterialUniformExpressionGenericParameter);
 public:
 
-	FMaterialUniformExpressionVectorParameter() {}
-	FMaterialUniformExpressionVectorParameter(const FMaterialParameterInfo& InParameterInfo, int32 InParameterIndex)
+	FMaterialUniformExpressionGenericParameter() {}
+	FMaterialUniformExpressionGenericParameter(const FMaterialParameterInfo& InParameterInfo, int32 InParameterIndex)
 		: ParameterInfo(InParameterInfo)
 		, ParameterIndex(InParameterIndex)
 	{
@@ -234,7 +234,7 @@ public:
 	// FMaterialUniformExpression interface.
 	virtual void WriteNumberOpcodes(UE::Shader::FPreshaderData& OutData) const override
 	{
-		OutData.WriteOpcode(UE::Shader::EPreshaderOpcode::VectorParameter);
+		OutData.WriteOpcode(UE::Shader::EPreshaderOpcode::Parameter);
 		OutData.Write((uint16)ParameterIndex);
 	}
 
@@ -259,62 +259,9 @@ public:
 		{
 			return false;
 		}
-		FMaterialUniformExpressionVectorParameter* OtherParameter = (FMaterialUniformExpressionVectorParameter*)OtherExpression;
+		FMaterialUniformExpressionGenericParameter* OtherParameter = (FMaterialUniformExpressionGenericParameter*)OtherExpression;
 		return ParameterInfo == OtherParameter->ParameterInfo && ParameterIndex == OtherParameter->ParameterIndex;
 	}
-
-private:
-	FHashedMaterialParameterInfo ParameterInfo;
-	int32 ParameterIndex;
-};
-
-/**
- */
-class FMaterialUniformExpressionScalarParameter: public FMaterialUniformExpression
-{
-	DECLARE_MATERIALUNIFORMEXPRESSION_TYPE(FMaterialUniformExpressionScalarParameter);
-public:
-
-	FMaterialUniformExpressionScalarParameter() {}
-	FMaterialUniformExpressionScalarParameter(const FMaterialParameterInfo& InParameterInfo, int32 InParameterIndex)
-		: ParameterInfo(InParameterInfo)
-		, ParameterIndex(InParameterIndex)
-	{
-		check(InParameterIndex >= 0 && InParameterIndex <= 0xffff);
-	}
-
-	// FMaterialUniformExpression interface.
-	virtual void WriteNumberOpcodes(UE::Shader::FPreshaderData& OutData) const override
-	{
-		OutData.WriteOpcode(UE::Shader::EPreshaderOpcode::ScalarParameter);
-		OutData.Write((uint16)ParameterIndex);
-	}
-
-	virtual bool IsConstant() const
-	{
-		return false;
-	}
-
-	const FHashedMaterialParameterInfo& GetParameterInfo() const
-	{
-		return ParameterInfo;
-	}
-
-	FName GetParameterName() const
-	{
-		return ParameterInfo.GetName();
-	}
-
-	virtual bool IsIdentical(const FMaterialUniformExpression* OtherExpression) const
-	{
-		if (GetType() != OtherExpression->GetType())
-		{
-			return false;
-		}
-		FMaterialUniformExpressionScalarParameter* OtherParameter = (FMaterialUniformExpressionScalarParameter*)OtherExpression;
-		return ParameterInfo == OtherParameter->ParameterInfo && ParameterIndex == OtherParameter->ParameterIndex;
-	}
-
 private:
 	FHashedMaterialParameterInfo ParameterInfo;
 	int32 ParameterIndex;
diff --git a/Engine/Source/Runtime/Engine/Private/Shader/Preshader.cpp b/Engine/Source/Runtime/Engine/Private/Shader/Preshader.cpp
--- a/Engine/Source/Runtime/Engine/Private/Shader/Preshader.cpp
+++ b/Engine/Source/Runtime/Engine/Private/Shader/Preshader.cpp
@@ -14,6 +14,19 @@ void UE::Shader::FPreshaderData::WriteData(const void* Value, uint32 Size)
 	Data.Append((uint8*)Value, Size);
 }
 
+template<>
+UE::Shader::FPreshaderData& UE::Shader::FPreshaderData::Write<UE::Shader::FValue>(const FValue& Value)
+{
+	const EValueType Type = Value.GetType();
+	FMemoryImageValue MemoryValue = Value.AsMemoryImage();
+	Data.Add((uint8)Type);
+	if (MemoryValue.Size > 0u)
+	{
+		Data.Append(MemoryValue.Bytes, MemoryValue.Size);
+	}
+	return *this;
+}
+
 void UE::Shader::FPreshaderData::WriteName(const FScriptName& Name)
 {
 	int32 Index = Names.Find(Name);
@@ -102,74 +115,61 @@ FHashedMaterialParameterInfo ReadPreshaderValue<FHashedMaterialParameterInfo>(UE
 	return FHashedMaterialParameterInfo(Name, Association, Index);
 }
 
-static void GetVectorParameter(const FUniformExpressionSet& UniformExpressionSet, uint32 ParameterIndex, const FMaterialRenderContext& Context, UE::Shader::FValue& OutValue)
+template<>
+UE::Shader::FValue ReadPreshaderValue<UE::Shader::FValue>(UE::Shader::FPreshaderDataContext& RESTRICT Data)
 {
-	const FMaterialVectorParameterInfo& Parameter = UniformExpressionSet.GetVectorParameter(ParameterIndex);
+	const UE::Shader::EValueType Type = (UE::Shader::EValueType)ReadPreshaderValue<uint8>(Data);
+	uint32 Size = 0u;
+	UE::Shader::FValue Result = UE::Shader::FValue::FromMemoryImage(Type, Data.Ptr, &Size);
+	Data.Ptr += Size;
+	checkSlow(Data.Ptr <= Data.EndPtr);
+	return Result;
+}
 
-	FLinearColor ParameterValue(ForceInitToZero);
-	bool bNeedsDefaultValue = false;
-	if (!Context.MaterialRenderProxy || !Context.MaterialRenderProxy->GetVectorValue(Parameter.ParameterInfo, &ParameterValue, Context))
-	{
-		if (Context.Material.HasMaterialLayers())
-		{
-			UMaterialInterface* Interface = Context.Material.GetMaterialInterface();
-			if (!Interface || !Interface->GetVectorParameterDefaultValue(Parameter.ParameterInfo, ParameterValue))
-			{
-				bNeedsDefaultValue = true;
-			}
-		}
-		else
-		{
-			bNeedsDefaultValue = true;
-		}
-	}
+static void EvaluateParameter(const FUniformExpressionSet& UniformExpressionSet, uint32 ParameterIndex, const FMaterialRenderContext& Context, UE::Shader::FValue& OutValue)
+{
+	const FMaterialNumericParameterInfo& Parameter = UniformExpressionSet.GetNumericParameter(ParameterIndex);
+	bool bFoundParameter = false;
 
-	if (bNeedsDefaultValue)
+	// First allow proxy the chance to override parameter
+	if (Context.MaterialRenderProxy)
 	{
-#if WITH_EDITOR
-		if (!Context.Material.TransientOverrides.GetVectorOverride(Parameter.ParameterInfo, ParameterValue))
-#endif // WITH_EDITOR
+		FMaterialParameterValue ParameterValue;
+		if (Context.MaterialRenderProxy->GetParameterValue(Parameter.ParameterType, Parameter.ParameterInfo, ParameterValue, Context))
 		{
-			Parameter.GetDefaultValue(ParameterValue);
+			OutValue = ParameterValue.AsShaderValue();
+			bFoundParameter = true;
 		}
 	}
 
-	OutValue = ParameterValue;
-}
-
-static void GetScalarParameter(const FUniformExpressionSet& UniformExpressionSet, uint32 ParameterIndex, const FMaterialRenderContext& Context, UE::Shader::FValue& OutValue)
-{
-	const FMaterialScalarParameterInfo& Parameter = UniformExpressionSet.GetScalarParameter(ParameterIndex);
-
-	float ParameterValue = 0.0f;
-	bool bNeedsDefaultValue = false;
-	if (!Context.MaterialRenderProxy || !Context.MaterialRenderProxy->GetScalarValue(Parameter.ParameterInfo, &ParameterValue, Context))
+	// Allow interface to override default value, only needed if using material layers
+	if (!bFoundParameter && Context.Material.HasMaterialLayers())
 	{
-		if (Context.Material.HasMaterialLayers())
+		UMaterialInterface* Interface = Context.Material.GetMaterialInterface();
+		if (Interface)
 		{
-			UMaterialInterface* Interface = Context.Material.GetMaterialInterface();
-			if (!Interface || !Interface->GetScalarParameterDefaultValue(Parameter.ParameterInfo, ParameterValue))
+			FMaterialParameterMetadata ParameterValue;
+			if (Interface->GetParameterDefaultValue(Parameter.ParameterType, Parameter.ParameterInfo, ParameterValue))
 			{
-				bNeedsDefaultValue = true;
+				OutValue = ParameterValue.Value.AsShaderValue();
+				bFoundParameter = true;
 			}
 		}
-		else
-		{
-			bNeedsDefaultValue = true;
-		}
 	}
 
-	if (bNeedsDefaultValue)
-	{
+	// Editor overrides
 #if WITH_EDITOR
-		if (!Context.Material.TransientOverrides.GetScalarOverride(Parameter.ParameterInfo, ParameterValue))
-#endif // WITH_EDITOR
-		{
-			Parameter.GetDefaultValue(ParameterValue);
-		}
+	if (!bFoundParameter)
+	{
+		bFoundParameter = Context.Material.TransientOverrides.GetNumericOverride(Parameter.ParameterType, Parameter.ParameterInfo, OutValue);
 	}
+#endif // WITH_EDITOR
 
-	OutValue = ParameterValue;
+	// Default value
+	if (!bFoundParameter)
+	{
+		OutValue = UniformExpressionSet.GetDefaultParameterValue(Parameter.ParameterType, Parameter.DefaultValueOffset);
+	}
 }
 
 template<typename Operation>
@@ -339,13 +339,9 @@ void UE::Shader::EvaluatePreshader(const FUniformExpressionSet* UniformExpressio
 		case EPreshaderOpcode::Constant:
 			Stack.Add(ReadPreshaderValue<UE::Shader::FValue>(Data));
 			break;
-		case EPreshaderOpcode::VectorParameter:
-			check(UniformExpressionSet);
-			GetVectorParameter(*UniformExpressionSet, ReadPreshaderValue<uint16>(Data), Context, Stack.AddDefaulted_GetRef());
-			break;
-		case EPreshaderOpcode::ScalarParameter:
+		case EPreshaderOpcode::Parameter:
 			check(UniformExpressionSet);
-			GetScalarParameter(*UniformExpressionSet, ReadPreshaderValue<uint16>(Data), Context, Stack.AddDefaulted_GetRef());
+			EvaluateParameter(*UniformExpressionSet, ReadPreshaderValue<uint16>(Data), Context, Stack.AddDefaulted_GetRef());
 			break;
 		case EPreshaderOpcode::Add: EvaluateBinaryOp(Stack, UE::Shader::Add); break;
 		case EPreshaderOpcode::Sub: EvaluateBinaryOp(Stack, UE::Shader::Sub); break;
diff --git a/Engine/Source/Runtime/Engine/Private/Shader/ShaderValue.cpp b/Engine/Source/Runtime/Engine/Private/Shader/ShaderValue.cpp
--- a/Engine/Source/Runtime/Engine/Private/Shader/ShaderValue.cpp
+++ b/Engine/Source/Runtime/Engine/Private/Shader/ShaderValue.cpp
@@ -47,7 +47,7 @@ struct FCastBool
 		{
 		case EValueComponentType::Float: return Component.Float != 0.0f;
 		case EValueComponentType::Int: return Component.Int != 0;
-		case EValueComponentType::Bool: return Component.Bool;
+		case EValueComponentType::Bool: return Component.AsBool();
 		default: return false;
 		}
 	}
@@ -84,6 +84,44 @@ void CastValue(const Operation& Op, const FValue& Value, ResultType& OutResult)
 } // namespace Shader
 } // namespace UE
 
+UE::Shader::FValue UE::Shader::FValue::FromMemoryImage(EValueType Type, const void* Data, uint32* OutSizeInBytes)
+{
+	FValue Result(Type);
+	const uint8* Bytes = static_cast<const uint8*>(Data);
+	const uint32 ComponentSizeInBytes = GetComponentTypeSizeInBytes(Result.ComponentType);
+	if (ComponentSizeInBytes > 0u)
+	{
+		for (int32 i = 0u; i < Result.NumComponents; ++i)
+		{
+			FMemory::Memcpy(&Result.Component[i].Packed, Bytes, ComponentSizeInBytes);
+			Bytes += ComponentSizeInBytes;
+		}
+	}
+	if (OutSizeInBytes)
+	{
+		*OutSizeInBytes = (uint32)(Bytes - static_cast<const uint8*>(Data));
+	}
+	return Result;
+}
+
+UE::Shader::FMemoryImageValue UE::Shader::FValue::AsMemoryImage() const
+{
+	FMemoryImageValue Result;
+	uint8* Bytes = Result.Bytes;
+	const uint32 ComponentSizeInBytes = GetComponentTypeSizeInBytes(ComponentType);
+	if (ComponentSizeInBytes > 0u)
+	{
+		for (int32 i = 0u; i < NumComponents; ++i)
+		{
+			FMemory::Memcpy(Bytes, &Component[i].Packed, ComponentSizeInBytes);
+			Bytes += ComponentSizeInBytes;
+		}
+	}
+	Result.Size = (uint32)(Bytes - Result.Bytes);
+	check(Result.Size <= FMemoryImageValue::MaxSize);
+	return Result;
+}
+
 UE::Shader::FFloatValue UE::Shader::FValue::AsFloat() const
 {
 	FFloatValue Result;
@@ -111,6 +149,13 @@ UE::Shader::FBoolValue UE::Shader::FValue::AsBool() const
 	return Result;
 }
 
+float UE::Shader::FValue::AsFloatScalar() const
+{
+	FFloatValue Result;
+	Private::CastValue(Private::FCastFloat(), *this, Result);
+	return Result[0];
+}
+
 bool UE::Shader::FValue::AsBoolScalar() const
 {
 	const FBoolValue Result = AsBool();
@@ -124,6 +169,19 @@ bool UE::Shader::FValue::AsBoolScalar() const
 	return false;
 }
 
+uint32 UE::Shader::GetComponentTypeSizeInBytes(EValueComponentType Type)
+{
+	switch (Type)
+	{
+	case EValueComponentType::Void: return 0u;
+	case EValueComponentType::Float: return sizeof(float);
+	case EValueComponentType::Int: return sizeof(int32);
+	case EValueComponentType::Bool: return 1u;
+	case EValueComponentType::MaterialAttributes: return 0u;
+	default: checkNoEntry(); return 0u;
+	}
+}
+
 UE::Shader::FValueTypeDescription UE::Shader::GetValueTypeDescription(EValueType Type)
 {
 	switch (Type)
@@ -442,6 +500,26 @@ bool UE::Shader::operator==(const FValue& Lhs, const FValue& Rhs)
 	return true;
 }
 
+uint32 UE::Shader::GetTypeHash(const FValue& Value)
+{
+	uint32 Result = ::GetTypeHash(Value.ComponentType);
+	Result = HashCombine(Result, ::GetTypeHash(Value.NumComponents));
+	for (int32 i = 0u; i < Value.NumComponents; ++i)
+	{
+		const FValueComponent& Component = Value.Component[i];
+		uint32 ComponentHash = 0u;
+		switch (Value.ComponentType)
+		{
+		case EValueComponentType::Float: ComponentHash = ::GetTypeHash(Component.Float); break;
+		case EValueComponentType::Int: ComponentHash = ::GetTypeHash(Component.Int); break;
+		case EValueComponentType::Bool: ComponentHash = ::GetTypeHash(Component.Bool); break;
+		default: checkNoEntry(); break;
+		}
+		Result = HashCombine(Result, ComponentHash);
+	}
+	return Result;
+}
+
 UE::Shader::FValue UE::Shader::Abs(const FValue& Value)
 {
 	return Private::UnaryOp(Private::FOpAbs(), Value);
diff --git a/Engine/Source/Runtime/Engine/Private/ShaderDerivedDataVersion.h b/Engine/Source/Runtime/Engine/Private/ShaderDerivedDataVersion.h
--- a/Engine/Source/Runtime/Engine/Private/ShaderDerivedDataVersion.h
+++ b/Engine/Source/Runtime/Engine/Private/ShaderDerivedDataVersion.h
@@ -11,4 +11,4 @@
 // In case of merge conflicts with DDC versions, you MUST generate a new GUID and set this new
 // guid as version
 #define GLOBALSHADERMAP_DERIVEDDATA_VER			TEXT("60CD2243B26E480D898BAA7C66094C0A")
-#define MATERIALSHADERMAP_DERIVEDDATA_VER		TEXT("2831FE97589D4E34A316BE837CCA0039")
\ No newline at end of file
+#define MATERIALSHADERMAP_DERIVEDDATA_VER		TEXT("81444FC0-337140F8A38670E251CDEF0C")
\ No newline at end of file
diff --git a/Engine/Source/Runtime/Engine/Public/EditorSupportDelegates.h b/Engine/Source/Runtime/Engine/Public/EditorSupportDelegates.h
--- a/Engine/Source/Runtime/Engine/Public/EditorSupportDelegates.h
+++ b/Engine/Source/Runtime/Engine/Public/EditorSupportDelegates.h
@@ -6,6 +6,16 @@
 
 #if WITH_EDITOR
 
+namespace UE
+{
+namespace Shader
+{
+struct FValue;
+}
+}
+
+enum class EMaterialParameterType : uint8;
+
 /** 
  * FEditorSupportDelegates
  * Delegates that are needed for proper editor functionality, but are accessed or triggered in engine code.
@@ -22,10 +32,8 @@ struct ENGINE_API FEditorSupportDelegates
 	DECLARE_MULTICAST_DELEGATE_TwoParams(FOnWindowsMessage, class FViewport*, uint32);
 	/** delegate type for material usage flags change events ( Params: UMaterial* material, int32 FlagThatChanged ) */
 	DECLARE_MULTICAST_DELEGATE_TwoParams(FOnMaterialUsageFlagsChanged, class UMaterial*, int32); 
-	/** delegate type for vector parameter default change event */
-	DECLARE_MULTICAST_DELEGATE_ThreeParams(FOnVectorParameterDefaultChanged, class UMaterialExpression*, FName, const FLinearColor&); 
-	/** delegate type for scalar parameter default change event */
-	DECLARE_MULTICAST_DELEGATE_ThreeParams(FOnScalarParameterDefaultChanged, class UMaterialExpression*, FName, float); 
+	/** delegate type for numeric parameter default change event */
+	DECLARE_MULTICAST_DELEGATE_FourParams(FOnNumericParameterDefaultChanged, class UMaterialExpression*, EMaterialParameterType, FName, const UE::Shader::FValue&);
 
 	/** Called when all viewports need to be redrawn */
 	static FSimpleMulticastDelegate RedrawAllViewports;
@@ -47,10 +55,8 @@ struct ENGINE_API FEditorSupportDelegates
 	static FOnWindowsMessage PostWindowsMessage;
 	/** Sent after the usages flags on a material have changed*/
 	static FOnMaterialUsageFlagsChanged MaterialUsageFlagsChanged;
-	/** Sent after vector param default changed */
-	static FOnVectorParameterDefaultChanged VectorParameterDefaultChanged;
-	/** Sent after scalar param default changed */
-	static FOnScalarParameterDefaultChanged ScalarParameterDefaultChanged;
+	/** Sent after numeric param default changed */
+	static FOnNumericParameterDefaultChanged NumericParameterDefaultChanged;
 };
 
 #endif // WITH_EDITOR
diff --git a/Engine/Source/Runtime/Engine/Public/MaterialCompiler.h b/Engine/Source/Runtime/Engine/Public/MaterialCompiler.h
--- a/Engine/Source/Runtime/Engine/Public/MaterialCompiler.h
+++ b/Engine/Source/Runtime/Engine/Public/MaterialCompiler.h
@@ -143,8 +143,10 @@ public:
 	virtual int32 GetCurrentFunctionStackDepth() = 0;
 
 	virtual int32 AccessCollectionParameter(UMaterialParameterCollection* ParameterCollection, int32 ParameterIndex, int32 ComponentIndex) = 0;	
-	virtual int32 ScalarParameter(FName ParameterName, float DefaultValue) = 0;
-	virtual int32 VectorParameter(FName ParameterName, const FLinearColor& DefaultValue) = 0;
+	virtual int32 GenericParameter(EMaterialParameterType ParameterType, FName ParameterName, const UE::Shader::FValue& DefaultValue) = 0;
+
+	ENGINE_API int32 ScalarParameter(FName ParameterName, float DefaultValue);
+	ENGINE_API int32 VectorParameter(FName ParameterName, const FLinearColor& DefaultValue);
 
 	virtual int32 Constant(float X) = 0;
 	virtual int32 Constant2(float X,float Y) = 0;
@@ -524,8 +526,7 @@ public:
 	virtual int32 ForceCast(int32 Code, EMaterialValueType DestType, uint32 ForceCastFlags = 0) override { return Compiler->ForceCast(Code, DestType, ForceCastFlags); }
 
 	virtual int32 AccessCollectionParameter(UMaterialParameterCollection* ParameterCollection, int32 ParameterIndex, int32 ComponentIndex) override { return Compiler->AccessCollectionParameter(ParameterCollection, ParameterIndex, ComponentIndex); }
-	virtual int32 ScalarParameter(FName ParameterName, float DefaultValue) override { return Compiler->ScalarParameter(ParameterName, DefaultValue); }
-	virtual int32 VectorParameter(FName ParameterName, const FLinearColor& DefaultValue) override { return Compiler->VectorParameter(ParameterName, DefaultValue); }
+	virtual int32 GenericParameter(EMaterialParameterType ParameterType, FName ParameterName, const UE::Shader::FValue& DefaultValue) override { return Compiler->GenericParameter(ParameterType, ParameterName, DefaultValue); }
 
 	virtual int32 Constant(float X) override { return Compiler->Constant(X); }
 	virtual int32 Constant2(float X, float Y) override { return Compiler->Constant2(X, Y); }
diff --git a/Engine/Source/Runtime/Engine/Public/MaterialShared.h b/Engine/Source/Runtime/Engine/Public/MaterialShared.h
--- a/Engine/Source/Runtime/Engine/Public/MaterialShared.h
+++ b/Engine/Source/Runtime/Engine/Public/MaterialShared.h
@@ -330,48 +330,23 @@ public:
 	LAYOUT_FIELD(uint32, OpcodeSize);
 };
 
-class FMaterialScalarParameterInfo
+class FMaterialNumericParameterInfo
 {
-	DECLARE_TYPE_LAYOUT(FMaterialScalarParameterInfo, NonVirtual);
+	DECLARE_TYPE_LAYOUT(FMaterialNumericParameterInfo, NonVirtual);
 public:
-	friend inline bool operator==(const FMaterialScalarParameterInfo& Lhs, const FMaterialScalarParameterInfo& Rhs)
-	{
-		return Lhs.ParameterInfo == Rhs.ParameterInfo && Lhs.DefaultValue == Rhs.DefaultValue;
-	}
-	friend inline bool operator!=(const FMaterialScalarParameterInfo& Lhs, const FMaterialScalarParameterInfo& Rhs)
-	{
-		return !operator==(Lhs, Rhs);
-	}
-
-	void GetGameThreadNumberValue(const UMaterialInterface* SourceMaterialToCopyFrom, float& OutValue) const;
-	void GetDefaultValue(float& OutValue) const { OutValue = DefaultValue; }
-	
-	LAYOUT_FIELD(FHashedMaterialParameterInfo, ParameterInfo);
-	LAYOUT_FIELD(float, DefaultValue);
-};
 
-
-class FMaterialVectorParameterInfo
-{
-	DECLARE_TYPE_LAYOUT(FMaterialVectorParameterInfo, NonVirtual);
-public:
-
-	FMaterialVectorParameterInfo() : DefaultValue(ForceInitToZero) {}
-
-	friend inline bool operator==(const FMaterialVectorParameterInfo& Lhs, const FMaterialVectorParameterInfo& Rhs)
+	friend inline bool operator==(const FMaterialNumericParameterInfo& Lhs, const FMaterialNumericParameterInfo& Rhs)
 	{
-		return Lhs.ParameterInfo == Rhs.ParameterInfo && Lhs.DefaultValue == Rhs.DefaultValue;
+		return Lhs.ParameterInfo == Rhs.ParameterInfo && Lhs.ParameterType == Rhs.ParameterType && Lhs.DefaultValueOffset == Rhs.DefaultValueOffset;
 	}
-	friend inline bool operator!=(const FMaterialVectorParameterInfo& Lhs, const FMaterialVectorParameterInfo& Rhs)
+	friend inline bool operator!=(const FMaterialNumericParameterInfo& Lhs, const FMaterialNumericParameterInfo& Rhs)
 	{
 		return !operator==(Lhs, Rhs);
 	}
 
-	void GetGameThreadNumberValue(const UMaterialInterface* SourceMaterialToCopyFrom, FLinearColor& OutValue) const;
-	void GetDefaultValue(FLinearColor& OutValue) const { OutValue = DefaultValue; }
-
 	LAYOUT_FIELD(FHashedMaterialParameterInfo, ParameterInfo);
-	LAYOUT_FIELD(FLinearColor, DefaultValue);
+	LAYOUT_FIELD(EMaterialParameterType, ParameterType);
+	LAYOUT_FIELD(uint32, DefaultValueOffset);
 };
 
 /** Must invalidate ShaderVersion.ush when changing */
@@ -435,19 +410,34 @@ public:
 class FUniformParameterOverrides
 {
 public:
-	void SetScalarOverride(const FHashedMaterialParameterInfo& ParameterInfo, float Value, bool bOverride);
-	void SetVectorOverride(const FHashedMaterialParameterInfo& ParameterInfo, const FLinearColor& Value, bool bOverride);
-
-	bool GetScalarOverride(const FHashedMaterialParameterInfo& ParameterInfo, float& OutValue) const;
-	bool GetVectorOverride(const FHashedMaterialParameterInfo& ParameterInfo, FLinearColor& OutValue) const;
+	void SetNumericOverride(EMaterialParameterType Type, const FHashedMaterialParameterInfo& ParameterInfo, const UE::Shader::FValue& Value, bool bOverride);
+	bool GetNumericOverride(EMaterialParameterType Type, const FHashedMaterialParameterInfo& ParameterInfo, UE::Shader::FValue& OutValue) const;
 
 	void SetTextureOverride(EMaterialTextureParameterType Type, const FHashedMaterialParameterInfo& ParameterInfo, UTexture* Texture);
 	UTexture* GetTextureOverride_GameThread(EMaterialTextureParameterType Type, const FHashedMaterialParameterInfo& ParameterInfo) const;
 	UTexture* GetTextureOverride_RenderThread(EMaterialTextureParameterType Type, const FHashedMaterialParameterInfo& ParameterInfo) const;
 
 private:
-	TMap<FHashedMaterialParameterInfo, float> ScalarOverrides;
-	TMap<FHashedMaterialParameterInfo, FLinearColor> VectorOverrides;
+	struct FNumericParameterKey
+	{
+		FHashedMaterialParameterInfo ParameterInfo;
+		EMaterialParameterType ParameterType;
+
+		friend inline bool operator==(const FNumericParameterKey& Lhs, const FNumericParameterKey& Rhs)
+		{
+			return Lhs.ParameterInfo == Rhs.ParameterInfo && Lhs.ParameterType == Rhs.ParameterType;
+		}
+		friend inline bool operator!=(const FNumericParameterKey& Lhs, const FNumericParameterKey& Rhs)
+		{
+			return !operator==(Lhs, Rhs);
+		}
+		friend inline uint32 GetTypeHash(const FNumericParameterKey& Value)
+		{
+			return HashCombine(GetTypeHash(Value.ParameterInfo), GetTypeHash(Value.ParameterType));
+		}
+	};
+
+	TMap<FNumericParameterKey, UE::Shader::FValue> NumericOverrides;
 	TMap<FHashedMaterialParameterInfo, UTexture*> GameThreadTextureOverides[NumMaterialTextureParameterTypes];
 	TMap<FHashedMaterialParameterInfo, UTexture*> RenderThreadTextureOverrides[NumMaterialTextureParameterTypes];
 };
@@ -482,21 +472,19 @@ public:
 		return UniformBufferLayoutInitializer;
 	}
 
-	inline const FMaterialVectorParameterInfo& GetVectorParameter(uint32 Index) const { return UniformVectorParameters[Index]; }
-	inline const FMaterialScalarParameterInfo& GetScalarParameter(uint32 Index) const { return UniformScalarParameters[Index]; }
-	inline const FMaterialTextureParameterInfo& GetTextureParameter(EMaterialTextureParameterType Type, int32 Index) const { return UniformTextureParameters[(uint32)Type][Index]; }
+	UE::Shader::FValue GetDefaultParameterValue(EMaterialParameterType Type, uint32 Offset) const
+	{
+		return UE::Shader::FValue::FromMemoryImage(GetShaderValueType(Type), DefaultValues.GetData() + Offset);
+	}
 
-	const FMaterialVectorParameterInfo* FindVectorParameter(const FHashedMaterialParameterInfo& ParameterInfo) const;
-	const FMaterialScalarParameterInfo* FindScalarParameter(const FHashedMaterialParameterInfo& ParameterInfo) const;
+	inline const FMaterialNumericParameterInfo& GetNumericParameter(uint32 Index) const { return UniformNumericParameters[Index]; }
+	inline const FMaterialTextureParameterInfo& GetTextureParameter(EMaterialTextureParameterType Type, int32 Index) const { return UniformTextureParameters[(uint32)Type][Index]; }
 
 	inline int32 GetNumTextures(EMaterialTextureParameterType Type) const { return UniformTextureParameters[(uint32)Type].Num(); }
 	ENGINE_API void GetGameThreadTextureValue(EMaterialTextureParameterType Type, int32 Index, const UMaterialInterface* MaterialInterface, const FMaterial& Material, UTexture*& OutValue, bool bAllowOverride = true) const;
 	ENGINE_API void GetTextureValue(EMaterialTextureParameterType Type, int32 Index, const FMaterialRenderContext& Context, const FMaterial& Material, const UTexture*& OutValue) const;
 	ENGINE_API void GetTextureValue(int32 Index, const FMaterialRenderContext& Context, const FMaterial& Material, const URuntimeVirtualTexture*& OutValue) const;
 
-	int32 FindOrAddScalarParameter(const FHashedMaterialParameterInfo& ParameterInfo, float DefaultValue);
-	int32 FindOrAddVectorParameter(const FHashedMaterialParameterInfo& ParameterInfo, const FLinearColor& DefaultValue);
-
 	int32 FindOrAddTextureParameter(EMaterialTextureParameterType Type, const FMaterialTextureParameterInfo& Info);
 
 protected:
@@ -525,12 +513,12 @@ protected:
 
 	LAYOUT_FIELD(TMemoryImageArray<FMaterialUniformPreshaderHeader>, UniformVectorPreshaders);
 	LAYOUT_FIELD(TMemoryImageArray<FMaterialUniformPreshaderHeader>, UniformScalarPreshaders);
-	LAYOUT_FIELD(TMemoryImageArray<FMaterialScalarParameterInfo>, UniformScalarParameters);
-	LAYOUT_FIELD(TMemoryImageArray<FMaterialVectorParameterInfo>, UniformVectorParameters);
+	LAYOUT_FIELD(TMemoryImageArray<FMaterialNumericParameterInfo>, UniformNumericParameters);
 	LAYOUT_ARRAY(TMemoryImageArray<FMaterialTextureParameterInfo>, UniformTextureParameters, NumMaterialTextureParameterTypes);
 	LAYOUT_FIELD(TMemoryImageArray<FMaterialExternalTextureParameterInfo>, UniformExternalTextureParameters);
 
 	LAYOUT_FIELD(UE::Shader::FPreshaderData, UniformPreshaderData);
+	LAYOUT_FIELD(TMemoryImageArray<uint8>, DefaultValues);
 
 	/** Virtual texture stacks found during compilation */
 	LAYOUT_FIELD(TMemoryImageArray<FMaterialVirtualTextureStack>, VTStacks);
@@ -1831,8 +1819,7 @@ public:
 	// Accessors.
 	ENGINE_API const FUniformExpressionSet& GetUniformExpressions() const;
 	ENGINE_API TArrayView<const FMaterialTextureParameterInfo> GetUniformTextureExpressions(EMaterialTextureParameterType Type) const;
-	ENGINE_API TArrayView<const FMaterialVectorParameterInfo> GetUniformVectorParameterExpressions() const;
-	ENGINE_API TArrayView<const FMaterialScalarParameterInfo> GetUniformScalarParameterExpressions() const;
+	ENGINE_API TArrayView<const FMaterialNumericParameterInfo> GetUniformNumericParameterExpressions() const;
 
 	inline TArrayView<const FMaterialTextureParameterInfo> GetUniform2DTextureExpressions() const { return GetUniformTextureExpressions(EMaterialTextureParameterType::Standard2D); }
 	inline TArrayView<const FMaterialTextureParameterInfo> GetUniformCubeTextureExpressions() const { return GetUniformTextureExpressions(EMaterialTextureParameterType::Cube); }
@@ -2339,10 +2326,13 @@ public:
 	}
 
 	virtual UMaterialInterface* GetMaterialInterface() const { return NULL; }
-	virtual bool GetVectorValue(const FHashedMaterialParameterInfo& ParameterInfo, FLinearColor* OutValue, const FMaterialRenderContext& Context) const = 0;
-	virtual bool GetScalarValue(const FHashedMaterialParameterInfo& ParameterInfo, float* OutValue, const FMaterialRenderContext& Context) const = 0;
-	virtual bool GetTextureValue(const FHashedMaterialParameterInfo& ParameterInfo,const UTexture** OutValue, const FMaterialRenderContext& Context) const = 0;
-	virtual bool GetTextureValue(const FHashedMaterialParameterInfo& ParameterInfo, const URuntimeVirtualTexture** OutValue, const FMaterialRenderContext& Context) const = 0;
+
+	ENGINE_API bool GetVectorValue(const FHashedMaterialParameterInfo& ParameterInfo, FLinearColor* OutValue, const FMaterialRenderContext& Context) const;
+	ENGINE_API bool GetScalarValue(const FHashedMaterialParameterInfo& ParameterInfo, float* OutValue, const FMaterialRenderContext& Context) const;
+	ENGINE_API bool GetTextureValue(const FHashedMaterialParameterInfo& ParameterInfo, const UTexture** OutValue, const FMaterialRenderContext& Context) const;
+	ENGINE_API bool GetTextureValue(const FHashedMaterialParameterInfo& ParameterInfo, const URuntimeVirtualTexture** OutValue, const FMaterialRenderContext& Context) const;
+	virtual bool GetParameterValue(EMaterialParameterType Type, const FHashedMaterialParameterInfo& ParameterInfo, FMaterialParameterValue& OutValue, const FMaterialRenderContext& Context) const = 0;
+
 	bool IsDeleted() const
 	{
 		return DeletedFlag != 0;
@@ -2427,10 +2417,7 @@ public:
 	// FMaterialRenderProxy interface.
 	ENGINE_API virtual const FMaterial* GetMaterialNoFallback(ERHIFeatureLevel::Type InFeatureLevel) const override;
 	ENGINE_API virtual const FMaterialRenderProxy* GetFallback(ERHIFeatureLevel::Type InFeatureLevel) const override;
-	ENGINE_API virtual bool GetVectorValue(const FHashedMaterialParameterInfo& ParameterInfo, FLinearColor* OutValue, const FMaterialRenderContext& Context) const override;
-	ENGINE_API virtual bool GetScalarValue(const FHashedMaterialParameterInfo& ParameterInfo, float* OutValue, const FMaterialRenderContext& Context) const override;
-	ENGINE_API virtual bool GetTextureValue(const FHashedMaterialParameterInfo& ParameterInfo,const UTexture** OutValue, const FMaterialRenderContext& Context) const override;
-	ENGINE_API virtual bool GetTextureValue(const FHashedMaterialParameterInfo& ParameterInfo, const URuntimeVirtualTexture** OutValue, const FMaterialRenderContext& Context) const;
+	ENGINE_API virtual bool GetParameterValue(EMaterialParameterType Type, const FHashedMaterialParameterInfo& ParameterInfo, FMaterialParameterValue& OutValue, const FMaterialRenderContext& Context) const override;
 };
 
 
@@ -2451,7 +2438,7 @@ public:
 		TextureParamName(InTextureParamName)
 	{}
 
-	ENGINE_API virtual bool GetTextureValue(const FHashedMaterialParameterInfo& ParameterInfo, const UTexture** OutValue, const FMaterialRenderContext& Context) const override;
+	ENGINE_API virtual bool GetParameterValue(EMaterialParameterType Type, const FHashedMaterialParameterInfo& ParameterInfo, FMaterialParameterValue& OutValue, const FMaterialRenderContext& Context) const override;
 };
 
 
@@ -2475,10 +2462,7 @@ public:
 	// FMaterialRenderProxy interface.
 	ENGINE_API virtual const FMaterial* GetMaterialNoFallback(ERHIFeatureLevel::Type InFeatureLevel) const override;
 	ENGINE_API virtual const FMaterialRenderProxy* GetFallback(ERHIFeatureLevel::Type InFeatureLevel) const override;
-	ENGINE_API virtual bool GetVectorValue(const FHashedMaterialParameterInfo& ParameterInfo, FLinearColor* OutValue, const FMaterialRenderContext& Context) const override;
-	ENGINE_API virtual bool GetScalarValue(const FHashedMaterialParameterInfo& ParameterInfo, float* OutValue, const FMaterialRenderContext& Context) const override;
-	ENGINE_API virtual bool GetTextureValue(const FHashedMaterialParameterInfo& ParameterInfo, const UTexture** OutValue, const FMaterialRenderContext& Context) const override;
-	ENGINE_API virtual bool GetTextureValue(const FHashedMaterialParameterInfo& ParameterInfo, const URuntimeVirtualTexture** OutValue, const FMaterialRenderContext& Context) const;
+	ENGINE_API virtual bool GetParameterValue(EMaterialParameterType Type, const FHashedMaterialParameterInfo& ParameterInfo, FMaterialParameterValue& OutValue, const FMaterialRenderContext& Context) const override;
 };
 
 
@@ -2497,7 +2481,7 @@ public:
 	{}
 
 	// FMaterialRenderProxy interface.
-	virtual bool GetVectorValue(const FHashedMaterialParameterInfo& ParameterInfo, FLinearColor* OutValue, const FMaterialRenderContext& Context) const override;
+	virtual bool GetParameterValue(EMaterialParameterType Type, const FHashedMaterialParameterInfo& ParameterInfo, FMaterialParameterValue& OutValue, const FMaterialRenderContext& Context) const override;
 };
 
 /**
diff --git a/Engine/Source/Runtime/Engine/Public/MaterialTypes.h b/Engine/Source/Runtime/Engine/Public/MaterialTypes.h
--- a/Engine/Source/Runtime/Engine/Public/MaterialTypes.h
+++ b/Engine/Source/Runtime/Engine/Public/MaterialTypes.h
@@ -5,7 +5,9 @@
 #include "CoreMinimal.h"
 #include "Misc/Guid.h"
 #include "Misc/EnumClassFlags.h"
+#include "Serialization/MemoryLayout.h"
 #include "UObject/SoftObjectPtr.h"
+#include "Shader/ShaderTypes.h"
 #include "MaterialTypes.generated.h"
 
 class UTexture;
@@ -31,6 +33,7 @@ enum class EMaterialParameterType : uint8
 	Num,
 	None = 0xff,
 };
+DECLARE_INTRINSIC_TYPE_LAYOUT(EMaterialParameterType);
 
 static const int32 NumMaterialParameterTypes = (int32)EMaterialParameterType::Num;
 static const int32 NumMaterialRuntimeParameterTypes = (int32)EMaterialParameterType::NumRuntime;
@@ -48,6 +51,8 @@ static inline bool IsStaticMaterialParameter(EMaterialParameterType InType)
 	}
 }
 
+extern ENGINE_API UE::Shader::EValueType GetShaderValueType(EMaterialParameterType Type);
+
 enum class EMaterialGetParameterValueFlags : uint32
 {
 	None = 0u,
@@ -93,19 +98,31 @@ struct FMaterialParameterValue
 	FMaterialParameterValue() : Type(EMaterialParameterType::None) {}
 	FMaterialParameterValue(float InValue) : Type(EMaterialParameterType::Scalar) { Float[0] = InValue; }
 	FMaterialParameterValue(const FLinearColor& InValue) : Type(EMaterialParameterType::Vector) { Float[0] = InValue.R; Float[1] = InValue.G; Float[2] = InValue.B; Float[3] = InValue.A; }
+	FMaterialParameterValue(const FVector3f& InValue) : Type(EMaterialParameterType::Vector) { Float[0] = InValue.X; Float[1] = InValue.Y; Float[2] = InValue.Z; Float[3] = 0.0f; }
 	FMaterialParameterValue(UTexture* InValue) : Type(EMaterialParameterType::Texture) { Texture = InValue; }
 	FMaterialParameterValue(const TObjectPtr<UTexture>& InValue) : Type(EMaterialParameterType::Texture) { Texture = InValue; }
 	FMaterialParameterValue(URuntimeVirtualTexture* InValue) : Type(EMaterialParameterType::RuntimeVirtualTexture) { RuntimeVirtualTexture = InValue; }
 	FMaterialParameterValue(const TObjectPtr<URuntimeVirtualTexture>& InValue) : Type(EMaterialParameterType::RuntimeVirtualTexture) { RuntimeVirtualTexture = InValue; }
+
+	// Gamethread parameters are typically non-const, but renderthread parameters are const
+	// Would be possible to store an additional const-flag member, and provide runtime checks to ensure constness is not violated...maybe worth doing in the future
+	FMaterialParameterValue(const UTexture* InValue) : Type(EMaterialParameterType::Texture) { Texture = const_cast<UTexture*>(InValue); }
+	FMaterialParameterValue(const TObjectPtr<const UTexture>& InValue) : Type(EMaterialParameterType::Texture) { Texture = const_cast<UTexture*>(InValue.Get()); }
+	FMaterialParameterValue(const URuntimeVirtualTexture* InValue) : Type(EMaterialParameterType::RuntimeVirtualTexture) { RuntimeVirtualTexture = const_cast<URuntimeVirtualTexture*>(InValue); }
+	FMaterialParameterValue(const TObjectPtr<const URuntimeVirtualTexture>& InValue) : Type(EMaterialParameterType::RuntimeVirtualTexture) { RuntimeVirtualTexture = const_cast<URuntimeVirtualTexture*>(InValue.Get()); }
+
 	FMaterialParameterValue(UFont* InValue, int32 InPage) : Type(EMaterialParameterType::Font) { Font.Value = InValue; Font.Page = InPage; }
 	FMaterialParameterValue(bool InValue) : Type(EMaterialParameterType::StaticSwitch) { Bool[0] = InValue; }
 	FMaterialParameterValue(const FStaticComponentMaskValue& InValue) : Type(EMaterialParameterType::StaticComponentMask) { Bool[0] = InValue.R; Bool[1] = InValue.G; Bool[2] = InValue.B; Bool[3] = InValue.A; }
 	FMaterialParameterValue(bool bMaskR, bool bMaskG, bool bMaskB, bool bMaskA) : Type(EMaterialParameterType::StaticComponentMask) { Bool[0] = bMaskR; Bool[1] = bMaskG; Bool[2] = bMaskB; Bool[3] = bMaskA; }
+	
+	ENGINE_API FMaterialParameterValue(EMaterialParameterType Type, const UE::Shader::FValue& InValue);
 
 	inline float AsScalar() const { check(Type == EMaterialParameterType::Scalar); return Float[0]; }
 	inline FLinearColor AsLinearColor() const { check(Type == EMaterialParameterType::Vector); return FLinearColor(Float[0], Float[1], Float[2], Float[3]); }
 	inline bool AsStaticSwitch() const { check(Type == EMaterialParameterType::StaticSwitch); return Bool[0]; }
 	inline FStaticComponentMaskValue AsStaticComponentMask() const { check(Type == EMaterialParameterType::StaticComponentMask); return FStaticComponentMaskValue(Bool[0], Bool[1], Bool[2], Bool[3]); }
+	ENGINE_API UE::Shader::FValue AsShaderValue() const;
 
 	union
 	{
diff --git a/Engine/Source/Runtime/Engine/Public/Shader/Preshader.h b/Engine/Source/Runtime/Engine/Public/Shader/Preshader.h
--- a/Engine/Source/Runtime/Engine/Public/Shader/Preshader.h
+++ b/Engine/Source/Runtime/Engine/Public/Shader/Preshader.h
@@ -19,8 +19,7 @@ enum class EPreshaderOpcode : uint8
 	Nop,
 	ConstantZero,
 	Constant,
-	ScalarParameter,
-	VectorParameter,
+	Parameter,
 	Add,
 	Sub,
 	Mul,
@@ -89,6 +88,9 @@ public:
 	template<>
 	FPreshaderData& Write<FScriptName>(const FScriptName& Value) { WriteName(Value); return *this; }
 
+	template<>
+	FPreshaderData& Write<FValue>(const FValue& Value);
+
 	/** Can't write FName, use FScriptName instead */
 	template<>
 	FPreshaderData& Write<FName>(const FName& Value) = delete;
diff --git a/Engine/Source/Runtime/Engine/Public/Shader/ShaderTypes.h b/Engine/Source/Runtime/Engine/Public/Shader/ShaderTypes.h
--- a/Engine/Source/Runtime/Engine/Public/Shader/ShaderTypes.h
+++ b/Engine/Source/Runtime/Engine/Public/Shader/ShaderTypes.h
@@ -2,6 +2,7 @@
 #pragma once
 
 #include "CoreMinimal.h"
+#include "Serialization/MemoryLayout.h"
 
 namespace UE
 {
@@ -49,6 +50,7 @@ enum class EValueType : uint8
 	MaterialAttributes,
 };
 
+uint32 GetComponentTypeSizeInBytes(EValueComponentType Type);
 FValueTypeDescription GetValueTypeDescription(EValueType Type);
 EValueType MakeValueType(EValueComponentType ComponentType, int32 NumComponents);
 EValueType MakeValueType(EValueType BaseType, int32 NumComponents);
@@ -67,14 +69,24 @@ using FFloatValue = TValue<float>;
 using FIntValue = TValue<int32>;
 using FBoolValue = TValue<bool>;
 
+struct FMemoryImageValue
+{
+	static const uint32 MaxSize = sizeof(double) * 4;
+	uint8 Bytes[MaxSize];
+	uint32 Size;
+};
+
 union FValueComponent
 {
 	FValueComponent() : Packed(0u) {}
 
+	// 'Bool' is stored as uint8 to avoid changing on different compilers
+	bool AsBool() const { return Bool != 0u; }
+
 	uint32 Packed;
 	float Float;
 	int32 Int;
-	bool Bool;
+	uint8 Bool;
 };
 static_assert(sizeof(FValueComponent) == sizeof(uint32), "bad packing");
 
@@ -151,6 +163,14 @@ struct FValue
 		Component[0].Bool = v;
 	}
 
+	inline FValue(bool X, bool Y, bool Z, bool W) : ComponentType(EValueComponentType::Bool), NumComponents(4)
+	{
+		Component[0].Bool = X;
+		Component[1].Bool = Y;
+		Component[2].Bool = Z;
+		Component[3].Bool = W;
+	}
+
 	inline FValue(int32 v) : ComponentType(EValueComponentType::Int), NumComponents(1)
 	{
 		Component[0].Int = v;
@@ -160,11 +180,16 @@ struct FValue
 
 	inline const FValueComponent& GetComponent(int32 i) const { check(i >= 0 && i < NumComponents); return Component[i]; }
 
+	static FValue FromMemoryImage(EValueType Type, const void* Data, uint32* OutSizeInBytes = nullptr);
+
+	FMemoryImageValue AsMemoryImage() const;
+
 	FFloatValue AsFloat() const;
 	FIntValue AsInt() const;
 	FBoolValue AsBool() const;
 
 	FLinearColor AsLinearColor() const;
+	float AsFloatScalar() const;
 	bool AsBoolScalar() const;
 
 	EValueComponentType ComponentType;
@@ -174,6 +199,7 @@ struct FValue
 
 ENGINE_API bool operator==(const FValue& Lhs, const FValue& Rhs);
 inline bool operator!=(const FValue& Lhs, const FValue& Rhs) { return !operator==(Lhs, Rhs); }
+ENGINE_API uint32 GetTypeHash(const FValue& Value);
 
 ENGINE_API FValue Abs(const FValue& Value);
 ENGINE_API FValue Saturate(const FValue& Value);
@@ -208,3 +234,6 @@ ENGINE_API FValue Cross(const FValue& Lhs, const FValue& Rhs);
 ENGINE_API FValue Append(const FValue& Lhs, const FValue& Rhs);
 }
 }
+
+DECLARE_INTRINSIC_TYPE_LAYOUT(UE::Shader::EValueType);
+DECLARE_INTRINSIC_TYPE_LAYOUT(UE::Shader::EValueComponentType);
diff --git a/Engine/Source/Runtime/Landscape/Private/LandscapeGizmoActor.cpp b/Engine/Source/Runtime/Landscape/Private/LandscapeGizmoActor.cpp
--- a/Engine/Source/Runtime/Landscape/Private/LandscapeGizmoActor.cpp
+++ b/Engine/Source/Runtime/Landscape/Private/LandscapeGizmoActor.cpp
@@ -57,68 +57,62 @@ public:
 		return Parent->GetFallback(InFeatureLevel);
 	}
 
-	virtual bool GetVectorValue(const FHashedMaterialParameterInfo& ParameterInfo, FLinearColor* OutValue, const FMaterialRenderContext& Context) const override
+	virtual bool GetParameterValue(EMaterialParameterType Type, const FHashedMaterialParameterInfo& ParameterInfo, FMaterialParameterValue& OutValue, const FMaterialRenderContext& Context) const
 	{
-		if (ParameterInfo.Name == FName(TEXT("AlphaScaleBias")))
+		switch (Type)
 		{
-			*OutValue = ScaleBias;
-			return true;
-		}
-		else
-		if (ParameterInfo.Name == FName(TEXT("MatrixRow1")))
-		{
-			*OutValue = FLinearColor(WorldToLandscapeMatrix.M[0][0], WorldToLandscapeMatrix.M[0][1], WorldToLandscapeMatrix.M[0][2],WorldToLandscapeMatrix.M[0][3]);
-			return true;
-		}
-		else
-		if (ParameterInfo.Name == FName(TEXT("MatrixRow2")))
-		{
-			*OutValue = FLinearColor(WorldToLandscapeMatrix.M[1][0], WorldToLandscapeMatrix.M[1][1], WorldToLandscapeMatrix.M[1][2],WorldToLandscapeMatrix.M[1][3]);
-			return true;
-		}
-		else
-		if (ParameterInfo.Name == FName(TEXT("MatrixRow3")))
-		{
-			*OutValue = FLinearColor(WorldToLandscapeMatrix.M[2][0], WorldToLandscapeMatrix.M[2][1], WorldToLandscapeMatrix.M[2][2],WorldToLandscapeMatrix.M[2][3]);
-			return true;
-		}
-		else
-		if (ParameterInfo.Name == FName(TEXT("MatrixRow4")))
-		{
-			*OutValue = FLinearColor(WorldToLandscapeMatrix.M[3][0], WorldToLandscapeMatrix.M[3][1], WorldToLandscapeMatrix.M[3][2],WorldToLandscapeMatrix.M[3][3]);
-			return true;
-		}
-
-		return Parent->GetVectorValue(ParameterInfo, OutValue, Context);
-	}
-	virtual bool GetScalarValue(const FHashedMaterialParameterInfo& ParameterInfo, float* OutValue, const FMaterialRenderContext& Context) const override
-	{
-		if (ParameterInfo.Name == FName(TEXT("Top")))
-		{
-			*OutValue = TopHeight;
-			return true;
-		}
-		else if (ParameterInfo.Name == FName(TEXT("Bottom")))
-		{
-			*OutValue = BottomHeight;
-			return true;
-		}
-		return Parent->GetScalarValue(ParameterInfo, OutValue, Context);
-	}
-	virtual bool GetTextureValue(const FHashedMaterialParameterInfo& ParameterInfo,const UTexture** OutValue, const FMaterialRenderContext& Context) const override
-	{
-		if (ParameterInfo.Name == FName(TEXT("AlphaTexture")))
-		{
-			// FIXME: This needs to return a black texture if AlphaTexture is NULL.
-			// Returning NULL will cause the material to use GWhiteTexture.
-			*OutValue = AlphaTexture;
-			return true;
+		case EMaterialParameterType::Vector:
+			if (ParameterInfo.Name == FName(TEXT("AlphaScaleBias")))
+			{
+				OutValue = ScaleBias;
+				return true;
+			}
+			else if (ParameterInfo.Name == FName(TEXT("MatrixRow1")))
+			{
+				OutValue = FLinearColor(WorldToLandscapeMatrix.M[0][0], WorldToLandscapeMatrix.M[0][1], WorldToLandscapeMatrix.M[0][2], WorldToLandscapeMatrix.M[0][3]);
+				return true;
+			}
+			else if (ParameterInfo.Name == FName(TEXT("MatrixRow2")))
+			{
+				OutValue = FLinearColor(WorldToLandscapeMatrix.M[1][0], WorldToLandscapeMatrix.M[1][1], WorldToLandscapeMatrix.M[1][2], WorldToLandscapeMatrix.M[1][3]);
+				return true;
+			}
+			else if (ParameterInfo.Name == FName(TEXT("MatrixRow3")))
+			{
+				OutValue = FLinearColor(WorldToLandscapeMatrix.M[2][0], WorldToLandscapeMatrix.M[2][1], WorldToLandscapeMatrix.M[2][2], WorldToLandscapeMatrix.M[2][3]);
+				return true;
+			}
+			else if (ParameterInfo.Name == FName(TEXT("MatrixRow4")))
+			{
+				OutValue = FLinearColor(WorldToLandscapeMatrix.M[3][0], WorldToLandscapeMatrix.M[3][1], WorldToLandscapeMatrix.M[3][2], WorldToLandscapeMatrix.M[3][3]);
+				return true;
+			}
+			break;
+		case EMaterialParameterType::Scalar:
+			if (ParameterInfo.Name == FName(TEXT("Top")))
+			{
+				OutValue = TopHeight;
+				return true;
+			}
+			else if (ParameterInfo.Name == FName(TEXT("Bottom")))
+			{
+				OutValue = BottomHeight;
+				return true;
+			}
+			break;
+		case EMaterialParameterType::Texture:
+			if (ParameterInfo.Name == FName(TEXT("AlphaTexture")))
+			{
+				// FIXME: This needs to return a black texture if AlphaTexture is NULL.
+				// Returning NULL will cause the material to use GWhiteTexture.
+				OutValue = AlphaTexture;
+				return true;
+			}
+			break;
+		default:
+			break;
 		}
-		return Parent->GetTextureValue(ParameterInfo, OutValue, Context);
-	}
-	virtual bool GetTextureValue(const FHashedMaterialParameterInfo& ParameterInfo, const URuntimeVirtualTexture** OutValue, const FMaterialRenderContext& Context) const
-	{
-		return Parent->GetTextureValue(ParameterInfo, OutValue, Context);
+		return Parent->GetParameterValue(Type, ParameterInfo, OutValue, Context);
 	}
 };
 
diff --git a/Engine/Source/Runtime/Landscape/Public/LandscapeRender.h b/Engine/Source/Runtime/Landscape/Public/LandscapeRender.h
--- a/Engine/Source/Runtime/Landscape/Public/LandscapeRender.h
+++ b/Engine/Source/Runtime/Landscape/Public/LandscapeRender.h
@@ -1024,59 +1024,48 @@ public:
 		return Parent->GetFallback(InFeatureLevel);
 	}
 
-	virtual bool GetVectorValue(const FHashedMaterialParameterInfo& ParameterInfo, FLinearColor* OutValue, const FMaterialRenderContext& Context) const override
+	virtual bool GetParameterValue(EMaterialParameterType Type, const FHashedMaterialParameterInfo& ParameterInfo, FMaterialParameterValue& OutValue, const FMaterialRenderContext& Context) const
 	{
-		if (ParameterInfo.Name == FName(TEXT("Landscape_RedMask")))
+		switch (Type)
 		{
-			*OutValue = R;
-			return true;
-		}
-		else if (ParameterInfo.Name == FName(TEXT("Landscape_GreenMask")))
-		{
-			*OutValue = G;
-			return true;
-		}
-		else if (ParameterInfo.Name == FName(TEXT("Landscape_BlueMask")))
-		{
-			*OutValue = B;
-			return true;
-		}
-		else
-		{
-			return Parent->GetVectorValue(ParameterInfo, OutValue, Context);
-		}
-	}
-	virtual bool GetScalarValue(const FHashedMaterialParameterInfo& ParameterInfo, float* OutValue, const FMaterialRenderContext& Context) const override
-	{
-		return Parent->GetScalarValue(ParameterInfo, OutValue, Context);
-	}
-	virtual bool GetTextureValue(const FHashedMaterialParameterInfo& ParameterInfo, const UTexture** OutValue, const FMaterialRenderContext& Context) const override
-	{
-		// NOTE: These should be returning black textures when NULL. The material will
-		// use a white texture if they are.
-		if (ParameterInfo.Name == FName(TEXT("Landscape_RedTexture")))
-		{
-			*OutValue = RedTexture;
-			return true;
-		}
-		else if (ParameterInfo.Name == FName(TEXT("Landscape_GreenTexture")))
-		{
-			*OutValue = GreenTexture;
-			return true;
-		}
-		else if (ParameterInfo.Name == FName(TEXT("Landscape_BlueTexture")))
-		{
-			*OutValue = BlueTexture;
-			return true;
-		}
-		else
-		{
-			return Parent->GetTextureValue(ParameterInfo, OutValue, Context);
+		case EMaterialParameterType::Vector:
+			if (ParameterInfo.Name == FName(TEXT("Landscape_RedMask")))
+			{
+				OutValue = R;
+				return true;
+			}
+			else if (ParameterInfo.Name == FName(TEXT("Landscape_GreenMask")))
+			{
+				OutValue = G;
+				return true;
+			}
+			else if (ParameterInfo.Name == FName(TEXT("Landscape_BlueMask")))
+			{
+				OutValue = B;
+				return true;
+			}
+			break;
+		case EMaterialParameterType::Texture:
+			if (ParameterInfo.Name == FName(TEXT("Landscape_RedTexture")))
+			{
+				OutValue = RedTexture;
+				return true;
+			}
+			else if (ParameterInfo.Name == FName(TEXT("Landscape_GreenTexture")))
+			{
+				OutValue = GreenTexture;
+				return true;
+			}
+			else if (ParameterInfo.Name == FName(TEXT("Landscape_BlueTexture")))
+			{
+				OutValue = BlueTexture;
+				return true;
+			}
+			break;
+		default:
+			break;
 		}
-	}
-	virtual bool GetTextureValue(const FHashedMaterialParameterInfo& ParameterInfo, const URuntimeVirtualTexture** OutValue, const FMaterialRenderContext& Context) const
-	{
-		return Parent->GetTextureValue(ParameterInfo, OutValue, Context);
+		return Parent->GetParameterValue(Type, ParameterInfo, OutValue, Context);
 	}
 };
 
@@ -1101,37 +1090,29 @@ public:
 	{
 		return Parent->GetFallback(InFeatureLevel);
 	}
-	virtual bool GetVectorValue(const FHashedMaterialParameterInfo& ParameterInfo, FLinearColor* OutValue, const FMaterialRenderContext& Context) const override
-	{
-		if (ParameterInfo.Name == FName(TEXT("HighlightColor")))
-		{
-			*OutValue = FLinearColor(1.f, 0.5f, 0.5f);
-			return true;
-		}
-		else
-		{
-			return Parent->GetVectorValue(ParameterInfo, OutValue, Context);
-		}
-	}
-	virtual bool GetScalarValue(const FHashedMaterialParameterInfo& ParameterInfo, float* OutValue, const FMaterialRenderContext& Context) const override
-	{
-		return Parent->GetScalarValue(ParameterInfo, OutValue, Context);
-	}
-	virtual bool GetTextureValue(const FHashedMaterialParameterInfo& ParameterInfo, const UTexture** OutValue, const FMaterialRenderContext& Context) const override
+
+	virtual bool GetParameterValue(EMaterialParameterType Type, const FHashedMaterialParameterInfo& ParameterInfo, FMaterialParameterValue& OutValue, const FMaterialRenderContext& Context) const
 	{
-		if (ParameterInfo.Name == FName(TEXT("SelectedData")))
+		switch (Type)
 		{
-			*OutValue = SelectTexture;
-			return true;
-		}
-		else
-		{
-			return Parent->GetTextureValue(ParameterInfo, OutValue, Context);
+		case EMaterialParameterType::Vector:
+			if (ParameterInfo.Name == FName(TEXT("HighlightColor")))
+			{
+				OutValue = FLinearColor(1.f, 0.5f, 0.5f);
+				return true;
+			}
+			break;
+		case EMaterialParameterType::Texture:
+			if (ParameterInfo.Name == FName(TEXT("SelectedData")))
+			{
+				OutValue = SelectTexture;
+				return true;
+			}
+			break;
+		default:
+			break;
 		}
-	}
-	virtual bool GetTextureValue(const FHashedMaterialParameterInfo& ParameterInfo, const URuntimeVirtualTexture** OutValue, const FMaterialRenderContext& Context) const
-	{
-		return Parent->GetTextureValue(ParameterInfo, OutValue, Context);
+		return Parent->GetParameterValue(Type, ParameterInfo, OutValue, Context);
 	}
 };
 
@@ -1158,34 +1139,29 @@ public:
 	{
 		return Parent->GetFallback(InFeatureLevel);
 	}
-	virtual bool GetVectorValue(const FHashedMaterialParameterInfo& ParameterInfo, FLinearColor* OutValue, const FMaterialRenderContext& Context) const override
-	{
-		return Parent->GetVectorValue(ParameterInfo, OutValue, Context);
-	}
-	virtual bool GetScalarValue(const FHashedMaterialParameterInfo& ParameterInfo, float* OutValue, const FMaterialRenderContext& Context) const override
-	{
-		if (ParameterInfo.Name == FName(TEXT("bInverted")))
-		{
-			*OutValue = bInverted;
-			return true;
-		}
-		return Parent->GetScalarValue(ParameterInfo, OutValue, Context);
-	}
-	virtual bool GetTextureValue(const FHashedMaterialParameterInfo& ParameterInfo, const UTexture** OutValue, const FMaterialRenderContext& Context) const override
+	
+	virtual bool GetParameterValue(EMaterialParameterType Type, const FHashedMaterialParameterInfo& ParameterInfo, FMaterialParameterValue& OutValue, const FMaterialRenderContext& Context) const
 	{
-		if (ParameterInfo.Name == FName(TEXT("SelectedData")))
-		{
-			*OutValue = SelectTexture;
-			return true;
-		}
-		else
+		switch (Type)
 		{
-			return Parent->GetTextureValue(ParameterInfo, OutValue, Context);
+		case EMaterialParameterType::Scalar:
+			if (ParameterInfo.Name == FName(TEXT("bInverted")))
+			{
+				OutValue = (float)bInverted;
+				return true;
+			}
+			break;
+		case EMaterialParameterType::Texture:
+			if (ParameterInfo.Name == FName(TEXT("SelectedData")))
+			{
+				OutValue = SelectTexture;
+				return true;
+			}
+			break;
+		default:
+			break;
 		}
-	}
-	virtual bool GetTextureValue(const FHashedMaterialParameterInfo& ParameterInfo, const URuntimeVirtualTexture** OutValue, const FMaterialRenderContext& Context) const
-	{
-		return Parent->GetTextureValue(ParameterInfo, OutValue, Context);
+		return Parent->GetParameterValue(Type, ParameterInfo, OutValue, Context);
 	}
 };
 
@@ -1213,7 +1189,8 @@ public:
 	{
 		return Parent->GetFallback(InFeatureLevel);
 	}
-	virtual bool GetVectorValue(const FHashedMaterialParameterInfo& ParameterInfo, FLinearColor* OutValue, const FMaterialRenderContext& Context) const override
+	
+	virtual bool GetParameterValue(EMaterialParameterType Type, const FHashedMaterialParameterInfo& ParameterInfo, FMaterialParameterValue& OutValue, const FMaterialRenderContext& Context) const
 	{
 		static FName ColorNames[] =
 		{
@@ -1235,41 +1212,39 @@ public:
 			FName(TEXT("Color15"))
 		};
 
-		for (int32 i = 0; i < UE_ARRAY_COUNT(ColorNames) && i < LayerColors.Num(); i++)
+		switch (Type)
 		{
-			if (ParameterInfo.Name == ColorNames[i])
+		case EMaterialParameterType::Vector:
+			for (int32 i = 0; i < UE_ARRAY_COUNT(ColorNames) && i < LayerColors.Num(); i++)
+			{
+				if (ParameterInfo.Name == ColorNames[i])
+				{
+					OutValue = LayerColors[i];
+					return true;
+				}
+			}
+			break;
+		case EMaterialParameterType::Scalar:
+			if (ParameterInfo.Name == FName(TEXT("Rotation")))
 			{
-				*OutValue = LayerColors[i];
+				OutValue = Rotation;
 				return true;
 			}
+			else if (ParameterInfo.Name == FName(TEXT("NumStripes")))
+			{
+				OutValue = (float)LayerColors.Num();
+				return true;
+			}
+			else if (ParameterInfo.Name == FName(TEXT("ComponentSizeVerts")))
+			{
+				OutValue = (float)ComponentSizeVerts;
+				return true;
+			}
+			break;
+		default:
+			break;
 		}
-		return Parent->GetVectorValue(ParameterInfo, OutValue, Context);
-	}
-	virtual bool GetScalarValue(const FHashedMaterialParameterInfo& ParameterInfo, float* OutValue, const FMaterialRenderContext& Context) const override
-	{
-		if (ParameterInfo.Name == FName(TEXT("Rotation")))
-		{
-			*OutValue = Rotation;
-			return true;
-		}
-		if (ParameterInfo.Name == FName(TEXT("NumStripes")))
-		{
-			*OutValue = LayerColors.Num();
-			return true;
-		}
-		if (ParameterInfo.Name == FName(TEXT("ComponentSizeVerts")))
-		{
-			*OutValue = ComponentSizeVerts;
-			return true;
-		}		
-		return Parent->GetScalarValue(ParameterInfo, OutValue, Context);
-	}
-	virtual bool GetTextureValue(const FHashedMaterialParameterInfo& ParameterInfo, const UTexture** OutValue, const FMaterialRenderContext& Context) const override
-	{
-		return Parent->GetTextureValue(ParameterInfo, OutValue, Context);
-	}
-	virtual bool GetTextureValue(const FHashedMaterialParameterInfo& ParameterInfo, const URuntimeVirtualTexture** OutValue, const FMaterialRenderContext& Context) const
-	{
-		return Parent->GetTextureValue(ParameterInfo, OutValue, Context);
+
+		return Parent->GetParameterValue(Type, ParameterInfo, OutValue, Context);
 	}
 };
