PlayerState - Introduce GetPingInMilliseconds() and expose it to blueprints.  Does the appropriate expansion of CompressedPing * 4.

diff --git a/Engine/Source/Runtime/Engine/Classes/GameFramework/PlayerState.h b/Engine/Source/Runtime/Engine/Classes/GameFramework/PlayerState.h
--- a/Engine/Source/Runtime/Engine/Classes/GameFramework/PlayerState.h
+++ b/Engine/Source/Runtime/Engine/Classes/GameFramework/PlayerState.h
@@ -360,7 +360,14 @@ public:
 	/** Sets the value of PlayerId without causing other side effects to this instance. */
 	void SetPlayerId(const int32 NewId);
 
-	/** Gets the literal value of Ping. */
+	/** Gets the value of ping in milliseconds. */
+	UFUNCTION(BlueprintCallable, Category = "PlayerState")
+	int32 GetPingInMilliseconds() const
+	{
+		return Ping * 4;
+	}
+
+	/** Gets the literal value of Ping.  This is the compressed ping, sent through replication (real ping / 4) */
 	uint8 GetPing() const
 	{
 		return Ping;
