Improved HLOD runtime grid computation to fix case where a parent grid has the same cell size as its child grid

diff --git a/Engine/Source/Runtime/Engine/Private/WorldPartition/RuntimeSpatialHash/RuntimeSpatialHashHLOD.cpp b/Engine/Source/Runtime/Engine/Private/WorldPartition/RuntimeSpatialHash/RuntimeSpatialHashHLOD.cpp
--- a/Engine/Source/Runtime/Engine/Private/WorldPartition/RuntimeSpatialHash/RuntimeSpatialHashHLOD.cpp
+++ b/Engine/Source/Runtime/Engine/Private/WorldPartition/RuntimeSpatialHash/RuntimeSpatialHashHLOD.cpp
@@ -11,6 +11,7 @@
 #include "WorldPartition/HLOD/HLODActorDesc.h"
 #include "WorldPartition/WorldPartitionHelpers.h"
 
+#include "UObject/GCObjectScopeGuard.h"
 #include "Engine/Engine.h"
 #include "HAL/PlatformFileManager.h"
 #include "Misc/PackageName.h"
@@ -190,7 +191,7 @@ static TArray<FGuid> GenerateHLODsForGrid(UWorldPartition* WorldPartition, const
 
 				TRACE_CPUPROFILER_EVENT_SCOPE_TEXT(*CellName.ToString());
 
-				UE_LOG(LogWorldPartitionRuntimeSpatialHashHLOD, Display, TEXT("[%d / %d] Creating/updating HLOD actor(s) for cell %s"), (int32)SlowTask.CompletedWork, (int32)SlowTask.TotalAmountOfWork, *CellName.ToString());
+				UE_LOG(LogWorldPartitionRuntimeSpatialHashHLOD, Display, TEXT("[%d / %d] Processing cell %s"), (int32)SlowTask.CompletedWork + 1, (int32)SlowTask.TotalAmountOfWork, *CellName.ToString());
 
 				FHLODCreationParams CreationParams;
 				CreationParams.WorldPartition = WorldPartition;
@@ -269,15 +270,11 @@ static TArray<FGuid> GenerateHLODsForGrid(UWorldPartition* WorldPartition, const
 	return GridHLODActors;
 }
 
-// Find all HLOD grids from the HLODLayer assets we use and build a dependency graph
-static void GatherHLODGrids(UWorldPartition* WorldPartition, TMap<FName, FSpatialHashRuntimeGrid>& OutHLODGrids, TMap<FName, TSet<FName>>& OutHLODGridsGraph)
+// Find all referenced HLODLayer assets
+static TMap<UHLODLayer*, int32> GatherHLODLayers(UWorldPartition* WorldPartition)
 {
-	// We don't know the HLOD level for those grids yet, we'll have to compute a graph of dependencies between our HLOD layers first
-	// Start with level 0, we'll rename once we have the correct info
-	const uint32 InitialHLODLevel = 0;
-
-	// Gather up all HLODLayer referenced by the actors
-	TSet<UHLODLayer*> HLODLayers;
+	// Gather up all HLODLayers referenced by the actors, along with the HLOD level at which it was used
+	TMap<UHLODLayer*, int32> HLODLayersLevel;
 
 	for (FActorDescList::TIterator<> ActorDescIterator(WorldPartition); ActorDescIterator; ++ActorDescIterator)
 	{
@@ -287,127 +284,57 @@ static void GatherHLODGrids(UWorldPartition* WorldPartition, TMap<FName, FSpatia
 			if (ActorDesc.GetActorIsHLODRelevant())
 			{
 				UHLODLayer* HLODLayer = UHLODLayer::GetHLODLayer(ActorDesc, WorldPartition);
-				while (HLODLayer != nullptr)
+
+				// If layer was already encountered, no need to process it again
+				if (!HLODLayersLevel.Contains(HLODLayer))
 				{
-					bool bAlreadyInSet = false;
-					HLODLayers.Add(HLODLayer, &bAlreadyInSet);
-					if (bAlreadyInSet)
+					// Walk up the parent HLOD layers, keep track of HLOD level
+					int32 CurrentHLODLevel = 0;
+					while (HLODLayer != nullptr)
 					{
-						break;
-					}
+						int32& HLODLevel = HLODLayersLevel.FindOrAdd(HLODLayer);
+						HLODLevel = FMath::Max(HLODLevel, CurrentHLODLevel);
 
-					HLODLayer = HLODLayer->GetParentLayer().LoadSynchronous();
+						HLODLayer = HLODLayer->GetParentLayer().LoadSynchronous();
+						CurrentHLODLevel++;
+					}
 				}
 			}
 		}
 	}
 
-	// Find all used HLOD grids & build a dependency graph
-	for (const UHLODLayer* HLODLayer : HLODLayers)
-	{
-		FSpatialHashRuntimeGrid HLODGrid;
-		HLODGrid.CellSize = HLODLayer->GetCellSize();
-		HLODGrid.LoadingRange = HLODLayer->GetLoadingRange();
-		HLODGrid.DebugColor = FLinearColor::Red;
-		HLODGrid.GridName = HLODLayer->GetRuntimeGrid(InitialHLODLevel);
-		HLODGrid.bClientOnlyVisible = true;
-		HLODGrid.HLODLayer = HLODLayer;
-
-		OutHLODGrids.Emplace(HLODGrid.GridName, HLODGrid);
-
-		TSet<FName>& HLODParentGrids = OutHLODGridsGraph.FindOrAdd(HLODGrid.GridName);
-		if (const UHLODLayer* ParentHLODLayer = HLODLayer->GetParentLayer().LoadSynchronous())
-		{
-			HLODParentGrids.Add(ParentHLODLayer->GetRuntimeGrid(InitialHLODLevel));
-		}
-	}
+	return HLODLayersLevel;
 }
 
-// Sort HLOD grids in the order they'll need to be processed for HLOD generation
-static bool SortHLODGrids(const TMap<FName, TSet<FName>>& InHLODGridsGraph, TArray<FName>& OutSortedGrids, TMap<FName, uint32>& OutGridsDepth)
+static TMap<FName, FSpatialHashRuntimeGrid> CreateHLODGrids(TMap<UHLODLayer*, int32>& InHLODLayersLevel)
 {
-	TSet<FName>		ProcessedGridsSet;	// Processed grids
-	TSet<FName>		VisitedGridsSet;	// Visited grids
+	TMap<FName, FSpatialHashRuntimeGrid> HLODGrids;
 
-	TFunction<bool(const FName, uint32)> VisitGraph = [&](const FName GridName, uint32 CurrentDepth) -> bool
+	// Create HLOD runtime grids
+	for (const TPair<UHLODLayer*, int32>& HLODLayerLevel : InHLODLayersLevel)
 	{
-		uint32& GridDepth = OutGridsDepth.FindOrAdd(GridName);
-		GridDepth = FMath::Max(GridDepth, CurrentDepth);
-
-		if (ProcessedGridsSet.Contains(GridName))
-		{
-			return true;
-		}
+		UHLODLayer* HLODLayer = HLODLayerLevel.Key;
 
-		// Detect cyclic dependencies between HLOD grids...
-		if (VisitedGridsSet.Contains(GridName))
+		// No need to create a runtime grid if the HLOD layer is set to be always loaded
+		if (!HLODLayer->IsAlwaysLoaded())
 		{
-			return false; // Not a DAG
-		}
+			int32 HLODLevel = HLODLayerLevel.Value;
 
-		VisitedGridsSet.Add(GridName);
+			FSpatialHashRuntimeGrid HLODGrid;
+			HLODGrid.CellSize = HLODLayer->GetCellSize();
+			HLODGrid.LoadingRange = HLODLayer->GetLoadingRange();
+			HLODGrid.DebugColor = FLinearColor::Red;
+			HLODGrid.GridName = HLODLayer->GetRuntimeGrid(HLODLevel);
+			HLODGrid.bClientOnlyVisible = true;
+			HLODGrid.HLODLayer = HLODLayer;
 
-		for (const FName ParentGridName : InHLODGridsGraph.FindChecked(GridName))
-		{
-			if (!VisitGraph(ParentGridName, CurrentDepth + 1))
-			{
-				return false;
-			}
-		}
-
-		VisitedGridsSet.Remove(GridName);
-
-		ProcessedGridsSet.Add(GridName);
-		OutSortedGrids.Insert(GridName, 0);
-		return true;
-	};
-
-	bool bIsADAG = true;
-	for (const auto& HLODGridEntry : InHLODGridsGraph)
-	{
-		FName GridName = HLODGridEntry.Key;
-		if (!ProcessedGridsSet.Contains(GridName))
-		{
-			const uint32 HLODDepth = 0;
-			bIsADAG = VisitGraph(GridName, HLODDepth);
-			if (!bIsADAG)
-			{
-				break;
-			}
+			HLODGrids.Emplace(HLODGrid.GridName, HLODGrid);
 		}
 	}
 
-	// Remove leafs, those grids are the standard runtime grids, we'll generate HLOD for those in a first pass
-	OutSortedGrids.RemoveAll([&InHLODGridsGraph](const FName GridName) { return !InHLODGridsGraph.Contains(GridName); });
+	HLODGrids.KeySort(FNameLexicalLess());
 
-	return bIsADAG;
-}
-
-static void RenameHLODGrids(TMap<FName, FSpatialHashRuntimeGrid>& HLODGrids, TArray<FName>& SortedGrids, TMap<FName, uint32>& GridsDepth)
-{
-	// Build a mapping of old -> new names
-	TMap<FName, FName> NewNamesMapping;
-	for (const auto& HLODGrid : HLODGrids)
-	{
-		NewNamesMapping.Emplace(HLODGrid.Key, UHLODLayer::GetRuntimeGridName(GridsDepth[HLODGrid.Key], HLODGrid.Value.CellSize, HLODGrid.Value.LoadingRange));
-	}
-
-	// Remplace map entries
-	for (const auto& Names : NewNamesMapping)
-	{
-		FSpatialHashRuntimeGrid RuntimeGrid = HLODGrids.FindAndRemoveChecked(Names.Key);
-		RuntimeGrid.GridName = Names.Value;
-		HLODGrids.Emplace(Names.Value, RuntimeGrid);
-
-		int32 GridDepth = GridsDepth.FindAndRemoveChecked(Names.Key);
-		GridsDepth.Emplace(Names.Value, GridDepth);
-	}
-
-	// Replace arrays entries
-	for (FName& Name : SortedGrids)
-	{
-		Name = NewNamesMapping[Name];
-	}
+	return HLODGrids;
 }
 
 // Create/destroy HLOD grid actors
@@ -482,26 +409,16 @@ bool UWorldPartitionRuntimeSpatialHash::GenerateHLOD(ISourceControlHelper* Sourc
 
 	UWorldPartition* WorldPartition = GetOuterUWorldPartition();
 
-	// Find all used HLOD grids & build a dependency graph
-	check(HLODGrids.IsEmpty());
-	ON_SCOPE_EXIT { HLODGrids.Empty(); };
-
-	TMap<FName, TSet<FName>>			 HLODGridsGraph;
-	GatherHLODGrids(WorldPartition, HLODGrids, HLODGridsGraph);
+	// Find all used HLOD layers
+	TMap<UHLODLayer*, int32> HLODLayersLevels = GatherHLODLayers(WorldPartition);
+	TArray<UHLODLayer*> HLODLayers;
+	HLODLayersLevels.GetKeys(HLODLayers);
 
-	// Sort HLOD grids in the order they'll need to be processed
-	TArray<FName>	SortedGrids;		// HLOD Grids, sorted in the order they'll need to be processed for HLOD generation
-	TMap<FName, uint32>	GridsDepth;		// Depth - Used to obtain the HLOD level
-	bool bIsADAG = SortHLODGrids(HLODGridsGraph, SortedGrids, GridsDepth);
-	if (!bIsADAG)
-	{
-		UE_LOG(LogWorldPartitionRuntimeSpatialHashHLOD, Error, TEXT("Invalid grids setup, cycles detected"));
-		return false;
-	}
-
-	// Now that we computed proper depth per grid, rename the grids
-	RenameHLODGrids(HLODGrids, SortedGrids, GridsDepth);
+	// Keep references to HLODLayers or they may be GC'd
+	TGCObjectsScopeGuard<UHLODLayer> KeepHLODLayersAlive(HLODLayers);
 
+	TMap<FName, FSpatialHashRuntimeGrid> HLODGrids = CreateHLODGrids(HLODLayersLevels);
+	
 	TMap<FName, int32> GridsMapping;
 	GridsMapping.Add(NAME_None, 0);
 	for (int32 i = 0; i < Grids.Num(); i++)
@@ -572,7 +489,7 @@ bool UWorldPartitionRuntimeSpatialHash::GenerateHLOD(ISourceControlHelper* Sourc
 	// Keep track of all valid HLOD actors, along with which runtime grid they live in
 	TMap<FName, TArray<FGuid>> GridsHLODActors;
 
-	auto GenerateHLODs = [&GridsHLODActors, &MainContainerInstance, WorldPartition, &GridsDepth, &Context, SourceControlHelper, bCreateActorsOnly](const FSpatialHashRuntimeGrid& RuntimeGrid, uint32 HLODLevel, const TArray<const FActorClusterInstance*>& ActorClusterInstances)
+	auto GenerateHLODs = [&GridsHLODActors, &MainContainerInstance, WorldPartition, &Context, SourceControlHelper, bCreateActorsOnly](const FSpatialHashRuntimeGrid& RuntimeGrid, uint32 HLODLevel, const TArray<const FActorClusterInstance*>& ActorClusterInstances)
 	{
 		// Generate HLODs for this grid
 		TArray<FGuid> HLODActors = GenerateHLODsForGrid(WorldPartition, MainContainerInstance, RuntimeGrid, HLODLevel, Context, SourceControlHelper, bCreateActorsOnly, ActorClusterInstances);
@@ -597,12 +514,14 @@ bool UWorldPartitionRuntimeSpatialHash::GenerateHLOD(ISourceControlHelper* Sourc
 	}
 
 	// Now, go on and create HLOD actors from HLOD grids (HLOD 1-N)
-	for (const FName HLODGridName : SortedGrids)
+	for (auto It = HLODGrids.CreateIterator(); It; ++It)
 	{
+		const FName HLODGridName = It.Key();
+
 		// No need to process empty grids
 		if (!GridsHLODActors.Contains(HLODGridName))
 		{
-			HLODGrids.Remove(HLODGridName); // No need to keep this grid around, we have no actors in it
+			It.RemoveCurrent(); // No need to keep this grid around, we have no actors in it
 			continue;
 		}
 
@@ -624,7 +543,7 @@ bool UWorldPartitionRuntimeSpatialHash::GenerateHLOD(ISourceControlHelper* Sourc
 
 		// Generate HLODs for this grid - retrieve actors from our ValidHLODActors map
 		// We can't rely on actor descs for newly created HLOD actors
-		GenerateHLODs(HLODGrids[HLODGridName], GridsDepth[HLODGridName] + 1, HLODActorClusterInstancePtrs);
+		GenerateHLODs(HLODGrids[HLODGridName], HLODLayersLevels[HLODGrids[HLODGridName].HLODLayer] + 1, HLODActorClusterInstancePtrs);
 	}
 
 	auto DeleteHLODActor = [&SourceControlHelper, WorldPartition](FWorldPartitionHandle ActorHandle)
diff --git a/Engine/Source/Runtime/Engine/Public/WorldPartition/WorldPartitionRuntimeSpatialHash.h b/Engine/Source/Runtime/Engine/Public/WorldPartition/WorldPartitionRuntimeSpatialHash.h
--- a/Engine/Source/Runtime/Engine/Public/WorldPartition/WorldPartitionRuntimeSpatialHash.h
+++ b/Engine/Source/Runtime/Engine/Public/WorldPartition/WorldPartitionRuntimeSpatialHash.h
@@ -274,9 +274,6 @@ private:
 	UPROPERTY(Transient)
 	mutable FWorldPartitionRuntimeSpatialHashGridPreviewer GridPreviewer;
 
-	UPROPERTY(Transient)
-	TMap<FName, FSpatialHashRuntimeGrid> HLODGrids;
-
 	TMap<FString, UWorldPartitionRuntimeCell*> PackagesToGenerateForCook;
 #endif
 
