Fix foliage instances appearing single sided in ray tracing even when the material is two sided.

diff --git a/Engine/Source/Runtime/Engine/Private/RayTracingInstance.cpp b/Engine/Source/Runtime/Engine/Private/RayTracingInstance.cpp
--- a/Engine/Source/Runtime/Engine/Private/RayTracingInstance.cpp
+++ b/Engine/Source/Runtime/Engine/Private/RayTracingInstance.cpp
@@ -47,7 +47,7 @@ FRayTracingMaskAndFlags BuildRayTracingInstanceMaskAndFlags(TArrayView<const FMe
 			bAllSegmentsOpaque &= BlendMode == BLEND_Opaque;
 			bAnySegmentsCastShadow |= MeshBatch.CastRayTracedShadow && Material.CastsRayTracedShadows();
 			bAllSegmentsCastShadow &= MeshBatch.CastRayTracedShadow && Material.CastsRayTracedShadows();
-			bDoubleSided |= MeshBatch.bDisableBackfaceCulling;
+			bDoubleSided |= MeshBatch.bDisableBackfaceCulling || Material.IsTwoSided();
 		}
 	}
 
