Default r.OpenGL.ForceDXC to 0 until UE-122325 is resolved

diff --git a/Engine/Source/Runtime/Engine/Private/ShaderCompiler/ShaderCompiler.cpp b/Engine/Source/Runtime/Engine/Private/ShaderCompiler/ShaderCompiler.cpp
--- a/Engine/Source/Runtime/Engine/Private/ShaderCompiler/ShaderCompiler.cpp
+++ b/Engine/Source/Runtime/Engine/Private/ShaderCompiler/ShaderCompiler.cpp
@@ -956,10 +956,10 @@ static TAutoConsoleVariable<int32> CVarD3DForceShaderConductorDXCRewrite(
 
 static TAutoConsoleVariable<int32> CVarOpenGLForceDXC(
 	TEXT("r.OpenGL.ForceDXC"),
-	1,
+	0,
 	TEXT("Forces DirectX Shader Compiler (DXC) to be used for all OpenGL shaders instead of hlslcc.\n")
-	TEXT(" 0: Disable\n")
-	TEXT(" 1: Force new compiler for all shaders (default)"),
+	TEXT(" 0: Disable (default)\n")
+	TEXT(" 1: Force new compiler for all shaders"),
 	ECVF_ReadOnly);
 
 int32 GCreateShadersOnLoad = 0;
