Fix load ordering dependency with parent->child hierarchies of anim BPs

In the case where child clases are PostLoaded() after their parents, incorrect node data can be assigned to the class causing OOB access. This fix switches the init of the node data to be lazy/latched, so subsequent subclass loads dont overwrite parent class node data

diff --git a/Engine/Source/Editor/AnimGraph/Private/AnimBlueprintCompiler.cpp b/Engine/Source/Editor/AnimGraph/Private/AnimBlueprintCompiler.cpp
--- a/Engine/Source/Editor/AnimGraph/Private/AnimBlueprintCompiler.cpp
+++ b/Engine/Source/Editor/AnimGraph/Private/AnimBlueprintCompiler.cpp
@@ -763,7 +763,7 @@ void FAnimBlueprintCompilerContext::CopyTermDefaultsToDefaultObject(UObject* Def
 		}
 		
 		// Re-initialize node data tables (they would be overwritten in the loop above)
-		NewAnimBlueprintClass->InitializeAnimNodeData(DefaultObject);
+		NewAnimBlueprintClass->InitializeAnimNodeData(DefaultObject, true);
 	}
 
 	// Give game-specific logic a chance to replace animations
diff --git a/Engine/Source/Runtime/Engine/Classes/Animation/AnimBlueprintGeneratedClass.h b/Engine/Source/Runtime/Engine/Classes/Animation/AnimBlueprintGeneratedClass.h
--- a/Engine/Source/Runtime/Engine/Classes/Animation/AnimBlueprintGeneratedClass.h
+++ b/Engine/Source/Runtime/Engine/Classes/Animation/AnimBlueprintGeneratedClass.h
@@ -474,7 +474,7 @@ private:
 
 	// Called by the compiler to make sure that data tables are initialized. This is needed to patch the sparse class
 	// data for child anim BP overrides 
-	void InitializeAnimNodeData(UObject* DefaultObject);
+	void InitializeAnimNodeData(UObject* DefaultObject, bool bForce);
 	
 public:
 #if WITH_EDITORONLY_DATA
diff --git a/Engine/Source/Runtime/Engine/Private/Animation/AnimBlueprintGeneratedClass.cpp b/Engine/Source/Runtime/Engine/Private/Animation/AnimBlueprintGeneratedClass.cpp
--- a/Engine/Source/Runtime/Engine/Private/Animation/AnimBlueprintGeneratedClass.cpp
+++ b/Engine/Source/Runtime/Engine/Private/Animation/AnimBlueprintGeneratedClass.cpp
@@ -672,7 +672,7 @@ void UAnimBlueprintGeneratedClass::GenerateAnimationBlueprintFunctions()
 	}
 }
 
-void UAnimBlueprintGeneratedClass::InitializeAnimNodeData(UObject* DefaultObject)
+void UAnimBlueprintGeneratedClass::InitializeAnimNodeData(UObject* DefaultObject, bool bForce)
 {
 	// Link functions to their nodes
 	for(int32 AnimNodeIndex = 0; AnimNodeIndex < AnimNodeProperties.Num(); ++AnimNodeIndex)
@@ -682,7 +682,13 @@ void UAnimBlueprintGeneratedClass::InitializeAnimNodeData(UObject* DefaultObject
 		if(StructProperty->Struct->IsChildOf(FAnimNode_Base::StaticStruct()))
 		{
 			FAnimNode_Base* Node = StructProperty->ContainerPtrToValuePtr<FAnimNode_Base>(DefaultObject);
-			Node->SetNodeData(AnimNodeData[AnimNodeIndex]);
+
+			// This function is called on child->parent hierarchies in postload. We dont want to overwrite child data
+			// with parent data so skip if the node data has already been set up by a previous call.	
+			if(bForce || Node->NodeData == nullptr)
+			{
+				Node->SetNodeData(AnimNodeData[AnimNodeIndex]);
+			}
 		}
 	}
 }
@@ -693,6 +699,10 @@ void UAnimBlueprintGeneratedClass::LinkFunctionsToDefaultObjectNodes(UObject* De
 	DynamicResetNodeProperties.Empty();
 	InitializationNodeProperties.Empty();
 
+	// Perform node init as a first pass as root nodes (needed for function names)
+	// may not be initialized in sparse class data before they need to be patched
+	InitializeAnimNodeData(DefaultObject, false);
+	
 	// Link functions to their nodes
 	for(int32 AnimNodeIndex = 0; AnimNodeIndex < AnimNodeProperties.Num(); ++AnimNodeIndex)
 	{
@@ -702,13 +712,6 @@ void UAnimBlueprintGeneratedClass::LinkFunctionsToDefaultObjectNodes(UObject* De
 		{
 			FAnimNode_Base* Node = StructProperty->ContainerPtrToValuePtr<FAnimNode_Base>(DefaultObject);
 
-			// This function is called on child->parent hierarchies in postload. We dont want to overwrite child data
-			// with parent data so skip if the node data has already been set up by a previous call.
-			if(Node->NodeData == nullptr)
-			{
-				Node->SetNodeData(AnimNodeData[AnimNodeIndex]);
-			}
-
 			if(Node->NeedsDynamicReset())
 			{
 				DynamicResetNodeProperties.Add(StructProperty);
