Fix crash when calling SetSkeletalMesh

Required bones need to be recalculated in this circumstance as the mesh has competely changed

diff --git a/Engine/Source/Runtime/Engine/Private/Components/SkeletalMeshComponent.cpp b/Engine/Source/Runtime/Engine/Private/Components/SkeletalMeshComponent.cpp
--- a/Engine/Source/Runtime/Engine/Private/Components/SkeletalMeshComponent.cpp
+++ b/Engine/Source/Runtime/Engine/Private/Components/SkeletalMeshComponent.cpp
@@ -2865,6 +2865,10 @@ void USkeletalMeshComponent::SetSkeletalMesh(USkeletalMesh* InSkelMesh, bool bRe
 		UpdateHasValidBodies();
 		ClearMorphTargets();
 
+		// Make sure that required bones are invalidated as we have just changed our mesh
+		// RecalcRequiredBones will be called by InitAnim below
+		bRequiredBonesUpToDate = false;
+		
 		InitAnim(bReinitPose);
 
 #if WITH_APEX_CLOTHING || WITH_CHAOS_CLOTHING
