Initialize Sematics

Added 'initial update' concept and a seperate tracker via node relevancy subsystem to track it.

diff --git a/Engine/Source/Editor/AnimGraph/Private/AnimBlueprintExtension.cpp b/Engine/Source/Editor/AnimGraph/Private/AnimBlueprintExtension.cpp
--- a/Engine/Source/Editor/AnimGraph/Private/AnimBlueprintExtension.cpp
+++ b/Engine/Source/Editor/AnimGraph/Private/AnimBlueprintExtension.cpp
@@ -78,7 +78,7 @@ void UAnimBlueprintExtension::RequestExtensionsForNode(UAnimGraphNode_Base* InAn
             UAnimBlueprintExtension_PropertyAccess::StaticClass()
         };
 
-		if(InAnimGraphNode->BecomeRelevantFunction.GetMemberGuid().IsValid())
+		if(InAnimGraphNode->InitialUpdateFunction.GetMemberGuid().IsValid() || InAnimGraphNode->BecomeRelevantFunction.GetMemberGuid().IsValid())
 		{
 			ExtensionClasses.Add(UAnimBlueprintExtension_NodeRelevancy::StaticClass());
 		}
diff --git a/Engine/Source/Editor/AnimGraph/Private/AnimGraphNode_Base.cpp b/Engine/Source/Editor/AnimGraph/Private/AnimGraphNode_Base.cpp
--- a/Engine/Source/Editor/AnimGraph/Private/AnimGraphNode_Base.cpp
+++ b/Engine/Source/Editor/AnimGraph/Private/AnimGraphNode_Base.cpp
@@ -85,6 +85,12 @@ void UAnimGraphNode_Base::PostEditChangeProperty(FPropertyChangedEvent& Property
 		FOptionalPinManager::EvaluateOldShownPins(ShowPinForProperties, OldShownPins, this);
 		GetSchema()->ReconstructNode(*this);
 	}
+	else if(PropertyName == GET_MEMBER_NAME_CHECKED(UAnimGraphNode_Base, InitialUpdateFunction))
+	{
+		GetFNode()->InitialUpdateFunction.SetFromFunction(InitialUpdateFunction.ResolveMember<UFunction>(GetAnimBlueprint()->SkeletonGeneratedClass));
+		GetAnimBlueprint()->RequestRefreshExtensions();
+		GetSchema()->ReconstructNode(*this);
+	}
 	else if(PropertyName == GET_MEMBER_NAME_CHECKED(UAnimGraphNode_Base, BecomeRelevantFunction))
 	{
 		GetFNode()->BecomeRelevantFunction.SetFromFunction(BecomeRelevantFunction.ResolveMember<UFunction>(GetAnimBlueprint()->SkeletonGeneratedClass));
@@ -209,6 +215,7 @@ void UAnimGraphNode_Base::ValidateAnimNodeDuringCompilation(USkeleton* ForSkelet
 		}
 	};
 	
+	ValidateFunctionRef(InitialUpdateFunction, LOCTEXT("MissingInitialUpdateFunction", "Could not resolve initial update function for @@"));
 	ValidateFunctionRef(BecomeRelevantFunction, LOCTEXT("MissingBecomeRelevantFunction", "Could not resolve become relevant function for @@"));
 	ValidateFunctionRef(UpdateFunction, LOCTEXT("MissingUpdateFunction", "Could not resolve update function for @@"));
 }
@@ -494,6 +501,7 @@ void UAnimGraphNode_Base::ProcessDuringCompilation(IAnimBlueprintCompilationCont
 	Extension->ProcessNodePins(this, InCompilationContext, OutCompiledData);
 
 	// Resolve functions
+	GetFNode()->InitialUpdateFunction.SetFromFunction(InitialUpdateFunction.ResolveMember<UFunction>(GetAnimBlueprint()->SkeletonGeneratedClass));  
 	GetFNode()->BecomeRelevantFunction.SetFromFunction(BecomeRelevantFunction.ResolveMember<UFunction>(GetAnimBlueprint()->SkeletonGeneratedClass));
 	GetFNode()->UpdateFunction.SetFromFunction(UpdateFunction.ResolveMember<UFunction>(GetAnimBlueprint()->SkeletonGeneratedClass)); 
 
diff --git a/Engine/Source/Editor/AnimGraph/Public/AnimGraphNode_Base.h b/Engine/Source/Editor/AnimGraph/Public/AnimGraphNode_Base.h
--- a/Engine/Source/Editor/AnimGraph/Public/AnimGraphNode_Base.h
+++ b/Engine/Source/Editor/AnimGraph/Public/AnimGraphNode_Base.h
@@ -209,6 +209,10 @@ class ANIMGRAPH_API UAnimGraphNode_Base : public UK2Node
 	UPROPERTY(Transient)
 	EBlueprintUsage BlueprintUsage;
 
+	// Function called before the node is updated for the first time
+	UPROPERTY(EditAnywhere, Category = Functions, meta=(FunctionReference, PrototypeFunction="/Script/AnimGraphRuntime.AnimExecutionContextLibrary.Prototype_ThreadSafeAnimUpdateCall"), DisplayName="On Initial Update")
+	FMemberReference InitialUpdateFunction;
+
 	// Function called when the node becomes relevant
 	UPROPERTY(EditAnywhere, Category = Functions, meta=(FunctionReference, PrototypeFunction="/Script/AnimGraphRuntime.AnimExecutionContextLibrary.Prototype_ThreadSafeAnimUpdateCall"), DisplayName="On Become Relevant")
 	FMemberReference BecomeRelevantFunction;
diff --git a/Engine/Source/Editor/AnimationBlueprintEditor/Private/AnimationNodes/SAnimationGraphNode.cpp b/Engine/Source/Editor/AnimationBlueprintEditor/Private/AnimationNodes/SAnimationGraphNode.cpp
--- a/Engine/Source/Editor/AnimationBlueprintEditor/Private/AnimationNodes/SAnimationGraphNode.cpp
+++ b/Engine/Source/Editor/AnimationBlueprintEditor/Private/AnimationNodes/SAnimationGraphNode.cpp
@@ -368,6 +368,10 @@ void SAnimationGraphNode::CreateBelowPinControls(TSharedPtr<SVerticalBox> MainBo
 			}
 		};
 
+		if(AnimNode->InitialUpdateFunction.GetMemberGuid().IsValid())
+		{
+			AddFunctionBindingWidget("Functions", GET_MEMBER_NAME_CHECKED(UAnimGraphNode_Base, InitialUpdateFunction));
+		}
 		if(AnimNode->BecomeRelevantFunction.GetMemberGuid().IsValid())
 		{
 			AddFunctionBindingWidget("Functions", GET_MEMBER_NAME_CHECKED(UAnimGraphNode_Base, BecomeRelevantFunction));
diff --git a/Engine/Source/Runtime/Engine/Classes/Animation/AnimNodeBase.h b/Engine/Source/Runtime/Engine/Classes/Animation/AnimNodeBase.h
--- a/Engine/Source/Runtime/Engine/Classes/Animation/AnimNodeBase.h
+++ b/Engine/Source/Runtime/Engine/Classes/Animation/AnimNodeBase.h
@@ -966,6 +966,7 @@ protected:
 
 private:
 	// Access functions
+	const FAnimNodeFunctionRef& GetInitialUpdateFunction() const;
 	const FAnimNodeFunctionRef& GetBecomeRelevantFunction() const;
 	const FAnimNodeFunctionRef& GetUpdateFunction() const;
 	
@@ -983,6 +984,10 @@ private:
 	const FAnimNodeData* NodeData = nullptr;
 
 #if WITH_EDITORONLY_DATA
+	// Function called on initial update
+	UPROPERTY(meta=(FoldProperty))
+	FAnimNodeFunctionRef InitialUpdateFunction;
+
 	// Function called on become relevant
 	UPROPERTY(meta=(FoldProperty))
 	FAnimNodeFunctionRef BecomeRelevantFunction;
diff --git a/Engine/Source/Runtime/Engine/Classes/Animation/AnimNodeFunctionRef.h b/Engine/Source/Runtime/Engine/Classes/Animation/AnimNodeFunctionRef.h
--- a/Engine/Source/Runtime/Engine/Classes/Animation/AnimNodeFunctionRef.h
+++ b/Engine/Source/Runtime/Engine/Classes/Animation/AnimNodeFunctionRef.h
@@ -66,6 +66,9 @@ private:
 	friend struct ::FComponentSpacePoseLink;
 	friend struct ::FAnimInstanceProxy;
 	
+	// Call the InitialUpdate function of this node
+	static void InitialUpdate(const FAnimationUpdateContext& InContext, FAnimNode_Base& InNode);
+
 	// Call the BecomeRelevant function of this node
 	static void BecomeRelevant(const FAnimationUpdateContext& InContext, FAnimNode_Base& InNode);
 
diff --git a/Engine/Source/Runtime/Engine/Private/Animation/AnimInstanceProxy.cpp b/Engine/Source/Runtime/Engine/Private/Animation/AnimInstanceProxy.cpp
--- a/Engine/Source/Runtime/Engine/Private/Animation/AnimInstanceProxy.cpp
+++ b/Engine/Source/Runtime/Engine/Private/Animation/AnimInstanceProxy.cpp
@@ -59,6 +59,7 @@ void FAnimInstanceProxy::UpdateAnimationNode_WithRoot(const FAnimationUpdateCont
 			UpdateCounter.Increment();
 		}
 
+		UE::Anim::FNodeFunctionCaller::InitialUpdate(InContext, *InRootNode);
 		UE::Anim::FNodeFunctionCaller::BecomeRelevant(InContext, *InRootNode);
 		UE::Anim::FNodeFunctionCaller::Update(InContext, *InRootNode);
 		InRootNode->Update_AnyThread(InContext);
diff --git a/Engine/Source/Runtime/Engine/Private/Animation/AnimNodeBase.cpp b/Engine/Source/Runtime/Engine/Private/Animation/AnimNodeBase.cpp
--- a/Engine/Source/Runtime/Engine/Private/Animation/AnimNodeBase.cpp
+++ b/Engine/Source/Runtime/Engine/Private/Animation/AnimNodeBase.cpp
@@ -149,6 +149,11 @@ void FAnimNode_Base::ResetDynamics(ETeleportType InTeleportType)
 	PRAGMA_ENABLE_DEPRECATION_WARNINGS
 }
 
+const FAnimNodeFunctionRef& FAnimNode_Base::GetInitialUpdateFunction() const
+{
+	return GET_ANIM_NODE_DATA(FAnimNodeFunctionRef, InitialUpdateFunction);
+}
+
 const FAnimNodeFunctionRef& FAnimNode_Base::GetBecomeRelevantFunction() const
 {
 	return GET_ANIM_NODE_DATA(FAnimNodeFunctionRef, BecomeRelevantFunction);
@@ -314,6 +319,7 @@ void FPoseLinkBase::Update(const FAnimationUpdateContext& InContext)
 	{
 		FAnimationUpdateContext LinkContext(InContext.WithNodeId(LinkID));
 		TRACE_SCOPED_ANIM_NODE(LinkContext);
+		UE::Anim::FNodeFunctionCaller::InitialUpdate(LinkContext, *LinkedNode);
 		UE::Anim::FNodeFunctionCaller::BecomeRelevant(LinkContext, *LinkedNode);
 		UE::Anim::FNodeFunctionCaller::Update(LinkContext, *LinkedNode);
 		LinkedNode->Update_AnyThread(LinkContext);
diff --git a/Engine/Source/Runtime/Engine/Private/Animation/AnimNodeFunctionRef.cpp b/Engine/Source/Runtime/Engine/Private/Animation/AnimNodeFunctionRef.cpp
--- a/Engine/Source/Runtime/Engine/Private/Animation/AnimNodeFunctionRef.cpp
+++ b/Engine/Source/Runtime/Engine/Private/Animation/AnimNodeFunctionRef.cpp
@@ -48,6 +48,23 @@ static void CallFunctionHelper(const FAnimNodeFunctionRef& InFunction, ContextTy
 		InFunction.Call(AnimInstance, &Params);
 	}
 }
+
+void FNodeFunctionCaller::InitialUpdate(const FAnimationUpdateContext& InContext, FAnimNode_Base& InNode)
+{
+	if(InNode.NodeData != nullptr)
+	{
+		const FAnimNodeFunctionRef& Function = InNode.GetInitialUpdateFunction();
+		if(Function.IsValid())
+		{
+			FAnimSubsystemInstance_NodeRelevancy& RelevancySubsystem = CastChecked<UAnimInstance>(InContext.GetAnimInstanceObject())->GetSubsystem<FAnimSubsystemInstance_NodeRelevancy>();
+			const EAnimNodeInitializationStatus Status = RelevancySubsystem.UpdateNodeInitializationStatus(InContext, InNode);
+			if(Status == EAnimNodeInitializationStatus::InitialUpdate)
+			{
+				CallFunctionHelper<FAnimUpdateContext>(Function, InContext, InNode);
+			}
+		}
+	}
+}
 	
 void FNodeFunctionCaller::BecomeRelevant(const FAnimationUpdateContext& InContext, FAnimNode_Base& InNode)
 {
diff --git a/Engine/Source/Runtime/Engine/Private/Animation/AnimSubsystem_NodeRelevancy.cpp b/Engine/Source/Runtime/Engine/Private/Animation/AnimSubsystem_NodeRelevancy.cpp
--- a/Engine/Source/Runtime/Engine/Private/Animation/AnimSubsystem_NodeRelevancy.cpp
+++ b/Engine/Source/Runtime/Engine/Private/Animation/AnimSubsystem_NodeRelevancy.cpp
@@ -76,3 +76,26 @@ FAnimNodeRelevancyStatus FAnimSubsystemInstance_NodeRelevancy::GetNodeRelevancy(
 
 	return FAnimNodeRelevancyStatus();
 }
+
+EAnimNodeInitializationStatus FAnimSubsystemInstance_NodeRelevancy::UpdateNodeInitializationStatus(const FAnimationUpdateContext& InContext, const FAnimNode_Base& InNode)
+{
+	EAnimNodeInitializationStatus& Status = NodeInitTrackers.FindOrAdd(&InNode);
+
+	switch(Status)
+	{
+	case EAnimNodeInitializationStatus::NotUpdated:
+		Status = EAnimNodeInitializationStatus::InitialUpdate;
+		break;
+	default:
+	case EAnimNodeInitializationStatus::InitialUpdate:
+		Status = EAnimNodeInitializationStatus::Updated;
+		break;
+	}
+
+	return Status;
+}
+
+EAnimNodeInitializationStatus FAnimSubsystemInstance_NodeRelevancy::GetNodeInitializationStatus(const FAnimNode_Base& InNode) const
+{
+	return NodeInitTrackers.FindRef(&InNode);
+}
\ No newline at end of file
diff --git a/Engine/Source/Runtime/Engine/Public/Animation/AnimSubsystem_NodeRelevancy.h b/Engine/Source/Runtime/Engine/Public/Animation/AnimSubsystem_NodeRelevancy.h
--- a/Engine/Source/Runtime/Engine/Public/Animation/AnimSubsystem_NodeRelevancy.h
+++ b/Engine/Source/Runtime/Engine/Public/Animation/AnimSubsystem_NodeRelevancy.h
@@ -51,6 +51,16 @@ private:
 	float CurrentWeight = 0.0f;
 };
 
+// Initialization status of a node
+enum class EAnimNodeInitializationStatus
+{
+	NotUpdated = 0,
+	
+	InitialUpdate,
+
+	Updated,
+};
+
 /** Allows anim node's relevancy (when they receive and lose weight in the graph) to be tracked */
 USTRUCT()
 struct ENGINE_API FAnimSubsystemInstance_NodeRelevancy : public FAnimSubsystemInstance
@@ -63,15 +73,15 @@ struct ENGINE_API FAnimSubsystemInstance_NodeRelevancy : public FAnimSubsystemIn
 	// Update the relevancy of the passed-in node using the supplied context
 	FAnimNodeRelevancyStatus UpdateNodeRelevancy(const FAnimationUpdateContext& InContext, const FAnimNode_Base& InNode);
 
-	// Update the relevancy of the passed-in instance using the supplied context
-	FAnimNodeRelevancyStatus UpdateInstanceRelevancy(const FAnimationUpdateContext& InContext, const UAnimInstance& InInstance);
-	
 	// Get the tracked relevancy of the passed-in node. If the node is not tracked the relevancy will be default (zero weighted).
 	FAnimNodeRelevancyStatus GetNodeRelevancy(const FAnimNode_Base& InNode) const;
 
-	// Get the tracked relevancy of the passed-in instance. If the instance is not tracked the relevancy will be default (zero weighted).
-	FAnimNodeRelevancyStatus GetInstanceRelevancy(const UAnimInstance& InInstance) const;
+	// Update the initialization state of the passed-in node using the supplied context
+	EAnimNodeInitializationStatus UpdateNodeInitializationStatus(const FAnimationUpdateContext& InContext, const FAnimNode_Base& InNode);
 
+	// Get the tracked initialization state of the passed-in node. If the node is not tracked the initialization state will be default (NotUpdated)
+	EAnimNodeInitializationStatus GetNodeInitializationStatus(const FAnimNode_Base& InNode) const;
+	
 private:
 	// Tracks the relevancy of a node
 	struct FTracker
@@ -86,4 +96,7 @@ private:
 private:
 	// Map of tracked anim nodes
 	TMap<const FAnimNode_Base*, FTracker> NodeTrackers;
+
+	// Map of tracked anim nodes initialization
+	TMap<const FAnimNode_Base*, EAnimNodeInitializationStatus> NodeInitTrackers;
 };
\ No newline at end of file
