Don't create the raytracing geometry for streamed static meshes if it's disabled on the UStaticMesh

diff --git a/Engine/Source/Runtime/Engine/Private/Streaming/StaticMeshUpdate.cpp b/Engine/Source/Runtime/Engine/Private/Streaming/StaticMeshUpdate.cpp
--- a/Engine/Source/Runtime/Engine/Private/Streaming/StaticMeshUpdate.cpp
+++ b/Engine/Source/Runtime/Engine/Private/Streaming/StaticMeshUpdate.cpp
@@ -218,7 +218,7 @@ void FStaticMeshStreamIn::DoFinishUpdate(const FContext& Context)
 
 #if RHI_RAYTRACING
 		// Must happen after the batched updates have been flushed
-		if (IsRayTracingEnabled())
+		if (IsRayTracingEnabled() && Context.Mesh->bSupportRayTracing)
 		{
 			for (int32 LODIndex = PendingFirstLODIdx; LODIndex < CurrentFirstLODIdx; ++LODIndex)
 			{
