Remove unused stat

diff --git a/Engine/Source/Runtime/Engine/Private/VT/VirtualTextureUploadCache.cpp b/Engine/Source/Runtime/Engine/Private/VT/VirtualTextureUploadCache.cpp
--- a/Engine/Source/Runtime/Engine/Private/VT/VirtualTextureUploadCache.cpp
+++ b/Engine/Source/Runtime/Engine/Private/VT/VirtualTextureUploadCache.cpp
@@ -16,7 +16,6 @@
 
 DECLARE_MEMORY_STAT_POOL(TEXT("Total GPU Upload Memory"), STAT_TotalGPUUploadSize, STATGROUP_VirtualTextureMemory, FPlatformMemory::MCR_GPU);
 DECLARE_MEMORY_STAT(TEXT("Total CPU Upload Memory"), STAT_TotalCPUUploadSize, STATGROUP_VirtualTextureMemory);
-DECLARE_DWORD_ACCUMULATOR_STAT(TEXT("Num Upload Entries"), STAT_NumUploadEntries, STATGROUP_VirtualTextureMemory);
 
 static TAutoConsoleVariable<int32> CVarVTUploadMemoryPageSize(
 	TEXT("r.VT.UploadMemoryPageSize"),
