PreShader: replace template specialization with method

diff --git a/Engine/Source/Runtime/Engine/Private/Shader/Preshader.cpp b/Engine/Source/Runtime/Engine/Private/Shader/Preshader.cpp
--- a/Engine/Source/Runtime/Engine/Private/Shader/Preshader.cpp
+++ b/Engine/Source/Runtime/Engine/Private/Shader/Preshader.cpp
@@ -14,19 +14,6 @@ void UE::Shader::FPreshaderData::WriteData(const void* Value, uint32 Size)
 	Data.Append((uint8*)Value, Size);
 }
 
-template<>
-UE::Shader::FPreshaderData& UE::Shader::FPreshaderData::Write<UE::Shader::FValue>(const FValue& Value)
-{
-	const EValueType Type = Value.GetType();
-	FMemoryImageValue MemoryValue = Value.AsMemoryImage();
-	Data.Add((uint8)Type);
-	if (MemoryValue.Size > 0u)
-	{
-		Data.Append(MemoryValue.Bytes, MemoryValue.Size);
-	}
-	return *this;
-}
-
 void UE::Shader::FPreshaderData::WriteName(const FScriptName& Name)
 {
 	int32 Index = Names.Find(Name);
diff --git a/Engine/Source/Runtime/Engine/Public/Shader/Preshader.h b/Engine/Source/Runtime/Engine/Public/Shader/Preshader.h
--- a/Engine/Source/Runtime/Engine/Public/Shader/Preshader.h
+++ b/Engine/Source/Runtime/Engine/Public/Shader/Preshader.h
@@ -105,5 +105,18 @@ public:
 	LAYOUT_FIELD(TMemoryImageArray<uint8>, Data);
 };
 
+template<>
+inline FPreshaderData& FPreshaderData::Write<FValue>(const FValue& Value)
+{
+	const EValueType Type = Value.GetType();
+	FMemoryImageValue MemoryValue = Value.AsMemoryImage();
+	Data.Add((uint8)Type);
+	if (MemoryValue.Size > 0u)
+	{
+		Data.Append(MemoryValue.Bytes, MemoryValue.Size);
+	}
+	return *this;
+}
+
 } // namespace Shader
 } // namespace UE
\ No newline at end of file
