World Partition: removed wp.Runtime.IgnoreStreamingPerformance and fixed wp.Runtime.GBlockOnSlowStreaming to be a cvar instead of a toggle.

diff --git a/Engine/Plugins/MovieScene/MovieRenderPipeline/Source/MovieRenderPipelineCore/Private/MoviePipelineGameOverrideSetting.cpp b/Engine/Plugins/MovieScene/MovieRenderPipeline/Source/MovieRenderPipelineCore/Private/MoviePipelineGameOverrideSetting.cpp
--- a/Engine/Plugins/MovieScene/MovieRenderPipeline/Source/MovieRenderPipelineCore/Private/MoviePipelineGameOverrideSetting.cpp
+++ b/Engine/Plugins/MovieScene/MovieRenderPipeline/Source/MovieRenderPipelineCore/Private/MoviePipelineGameOverrideSetting.cpp
@@ -107,7 +107,7 @@ void UMoviePipelineGameOverrideSetting::ApplyCVarSettings(const bool bOverrideVa
 	MOVIEPIPELINE_STORE_AND_OVERRIDE_CVAR_INT(PreviousVolumetricRenderTarget, TEXT("r.VolumetricRenderTarget"), 0, bOverrideValues);
 
 	// To make sure that the world partition streaming doesn't end up in critical streaming performances and stops streaming low priority cells.
-	MOVIEPIPELINE_STORE_AND_OVERRIDE_CVAR_INT(PreviousIgnoreStreamingPerformance, TEXT("wp.Runtime.IgnoreStreamingPerformance"), 1, bOverrideValues);
+	MOVIEPIPELINE_STORE_AND_OVERRIDE_CVAR_INT(PreviousIgnoreStreamingPerformance, TEXT("wp.Runtime.BlockOnSlowStreaming"), 0, bOverrideValues);
 }
 
 void UMoviePipelineGameOverrideSetting::BuildNewProcessCommandLineArgsImpl(TArray<FString>& InOutUnrealURLParams, TArray<FString>& InOutCommandLineArgs, TArray<FString>& InOutDeviceProfileCvars, TArray<FString>& InOutExecCmds) const
diff --git a/Engine/Source/Runtime/Engine/Private/WorldPartition/WorldPartitionRuntimeHash.cpp b/Engine/Source/Runtime/Engine/Private/WorldPartition/WorldPartitionRuntimeHash.cpp
--- a/Engine/Source/Runtime/Engine/Private/WorldPartition/WorldPartitionRuntimeHash.cpp
+++ b/Engine/Source/Runtime/Engine/Private/WorldPartition/WorldPartitionRuntimeHash.cpp
@@ -17,14 +17,6 @@
 
 #define LOCTEXT_NAMESPACE "WorldPartition"
 
-static int32 GIgnoreStreamingPerformance = 0;
-static FAutoConsoleVariableRef CVarGIgnoreStreamingPerformance(
-	TEXT("wp.Runtime.IgnoreStreamingPerformance"),
-	GIgnoreStreamingPerformance,
-	TEXT("Should we ignore streaming performance checks? Enabling this will not block on critical streaming situations and may break gameplay. Useful for scenarios where we flush streaming every frame, such as movie rendering."),
-	ECVF_Default
-);
-
 UWorldPartitionRuntimeHash::UWorldPartitionRuntimeHash(const FObjectInitializer& ObjectInitializer)
 	: Super(ObjectInitializer)
 {}
@@ -290,15 +282,9 @@ void UWorldPartitionRuntimeHash::SortStreamingCellsByImportance(const TSet<const
 
 EWorldPartitionStreamingPerformance UWorldPartitionRuntimeHash::GetStreamingPerformance(const TSet<const UWorldPartitionRuntimeCell*>& CellsToActivate) const
 {
-	bool bForceGoodStreamingPerformance = !GetWorld()->bMatchStarted;
-
-#if WITH_EDITOR
-	bForceGoodStreamingPerformance |= GIgnoreStreamingPerformance > 0;
-#endif
-
 	EWorldPartitionStreamingPerformance StreamingPerformance = EWorldPartitionStreamingPerformance::Good;
 
-	if (!CellsToActivate.IsEmpty() && !bForceGoodStreamingPerformance)
+	if (!CellsToActivate.IsEmpty() && GetWorld()->bMatchStarted)
 	{
 		UWorld* World = GetWorld();
 
diff --git a/Engine/Source/Runtime/Engine/Private/WorldPartition/WorldPartitionStreamingPolicy.cpp b/Engine/Source/Runtime/Engine/Private/WorldPartition/WorldPartitionStreamingPolicy.cpp
--- a/Engine/Source/Runtime/Engine/Private/WorldPartition/WorldPartitionStreamingPolicy.cpp
+++ b/Engine/Source/Runtime/Engine/Private/WorldPartition/WorldPartitionStreamingPolicy.cpp
@@ -48,11 +48,11 @@ static FAutoConsoleVariableRef CMaxLoadingStreamingCells(
 	GMaxLoadingStreamingCells,
 	TEXT("Used to limit the number of concurrent loading world partition streaming cells."));
 
-static int32 GBlockOnSlowStreaming = 1;
-static FAutoConsoleCommand CVarBlockOnSlowStreaming(
-	TEXT("wp.Runtime.ToggleBlockOnSlowStreaming"),
-	TEXT("Toggles if streaming needs to block when to slow to catchup."),
-	FConsoleCommandDelegate::CreateLambda([] { GBlockOnSlowStreaming = !GBlockOnSlowStreaming; }));
+int32 GBlockOnSlowStreaming = 1;
+static FAutoConsoleVariableRef CVarBlockOnSlowStreaming(
+	TEXT("wp.Runtime.BlockOnSlowStreaming"),
+	GBlockOnSlowStreaming,
+	TEXT("Set if streaming needs to block when to slow to catchup."));
 
 
 UWorldPartitionStreamingPolicy::UWorldPartitionStreamingPolicy(const FObjectInitializer& ObjectInitializer)
