Display in log whether the global shader map is loaded or not.

diff --git a/Engine/Source/Runtime/Engine/Private/ShaderCompiler/ShaderCompiler.cpp b/Engine/Source/Runtime/Engine/Private/ShaderCompiler/ShaderCompiler.cpp
--- a/Engine/Source/Runtime/Engine/Private/ShaderCompiler/ShaderCompiler.cpp
+++ b/Engine/Source/Runtime/Engine/Private/ShaderCompiler/ShaderCompiler.cpp
@@ -6406,7 +6406,7 @@ void CompileGlobalShaderMap(EShaderPlatform Platform, const ITargetPlatform* Tar
 
 				if (!bFileExist)
 				{
-					UE_LOG(LogShaders, Log, TEXT("%s doesn't exists"), *OverrideGlobalShaderCacheFilename);
+					UE_LOG(LogShaders, Display, TEXT("%s doesn't exists"), *OverrideGlobalShaderCacheFilename);
 				}
 				else
 				{
@@ -6414,7 +6414,7 @@ void CompileGlobalShaderMap(EShaderPlatform Platform, const ITargetPlatform* Tar
 
 					if (bLoadedFromCacheFile)
 					{
-						UE_LOG(LogShaders, Log, TEXT("%s has been loaded successfully"), *OverrideGlobalShaderCacheFilename);
+						UE_LOG(LogShaders, Display, TEXT("%s has been loaded successfully"), *OverrideGlobalShaderCacheFilename);
 					}
 					else
 					{
