Fix unnecessary file reads of empty codec payloads

diff --git a/Engine/Source/Runtime/Engine/Private/VT/UploadingVirtualTexture.cpp b/Engine/Source/Runtime/Engine/Private/VT/UploadingVirtualTexture.cpp
--- a/Engine/Source/Runtime/Engine/Private/VT/UploadingVirtualTexture.cpp
+++ b/Engine/Source/Runtime/Engine/Private/VT/UploadingVirtualTexture.cpp
@@ -301,8 +301,9 @@ FVTCodecAndStatus FUploadingVirtualTexture::GetCodecForChunk(FGraphEventArray& O
 {
 	TRACE_CPUPROFILER_EVENT_SCOPE(FUploadingVirtualTexture::GetCodecForChunk);
 
-	const FVirtualTextureDataChunk& Chunk = Data->Chunks[ChunkIndex];
-	if (Chunk.CodecPayloadSize == 0u)
+	// CodecPayloadSize includes the header that is always present even if there is no payload.
+ 	const FVirtualTextureDataChunk& Chunk = Data->Chunks[ChunkIndex];
+ 	if (Chunk.CodecPayloadSize <= sizeof(FVirtualTextureChunkHeader))
 	{
 		// Chunk has no codec
 		return EVTRequestPageStatus::Available;
