Fix WIP new material translator to work with updated parameter code
Some minor cleanup

diff --git a/Engine/Source/Runtime/Engine/Private/HLSLTree/HLSLTree.cpp b/Engine/Source/Runtime/Engine/Private/HLSLTree/HLSLTree.cpp
--- a/Engine/Source/Runtime/Engine/Private/HLSLTree/HLSLTree.cpp
+++ b/Engine/Source/Runtime/Engine/Private/HLSLTree/HLSLTree.cpp
@@ -682,11 +682,6 @@ UE::HLSLTree::ENodeVisitResult UE::HLSLTree::FExpression::Visit(FNodeVisitor& Vi
 	return Visitor.OnExpression(*this);
 }
 
-UE::HLSLTree::ENodeVisitResult UE::HLSLTree::FParameterDeclaration::Visit(FNodeVisitor& Visitor)
-{
-	return Visitor.OnParameterDeclaration(*this);
-}
-
 UE::HLSLTree::ENodeVisitResult UE::HLSLTree::FTextureParameterDeclaration::Visit(FNodeVisitor& Visitor)
 {
 	return Visitor.OnTextureParameterDeclaration(*this);
@@ -927,12 +922,6 @@ UE::HLSLTree::FScope* UE::HLSLTree::FTree::NewScope(FScope& Scope)
 	return NewScope;
 }
 
-UE::HLSLTree::FParameterDeclaration* UE::HLSLTree::FTree::NewParameterDeclaration(const FName& Name, const Shader::FValue& DefaultValue)
-{
-	FParameterDeclaration* Declaration = NewNode<FParameterDeclaration>(Name, DefaultValue);
-	return Declaration;
-}
-
 UE::HLSLTree::FTextureParameterDeclaration* UE::HLSLTree::FTree::NewTextureParameterDeclaration(const FName& Name, const FTextureDescription& DefaultValue)
 {
 	FTextureParameterDeclaration* Declaration = NewNode<FTextureParameterDeclaration>(Name, DefaultValue);
diff --git a/Engine/Source/Runtime/Engine/Private/HLSLTree/HLSLTreeCommon.cpp b/Engine/Source/Runtime/Engine/Private/HLSLTree/HLSLTreeCommon.cpp
--- a/Engine/Source/Runtime/Engine/Private/HLSLTree/HLSLTreeCommon.cpp
+++ b/Engine/Source/Runtime/Engine/Private/HLSLTree/HLSLTreeCommon.cpp
@@ -82,14 +82,13 @@ bool UE::HLSLTree::FExpressionConstant::EmitCode(FEmitContext& Context, FExpress
 	return true;
 }
 
-bool UE::HLSLTree::FExpressionParameter::EmitCode(FEmitContext& Context, FExpressionEmitResult& OutResult) const
+bool UE::HLSLTree::FExpressionMaterialParameter::EmitCode(FEmitContext& Context, FExpressionEmitResult& OutResult) const
 {
-	OutResult.Type = Declaration->DefaultValue.GetType();
-
-	if (OutResult.Type == Shader::EValueType::Bool1)
+	OutResult.Type = GetShaderValueType(ParameterType);
+	if (ParameterType == EMaterialParameterType::StaticSwitch)
 	{
-		const FMaterialParameterInfo ParameterInfo(Declaration->Name);
-		bool bValue = Declaration->DefaultValue.Component[0].AsBool();
+		const FMaterialParameterInfo ParameterInfo(ParameterName);
+		bool bValue = DefaultValue.Component[0].AsBool();
 		for (const FStaticSwitchParameter& Parameter : Context.StaticParameters->StaticSwitchParameters)
 		{
 			if (Parameter.ParameterInfo == ParameterInfo)
@@ -101,21 +100,23 @@ bool UE::HLSLTree::FExpressionParameter::EmitCode(FEmitContext& Context, FExpres
 		OutResult.EvaluationType = EExpressionEvaluationType::Constant;
 		OutResult.Preshader.WriteOpcode(Shader::EPreshaderOpcode::Constant).Write(Shader::FValue(bValue));
 	}
-	else if (OutResult.Type == Shader::EValueType::Float1)
-	{
-		/*const int32 ParameterIndex = Context.MaterialCompilationOutput->UniformExpressionSet.FindOrAddScalarParameter(Declaration->Name, Declaration->DefaultValue.Component[0].Float);
-		check(ParameterIndex >= 0 && ParameterIndex <= 0xffff);
-		OutResult.EvaluationType = EExpressionEvaluationType::Preshader;
-		OutResult.Preshader.WriteOpcode(Shader::EPreshaderOpcode::ScalarParameter).Write((uint16)ParameterIndex);*/
-		check(false);
-	}
 	else
 	{
-		/*const int32 ParameterIndex = Context.MaterialCompilationOutput->UniformExpressionSet.FindOrAddVectorParameter(Declaration->Name, Declaration->DefaultValue.AsLinearColor());
+		const uint32* PrevDefaultOffset = Context.DefaultUniformValues.Find(DefaultValue);
+		uint32 DefaultOffset;
+		if (PrevDefaultOffset)
+		{
+			DefaultOffset = *PrevDefaultOffset;
+		}
+		else
+		{
+			DefaultOffset = Context.MaterialCompilationOutput->UniformExpressionSet.AddDefaultParameterValue(DefaultValue);
+			Context.DefaultUniformValues.Add(DefaultValue, DefaultOffset);
+		}
+		const int32 ParameterIndex = Context.MaterialCompilationOutput->UniformExpressionSet.FindOrAddNumericParameter(ParameterType, ParameterName, DefaultOffset);
 		check(ParameterIndex >= 0 && ParameterIndex <= 0xffff);
 		OutResult.EvaluationType = EExpressionEvaluationType::Preshader;
-		OutResult.Preshader.WriteOpcode(Shader::EPreshaderOpcode::VectorParameter).Write((uint16)ParameterIndex);*/
-		check(false);
+		OutResult.Preshader.WriteOpcode(Shader::EPreshaderOpcode::Parameter).Write((uint16)ParameterIndex); 
 	}
 	return true;
 }
diff --git a/Engine/Source/Runtime/Engine/Private/Materials/HLSLMaterialTranslator.cpp b/Engine/Source/Runtime/Engine/Private/Materials/HLSLMaterialTranslator.cpp
--- a/Engine/Source/Runtime/Engine/Private/Materials/HLSLMaterialTranslator.cpp
+++ b/Engine/Source/Runtime/Engine/Private/Materials/HLSLMaterialTranslator.cpp
@@ -4012,7 +4012,7 @@ int32 FHLSLMaterialTranslator::AccessCollectionParameter(UMaterialParameterColle
 		ComponentIndex == -1 ? true : ComponentIndex % 4 == 3);
 }
 
-int32 FHLSLMaterialTranslator::GenericParameter(EMaterialParameterType ParameterType, FName ParameterName, const UE::Shader::FValue& DefaultValue)
+int32 FHLSLMaterialTranslator::NumericParameter(EMaterialParameterType ParameterType, FName ParameterName, const UE::Shader::FValue& DefaultValue)
 {
 	const UE::Shader::EValueType ValueType = GetShaderValueType(ParameterType);
 	check(DefaultValue.GetType() == ValueType);
@@ -4025,35 +4025,15 @@ int32 FHLSLMaterialTranslator::GenericParameter(EMaterialParameterType Parameter
 	}
 	else
 	{
-		DefaultOffset = MaterialCompilationOutput.UniformExpressionSet.DefaultValues.Num();
+		DefaultOffset = MaterialCompilationOutput.UniformExpressionSet.AddDefaultParameterValue(DefaultValue);
 		DefaultUniformValues.Add(DefaultValue, DefaultOffset);
-		UE::Shader::FMemoryImageValue DefaultValueMemory = DefaultValue.AsMemoryImage();
-		MaterialCompilationOutput.UniformExpressionSet.DefaultValues.Append(DefaultValueMemory.Bytes, DefaultValueMemory.Size);
 	}
 
 	FMaterialParameterInfo ParameterInfo = GetParameterAssociationInfo();
 	ParameterInfo.Name = ParameterName;
 
-	int32 ParameterIndex = INDEX_NONE;
-	for (int32 i = 0; i < MaterialCompilationOutput.UniformExpressionSet.UniformNumericParameters.Num(); ++i)
-	{
-		const FMaterialNumericParameterInfo& Parameter = MaterialCompilationOutput.UniformExpressionSet.UniformNumericParameters[i];
-		if (Parameter.ParameterType == ParameterType && Parameter.ParameterInfo == ParameterInfo)
-		{
-			ParameterIndex = i;
-			break;
-		}
-	}
-	if (ParameterIndex == INDEX_NONE)
-	{
-		ParameterIndex = MaterialCompilationOutput.UniformExpressionSet.UniformNumericParameters.Num();
-		FMaterialNumericParameterInfo& Parameter = MaterialCompilationOutput.UniformExpressionSet.UniformNumericParameters.AddDefaulted_GetRef();
-		Parameter.ParameterInfo = ParameterInfo;
-		Parameter.ParameterType = ParameterType;
-		Parameter.DefaultValueOffset = DefaultOffset;
-	}
-
-	return AddUniformExpression(new FMaterialUniformExpressionGenericParameter(ParameterInfo, ParameterIndex), GetMaterialValueType(ParameterType), TEXT(""));
+	const int32 ParameterIndex = MaterialCompilationOutput.UniformExpressionSet.FindOrAddNumericParameter(ParameterType, ParameterInfo, DefaultOffset);
+	return AddUniformExpression(new FMaterialUniformExpressionNumericParameter(ParameterInfo, ParameterIndex), GetMaterialValueType(ParameterType), TEXT(""));
 }
 
 int32 FHLSLMaterialTranslator::Constant(float X)
diff --git a/Engine/Source/Runtime/Engine/Private/Materials/HLSLMaterialTranslator.h b/Engine/Source/Runtime/Engine/Private/Materials/HLSLMaterialTranslator.h
--- a/Engine/Source/Runtime/Engine/Private/Materials/HLSLMaterialTranslator.h
+++ b/Engine/Source/Runtime/Engine/Private/Materials/HLSLMaterialTranslator.h
@@ -676,7 +676,7 @@ protected:
 
 	virtual int32 AccessCollectionParameter(UMaterialParameterCollection* ParameterCollection, int32 ParameterIndex, int32 ComponentIndex) override;
 
-	virtual int32 GenericParameter(EMaterialParameterType ParameterType, FName ParameterName, const UE::Shader::FValue& DefaultValue) override;
+	virtual int32 NumericParameter(EMaterialParameterType ParameterType, FName ParameterName, const UE::Shader::FValue& DefaultValue) override;
 	virtual int32 Constant(float X) override;
 	virtual int32 Constant2(float X, float Y) override;
 	virtual int32 Constant3(float X, float Y, float Z) override;
diff --git a/Engine/Source/Runtime/Engine/Private/Materials/MaterialExpressionHLSL.cpp b/Engine/Source/Runtime/Engine/Private/Materials/MaterialExpressionHLSL.cpp
--- a/Engine/Source/Runtime/Engine/Private/Materials/MaterialExpressionHLSL.cpp
+++ b/Engine/Source/Runtime/Engine/Private/Materials/MaterialExpressionHLSL.cpp
@@ -117,22 +117,19 @@ EMaterialGenerateHLSLStatus UMaterialExpressionGetLocal::GenerateHLSLExpression(
 
 EMaterialGenerateHLSLStatus UMaterialExpressionVectorParameter::GenerateHLSLExpression(FMaterialHLSLGenerator& Generator, UE::HLSLTree::FScope& Scope, int32 OutputIndex, UE::HLSLTree::FExpression*& OutExpression)
 {
-	UE::HLSLTree::FParameterDeclaration* Declaration = Generator.AcquireParameterDeclaration(ParameterName, DefaultValue);
-	OutExpression = Generator.GetTree().NewExpression<UE::HLSLTree::FExpressionParameter>(Scope, Declaration);
+	OutExpression = Generator.GetTree().NewExpression<UE::HLSLTree::FExpressionMaterialParameter>(Scope, EMaterialParameterType::Vector, ParameterName, DefaultValue);
 	return EMaterialGenerateHLSLStatus::Success;
 }
 
 EMaterialGenerateHLSLStatus UMaterialExpressionScalarParameter::GenerateHLSLExpression(FMaterialHLSLGenerator& Generator, UE::HLSLTree::FScope& Scope, int32 OutputIndex, UE::HLSLTree::FExpression*& OutExpression)
 {
-	UE::HLSLTree::FParameterDeclaration* Declaration = Generator.AcquireParameterDeclaration(ParameterName, DefaultValue);
-	OutExpression = Generator.GetTree().NewExpression<UE::HLSLTree::FExpressionParameter>(Scope, Declaration);
+	OutExpression = Generator.GetTree().NewExpression<UE::HLSLTree::FExpressionMaterialParameter>(Scope, EMaterialParameterType::Scalar, ParameterName, DefaultValue);
 	return EMaterialGenerateHLSLStatus::Success;
 }
 
 EMaterialGenerateHLSLStatus UMaterialExpressionStaticBoolParameter::GenerateHLSLExpression(FMaterialHLSLGenerator& Generator, UE::HLSLTree::FScope& Scope, int32 OutputIndex, UE::HLSLTree::FExpression*& OutExpression)
 {
-	UE::HLSLTree::FParameterDeclaration* Declaration = Generator.AcquireParameterDeclaration(ParameterName, (bool)DefaultValue);
-	OutExpression = Generator.GetTree().NewExpression<UE::HLSLTree::FExpressionParameter>(Scope, Declaration);
+	OutExpression = Generator.GetTree().NewExpression<UE::HLSLTree::FExpressionMaterialParameter>(Scope, EMaterialParameterType::StaticSwitch, ParameterName, (bool)DefaultValue);
 	return EMaterialGenerateHLSLStatus::Success;
 }
 
diff --git a/Engine/Source/Runtime/Engine/Private/Materials/MaterialHLSLGenerator.cpp b/Engine/Source/Runtime/Engine/Private/Materials/MaterialHLSLGenerator.cpp
--- a/Engine/Source/Runtime/Engine/Private/Materials/MaterialHLSLGenerator.cpp
+++ b/Engine/Source/Runtime/Engine/Private/Materials/MaterialHLSLGenerator.cpp
@@ -291,16 +291,6 @@ UE::HLSLTree::FExpression* FMaterialHLSLGenerator::NewFunctionInput(UE::HLSLTree
 	return Expression;
 }
 
-UE::HLSLTree::FParameterDeclaration* FMaterialHLSLGenerator::AcquireParameterDeclaration(const FName& Name, const UE::Shader::FValue& DefaultValue)
-{
-	UE::HLSLTree::FParameterDeclaration*& Declaration = ParameterDeclarationMap.FindOrAdd(Name);
-	if (!Declaration)
-	{
-		Declaration = HLSLTree->NewParameterDeclaration(Name, DefaultValue);
-	}
-	return Declaration;
-}
-
 UE::HLSLTree::FTextureParameterDeclaration* FMaterialHLSLGenerator::AcquireTextureDeclaration(const UE::HLSLTree::FTextureDescription& Value)
 {
 	FString SamplerTypeError;
diff --git a/Engine/Source/Runtime/Engine/Private/Materials/MaterialShared.cpp b/Engine/Source/Runtime/Engine/Private/Materials/MaterialShared.cpp
--- a/Engine/Source/Runtime/Engine/Private/Materials/MaterialShared.cpp
+++ b/Engine/Source/Runtime/Engine/Private/Materials/MaterialShared.cpp
@@ -231,12 +231,12 @@ int32 FMaterialCompiler::Errorf(const TCHAR* Format,...)
 
 int32 FMaterialCompiler::ScalarParameter(FName ParameterName, float DefaultValue)
 {
-	return GenericParameter(EMaterialParameterType::Scalar, ParameterName, DefaultValue);
+	return NumericParameter(EMaterialParameterType::Scalar, ParameterName, DefaultValue);
 }
 
 int32 FMaterialCompiler::VectorParameter(FName ParameterName, const FLinearColor& DefaultValue)
 {
-	return GenericParameter(EMaterialParameterType::Vector, ParameterName, DefaultValue);
+	return NumericParameter(EMaterialParameterType::Vector, ParameterName, DefaultValue);
 }
 
 IMPLEMENT_STRUCT(ExpressionInput);
diff --git a/Engine/Source/Runtime/Engine/Private/Materials/MaterialUniformExpressions.cpp b/Engine/Source/Runtime/Engine/Private/Materials/MaterialUniformExpressions.cpp
--- a/Engine/Source/Runtime/Engine/Private/Materials/MaterialUniformExpressions.cpp
+++ b/Engine/Source/Runtime/Engine/Private/Materials/MaterialUniformExpressions.cpp
@@ -574,6 +574,33 @@ int32 FUniformExpressionSet::FindOrAddTextureParameter(EMaterialTextureParameter
 	return UniformTextureParameters[(int32)Type].Add(Info);
 }
 
+int32 FUniformExpressionSet::FindOrAddNumericParameter(EMaterialParameterType Type, const FMaterialParameterInfo& ParameterInfo, uint32 DefaultValueOffset)
+{
+	for (int32 i = 0; i < UniformNumericParameters.Num(); ++i)
+	{
+		const FMaterialNumericParameterInfo& Parameter = UniformNumericParameters[i];
+		if (Parameter.ParameterType == Type && Parameter.DefaultValueOffset == DefaultValueOffset && Parameter.ParameterInfo == ParameterInfo)
+		{
+			return i;
+		}
+	}
+
+	const int32 Index = UniformNumericParameters.Num();
+	FMaterialNumericParameterInfo& Parameter = UniformNumericParameters.AddDefaulted_GetRef();
+	Parameter.ParameterType = Type;
+	Parameter.ParameterInfo = ParameterInfo;
+	Parameter.DefaultValueOffset = DefaultValueOffset;
+	return Index;
+}
+
+uint32 FUniformExpressionSet::AddDefaultParameterValue(const UE::Shader::FValue& Value)
+{
+	const uint32 Offset = DefaultValues.Num();
+	UE::Shader::FMemoryImageValue DefaultValueMemory = Value.AsMemoryImage();
+	DefaultValues.Append(DefaultValueMemory.Bytes, DefaultValueMemory.Size);
+	return Offset;
+}
+
 void FUniformExpressionSet::GetGameThreadTextureValue(EMaterialTextureParameterType Type, int32 Index, const UMaterialInterface* MaterialInterface, const FMaterial& Material, UTexture*& OutValue, bool bAllowOverride) const
 {
 	check(IsInGameThread());
@@ -1434,7 +1461,7 @@ class FMaterialUniformExpressionRuntimeVirtualTextureParameter_DEPRECATED : publ
 
 IMPLEMENT_MATERIALUNIFORMEXPRESSION_TYPE(FMaterialUniformExpressionTexture);
 IMPLEMENT_MATERIALUNIFORMEXPRESSION_TYPE(FMaterialUniformExpressionConstant);
-IMPLEMENT_MATERIALUNIFORMEXPRESSION_TYPE(FMaterialUniformExpressionGenericParameter);
+IMPLEMENT_MATERIALUNIFORMEXPRESSION_TYPE(FMaterialUniformExpressionNumericParameter);
 IMPLEMENT_MATERIALUNIFORMEXPRESSION_TYPE(FMaterialUniformExpressionTextureParameter);
 IMPLEMENT_MATERIALUNIFORMEXPRESSION_TYPE(FMaterialUniformExpressionExternalTextureBase);
 IMPLEMENT_MATERIALUNIFORMEXPRESSION_TYPE(FMaterialUniformExpressionExternalTexture);
diff --git a/Engine/Source/Runtime/Engine/Private/Materials/MaterialUniformExpressions.h b/Engine/Source/Runtime/Engine/Private/Materials/MaterialUniformExpressions.h
--- a/Engine/Source/Runtime/Engine/Private/Materials/MaterialUniformExpressions.h
+++ b/Engine/Source/Runtime/Engine/Private/Materials/MaterialUniformExpressions.h
@@ -218,13 +218,13 @@ private:
 
 /**
  */
-class FMaterialUniformExpressionGenericParameter : public FMaterialUniformExpression
+class FMaterialUniformExpressionNumericParameter : public FMaterialUniformExpression
 {
-	DECLARE_MATERIALUNIFORMEXPRESSION_TYPE(FMaterialUniformExpressionGenericParameter);
+	DECLARE_MATERIALUNIFORMEXPRESSION_TYPE(FMaterialUniformExpressionNumericParameter);
 public:
 
-	FMaterialUniformExpressionGenericParameter() {}
-	FMaterialUniformExpressionGenericParameter(const FMaterialParameterInfo& InParameterInfo, int32 InParameterIndex)
+	FMaterialUniformExpressionNumericParameter() {}
+	FMaterialUniformExpressionNumericParameter(const FMaterialParameterInfo& InParameterInfo, int32 InParameterIndex)
 		: ParameterInfo(InParameterInfo)
 		, ParameterIndex(InParameterIndex)
 	{
@@ -259,7 +259,7 @@ public:
 		{
 			return false;
 		}
-		FMaterialUniformExpressionGenericParameter* OtherParameter = (FMaterialUniformExpressionGenericParameter*)OtherExpression;
+		FMaterialUniformExpressionNumericParameter* OtherParameter = (FMaterialUniformExpressionNumericParameter*)OtherExpression;
 		return ParameterInfo == OtherParameter->ParameterInfo && ParameterIndex == OtherParameter->ParameterIndex;
 	}
 private:
diff --git a/Engine/Source/Runtime/Engine/Public/HLSLTree/HLSLTree.h b/Engine/Source/Runtime/Engine/Public/HLSLTree/HLSLTree.h
--- a/Engine/Source/Runtime/Engine/Public/HLSLTree/HLSLTree.h
+++ b/Engine/Source/Runtime/Engine/Public/HLSLTree/HLSLTree.h
@@ -270,6 +270,7 @@ public:
 	const FMaterial* Material = nullptr; // TODO - remove preshader material dependency
 	const FStaticParameterSet* StaticParameters = nullptr;
 	FMaterialCompilationOutput* MaterialCompilationOutput = nullptr;
+	TMap<Shader::FValue, uint32> DefaultUniformValues;
 	int32 TotalCodeLength = 0;
 	int32 NumExpressionLocals = 0;
 	int32 NumTexCoords = 0;
@@ -311,7 +312,6 @@ public:
 	virtual ENodeVisitResult OnScope(FScope& Scope) { return ENodeVisitResult::VisitDependentNodes; }
 	virtual ENodeVisitResult OnStatement(FStatement& Statement) { return ENodeVisitResult::VisitDependentNodes; }
 	virtual ENodeVisitResult OnExpression(FExpression& Expression) { return ENodeVisitResult::VisitDependentNodes; }
-	virtual ENodeVisitResult OnParameterDeclaration(FParameterDeclaration& Declaration) { return ENodeVisitResult::VisitDependentNodes; }
 	virtual ENodeVisitResult OnTextureParameterDeclaration(FTextureParameterDeclaration& Declaration) { return ENodeVisitResult::VisitDependentNodes; }
 	virtual ENodeVisitResult OnFunctionCall(FFunctionCall& FunctionCall) { return ENodeVisitResult::VisitDependentNodes; }
 };
@@ -379,20 +379,6 @@ public:
 	int32 NumValues;
 };
 
-/**
- * Represents an HLSL parameter.  This is a uniform constant passed to the generated shader
- */
-class FParameterDeclaration final : public FNode
-{
-public:
-	FParameterDeclaration(const FName& InName, const Shader::FValue& InDefaultValue) : Name(InName), DefaultValue(InDefaultValue) {}
-
-	virtual ENodeVisitResult Visit(FNodeVisitor& Visitor) override;
-
-	FName Name;
-	Shader::FValue DefaultValue;
-};
-
 /**
  * Represents an HLSL texture parameter.
  */
@@ -497,7 +483,6 @@ public:
 	}
 
 	FScope* NewScope(FScope& Scope);
-	FParameterDeclaration* NewParameterDeclaration(const FName& Name, const Shader::FValue& DefaultValue);
 	FTextureParameterDeclaration* NewTextureParameterDeclaration(const FName& Name, const FTextureDescription& DefaultValue);
 
 	FFunctionCall* NewFunctionCall(FScope& Scope,
diff --git a/Engine/Source/Runtime/Engine/Public/HLSLTree/HLSLTreeCommon.h b/Engine/Source/Runtime/Engine/Public/HLSLTree/HLSLTreeCommon.h
--- a/Engine/Source/Runtime/Engine/Public/HLSLTree/HLSLTreeCommon.h
+++ b/Engine/Source/Runtime/Engine/Public/HLSLTree/HLSLTreeCommon.h
@@ -3,6 +3,8 @@
 
 #include "HLSLTree/HLSLTree.h"
 
+enum class EMaterialParameterType : uint8;
+
 namespace UE
 {
 namespace HLSLTree
@@ -46,24 +48,16 @@ public:
 	virtual bool EmitCode(FEmitContext& Context, FExpressionEmitResult& OutResult) const override;
 };
 
-class FExpressionParameter : public FExpression
+class FExpressionMaterialParameter : public FExpression
 {
 public:
-	explicit FExpressionParameter(FParameterDeclaration* InDeclaration)
-		: Declaration(InDeclaration)
+	explicit FExpressionMaterialParameter(EMaterialParameterType InType, const FName& InName, const Shader::FValue& InDefaultValue)
+		: ParameterName(InName), DefaultValue(InDefaultValue), ParameterType(InType)
 	{}
 
-	FParameterDeclaration* Declaration;
-
-	virtual ENodeVisitResult Visit(FNodeVisitor& Visitor) override
-	{
-		const ENodeVisitResult Result = FExpression::Visit(Visitor);
-		if (ShouldVisitDependentNodes(Result))
-		{
-			Visitor.VisitNode(Declaration);
-		}
-		return Result;
-	}
+	FName ParameterName;
+	Shader::FValue DefaultValue;
+	EMaterialParameterType ParameterType;
 
 	virtual bool EmitCode(FEmitContext& Context, FExpressionEmitResult& OutResult) const override;
 };
diff --git a/Engine/Source/Runtime/Engine/Public/HLSLTree/HLSLTreeTypes.h b/Engine/Source/Runtime/Engine/Public/HLSLTree/HLSLTreeTypes.h
--- a/Engine/Source/Runtime/Engine/Public/HLSLTree/HLSLTreeTypes.h
+++ b/Engine/Source/Runtime/Engine/Public/HLSLTree/HLSLTreeTypes.h
@@ -13,7 +13,6 @@ namespace HLSLTree
 
 class FNode;
 class FExpression;
-class FParameterDeclaration;
 class FTextureParameterDeclaration;
 class FStatement;
 class FScope;
diff --git a/Engine/Source/Runtime/Engine/Public/MaterialCompiler.h b/Engine/Source/Runtime/Engine/Public/MaterialCompiler.h
--- a/Engine/Source/Runtime/Engine/Public/MaterialCompiler.h
+++ b/Engine/Source/Runtime/Engine/Public/MaterialCompiler.h
@@ -143,7 +143,7 @@ public:
 	virtual int32 GetCurrentFunctionStackDepth() = 0;
 
 	virtual int32 AccessCollectionParameter(UMaterialParameterCollection* ParameterCollection, int32 ParameterIndex, int32 ComponentIndex) = 0;	
-	virtual int32 GenericParameter(EMaterialParameterType ParameterType, FName ParameterName, const UE::Shader::FValue& DefaultValue) = 0;
+	virtual int32 NumericParameter(EMaterialParameterType ParameterType, FName ParameterName, const UE::Shader::FValue& DefaultValue) = 0;
 
 	ENGINE_API int32 ScalarParameter(FName ParameterName, float DefaultValue);
 	ENGINE_API int32 VectorParameter(FName ParameterName, const FLinearColor& DefaultValue);
@@ -526,7 +526,7 @@ public:
 	virtual int32 ForceCast(int32 Code, EMaterialValueType DestType, uint32 ForceCastFlags = 0) override { return Compiler->ForceCast(Code, DestType, ForceCastFlags); }
 
 	virtual int32 AccessCollectionParameter(UMaterialParameterCollection* ParameterCollection, int32 ParameterIndex, int32 ComponentIndex) override { return Compiler->AccessCollectionParameter(ParameterCollection, ParameterIndex, ComponentIndex); }
-	virtual int32 GenericParameter(EMaterialParameterType ParameterType, FName ParameterName, const UE::Shader::FValue& DefaultValue) override { return Compiler->GenericParameter(ParameterType, ParameterName, DefaultValue); }
+	virtual int32 NumericParameter(EMaterialParameterType ParameterType, FName ParameterName, const UE::Shader::FValue& DefaultValue) override { return Compiler->NumericParameter(ParameterType, ParameterName, DefaultValue); }
 
 	virtual int32 Constant(float X) override { return Compiler->Constant(X); }
 	virtual int32 Constant2(float X, float Y) override { return Compiler->Constant2(X, Y); }
diff --git a/Engine/Source/Runtime/Engine/Public/MaterialHLSLGenerator.h b/Engine/Source/Runtime/Engine/Public/MaterialHLSLGenerator.h
--- a/Engine/Source/Runtime/Engine/Public/MaterialHLSLGenerator.h
+++ b/Engine/Source/Runtime/Engine/Public/MaterialHLSLGenerator.h
@@ -92,8 +92,6 @@ public:
 
 	UE::HLSLTree::FExpression* NewFunctionInput(UE::HLSLTree::FScope& Scope, int32 InputIndex, UMaterialExpressionFunctionInput* MaterialFunctionInput);
 
-	UE::HLSLTree::FParameterDeclaration* AcquireParameterDeclaration(const FName& Name, const UE::Shader::FValue& DefaultValue);
-
 	/** Returns a declaration to access the given texture, with no parameter */
 	UE::HLSLTree::FTextureParameterDeclaration* AcquireTextureDeclaration(const UE::HLSLTree::FTextureDescription& Value);
 
@@ -233,7 +231,6 @@ private:
 	TArray<UE::HLSLTree::FExpressionLocalPHI*> PHIExpressions;
 	TArray<FString> CompileErrors;
 	TArray<UMaterialExpression*> ErrorExpressions;
-	TMap<FName, UE::HLSLTree::FParameterDeclaration*> ParameterDeclarationMap;
 	TMap<UE::HLSLTree::FTextureDescription, UE::HLSLTree::FTextureParameterDeclaration*> TextureDeclarationMap;
 	TMap<FName, UE::HLSLTree::FTextureParameterDeclaration*> TextureParameterDeclarationMap;
 	TMap<FFunctionCallKey, UE::HLSLTree::FFunctionCall*> FunctionCallMap;
diff --git a/Engine/Source/Runtime/Engine/Public/MaterialShared.h b/Engine/Source/Runtime/Engine/Public/MaterialShared.h
--- a/Engine/Source/Runtime/Engine/Public/MaterialShared.h
+++ b/Engine/Source/Runtime/Engine/Public/MaterialShared.h
@@ -486,6 +486,8 @@ public:
 	ENGINE_API void GetTextureValue(int32 Index, const FMaterialRenderContext& Context, const FMaterial& Material, const URuntimeVirtualTexture*& OutValue) const;
 
 	int32 FindOrAddTextureParameter(EMaterialTextureParameterType Type, const FMaterialTextureParameterInfo& Info);
+	int32 FindOrAddNumericParameter(EMaterialParameterType Type, const FMaterialParameterInfo& ParameterInfo, uint32 DefaultValueOffset);
+	uint32 AddDefaultParameterValue(const UE::Shader::FValue& Value);
 
 protected:
 	union FVTPackedStackAndLayerIndex
