Added console command `wp.Runtime.ToggleDrawStreamingPerfs` to only display streaming perfs debug info.
Also display if block on slow grid performance is enable or not in debug display.

diff --git a/Engine/Source/Runtime/Engine/Private/WorldPartition/WorldPartitionSubsystem.cpp b/Engine/Source/Runtime/Engine/Private/WorldPartition/WorldPartitionSubsystem.cpp
--- a/Engine/Source/Runtime/Engine/Private/WorldPartition/WorldPartitionSubsystem.cpp
+++ b/Engine/Source/Runtime/Engine/Private/WorldPartition/WorldPartitionSubsystem.cpp
@@ -12,6 +12,7 @@
 #include "ConsoleSettings.h"
 #include "Debug/DebugDrawService.h"
 
+extern int32 GBlockOnSlowStreaming;
 static const FName NAME_WorldPartitionRuntimeHash("WorldPartitionRuntimeHash");
 
 static int32 GDrawRuntimeHash3D = 0;
@@ -32,6 +33,12 @@ static FAutoConsoleCommand CVarDrawStreamingSources(
 	TEXT("Toggles debug display of world partition streaming sources."),
 	FConsoleCommandDelegate::CreateLambda([] { GDrawStreamingSources = !GDrawStreamingSources; }));
 
+static int32 GDrawStreamingPerfs = 0;
+static FAutoConsoleCommand CVarDrawStreamingPerfs(
+	TEXT("wp.Runtime.ToggleDrawStreamingPerfs"),
+	TEXT("Toggles debug display of world partition streaming perfs."),
+	FConsoleCommandDelegate::CreateLambda([] { GDrawStreamingPerfs = !GDrawStreamingPerfs; }));
+
 static int32 GDrawLegends = 0;
 static FAutoConsoleCommand CVarGDrawLegends(
 	TEXT("wp.Runtime.ToggleDrawLegends"),
@@ -272,7 +279,10 @@ void UWorldPartitionSubsystem::Draw(UCanvas* Canvas, class APlayerController* PC
 			}
 			CurrentOffset.X = CanvasBottomRightPadding.X;
 		}
-
+	}
+	
+	if (GDrawStreamingPerfs || GDrawRuntimeHash2D)
+	{
 		{
 			FString StatusText;
 			if (IsIncrementalPurgePending()) { StatusText += TEXT("(Purging) "); }
@@ -302,7 +312,7 @@ void UWorldPartitionSubsystem::Draw(UCanvas* Canvas, class APlayerController* PC
 				StatusText = TEXT("Unknown");
 				break;
 			}
-			const FString Text = FString::Printf(TEXT("Streaming Performance: %s"), *StatusText);
+			const FString Text = FString::Printf(TEXT("Streaming Performance: %s (Blocking %s)"), *StatusText, GBlockOnSlowStreaming ? TEXT("Enabled") : TEXT("Disabled"));
 			FWorldPartitionDebugHelper::DrawText(Canvas, Text, GEngine->GetSmallFont(), FColor::White, CurrentOffset);
 		}
 	}
