Clean up some more MI code to avoid duplication for different parameter types

diff --git a/Engine/Source/Runtime/Engine/Classes/Materials/MaterialInstance.h b/Engine/Source/Runtime/Engine/Classes/Materials/MaterialInstance.h
--- a/Engine/Source/Runtime/Engine/Classes/Materials/MaterialInstance.h
+++ b/Engine/Source/Runtime/Engine/Classes/Materials/MaterialInstance.h
@@ -88,8 +88,8 @@ struct FScalarParameterValue
 	UPROPERTY()
 	FGuid ExpressionGUID;
 
-	FScalarParameterValue()
-		: ParameterValue(0)
+	explicit FScalarParameterValue(const FMaterialParameterInfo& InParameterInfo = FMaterialParameterInfo(), float InValue = 0.0f, const FScalarParameterAtlasInstanceData& InAtlasData = FScalarParameterAtlasInstanceData())
+		: AtlasData(InAtlasData), ParameterInfo(InParameterInfo), ParameterValue(InValue)
 	{
 	}
 
@@ -142,8 +142,8 @@ struct FVectorParameterValue
 	UPROPERTY()
 	FGuid ExpressionGUID;
 
-	FVectorParameterValue()
-		: ParameterValue(ForceInit)
+	explicit FVectorParameterValue(const FMaterialParameterInfo& InParameterInfo = FMaterialParameterInfo(), const FLinearColor& InValue = FLinearColor(ForceInit))
+		: ParameterInfo(InParameterInfo), ParameterValue(InValue)
 	{
 	}
 
@@ -193,8 +193,8 @@ struct FTextureParameterValue
 	UPROPERTY()
 	FGuid ExpressionGUID;
 
-	FTextureParameterValue()
-		: ParameterValue(NULL)
+	explicit FTextureParameterValue(const FMaterialParameterInfo& InParameterInfo = FMaterialParameterInfo(), class UTexture* InValue = nullptr)
+		: ParameterInfo(InParameterInfo), ParameterValue(InValue)
 	{
 	}
 
@@ -239,8 +239,8 @@ struct FRuntimeVirtualTextureParameterValue
 	UPROPERTY()
 	FGuid ExpressionGUID;
 
-	FRuntimeVirtualTextureParameterValue()
-		: ParameterValue(NULL)
+	explicit FRuntimeVirtualTextureParameterValue(const FMaterialParameterInfo& InParameterInfo = FMaterialParameterInfo(), class URuntimeVirtualTexture* InValue = nullptr)
+		: ParameterInfo(InParameterInfo), ParameterValue(InValue)
 	{
 	}
 
@@ -293,9 +293,8 @@ struct FFontParameterValue
 	UPROPERTY()
 	FGuid ExpressionGUID;
 
-	FFontParameterValue()
-		: FontValue(nullptr)
-		, FontPage(0)
+	explicit FFontParameterValue(const FMaterialParameterInfo& InParameterInfo = FMaterialParameterInfo(), class UFont* InFont = nullptr, int32 InPage = 0)
+		: ParameterInfo(InParameterInfo), FontValue(InFont), FontPage(InPage)
 	{
 	}
 
@@ -816,6 +815,8 @@ protected:
 	/**
 	 * Internal interface for setting / updating values for material instances.
 	 */
+	void ReserveParameterValuesInternal(EMaterialParameterType Type, int32 Capacity);
+	void AddParameterValueInternal(const FMaterialParameterInfo& ParameterInfo, const FMaterialParameterMetadata& Meta, EMaterialSetParameterValueFlags Flags = EMaterialSetParameterValueFlags::None);
 	void SetParameterValueInternal(const FMaterialParameterInfo& ParameterInfo, const FMaterialParameterMetadata& Meta, EMaterialSetParameterValueFlags Flags = EMaterialSetParameterValueFlags::None);
 	void SetVectorParameterValueInternal(const FMaterialParameterInfo& ParameterInfo, FLinearColor Value);
 	bool SetVectorParameterByIndexInternal(int32 ParameterIndex, FLinearColor Value);
diff --git a/Engine/Source/Runtime/Engine/Private/Materials/MaterialInstance.cpp b/Engine/Source/Runtime/Engine/Private/Materials/MaterialInstance.cpp
--- a/Engine/Source/Runtime/Engine/Private/Materials/MaterialInstance.cpp
+++ b/Engine/Source/Runtime/Engine/Private/Materials/MaterialInstance.cpp
@@ -1464,7 +1464,7 @@ void UMaterialInstanceDynamic::CopyScalarAndVectorParameters(const UMaterialInte
 				ParameterValue.Value = FMaterialParameterValue(Parameter.ParameterType, DefaultValue);
 			}
 
-			SetParameterValueInternal(FMaterialParameterInfo(Parameter.ParameterInfo), ParameterValue);
+			AddParameterValueInternal(FMaterialParameterInfo(Parameter.ParameterInfo), ParameterValue);
 		}
 
 		// now, init the resources
@@ -1530,90 +1530,25 @@ void UMaterialInstance::CopyMaterialInstanceParameters(UMaterialInterface* Sourc
 		//setup some arrays to use
 		TArray<FMaterialParameterInfo> OutParameterInfo;
 		TArray<FGuid> Guids;
-		
-		// Handle all the fonts
-		GetAllFontParameterInfo(OutParameterInfo, Guids);
-		for (const FMaterialParameterInfo& ParameterInfo : OutParameterInfo)
-		{
-			UFont* FontValue = nullptr;
-			int32 FontPage;
-			if (Source->GetFontParameterValue(ParameterInfo, FontValue, FontPage))
-			{
-				FFontParameterValue* ParameterValue = new(FontParameterValues) FFontParameterValue;
-				ParameterValue->ParameterInfo = ParameterInfo;
-				ParameterValue->ExpressionGUID.Invalidate();
-				ParameterValue->FontValue = FontValue;
-				ParameterValue->FontPage = FontPage;
-			}
-		}
-
-		// Now do the scalar params
-		OutParameterInfo.Reset();
-		Guids.Reset();
-		GetAllScalarParameterInfo(OutParameterInfo, Guids);
-		for (const FMaterialParameterInfo& ParameterInfo : OutParameterInfo)
-		{
-			float ScalarValue = 1.0f;
-			if (Source->GetScalarParameterValue(ParameterInfo, ScalarValue))
-			{
-				FScalarParameterValue* ParameterValue = new(ScalarParameterValues) FScalarParameterValue;
-				ParameterValue->ParameterInfo = ParameterInfo;
-				ParameterValue->ExpressionGUID.Invalidate();
-				ParameterValue->ParameterValue = ScalarValue;
-#if WITH_EDITOR
-				IsScalarParameterUsedAsAtlasPosition(ParameterValue->ParameterInfo, ParameterValue->AtlasData.bIsUsedAsAtlasPosition, ParameterValue->AtlasData.Curve, ParameterValue->AtlasData.Atlas);
-#endif
-			}
-		}
-
-		// Now do the vector params
-		OutParameterInfo.Reset();
-		Guids.Reset();
-		GetAllVectorParameterInfo(OutParameterInfo, Guids);
-		for (const FMaterialParameterInfo& ParameterInfo : OutParameterInfo)
-		{
-			FLinearColor VectorValue;
-			if (Source->GetVectorParameterValue(ParameterInfo, VectorValue))
-			{
-				FVectorParameterValue* ParameterValue = new(VectorParameterValues) FVectorParameterValue;
-				ParameterValue->ParameterInfo = ParameterInfo;
-				ParameterValue->ExpressionGUID.Invalidate();
-				ParameterValue->ParameterValue = VectorValue;
-			}
-		}
-
-		// Now do the texture params
-		OutParameterInfo.Reset();
-		Guids.Reset();
-		GetAllTextureParameterInfo(OutParameterInfo, Guids);
-		for (const FMaterialParameterInfo& ParameterInfo : OutParameterInfo)
-		{
-			UTexture* TextureValue = nullptr;
-			if (Source->GetTextureParameterValue(ParameterInfo, TextureValue))
-			{
-				FTextureParameterValue* ParameterValue = new(TextureParameterValues) FTextureParameterValue;
-				ParameterValue->ParameterInfo = ParameterInfo;
-				ParameterValue->ExpressionGUID.Invalidate();
-				ParameterValue->ParameterValue = TextureValue;
-			}
-		}
 
-		// Now do the runtime virtual texture params
-		OutParameterInfo.Reset();
-		Guids.Reset();
-		GetAllRuntimeVirtualTextureParameterInfo(OutParameterInfo, Guids);
-		for (const FMaterialParameterInfo& ParameterInfo : OutParameterInfo)
+		for (int32 ParameterTypeIndex = 0; ParameterTypeIndex < NumMaterialParameterTypes; ++ParameterTypeIndex)
 		{
-			URuntimeVirtualTexture* Value = nullptr;
-			if (Source->GetRuntimeVirtualTextureParameterValue(ParameterInfo, Value))
+			const EMaterialParameterType ParameterType = (EMaterialParameterType)ParameterTypeIndex;
+			if (!IsStaticMaterialParameter(ParameterType))
 			{
-				FRuntimeVirtualTextureParameterValue* ParameterValue = new(RuntimeVirtualTextureParameterValues) FRuntimeVirtualTextureParameterValue;
-				ParameterValue->ParameterInfo = ParameterInfo;
-				ParameterValue->ExpressionGUID.Invalidate();
-				ParameterValue->ParameterValue = Value;
+				GetAllParameterInfoOfType(ParameterType, OutParameterInfo, Guids);
+				ReserveParameterValuesInternal(ParameterType, OutParameterInfo.Num());
+				for (const FMaterialParameterInfo& ParameterInfo : OutParameterInfo)
+				{
+					FMaterialParameterMetadata SourceValue;
+					if (Source->GetParameterValue(ParameterType, ParameterInfo, SourceValue, EMaterialGetParameterValueFlags::CheckAll))
+					{
+						AddParameterValueInternal(ParameterInfo, SourceValue, EMaterialSetParameterValueFlags::SetCurveAtlas);
+					}
+				}
 			}
 		}
-
+		
 		// Now, init the resources
 		InitResources();
 
@@ -3000,6 +2935,45 @@ void FMaterialInstanceParameterUpdateContext::SetBasePropertyOverrides(const FMa
 }
 #endif // WITH_EDITORONLY_DATA
 
+void UMaterialInstance::ReserveParameterValuesInternal(EMaterialParameterType Type, int32 Capacity)
+{
+	switch (Type)
+	{
+	case EMaterialParameterType::Scalar: ScalarParameterValues.Reserve(Capacity); break;
+	case EMaterialParameterType::Vector: VectorParameterValues.Reserve(Capacity); break;
+	case EMaterialParameterType::Texture: TextureParameterValues.Reserve(Capacity); break;
+	case EMaterialParameterType::Font: FontParameterValues.Reserve(Capacity); break;
+	case EMaterialParameterType::RuntimeVirtualTexture: RuntimeVirtualTextureParameterValues.Reserve(Capacity); break;
+	default: checkNoEntry();
+	}
+}
+
+void UMaterialInstance::AddParameterValueInternal(const FMaterialParameterInfo& ParameterInfo, const FMaterialParameterMetadata& Meta, EMaterialSetParameterValueFlags Flags)
+{
+	const bool bUseAtlas = EnumHasAnyFlags(Flags, EMaterialSetParameterValueFlags::SetCurveAtlas);
+	const FMaterialParameterValue& Value = Meta.Value;
+	FScalarParameterAtlasInstanceData AtlasData;
+	switch (Value.Type)
+	{
+	case EMaterialParameterType::Scalar:
+#if WITH_EDITORONLY_DATA
+		if (bUseAtlas)
+		{
+			AtlasData.bIsUsedAsAtlasPosition = Meta.bUsedAsAtlasPosition;
+			AtlasData.Atlas = Meta.ScalarAtlas;
+			AtlasData.Curve = Meta.ScalarCurve;
+		}
+#endif // WITH_EDITORONLY_DATA
+		ScalarParameterValues.Emplace(ParameterInfo, Value.AsScalar(), AtlasData);
+		break;
+	case EMaterialParameterType::Vector: VectorParameterValues.Emplace(ParameterInfo, Value.AsLinearColor()); break;
+	case EMaterialParameterType::Texture: TextureParameterValues.Emplace(ParameterInfo, Value.Texture); break;
+	case EMaterialParameterType::Font: FontParameterValues.Emplace(ParameterInfo, Value.Font.Value, Value.Font.Page); break;
+	case EMaterialParameterType::RuntimeVirtualTexture: RuntimeVirtualTextureParameterValues.Emplace(ParameterInfo, Value.RuntimeVirtualTexture); break;
+	default: checkNoEntry();
+	}
+}
+
 void UMaterialInstance::SetParameterValueInternal(const FMaterialParameterInfo& ParameterInfo, const FMaterialParameterMetadata& Meta, EMaterialSetParameterValueFlags Flags)
 {
 	const bool bUseAtlas = EnumHasAnyFlags(Flags, EMaterialSetParameterValueFlags::SetCurveAtlas);
@@ -4301,22 +4275,8 @@ void UMaterialInstance::CopyMaterialUniformParametersInternal(UMaterialInterface
 					for (const FMaterialNumericParameterInfo& Parameter : MaterialResource->GetUniformNumericParameterExpressions())
 					{
 						const UE::Shader::FValue DefaultValue = MaterialResource->GetUniformExpressions().GetDefaultParameterValue(Parameter.ParameterType, Parameter.DefaultValueOffset);
-						if (Parameter.ParameterType == EMaterialParameterType::Scalar)
-						{
-							FScalarParameterValue* ParameterValue = new(ScalarParameterValues) FScalarParameterValue;
-							ParameterValue->ParameterInfo.Name = Parameter.ParameterInfo.GetName();
-							ParameterValue->ParameterValue = DefaultValue.AsFloatScalar();
-						}
-						else if (Parameter.ParameterType == EMaterialParameterType::Vector)
-						{
-							FVectorParameterValue* ParameterValue = new(VectorParameterValues) FVectorParameterValue;
-							ParameterValue->ParameterInfo.Name = Parameter.ParameterInfo.GetName();
-							ParameterValue->ParameterValue = DefaultValue.AsLinearColor();
-						}
-						else
-						{
-							ensure(false);
-						}
+						const FMaterialParameterMetadata Meta(Parameter.ParameterType, DefaultValue);
+						AddParameterValueInternal(Parameter.ParameterInfo.GetName(), Meta);
 					}
 
 					// Textures
diff --git a/Engine/Source/Runtime/Engine/Public/MaterialTypes.h b/Engine/Source/Runtime/Engine/Public/MaterialTypes.h
--- a/Engine/Source/Runtime/Engine/Public/MaterialTypes.h
+++ b/Engine/Source/Runtime/Engine/Public/MaterialTypes.h
@@ -178,6 +178,10 @@ struct FMaterialParameterMetadata
 {
 	FMaterialParameterValue Value;
 
+	FMaterialParameterMetadata() = default;
+	FMaterialParameterMetadata(const FMaterialParameterValue& InValue) : Value(InValue) {}
+	FMaterialParameterMetadata(EMaterialParameterType Type, const UE::Shader::FValue& InValue) : Value(Type, InValue) {}
+
 #if WITH_EDITORONLY_DATA
 	/** Name of channels, for Vectors/Textures */
 	FText ChannelName[4];
