Fix non-editor compile

diff --git a/Engine/Source/Runtime/Engine/Classes/Materials/MaterialInstance.h b/Engine/Source/Runtime/Engine/Classes/Materials/MaterialInstance.h
--- a/Engine/Source/Runtime/Engine/Classes/Materials/MaterialInstance.h
+++ b/Engine/Source/Runtime/Engine/Classes/Materials/MaterialInstance.h
@@ -89,8 +89,11 @@ struct FScalarParameterValue
 	FGuid ExpressionGUID;
 
 	explicit FScalarParameterValue(const FMaterialParameterInfo& InParameterInfo = FMaterialParameterInfo(), float InValue = 0.0f, const FScalarParameterAtlasInstanceData& InAtlasData = FScalarParameterAtlasInstanceData())
-		: AtlasData(InAtlasData), ParameterInfo(InParameterInfo), ParameterValue(InValue)
+		: ParameterInfo(InParameterInfo), ParameterValue(InValue)
 	{
+#if WITH_EDITORONLY_DATA
+		AtlasData = InAtlasData;
+#endif
 	}
 
 	bool IsOverride() const { return true; }
