Rename APlayerState::Ping and associated functions to CompressedPing, and introduced a new getter GetPingInMilliseconds that returns either the accurate ping if available or the decompressed replicated ping otherwise

diff --git a/Engine/Config/BaseEngine.ini b/Engine/Config/BaseEngine.ini
--- a/Engine/Config/BaseEngine.ini
+++ b/Engine/Config/BaseEngine.ini
@@ -919,6 +919,7 @@ NearClipPlane=10.0
 +PropertyRedirects=(OldName="GeometryCollectionComponent.Simulating", NewName="SimulatePhysics")
 +PropertyRedirects=(OldName="ControlRigComponent.OnPreUpdateDelegate", NewName="OnPreForwardsSolveDelegate")
 +PropertyRedirects=(OldName="ControlRigComponent.OnPostUpdateDelegate", NewName="OnPostForwardsSolveDelegate")
++PropertyRedirects=(OldName="PlayerState.Ping",NewName="CompressedPing")
 ; Sequencer LWC support
 +ClassRedirects=(OldName="/Script/MovieSceneTracks.MovieSceneVectorSection",NewName="/Script/MovieSceneTracks.MovieSceneFloatVectorSection")
 +ClassRedirects=(OldName="/Script/MovieSceneTracks.MovieSceneVectorTrack",NewName="/Script/MovieSceneTracks.MovieSceneFloatVectorTrack")
diff --git a/Engine/Source/Runtime/Engine/Classes/GameFramework/PlayerState.h b/Engine/Source/Runtime/Engine/Classes/GameFramework/PlayerState.h
--- a/Engine/Source/Runtime/Engine/Classes/GameFramework/PlayerState.h
+++ b/Engine/Source/Runtime/Engine/Classes/GameFramework/PlayerState.h
@@ -82,17 +82,16 @@ class ENGINE_API APlayerState : public AInfo
 	UPROPERTY(ReplicatedUsing=OnRep_PlayerId, BlueprintReadOnly, Category=PlayerState)
 	int32 PlayerId;
 
+private:
 	/** Replicated compressed ping for this player (holds ping in msec divided by 4) */
-	UE_DEPRECATED(4.25, "This member will be made private. Use GetPing or SetPing instead.")
-	UPROPERTY(Replicated, BlueprintReadOnly, Category=PlayerState)
-	uint8 Ping;
+	UPROPERTY(Replicated, BlueprintReadOnly, Category=PlayerState, meta=(AllowPrivateAccess))
+	uint8 CompressedPing;
 
-private:
 	/** The current PingBucket index that is being filled */
-	uint8			CurPingBucket;
+	uint8 CurPingBucket;
 
 	/**
-	 * Whether or not this player's replicated Ping value is updated automatically.
+	 * Whether or not this player's replicated CompressedPing value is updated automatically.
 	 * Since player states are always relevant by default, in cases where there are many players replicating,
 	 * replicating the ping value can cause additional unnecessary overhead on servers if the value isn't
 	 * needed on clients.
@@ -144,7 +143,7 @@ public:
 	UPROPERTY()
 	TSubclassOf<class ULocalMessage> EngineMessageClass;
 
-	/** Exact ping as float (rounded and compressed in replicated Ping) */
+	/** Exact ping as float (rounded and compressed in replicated CompressedPing) */
 	float ExactPing;
 
 	UE_DEPRECATED(4.27, "Please use ExactPing instead.")
@@ -360,14 +359,30 @@ public:
 	/** Sets the value of PlayerId without causing other side effects to this instance. */
 	void SetPlayerId(const int32 NewId);
 
-	/** Gets the literal value of Ping. */
-	uint8 GetPing() const
+	/** Gets the literal value of the compressed Ping value (Ping = PingInMS / 4). */
+	uint8 GetCompressedPing() const
 	{
-		return Ping;
+		return CompressedPing;
 	}
 
-	/** Sets the value of Ping without causing other side effects to this instance. */
-	void SetPing(const uint8 NewPing);
+	UE_DEPRECATED(5.0, "Use GetPingInMilliseconds() or GetCompressedPing() instead")
+	uint8 GetPing() const { return GetCompressedPing(); }
+
+	/** Sets the value of CompressedPing without causing other side effects to this instance. */
+	void SetCompressedPing(const uint8 NewPing);
+
+	/**
+	 * Returns the ping (in milliseconds)
+	 *
+	 * Returns ExactPing if available (local players or when running on the server), and
+	 * the replicated CompressedPing (converted back to milliseconds) otherwise.
+	 * 
+	 * Note that replication of CompressedPing is controlled by bShouldUpdateReplicatedPing,
+	 * and if disabled then this will return 0 or a stale value on clients for player states
+	 * that aren't related to local players
+	 */
+	UFUNCTION(BlueprintCallable, Category = "PlayerState")
+	float GetPingInMilliseconds() const;
 
 	/** Gets the literal value of bIsSpectator. */
 	bool IsSpectator() const
diff --git a/Engine/Source/Runtime/Engine/Private/PlayerState.cpp b/Engine/Source/Runtime/Engine/Private/PlayerState.cpp
--- a/Engine/Source/Runtime/Engine/Private/PlayerState.cpp
+++ b/Engine/Source/Runtime/Engine/Private/PlayerState.cpp
@@ -1,9 +1,5 @@
 // Copyright Epic Games, Inc. All Rights Reserved.
 
-/*=============================================================================
-	PlayerState.cpp: 
-=============================================================================*/
-
 #include "GameFramework/PlayerState.h"
 #include "Engine/World.h"
 #include "GameFramework/Controller.h"
@@ -141,7 +137,7 @@ void APlayerState::RecalculateAvgPing()
 
 	if (bShouldUpdateReplicatedPing || !HasAuthority())
 	{
-		SetPing(FMath::Min(255, (int32)(ExactPing * 0.25f)));
+		SetCompressedPing(FMath::Min(255, (int32)(ExactPing * 0.25f)));
 	}
 }
 
@@ -169,7 +165,7 @@ void APlayerState::OverrideWith(APlayerState* PlayerState)
 void APlayerState::CopyProperties(APlayerState* PlayerState)
 {
 	PlayerState->SetScore(GetScore());
-	PlayerState->SetPing(GetPing());
+	PlayerState->SetCompressedPing(GetCompressedPing());
 	PlayerState->ExactPing = ExactPing;
 	PlayerState->SetPlayerId(GetPlayerId());
 	PlayerState->SetUniqueId(GetUniqueId().GetUniqueNetId());
@@ -462,7 +458,7 @@ void APlayerState::GetLifetimeReplicatedProps(TArray< FLifetimeProperty > & OutL
 	DOREPLIFETIME_WITH_PARAMS_FAST(APlayerState, PlayerNamePrivate, SharedParams);
 
 	SharedParams.Condition = COND_SkipOwner;
-	DOREPLIFETIME_WITH_PARAMS_FAST(APlayerState, Ping, SharedParams);
+	DOREPLIFETIME_WITH_PARAMS_FAST(APlayerState, CompressedPing, SharedParams);
 
 	SharedParams.Condition = COND_InitialOnly;
 	DOREPLIFETIME_WITH_PARAMS_FAST(APlayerState, PlayerId, SharedParams);
@@ -485,10 +481,24 @@ void APlayerState::SetPlayerId(const int32 NewId)
 	PlayerId = NewId;
 }
 
-void APlayerState::SetPing(const uint8 NewPing)
+void APlayerState::SetCompressedPing(const uint8 NewPing)
 {
-	MARK_PROPERTY_DIRTY_FROM_NAME(APlayerState, Ping, this);
-	Ping = NewPing;
+	MARK_PROPERTY_DIRTY_FROM_NAME(APlayerState, CompressedPing, this);
+	CompressedPing = NewPing;
+}
+
+float APlayerState::GetPingInMilliseconds() const
+{
+	if (ExactPing > 0.0f)
+	{
+		// Prefer the exact ping if set (only on the server or for the local players)
+		return ExactPing;
+	}
+	else
+	{
+		// Otherwise, use the replicated compressed ping
+		return CompressedPing * 4.0f;
+	}
 }
 
 void APlayerState::SetIsSpectator(const bool bNewSpectator)
