Reverting shader map log spam for the majority of users.

diff --git a/Engine/Source/Runtime/Engine/Private/Materials/MaterialShader.cpp b/Engine/Source/Runtime/Engine/Private/Materials/MaterialShader.cpp
--- a/Engine/Source/Runtime/Engine/Private/Materials/MaterialShader.cpp
+++ b/Engine/Source/Runtime/Engine/Private/Materials/MaterialShader.cpp
@@ -1351,7 +1351,7 @@ void FMaterialShaderMap::LoadFromDerivedDataCache(const FMaterial* Material, con
 
 			if (InOutShaderMap)
 			{
-				UE_LOG(LogMaterial, Display, TEXT("Loaded shaders for %s from DDC (key hash: %s)"), *Material->GetAssetName(), *OutDDCKeyDesc);
+				UE_LOG(LogMaterial, Verbose, TEXT("Loaded shaders for %s from DDC (key hash: %s)"), *Material->GetAssetName(), *OutDDCKeyDesc);
 			}
 		}
 		INC_FLOAT_STAT_BY(STAT_ShaderCompiling_DDCLoading,(float)MaterialDDCTime);
@@ -1367,7 +1367,7 @@ void FMaterialShaderMap::SaveToDerivedDataCache()
 	Serialize(Ar);
 
 	FString DataKey = GetMaterialShaderMapKeyString(ShaderMapId, GetShaderPlatform());
-	UE_LOG(LogMaterial, Display, TEXT("Saved shaders for %s from DDC (key hash: %s)"), GetMaterialPath(), *FSHA1_HashString(DataKey));
+	UE_LOG(LogMaterial, Verbose, TEXT("Saved shaders for %s from DDC (key hash: %s)"), GetMaterialPath(), *FSHA1_HashString(DataKey));
 	GetDerivedDataCacheRef().Put(*DataKey, SaveData, FStringView(GetFriendlyName()));
 	COOK_STAT(Timer.AddMiss(SaveData.Num()));
 }
