Add console command to show/hide HUD

diff --git a/Engine/Source/Runtime/Engine/Private/Slate/SGameLayerManager.cpp b/Engine/Source/Runtime/Engine/Private/Slate/SGameLayerManager.cpp
--- a/Engine/Source/Runtime/Engine/Private/Slate/SGameLayerManager.cpp
+++ b/Engine/Source/Runtime/Engine/Private/Slate/SGameLayerManager.cpp
@@ -20,6 +20,9 @@
 #include "Types/InvisibleToWidgetReflectorMetaData.h"
 #include "Framework/Application/SlateApplication.h"
 
+#ifndef UE_SLATE_WITH_GAMELAYER_CANVAS_VISIBILITY_COMMANDS
+	#define UE_SLATE_WITH_GAMELAYER_CANVAS_VISIBILITY_COMMANDS !(UE_BUILD_SHIPPING || UE_BUILD_TEST) && WITH_SLATE_DEBUGGING
+#endif
 
 /* SGameLayerManager interface
  *****************************************************************************/
@@ -37,6 +40,49 @@ static FAutoConsoleVariableRef CVarEnablePerUserHitTesting(
 	FConsoleVariableDelegate::CreateStatic(&HandlePerUserHitTestingToggled)
 );
 
+#if UE_SLATE_WITH_GAMELAYER_CANVAS_VISIBILITY_COMMANDS
+namespace CanvasVisibility
+{
+	static TArray<SGameLayerManager*> GameLayerManagerInstances;
+	static bool ViewportSlotVisibilityConsoleValue = true;
+	static bool DebugCanvasVisibilityConsoleValue = true;
+	static bool PlayerCanvasVisibilityConsoleValue = true;
+
+	enum class CanvasType
+	{
+		PlayerCanvas,
+		DebugCanvas,
+		Viewport,
+		AllCanvases
+	};
+}
+
+static FAutoConsoleVariableRef CVarShowViewportSlot(
+	TEXT("Slate.GameLayer.ViewportSlotVisible"),
+	CanvasVisibility::ViewportSlotVisibilityConsoleValue,
+	TEXT("Show/Hide the slot on viewport"),
+	ECVF_Cheat
+);
+static FAutoConsoleVariableRef CVarShowDebugCanvas(
+	TEXT("Slate.GameLayer.DebugCanvasVisible"),
+	CanvasVisibility::DebugCanvasVisibilityConsoleValue,
+	TEXT("Show/Hide the debug canvas."),
+	ECVF_Cheat
+);
+static FAutoConsoleVariableRef CVarShowPlayerCanvas(
+	TEXT("Slate.GameLayer.PlayerCanvasVisible"),
+	CanvasVisibility::PlayerCanvasVisibilityConsoleValue,
+	TEXT("Show/Hide the player canvas."),
+	ECVF_Cheat
+);
+static TAutoConsoleVariable<bool> CVarShowAllCanvases(
+	TEXT("Slate.GameLayer.AllCanvasesVisible"),
+	true,
+	TEXT("Show/Hide the viewport slot, player canvas, and debug canvas."),
+	ECVF_Cheat
+);
+#endif
+
 SGameLayerManager::SGameLayerManager()
 :	DefaultWindowTitleBarHeight(64.0f)
 ,	bIsGameUsingBorderlessWindow(false)
@@ -82,7 +128,10 @@ void SGameLayerManager::Construct(const SGameLayerManager::FArguments& InArgs)
 
 				+ SOverlay::Slot()
 				[
-					InArgs._Content.Widget
+					SAssignNew(ViewportSlotContainer, SBox)
+					[
+						InArgs._Content.Widget
+					]
 				]
 
 				+ SOverlay::Slot()
@@ -143,8 +192,55 @@ void SGameLayerManager::Construct(const SGameLayerManager::FArguments& InArgs)
 	TitleBarAreaVerticalBox->SetRequestToggleFullscreenCallback(FSimpleDelegate::CreateSP(this, &SGameLayerManager::RequestToggleFullscreen));
 
 	SetWindowTitleBarState(nullptr, EWindowTitleBarMode::Overlay, false, false, false);
-
+	
 	bIsGameUsingBorderlessWindow = GetDefault<UGeneralProjectSettings>()->bUseBorderlessWindow && PLATFORM_WINDOWS;
+
+#if UE_SLATE_WITH_GAMELAYER_CANVAS_VISIBILITY_COMMANDS
+	CanvasVisibility::GameLayerManagerInstances.Add(this);
+	if (!CVarShowAllCanvases->AsVariable()->OnChangedDelegate().IsBound())
+	{
+		auto ShowHUD = [](IConsoleVariable* Variable, CanvasVisibility::CanvasType CanvasType)
+		{
+			bool bVisible = Variable->GetBool();
+			for (int32 Index = 0; Index < CanvasVisibility::GameLayerManagerInstances.Num(); ++Index)
+			{
+				if (SGameLayerManager* GameManagerInstance = CanvasVisibility::GameLayerManagerInstances[Index])
+				{
+					switch (CanvasType)
+					{
+					case CanvasVisibility::CanvasType::PlayerCanvas:
+						GameManagerInstance->ShowPlayerCanvas(bVisible);
+						break;
+					case CanvasVisibility::CanvasType::DebugCanvas:
+						GameManagerInstance->ShowDebugCanvas(bVisible);
+						break;
+					case CanvasVisibility::CanvasType::Viewport:
+						GameManagerInstance->ShowViewportSlot(bVisible);
+						break;
+					case CanvasVisibility::CanvasType::AllCanvases:
+						GameManagerInstance->ShowViewportSlot(bVisible);
+						GameManagerInstance->ShowDebugCanvas(bVisible);
+						GameManagerInstance->ShowPlayerCanvas(bVisible);
+						CanvasVisibility::ViewportSlotVisibilityConsoleValue = bVisible;
+						CanvasVisibility::PlayerCanvasVisibilityConsoleValue = bVisible;
+						CanvasVisibility::DebugCanvasVisibilityConsoleValue = bVisible;
+						break;
+					default:
+						break;
+					}
+				}
+			};
+		};
+		CVarShowAllCanvases->AsVariable()->SetOnChangedCallback(FConsoleVariableDelegate::CreateLambda(ShowHUD, CanvasVisibility::CanvasType::AllCanvases));
+		CVarShowPlayerCanvas->AsVariable()->SetOnChangedCallback(FConsoleVariableDelegate::CreateLambda(ShowHUD, CanvasVisibility::CanvasType::PlayerCanvas));
+		CVarShowDebugCanvas->AsVariable()->SetOnChangedCallback(FConsoleVariableDelegate::CreateLambda(ShowHUD, CanvasVisibility::CanvasType::DebugCanvas));
+		CVarShowViewportSlot->AsVariable()->SetOnChangedCallback(FConsoleVariableDelegate::CreateLambda(ShowHUD, CanvasVisibility::CanvasType::Viewport));
+	}
+
+	ShowViewportSlot(CanvasVisibility::ViewportSlotVisibilityConsoleValue);
+	ShowDebugCanvas(CanvasVisibility::DebugCanvasVisibilityConsoleValue);
+	ShowPlayerCanvas(CanvasVisibility::PlayerCanvasVisibilityConsoleValue);
+#endif
 }
 
 void SGameLayerManager::SetSceneViewport(FSceneViewport* InSceneViewport)
@@ -630,3 +726,25 @@ void SGameLayerManager::RequestToggleFullscreen()
 		GEngine->DeferredCommands.Add(TEXT("TOGGLE_FULLSCREEN"));
 	}
 }
+
+void SGameLayerManager::ShowPlayerCanvas(bool bIsVisible)
+{
+	PlayerCanvas->SetVisibility(bIsVisible ? EVisibility::SelfHitTestInvisible : EVisibility::Collapsed);
+}
+
+void SGameLayerManager::ShowDebugCanvas(bool bIsVisible)
+{
+	DebugCanvas->SetVisibility(bIsVisible ? EVisibility::SelfHitTestInvisible : EVisibility::Collapsed);
+}
+
+void SGameLayerManager::ShowViewportSlot(bool bIsVisible)
+{
+	ViewportSlotContainer->SetVisibility(bIsVisible ? EVisibility::SelfHitTestInvisible : EVisibility::Collapsed);
+}
+
+SGameLayerManager::~SGameLayerManager() 
+{
+#if UE_SLATE_WITH_GAMELAYER_CANVAS_VISIBILITY_COMMANDS
+	CanvasVisibility::GameLayerManagerInstances.RemoveSingleSwap(this);
+#endif
+}
\ No newline at end of file
diff --git a/Engine/Source/Runtime/Engine/Public/Slate/SGameLayerManager.h b/Engine/Source/Runtime/Engine/Public/Slate/SGameLayerManager.h
--- a/Engine/Source/Runtime/Engine/Public/Slate/SGameLayerManager.h
+++ b/Engine/Source/Runtime/Engine/Public/Slate/SGameLayerManager.h
@@ -99,6 +99,8 @@ public:
 
 	SGameLayerManager();
 
+	virtual ~SGameLayerManager();
+
 	/**
 	 * Construct this widget
 	 *
@@ -131,11 +133,11 @@ public:
 
 public:
 
-	// Begin SWidget overrides
+	//~ Begin SWidget overrides
 	virtual void Tick(const FGeometry& AllottedGeometry, const double InCurrentTime, const float InDeltaTime) override;
 	virtual int32 OnPaint(const FPaintArgs& Args, const FGeometry& AllottedGeometry, const FSlateRect& MyCullingRect, FSlateWindowElementList& OutDrawElements, int32 LayerId, const FWidgetStyle& InWidgetStyle, bool bParentEnabled) const override;
 	virtual bool OnVisualizeTooltip(const TSharedPtr<SWidget>& TooltipContent) override;
-	// End SWidget overrides
+	//~ End SWidget overrides
 
 	/**
 	 * Function will instruct internal DPI computations to use a provided reference viewport size instead of the actual viewport size.
@@ -149,6 +151,11 @@ private:
 	float GetGameViewportDPIScale() const;
 	FOptionalSize GetDefaultWindowTitleBarHeight() const;
 
+	//~ Functions for setting the visibility of different canvases containing HUD.
+	void ShowPlayerCanvas(bool bIsVisible);
+	void ShowDebugCanvas(bool bIsVisible);
+	void ShowViewportSlot(bool bIsVisible);
+
 private:
 
 	struct FPlayerLayer : TSharedFromThis<FPlayerLayer>
@@ -185,6 +192,7 @@ private:
 	TSharedPtr<SVerticalBox> WidgetHost;
 	TSharedPtr<SCanvas> PlayerCanvas;
 	TSharedPtr<SDebugCanvas> DebugCanvas;
+	TSharedPtr<SBox> ViewportSlotContainer;
 	TSharedPtr<STooltipPresenter> TooltipPresenter;
 
 	TSharedPtr<SWindowTitleBarArea> TitleBarAreaOverlay;
