#8100 : PR : Don't do redundant work in AActor::SetReplicates

diff --git a/Engine/Source/Runtime/Engine/Private/Actor.cpp b/Engine/Source/Runtime/Engine/Private/Actor.cpp
--- a/Engine/Source/Runtime/Engine/Private/Actor.cpp
+++ b/Engine/Source/Runtime/Engine/Private/Actor.cpp
@@ -3582,15 +3582,19 @@ void AActor::PostActorConstruction()
 
 void AActor::SetReplicates(bool bInReplicates)
 { 
-	if (GetLocalRole() == ROLE_Authority)
+	// Due to SetRemoteRoleForBackwardsCompat, it's possible that bReplicates is false, but RemoteRole is something other than ROLE_None.
+	// So, we'll also make sure that we don't need to update RemoteRole here, even if not bReplicates wouldn't change, to fix up that case.
+	const ENetRole ExpectedRemoteRole = bInReplicates ? ROLE_SimulatedProxy : ROLE_None;
+
+	if (GetLocalRole() != ROLE_Authority)
+	{
+		UE_LOG(LogActor, Warning, TEXT("SetReplicates called on actor '%s' that is not valid for having its role modified."), *GetName());
+	}
+	else if (bReplicates != bInReplicates || RemoteRole != ExpectedRemoteRole) 
 	{
 		const bool bNewlyReplicates = (bReplicates == false && bInReplicates == true);
 	
-		// Due to SetRemoteRoleForBackwardsCompat, it's possible that bReplicates is false, but RemoteRole is something
-		// other than ROLE_None.
-		// So, we'll always set RemoteRole here regardless of whether or not bReplicates would change, to fix up that
-		// case.
-		RemoteRole = bInReplicates ? ROLE_SimulatedProxy : ROLE_None;
+		RemoteRole = ExpectedRemoteRole;
 		bReplicates = bInReplicates;
 
 		if (bActorInitialized)
@@ -3618,10 +3622,6 @@ void AActor::SetReplicates(bool bInReplicates)
 			UE_LOG(LogActor, Warning, TEXT("SetReplicates called on non-initialized actor %s. Directly setting bReplicates is the correct procedure for pre-init actors."), *GetName());
 		}
 	}
-	else
-	{
-		UE_LOG(LogActor, Warning, TEXT("SetReplicates called on actor '%s' that is not valid for having its role modified."), *GetName());
-	}
 }
 
 void AActor::SetReplicateMovement(bool bInReplicateMovement)
