reduced audio tasks launches to improve perf. previously a new task was launched on adding a new audio command to the batch if the previous batch was completed or max batch size was reached. now it launches a new task only if max batch size is reached. batcher is flushed once a frame explicitly by FAudioDevice::Update(), this should be a sufficient frequency

diff --git a/Engine/Source/Runtime/Engine/Private/AudioThread.cpp b/Engine/Source/Runtime/Engine/Private/AudioThread.cpp
--- a/Engine/Source/Runtime/Engine/Private/AudioThread.cpp
+++ b/Engine/Source/Runtime/Engine/Private/AudioThread.cpp
@@ -316,8 +316,7 @@ struct FAudioAsyncBatcher
 #if !WITH_EDITOR
 		if (GCVarEnableBatchProcessing)
 		{
-			if (WorkItems.Num() >= GBatchAudioAsyncBatchSize // collected enough work
-				|| LastBatch.IsCompleted()) // the previous batch is completed, kick off a new one
+			if (WorkItems.Num() >= GBatchAudioAsyncBatchSize) // collected enough work
 			{
 				Flush();
 			}
