Re-cache component transform right before evaluation to ensure it is up-to-date when playing root motion

diff --git a/Engine/Source/Runtime/Engine/Private/Animation/AnimInstanceProxy.cpp b/Engine/Source/Runtime/Engine/Private/Animation/AnimInstanceProxy.cpp
--- a/Engine/Source/Runtime/Engine/Private/Animation/AnimInstanceProxy.cpp
+++ b/Engine/Source/Runtime/Engine/Private/Animation/AnimInstanceProxy.cpp
@@ -1183,6 +1183,17 @@ void FAnimInstanceProxy::PreEvaluateAnimation(UAnimInstance* InAnimInstance)
 	DECLARE_SCOPE_HIERARCHICAL_COUNTER_FUNC()
 
 	InitializeObjects(InAnimInstance);
+
+	// Re-cache Component Transforms if needed
+	// When playing root motion is the CharacterMovementComp who triggers the animation update and this happens before root motion is consumed and the position of the character is updated for this frame
+	// which means that at the point ComponentTransform is cached in the PreUpdate function it contain the previous frame transform
+	if (SkeletalMeshComponent && SkeletalMeshComponent->IsPlayingRootMotion())
+	{
+		ComponentTransform = SkeletalMeshComponent->GetComponentTransform();
+		ComponentRelativeTransform = SkeletalMeshComponent->GetRelativeTransform();
+		ActorTransform = SkeletalMeshComponent->GetOwner() ? SkeletalMeshComponent->GetOwner()->GetActorTransform() : FTransform::Identity;
+	}
+
 #if ENABLE_ANIM_LOGGING
 	LoggedMessagesMap.FindOrAdd(NAME_Evaluate).Reset();
 #endif
