- Edigrate changes to support multiple ActiveSound instances to play per AudioComponent
- Fix for virtualization not observing if an ActiveSound is associated with a Component set to allow multiple ActiveSound instances to play at once
- Fix component map member lifetime descrepancies in FAudioDevice
- Fix for incorrect member name checked in AudioComponent PropertyChange
- Add boolean option to disable parameter updates for playing sounds

diff --git a/Engine/Source/Runtime/Engine/Classes/Audio/SoundGeneratorParameterInterface.h b/Engine/Source/Runtime/Engine/Classes/Audio/SoundGeneratorParameterInterface.h
--- a/Engine/Source/Runtime/Engine/Classes/Audio/SoundGeneratorParameterInterface.h
+++ b/Engine/Source/Runtime/Engine/Classes/Audio/SoundGeneratorParameterInterface.h
@@ -1,8 +1,9 @@
 // Copyright Epic Games, Inc. All Rights Reserved.
 #pragma once
 
-#include "UObject/Object.h"
 #include "AudioParameterInterface.h"
+#include "CoreMinimal.h"
+#include "UObject/Object.h"
 
 #include "SoundGeneratorParameterInterface.generated.h"
 
@@ -10,6 +11,7 @@
 // Forward Declarations
 class FAudioDevice;
 class USoundBase;
+struct FActiveSound;
 
 UINTERFACE()
 class ENGINE_API USoundGeneratorParameterInterface : public UAudioParameterInterface
@@ -43,8 +45,8 @@ public:
 	/** Returns the active audio device to use for this component based on whether or not the component is playing in a world. */
 	virtual FAudioDevice* GetAudioDevice() const = 0;
 
-	/** Returns the id of the sound instance associated with the parameter interface. */
-	virtual uint64 GetInstanceID() const = 0;
+	/** Returns the id of the sound owner's instance associated with the parameter interface. */
+	virtual uint64 GetInstanceOwnerID() const = 0;
 
 	/** Returns the Game Thread copy of parameters to modify in place. */
 	virtual TArray<FAudioParameter>& GetInstanceParameters() = 0;
@@ -54,6 +56,8 @@ public:
 
 	virtual bool IsPlaying() const = 0;
 
+	virtual bool GetDisableParameterUpdatesWhilePlaying() const = 0;
+
 private:
 	void SetParameterInternal(FAudioParameter&& InValue);
 };
diff --git a/Engine/Source/Runtime/Engine/Classes/Components/AudioComponent.h b/Engine/Source/Runtime/Engine/Classes/Components/AudioComponent.h
--- a/Engine/Source/Runtime/Engine/Classes/Components/AudioComponent.h
+++ b/Engine/Source/Runtime/Engine/Classes/Components/AudioComponent.h
@@ -102,7 +102,6 @@ DECLARE_DYNAMIC_MULTICAST_DELEGATE_ThreeParams(FOnAudioMultiEnvelopeValue, const
 DECLARE_MULTICAST_DELEGATE_FourParams(FOnAudioMultiEnvelopeValueNative, const UAudioComponent*, const float, const float, const int32);
 
 
-
 /** Type of fade to use when adjusting the audio component's volume. */
 UENUM(BlueprintType)
 enum class EAudioFaderCurve : uint8
@@ -160,22 +159,22 @@ class UInitialActiveSoundParams : public UObject
  * @see https://docs.unrealengine.com/latest/INT/Audio/Overview/index.html
  * @see USoundBase
  */
-UCLASS(ClassGroup=(Audio, Common), hidecategories=(Object, ActorComponent, Physics, Rendering, Mobility, LOD), ShowCategories=Trigger, meta=(BlueprintSpawnableComponent))
+UCLASS(ClassGroup=(Audio, Common), HideCategories=(Object, ActorComponent, Physics, Rendering, Mobility, LOD), ShowCategories=Trigger, meta=(BlueprintSpawnableComponent))
 class ENGINE_API UAudioComponent : public USceneComponent, public ISoundGeneratorParameterInterface
 {
 	GENERATED_UCLASS_BODY()
 
 	/** The sound to be played */
-	UPROPERTY(EditAnywhere, BlueprintReadWrite, Category=Sound)
+	UPROPERTY(EditAnywhere, BlueprintReadWrite, Category = Sound)
 	TObjectPtr<USoundBase> Sound;
 
 	/** Array of per-instance parameters for this AudioComponent. Direct changes to this array from Blueprint (as opposed to updates
 	  * made via the Set functions) will not be forwarded to the sound if the component is actively playing. */
-	UPROPERTY(EditAnywhere, BlueprintReadWrite, Category=Sound, AdvancedDisplay)
+	UPROPERTY(EditAnywhere, BlueprintReadWrite, Category = Parameters, meta = (DisplayAfter = "bDisableParameterUpdatesWhilePlaying"))
 	TArray<FAudioParameter> InstanceParameters;
 
 	/** SoundClass that overrides that set on the referenced SoundBase when component is played. */
-	UPROPERTY(EditAnywhere, Category=Sound, AdvancedDisplay)
+	UPROPERTY(EditAnywhere, Category = Sound, AdvancedDisplay)
 	TObjectPtr<USoundClass> SoundClassOverride;
 
 	/** Auto destroy this component on completion */
@@ -191,32 +190,42 @@ class ENGINE_API UAudioComponent : public USceneComponent, public ISoundGenerato
 	uint8 bShouldRemainActiveIfDropped:1;
 
 	/** Overrides spatialization enablement in either the attenuation asset or on this audio component's attenuation settings override. */
-	UPROPERTY(EditAnywhere, BlueprintReadWrite, Category=Attenuation)
+	UPROPERTY(EditAnywhere, BlueprintReadWrite, Category = Attenuation)
 	uint8 bAllowSpatialization:1;
 
 	/** Allows defining attenuation settings directly on this audio component without using an attenuation settings asset. */
-	UPROPERTY(EditAnywhere, BlueprintReadWrite, Category=Attenuation)
+	UPROPERTY(EditAnywhere, BlueprintReadWrite, Category = Attenuation)
 	uint8 bOverrideAttenuation:1;
 
 	/** Whether or not to override the sound's subtitle priority. */
-	UPROPERTY(EditAnywhere, BlueprintReadWrite, Category = Sound)
+	UPROPERTY(EditAnywhere, BlueprintReadWrite, Category = Subtitles, meta = (InlineEditConditionToggle))
 	uint8 bOverrideSubtitlePriority:1;
 
 	/** Whether or not this sound plays when the game is paused in the UI */
-	UPROPERTY(EditAnywhere, BlueprintReadWrite, Category = Sound)
+	UPROPERTY(EditAnywhere, BlueprintReadWrite, Category = Sound, AdvancedDisplay)
 	uint8 bIsUISound : 1;
 
 	/** Whether or not to apply a low-pass filter to the sound that plays in this audio component. */
-	UPROPERTY(EditAnywhere, BlueprintReadWrite, Category = Filter)
+	UPROPERTY(EditAnywhere, BlueprintReadWrite, Category = Sound, meta = (InlineEditConditionToggle, DisplayAfter = "PitchMultiplier"))
 	uint8 bEnableLowPassFilter : 1;
 
-	UPROPERTY(EditAnywhere, BlueprintReadWrite, Category = Sound)
+	/** Whether or not to override the priority of the given sound with the value provided. */
+	UPROPERTY(EditAnywhere, BlueprintReadWrite, Category = Sound, meta = (InlineEditConditionToggle, DisplayAfter = "LowPassFilterFrequency"))
 	uint8 bOverridePriority:1;
 
 	/** If true, subtitles in the sound data will be ignored. */
-	UPROPERTY(EditAnywhere, BlueprintReadWrite, Category = Sound)
+	UPROPERTY(EditAnywhere, BlueprintReadWrite, Category = Subtitles)
 	uint8 bSuppressSubtitles:1;
 
+	/** If true, the Audio Component will play multiple sound instances at once. Switching sounds or calling play while already playing
+	  * will not stop already active instances. */
+	UPROPERTY(EditAnywhere, BlueprintReadWrite, Category = Sound, meta = (DisplayName = "Play Multiple Instances", DisplayAfter = "Priority"))
+	uint8 bCanPlayMultipleInstances:1;
+
+	/** If true, the Audio Component will ignore parameter updates for already-playing sound(s). */
+	UPROPERTY(EditAnywhere, BlueprintReadWrite, Category = Parameters)
+	uint8 bDisableParameterUpdatesWhilePlaying : 1;
+
 	/** Whether this audio component is previewing a sound */
 	uint8 bPreviewComponent:1;
 
@@ -256,7 +265,7 @@ class ENGINE_API UAudioComponent : public USceneComponent, public ISoundGenerato
 	* This also disables attachment on dedicated servers, where we don't actually activate even if bAutoActivate is true.
 	* @see AutoAttachParent, AutoAttachSocketName, AutoAttachLocationType
 	*/
-	UPROPERTY(EditAnywhere, BlueprintReadOnly, Category=Attachment)
+	UPROPERTY(EditAnywhere, BlueprintReadOnly, Category = Attachment)
 	uint8 bAutoManageAttachment:1;
 
 private:
@@ -272,47 +281,47 @@ public:
 	FName AudioComponentUserID;
 
 	/** The lower bound to use when randomly determining a pitch multiplier */
-	UPROPERTY(EditAnywhere, BlueprintReadWrite, Category= "Randomization|Pitch", meta = (DisplayName = "Pitch (Min)"))
+	UPROPERTY(EditAnywhere, BlueprintReadWrite, Category = "Randomization|Pitch", meta = (DisplayName = "Pitch (Min)"))
 	float PitchModulationMin;
 
 	/** The upper bound to use when randomly determining a pitch multiplier */
-	UPROPERTY(EditAnywhere, BlueprintReadWrite, Category="Randomization|Pitch", meta = (DisplayName = "Pitch (Max)"))
+	UPROPERTY(EditAnywhere, BlueprintReadWrite, Category = "Randomization|Pitch", meta = (DisplayName = "Pitch (Max)"))
 	float PitchModulationMax;
 
 	/** The lower bound to use when randomly determining a volume multiplier */
-	UPROPERTY(EditAnywhere, BlueprintReadWrite, Category= "Randomization|Volume", meta = (DisplayName = "Volume (Min)"))
+	UPROPERTY(EditAnywhere, BlueprintReadWrite, Category = "Randomization|Volume", meta = (DisplayName = "Volume (Min)"))
 	float VolumeModulationMin;
 
 	/** The upper bound to use when randomly determining a volume multiplier */
-	UPROPERTY(EditAnywhere, BlueprintReadWrite, Category= "Randomization|Volume", meta = (DisplayName = "Volume (Max)"))
+	UPROPERTY(EditAnywhere, BlueprintReadWrite, Category = "Randomization|Volume", meta = (DisplayName = "Volume (Max)"))
 	float VolumeModulationMax;
 
 	/** A volume multiplier to apply to sounds generated by this component */
-	UPROPERTY(EditAnywhere, BlueprintReadWrite, Category=Sound)
+	UPROPERTY(EditAnywhere, BlueprintReadWrite, Category = Sound, meta = (DisplayAfter = "Sound"))
 	float VolumeMultiplier;
 
 	/** The attack time in milliseconds for the envelope follower. Delegate callbacks can be registered to get the 
-	 *  envelope value of sounds played with this audio component. Only used in audio mixer. 
+	 *  envelope value of sounds played with this audio component.
 	 */
-	UPROPERTY(EditAnywhere, BlueprintReadWrite, Category = Sound, meta = (ClampMin = "0", UIMin = "0"))
+	UPROPERTY(EditAnywhere, BlueprintReadWrite, Category = Analysis, meta = (ClampMin = "0", UIMin = "0"))
 	int32 EnvelopeFollowerAttackTime;
 
 	/** The release time in milliseconds for the envelope follower. Delegate callbacks can be registered to get the
-	 *  envelope value of sounds played with this audio component. Only used in audio mixer. 
+	 *  envelope value of sounds played with this audio component.
 	 */
-	UPROPERTY(EditAnywhere, BlueprintReadWrite, Category = Sound, meta = (ClampMin = "0", UIMin = "0"))
+	UPROPERTY(EditAnywhere, BlueprintReadWrite, Category = Analysis, meta = (ClampMin = "0", UIMin = "0"))
 	int32 EnvelopeFollowerReleaseTime;
 
-	/** A priority value that is used for sounds that play on this component that scales against final output volume. */
-	UPROPERTY(EditAnywhere, BlueprintReadWrite, Category = Sound, meta = (ClampMin = "0.0", UIMin = "0.0", EditCondition = "bOverridePriority"))
+	/** If enabled, overrides the priority of the selected sound with the value provided. */
+	UPROPERTY(EditAnywhere, BlueprintReadWrite, Category = Sound, meta = (ClampMin = "0.0", UIMin = "0.0", EditCondition = "bOverridePriority", DisplayAfter = "LowPassFilterFrequency"))
 	float Priority;
 
 	/** Used by the subtitle manager to prioritize subtitles wave instances spawned by this component. */
-	UPROPERTY(EditAnywhere, BlueprintReadWrite, Category = Sound, meta = (ClampMin = "0.0", UIMin = "0.0", EditCondition = "bOverrideSubtitlePriority"))
+	UPROPERTY(EditAnywhere, BlueprintReadWrite, Category = Subtitles, meta = (ClampMin = "0.0", UIMin = "0.0", EditCondition = "bOverrideSubtitlePriority", DisplayAfter = "SuppressSubtitles"))
 	float SubtitlePriority;
 
 	/** The chain of Source Effects to apply to the sounds playing on the Audio Component */
-	UPROPERTY(EditAnywhere, BlueprintReadWrite, Category = Sound)
+	UPROPERTY(EditAnywhere, BlueprintReadWrite, Category = Sound, meta = (DisplayAfter = "CanPlayMultipleInstances"))
 	TObjectPtr<USoundEffectSourcePresetChain> SourceEffectChain;
 
 #if WITH_EDITORONLY_DATA
@@ -324,22 +333,22 @@ public:
 #endif
 
 	/** A pitch multiplier to apply to sounds generated by this component */
-	UPROPERTY(EditAnywhere, BlueprintReadWrite, Category=Sound)
+	UPROPERTY(EditAnywhere, BlueprintReadWrite, Category = Sound, meta = (DisplayAfter = "VolumeMultiplier"))
 	float PitchMultiplier;
 
-	/** The frequency of the Lowpass Filter (in Hz) to apply to this voice. A frequency of 0.0 is the device sample rate and will bypass the filter. */
-	UPROPERTY(EditAnywhere, BlueprintReadWrite, Category = Filter, meta = (ClampMin = "0.0", UIMin = "0.0", EditCondition = "bEnableLowPassFilter"))
+	/** If enabled, the frequency of the Lowpass Filter (in Hz) to apply to this voice. A frequency of 0.0 is the device sample rate and will bypass the filter. */
+	UPROPERTY(EditAnywhere, BlueprintReadWrite, Category = Sound, meta = (ClampMin = "0.0", UIMin = "0.0", EditCondition = "bEnableLowPassFilter", DisplayAfter = "PitchMultiplier"))
 	float LowPassFilterFrequency;
 
 	/** A count of how many times we've started playing */
 	int32 ActiveCount;
 
 	/** If bOverrideSettings is false, the asset to use to determine attenuation properties for sounds generated by this component */
-	UPROPERTY(EditAnywhere, BlueprintReadWrite, Category=Attenuation, meta=(EditCondition="!bOverrideAttenuation"))
+	UPROPERTY(EditAnywhere, BlueprintReadWrite, Category = Attenuation, meta = (EditCondition = "!bOverrideAttenuation", DisplayAfter = "bOverrideAttenuation", EditConditionHides))
 	TObjectPtr<class USoundAttenuation> AttenuationSettings;
 
 	/** If bOverrideSettings is true, the attenuation properties to use for sounds generated by this component */
-	UPROPERTY(EditAnywhere, BlueprintReadWrite, Category= Attenuation, meta=(EditCondition="bOverrideAttenuation"))
+	UPROPERTY(EditAnywhere, BlueprintReadWrite, Category = Attenuation, meta = (EditCondition = "bOverrideAttenuation", DisplayAfter = "bOverrideAttenuation", EditConditionHides))
 	struct FSoundAttenuationSettings AttenuationOverrides;
 
 	/** What sound concurrency to use for sounds generated by this audio component */
@@ -383,7 +392,7 @@ public:
 	UPROPERTY(EditAnywhere, BlueprintReadWrite, Category = Modulation)
 	FSoundModulationDefaultRoutingSettings ModulationRouting;
 
-	/** This function returns the Targeted Audio Component’s current Play State.
+	/** This function returns the Targeted Audio Component's current Play State.
 	  * Playing, if the sound is currently playing.
 	  * Stopped, if the sound is stopped.
 	  * Paused, if the sound is currently playing, but paused.
@@ -438,7 +447,7 @@ public:
 
 	// Set what sound is played by this component
 	UFUNCTION(BlueprintCallable, Category="Audio|Components|Audio")
-	void SetSound( USoundBase* NewSound );
+	void SetSound(USoundBase* NewSound);
 
 	/**
 	 * This function allows designers to call Play on an Audio Component instance while applying a volume curve over time. 
@@ -462,14 +471,14 @@ public:
 	UFUNCTION(BlueprintCallable, Category="Audio|Components|Audio", meta = (AdvancedDisplay = 1))
 	virtual	void FadeOut(float FadeOutDuration, float FadeVolumeLevel, const EAudioFaderCurve FadeCurve = EAudioFaderCurve::Linear);
 
-	/** Begins playing the targeted Audio Component’s sound at the designated Start Time, seeking into a sound. 
+	/** Begins playing the targeted Audio Component's sound at the designated Start Time, seeking into a sound.
 	 * @param StartTime The offset, in seconds, to begin reading the sound at
 	 */
 	UFUNCTION(BlueprintCallable, Category="Audio|Components|Audio")
 	virtual void Play(float StartTime = 0.0f);
 
 	/** Start a sound playing on an audio component on a given quantization boundary with the handle to an existing clock */
-	UFUNCTION(BlueprintCallable, Category = "Audio|Components|Audio", meta=(WorldContext = "WorldContextObject", AdvancedDisplay = "3", UnsafeDuringActorConstruction = "true", Keywords = "play", AutoCreateRefTerm = "InDelegate"))
+	UFUNCTION(BlueprintCallable, Category = "Audio|Components|Audio", meta = (WorldContext = "WorldContextObject", AdvancedDisplay = "3", UnsafeDuringActorConstruction = "true", Keywords = "play", AutoCreateRefTerm = "InDelegate"))
 	virtual void PlayQuantized(
 		  const UObject* WorldContextObject
 		, UPARAM(ref) UQuartzClockHandle*& InClockHandle
@@ -493,7 +502,7 @@ public:
 	UFUNCTION(BlueprintCallable, Category = "Audio|Components|Audio")
 	void SetPaused(bool bPause);
 
-	/** Returns TRUE if the targeted Audio Component’s sound is playing. 
+	/** Returns TRUE if the targeted Audio Component’s sound is playing.
 	 *  Doesn't indicate if the sound is paused or fading in/out. Use GetPlayState() to get the full play state.
 	 */
 	UFUNCTION(BlueprintCallable, Category="Audio|Components|Audio")
@@ -639,6 +648,8 @@ public:
 
 	static void PlaybackCompleted(uint64 AudioComponentID, bool bFailedToStart);
 
+	bool GetDisableParameterUpdatesWhilePlaying() const override { return static_cast<bool>(bDisableParameterUpdatesWhilePlaying); }
+
 private:
 	/** Called by the ActiveSound to inform the component that playback is finished */
 	void PlaybackCompleted(bool bFailedToStart);
@@ -712,7 +723,7 @@ public:
 	/** SoundGeneratorParameterInterface Implementation */
 	FAudioDevice* GetAudioDevice() const override;
 	TArray<FAudioParameter>& GetInstanceParameters() override { return InstanceParameters; }
-	uint64 GetInstanceID() const override { return AudioComponentID; }
+	uint64 GetInstanceOwnerID() const override { return AudioComponentID; }
 	USoundBase* GetSound() override { return Sound; }
 
 public:
@@ -790,7 +801,7 @@ private:
 protected:
 
 	/** Utility function called by Play and FadeIn to start a sound playing. */
-	void PlayInternal(const PlayInternalRequestData& InPlayRequestData);
+	void PlayInternal(const PlayInternalRequestData& InPlayRequestData, USoundBase * InSound = nullptr);
 
 #if WITH_EDITORONLY_DATA
 	/** Utility function that updates which texture is displayed on the sprite dependent on the properties of the Audio Component. */
diff --git a/Engine/Source/Runtime/Engine/Private/Audio/SoundGeneratorParameterInterface.cpp b/Engine/Source/Runtime/Engine/Private/Audio/SoundGeneratorParameterInterface.cpp
--- a/Engine/Source/Runtime/Engine/Private/Audio/SoundGeneratorParameterInterface.cpp
+++ b/Engine/Source/Runtime/Engine/Private/Audio/SoundGeneratorParameterInterface.cpp
@@ -5,7 +5,6 @@
 #include "Audio.h"
 #include "AudioDevice.h"
 #include "AudioThread.h"
-#include "Components/AudioComponent.h"
 #include "IAudioExtensionPlugin.h"
 
 
@@ -23,16 +22,17 @@ void ISoundGeneratorParameterInterface::ResetParameters()
 {
 	if (FAudioDevice* AudioDevice = GetAudioDevice())
 	{
-		FAudioThread::RunCommandOnAudioThread([AudioDevice, InstanceID = GetInstanceID()]() mutable
+		if (IsPlaying() && !GetDisableParameterUpdatesWhilePlaying())
 		{
-			if (FActiveSound* ActiveSound = AudioDevice->FindActiveSound(InstanceID))
+			DECLARE_CYCLE_STAT(TEXT("FAudioThreadTask.SoundGenerator.ResetParameters"), STAT_AudioResetParameters, STATGROUP_AudioThreadCommands);
+			AudioDevice->SendCommandToActiveSounds(GetInstanceOwnerID(), [] (FActiveSound& ActiveSound)
 			{
-				if (Audio::IParameterTransmitter* Transmitter = ActiveSound->GetTransmitter())
+				if (Audio::IParameterTransmitter* Transmitter = ActiveSound.GetTransmitter())
 				{
 					Transmitter->Reset();
 				}
-			}
-		});
+			}, GET_STATID(STAT_AudioResetParameters));
+		}
 	}
 }
 
@@ -99,7 +99,7 @@ void ISoundGeneratorParameterInterface::SetParameter(FAudioParameter&& InValue)
 
 void ISoundGeneratorParameterInterface::SetParameters(TArray<FAudioParameter>&& InValues)
 {
-	for (FAudioParameter& Value : InValues)
+	for (const FAudioParameter& Value : InValues)
 	{
 		TArray<FAudioParameter>& InstanceParameters = GetInstanceParameters();
 		if (FAudioParameter* CurrentParam = FAudioParameter::FindOrAddParam(InstanceParameters, Value.ParamName))
@@ -108,7 +108,7 @@ void ISoundGeneratorParameterInterface::SetParameters(TArray<FAudioParameter>&&
 		}
 	}
 
-	if (IsPlaying())
+	if (IsPlaying() && !GetDisableParameterUpdatesWhilePlaying())
 	{
 		if (FAudioDevice* AudioDevice = GetAudioDevice())
 		{
@@ -121,17 +121,18 @@ void ISoundGeneratorParameterInterface::SetParameters(TArray<FAudioParameter>&&
 			ParamsToSet = MoveTemp(InValues);
 
 			DECLARE_CYCLE_STAT(TEXT("FAudioThreadTask.SoundGenerator.SetParameters"), STAT_AudioSetParameters, STATGROUP_AudioThreadCommands);
-
-			FAudioThread::RunCommandOnAudioThread([AudioDevice, InstanceID = GetInstanceID(), Params { MoveTemp(ParamsToSet) }]() mutable
+			AudioDevice->SendCommandToActiveSounds(GetInstanceOwnerID(), [AudioDevice, Params = MoveTemp(ParamsToSet)](FActiveSound& ActiveSound)
 			{
-				if (FActiveSound* ActiveSound = AudioDevice->FindActiveSound(InstanceID))
+				if (Audio::IParameterTransmitter* Transmitter = ActiveSound.GetTransmitter())
 				{
-					if (Audio::IParameterTransmitter* Transmitter = ActiveSound->GetTransmitter())
+					for (const FAudioParameter& Param : Params)
 					{
-						for (FAudioParameter& Param : Params)
+						const FName ParamName = Param.ParamName;
+						if (!ParamName.IsNone())
 						{
-							const FName ParamName = Param.ParamName;
-							if (!Transmitter->SetParameter(MoveTemp(Param)))
+							// Must be copied as original version must be preserved in case command is called on multiple ActiveSounds.
+							FAudioParameter TempParam = Param;
+							if (!Transmitter->SetParameter(MoveTemp(TempParam)))
 							{
 								UE_LOG(LogAudio, Warning, TEXT("Failed to set parameter '%s'"), *ParamName.ToString());
 							}
@@ -145,7 +146,7 @@ void ISoundGeneratorParameterInterface::SetParameters(TArray<FAudioParameter>&&
 
 void ISoundGeneratorParameterInterface::SetParameterInternal(FAudioParameter&& InParam)
 {
-	if (!InParam.ParamName.IsValid())
+	if (InParam.ParamName.IsNone())
 	{
 		return;
 	}
@@ -156,7 +157,7 @@ void ISoundGeneratorParameterInterface::SetParameterInternal(FAudioParameter&& I
 		CurrentParam->Merge(InParam, false /* bInTakeName */);
 	}
 
-	if (IsPlaying())
+	if (IsPlaying() && !GetDisableParameterUpdatesWhilePlaying())
 	{
 		if (FAudioDevice* AudioDevice = GetAudioDevice())
 		{
@@ -174,17 +175,17 @@ void ISoundGeneratorParameterInterface::SetParameterInternal(FAudioParameter&& I
 
 			DECLARE_CYCLE_STAT(TEXT("FAudioThreadTask.SoundGenerator.SetParameter"), STAT_AudioSetParameter, STATGROUP_AudioThreadCommands);
 
-			FAudioThread::RunCommandOnAudioThread([AudioDevice, InstanceID = GetInstanceID(), Param { MoveTemp(ParamToSet) }]() mutable
+			AudioDevice->SendCommandToActiveSounds(GetInstanceOwnerID(), [AudioDevice, Param = MoveTemp(ParamToSet)](FActiveSound& ActiveSound)
 			{
-				if (FActiveSound* ActiveSound = AudioDevice->FindActiveSound(InstanceID))
+				if (Audio::IParameterTransmitter* Transmitter = ActiveSound.GetTransmitter())
 				{
-					if (Audio::IParameterTransmitter* Transmitter = ActiveSound->GetTransmitter())
+					const FName ParamName = Param.ParamName;
+
+					// Must be copied as original version must be preserved in case command is called on multiple ActiveSounds.
+					FAudioParameter TempParam = Param;
+					if (!Transmitter->SetParameter(MoveTemp(TempParam)))
 					{
-						const FName ParamName = Param.ParamName;
-						if (!Transmitter->SetParameter(MoveTemp(Param)))
-						{
-							UE_LOG(LogAudio, Warning, TEXT("Failed to set parameter '%s'"), *ParamName.ToString());
-						}
+						UE_LOG(LogAudio, Warning, TEXT("Failed to set parameter '%s'"), *ParamName.ToString());
 					}
 				}
 			}, GET_STATID(STAT_AudioSetParameter));
diff --git a/Engine/Source/Runtime/Engine/Private/AudioDevice.cpp b/Engine/Source/Runtime/Engine/Private/AudioDevice.cpp
--- a/Engine/Source/Runtime/Engine/Private/AudioDevice.cpp
+++ b/Engine/Source/Runtime/Engine/Private/AudioDevice.cpp
@@ -1117,7 +1117,7 @@ void FAudioDevice::ShowSoundClassHierarchy(FOutputDevice& Ar, USoundClass* InSou
 		{
 			Ar.Logf(TEXT("%s"), *SoundClass->GetName());
 		}
-		for (int i = 0; i < SoundClass->ChildClasses.Num(); ++i)
+		for (int32 i = 0; i < SoundClass->ChildClasses.Num(); ++i)
 		{
 			if (SoundClass->ChildClasses[i])
 			{
@@ -4885,7 +4885,8 @@ void FAudioDevice::AddNewActiveSoundInternal(const FActiveSound& InNewActiveSoun
 
 	if (ActiveSound->GetAudioComponentID() > 0)
 	{
-		AudioComponentIDToActiveSoundMap.Add(ActiveSound->GetAudioComponentID(), ActiveSound);
+		TArray<FActiveSound*>& ActiveSoundArray = AudioComponentIDToActiveSoundMap.FindOrAdd(ActiveSound->GetAudioComponentID());
+		ActiveSoundArray.Add(ActiveSound);
 	}
 }
 
@@ -4957,15 +4958,27 @@ void FAudioDevice::AddVirtualLoop(const FAudioVirtualLoop& InVirtualLoop)
 	const int64 ComponentID = ActiveSound.GetAudioComponentID();
 	if (ComponentID > 0)
 	{
-		if (FActiveSound* ExistingSound = AudioComponentIDToActiveSoundMap.FindRef(ComponentID))
+		if (TArray<FActiveSound*>* ExistingSounds = AudioComponentIDToActiveSoundMap.Find(ComponentID))
 		{
-			UE_LOG(LogAudio, Warning, TEXT("Adding ComponentID for Sound '%s' when map already contains ID for Sound '%s'."),
-				ActiveSound.Sound ? *ActiveSound.Sound->GetName() : TEXT("N/A"),
-				ExistingSound->Sound ? *ExistingSound->Sound->GetName() : TEXT("N/A")
-				);
-			AudioComponentIDToActiveSoundMap.Remove(ComponentID);
+			if (!CanHaveMultipleActiveSounds(ComponentID))
+			{
+				// if we can't have more than one sound, we should have only one element (or less) here
+				ensure(ExistingSounds->Num() <= 1);
+				for (int32 i = ExistingSounds->Num() - 1; i >= 0; --i)
+				{
+					const FActiveSound* ExistingSound = (*ExistingSounds)[i];
+					if (ensure(ExistingSound))
+					{
+						UE_LOG(LogAudio, Warning, TEXT("Attempting to add Sound '%s' to ComponentID when map already contains ID for Sound '%s', and ID can only reference one active sound."),
+							ActiveSound.Sound ? *ActiveSound.Sound->GetName() : TEXT("N/A"),
+							ExistingSound->Sound ? *ExistingSound->Sound->GetName() : TEXT("N/A")
+						);
+						ExistingSounds->RemoveAtSwap(i, 1, false /* bAllowShrinking */);
+					}
+				}
+			}	
+			ExistingSounds->AddUnique(&ActiveSound);
 		}
-		AudioComponentIDToActiveSoundMap.Add(ComponentID, &ActiveSound);
 	}
 
 #if !(UE_BUILD_SHIPPING || UE_BUILD_TEST)
@@ -5175,11 +5188,21 @@ void FAudioDevice::StopActiveSound(const uint64 AudioComponentID)
 		return;
 	}
 
-	FActiveSound* ActiveSound = FindActiveSound(AudioComponentID);
-	if (ActiveSound)
+	const Audio::FDeviceId AudioDeviceID = DeviceID;
+
+	SendCommandToActiveSounds(AudioComponentID, [AudioDeviceID](FActiveSound& ActiveSound)
 	{
-		AddSoundToStop(ActiveSound);
-	}
+		FAudioDeviceManager* AudioDeviceManager = GEngine->GetAudioDeviceManager();
+			
+		if (AudioDeviceManager)
+		{
+			FAudioDeviceHandle AudioDevice = AudioDeviceManager->GetAudioDevice(AudioDeviceID);
+			if (AudioDevice.IsValid())
+			{
+				AudioDevice->AddSoundToStop(&ActiveSound);
+			}
+		}
+	});
 }
 
 void FAudioDevice::StopActiveSound(FActiveSound* ActiveSound)
@@ -5191,11 +5214,11 @@ void FAudioDevice::StopActiveSound(FActiveSound* ActiveSound)
 void FAudioDevice::PauseActiveSound(const uint64 AudioComponentID, const bool bInIsPaused)
 {
 	check(IsInAudioThread());
-	FActiveSound* ActiveSound = FindActiveSound(AudioComponentID);
-	if (ActiveSound)
+	
+	SendCommandToActiveSounds(AudioComponentID, [bInIsPaused](FActiveSound& ActiveSound)
 	{
-		ActiveSound->bIsPaused = bInIsPaused;
-	}
+		ActiveSound.bIsPaused = bInIsPaused;
+	});
 }
 
 void FAudioDevice::NotifyActiveSoundOcclusionTraceDone(FActiveSound* InActiveSound, bool bIsOccluded)
@@ -5210,10 +5233,69 @@ void FAudioDevice::NotifyActiveSoundOcclusionTraceDone(FActiveSound* InActiveSou
 
 FActiveSound* FAudioDevice::FindActiveSound(const uint64 AudioComponentID)
 {
-	check(IsInAudioThread());
+	ensure(IsInAudioThread());
+
+	if (TArray<FActiveSound*>* ActiveSoundsInComponent = AudioComponentIDToActiveSoundMap.Find(AudioComponentID))
+	{
+		// return the first active sound corresponding to this audio component, if it exists, or a nullptr if not
+		if (!ActiveSoundsInComponent->IsEmpty())
+		{
+			return (*ActiveSoundsInComponent)[0];
+		}
+	}
+
+	return nullptr;
+}
+
+void FAudioDevice::SendCommandToActiveSounds(uint64 InAudioComponentID, TUniqueFunction<void(FActiveSound&)> InFunc, const TStatId InStatId)
+{
+	if (!IsInAudioThread())
+	{
+		FAudioThread::RunCommandOnAudioThread([this, InAudioComponentID, Func = MoveTemp(InFunc)]() mutable
+		{
+			SendCommandToActiveSounds(InAudioComponentID, MoveTemp(Func));
+		}, InStatId);
+		return;
+	}
+
+	// Must cache active sounds as the AudioComponentIDToActiveSoundMap can potentially be modified by the command function provided
+	// (ex. when stopping a sound or array of sounds).
+	const TArray<FActiveSound*> CachedActiveSounds = AudioComponentIDToActiveSoundMap.FindRef(InAudioComponentID);
+	for (FActiveSound* ActiveSound : CachedActiveSounds)
+	{
+		// This should never be null
+		if (ensure(ActiveSound))
+		{
+			InFunc(*ActiveSound);
+		}
+	}
+}
 
-	// find the active sound corresponding to this audio component
-	return AudioComponentIDToActiveSoundMap.FindRef(AudioComponentID);
+bool FAudioDevice::CanHaveMultipleActiveSounds(uint64 AudioComponentID) const
+{
+	ensure(IsInAudioThread());
+
+	if (const bool* bCanHaveMultipleActiveSounds = AudioComponentIDToCanHaveMultipleActiveSoundsMap.Find(AudioComponentID))
+	{
+		return *bCanHaveMultipleActiveSounds;
+	}
+
+	return false;
+}
+
+void FAudioDevice::SetCanHaveMultipleActiveSounds(uint64 InAudioComponentID, bool InCanHaveMultipleActiveSounds)
+{
+	if (!IsInAudioThread())
+	{
+		FAudioThread::RunCommandOnAudioThread([this, ComponentID = InAudioComponentID, bNewValue = InCanHaveMultipleActiveSounds]()
+		{
+			SetCanHaveMultipleActiveSounds(ComponentID, bNewValue);
+		});
+
+		return;
+	}
+
+	AudioComponentIDToCanHaveMultipleActiveSoundsMap.FindOrAdd(InAudioComponentID, InCanHaveMultipleActiveSounds) = InCanHaveMultipleActiveSounds;
 }
 
 void FAudioDevice::RemoveActiveSound(FActiveSound* ActiveSound)
@@ -5615,7 +5697,29 @@ void FAudioDevice::UnlinkActiveSoundFromComponent(const FActiveSound& InActiveSo
 	const uint64 AudioComponentID = InActiveSound.GetAudioComponentID();
 	if (AudioComponentID > 0)
 	{
-		AudioComponentIDToActiveSoundMap.Remove(AudioComponentID);
+		if (TArray<FActiveSound*>* ActiveSoundsInComponent = AudioComponentIDToActiveSoundMap.Find(AudioComponentID))
+		{
+			for (int32 i = ActiveSoundsInComponent->Num() - 1; i >= 0; --i)
+			{
+				const FActiveSound* ActiveSound = (*ActiveSoundsInComponent)[i];
+				
+				if (ensure(ActiveSound))
+				{
+					if (ActiveSound->GetInstanceID() == InActiveSound.GetInstanceID())
+					{
+						ActiveSoundsInComponent->RemoveAtSwap(i, 1, false);
+						break;
+					}
+				}
+			}
+			
+			if (ActiveSoundsInComponent->IsEmpty())
+			{
+				AudioComponentIDToActiveSoundMap.Remove(AudioComponentID);
+				AudioComponentIDToCanHaveMultipleActiveSoundsMap.Remove(AudioComponentID);
+			}
+		}
+
 	}
 }
 
diff --git a/Engine/Source/Runtime/Engine/Private/Components/AudioComponent.cpp b/Engine/Source/Runtime/Engine/Private/Components/AudioComponent.cpp
--- a/Engine/Source/Runtime/Engine/Private/Components/AudioComponent.cpp
+++ b/Engine/Source/Runtime/Engine/Private/Components/AudioComponent.cpp
@@ -61,6 +61,8 @@ UAudioComponent::UAudioComponent(const FObjectInitializer& ObjectInitializer)
 	VolumeMultiplier = 1.f;
 	bOverridePriority = false;
 	bOverrideSubtitlePriority = false;
+	bCanPlayMultipleInstances = false;
+	bDisableParameterUpdatesWhilePlaying = false;
 	bIsPreviewSound = false;
 	bIsPaused = false;
 
@@ -270,7 +272,11 @@ void UAudioComponent::SetSound(USoundBase* NewSound)
 	// If this is an auto destroy component we need to prevent it from being auto-destroyed since we're really just restarting it
 	const bool bWasAutoDestroy = bAutoDestroy;
 	bAutoDestroy = false;
-	Stop();
+	//Only stop the existing sound if we are limited to one sound per component
+	if (!bCanPlayMultipleInstances)
+	{
+		Stop();
+	}
 	bAutoDestroy = bWasAutoDestroy;
 
 	Sound = NewSound;
@@ -280,7 +286,7 @@ void UAudioComponent::SetSound(USoundBase* NewSound)
 		UGameplayStatics::PrimeSound(Sound);
 	}
 
-	if (bPlay)
+	if (bPlay && !bCanPlayMultipleInstances)
 	{
 		Play();
 	}
@@ -306,18 +312,10 @@ void UAudioComponent::OnUpdateTransform(EUpdateTransformFlags UpdateTransformFla
 		{
 			DECLARE_CYCLE_STAT(TEXT("FAudioThreadTask.UpdateAudioComponentTransform"), STAT_AudioUpdateComponentTransform, STATGROUP_AudioThreadCommands);
 
-			const uint64 MyAudioComponentID = AudioComponentID;
-			const FTransform& MyTransform = GetComponentTransform();
-
-			FAudioThread::RunCommandOnAudioThread([AudioDevice, MyAudioComponentID, MyTransform]()
+			AudioDevice->SendCommandToActiveSounds(AudioComponentID, [NewTransform = GetComponentTransform()](FActiveSound& ActiveSound)
 			{
-				FActiveSound* ActiveSound = AudioDevice->FindActiveSound(MyAudioComponentID);
-				if (ActiveSound)
-				{
-					ActiveSound->Transform = MyTransform;
-
-				}
-			}, GET_STATID(STAT_AudioUpdateComponentTransform));
+				ActiveSound.Transform = NewTransform;
+			});
 		}
 	}
 };
@@ -433,7 +431,7 @@ void UAudioComponent::PlayQuantized(
 	PlayInternal(Data);
 }
 
-void UAudioComponent::PlayInternal(const PlayInternalRequestData& InPlayRequestData)
+void UAudioComponent::PlayInternal(const PlayInternalRequestData& InPlayRequestData, USoundBase* InSound)
 {
 	SCOPE_CYCLE_COUNTER(STAT_AudioComp_Play);
 
@@ -449,7 +447,11 @@ void UAudioComponent::PlayInternal(const PlayInternalRequestData& InPlayRequestD
 		// If this is an auto destroy component we need to prevent it from being auto-destroyed since we're really just restarting it
 		bool bCurrentAutoDestroy = bAutoDestroy;
 		bAutoDestroy = false;
-		Stop();
+		//Only stop if this component is limited to one active sound
+		if (!bCanPlayMultipleInstances)
+		{
+			Stop();
+		}
 		bAutoDestroy = bCurrentAutoDestroy;
 	}
 
@@ -507,7 +509,17 @@ void UAudioComponent::PlayInternal(const PlayInternalRequestData& InPlayRequestD
 			FActiveSound NewActiveSound;
 			NewActiveSound.SetAudioComponent(*this);
 			NewActiveSound.SetWorld(GetWorld());
-			NewActiveSound.SetSound(Sound);
+
+			//If the InSound pointer is non-null, use its data. Otherwise, use the base Sound
+			if (InSound == nullptr)
+			{
+				NewActiveSound.SetSound(Sound);
+			}
+			else
+			{
+				NewActiveSound.SetSound(InSound);
+			}
+
 			NewActiveSound.SetSourceEffectChain(SourceEffectChain);
 			NewActiveSound.SetSoundClass(SoundClassOverride);
 			NewActiveSound.ConcurrencySet = ConcurrencySet;
@@ -603,6 +615,9 @@ void UAudioComponent::PlayInternal(const PlayInternalRequestData& InPlayRequestD
 			// Bump ActiveCount... this is used to determine if an audio component is still active after a sound reports back as completed
 			++ActiveCount;
 
+			// Pass along whether or not component is setup to support multiple active sounds. This is to ensure virtualization will function accordingly.
+			AudioDevice->SetCanHaveMultipleActiveSounds(AudioComponentID, bCanPlayMultipleInstances);
+
 			AudioDevice->AddNewActiveSound(NewActiveSound, &InstanceParameters);
 
 			// In editor, the audio thread is not run separate from the game thread, and can result in calling PlaybackComplete prior
@@ -689,17 +704,10 @@ void UAudioComponent::AdjustVolumeInternal(float AdjustVolumeDuration, float Adj
 		BroadcastPlayState();
 	}
 
-	const uint64 InAudioComponentID = AudioComponentID;
 	DECLARE_CYCLE_STAT(TEXT("FAudioThreadTask.AdjustVolume"), STAT_AudioAdjustVolume, STATGROUP_AudioThreadCommands);
-	FAudioThread::RunCommandOnAudioThread([AudioDevice, InAudioComponentID, AdjustVolumeDuration, AdjustVolumeLevel, bInIsFadeOut, FadeCurve]()
+	AudioDevice->SendCommandToActiveSounds(AudioComponentID, [bInIsFadeOut, AdjustVolumeLevel, AdjustVolumeDuration, FadeCurve](FActiveSound& ActiveSound)
 	{
-		FActiveSound* ActiveSound = AudioDevice->FindActiveSound(InAudioComponentID);
-		if (!ActiveSound)
-		{
-			return;
-		}
-
-		Audio::FVolumeFader& Fader = ActiveSound->ComponentVolumeFader;
+		Audio::FVolumeFader& Fader = ActiveSound.ComponentVolumeFader;
 		const float InitialTargetVolume = Fader.GetTargetVolume();
 
 		// Ignore fade out request if requested volume is higher than current target.
@@ -709,7 +717,7 @@ void UAudioComponent::AdjustVolumeInternal(float AdjustVolumeDuration, float Adj
 		}
 
 		const bool ToZeroVolume = FMath::IsNearlyZero(AdjustVolumeLevel);
-		if (ActiveSound->FadeOut == FActiveSound::EFadeOut::Concurrency)
+		if (ActiveSound.FadeOut == FActiveSound::EFadeOut::Concurrency)
 		{
 			// Ignore adjust volume request if non-zero and currently voice stealing.
 			if (!FMath::IsNearlyZero(AdjustVolumeLevel))
@@ -725,7 +733,7 @@ void UAudioComponent::AdjustVolumeInternal(float AdjustVolumeDuration, float Adj
 		}
 		else
 		{
-			ActiveSound->FadeOut = bInIsFadeOut || ToZeroVolume ? FActiveSound::EFadeOut::User : FActiveSound::EFadeOut::None;
+			ActiveSound.FadeOut = bInIsFadeOut || ToZeroVolume ? FActiveSound::EFadeOut::User : FActiveSound::EFadeOut::None;
 		}
 
 		if (bInIsFadeOut || ToZeroVolume)
@@ -793,41 +801,34 @@ void UAudioComponent::StopDelayed(float DelayTime)
 	}
 
 	// 2. Performs delayed stop with no fade
-	const uint64 InAudioComponentID = AudioComponentID;
 	DECLARE_CYCLE_STAT(TEXT("FAudioThreadTask.StopDelayed"), STAT_AudioStopDelayed, STATGROUP_AudioThreadCommands);
-	FAudioThread::RunCommandOnAudioThread([AudioDevice, InAudioComponentID, DelayTime]()
+	AudioDevice->SendCommandToActiveSounds(AudioComponentID, [DelayTime](FActiveSound& ActiveSound)
 	{
-		FActiveSound* ActiveSound = AudioDevice->FindActiveSound(InAudioComponentID);
-		if (!ActiveSound)
-		{
-			return;
-		}
-
-		if (const USoundBase* StoppingSound = ActiveSound->GetSound())
+		if (const USoundBase* StoppingSound = ActiveSound.GetSound())
 		{
 			UE_LOG(LogAudio, Verbose, TEXT("%g: Delayed Stop requested for sound '%s'"),
-				ActiveSound->GetWorld() ? ActiveSound->GetWorld()->GetAudioTimeSeconds() : 0.0f,
+				ActiveSound.GetWorld() ? ActiveSound.GetWorld()->GetAudioTimeSeconds() : 0.0f,
 				*StoppingSound->GetName());
 		}
 
-		Audio::FVolumeFader& Fader = ActiveSound->ComponentVolumeFader;
-		switch (ActiveSound->FadeOut)
+		Audio::FVolumeFader& Fader = ActiveSound.ComponentVolumeFader;
+		switch (ActiveSound.FadeOut)
 		{
 			case FActiveSound::EFadeOut::Concurrency:
-			{
-				// Ignore request of longer fade out than active target if active is concurrency (voice stealing) fade.
-				if (DelayTime < Fader.GetFadeDuration())
 				{
-					Fader.SetActiveDuration(DelayTime);
+					// Ignore request of longer fade out than active target if active is concurrency (voice stealing) fade.
+					if (DelayTime < Fader.GetFadeDuration())
+					{
+						Fader.SetActiveDuration(DelayTime);
+					}
 				}
-			}
-			break;
-			
+				break;
+
 			case FActiveSound::EFadeOut::User:
 			case FActiveSound::EFadeOut::None:
 			default:
 			{
-				ActiveSound->FadeOut = FActiveSound::EFadeOut::User;
+				ActiveSound.FadeOut = FActiveSound::EFadeOut::User;
 				Fader.SetActiveDuration(DelayTime);
 			}
 			break;
@@ -985,6 +986,11 @@ void UAudioComponent::PostEditChangeProperty(FPropertyChangedEvent& PropertyChan
 		Play();
 	}
 
+	if (PropertyChangedEvent.GetPropertyName() == GET_MEMBER_NAME_CHECKED(UAudioComponent, bCanPlayMultipleInstances))
+	{
+		GetAudioDevice()->SetCanHaveMultipleActiveSounds(AudioComponentID, bCanPlayMultipleInstances);
+	}
+
 #if WITH_EDITORONLY_DATA
 	UpdateSpriteTexture();
 #endif
@@ -1022,7 +1028,7 @@ bool UAudioComponent::BP_GetAttenuationSettingsToApply(FSoundAttenuationSettings
 
 void UAudioComponent::CollectAttenuationShapesForVisualization(TMultiMap<EAttenuationShape::Type, FBaseAttenuationSettings::AttenuationShapeDetails>& ShapeDetailsMap) const
 {
-	const FSoundAttenuationSettings *AttenuationSettingsToApply = GetAttenuationSettingsToApply();
+	const FSoundAttenuationSettings* AttenuationSettingsToApply = GetAttenuationSettingsToApply();
 
 	if (AttenuationSettingsToApply)
 	{
@@ -1034,7 +1040,7 @@ void UAudioComponent::CollectAttenuationShapesForVisualization(TMultiMap<EAttenu
 	if (SoundCue)
 	{
 		TArray<USoundNodeAttenuation*> AttenuationNodes;
-		SoundCue->RecursiveFindAttenuation( SoundCue->FirstNode, AttenuationNodes );
+		SoundCue->RecursiveFindAttenuation(SoundCue->FirstNode, AttenuationNodes);
 		for (int32 NodeIndex = 0; NodeIndex < AttenuationNodes.Num(); ++NodeIndex)
 		{
 			AttenuationSettingsToApply = AttenuationNodes[NodeIndex]->GetAttenuationSettingsToApply();
@@ -1113,15 +1119,9 @@ void UAudioComponent::SetVolumeMultiplier(const float NewVolumeMultiplier)
 		if (FAudioDevice* AudioDevice = GetAudioDevice())
 		{
 			DECLARE_CYCLE_STAT(TEXT("FAudioThreadTask.SetVolumeMultiplier"), STAT_AudioSetVolumeMultiplier, STATGROUP_AudioThreadCommands);
-
-			const uint64 MyAudioComponentID = AudioComponentID;
-			FAudioThread::RunCommandOnAudioThread([AudioDevice, MyAudioComponentID, NewVolumeMultiplier]()
+			AudioDevice->SendCommandToActiveSounds(AudioComponentID, [NewVolumeMultiplier](FActiveSound& ActiveSound)
 			{
-				FActiveSound* ActiveSound = AudioDevice->FindActiveSound(MyAudioComponentID);
-				if (ActiveSound)
-				{
-					ActiveSound->SetVolume(NewVolumeMultiplier);
-				}
+				ActiveSound.SetVolume(NewVolumeMultiplier);
 			}, GET_STATID(STAT_AudioSetVolumeMultiplier));
 		}
 	}
@@ -1137,15 +1137,9 @@ void UAudioComponent::SetPitchMultiplier(const float NewPitchMultiplier)
 		if (FAudioDevice* AudioDevice = GetAudioDevice())
 		{
 			DECLARE_CYCLE_STAT(TEXT("FAudioThreadTask.SetPitchMultiplier"), STAT_AudioSetPitchMultiplier, STATGROUP_AudioThreadCommands);
-
-			const uint64 MyAudioComponentID = AudioComponentID;
-			FAudioThread::RunCommandOnAudioThread([AudioDevice, MyAudioComponentID, NewPitchMultiplier]()
+			AudioDevice->SendCommandToActiveSounds(AudioComponentID, [NewPitchMultiplier](FActiveSound& ActiveSound)
 			{
-				FActiveSound* ActiveSound = AudioDevice->FindActiveSound(MyAudioComponentID);
-				if (ActiveSound)
-				{
-					ActiveSound->SetPitch(NewPitchMultiplier);
-				}
+				ActiveSound.SetPitch(NewPitchMultiplier);
 			}, GET_STATID(STAT_AudioSetPitchMultiplier));
 		}
 	}
@@ -1160,15 +1154,9 @@ void UAudioComponent::SetUISound(const bool bInIsUISound)
 		if (FAudioDevice* AudioDevice = GetAudioDevice())
 		{
 			DECLARE_CYCLE_STAT(TEXT("FAudioThreadTask.SetIsUISound"), STAT_AudioSetIsUISound, STATGROUP_AudioThreadCommands);
-
-			const uint64 MyAudioComponentID = AudioComponentID;
-			FAudioThread::RunCommandOnAudioThread([AudioDevice, MyAudioComponentID, bInIsUISound]()
+			AudioDevice->SendCommandToActiveSounds(AudioComponentID, [bInIsUISound](FActiveSound& ActiveSound)
 			{
-				FActiveSound* ActiveSound = AudioDevice->FindActiveSound(MyAudioComponentID);
-				if (ActiveSound)
-				{
-					ActiveSound->bIsUISound = bInIsUISound;
-				}
+				ActiveSound.bIsUISound = bInIsUISound;
 			}, GET_STATID(STAT_AudioSetIsUISound));
 		}
 	}
@@ -1184,15 +1172,9 @@ void UAudioComponent::AdjustAttenuation(const FSoundAttenuationSettings& InAtten
 		if (FAudioDevice* AudioDevice = GetAudioDevice())
 		{
 			DECLARE_CYCLE_STAT(TEXT("FAudioThreadTask.AdjustAttenuation"), STAT_AudioAdjustAttenuation, STATGROUP_AudioThreadCommands);
-
-			const uint64 MyAudioComponentID = AudioComponentID;
-			FAudioThread::RunCommandOnAudioThread([AudioDevice, MyAudioComponentID, InAttenuationSettings]()
+			AudioDevice->SendCommandToActiveSounds(AudioComponentID, [InAttenuationSettings](FActiveSound& ActiveSound)
 			{
-				FActiveSound* ActiveSound = AudioDevice->FindActiveSound(MyAudioComponentID);
-				if (ActiveSound)
-				{
-					ActiveSound->AttenuationSettings = InAttenuationSettings;
-				}
+				ActiveSound.AttenuationSettings = InAttenuationSettings;
 			}, GET_STATID(STAT_AudioAdjustAttenuation));
 		}
 	}
@@ -1202,19 +1184,14 @@ void UAudioComponent::SetSubmixSend(USoundSubmixBase* Submix, float SendLevel)
 {
 	if (FAudioDevice* AudioDevice = GetAudioDevice())
 	{
-		DECLARE_CYCLE_STAT(TEXT("FAudioThreadTask.AudioSetSubmixSend"), STAT_SetSubmixSend, STATGROUP_AudioThreadCommands);
+		FSoundSubmixSendInfo SendInfo;
+		SendInfo.SoundSubmix = Submix;
+		SendInfo.SendLevel = SendLevel;
 
-		const uint64 MyAudioComponentID = AudioComponentID;
-		FAudioThread::RunCommandOnAudioThread([AudioDevice, MyAudioComponentID, Submix, SendLevel]()
+		DECLARE_CYCLE_STAT(TEXT("FAudioThreadTask.AudioSetSubmixSend"), STAT_SetSubmixSend, STATGROUP_AudioThreadCommands);
+		AudioDevice->SendCommandToActiveSounds(AudioComponentID, [SendInfo](FActiveSound& ActiveSound)
 		{
-			FActiveSound* ActiveSound = AudioDevice->FindActiveSound(MyAudioComponentID);
-			if (ActiveSound)
-			{
-				FSoundSubmixSendInfo SendInfo;
-				SendInfo.SoundSubmix = Submix;
-				SendInfo.SendLevel = SendLevel;
-				ActiveSound->SetSubmixSend(SendInfo);
-			}
+			ActiveSound.SetSubmixSend(SendInfo);
 		}, GET_STATID(STAT_SetSubmixSend));
 	}
 }
@@ -1223,19 +1200,14 @@ void UAudioComponent::SetBusSendEffectInternal(USoundSourceBus* InSourceBus, UAu
 {
 	if (FAudioDevice* AudioDevice = GetAudioDevice())
 	{
-		const uint64 MyAudioComponentID = AudioComponentID;
-		FAudioThread::RunCommandOnAudioThread([AudioDevice, MyAudioComponentID, InSourceBus, InAudioBus, SendLevel, InBusSendType]()
-		{
-			FActiveSound* ActiveSound = AudioDevice->FindActiveSound(MyAudioComponentID);
-			if (ActiveSound)
-			{
-				FSoundSourceBusSendInfo SourceBusSendInfo;
-				SourceBusSendInfo.SoundSourceBus = InSourceBus;
-				SourceBusSendInfo.AudioBus = InAudioBus;
-				SourceBusSendInfo.SendLevel = SendLevel;
+		FSoundSourceBusSendInfo SourceBusSendInfo;
+		SourceBusSendInfo.SoundSourceBus = InSourceBus;
+		SourceBusSendInfo.AudioBus = InAudioBus;
+		SourceBusSendInfo.SendLevel = SendLevel;
 
-				ActiveSound->SetSourceBusSend(InBusSendType, SourceBusSendInfo);
-			}
+		AudioDevice->SendCommandToActiveSounds(AudioComponentID, [InBusSendType, SourceBusSendInfo](FActiveSound& ActiveSound)
+		{
+			ActiveSound.SetSourceBusSend(InBusSendType, SourceBusSendInfo);
 		});
 	}
 }
@@ -1265,15 +1237,9 @@ void UAudioComponent::SetLowPassFilterEnabled(bool InLowPassFilterEnabled)
 	if (FAudioDevice* AudioDevice = GetAudioDevice())
 	{
 		DECLARE_CYCLE_STAT(TEXT("FAudioThreadTask.SetLowPassFilterFrequency"), STAT_AudioSetLowPassFilterEnabled, STATGROUP_AudioThreadCommands);
-
-		const uint64 MyAudioComponentID = AudioComponentID;
-		FAudioThread::RunCommandOnAudioThread([AudioDevice, MyAudioComponentID, InLowPassFilterEnabled]()
+		AudioDevice->SendCommandToActiveSounds(AudioComponentID, [InLowPassFilterEnabled](FActiveSound& ActiveSound)
 		{
-			FActiveSound* ActiveSound = AudioDevice->FindActiveSound(MyAudioComponentID);
-			if (ActiveSound)
-			{
-				ActiveSound->bEnableLowPassFilter = InLowPassFilterEnabled;
-			}
+			ActiveSound.bEnableLowPassFilter = InLowPassFilterEnabled;
 		}, GET_STATID(STAT_AudioSetLowPassFilterEnabled));
 	}
 }
@@ -1283,15 +1249,9 @@ void UAudioComponent::SetLowPassFilterFrequency(float InLowPassFilterFrequency)
 	if (FAudioDevice* AudioDevice = GetAudioDevice())
 	{
 		DECLARE_CYCLE_STAT(TEXT("FAudioThreadTask.SetLowPassFilterFrequency"), STAT_AudioSetLowPassFilterFrequency, STATGROUP_AudioThreadCommands);
-
-		const uint64 MyAudioComponentID = AudioComponentID;
-		FAudioThread::RunCommandOnAudioThread([AudioDevice, MyAudioComponentID, InLowPassFilterFrequency]()
+		AudioDevice->SendCommandToActiveSounds(AudioComponentID, [InLowPassFilterFrequency](FActiveSound& ActiveSound)
 		{
-			FActiveSound* ActiveSound = AudioDevice->FindActiveSound(MyAudioComponentID);
-			if (ActiveSound)
-			{
-				ActiveSound->LowPassFilterFrequency = InLowPassFilterFrequency;
-			}
+			ActiveSound.LowPassFilterFrequency = InLowPassFilterFrequency;
 		}, GET_STATID(STAT_AudioSetLowPassFilterFrequency));
 	}
 }
@@ -1301,23 +1261,17 @@ void UAudioComponent::SetOutputToBusOnly(bool bInOutputToBusOnly)
 	if (FAudioDevice* AudioDevice = GetAudioDevice())
 	{
 		DECLARE_CYCLE_STAT(TEXT("FAudioThreadTask.SetOutputToBusOnly"), STAT_AudioSetOutputToBusOnly, STATGROUP_AudioThreadCommands);
-
-		const uint64 MyAudioComponentID = AudioComponentID;
-		FAudioThread::RunCommandOnAudioThread([AudioDevice, MyAudioComponentID, bInOutputToBusOnly]()
+		AudioDevice->SendCommandToActiveSounds(AudioComponentID, [bInOutputToBusOnly](FActiveSound& ActiveSound)
 		{
-			FActiveSound* ActiveSound = AudioDevice->FindActiveSound(MyAudioComponentID);
-			if (ActiveSound)
+			ActiveSound.bHasActiveMainSubmixOutputOverride = true;
+			ActiveSound.bHasActiveSubmixSendRoutingOverride = true;
+			if (bInOutputToBusOnly)
 			{
-				ActiveSound->bHasActiveMainSubmixOutputOverride = true;
-				ActiveSound->bHasActiveSubmixSendRoutingOverride = true;
-				if (bInOutputToBusOnly)
-				{
-					ActiveSound->bHasActiveBusSendRoutingOverride = true;
-					ActiveSound->bEnableBusSendRoutingOverride = true;
-				}
-				ActiveSound->bEnableMainSubmixOutputOverride = !bInOutputToBusOnly;
-				ActiveSound->bEnableSubmixSendRoutingOverride = !bInOutputToBusOnly;
+				ActiveSound.bHasActiveBusSendRoutingOverride = true;
+				ActiveSound.bEnableBusSendRoutingOverride = true;
 			}
+			ActiveSound.bEnableMainSubmixOutputOverride = !bInOutputToBusOnly;
+			ActiveSound.bEnableSubmixSendRoutingOverride = !bInOutputToBusOnly;
 		}, GET_STATID(STAT_AudioSetOutputToBusOnly));
 	}
 }
@@ -1537,4 +1491,4 @@ bool UAudioComponent::GetCookedEnvelopeDataForAllPlayingSounds(TArray<FSoundWave
 void UAudioComponent::SetSourceEffectChain(USoundEffectSourcePresetChain* InSourceEffectChain)
 {
 	SourceEffectChain = InSourceEffectChain;
-}
\ No newline at end of file
+}
diff --git a/Engine/Source/Runtime/Engine/Public/ActiveSound.h b/Engine/Source/Runtime/Engine/Public/ActiveSound.h
--- a/Engine/Source/Runtime/Engine/Public/ActiveSound.h
+++ b/Engine/Source/Runtime/Engine/Public/ActiveSound.h
@@ -286,6 +286,8 @@ public:
 	bool IsPreviewSound() const override { return bIsPreviewSound; }
 	void Stop() override;
 
+	/** Returns a unique identifier for this active sound object */
+	uint32 GetInstanceID() const { return PlayOrder; }
 
 	uint64 GetAudioComponentID() const { return AudioComponentID; }
 	FName GetAudioComponentUserID() const { return AudioComponentUserID; }
diff --git a/Engine/Source/Runtime/Engine/Public/AudioDevice.h b/Engine/Source/Runtime/Engine/Public/AudioDevice.h
--- a/Engine/Source/Runtime/Engine/Public/AudioDevice.h
+++ b/Engine/Source/Runtime/Engine/Public/AudioDevice.h
@@ -796,8 +796,21 @@ public:
 	/**
 	* Finds the active sound for the specified audio component ID
 	*/
+	UE_DEPRECATED(5.0, "This call cannot be called from the GameThread & AudioComponents can now execute multiple ActiveSound instances at once.  Use 'SendCommandToActiveSounds' instead.")
 	FActiveSound* FindActiveSound(uint64 AudioComponentID);
 
+	/**
+	 * Whether a given Audio Component ID should be allowed to have multiple
+	 * associated Active Sounds
+	 */
+	bool CanHaveMultipleActiveSounds(uint64 AudioComponentID) const;
+
+	/**
+	 * Set whether a given Audio Component ID should be allowed to have multiple
+	 * associated Active Sounds
+	 */
+	void SetCanHaveMultipleActiveSounds(uint64 AudioComponentID, bool InCanHaveMultipleActiveSounds);
+
 	/**
 	 * Removes an active sound from the active sounds array
 	 */
@@ -1325,6 +1338,9 @@ public:
 		return bIsBakedAnalysisEnabled;
 	}
 
+	/** Performs an operation on all active sounds requested to execute by an audio component */
+	void SendCommandToActiveSounds(uint64 InAudioComponentID, TUniqueFunction<void(FActiveSound&)> InFunc, const TStatId InStatId = TStatId());
+
 	virtual bool IsNonRealtime() const
 	{
 		return false;
@@ -1450,7 +1466,7 @@ protected:
 	 */
 	virtual void OnListenerUpdated(const TArray<FListener>& InListeners) {};
 
-private:
+	private:
 
 	/**
 	 * Adds an active sound to the audio device. Can be a new active sound or one provided by the re-triggering
@@ -2120,7 +2136,10 @@ private:
 	/** A set of sounds which need to be deleted but weren't able to be deleted due to pending async operations */
 	TArray<FActiveSound*> PendingSoundsToDelete;
 
-	TMap<uint64, FActiveSound*> AudioComponentIDToActiveSoundMap;
+	TMap<uint64, TArray<FActiveSound*>> AudioComponentIDToActiveSoundMap;
+
+	/** Used to determine if a given Audio Component should be allowed to have multiple simultaneous associated active sounds */
+	TMap<uint64, bool> AudioComponentIDToCanHaveMultipleActiveSoundsMap;
 
 	TMap<uint32, FAudioVolumeProxy> AudioVolumeProxies;
 
diff --git a/Engine/Source/Runtime/MovieSceneTracks/Private/Evaluation/MovieSceneAudioTemplate.cpp b/Engine/Source/Runtime/MovieSceneTracks/Private/Evaluation/MovieSceneAudioTemplate.cpp
--- a/Engine/Source/Runtime/MovieSceneTracks/Private/Evaluation/MovieSceneAudioTemplate.cpp
+++ b/Engine/Source/Runtime/MovieSceneTracks/Private/Evaluation/MovieSceneAudioTemplate.cpp
@@ -537,16 +537,10 @@ struct FAudioSectionExecutionToken : IMovieSceneExecutionToken
 			if (FAudioDevice* AudioDevice = AudioComponent.GetAudioDevice())
 			{
 				DECLARE_CYCLE_STAT(TEXT("FAudioThreadTask.MovieSceneUpdateAudioTransform"), STAT_MovieSceneUpdateAudioTransform, STATGROUP_TaskGraphTasks);
-
-				const FTransform ActorTransform = AudioComponent.GetComponentTransform();
-				const uint64 ActorComponentID = AudioComponent.GetAudioComponentID();
-				FAudioThread::RunCommandOnAudioThread([AudioDevice, ActorComponentID, ActorTransform]()
+				AudioDevice->SendCommandToActiveSounds(AudioComponent.GetAudioComponentID(), [ActorTransform = AudioComponent.GetComponentTransform()](FActiveSound& ActiveSound)
 				{
-					if (FActiveSound* ActiveSound = AudioDevice->FindActiveSound(ActorComponentID))
-					{
-						ActiveSound->bLocationDefined = true;
-						ActiveSound->Transform = ActorTransform;
-					}
+					ActiveSound.bLocationDefined = true;
+					ActiveSound.Transform = ActorTransform;
 				}, GET_STATID(STAT_MovieSceneUpdateAudioTransform));
 			}
 		}
