Changed SortGPUBuffers to use an FRHICommandList instead of FRHICommandListImmediate

diff --git a/Engine/Source/Runtime/Engine/Private/GPUSort.cpp b/Engine/Source/Runtime/Engine/Private/GPUSort.cpp
--- a/Engine/Source/Runtime/Engine/Private/GPUSort.cpp
+++ b/Engine/Source/Runtime/Engine/Private/GPUSort.cpp
@@ -664,7 +664,7 @@ int32 GetGPUSortPassCount(uint32 KeyMask)
  * @param Count - How many items in the buffer need to be sorted.
  * @returns The index of the buffer containing sorted results.
  */
-int32 SortGPUBuffers(FRHICommandListImmediate& RHICmdList, FGPUSortBuffers SortBuffers, int32 BufferIndex, uint32 KeyMask, int32 Count, ERHIFeatureLevel::Type FeatureLevel)
+int32 SortGPUBuffers(FRHICommandList& RHICmdList, FGPUSortBuffers SortBuffers, int32 BufferIndex, uint32 KeyMask, int32 Count, ERHIFeatureLevel::Type FeatureLevel)
 {
 	FRadixSortParameters SortParameters;
 	FRadixSortUniformBufferRef SortUniformBufferRef;
diff --git a/Engine/Source/Runtime/Engine/Public/GPUSort.h b/Engine/Source/Runtime/Engine/Public/GPUSort.h
--- a/Engine/Source/Runtime/Engine/Public/GPUSort.h
+++ b/Engine/Source/Runtime/Engine/Public/GPUSort.h
@@ -49,7 +49,7 @@ int32 GetGPUSortPassCount(uint32 KeyMask);
  * @param Count - How many items in the buffer need to be sorted.
  * @returns The index of the buffer containing sorted results.
  */
-ENGINE_API int32 SortGPUBuffers(FRHICommandListImmediate& RHICmdList, FGPUSortBuffers SortBuffers, int32 BufferIndex, uint32 KeyMask, int32 Count, ERHIFeatureLevel::Type FeatureLevel);
+ENGINE_API int32 SortGPUBuffers(FRHICommandList& RHICmdList, FGPUSortBuffers SortBuffers, int32 BufferIndex, uint32 KeyMask, int32 Count, ERHIFeatureLevel::Type FeatureLevel);
 
 /**
  * GPU sorting tests.
