Include loaded and active data layers in mem reports via wp.DumpDataLayers console command.

diff --git a/Engine/Config/BaseEngine.ini b/Engine/Config/BaseEngine.ini
--- a/Engine/Config/BaseEngine.ini
+++ b/Engine/Config/BaseEngine.ini
@@ -1732,6 +1732,7 @@ bRecompileOnLoad=true
 +Cmd="listtextures uncompressed"
 +Cmd="LogOutStatLevels"
 +Cmd="ListSpawnedActors"
++Cmd="wp.DumpDataLayers"
 
 [MemReportFullCommands]
 ; Additional commands to run with memreport -full
diff --git a/Engine/Source/Runtime/Engine/Private/WorldPartition/DataLayer/DataLayerSubsystem.cpp b/Engine/Source/Runtime/Engine/Private/WorldPartition/DataLayer/DataLayerSubsystem.cpp
--- a/Engine/Source/Runtime/Engine/Private/WorldPartition/DataLayer/DataLayerSubsystem.cpp
+++ b/Engine/Source/Runtime/Engine/Private/WorldPartition/DataLayer/DataLayerSubsystem.cpp
@@ -14,6 +14,42 @@
 #include "WorldPartition/DataLayer/IDataLayerEditorModule.h"
 #endif
 
+static FAutoConsoleCommandWithOutputDevice GDumpDumpDataLayersCmd(
+	TEXT("wp.DumpDataLayers"),
+	TEXT("Dumps active data layers memory to the log"),
+	FConsoleCommandWithOutputDeviceDelegate::CreateStatic([](FOutputDevice& OutputDevice)
+	{
+		for (const FWorldContext& Context : GEngine->GetWorldContexts())
+		{
+			UWorld* World = Context.World();
+			if (World && World->IsGameWorld())
+			{
+				if (const UDataLayerSubsystem* DataLayerSubsystem = World->GetSubsystem<UDataLayerSubsystem>())
+				{
+					auto DumpDataLayers = [&OutputDevice, DataLayerSubsystem](const TCHAR* StateName, const TSet<FName>& DataLayers)
+					{
+						if (DataLayers.Num())
+						{
+							OutputDevice.Logf(TEXT("  - %s Data Layers:"), StateName);
+							for (const FName& DataLayerName : DataLayers)
+							{
+								if (UDataLayer* DataLayer = DataLayerSubsystem->GetDataLayerFromName(DataLayerName))
+								{
+									OutputDevice.Logf(TEXT("    - %s"), *DataLayer->GetDataLayerLabel().ToString());
+								}
+							}
+						}
+					};
+
+					OutputDevice.Logf(TEXT("Data Layers:"));
+					DumpDataLayers(TEXT("Loaded"), DataLayerSubsystem->GetLoadedDataLayerNames());
+					DumpDataLayers(TEXT("Active"), DataLayerSubsystem->GetActiveDataLayerNames());
+				}
+			}
+		}
+	})
+);
+
 UDataLayerSubsystem::UDataLayerSubsystem()
 {}
 
