Fix and optimize toggling NaniteEnabled from a multi-selection in the content browser

  - Fix crash caused by the CompletedTasks potentially containing non-completed tasks due to pointer reuse when reallocating new tasks after some cancellations
  - Add batched cancellations to DistanceField and MeshCards and use it from the BatchBuild function
  - Use BatchBuild when applying NaniteEnabled to multiple meshes at once
  - Some refactor and simplification of the cancellation logic for both DistanceField and MeshCards

  - 1h23m to 1m9s to right-click and change 50000 unloaded meshes from the content browser

diff --git a/Engine/Source/Developer/AssetTools/Private/AssetTypeActions/AssetTypeActions_StaticMesh.cpp b/Engine/Source/Developer/AssetTools/Private/AssetTypeActions/AssetTypeActions_StaticMesh.cpp
--- a/Engine/Source/Developer/AssetTools/Private/AssetTypeActions/AssetTypeActions_StaticMesh.cpp
+++ b/Engine/Source/Developer/AssetTools/Private/AssetTypeActions/AssetTypeActions_StaticMesh.cpp
@@ -9,6 +9,7 @@
 #include "AssetTools.h"
 #include "Editor/StaticMeshEditor/Public/StaticMeshEditorModule.h"
 #include "FbxMeshUtils.h"
+#include "StaticMeshCompiler.h"
 #include "Framework/Notifications/NotificationManager.h"
 #include "Widgets/Notifications/SNotificationList.h"
 #include "Misc/MessageDialog.h"
@@ -287,34 +288,44 @@ void FAssetTypeActions_StaticMesh::ExecuteRemoveVertexColors(TArray<TWeakObjectP
 	}
 }
 
-void FAssetTypeActions_StaticMesh::ExecuteNaniteEnable(TArray<TWeakObjectPtr<UStaticMesh>> Objects)
+void FAssetTypeActions_StaticMesh::ModifyNaniteEnable(TArray<TWeakObjectPtr<UStaticMesh>> Objects, bool bNaniteEnable)
 {
-	for (auto StaticMeshPtr : Objects)
+	TArray<UStaticMesh*> Meshes;
+	Meshes.Reserve(Objects.Num());
+
+	for (TWeakObjectPtr<UStaticMesh>& StaticMeshPtr : Objects)
 	{
 		UStaticMesh* Mesh = StaticMeshPtr.Get();
-		if (Mesh)
+		if (Mesh && Mesh->NaniteSettings.bEnabled != bNaniteEnable)
 		{
-			Mesh->NaniteSettings.bEnabled = true;
-			Mesh->Build();
-			Mesh->MarkPackageDirty();
-			Mesh->OnMeshChanged.Broadcast();
+			Meshes.Add(Mesh);
 		}
 	}
+
+	FStaticMeshCompilingManager::Get().FinishCompilation(Meshes);
+
+	for (UStaticMesh* Mesh : Meshes)
+	{
+		Mesh->NaniteSettings.bEnabled = bNaniteEnable;
+	}
+
+	UStaticMesh::BatchBuild(Meshes);
+
+	for (UStaticMesh* Mesh : Meshes)
+	{
+		Mesh->MarkPackageDirty();
+		Mesh->OnMeshChanged.Broadcast();
+	}
+}
+
+void FAssetTypeActions_StaticMesh::ExecuteNaniteEnable(TArray<TWeakObjectPtr<UStaticMesh>> Objects)
+{
+	ModifyNaniteEnable(Objects, true);
 }
 
 void FAssetTypeActions_StaticMesh::ExecuteNaniteDisable(TArray<TWeakObjectPtr<UStaticMesh>> Objects)
 {
-	for (auto StaticMeshPtr : Objects)
-	{
-		UStaticMesh* Mesh = StaticMeshPtr.Get();
-		if (Mesh)
-		{
-			Mesh->NaniteSettings.bEnabled = false;
-			Mesh->Build();
-			Mesh->MarkPackageDirty();
-			Mesh->OnMeshChanged.Broadcast();
-		}
-	}
+	ModifyNaniteEnable(Objects, false);
 }
 
 #undef LOCTEXT_NAMESPACE
diff --git a/Engine/Source/Developer/AssetTools/Private/AssetTypeActions/AssetTypeActions_StaticMesh.h b/Engine/Source/Developer/AssetTools/Private/AssetTypeActions/AssetTypeActions_StaticMesh.h
--- a/Engine/Source/Developer/AssetTools/Private/AssetTypeActions/AssetTypeActions_StaticMesh.h
+++ b/Engine/Source/Developer/AssetTools/Private/AssetTypeActions/AssetTypeActions_StaticMesh.h
@@ -46,6 +46,9 @@ private:
 	/** Handler for when RemoveVertexColors is selected */
 	void ExecuteRemoveVertexColors(TArray<TWeakObjectPtr<UStaticMesh>> Objects);
 
+	/** Handler for when NaniteEnable is modified */
+	void ModifyNaniteEnable(TArray<TWeakObjectPtr<UStaticMesh>> Objects, bool bNaniteEnable);
+
 	/** Handler for when NaniteEnable is selected */
 	void ExecuteNaniteEnable(TArray<TWeakObjectPtr<UStaticMesh>> Objects);
 
diff --git a/Engine/Source/Runtime/Engine/Private/DistanceFieldAtlas.cpp b/Engine/Source/Runtime/Engine/Private/DistanceFieldAtlas.cpp
--- a/Engine/Source/Runtime/Engine/Private/DistanceFieldAtlas.cpp
+++ b/Engine/Source/Runtime/Engine/Private/DistanceFieldAtlas.cpp
@@ -307,35 +307,9 @@ FDistanceFieldAsyncQueue::~FDistanceFieldAsyncQueue()
 
 void FDistanceFieldAsyncQueue::OnPostReachabilityAnalysis()
 {
-	FScopeLock Lock(&CriticalSection);
-
-	if (ReferencedTasks.Num() || PendingTasks.Num())
-	{
-		TRACE_CPUPROFILER_EVENT_SCOPE(FDistanceFieldAsyncQueue::CancelUnreachableMeshes);
-
-		TArray<FAsyncDistanceFieldTask*> TasksToCancel;
-		for (auto Iterator = ReferencedTasks.CreateIterator(); Iterator; ++Iterator)
-		{
-			FAsyncDistanceFieldTask* Task = *Iterator;
-			if (IsTaskInvalid(Task))
-			{
-				Iterator.RemoveCurrent();
-				TasksToCancel.Add(Task);
-			}
-		}
-
-		for (auto Iterator = PendingTasks.CreateIterator(); Iterator; ++Iterator)
-		{
-			FAsyncDistanceFieldTask* Task = *Iterator;
-			if (IsTaskInvalid(Task))
-			{
-				Iterator.RemoveCurrent();
-			}
-		}
+	TRACE_CPUPROFILER_EVENT_SCOPE(FDistanceFieldAsyncQueue::CancelUnreachableMeshes);
 
-		Lock.Unlock();
-		CancelAndDeleteBackgroundTask(TasksToCancel);
-	}
+	CancelAndDeleteTaskByPredicate([this](FAsyncDistanceFieldTask* Task) { return IsTaskInvalid(Task); });
 }
 
 void FAsyncDistanceFieldTaskWorker::DoWork()
@@ -390,9 +364,43 @@ bool FDistanceFieldAsyncQueue::IsTaskInvalid(FAsyncDistanceFieldTask* Task) cons
 	return (Task->StaticMesh && Task->StaticMesh->IsUnreachable()) || (Task->GenerateSource && Task->GenerateSource->IsUnreachable());
 }
 
-void FDistanceFieldAsyncQueue::CancelAndDeleteBackgroundTask(TArray<FAsyncDistanceFieldTask*> Tasks)
+void FDistanceFieldAsyncQueue::CancelAndDeleteTaskByPredicate(TFunctionRef<bool (FAsyncDistanceFieldTask*)> InShouldCancelPredicate)
 {
-	TRACE_CPUPROFILER_EVENT_SCOPE(FDistanceFieldAsyncQueue::CancelAndDeleteBackgroundTask);
+	TRACE_CPUPROFILER_EVENT_SCOPE(FDistanceFieldAsyncQueue::CancelAndDeleteTaskByPredicate);
+
+	FScopeLock Lock(&CriticalSection);
+
+	if (ReferencedTasks.Num() || PendingTasks.Num() || CompletedTasks.Num())
+	{
+		TSet<FAsyncDistanceFieldTask*> Removed;
+
+		auto RemoveByPredicate =
+			[&InShouldCancelPredicate, &Removed](TSet<FAsyncDistanceFieldTask*>& Tasks)
+		{
+			for (auto It = Tasks.CreateIterator(); It; ++It)
+			{
+				FAsyncDistanceFieldTask* Task = *It;
+				if (InShouldCancelPredicate(Task))
+				{
+					Removed.Add(Task);
+					It.RemoveCurrent();
+				}
+			}
+		};
+
+		RemoveByPredicate(PendingTasks);
+		RemoveByPredicate(ReferencedTasks);
+		RemoveByPredicate(CompletedTasks);
+
+		Lock.Unlock();
+
+		CancelAndDeleteTask(Removed);
+	}
+}
+
+void FDistanceFieldAsyncQueue::CancelAndDeleteTask(const TSet<FAsyncDistanceFieldTask*>& Tasks)
+{
+	TRACE_CPUPROFILER_EVENT_SCOPE(FDistanceFieldAsyncQueue::CancelAndDeleteTask);
 
 	// Do all the cancellation first to make sure none of these tasks
 	// get scheduled as we're waiting for completion.
@@ -421,6 +429,15 @@ void FDistanceFieldAsyncQueue::CancelAndDeleteBackgroundTask(TArray<FAsyncDistan
 			BeginCleanup(Task->GeneratedVolumeData);
 		}
 
+#if DO_GUARD_SLOW
+		{
+			FScopeLock Lock(&CriticalSection);
+			check(!PendingTasks.Contains(Task));
+			check(!ReferencedTasks.Contains(Task));
+			check(!CompletedTasks.Contains(Task));
+		}
+#endif
+
 		delete Task;
 	}
 }
@@ -435,16 +452,15 @@ void FDistanceFieldAsyncQueue::StartBackgroundTask(FAsyncDistanceFieldTask* Task
 void FDistanceFieldAsyncQueue::ProcessPendingTasks()
 {
 	FScopeLock Lock(&CriticalSection);
-	TArray<FAsyncDistanceFieldTask*> Tasks = MoveTemp(PendingTasks);
-	for (FAsyncDistanceFieldTask* Task : Tasks)
+
+	for (auto It = PendingTasks.CreateIterator(); It; ++It)
 	{
-		if (Task->GenerateSource && Task->GenerateSource->IsCompiling())
-		{
-			PendingTasks.Add(Task);
-		}
-		else
+		FAsyncDistanceFieldTask* Task = *It;
+		if (Task->GenerateSource == nullptr || 
+			Task->GenerateSource->IsCompiling() == false)
 		{
 			StartBackgroundTask(Task);
+			It.RemoveCurrent();
 		}
 	}
 }
@@ -456,7 +472,7 @@ void FDistanceFieldAsyncQueue::AddTask(FAsyncDistanceFieldTask* Task)
 	// Simply delete the task if the static mesh are being garbage collected
 	if (IsTaskInvalid(Task))
 	{
-		CancelAndDeleteBackgroundTask({Task});
+		CancelAndDeleteTask({Task});
 		return;
 	}
 
@@ -465,96 +481,75 @@ void FDistanceFieldAsyncQueue::AddTask(FAsyncDistanceFieldTask* Task)
 		MeshUtilities = &FModuleManager::Get().LoadModuleChecked<IMeshUtilities>(TEXT("MeshUtilities"));
 	}
 	
+	const bool bUseAsyncBuild = GUseAsyncDistanceFieldBuildQueue || !IsInGameThread();
+	const bool bIsCompiling = Task->GenerateSource->IsCompiling();
 	{
 		// Array protection when called from multiple threads
 		FScopeLock Lock(&CriticalSection);
+		check(!CompletedTasks.Contains(Task)); // reusing same pointer for a new task that is marked completed but has been canceled...
 		ReferencedTasks.Add(Task);
-	}
 
-	// The Source Mesh's RenderData is not yet ready, postpone the build
-	if (Task->GenerateSource->IsCompiling())
-	{
-		// Array protection when called from multiple threads
-		FScopeLock Lock(&CriticalSection);
-		PendingTasks.Add(Task);
-	}
-	else
-	{
-		// If we're already in worker threads, there is no need to launch an async task.
-		if (GUseAsyncDistanceFieldBuildQueue || !IsInGameThread())
+		// The Source Mesh's RenderData is not ready yet, postpone the build
+		if (bIsCompiling)
 		{
-			StartBackgroundTask(Task);
+			PendingTasks.Add(Task);
 		}
-		else
+		else if (bUseAsyncBuild)
 		{
-			// To avoid deadlocks, we must queue the inner build tasks on another thread pool, so use the task graph.
-			// Put on background thread to avoid interfering with game-thread bound tasks
-			FQueuedThreadPoolTaskGraphWrapper TaskGraphWrapper(ENamedThreads::AnyBackgroundThreadNormalTask);
-			Build(Task, TaskGraphWrapper);
+			// Make sure the Task is launched while we hold the lock to avoid
+			// race with cancellation
+			StartBackgroundTask(Task);
 		}
 	}
+	
+	if (!bIsCompiling && !bUseAsyncBuild)
+	{
+		// To avoid deadlocks, we must queue the inner build tasks on another thread pool, so use the task graph.
+		// Put on background thread to avoid interfering with game-thread bound tasks
+		FQueuedThreadPoolTaskGraphWrapper TaskGraphWrapper(ENamedThreads::AnyBackgroundThreadNormalTask);
+		Build(Task, TaskGraphWrapper);
+	}
 #else
 	UE_LOG(LogStaticMesh,Fatal,TEXT("Tried to build a distance field without editor support (this should have been done during cooking)"));
 #endif
 }
 
-void FDistanceFieldAsyncQueue::CancelBuild(UStaticMesh* StaticMesh)
+void FDistanceFieldAsyncQueue::CancelBuild(UStaticMesh* InStaticMesh)
 {
-	TRACE_CPUPROFILER_EVENT_SCOPE(FDistanceFieldAsyncQueue::CancelBuild)
+	CancelBuilds({ InStaticMesh });
+}
 
-	TArray<FAsyncDistanceFieldTask*> TasksToCancel;
-	{
-		FScopeLock Lock(&CriticalSection);
-		TArray<FAsyncDistanceFieldTask*> Tasks = MoveTemp(PendingTasks);
-		PendingTasks.Reserve(Tasks.Num());
-		for (FAsyncDistanceFieldTask* Task : Tasks)
-		{
-			if (Task->GenerateSource != StaticMesh && Task->StaticMesh != StaticMesh)
-			{
-				PendingTasks.Add(Task);
-			}
-		}
+void FDistanceFieldAsyncQueue::CancelBuilds(const TSet<UStaticMesh*>& InStaticMeshes)
+{
+	TRACE_CPUPROFILER_EVENT_SCOPE(FDistanceFieldAsyncQueue::CancelBuilds)
 
-		Tasks = MoveTemp(ReferencedTasks);
-		ReferencedTasks.Reserve(Tasks.Num());
-		for (FAsyncDistanceFieldTask* Task : Tasks)
+	CancelAndDeleteTaskByPredicate(
+		[&InStaticMeshes](FAsyncDistanceFieldTask* Task)
 		{
-			if (Task->GenerateSource != StaticMesh && Task->StaticMesh != StaticMesh)
-			{
-				ReferencedTasks.Add(Task);
-			}
-			else
-			{
-				TasksToCancel.Add(Task);
-			}
+			return InStaticMeshes.Contains(Task->GenerateSource) || InStaticMeshes.Contains(Task->StaticMesh);
 		}
-	}
-
-	CancelAndDeleteBackgroundTask(TasksToCancel);
+	);
 }
 
 void FDistanceFieldAsyncQueue::CancelAllOutstandingBuilds()
 {
 	TRACE_CPUPROFILER_EVENT_SCOPE(FDistanceFieldAsyncQueue::CancelAllOutstandingBuilds)
 
-	TArray<FAsyncDistanceFieldTask*> OutstandingTasks;
+	TSet<FAsyncDistanceFieldTask*> OutstandingTasks;
 	{
 		FScopeLock Lock(&CriticalSection);
-		PendingTasks.Empty();
+		PendingTasks.Reset();
 		OutstandingTasks = MoveTemp(ReferencedTasks);
 	}
 
-	CancelAndDeleteBackgroundTask(OutstandingTasks);
+	CancelAndDeleteTask(OutstandingTasks);
 }
 
 void FDistanceFieldAsyncQueue::RescheduleBackgroundTask(FAsyncDistanceFieldTask* InTask, EQueuedWorkPriority InPriority)
 {
-	if (InTask->AsyncTask)
+	if (InTask->AsyncTask && InTask->AsyncTask->GetPriority() != InPriority)
 	{
-		if (InTask->AsyncTask->GetPriority() != InPriority)
-		{
-			InTask->AsyncTask->Reschedule(GThreadPool, InPriority);
-		}
+		InTask->AsyncTask->Reschedule(GThreadPool, InPriority);
 	}
 }
 
@@ -582,10 +577,10 @@ void FDistanceFieldAsyncQueue::BlockUntilBuildComplete(UStaticMesh* StaticMesh,
 		// Reschedule the tasks we're waiting on as highest prio
 		{
 			FScopeLock Lock(&CriticalSection);
-			for (int TaskIndex = 0; TaskIndex < ReferencedTasks.Num(); TaskIndex++)
+			for (FAsyncDistanceFieldTask* Task : ReferencedTasks)
 			{
-				if (ReferencedTasks[TaskIndex]->StaticMesh == StaticMesh ||
-					ReferencedTasks[TaskIndex]->GenerateSource == StaticMesh)
+				if (Task->StaticMesh == StaticMesh ||
+					Task->GenerateSource == StaticMesh)
 				{
 					bReferenced = true;
 
@@ -593,20 +588,20 @@ void FDistanceFieldAsyncQueue::BlockUntilBuildComplete(UStaticMesh* StaticMesh,
 					// we need to force finish them too.
 #if WITH_EDITOR
 					
-					if (ReferencedTasks[TaskIndex]->GenerateSource != nullptr &&
-						ReferencedTasks[TaskIndex]->GenerateSource->IsCompiling())
+					if (Task->GenerateSource != nullptr &&
+						Task->GenerateSource->IsCompiling())
 					{
-						RequiredFinishCompilation.Add(ReferencedTasks[TaskIndex]->GenerateSource);
+						RequiredFinishCompilation.Add(Task->GenerateSource);
 					}
 
-					if (ReferencedTasks[TaskIndex]->StaticMesh != nullptr &&
-						ReferencedTasks[TaskIndex]->StaticMesh->IsCompiling())
+					if (Task->StaticMesh != nullptr &&
+						Task->StaticMesh->IsCompiling())
 					{
-						RequiredFinishCompilation.Add(ReferencedTasks[TaskIndex]->StaticMesh);
+						RequiredFinishCompilation.Add(Task->StaticMesh);
 					}
 #endif
 
-					RescheduleBackgroundTask(ReferencedTasks[TaskIndex], EQueuedWorkPriority::Highest);
+					RescheduleBackgroundTask(Task, EQueuedWorkPriority::Highest);
 				}
 			}
 		}
@@ -657,9 +652,9 @@ void FDistanceFieldAsyncQueue::BlockUntilAllBuildsComplete()
 		// Reschedule as highest prio since we're explicitly waiting on them
 		{
 			FScopeLock Lock(&CriticalSection);
-			for (int TaskIndex = 0; TaskIndex < ReferencedTasks.Num(); TaskIndex++)
+			for (FAsyncDistanceFieldTask* Task : ReferencedTasks)
 			{
-				RescheduleBackgroundTask(ReferencedTasks[TaskIndex], EQueuedWorkPriority::Highest);
+				RescheduleBackgroundTask(Task, EQueuedWorkPriority::Highest);
 			}
 		}
 
@@ -695,7 +690,15 @@ void FDistanceFieldAsyncQueue::Build(FAsyncDistanceFieldTask* Task, FQueuedThrea
 			*Task->GeneratedVolumeData);
 	}
 
-    CompletedTasks.Push(Task);
+	{
+		FScopeLock Lock(&CriticalSection);
+		// Avoid adding to the completed list if the task has been canceled
+		if (ReferencedTasks.Contains(Task))
+		{
+			CompletedTasks.Add(Task);
+		}
+	}
+
 #endif
 }
 
@@ -711,7 +714,20 @@ void FDistanceFieldAsyncQueue::ProcessAsyncTasks(bool bLimitExecutionTime)
 	double StartTime = FPlatformTime::Seconds();
 	while (!bLimitExecutionTime || (FPlatformTime::Seconds() - StartTime) < MaxProcessingTime)
 	{
-		FAsyncDistanceFieldTask* Task = CompletedTasks.Pop();
+		FAsyncDistanceFieldTask* Task = nullptr;
+		{
+			FScopeLock Lock(&CriticalSection);
+			if (CompletedTasks.Num() > 0)
+			{
+				// Pop any element from the set
+				auto Iterator = CompletedTasks.CreateIterator();
+				Task = *Iterator;
+				Iterator.RemoveCurrent();
+
+				verify(ReferencedTasks.Remove(Task));
+			}
+		}
+
 		if (Task == nullptr)
 		{
 			break;
@@ -720,17 +736,6 @@ void FDistanceFieldAsyncQueue::ProcessAsyncTasks(bool bLimitExecutionTime)
 		// We want to count each resource built from a DDC miss, so count each iteration of the loop separately.
 		COOK_STAT(auto Timer = DistanceFieldCookStats::UsageStats.TimeSyncWork());
 
-		bool bWasCancelled = false;
-		{
-			FScopeLock Lock(&CriticalSection);
-			bWasCancelled = ReferencedTasks.Remove(Task) == 0;
-		}
-
-		if (bWasCancelled)
-		{
-			continue;
-		}
-
 		if (Task->AsyncTask)
 		{
 			Task->AsyncTask->EnsureCompletion();
@@ -740,6 +745,8 @@ void FDistanceFieldAsyncQueue::ProcessAsyncTasks(bool bLimitExecutionTime)
 		// Editor 'force delete' can null any UObject pointers which are seen by reference collecting (eg FProperty or serialized)
 		if (Task->StaticMesh)
 		{
+			check(!Task->StaticMesh->IsCompiling());
+
 			Task->GeneratedVolumeData->bAsyncBuilding = false;
 
 			FStaticMeshRenderData* RenderData = Task->StaticMesh->GetRenderData();
diff --git a/Engine/Source/Runtime/Engine/Private/MeshCardRepresentation.cpp b/Engine/Source/Runtime/Engine/Private/MeshCardRepresentation.cpp
--- a/Engine/Source/Runtime/Engine/Private/MeshCardRepresentation.cpp
+++ b/Engine/Source/Runtime/Engine/Private/MeshCardRepresentation.cpp
@@ -223,35 +223,9 @@ FCardRepresentationAsyncQueue::~FCardRepresentationAsyncQueue()
 
 void FCardRepresentationAsyncQueue::OnPostReachabilityAnalysis()
 {
-	FScopeLock Lock(&CriticalSection);
-
-	if (ReferencedTasks.Num() || PendingTasks.Num())
-	{
-		TRACE_CPUPROFILER_EVENT_SCOPE(FCardRepresentationAsyncQueue::CancelUnreachableMeshes);
-
-		TArray<FAsyncCardRepresentationTask*> TasksToCancel;
-		for (auto Iterator = ReferencedTasks.CreateIterator(); Iterator; ++Iterator)
-		{
-			FAsyncCardRepresentationTask* Task = *Iterator;
-			if (IsTaskInvalid(Task))
-			{
-				Iterator.RemoveCurrent();
-				TasksToCancel.Add(Task);
-			}
-		}
-
-		for (auto Iterator = PendingTasks.CreateIterator(); Iterator; ++Iterator)
-		{
-			FAsyncCardRepresentationTask* Task = *Iterator;
-			if (IsTaskInvalid(Task))
-			{
-				Iterator.RemoveCurrent();
-			}
-		}
+	TRACE_CPUPROFILER_EVENT_SCOPE(FCardRepresentationAsyncQueue::CancelUnreachableMeshes);
 
-		Lock.Unlock();
-		CancelAndDeleteBackgroundTask(TasksToCancel);
-	}
+	CancelAndDeleteTaskByPredicate([this](FAsyncCardRepresentationTask* Task) { return IsTaskInvalid(Task); });
 }
 void FAsyncCardRepresentationTaskWorker::DoWork()
 {
@@ -296,9 +270,43 @@ bool FCardRepresentationAsyncQueue::IsTaskInvalid(FAsyncCardRepresentationTask*
 	return (Task->StaticMesh && Task->StaticMesh->IsUnreachable()) || (Task->GenerateSource && Task->GenerateSource->IsUnreachable());
 }
 
-void FCardRepresentationAsyncQueue::CancelAndDeleteBackgroundTask(TArray<FAsyncCardRepresentationTask*> Tasks)
+void FCardRepresentationAsyncQueue::CancelAndDeleteTaskByPredicate(TFunctionRef<bool (FAsyncCardRepresentationTask*)> InShouldCancelPredicate)
+{
+	TRACE_CPUPROFILER_EVENT_SCOPE(FCardRepresentationAsyncQueue::CancelAndDeleteTaskByPredicate);
+
+	FScopeLock Lock(&CriticalSection);
+
+	if (ReferencedTasks.Num() || PendingTasks.Num() || CompletedTasks.Num())
+	{
+		TSet<FAsyncCardRepresentationTask*> Removed;
+
+		auto RemoveByPredicate =
+			[&InShouldCancelPredicate, &Removed](TSet<FAsyncCardRepresentationTask*>& Tasks)
+		{
+			for (auto It = Tasks.CreateIterator(); It; ++It)
+			{
+				FAsyncCardRepresentationTask* Task = *It;
+				if (InShouldCancelPredicate(Task))
+				{
+					Removed.Add(Task);
+					It.RemoveCurrent();
+				}
+			}
+		};
+
+		RemoveByPredicate(PendingTasks);
+		RemoveByPredicate(ReferencedTasks);
+		RemoveByPredicate(CompletedTasks);
+
+		Lock.Unlock();
+
+		CancelAndDeleteTask(Removed);
+	}
+}
+
+void FCardRepresentationAsyncQueue::CancelAndDeleteTask(const TSet<FAsyncCardRepresentationTask*>& Tasks)
 {
-	TRACE_CPUPROFILER_EVENT_SCOPE(FCardRepresentationAsyncQueue::CancelAndDeleteBackgroundTask);
+	TRACE_CPUPROFILER_EVENT_SCOPE(FCardRepresentationAsyncQueue::CancelAndDeleteTask);
 
 	// Do all the cancellation first to make sure none of these tasks
 	// get scheduled as we're waiting for completion.
@@ -327,6 +335,14 @@ void FCardRepresentationAsyncQueue::CancelAndDeleteBackgroundTask(TArray<FAsyncC
 			BeginCleanup(Task->GeneratedCardRepresentation);
 		}
 
+#if DO_GUARD_SLOW
+		{
+			FScopeLock Lock(&CriticalSection);
+			check(!PendingTasks.Contains(Task));
+			check(!ReferencedTasks.Contains(Task));
+			check(!CompletedTasks.Contains(Task));
+		}
+#endif
 		delete Task;
 	}
 }
@@ -341,16 +357,15 @@ void FCardRepresentationAsyncQueue::StartBackgroundTask(FAsyncCardRepresentation
 void FCardRepresentationAsyncQueue::ProcessPendingTasks()
 {
 	FScopeLock Lock(&CriticalSection);
-	TArray<FAsyncCardRepresentationTask*> Tasks = MoveTemp(PendingTasks);
-	for (FAsyncCardRepresentationTask* Task : Tasks)
+
+	for (auto It = PendingTasks.CreateIterator(); It; ++It)
 	{
-		if (Task->GenerateSource && Task->GenerateSource->IsCompiling())
-		{
-			PendingTasks.Add(Task);
-		}
-		else
+		FAsyncCardRepresentationTask* Task = *It;
+		if (Task->GenerateSource == nullptr || 
+			Task->GenerateSource->IsCompiling() == false)
 		{
 			StartBackgroundTask(Task);
+			It.RemoveCurrent();
 		}
 	}
 }
@@ -362,7 +377,7 @@ void FCardRepresentationAsyncQueue::AddTask(FAsyncCardRepresentationTask* Task)
 	// Simply delete the task if the static mesh are being garbage collected
 	if (IsTaskInvalid(Task))
 	{
-		CancelAndDeleteBackgroundTask({ Task });
+		CancelAndDeleteTask({ Task });
 		return;
 	}
 
@@ -370,97 +385,76 @@ void FCardRepresentationAsyncQueue::AddTask(FAsyncCardRepresentationTask* Task)
 	{
 		MeshUtilities = &FModuleManager::Get().LoadModuleChecked<IMeshUtilities>(TEXT("MeshUtilities"));
 	}
-	
+
+	const bool bUseAsyncBuild = GUseAsyncCardRepresentationBuildQueue || !IsInGameThread();
+	const bool bIsCompiling = Task->GenerateSource->IsCompiling();
 	{
 		// Array protection when called from multiple threads
 		FScopeLock Lock(&CriticalSection);
+		check(!CompletedTasks.Contains(Task)); // reusing same pointer for a new task that is marked completed but has been canceled...
 		ReferencedTasks.Add(Task);
-	}
 
-	// The Source Mesh's RenderData is not yet ready, postpone the build
-	if (Task->GenerateSource->IsCompiling())
-	{
-		// Array protection when called from multiple threads
-		FScopeLock Lock(&CriticalSection);
-		PendingTasks.Add(Task);
-	}
-	else
-	{
-		// If we're already in worker threads there is no need to launch an async task.
-		if (GUseAsyncCardRepresentationBuildQueue || !IsInGameThread())
+		// The Source Mesh's RenderData is not ready yet, postpone the build
+		if (bIsCompiling)
 		{
-			StartBackgroundTask(Task);
+			PendingTasks.Add(Task);
 		}
-		else
+		else if (bUseAsyncBuild)
 		{
-			// To avoid deadlocks, we must queue the inner build tasks on another thread pool, so use the task graph.
-			// Put on background thread to avoid interfering with game-thread bound tasks
-			FQueuedThreadPoolTaskGraphWrapper TaskGraphWrapper(ENamedThreads::AnyBackgroundThreadNormalTask);
-			Build(Task, TaskGraphWrapper);
+			// Make sure the Task is launched while we hold the lock to avoid
+			// race with cancellation
+			StartBackgroundTask(Task);
 		}
 	}
+
+	if (!bIsCompiling && !bUseAsyncBuild)
+	{
+		// To avoid deadlocks, we must queue the inner build tasks on another thread pool, so use the task graph.
+		// Put on background thread to avoid interfering with game-thread bound tasks
+		FQueuedThreadPoolTaskGraphWrapper TaskGraphWrapper(ENamedThreads::AnyBackgroundThreadNormalTask);
+		Build(Task, TaskGraphWrapper);
+	}
 #else
 	UE_LOG(LogStaticMesh,Fatal,TEXT("Tried to build a card representation without editor support (this should have been done during cooking)"));
 #endif
 }
 
-void FCardRepresentationAsyncQueue::CancelBuild(UStaticMesh* StaticMesh)
+void FCardRepresentationAsyncQueue::CancelBuild(UStaticMesh* InStaticMesh)
 {
-	TRACE_CPUPROFILER_EVENT_SCOPE(FCardRepresentationAsyncQueue::CancelBuild);
+	CancelBuilds({ InStaticMesh });
+}
 
-	TArray<FAsyncCardRepresentationTask*> TasksToCancel;
-	{
-		FScopeLock Lock(&CriticalSection);
-		TArray<FAsyncCardRepresentationTask*> Tasks = MoveTemp(PendingTasks);
-		PendingTasks.Reserve(Tasks.Num());
-		for (FAsyncCardRepresentationTask* Task : Tasks)
-		{
-			if (Task->GenerateSource != StaticMesh && Task->StaticMesh != StaticMesh)
-			{
-				PendingTasks.Add(Task);
-			}
-		}
+void FCardRepresentationAsyncQueue::CancelBuilds(const TSet<UStaticMesh*>& InStaticMeshes)
+{
+	TRACE_CPUPROFILER_EVENT_SCOPE(FCardRepresentationAsyncQueue::CancelBuilds);
 
-		Tasks = MoveTemp(ReferencedTasks);
-		ReferencedTasks.Reserve(Tasks.Num());
-		for (FAsyncCardRepresentationTask* Task : Tasks)
+	CancelAndDeleteTaskByPredicate(
+		[&InStaticMeshes](FAsyncCardRepresentationTask* Task) 
 		{
-			if (Task->GenerateSource != StaticMesh && Task->StaticMesh != StaticMesh)
-			{
-				ReferencedTasks.Add(Task);
-			}
-			else
-			{
-				TasksToCancel.Add(Task);
-			}
+			return InStaticMeshes.Contains(Task->GenerateSource) || InStaticMeshes.Contains(Task->StaticMesh);
 		}
-	}
-
-	CancelAndDeleteBackgroundTask(TasksToCancel);
+	);
 }
 
 void FCardRepresentationAsyncQueue::CancelAllOutstandingBuilds()
 {
 	TRACE_CPUPROFILER_EVENT_SCOPE(FCardRepresentationAsyncQueue::CancelAllOutstandingBuilds);
 
-	TArray<FAsyncCardRepresentationTask*> OutstandingTasks;
+	TSet<FAsyncCardRepresentationTask*> OutstandingTasks;
 	{
 		FScopeLock Lock(&CriticalSection);
-		PendingTasks.Empty();
+		PendingTasks.Reset();
 		OutstandingTasks = MoveTemp(ReferencedTasks);
 	}
 
-	CancelAndDeleteBackgroundTask(OutstandingTasks);
+	CancelAndDeleteTask(OutstandingTasks);
 }
 
 void FCardRepresentationAsyncQueue::RescheduleBackgroundTask(FAsyncCardRepresentationTask* InTask, EQueuedWorkPriority InPriority)
 {
-	if (InTask->AsyncTask)
+	if (InTask->AsyncTask && InTask->AsyncTask->GetPriority() != InPriority)
 	{
-		if (InTask->AsyncTask->GetPriority() != InPriority)
-		{
-			InTask->AsyncTask->Reschedule(GThreadPool, InPriority);
-		}
+		InTask->AsyncTask->Reschedule(GThreadPool, InPriority);
 	}
 }
 
@@ -479,6 +473,7 @@ void FCardRepresentationAsyncQueue::BlockUntilBuildComplete(UStaticMesh* StaticM
 	FStaticMeshCompilingManager::Get().FinishCompilation({ StaticMesh });
 #endif
 
+	TSet<UStaticMesh*> RequiredFinishCompilation;
 	do 
 	{
 		ProcessAsyncTasks();
@@ -487,17 +482,44 @@ void FCardRepresentationAsyncQueue::BlockUntilBuildComplete(UStaticMesh* StaticM
 
 		{
 			FScopeLock Lock(&CriticalSection);
-			for (int TaskIndex = 0; TaskIndex < ReferencedTasks.Num(); TaskIndex++)
+			for (FAsyncCardRepresentationTask* Task : ReferencedTasks)
 			{
-				if (ReferencedTasks[TaskIndex]->StaticMesh == StaticMesh ||
-					ReferencedTasks[TaskIndex]->GenerateSource == StaticMesh)
+				if (Task->StaticMesh == StaticMesh ||
+					Task->GenerateSource == StaticMesh)
 				{
 					bReferenced = true;
-					RescheduleBackgroundTask(ReferencedTasks[TaskIndex], EQueuedWorkPriority::Highest);
+
+					// If the task we are waiting on depends on other static meshes
+					// we need to force finish them too.
+#if WITH_EDITOR
+					
+					if (Task->GenerateSource != nullptr &&
+						Task->GenerateSource->IsCompiling())
+					{
+						RequiredFinishCompilation.Add(Task->GenerateSource);
+					}
+
+					if (Task->StaticMesh != nullptr &&
+						Task->StaticMesh->IsCompiling())
+					{
+						RequiredFinishCompilation.Add(Task->StaticMesh);
+					}
+#endif
+
+					RescheduleBackgroundTask(Task, EQueuedWorkPriority::Highest);
 				}
 			}
 		}
 
+#if WITH_EDITOR
+		// Call the finish compilation outside of the critical section since those compilations
+		// might need to register new distance field tasks which also uses the critical section.
+		if (RequiredFinishCompilation.Num())
+		{
+			FStaticMeshCompilingManager::Get().FinishCompilation(RequiredFinishCompilation.Array());
+		}
+#endif
+
 		if (bReferenced)
 		{
 			if (!bHadToBlock)
@@ -535,9 +557,9 @@ void FCardRepresentationAsyncQueue::BlockUntilAllBuildsComplete()
 		// Reschedule as highest prio since we're explicitly waiting on them
 		{
 			FScopeLock Lock(&CriticalSection);
-			for (int TaskIndex = 0; TaskIndex < ReferencedTasks.Num(); TaskIndex++)
+			for (FAsyncCardRepresentationTask* Task : ReferencedTasks)
 			{
-				RescheduleBackgroundTask(ReferencedTasks[TaskIndex], EQueuedWorkPriority::Highest);
+				RescheduleBackgroundTask(Task, EQueuedWorkPriority::Highest);
 			}
 		}
 
@@ -575,7 +597,14 @@ void FCardRepresentationAsyncQueue::Build(FAsyncCardRepresentationTask* Task, FQ
 			*Task->GeneratedCardRepresentation);
 	}
 
-    CompletedTasks.Push(Task);
+	{
+		FScopeLock Lock(&CriticalSection);
+		// Avoid adding to the completed list if the task has been canceled
+		if (ReferencedTasks.Contains(Task))
+		{
+			CompletedTasks.Add(Task);
+		}
+	}
 
 #endif
 }
@@ -592,7 +621,20 @@ void FCardRepresentationAsyncQueue::ProcessAsyncTasks(bool bLimitExecutionTime)
 	double StartTime = FPlatformTime::Seconds();
 	while (!bLimitExecutionTime || (FPlatformTime::Seconds() - StartTime) < MaxProcessingTime)
 	{
-		FAsyncCardRepresentationTask* Task = CompletedTasks.Pop();
+		FAsyncCardRepresentationTask* Task = nullptr;
+		{
+			FScopeLock Lock(&CriticalSection);
+			if (CompletedTasks.Num() > 0)
+			{
+				// Pop any element from the set
+				auto Iterator = CompletedTasks.CreateIterator();
+				Task = *Iterator;
+				Iterator.RemoveCurrent();
+
+				verify(ReferencedTasks.Remove(Task));
+			}
+		}
+
 		if (Task == nullptr)
 		{
 			break;
@@ -601,17 +643,6 @@ void FCardRepresentationAsyncQueue::ProcessAsyncTasks(bool bLimitExecutionTime)
 		// We want to count each resource built from a DDC miss, so count each iteration of the loop separately.
 		COOK_STAT(auto Timer = CardRepresentationCookStats::UsageStats.TimeSyncWork());
 
-		bool bWasCancelled = false;
-		{
-			FScopeLock Lock(&CriticalSection);
-			bWasCancelled = ReferencedTasks.Remove(Task) == 0;
-		}
-
-		if (bWasCancelled)
-		{
-			continue;
-		}
-
 		if (Task->AsyncTask)
 		{
 			Task->AsyncTask->EnsureCompletion();
@@ -621,6 +652,8 @@ void FCardRepresentationAsyncQueue::ProcessAsyncTasks(bool bLimitExecutionTime)
 		// Editor 'force delete' can null any UObject pointers which are seen by reference collecting (eg UProperty or serialized)
 		if (Task->StaticMesh && Task->bSuccess)
 		{
+			check(!Task->StaticMesh->IsCompiling());
+
 			FStaticMeshRenderData* RenderData = Task->StaticMesh->GetRenderData();
 			FCardRepresentationData* OldCardData = RenderData->LODResources[0].CardRepresentationData;
 
diff --git a/Engine/Source/Runtime/Engine/Private/StaticMeshBuild.cpp b/Engine/Source/Runtime/Engine/Private/StaticMeshBuild.cpp
--- a/Engine/Source/Runtime/Engine/Private/StaticMeshBuild.cpp
+++ b/Engine/Source/Runtime/Engine/Private/StaticMeshBuild.cpp
@@ -172,22 +172,14 @@ void UStaticMesh::BatchBuild(const TArray<UStaticMesh*>& InStaticMeshes, bool bI
 			IMeshBuilderModule::GetForPlatform(TargetPlatform);
 		}
 
-		for (UStaticMesh* StaticMesh : StaticMeshesToProcess)
+		if (GDistanceFieldAsyncQueue)
 		{
-			if (StaticMesh->GetRenderData())
-			{
-				// Cancel any previous async builds before modifying RenderData
-				// This can happen during import as the mesh is rebuilt redundantly
-				if (GDistanceFieldAsyncQueue)
-				{
-					GDistanceFieldAsyncQueue->CancelBuild(StaticMesh);
-				}
+			GDistanceFieldAsyncQueue->CancelBuilds(StaticMeshesToProcess);
+		}
 
-				if (GCardRepresentationAsyncQueue)
-				{
-					GCardRepresentationAsyncQueue->CancelBuild(StaticMesh);
-				}
-			}
+		if (GCardRepresentationAsyncQueue)
+		{
+			GCardRepresentationAsyncQueue->CancelBuilds(StaticMeshesToProcess);
 		}
 		
 		TMap<UStaticMesh*, TArray<UStaticMeshComponent*>> StaticMeshComponents;
diff --git a/Engine/Source/Runtime/Engine/Public/DistanceFieldAtlas.h b/Engine/Source/Runtime/Engine/Public/DistanceFieldAtlas.h
--- a/Engine/Source/Runtime/Engine/Public/DistanceFieldAtlas.h
+++ b/Engine/Source/Runtime/Engine/Public/DistanceFieldAtlas.h
@@ -351,7 +351,10 @@ public:
 	/** Cancel the build on this specific static mesh or block until it is completed if already started. */
 	ENGINE_API void CancelBuild(UStaticMesh* StaticMesh);
 
-	/** Blocks the main thread until the async build are either cancelled or completed. */
+	/** Cancel the build on these meshes or block until they are completed if already started. */
+	ENGINE_API void CancelBuilds(const TSet<UStaticMesh*>& InStaticMeshes);
+
+	/** Blocks the main thread until the async build are either canceled or completed. */
 	ENGINE_API void CancelAllOutstandingBuilds();
 
 	/** Blocks the main thread until the async build of the specified mesh is complete. */
@@ -396,8 +399,11 @@ private:
 	/** Task will be sent to a background worker. */
 	void StartBackgroundTask(FAsyncDistanceFieldTask* Task);
 
-	/** Cancel or finish any background work for the given task. */
-	static void CancelAndDeleteBackgroundTask(TArray<FAsyncDistanceFieldTask*> Tasks);
+	/** Cancel or finish any work for any task matching the predicate. */
+	void CancelAndDeleteTaskByPredicate(TFunctionRef<bool(FAsyncDistanceFieldTask*)> ShouldCancelPredicate);
+
+	/** Cancel or finish any work for the given task. */
+	void CancelAndDeleteTask(const TSet<FAsyncDistanceFieldTask*>& Tasks);
 
 	/** Return whether the task has become invalid and should be cancelled (i.e. reference unreachable objects) */
 	bool IsTaskInvalid(FAsyncDistanceFieldTask* Task) const;
@@ -409,14 +415,13 @@ private:
 	void OnAssetPostCompile(const TArray<FAssetCompileData>& CompiledAssets);
 
 	/** Game-thread managed list of tasks in the async system. */
-	TArray<FAsyncDistanceFieldTask*> ReferencedTasks;
+	TSet<FAsyncDistanceFieldTask*> ReferencedTasks;
 
 	/** Tasks that are waiting on static mesh compilation to proceed */
-	TArray<FAsyncDistanceFieldTask*> PendingTasks;
+	TSet<FAsyncDistanceFieldTask*> PendingTasks;
 
 	/** Tasks that have completed processing. */
-	// consider changing this from FIFO to Unordered, which may be faster
-	TLockFreePointerListLIFO<FAsyncDistanceFieldTask> CompletedTasks;
+	TSet<FAsyncDistanceFieldTask*> CompletedTasks;
 
 	FDelegateHandle PostReachabilityAnalysisHandle;
 
diff --git a/Engine/Source/Runtime/Engine/Public/MeshCardRepresentation.h b/Engine/Source/Runtime/Engine/Public/MeshCardRepresentation.h
--- a/Engine/Source/Runtime/Engine/Public/MeshCardRepresentation.h
+++ b/Engine/Source/Runtime/Engine/Public/MeshCardRepresentation.h
@@ -288,7 +288,10 @@ public:
 	/** Cancel the build on this specific static mesh or block until it is completed if already started. */
 	ENGINE_API void CancelBuild(UStaticMesh* StaticMesh);
 
-	/** Blocks the main thread until the async build are either cancelled or completed. */
+	/** Cancel the build on these meshes or block until they are completed if already started. */
+	ENGINE_API void CancelBuilds(const TSet<UStaticMesh*>& InStaticMeshes);
+
+	/** Blocks the main thread until the async build are either canceled or completed. */
 	ENGINE_API void CancelAllOutstandingBuilds();
 
 	/** Blocks the main thread until the async build of the specified mesh is complete. */
@@ -334,23 +337,26 @@ private:
 	/** Task will be sent to a background worker. */
 	void StartBackgroundTask(FAsyncCardRepresentationTask* Task);
 
-	/** Cancel or finish any background work for the given task. */
-	static void CancelAndDeleteBackgroundTask(TArray<FAsyncCardRepresentationTask*> Tasks);
+	/** Cancel or finish any work for any task matching the predicate. */
+	void CancelAndDeleteTaskByPredicate(TFunctionRef<bool(FAsyncCardRepresentationTask*)> ShouldCancelPredicate);
+
+	/** Cancel or finish any work for the given task. */
+	void CancelAndDeleteTask(const TSet<FAsyncCardRepresentationTask*>& Tasks);
 
-	/** Return whether the task has become invalid and should be cancelled (i.e. reference unreachable objects) */
+	/** Return whether the task has become invalid and should be canceled (i.e. reference unreachable objects) */
 	bool IsTaskInvalid(FAsyncCardRepresentationTask* Task) const;
 
 	/** Used to cancel tasks that are not needed anymore when garbage collection occurs */
 	void OnPostReachabilityAnalysis();
 
 	/** Game-thread managed list of tasks in the async system. */
-	TArray<FAsyncCardRepresentationTask*> ReferencedTasks;
+	TSet<FAsyncCardRepresentationTask*> ReferencedTasks;
 
 	/** Tasks that are waiting on static mesh compilation to proceed */
-	TArray<FAsyncCardRepresentationTask*> PendingTasks;
+	TSet<FAsyncCardRepresentationTask*> PendingTasks;
 
 	/** Tasks that have completed processing. */
-	TLockFreePointerListLIFO<FAsyncCardRepresentationTask> CompletedTasks;
+	TSet<FAsyncCardRepresentationTask*> CompletedTasks;
 
 	FDelegateHandle PostReachabilityAnalysisHandle;
 
