Fix incorrectly named "LightingChannelMask" mutator in FPrimitiveUniformShaderParametersBuilder (it's actually setting ObjectBounds) - thankfully nothing was calling this yet.

diff --git a/Engine/Source/Runtime/Engine/Public/PrimitiveUniformShaderParameters.h b/Engine/Source/Runtime/Engine/Public/PrimitiveUniformShaderParameters.h
--- a/Engine/Source/Runtime/Engine/Public/PrimitiveUniformShaderParameters.h
+++ b/Engine/Source/Runtime/Engine/Public/PrimitiveUniformShaderParameters.h
@@ -150,7 +150,7 @@ public:
 		return *this;
 	}
 
-	inline FPrimitiveUniformShaderParametersBuilder& LightingChannelMask(const FVector3f& InObjectBounds)
+	inline FPrimitiveUniformShaderParametersBuilder& ObjectBounds(const FVector3f& InObjectBounds)
 	{
 		Parameters.ObjectBoundsX = InObjectBounds.X;
 		Parameters.ObjectBoundsY = InObjectBounds.Y;
