Include active streaming sources in mem reports via wp.DumpStreamingSources console command.
Also tweaked streaming sources debug display to show blocking state and position in units.

diff --git a/Engine/Config/BaseEngine.ini b/Engine/Config/BaseEngine.ini
--- a/Engine/Config/BaseEngine.ini
+++ b/Engine/Config/BaseEngine.ini
@@ -1733,6 +1733,7 @@ bRecompileOnLoad=true
 +Cmd="LogOutStatLevels"
 +Cmd="ListSpawnedActors"
 +Cmd="wp.DumpDataLayers"
++Cmd="wp.DumpStreamingSources"
 
 [MemReportFullCommands]
 ; Additional commands to run with memreport -full
diff --git a/Engine/Source/Runtime/Engine/Private/WorldPartition/DataLayer/DataLayerSubsystem.cpp b/Engine/Source/Runtime/Engine/Private/WorldPartition/DataLayer/DataLayerSubsystem.cpp
--- a/Engine/Source/Runtime/Engine/Private/WorldPartition/DataLayer/DataLayerSubsystem.cpp
+++ b/Engine/Source/Runtime/Engine/Private/WorldPartition/DataLayer/DataLayerSubsystem.cpp
@@ -16,7 +16,7 @@
 
 static FAutoConsoleCommandWithOutputDevice GDumpDumpDataLayersCmd(
 	TEXT("wp.DumpDataLayers"),
-	TEXT("Dumps active data layers memory to the log"),
+	TEXT("Dumps active data layers to the log"),
 	FConsoleCommandWithOutputDeviceDelegate::CreateStatic([](FOutputDevice& OutputDevice)
 	{
 		for (const FWorldContext& Context : GEngine->GetWorldContexts())
@@ -26,24 +26,7 @@ static FAutoConsoleCommandWithOutputDevice GDumpDumpDataLayersCmd(
 			{
 				if (const UDataLayerSubsystem* DataLayerSubsystem = World->GetSubsystem<UDataLayerSubsystem>())
 				{
-					auto DumpDataLayers = [&OutputDevice, DataLayerSubsystem](const TCHAR* StateName, const TSet<FName>& DataLayers)
-					{
-						if (DataLayers.Num())
-						{
-							OutputDevice.Logf(TEXT("  - %s Data Layers:"), StateName);
-							for (const FName& DataLayerName : DataLayers)
-							{
-								if (UDataLayer* DataLayer = DataLayerSubsystem->GetDataLayerFromName(DataLayerName))
-								{
-									OutputDevice.Logf(TEXT("    - %s"), *DataLayer->GetDataLayerLabel().ToString());
-								}
-							}
-						}
-					};
-
-					OutputDevice.Logf(TEXT("Data Layers:"));
-					DumpDataLayers(TEXT("Loaded"), DataLayerSubsystem->GetLoadedDataLayerNames());
-					DumpDataLayers(TEXT("Active"), DataLayerSubsystem->GetActiveDataLayerNames());
+					DataLayerSubsystem->DumpDataLayers(OutputDevice);
 				}
 			}
 		}
@@ -333,6 +316,28 @@ TArray<UDataLayer*> UDataLayerSubsystem::ConvertArgsToDataLayers(UWorld* World,
 	return OutDataLayers.Array();
 }
 
+void UDataLayerSubsystem::DumpDataLayers(FOutputDevice& OutputDevice) const
+{
+	auto DumpDataLayers = [this, &OutputDevice](const TCHAR* StateName, const TSet<FName>& DataLayers)
+	{
+		if (DataLayers.Num())
+		{
+			OutputDevice.Logf(TEXT("  - %s Data Layers:"), StateName);
+			for (const FName& DataLayerName : DataLayers)
+			{
+				if (UDataLayer* DataLayer = GetDataLayerFromName(DataLayerName))
+				{
+					OutputDevice.Logf(TEXT("    - %s"), *DataLayer->GetDataLayerLabel().ToString());
+				}
+			}
+		}
+	};
+
+	OutputDevice.Logf(TEXT("Data Layers:"));
+	DumpDataLayers(TEXT("Loaded"), GetLoadedDataLayerNames());
+	DumpDataLayers(TEXT("Active"), GetActiveDataLayerNames());
+}
+
 FAutoConsoleCommand UDataLayerSubsystem::ToggleDataLayerActivation(
 	TEXT("wp.Runtime.ToggleDataLayerActivation"),
 	TEXT("Toggles DataLayers active state. Args [DataLayerLabels]"),
diff --git a/Engine/Source/Runtime/Engine/Private/WorldPartition/WorldPartitionSubsystem.cpp b/Engine/Source/Runtime/Engine/Private/WorldPartition/WorldPartitionSubsystem.cpp
--- a/Engine/Source/Runtime/Engine/Private/WorldPartition/WorldPartitionSubsystem.cpp
+++ b/Engine/Source/Runtime/Engine/Private/WorldPartition/WorldPartitionSubsystem.cpp
@@ -65,6 +65,25 @@ static FAutoConsoleVariableRef CVarGLevelStreamingContinuouslyIncrementalGCWhile
 	ECVF_Default
 );
 
+static FAutoConsoleCommandWithOutputDevice GDumpStreamingSourcesCmd(
+	TEXT("wp.DumpstreamingSources"),
+	TEXT("Dumps active streaming sources to the log"),
+	FConsoleCommandWithOutputDeviceDelegate::CreateStatic([](FOutputDevice& OutputDevice)
+	{
+		for (const FWorldContext& Context : GEngine->GetWorldContexts())
+		{
+			UWorld* World = Context.World();
+			if (World && World->IsGameWorld())
+			{
+				if (const UWorldPartitionSubsystem* WorldPartitionSubsystem = World->GetSubsystem<UWorldPartitionSubsystem>())
+				{				
+					WorldPartitionSubsystem->DumpStreamingSources(OutputDevice);
+				}
+			}
+		}
+	})
+);
+
 UWorldPartitionSubsystem::UWorldPartitionSubsystem()
 {}
 
@@ -220,6 +239,20 @@ bool UWorldPartitionSubsystem::IsStreamingCompleted(EWorldPartitionRuntimeCellSt
 	return true;
 }
 
+void UWorldPartitionSubsystem::DumpStreamingSources(FOutputDevice& OutputDevice) const
+{
+	const UWorldPartition* WorldPartition = GetMainWorldPartition();
+	const TArray<FWorldPartitionStreamingSource>* StreamingSources = WorldPartition ? &WorldPartition->GetStreamingSources() : nullptr;
+	if (StreamingSources && (StreamingSources->Num() > 0))
+	{
+		OutputDevice.Logf(TEXT("Streaming Sources:"));
+		for (const FWorldPartitionStreamingSource& StreamingSource : *StreamingSources)
+		{
+			OutputDevice.Logf(TEXT("  - %s: %s"), *StreamingSource.Name.ToString(), *StreamingSource.ToString());
+		}
+	}
+}
+
 void UWorldPartitionSubsystem::UpdateStreamingState()
 {
 #if WITH_EDITOR
@@ -336,8 +369,7 @@ void UWorldPartitionSubsystem::Draw(UCanvas* Canvas, class APlayerController* PC
 			Pos = CurrentOffset + FVector2D(MaxTextWidth + 10, 0.f);
 			for (const FWorldPartitionStreamingSource& StreamingSource : *StreamingSources)
 			{
-				const FString Text = FString::Printf(TEXT("Priority: %d | Pos: %s | Rot: %s | Vel: %3.2f m/s (%d mph)"), StreamingSource.Priority, *StreamingSource.Location.ToString(), *StreamingSource.Rotation.ToString(), StreamingSource.Velocity, (int32)(StreamingSource.Velocity*2.23694f));
-				FWorldPartitionDebugHelper::DrawText(Canvas, Text, GEngine->GetSmallFont(), FColor::White, Pos);
+				FWorldPartitionDebugHelper::DrawText(Canvas, *StreamingSource.ToString(), GEngine->GetSmallFont(), FColor::White, Pos);
 			}
 			CurrentOffset.Y = Pos.Y;
 		}
diff --git a/Engine/Source/Runtime/Engine/Public/WorldPartition/DataLayer/DataLayerSubsystem.h b/Engine/Source/Runtime/Engine/Public/WorldPartition/DataLayer/DataLayerSubsystem.h
--- a/Engine/Source/Runtime/Engine/Public/WorldPartition/DataLayer/DataLayerSubsystem.h
+++ b/Engine/Source/Runtime/Engine/Public/WorldPartition/DataLayer/DataLayerSubsystem.h
@@ -81,6 +81,8 @@ public:
 	void DrawDataLayersStatus(UCanvas* Canvas, FVector2D& Offset) const;
 	static TArray<UDataLayer*> ConvertArgsToDataLayers(UWorld* World, const TArray<FString>& InArgs);
 
+	void DumpDataLayers(FOutputDevice& OutputDevice) const;
+
 #if WITH_EDITOR
 public:
 	FORCEINLINE const FDataLayerEditorContext& GetDataLayerEditorContext() const { return DataLayerEditorContext; }
diff --git a/Engine/Source/Runtime/Engine/Public/WorldPartition/WorldPartitionStreamingSource.h b/Engine/Source/Runtime/Engine/Public/WorldPartition/WorldPartitionStreamingSource.h
--- a/Engine/Source/Runtime/Engine/Public/WorldPartition/WorldPartitionStreamingSource.h
+++ b/Engine/Source/Runtime/Engine/Public/WorldPartition/WorldPartitionStreamingSource.h
@@ -321,8 +321,7 @@ struct ENGINE_API FWorldPartitionStreamingSource
 		, Priority(InPriority)
 		, Velocity(InVelocity)
 		, TargetGrid(NAME_None)
-	{
-	}
+	{}
 
 	FColor GetDebugColor() const { return FColor::MakeRedToGreenColorFromScalar(FRandomStream(Name).GetFraction()); }
 
@@ -375,6 +374,20 @@ struct ENGINE_API FWorldPartitionStreamingSource
 			FStreamingSourceShapeHelper::ForEachShape(InGridLoadingRange, InGridLoadingRange, bInProjectIn2D, Location, Rotation, Shapes, InOperation);
 		}
 	}
+
+	FString ToString() const
+	{
+		const FVector Direction = Rotation.Euler();
+		return FString::Printf(
+			TEXT("Priority: %d | %s | Pos: X=%d,Y=%d,Z=%d | Rot: X=%d,Y=%d,Z=%d | Vel: %3.2f m/s (%d mph)"), 
+			Priority, 
+			bBlockOnSlowLoading ? TEXT("Blocking") : TEXT("NonBlocking"),
+			(int32)Location.X, (int32)Location.Y, (int32)Location.Z, 
+			(int32)Direction.X, (int32)Direction.Y, (int32)Direction.Z, 
+			Velocity, 
+			(int32)(Velocity*2.23694f)
+		);
+	}
 };
 
 /**
diff --git a/Engine/Source/Runtime/Engine/Public/WorldPartition/WorldPartitionSubsystem.h b/Engine/Source/Runtime/Engine/Public/WorldPartition/WorldPartitionSubsystem.h
--- a/Engine/Source/Runtime/Engine/Public/WorldPartition/WorldPartitionSubsystem.h
+++ b/Engine/Source/Runtime/Engine/Public/WorldPartition/WorldPartitionSubsystem.h
@@ -50,7 +50,7 @@ public:
 	UFUNCTION(BlueprintCallable, Category = Streaming)
 	bool IsStreamingCompleted(EWorldPartitionRuntimeCellState QueryState, const TArray<FWorldPartitionStreamingQuerySource>& QuerySources, bool bExactState) const;
 
-	void ToggleDrawRuntimeHash2D();
+	void DumpStreamingSources(FOutputDevice& OutputDevice) const;
 
 private:
 	UWorldPartition* GetMainWorldPartition();
