PR #8320: Prevent memory corruption when appending container with itself (Contributed by [REDACTED_PERSON])

#p4v-preflight-copy 17421780

diff --git a/Engine/Source/Runtime/Engine/Private/KismetArrayLibrary.cpp b/Engine/Source/Runtime/Engine/Private/KismetArrayLibrary.cpp
--- a/Engine/Source/Runtime/Engine/Private/KismetArrayLibrary.cpp
+++ b/Engine/Source/Runtime/Engine/Private/KismetArrayLibrary.cpp
@@ -153,8 +153,9 @@ void UKismetArrayLibrary::GenericArray_Append(void* TargetArray, const FArrayPro
 		{
 			FProperty* InnerProp = TargetArrayProp->Inner;
 
+			const int32 EndIdx = SourceArrayHelper.Num();
 			int32 StartIdx = TargetArrayHelper.AddValues(SourceArrayHelper.Num());
-			for(int32 x = 0; x < SourceArrayHelper.Num(); ++x, ++StartIdx)
+			for(int32 x = 0; x < EndIdx; ++x, ++StartIdx)
 			{
 				InnerProp->CopySingleValueToScriptVM(TargetArrayHelper.GetRawPtr(StartIdx), SourceArrayHelper.GetRawPtr(x));
 			}
