Added streaming sources target state in on-screen debug views and memreports.

diff --git a/Engine/Source/Runtime/Engine/Public/WorldPartition/WorldPartitionStreamingSource.h b/Engine/Source/Runtime/Engine/Public/WorldPartition/WorldPartitionStreamingSource.h
--- a/Engine/Source/Runtime/Engine/Public/WorldPartition/WorldPartitionStreamingSource.h
+++ b/Engine/Source/Runtime/Engine/Public/WorldPartition/WorldPartitionStreamingSource.h
@@ -214,6 +214,18 @@ enum class EStreamingSourceTargetState : uint8
 	Activated
 };
 
+inline const TCHAR* GetStreamingSourceTargetStateName(EStreamingSourceTargetState StreamingSourceTargetState)
+{
+	switch(StreamingSourceTargetState)
+
+	{
+	case EStreamingSourceTargetState::Loaded: return TEXT("Loaded");
+	case EStreamingSourceTargetState::Activated: return TEXT("Activated");
+	default: check(0);
+	}
+	return TEXT("Invalid");
+}
+
 /**
  * Structure containing all properties required to query a streaming state
  */
@@ -379,8 +391,9 @@ struct ENGINE_API FWorldPartitionStreamingSource
 	{
 		const FVector Direction = Rotation.Euler();
 		return FString::Printf(
-			TEXT("Priority: %d | %s | Pos: X=%d,Y=%d,Z=%d | Rot: X=%d,Y=%d,Z=%d | Vel: %3.2f m/s (%d mph)"), 
+			TEXT("Priority: %d | %s | %s | Pos: X=%d,Y=%d,Z=%d | Rot: X=%d,Y=%d,Z=%d | Vel: %3.2f m/s (%d mph)"), 
 			Priority, 
+			GetStreamingSourceTargetStateName(TargetState),
 			bBlockOnSlowLoading ? TEXT("Blocking") : TEXT("NonBlocking"),
 			(int32)Location.X, (int32)Location.Y, (int32)Location.Z, 
 			(int32)Direction.X, (int32)Direction.Y, (int32)Direction.Z, 
