- Bump MESHDATAKEY_STATICMESH_DERIVEDDATA_VER to clear out possible duplicate entries for the same source model in the Jupiter ddc caches
- Disable coarse mesh streaming again for safety

diff --git a/Engine/Source/Runtime/Engine/Private/StaticMeshSourceData.cpp b/Engine/Source/Runtime/Engine/Private/StaticMeshSourceData.cpp
--- a/Engine/Source/Runtime/Engine/Private/StaticMeshSourceData.cpp
+++ b/Engine/Source/Runtime/Engine/Private/StaticMeshSourceData.cpp
@@ -360,7 +360,7 @@ void FStaticMeshSourceModel::FillMaterialName(TMap<int32, FName>& OutMaterialMap
 // This is the key which is used for DDC data when legacy RawMesh is converted to MeshDescription.
 // If static mesh derived data needs to be rebuilt (new format, serialization differences, etc.) replace the version GUID below with a new one.
 // In case of merge conflicts with DDC versions, you MUST generate a new GUID and set this new GUID as the version.
-#define MESHDATAKEY_STATICMESH_DERIVEDDATA_VER TEXT("97E1F285E9824AD83693BB8BF976461")
+#define MESHDATAKEY_STATICMESH_DERIVEDDATA_VER TEXT("2BF8C1219912409497797709F70B5105")
 
 static bool GetMeshDataKey(FString& OutKey, const FRawMeshBulkData* RawMeshBulkData, int LodIndex)
 {
