Player Controillers acting as streaming sources can now specify loaded stated instead of activated.

diff --git a/Engine/Source/Runtime/Engine/Classes/GameFramework/PlayerController.h b/Engine/Source/Runtime/Engine/Classes/GameFramework/PlayerController.h
--- a/Engine/Source/Runtime/Engine/Classes/GameFramework/PlayerController.h
+++ b/Engine/Source/Runtime/Engine/Classes/GameFramework/PlayerController.h
@@ -506,6 +506,10 @@ public:
 	UPROPERTY(EditAnywhere, BlueprintReadWrite, Category = WorldPartition)
 	uint32 bEnableStreamingSource:1;
 
+	/** Whether the PlayerController streaming source should activate cells after loading.  */
+	UPROPERTY(EditAnywhere, BlueprintReadWrite, Category = WorldPartition, meta=(EditCondition="bEnableStreamingSource"))
+	uint32 bStreamingSourceShouldActivate:1;
+
 	/** Scale applied to force feedback values */
 	UPROPERTY(config)
 	float ForceFeedbackScale;
@@ -706,6 +710,14 @@ public:
 	UFUNCTION(BlueprintCallable, Category = WorldPartition)
 	virtual bool IsStreamingSourceEnabled() const { return bEnableStreamingSource; }
 
+	/**
+	* Whether the PlayerController streaming source should activate cells after loading.
+	* Default implementation returns bStreamingSourceShouldActivate but can be overriden in child classes.
+	* @return true if it should.
+	*/
+	UFUNCTION(BlueprintCallable, Category = WorldPartition)
+	virtual bool StreamingSourceShouldActivate() const { return bEnableStreamingSource && bStreamingSourceShouldActivate; }
+
 protected:
 	/** Pawn has been possessed, so changing state to NAME_Playing. Start it walking and begin playing with it. */
 	virtual void BeginPlayingState();
diff --git a/Engine/Source/Runtime/Engine/Private/PlayerController.cpp b/Engine/Source/Runtime/Engine/Private/PlayerController.cpp
--- a/Engine/Source/Runtime/Engine/Private/PlayerController.cpp
+++ b/Engine/Source/Runtime/Engine/Private/PlayerController.cpp
@@ -128,6 +128,8 @@ APlayerController::APlayerController(const FObjectInitializer& ObjectInitializer
 	bForceFeedbackEnabled = true;
 	ForceFeedbackScale = 1.f;
 
+	bStreamingSourceShouldActivate = true;
+
 	bAutoManageActiveCameraTarget = true;
 	bRenderPrimitiveComponents = true;
 	SmoothTargetViewRotationSpeed = 20.f;
diff --git a/Engine/Source/Runtime/Engine/Private/WorldPartition/WorldPartitionStreamingPolicy.cpp b/Engine/Source/Runtime/Engine/Private/WorldPartition/WorldPartitionStreamingPolicy.cpp
--- a/Engine/Source/Runtime/Engine/Private/WorldPartition/WorldPartitionStreamingPolicy.cpp
+++ b/Engine/Source/Runtime/Engine/Private/WorldPartition/WorldPartitionStreamingPolicy.cpp
@@ -132,8 +132,9 @@ void UWorldPartitionStreamingPolicy::UpdateStreamingSources()
 
 						// Transform to Local
 						ViewLocation = WorldToLocal.TransformPosition(ViewLocation);
-						ViewRotation = WorldToLocal.TransformRotation(ViewRotation.Quaternion()).Rotator();
-						StreamingSources.Add(FWorldPartitionStreamingSource(Player->GetFName(), ViewLocation, ViewRotation, EStreamingSourceTargetState::Activated, /*bBlockOnSlowLoading=*/true));
+						ViewRotation = WorldToLocal.TransformRotation(ViewRotation.Quaternion()).Rotator();						
+						EStreamingSourceTargetState TargetState = Player->PlayerController->StreamingSourceShouldActivate() ? EStreamingSourceTargetState::Activated : EStreamingSourceTargetState::Loaded;
+						StreamingSources.Add(FWorldPartitionStreamingSource(Player->GetFName(), ViewLocation, ViewRotation, TargetState, /*bBlockOnSlowLoading=*/true));
 					}
 				}
 			}
