Fix copy/paste error

diff --git a/Engine/Source/Runtime/Engine/Private/Materials/MaterialInterface.cpp b/Engine/Source/Runtime/Engine/Private/Materials/MaterialInterface.cpp
--- a/Engine/Source/Runtime/Engine/Private/Materials/MaterialInterface.cpp
+++ b/Engine/Source/Runtime/Engine/Private/Materials/MaterialInterface.cpp
@@ -611,7 +611,7 @@ bool UMaterialInterface::GetFontParameterDefaultValue(const FHashedMaterialParam
 bool UMaterialInterface::GetStaticSwitchParameterDefaultValue(const FHashedMaterialParameterInfo& ParameterInfo, bool& OutValue, FGuid& OutExpressionGuid) const
 {
 	FMaterialParameterMetadata Result;
-	if (GetParameterDefaultValue(EMaterialParameterType::RuntimeVirtualTexture, ParameterInfo, Result))
+	if (GetParameterDefaultValue(EMaterialParameterType::StaticSwitch, ParameterInfo, Result))
 	{
 		OutExpressionGuid = Result.ExpressionGuid;
 		OutValue = Result.Value.AsStaticSwitch();
@@ -623,7 +623,7 @@ bool UMaterialInterface::GetStaticSwitchParameterDefaultValue(const FHashedMater
 bool UMaterialInterface::GetStaticComponentMaskParameterDefaultValue(const FHashedMaterialParameterInfo& ParameterInfo, bool& OutR, bool& OutG, bool& OutB, bool& OutA, FGuid& OutExpressionGuid) const
 {
 	FMaterialParameterMetadata Result;
-	if (GetParameterDefaultValue(EMaterialParameterType::RuntimeVirtualTexture, ParameterInfo, Result))
+	if (GetParameterDefaultValue(EMaterialParameterType::StaticComponentMask, ParameterInfo, Result))
 	{
 		OutExpressionGuid = Result.ExpressionGuid;
 		OutR = Result.Value.Bool[0];
