an attempt to fix hanging references to audio resources

diff --git a/Engine/Source/Runtime/Engine/Private/AudioThread.cpp b/Engine/Source/Runtime/Engine/Private/AudioThread.cpp
--- a/Engine/Source/Runtime/Engine/Private/AudioThread.cpp
+++ b/Engine/Source/Runtime/Engine/Private/AudioThread.cpp
@@ -610,6 +610,7 @@ void FAudioThread::StopAudioThread()
 
 	GAudioAsyncBatcher.Flush();
 	GAudioAsyncBatcher.LastBatch.Wait();
+	GAudioAsyncBatcher.LastBatch = UE::Tasks::FTask{}; // release the task as it can hold some references
 
 	FCoreUObjectDelegates::GetPreGarbageCollectDelegate().Remove(PreGC);
 	FCoreUObjectDelegates::GetPostGarbageCollect().Remove(PostGC);
@@ -652,6 +653,7 @@ void FAudioCommandFence::Wait(bool bProcessGameThreadTasks) const
 {
 	check(IsInGameThread());
 	Fence.Wait();
+	Fence = UE::Tasks::FTask{}; // release the task as it can hold some references
 }
 
 #else // UE_AUDIO_THREAD_AS_PIPE
