Adds/improves tooltips in BlendSpace

diff --git a/Engine/Source/Editor/Persona/Public/BlendSpaceAnalysis.h b/Engine/Source/Editor/Persona/Public/BlendSpaceAnalysis.h
--- a/Engine/Source/Editor/Persona/Public/BlendSpaceAnalysis.h
+++ b/Engine/Source/Editor/Persona/Public/BlendSpaceAnalysis.h
@@ -45,21 +45,21 @@ public:
 UENUM()
 enum class EAnalysisSpace : uint8
 {
-	World,
-	Fixed,
-	Changing,
-	Moving,
+	World    UMETA(ToolTip = "Analysis is done in world space (relative to the root of the character)"),
+	Fixed    UMETA(ToolTip = "Analysis is done in the space of the specified bone or socket based on the first frame of the animation used"),
+	Changing UMETA(ToolTip = "Analysis is done in the space of the specified bone or socket based, but velocities are calculated as if this space is not moving"),
+	Moving   UMETA(ToolTip = "Analysis is done in the space of the specified bone or socket"),
 };
 
 UENUM()
 enum class EAnalysisLinearAxis : uint8
 {
-	PlusX UMETA(DisplayName = "+X"),
-	PlusY UMETA(DisplayName = "+Y"),
-	PlusZ UMETA(DisplayName = "+Z"),
-	MinusX UMETA(DisplayName = "-X"),
-	MinusY UMETA(DisplayName = "-Y"),
-	MinusZ UMETA(DisplayName = "-Z"),
+	PlusX UMETA(DisplayName = "+X", ToolTip = "The axis points in the positive X direction"),
+	PlusY UMETA(DisplayName = "+Y", ToolTip = "The axis points in the positive Y direction"),
+	PlusZ UMETA(DisplayName = "+Z", ToolTip = "The axis points in the positive Z direction"),
+	MinusX UMETA(DisplayName = "-X", ToolTip = "The axis points in the negative X direction"),
+	MinusY UMETA(DisplayName = "-Y", ToolTip = "The axis points in the negative Y direction"),
+	MinusZ UMETA(DisplayName = "-Z", ToolTip = "The axis points in the negative Z direction"),
 };
 
 UENUM()
diff --git a/Engine/Source/Runtime/Engine/Classes/Animation/BlendSpace.h b/Engine/Source/Runtime/Engine/Classes/Animation/BlendSpace.h
--- a/Engine/Source/Runtime/Engine/Classes/Animation/BlendSpace.h
+++ b/Engine/Source/Runtime/Engine/Classes/Animation/BlendSpace.h
@@ -55,9 +55,9 @@ enum EBlendSpaceAxis
 UENUM()
 enum class EPreferredTriangulationDirection : uint8
 {
-	None UMETA(DisplayName = "None"),
-	Tangential UMETA(DisplayName = "Tangential"),
-	Radial UMETA(DisplayName = "Radial")
+	None UMETA(DisplayName = "None", ToolTip = "None"),
+	Tangential UMETA(DisplayName = "Tangential", ToolTip = "When there is ambiguity, rectangles will be split so that the inserted edge tends to not point towards the origin"),
+	Radial UMETA(DisplayName = "Radial", ToolTip = "When there is ambiguity, rectangles will be split so that the inserted edge tends to point towards the origin")
 };
 
 
@@ -103,15 +103,15 @@ struct FBlendParameter
 	UPROPERTY(EditAnywhere, DisplayName = "Name", Category=BlendParameter)
 	FString DisplayName;
 
-	/** Min value for this parameter. */
+	/** Minimum value for this axis range. */
 	UPROPERTY(EditAnywhere, DisplayName = "Minimum Axis Value", Category=BlendParameter)
 	float Min;
 
-	/** Max value for this parameter. */
+	/** Maximum value for this axis range. */
 	UPROPERTY(EditAnywhere, DisplayName = "Maximum Axis Value", Category=BlendParameter)
 	float Max;
 
-	/** The number of grid divisions for this parameter (axis). */
+	/** The number of grid divisions along this axis. */
 	UPROPERTY(EditAnywhere, DisplayName = "Grid Divisions", Category=BlendParameter, meta=(UIMin="1", ClampMin="1"))
 	int32 GridNum;
 
@@ -119,7 +119,7 @@ struct FBlendParameter
 	UPROPERTY(EditAnywhere, DisplayName = "Snap to Grid", Category = BlendParameter)
 	bool bSnapToGrid;
 
-	/** If false then input parameters are clamped to the min/max values on this axis. If true then the input can go outside the min/max range and the blend space is treated as being cyclic on this axis. */
+	/** If true then the input can go outside the min/max range and the blend space is treated as being cyclic on this axis. If false then input parameters are clamped to the min/max values on this axis. */
 	UPROPERTY(EditAnywhere, DisplayName = "Wrap Input", Category = BlendParameter)
 	bool bWrapInput;
 
@@ -762,9 +762,9 @@ public:
 	float AnimLength;
 
 	/** The current mode used by the BlendSpace to decide which animation notifies to fire. Valid options are:
-	- AllAnimations - All notify events will fire
-	- HighestWeightedAnimation - Notify events will only fire from the highest weighted animation
-	- None - No notify events will fire from any animations
+	- AllAnimations: All notify events will fire
+	- HighestWeightedAnimation: Notify events will only fire from the highest weighted animation
+	- None: No notify events will fire from any animations
 	*/
 	UPROPERTY(EditAnywhere, Category = AnimationNotifies)
 	TEnumAsByte<ENotifyTriggerMode::Type> NotifyTriggerMode;
@@ -810,9 +810,9 @@ protected:
 	struct FBlendParameter BlendParameters[3];
 
 	/**
-	 * If you have input smoothing, which axis to scale the animation speed. E.g. for locomotion animation, 
-	 * the speed axis will scale the animation speed in order to make up the difference between the target 
-	 * and the result of blending the samples.
+	 * If you have input smoothing, this specifies the axis on which to scale the animation playback speed. E.g. for 
+	 * locomotion animation, the speed axis will scale the animation speed in order to make up the difference 
+	 * between the target and the result of blending the samples.
 	 */
 	UPROPERTY(EditAnywhere, Category = InputInterpolation)
 	TEnumAsByte<EBlendSpaceAxis> AxisToScaleAnimation;
