Chaos - Joint Constraint
- Added transfer attribute

diff --git a/Engine/Source/Editor/DetailCustomizations/Private/PhysicsConstraintComponentDetails.cpp b/Engine/Source/Editor/DetailCustomizations/Private/PhysicsConstraintComponentDetails.cpp
--- a/Engine/Source/Editor/DetailCustomizations/Private/PhysicsConstraintComponentDetails.cpp
+++ b/Engine/Source/Editor/DetailCustomizations/Private/PhysicsConstraintComponentDetails.cpp
@@ -336,6 +336,9 @@ void FPhysicsConstraintComponentDetails::AddLinearLimits(IDetailLayoutBuilder& D
 	LinearLimitCat.AddProperty(ProfilePropertiesProperty->GetChildHandle(GET_MEMBER_NAME_CHECKED(FConstraintProfileProperties, bLinearPlasticity)).ToSharedRef());
 	LinearLimitCat.AddProperty(ProfilePropertiesProperty->GetChildHandle(GET_MEMBER_NAME_CHECKED(FConstraintProfileProperties, LinearPlasticityType)).ToSharedRef());
 	LinearLimitCat.AddProperty(ProfilePropertiesProperty->GetChildHandle(GET_MEMBER_NAME_CHECKED(FConstraintProfileProperties, LinearPlasticityThreshold)).ToSharedRef());
+
+	// Mass Scale Properties
+	LinearLimitCat.AddProperty(ProfilePropertiesProperty->GetChildHandle(GET_MEMBER_NAME_CHECKED(FConstraintProfileProperties, ContactTransferScale)).ToSharedRef());
 }
 
 
@@ -631,7 +634,6 @@ void FPhysicsConstraintComponentDetails::AddLinearDrive(IDetailLayoutBuilder& De
 	[
 		MaxForceWidget
 	];
-
 }
 
 void FPhysicsConstraintComponentDetails::AddAngularDrive(IDetailLayoutBuilder& DetailBuilder, TSharedPtr<IPropertyHandle> ConstraintInstance, TSharedPtr<IPropertyHandle> ProfilePropertiesProperty)
diff --git a/Engine/Source/Runtime/Engine/Classes/PhysicsEngine/ConstraintInstance.h b/Engine/Source/Runtime/Engine/Classes/PhysicsEngine/ConstraintInstance.h
--- a/Engine/Source/Runtime/Engine/Classes/PhysicsEngine/ConstraintInstance.h
+++ b/Engine/Source/Runtime/Engine/Classes/PhysicsEngine/ConstraintInstance.h
@@ -58,6 +58,10 @@ struct ENGINE_API FConstraintProfileProperties
 	UPROPERTY(EditAnywhere, AdvancedDisplay, Category = Angular, meta = (editcondition = "bAngularPlasticity", ClampMin = "0.0"))
 	float AngularPlasticityThreshold;
 
+	/** [Chaos Only] Colliison transfer on parent from the joints child. Range is 0.0-MAX*/
+	UPROPERTY(EditAnywhere, AdvancedDisplay, Category = Linear, meta = (ClampMin = "0.0"))
+	float ContactTransferScale;
+
 	UPROPERTY(EditAnywhere, Category = Linear)
 	FLinearConstraint LinearLimit;
 
@@ -136,6 +140,9 @@ struct ENGINE_API FConstraintProfileProperties
 	/** Updates joint breakable properties (threshold, etc...)*/
 	void UpdatePlasticity_AssumesLocked(const FPhysicsConstraintHandle& InConstraintRef) const;
 
+	/** Updates joint linear mass scales.*/
+	void UpdateContactTransferScale_AssumesLocked(const FPhysicsConstraintHandle& InConstraintRef) const;
+
 	/** Updates joint flag based on profile properties */
 	void UpdateConstraintFlags_AssumesLocked(const FPhysicsConstraintHandle& InConstraintRef) const;
 
@@ -579,6 +586,21 @@ public:
 		return ProfileInstance.LinearPlasticityType;
 	}
 
+	/** Sets the Contact Transfer Scale properties
+	*	@param InContactTransferScale 	Contact transfer scale to joints parent
+	*/
+	void SetContactTransferScale(float InContactTransferScale)
+	{
+		ProfileInstance.ContactTransferScale = InContactTransferScale;
+		UpdateContactTransferScale();
+	}
+
+	/** Get the Contact Transfer Scale for the parent of the joint */
+	float GetContactTransferScale() const
+	{
+		return ProfileInstance.ContactTransferScale;
+	}
+
 	/** Sets the Angular Breakable properties
 	*	@param bInAngularBreakable		Whether it is possible to break the joint with angular force
 	*	@param InAngularBreakThreshold	Torque needed to break the joint
@@ -956,6 +978,7 @@ private:
 
 	void UpdateBreakable();
 	void UpdatePlasticity();
+	void UpdateContactTransferScale();
 	void UpdateDriveTarget();
 
 public:
diff --git a/Engine/Source/Runtime/Engine/Classes/PhysicsEngine/ConstraintInstanceBlueprintLibrary.h b/Engine/Source/Runtime/Engine/Classes/PhysicsEngine/ConstraintInstanceBlueprintLibrary.h
--- a/Engine/Source/Runtime/Engine/Classes/PhysicsEngine/ConstraintInstanceBlueprintLibrary.h
+++ b/Engine/Source/Runtime/Engine/Classes/PhysicsEngine/ConstraintInstanceBlueprintLibrary.h
@@ -210,7 +210,7 @@ public:
 			TEnumAsByte<EConstraintPlasticityType> PlasticityType
 		);
 
-	/** Sets Constraint Linear Plasticity properties
+	/** Gets Constraint Linear Plasticity properties
 	*	@param Accessor						Constraint accessor to query
 	*	@param bAngularPlasticity			Whether it is possible to reset the target position from the linear displacement
 	*	@param AngularPlasticityThreshold	Delta from target needed to reset the target joint
@@ -223,6 +223,32 @@ public:
 			TEnumAsByte<EConstraintPlasticityType>& PlasticityType
 		);
 
+	//---------------------------------------------------------------------------------------------------
+	//
+	// Contact Transfer
+	//
+	//---------------------------------------------------------------------------------------------------
+
+	/** Gets Constraint Contact Transfer Scale properties
+	*	@param Accessor						Constraint accessor to query
+	*	@param ContactTransferScale			Scale for transfer of child energy to parent.
+	*/
+	UFUNCTION(BlueprintCallable, Category = "Physics|Constraints")
+		static void GetContactTransferScale(
+			UPARAM(ref) FConstraintInstanceAccessor& Accessor,
+			float& ContactTransferScale
+		);
+
+	/** Set Contact Transfer Scale
+	*	@param Accessor						Constraint accessor to change
+	*	@param ContactTransferScale			Set Contact Transfer Scale onto joints parent
+	*/
+	UFUNCTION(BlueprintCallable, Category = "Physics|Constraints")
+	static void SetContactTransferScale(
+		UPARAM(ref) FConstraintInstanceAccessor& Accessor,
+		float ContactTransferScale
+	);
+
 	//---------------------------------------------------------------------------------------------------
 	//
 	// ANGULAR LIMITS 
diff --git a/Engine/Source/Runtime/Engine/Classes/PhysicsEngine/PhysicsConstraintComponent.h b/Engine/Source/Runtime/Engine/Classes/PhysicsEngine/PhysicsConstraintComponent.h
--- a/Engine/Source/Runtime/Engine/Classes/PhysicsEngine/PhysicsConstraintComponent.h
+++ b/Engine/Source/Runtime/Engine/Classes/PhysicsEngine/PhysicsConstraintComponent.h
@@ -298,6 +298,12 @@ public:
 	UFUNCTION(BlueprintCallable, Category = "Physics|Components|PhysicsConstraint")
 	void SetAngularPlasticity(bool bAngularPlasticity, float AngularPlasticityThreshold);
 
+	/** Sets the contact transfer scale properties
+	*	@param ContactTransferScale		   Set the contact transfer scale for the parent of the joint
+	*/
+	UFUNCTION(BlueprintCallable, Category = "Physics|Components|PhysicsConstraint")
+	void SetContactTransferScale(float ContactTransferScale);
+
 	/** Gets the current Angular Twist of the constraint */
 	UFUNCTION(BlueprintCallable, Category = "Physics|Components|PhysicsConstraint")
 	float GetCurrentTwist() const;
diff --git a/Engine/Source/Runtime/Engine/Private/PhysicsEngine/ConstraintInstance.cpp b/Engine/Source/Runtime/Engine/Private/PhysicsEngine/ConstraintInstance.cpp
--- a/Engine/Source/Runtime/Engine/Private/PhysicsEngine/ConstraintInstance.cpp
+++ b/Engine/Source/Runtime/Engine/Private/PhysicsEngine/ConstraintInstance.cpp
@@ -135,6 +135,7 @@ FConstraintProfileProperties::FConstraintProfileProperties()
 	, LinearPlasticityThreshold(0.1f)
 	, AngularBreakThreshold(500.f)
 	, AngularPlasticityThreshold(10.f)
+	, ContactTransferScale(0.f)
 	, bDisableCollision(false)
 	, bParentDominates(false)
 	, bEnableProjection(true)
@@ -188,6 +189,16 @@ void FConstraintInstance::UpdatePlasticity()
 		});
 }
 
+void FConstraintInstance::UpdateContactTransferScale()
+{
+#if WITH_CHAOS
+	FPhysicsInterface::ExecuteOnUnbrokenConstraintReadWrite(ConstraintHandle, [&](const FPhysicsConstraintHandle& InConstraint)
+		{
+			ProfileInstance.UpdateContactTransferScale_AssumesLocked(InConstraint);
+		});
+#endif
+}
+
 void FConstraintProfileProperties::UpdatePlasticity_AssumesLocked(const FPhysicsConstraintHandle& InConstraintRef) const
 {
 #if WITH_CHAOS
@@ -198,6 +209,13 @@ void FConstraintProfileProperties::UpdatePlasticity_AssumesLocked(const FPhysics
 #endif
 }
 
+void FConstraintProfileProperties::UpdateContactTransferScale_AssumesLocked(const FPhysicsConstraintHandle& InConstraintRef) const
+{
+#if WITH_CHAOS
+	FPhysicsInterface::SetContactTransferScale_AssumesLocked(InConstraintRef, ContactTransferScale);
+#endif
+}
+
 void FConstraintInstance::UpdateDriveTarget()
 {
 	FPhysicsInterface::ExecuteOnUnbrokenConstraintReadWrite(ConstraintHandle, [&](const FPhysicsConstraintHandle& InUnbrokenConstraint)
@@ -534,6 +552,7 @@ void FConstraintProfileProperties::Update_AssumesLocked(const FPhysicsConstraint
 
 	UpdateBreakable_AssumesLocked(InConstraintRef);
 	UpdatePlasticity_AssumesLocked(InConstraintRef);
+	UpdateContactTransferScale_AssumesLocked(InConstraintRef);
 
 	// Target
 	FPhysicsInterface::UpdateDriveTarget_AssumesLocked(InConstraintRef, LinearDrive, AngularDrive, InInitialize);
diff --git a/Engine/Source/Runtime/Engine/Private/PhysicsEngine/ConstraintInstanceBlueprintLibrary.cpp b/Engine/Source/Runtime/Engine/Private/PhysicsEngine/ConstraintInstanceBlueprintLibrary.cpp
--- a/Engine/Source/Runtime/Engine/Private/PhysicsEngine/ConstraintInstanceBlueprintLibrary.cpp
+++ b/Engine/Source/Runtime/Engine/Private/PhysicsEngine/ConstraintInstanceBlueprintLibrary.cpp
@@ -264,6 +264,39 @@ void UConstraintInstanceBlueprintLibrary::GetLinearPlasticity(
 	}
 }
 
+//---------------------------------------------------------------------------------------------------
+//
+// Contact Transfer Scale
+//
+//---------------------------------------------------------------------------------------------------
+
+void UConstraintInstanceBlueprintLibrary::SetContactTransferScale(
+	FConstraintInstanceAccessor& Accessor,
+	float ContactTransferScale
+)
+{
+	if (FConstraintInstance* ConstraintInstance = Accessor.Get())
+	{
+		Accessor.Modify();
+		ConstraintInstance->SetContactTransferScale(ContactTransferScale);
+	}
+}
+
+void UConstraintInstanceBlueprintLibrary::GetContactTransferScale(
+	FConstraintInstanceAccessor& Accessor,
+	float& ContactTransferScale
+)
+{
+	if (FConstraintInstance* ConstraintInstance = Accessor.Get())
+	{
+		ContactTransferScale = ConstraintInstance->GetContactTransferScale();
+	}
+	else
+	{
+		ContactTransferScale = 0.f;
+	}
+}
+
 //---------------------------------------------------------------------------------------------------
 //
 // ANGULAR LIMITS 
diff --git a/Engine/Source/Runtime/Engine/Private/PhysicsEngine/ImmediatePhysics/ImmediatePhysicsChaos/ImmediatePhysicsJointHandle_Chaos.cpp b/Engine/Source/Runtime/Engine/Private/PhysicsEngine/ImmediatePhysics/ImmediatePhysicsChaos/ImmediatePhysicsJointHandle_Chaos.cpp
--- a/Engine/Source/Runtime/Engine/Private/PhysicsEngine/ImmediatePhysics/ImmediatePhysicsChaos/ImmediatePhysicsJointHandle_Chaos.cpp
+++ b/Engine/Source/Runtime/Engine/Private/PhysicsEngine/ImmediatePhysics/ImmediatePhysicsChaos/ImmediatePhysicsJointHandle_Chaos.cpp
@@ -116,6 +116,9 @@ namespace ImmediatePhysics_Chaos
 		ConstraintSettings.AngularBreakTorque = (Profile.bAngularBreakable) ? Chaos::ConstraintSettings::AngularBreakScale() * Profile.AngularBreakThreshold : FLT_MAX;
 		ConstraintSettings.AngularPlasticityLimit = (Profile.bAngularPlasticity) ? Profile.AngularPlasticityThreshold : FLT_MAX;
 
+		ConstraintSettings.ContactTransferScale = 0.0f;
+
+
 		// UE Disables Soft Limits when the Limit is less than some threshold. This is not necessary in Chaos but for now we also do it for parity's sake (See FLinearConstraint::UpdateLinearLimit_AssumesLocked).
 		if (ConstraintSettings.LinearLimit < RB_MinSizeToLockDOF)
 		{
diff --git a/Engine/Source/Runtime/Engine/Private/PhysicsEngine/PhysicsConstraintComponent.cpp b/Engine/Source/Runtime/Engine/Private/PhysicsEngine/PhysicsConstraintComponent.cpp
--- a/Engine/Source/Runtime/Engine/Private/PhysicsEngine/PhysicsConstraintComponent.cpp
+++ b/Engine/Source/Runtime/Engine/Private/PhysicsEngine/PhysicsConstraintComponent.cpp
@@ -697,6 +697,11 @@ void UPhysicsConstraintComponent::SetAngularPlasticity(bool bAngularPlasticity,
 	ConstraintInstance.SetAngularPlasticity(bAngularPlasticity, AngularPlasticityThreshold);
 }
 
+void UPhysicsConstraintComponent::SetContactTransferScale(float ContactTransferScale)
+{
+	ConstraintInstance.SetContactTransferScale(ContactTransferScale);
+}
+
 float UPhysicsConstraintComponent::GetCurrentTwist() const
 {
 	const float CurrentTwistRads = ConstraintInstance.GetCurrentTwist();
diff --git a/Engine/Source/Runtime/Experimental/Chaos/Private/Chaos/PBDJointConstraints.cpp b/Engine/Source/Runtime/Experimental/Chaos/Private/Chaos/PBDJointConstraints.cpp
--- a/Engine/Source/Runtime/Experimental/Chaos/Private/Chaos/PBDJointConstraints.cpp
+++ b/Engine/Source/Runtime/Experimental/Chaos/Private/Chaos/PBDJointConstraints.cpp
@@ -183,6 +183,7 @@ namespace Chaos
 		, LinearPlasticityInitialDistanceSquared(FLT_MAX)
 		, AngularBreakTorque(FLT_MAX)
 		, AngularPlasticityLimit(FLT_MAX)
+		, ContactTransferScale(0.f)
 		, UserData(nullptr)
 	{
 	}
diff --git a/Engine/Source/Runtime/Experimental/Chaos/Public/Chaos/PBDJointConstraintData.h b/Engine/Source/Runtime/Experimental/Chaos/Public/Chaos/PBDJointConstraintData.h
--- a/Engine/Source/Runtime/Experimental/Chaos/Public/Chaos/PBDJointConstraintData.h
+++ b/Engine/Source/Runtime/Experimental/Chaos/Public/Chaos/PBDJointConstraintData.h
@@ -162,7 +162,9 @@ namespace Chaos
 		//void SetLinearDriveDamping(FReal InLinearDriveStffness);
 		//FReal GetLinearDriveDamping() const
 
-
+		CONSTRAINT_JOINT_PROPERPETY_IMPL(FReal, ContactTransferScale, EJointConstraintFlags::LinearDrive, JointSettings.ContactTransferScale);
+		//void SetContactTransferScale(FReal InContactTransferScale);
+		//FReal GetContactTransferScale() const
 
 		CONSTRAINT_JOINT_PROPERPETY_IMPL(bool, AngularSLerpPositionDriveEnabled, EJointConstraintFlags::AngularDrive, JointSettings.bAngularSLerpPositionDriveEnabled);
 		//void SetAngularSLerpPositionDriveEnabled(bool InAngularSLerpPositionDriveEnabled);
diff --git a/Engine/Source/Runtime/Experimental/Chaos/Public/Chaos/PBDJointConstraintTypes.h b/Engine/Source/Runtime/Experimental/Chaos/Public/Chaos/PBDJointConstraintTypes.h
--- a/Engine/Source/Runtime/Experimental/Chaos/Public/Chaos/PBDJointConstraintTypes.h
+++ b/Engine/Source/Runtime/Experimental/Chaos/Public/Chaos/PBDJointConstraintTypes.h
@@ -186,6 +186,8 @@ namespace Chaos
 		FReal AngularBreakTorque;
 		FReal AngularPlasticityLimit;
 
+		FReal ContactTransferScale;
+
 		void* UserData;
 	};
 
diff --git a/Engine/Source/Runtime/PhysicsCore/Private/ChaosEngineInterface.cpp b/Engine/Source/Runtime/PhysicsCore/Private/ChaosEngineInterface.cpp
--- a/Engine/Source/Runtime/PhysicsCore/Private/ChaosEngineInterface.cpp
+++ b/Engine/Source/Runtime/PhysicsCore/Private/ChaosEngineInterface.cpp
@@ -1315,6 +1315,17 @@ void FChaosEngineInterface::SetPlasticityLimits_AssumesLocked(const FPhysicsCons
 	}
 }
 
+void FChaosEngineInterface::SetContactTransferScale_AssumesLocked(const FPhysicsConstraintHandle& InConstraintRef, float InContactTransferScale)
+{
+	if (InConstraintRef.IsValid() && InConstraintRef.Constraint->IsType(Chaos::EConstraintType::JointConstraintType))
+	{
+		if (Chaos::FJointConstraint* Constraint = static_cast<Chaos::FJointConstraint*>(InConstraintRef.Constraint))
+		{
+			Constraint->SetContactTransferScale(InContactTransferScale);
+		}
+	}
+}
+
 void FChaosEngineInterface::SetLocalPose(const FPhysicsConstraintHandle& InConstraintRef,const FTransform& InPose,EConstraintFrame::Type InFrame)
 {
 	if (InConstraintRef.IsValid() && InConstraintRef.Constraint->IsType(Chaos::EConstraintType::JointConstraintType))
diff --git a/Engine/Source/Runtime/PhysicsCore/Public/Chaos/ChaosEngineInterface.h b/Engine/Source/Runtime/PhysicsCore/Public/Chaos/ChaosEngineInterface.h
--- a/Engine/Source/Runtime/PhysicsCore/Public/Chaos/ChaosEngineInterface.h
+++ b/Engine/Source/Runtime/PhysicsCore/Public/Chaos/ChaosEngineInterface.h
@@ -499,6 +499,7 @@ public:
 	static void SetParentDominates_AssumesLocked(const FPhysicsConstraintHandle& InConstraintRef,bool bInParentDominates);
 	static void SetBreakForces_AssumesLocked(const FPhysicsConstraintHandle& InConstraintRef,float InLinearBreakForce,float InAngularBreakForce);
 	static void SetPlasticityLimits_AssumesLocked(const FPhysicsConstraintHandle& InConstraintRef, float InLinearPlasticityLimit, float InAngularPlasticityLimit, EConstraintPlasticityType InLinearPlasticityType);
+	static void SetContactTransferScale_AssumesLocked(const FPhysicsConstraintHandle& InConstraintRef, float InContactTransferScale);
 	static void SetLocalPose(const FPhysicsConstraintHandle& InConstraintRef, const FTransform& InPose, EConstraintFrame::Type InFrame);
 
 	static void SetDrivePosition(const FPhysicsConstraintHandle& InConstraintRef,const FVector& InPosition);
