Use a more compact data structure to store VT tile offsets for the simple case where there is no compression on any layer.

diff --git a/Engine/Source/Runtime/Engine/Private/TextureDerivedData.cpp b/Engine/Source/Runtime/Engine/Private/TextureDerivedData.cpp
--- a/Engine/Source/Runtime/Engine/Private/TextureDerivedData.cpp
+++ b/Engine/Source/Runtime/Engine/Private/TextureDerivedData.cpp
@@ -60,7 +60,7 @@
 
 // This GUID is mixed into DDC version for virtual textures only, this allows updating DDC version for VT without invalidating DDC for all textures
 // This is useful during development, but once large numbers of VT are present in shipped content, it will have the same problem as TEXTURE_DERIVEDDATA_VER
-#define TEXTURE_VT_DERIVEDDATA_VER	TEXT("E541C13E65B24669B01DDE3DF18C7312")
+#define TEXTURE_VT_DERIVEDDATA_VER	TEXT("93F54B64FA524F78BCA4EA2CD30B1E9F")
 
 #if ENABLE_COOK_STATS
 namespace TextureCookStats
diff --git a/Engine/Source/Runtime/Engine/Private/VT/UploadingVirtualTexture.cpp b/Engine/Source/Runtime/Engine/Private/VT/UploadingVirtualTexture.cpp
--- a/Engine/Source/Runtime/Engine/Private/VT/UploadingVirtualTexture.cpp
+++ b/Engine/Source/Runtime/Engine/Private/VT/UploadingVirtualTexture.cpp
@@ -101,16 +101,15 @@ uint32 FUploadingVirtualTexture::GetLocalMipBias(uint8 vLevel, uint32 vAddress)
 	uint32 Current_vAddress = vAddress;
 	while (Current_vLevel < NumMips)
 	{
-		const uint32 TileIndex = Data->GetTileIndex(Current_vLevel, Current_vAddress);
-		if (TileIndex == ~0u)
+		if (!Data->IsValidAddress(Current_vLevel, Current_vAddress))
 		{
 			// vAddress is out-of-bounds for the given producer
 			Current_vLevel = NumMips - 1u;
 			break;
 		}
 
-		const int32 ChunkIndex = Data->GetChunkIndex(TileIndex);
-		if (ChunkIndex >= 0)
+		const uint32 TileOffset = Data->GetTileOffset(Current_vLevel, Current_vAddress, 0);
+		if (TileOffset != ~0u)
 		{
 			break;
 		}
diff --git a/Engine/Source/Runtime/Engine/Private/VT/VirtualTextureBuiltData.cpp b/Engine/Source/Runtime/Engine/Private/VT/VirtualTextureBuiltData.cpp
--- a/Engine/Source/Runtime/Engine/Private/VT/VirtualTextureBuiltData.cpp
+++ b/Engine/Source/Runtime/Engine/Private/VT/VirtualTextureBuiltData.cpp
@@ -1,10 +1,75 @@
 // Copyright Epic Games, Inc. All Rights Reserved.
 
 #include "VirtualTextureBuiltData.h"
+
+#include "DerivedDataCacheInterface.h"
+#include "Misc/CoreMisc.h"
 #include "Serialization/CustomVersion.h"
 #include "Serialization/MemoryWriter.h"
-#include "Misc/CoreMisc.h"
-#include "DerivedDataCacheInterface.h"
+
+#if WITH_EDITOR
+
+void FVirtualTextureTileOffsetData::Init(uint32 InWidth, uint32 InHeight)
+{
+	Width = FMath::Max(InWidth, 1u);
+	Height = FMath::Max(InHeight, 1u);
+
+	const uint32 SizePadded = FMath::RoundUpToPowerOfTwo(FMath::Max(Width, Height));
+	MaxAddress = SizePadded * SizePadded;
+
+	Addresses.Empty();
+	Offsets.Empty();
+	TileStates.Empty(MaxAddress);
+	TileStates.Add(false, MaxAddress);
+}
+
+void FVirtualTextureTileOffsetData::AddTile(uint32 InAddress)
+{
+	TileStates[InAddress] = true;
+}
+
+void FVirtualTextureTileOffsetData::Finalize()
+{
+	uint32 Offset = 0;
+	uint32 StartAddress = 0;
+	bool bCurrentState = TileStates[0];
+
+	Addresses.Add(StartAddress);
+	Offsets.Add(bCurrentState ? Offset : ~0u);
+
+	for (uint32 Address = 1; Address < MaxAddress; ++Address)
+	{
+		const bool bState = TileStates[Address];
+		if (bState != bCurrentState)
+		{
+			Offset += bCurrentState ? Address - StartAddress : 0;
+
+			StartAddress = Address;
+			bCurrentState = bState;
+
+			Addresses.Add(StartAddress);
+			Offsets.Add(bCurrentState ? Offset : ~0u);
+		}
+	}
+
+	TileStates.Empty();
+}
+
+#endif // WITH_EDITOR
+
+uint32 FVirtualTextureTileOffsetData::GetTileOffset(uint32 InAddress) const
+{
+	const int32 BlockIndex = Algo::UpperBound(Addresses, InAddress) - 1;
+	const uint32 BaseOffset = Offsets[BlockIndex];
+	if (BaseOffset == ~0u)
+	{
+		// Address is in empty space.
+		return ~0u;
+	}
+	const uint32 BaseAddress = Addresses[BlockIndex];
+	const uint32 LocalOffset = InAddress - BaseAddress;
+	return BaseOffset + LocalOffset;
+}
 
 uint64 FVirtualTextureBuiltData::GetDiskMemoryFootprint() const
 {
@@ -33,7 +98,23 @@ uint32 FVirtualTextureBuiltData::GetMemoryFootprint() const
 
 uint32 FVirtualTextureBuiltData::GetTileMemoryFootprint() const
 {
-	return TileOffsetInChunk.GetAllocatedSize() + TileIndexPerChunk.GetAllocatedSize() + TileIndexPerMip.GetAllocatedSize();
+	uint32 TotalSize = 0;
+
+	// Legacy tile offsets are used if tiles are compressed.
+	TotalSize += TileOffsetInChunk.GetAllocatedSize();
+	TotalSize += TileIndexPerChunk.GetAllocatedSize();
+	TotalSize += TileIndexPerMip.GetAllocatedSize();
+
+	// Implicit tile offsets are used if tiles are uncompressed.
+	TotalSize += ChunkIndexPerMip.GetAllocatedSize();
+	TotalSize += BaseOffsetPerMip.GetAllocatedSize();
+	for (FVirtualTextureTileOffsetData const& Data : TileOffsetData)
+	{
+		TotalSize += Data.Addresses.GetAllocatedSize();
+		TotalSize += Data.Offsets.GetAllocatedSize();
+	}
+
+	return TotalSize;
 }
 
 uint32 FVirtualTextureBuiltData::GetNumTileHeaders() const
@@ -41,6 +122,101 @@ uint32 FVirtualTextureBuiltData::GetNumTileHeaders() const
 	return TileOffsetInChunk.Num();
 }
 
+bool FVirtualTextureBuiltData::IsLegacyData() const
+{
+	return TileOffsetInChunk.Num() > 0;
+}
+
+uint32 FVirtualTextureBuiltData::GetTileIndex_Legacy(uint8 vLevel, uint32 vAddress) const
+{
+	check(vLevel < NumMips);
+	const uint32 TileIndex = TileIndexPerMip[vLevel] + vAddress * NumLayers;
+	if (TileIndex >= TileIndexPerMip[vLevel + 1])
+	{
+		// vAddress is out of bounds for this texture/mip level
+		return ~0u;
+	}
+	return TileIndex;
+}
+
+uint32 FVirtualTextureBuiltData::GetTileOffset_Legacy(uint32 ChunkIndex, uint32 TileIndex) const
+{
+	check(TileIndex >= TileIndexPerChunk[ChunkIndex]);
+	if (TileIndex < TileIndexPerChunk[ChunkIndex + 1])
+	{
+		return TileOffsetInChunk[TileIndex];
+	}
+
+	// If TileIndex is past the end of chunk, return the size of chunk
+	// This allows us to determine size of region by asking for start/end offsets
+	return Chunks[ChunkIndex].SizeInBytes;
+}
+
+bool FVirtualTextureBuiltData::IsValidAddress(uint32 vLevel, uint32 vAddress)
+{
+	bool bIsValid = false;
+
+	if (IsLegacyData())
+	{
+		bIsValid = GetTileIndex_Legacy(vLevel, vAddress) != ~0u;
+	}
+	else
+	{
+		if (TileOffsetData.IsValidIndex(vLevel))
+		{
+			const uint32 X = FMath::ReverseMortonCode2(vAddress);
+			const uint32 Y = FMath::ReverseMortonCode2(vAddress >> 1);
+			bIsValid = X < TileOffsetData[vLevel].Width && Y < TileOffsetData[vLevel].Height;
+		}
+	}
+
+	return bIsValid;
+}
+
+int32 FVirtualTextureBuiltData::GetChunkIndex(uint8 vLevel) const
+{
+	return ChunkIndexPerMip.IsValidIndex(vLevel) ? ChunkIndexPerMip[vLevel] : -1;
+}
+
+uint32 FVirtualTextureBuiltData::GetTileOffset(uint32 vLevel, uint32 vAddress, uint32 LayerIndex) const
+{
+	uint32 Offset = ~0u;
+
+	if (IsLegacyData())
+	{
+		const uint32 TileIndex = GetTileIndex_Legacy(vLevel, vAddress);
+		if (TileIndex != ~0u)
+		{
+			// If size of the tile is 0 we return ~0u to indicate that there is no data present.
+			const int32 ChunkIndex = GetChunkIndex(vLevel);
+			const uint32 TileOffset = GetTileOffset_Legacy(ChunkIndex, TileIndex);
+			const uint32 NextTileOffset = GetTileOffset_Legacy(ChunkIndex, TileIndex + GetNumLayers());
+			if (TileOffset != NextTileOffset)
+			{
+				Offset = GetTileOffset_Legacy(ChunkIndex, TileIndex + LayerIndex);
+			}
+		}
+	}
+	else
+	{
+		if (BaseOffsetPerMip.IsValidIndex(vLevel) && TileOffsetData.IsValidIndex(vLevel))
+		{
+			// If the tile offset is ~0u there is no data present so we return ~0u to indicate that.
+			const uint32 TileOffset = TileOffsetData[vLevel].GetTileOffset(vAddress);
+			if (TileOffset != ~0u)
+			{
+				const uint32 BaseOffset = BaseOffsetPerMip[vLevel];
+				const uint32 TileDataSize = TileDataOffsetPerLayer.Last();
+				const uint32 LayerDataOffset = LayerIndex == 0 ? 0 : TileDataOffsetPerLayer[LayerIndex - 1];
+				
+				Offset = BaseOffset + (TileOffset * TileDataSize) + LayerDataOffset;
+			}
+		}
+	}
+
+	return Offset;
+}
+
 void FVirtualTextureBuiltData::Serialize(FArchive& Ar, UObject* Owner, int32 FirstMipToSerialize)
 {
 	TRACE_CPUPROFILER_EVENT_SCOPE(FVirtualTextureDataChunk::Serialize);
@@ -57,12 +233,16 @@ void FVirtualTextureBuiltData::Serialize(FArchive& Ar, UObject* Owner, int32 Fir
 	Ar << HeightInBlocks;
 	Ar << TileSize;
 	Ar << TileBorderSize;
+	Ar << TileDataOffsetPerLayer;
 
 	if (!bStripMips)
 	{
 		Ar << NumMips;
 		Ar << Width;
 		Ar << Height;
+		Ar << ChunkIndexPerMip;
+		Ar << BaseOffsetPerMip;
+		Ar << TileOffsetData;
 		Ar << TileIndexPerChunk;
 		Ar << TileIndexPerMip;
 		Ar << TileOffsetInChunk;
@@ -70,51 +250,72 @@ void FVirtualTextureBuiltData::Serialize(FArchive& Ar, UObject* Owner, int32 Fir
 	else
 	{
 		check((uint32)FirstMipToSerialize < NumMips);
-		const uint32 NumTilesToStrip = TileIndexPerMip[FirstMipToSerialize];
-		check(NumTilesToStrip < (uint32)TileOffsetInChunk.Num());
-		for (int32 ChunkIndex = 0; ChunkIndex < Chunks.Num(); ++ChunkIndex)
-		{
-			if (TileIndexPerChunk[ChunkIndex + 1] <= NumTilesToStrip)
-			{
-				++NumChunksToStrip;
-			}
-			else
-			{
-				break;
-			}
-		}
+		NumChunksToStrip = ChunkIndexPerMip[FirstMipToSerialize];
 
 		uint32 NumMipsToSerialize = NumMips - FirstMipToSerialize;
 		uint32 WidthToSerialize = Width >> FirstMipToSerialize;
 		uint32 HeightToSerialize = Height >> FirstMipToSerialize;
+
+		TArray<uint32> StrippedChunkIndexPerMip;
+		TArray<uint32> StrippedBaseOffsetPerMip;
+		TArray<FVirtualTextureTileOffsetData> StrippedTileOffsetData;
 		TArray<uint32> StrippedTileIndexPerChunk;
 		TArray<uint32> StrippedTileIndexPerMip;
 		TArray<uint32> StrippedTileOffsetInChunk;
 
-		StrippedTileIndexPerChunk.Reserve(TileIndexPerChunk.Num() - NumChunksToStrip);
-		for (int32 i = NumChunksToStrip; i < TileIndexPerChunk.Num(); ++i)
+		StrippedChunkIndexPerMip.Reserve(ChunkIndexPerMip.Num() - FirstMipToSerialize);
+		for (int32 i = FirstMipToSerialize; i < ChunkIndexPerMip.Num(); ++i)
+		{
+			check(ChunkIndexPerMip[i] >= NumChunksToStrip);
+			StrippedChunkIndexPerMip.Add(ChunkIndexPerMip[i] - NumChunksToStrip);
+		}
+		
+		StrippedBaseOffsetPerMip.Reserve(BaseOffsetPerMip.Num() - FirstMipToSerialize);
+		for (int32 i = FirstMipToSerialize; i < BaseOffsetPerMip.Num(); ++i)
 		{
-			// Since we can only exclude data by chunk, it's possible that the first chunk we need to include will contain some initial tiles from mip that's been exclude
-			StrippedTileIndexPerChunk.Add(TileIndexPerChunk[i] - FMath::Min(NumTilesToStrip, TileIndexPerChunk[i]));
+			StrippedBaseOffsetPerMip.Add(BaseOffsetPerMip[i]);
 		}
 
-		StrippedTileIndexPerMip.Reserve(TileIndexPerMip.Num() - FirstMipToSerialize);
-		for (int32 i = FirstMipToSerialize; i < TileIndexPerMip.Num(); ++i)
+		StrippedTileOffsetData.Reserve(TileOffsetData.Num() - FirstMipToSerialize);
+		for (int32 i = FirstMipToSerialize; i < TileOffsetData.Num(); ++i)
 		{
-			check(TileIndexPerMip[i] >= NumTilesToStrip);
-			StrippedTileIndexPerMip.Add(TileIndexPerMip[i] - NumTilesToStrip);
+			StrippedTileOffsetData.Add(TileOffsetData[i]);
 		}
 
-		StrippedTileOffsetInChunk.Reserve(TileOffsetInChunk.Num() - NumTilesToStrip);
-		for (int32 i = NumTilesToStrip; i < TileOffsetInChunk.Num(); ++i)
+		const bool bHasLegacyData = TileOffsetInChunk.Num() > 0;
+		if (bHasLegacyData)
 		{
-			// offsets within each chunk are unchanged...we are removing chunks that are no longer referenced, but not truncating any existing chunks
-			StrippedTileOffsetInChunk.Add(TileOffsetInChunk[i]);
+			const uint32 NumTilesToStrip = TileIndexPerMip[FirstMipToSerialize];
+			check(NumTilesToStrip < (uint32)TileOffsetInChunk.Num());
+
+			StrippedTileIndexPerChunk.Reserve(TileIndexPerChunk.Num() - NumChunksToStrip);
+			for (int32 i = NumChunksToStrip; i < TileIndexPerChunk.Num(); ++i)
+			{
+				// Since we can only exclude data by chunk, it's possible that the first chunk we need to include will contain some initial tiles from a mip that's been excluded
+				StrippedTileIndexPerChunk.Add(TileIndexPerChunk[i] - FMath::Min(NumTilesToStrip, TileIndexPerChunk[i]));
+			}
+
+			StrippedTileIndexPerMip.Reserve(TileIndexPerMip.Num() - FirstMipToSerialize);
+			for (int32 i = FirstMipToSerialize; i < TileIndexPerMip.Num(); ++i)
+			{
+				check(TileIndexPerMip[i] >= NumTilesToStrip);
+				StrippedTileIndexPerMip.Add(TileIndexPerMip[i] - NumTilesToStrip);
+			}
+
+			StrippedTileOffsetInChunk.Reserve(TileOffsetInChunk.Num() - NumTilesToStrip);
+			for (int32 i = NumTilesToStrip; i < TileOffsetInChunk.Num(); ++i)
+			{
+				// offsets within each chunk are unchanged...we are removing chunks that are no longer referenced, but not truncating any existing chunks
+				StrippedTileOffsetInChunk.Add(TileOffsetInChunk[i]);
+			}
 		}
 
 		Ar << NumMipsToSerialize;
 		Ar << WidthToSerialize;
 		Ar << HeightToSerialize;
+		Ar << StrippedChunkIndexPerMip;
+		Ar << StrippedBaseOffsetPerMip;
+		Ar << StrippedTileOffsetData;
 		Ar << StrippedTileIndexPerChunk;
 		Ar << StrippedTileIndexPerMip;
 		Ar << StrippedTileOffsetInChunk;
@@ -253,7 +454,7 @@ bool FVirtualTextureBuiltData::ValidateData(FStringView const& InDDCDebugContext
 			break;
 		}
 
-		if (bValidateCompression)
+		if (bValidateCompression && IsLegacyData())
 		{
 			uint32 TileIndex = TileIndexPerChunk[ChunkIndex];
 			while (bResult && TileIndex < TileIndexPerChunk[ChunkIndex + 1])
@@ -269,8 +470,8 @@ bool FVirtualTextureBuiltData::ValidateData(FStringView const& InDDCDebugContext
 
 					if (VTCodec == EVirtualTextureCodec::ZippedGPU)
 					{
-						const uint32 TileOffset = GetTileOffset(ChunkIndex, TileIndex);
-						const uint32 NextTileOffset = GetTileOffset(ChunkIndex, TileIndex + 1);
+						const uint32 TileOffset = GetTileOffset_Legacy(ChunkIndex, TileIndex);
+						const uint32 NextTileOffset = GetTileOffset_Legacy(ChunkIndex, TileIndex + 1);
 						check(NextTileOffset >= TileOffset);
 						if (NextTileOffset > TileOffset)
 						{
@@ -349,4 +550,5 @@ uint32 FVirtualTextureDataChunk::StoreInDerivedDataCache(const FString& InDerive
 	BulkData.RemoveBulkData();
 	return BulkDataSizeInBytes;
 }
+
 #endif // WITH_EDITORONLY_DATA
diff --git a/Engine/Source/Runtime/Engine/Private/VT/VirtualTextureBuiltData.h b/Engine/Source/Runtime/Engine/Private/VT/VirtualTextureBuiltData.h
--- a/Engine/Source/Runtime/Engine/Private/VT/VirtualTextureBuiltData.h
+++ b/Engine/Source/Runtime/Engine/Private/VT/VirtualTextureBuiltData.h
@@ -71,6 +71,59 @@ struct FVirtualTextureDataChunk
 #endif // WITH_EDITORONLY_DATA
 };
 
+
+/** 
+ * Compact structure to find tile offsets within a sparse chunk. 
+ * The tiles are stored in Morton order inside the chunks. But for UDIM textures some areas of a mip level may be empty and so some tiles don't exist.
+ * We could store one chunk offset per tile, but that uses a large amount of memory.
+ * So we need a fast, but memory efficient way to get from tile index to chunk location.
+ * This structure splits the space into contiguous blocks of empty or non-empty tiles.
+ * If a block is non-empty it has an offset into the chunk, and if it is empty it has a special "empty" offset.
+ * Within a non-empty block, all tiles exist and are contiguous.
+ * Because the tiles are in Morton order and UDIMs are aligned on power of 2 boundaries the sequences are not usually very fragmented, so textures can be described with a low number of blocks.
+ * Lookup of tile index is done at the cost of a binary search of Morton address through the stored block addresses.
+ */
+struct FVirtualTextureTileOffsetData
+{
+#if WITH_EDITOR
+	/* Call at start of building the data before any calls to AddTile(). The entire area is initialized as empty. */
+	void Init(uint32 InWidth, uint32 InHeight);
+	/* Call to mark a tile as non-empty. */
+	void AddTile(uint32 InAddress);
+	/* Call after all calls to AddTile() to build the final blocks and offsets. */
+	void Finalize();
+#endif // WITH_EDITOR
+
+	/** Call at runtime to get the final tile offset in the chunk. */
+	uint32 GetTileOffset(uint32 InAddress) const;
+
+	/** Serialization helper. */
+	friend FArchive& operator<<(FArchive& Ar, FVirtualTextureTileOffsetData& TileOffsetData)
+	{
+		Ar << TileOffsetData.Width;
+		Ar << TileOffsetData.Height;
+		Ar << TileOffsetData.MaxAddress;
+		Ar << TileOffsetData.Addresses;
+		Ar << TileOffsetData.Offsets;
+		return Ar;
+	}
+
+	uint32 Width = 0;
+	uint32 Height = 0;
+	/** Upper bound Morton address for managed area. */
+	uint32 MaxAddress = 0;
+	/** Sorted list of contiguous tile block addresses. */
+	TArray<uint32> Addresses;
+	/** Offset for each block in Addresses. An empty block is marked with ~0u. */
+	TArray<uint32> Offsets;
+	
+#if WITH_EDITOR
+	/** Tile state scratch buffer. Created in Init() and destroyed in Finalize(). */
+	TBitArray<> TileStates;
+#endif
+};
+
+
 struct FVirtualTextureBuiltData
 {
 	uint32 NumLayers;
@@ -94,6 +147,22 @@ struct FVirtualTextureBuiltData
 	 */
 	TArray<FVirtualTextureDataChunk> Chunks;
 
+	/** 
+	 * Per layer tile data offset in bytes. 
+	 * This is empty if compression (zlib etc) is used, since in that case offsets will vary per tile.
+	 * Each value is the sum of the packed tile data sizes for the current and the preceding layers.
+	 */
+	TArray<uint32> TileDataOffsetPerLayer;
+
+	/** Chunk index that contains the mip. */
+	TArray<uint32> ChunkIndexPerMip;
+
+	/** Base offset in chunk for the mip. */
+	TArray<uint32> BaseOffsetPerMip;
+
+	/** Holds the tile to chunk offset lookup data. One array entry per mip level. */
+	TArray<FVirtualTextureTileOffsetData> TileOffsetData;
+
 	/** Index of the first tile within each chunk */
 	TArray<uint32> TileIndexPerChunk;
 
@@ -102,7 +171,7 @@ struct FVirtualTextureBuiltData
 
 	/**
 	 * Info for the tiles organized per level. Within a level tile info is organized in Morton order.
-	 * This is in morton order which can waste a lot of space in this array for non-square images
+	 * This is in Morton order which can waste a lot of space in this array for non-square images
 	 * e.g.:
 	 * - An 8x1 tile image will allocate 8x4 indexes in this array.
 	 * - An 1x8 tile image will allocate 8x8 indexes in this array.
@@ -133,72 +202,26 @@ struct FVirtualTextureBuiltData
 	uint32 GetTileMemoryFootprint() const;
 	uint64 GetDiskMemoryFootprint() const;
 	uint32 GetNumTileHeaders() const;
+	
 	void Serialize(FArchive& Ar, UObject* Owner, int32 FirstMipToSerialize);
 
-	/**
-	 * Return the index of the given tile
-	 */
-	inline uint32 GetTileIndex(uint8 vLevel, uint32 vAddress) const
-	{
-		check(vLevel < NumMips);
-		const uint32 TileIndex = TileIndexPerMip[vLevel] + vAddress * NumLayers;
-		if (TileIndex >= TileIndexPerMip[vLevel + 1])
-		{
-			// vAddress is out of bounds for this texture/mip level
-			return ~0u;
-		}
-		return TileIndex;
-	}
+	/** Returns false if the address isn't inside the texture bounds. */
+	bool IsValidAddress(uint32 vLevel, uint32 vAddress);
 
-	/**
-	 * Return the index of the chunk that contains the given tile
-	 */
-	inline int32 GetChunkIndex(uint32 TileIndex) const
-	{
-		if (TileIndex >= (uint32)TileOffsetInChunk.Num())
-		{
-			return -1;
-		}
-
-		for (int ChunkIndex = 0; ChunkIndex < TileIndexPerChunk.Num(); ++ChunkIndex)
-		{
-			if (TileIndex < TileIndexPerChunk[ChunkIndex + 1])
-			{
-				const uint32 NextTileOffset = GetTileOffset(ChunkIndex, TileIndex + GetNumLayers());
-				if (TileOffsetInChunk[TileIndex] == NextTileOffset)
-				{
-					// Size of the tile is 0, this means the given tile is not valid for this virtual texture
-					// This can happen since tile offset are stored with morton encoding, so non-square VTs will have some empty tiles allocated
-					return -1;
-				}
-
-				return ChunkIndex;
-			}
-		}
-
-		checkNoEntry();
-		return -1;
-	}
+	/** Return the index of the chunk that contains the given mip level. Returns -1 if the mip level doesn't exist. */
+	int32 GetChunkIndex(uint8 vLevel) const;
 
-	/**
-	* Return the offset of this tile within the chunk
-	*/
-	inline uint32 GetTileOffset(uint32 ChunkIndex, uint32 TileIndex) const
-	{
-		check(TileIndex >= TileIndexPerChunk[ChunkIndex]);
-		if (TileIndex < TileIndexPerChunk[ChunkIndex + 1])
-		{
-			return TileOffsetInChunk[TileIndex];
-		}
-
-		// If TileIndex is past the end of chunk, return the size of chunk
-		// This allows us to determine size of region by asking for start/end offsets
-		return Chunks[ChunkIndex].SizeInBytes;
-	}
+	/** Return the byte offset of a tile within a chunk. Returns ~0u if the tile doesn't exist. */
+	uint32 GetTileOffset(uint32 vLevel, uint32 vAddress, uint32 LayerIndex) const;
 
 	/**
 	* Validates the VT data
 	* If bValidateCompression is true, attempts to decompress every VT tile, to verify data is valid. This will be very slow
 	*/
 	bool ValidateData(FStringView const& InDDCDebugContext, bool bValidateCompression) const;
+
+private:
+	bool IsLegacyData() const;
+	uint32 GetTileIndex_Legacy(uint8 vLevel, uint32 vAddress) const;
+	uint32 GetTileOffset_Legacy(uint32 ChunkIndex, uint32 TileIndex) const;
 };
diff --git a/Engine/Source/Runtime/Engine/Private/VT/VirtualTextureChunkManager.cpp b/Engine/Source/Runtime/Engine/Private/VT/VirtualTextureChunkManager.cpp
--- a/Engine/Source/Runtime/Engine/Private/VT/VirtualTextureChunkManager.cpp
+++ b/Engine/Source/Runtime/Engine/Private/VT/VirtualTextureChunkManager.cpp
@@ -67,9 +67,8 @@ FVTRequestPageResult FVirtualTextureChunkStreamingManager::RequestTile(FUploadin
 	SCOPE_CYCLE_COUNTER(STAT_VTP_RequestTile);
 
 	const FVirtualTextureBuiltData* VTData = VTexture->GetVTData();
-	const uint32 TileIndex = VTData->GetTileIndex(vLevel, vAddress);
-	const int32 ChunkIndex = VTData->GetChunkIndex(TileIndex);
-	if (ChunkIndex == -1)
+	const uint32 TileOffset = VTData->GetTileOffset(vLevel, vAddress, 0);
+	if (TileOffset == ~0u)
 	{
 		static const auto CVar = IConsoleManager::Get().FindTConsoleVariableDataInt(TEXT("r.VT.Verbose"));
 		if (CVar->GetValueOnRenderThread())
@@ -80,6 +79,9 @@ FVTRequestPageResult FVirtualTextureChunkStreamingManager::RequestTile(FUploadin
 		return EVTRequestPageStatus::Invalid;
 	}
 
+	const int32 ChunkIndex = VTData->GetChunkIndex(vLevel);
+	check(ChunkIndex >= 0);
+
 	// tile is being transcoded/is done transcoding
 	const FVTTranscodeKey TranscodeKey = FVirtualTextureTranscodeCache::GetKey(ProducerHandle, LayerMask, vLevel, vAddress);
 	const FVTTranscodeTileHandleAndStatus TranscodeTaskResult = TranscodeCache.FindTask(TranscodeKey);
@@ -117,8 +119,8 @@ FVTRequestPageResult FVirtualTextureChunkStreamingManager::RequestTile(FUploadin
 	}
 
 	// make a single read request that covers region of all requested tiles
-	const uint32 OffsetStart = VTData->GetTileOffset(ChunkIndex, TileIndex + MinLayerIndex);
-	const uint32 OffsetEnd = VTData->GetTileOffset(ChunkIndex, TileIndex + MaxLayerIndex + 1u);
+	const uint32 OffsetStart = VTData->GetTileOffset(vLevel, vAddress, MinLayerIndex);
+	const uint32 OffsetEnd = VTData->GetTileOffset(vLevel, vAddress, MaxLayerIndex + 1u);
 	const uint32 RequestSize = OffsetEnd - OffsetStart;
 
 	const FVTDataAndStatus TileDataResult = VTexture->ReadData(GraphCompletionEvents, ChunkIndex, OffsetStart, RequestSize, Priority);
diff --git a/Engine/Source/Runtime/Engine/Private/VT/VirtualTextureDataBuilder.cpp b/Engine/Source/Runtime/Engine/Private/VT/VirtualTextureDataBuilder.cpp
--- a/Engine/Source/Runtime/Engine/Private/VT/VirtualTextureDataBuilder.cpp
+++ b/Engine/Source/Runtime/Engine/Private/VT/VirtualTextureDataBuilder.cpp
@@ -352,9 +352,15 @@ void FVirtualTextureDataBuilder::Build(const FTextureSourceData& InSourceData, c
 	OutData.WidthInBlocks = InSourceData.SizeInBlocksX;
 	OutData.HeightInBlocks = InSourceData.SizeInBlocksY;
 
+	OutData.TileDataOffsetPerLayer.Empty();
+	OutData.ChunkIndexPerMip.Empty();
+	OutData.BaseOffsetPerMip.Empty();
+	OutData.TileOffsetData.Empty();
+
 	OutData.TileIndexPerChunk.Empty();
 	OutData.TileIndexPerMip.Empty();
 	OutData.TileOffsetInChunk.Empty();
+
 	OutData.Chunks.Empty();
 
 	const uint32 Size = FMath::Max(SizeX, SizeY);
@@ -384,6 +390,31 @@ void FVirtualTextureDataBuilder::BuildPagesForChunk(const TArray<FVTSourceTileEn
 		BuildTiles(ActiveTileList, LayerIndex, LayerData[LayerIndex], bAllowAsync);
 	}
 
+	// Fill out tile offsets per layer if we haven't yet and if all layers are raw uncompressed data.
+	if (OutData.TileDataOffsetPerLayer.Num() == 0)
+	{
+		bool bIsRawGPUData = true;
+		for (int32 LayerIndex = 0; LayerIndex < LayerData.Num(); LayerIndex++)
+		{
+			if (LayerData[LayerIndex].Codec != EVirtualTextureCodec::RawGPU)
+			{
+				bIsRawGPUData = false;
+				break;
+			}
+		}
+		if (bIsRawGPUData)
+		{
+			uint32 TileDataOffset = 0;
+			OutData.TileDataOffsetPerLayer.Reserve(LayerData.Num());
+			for (int32 LayerIndex = 0; LayerIndex < LayerData.Num(); LayerIndex++)
+			{
+				TileDataOffset += LayerData[LayerIndex].TilePayload[0].Num();
+				OutData.TileDataOffsetPerLayer.Add(TileDataOffset);
+			}
+		}
+	}
+
+	// Write tiles out to chunk.
 	PushDataToChunk(ActiveTileList, LayerData);
 }
 
@@ -409,17 +440,20 @@ void FVirtualTextureDataBuilder::BuildPagesMacroBlocks(bool bAllowAsync)
 		MipHeightInTiles = FMath::DivideAndRoundUp(MipHeightInTiles, 2u);
 	}
 
+	// loop over each macro block and assemble the tiles
 	FScopedSlowTask BuildTask(NumTiles);
 
-	//
 	TArray<FVTSourceTileEntry> TilesInChunk;
 	TilesInChunk.Reserve(NumTiles);
 
-	// loop over each macro block and assemble the tiles
 	{
 		uint32 TileIndex = 0u;
 		bool bInFinalChunk = false;
 
+		OutData.ChunkIndexPerMip.Reserve(OutData.NumMips);
+		OutData.BaseOffsetPerMip.Reserve(OutData.NumMips);
+		OutData.TileOffsetData.Reserve(OutData.NumMips);
+
 		OutData.TileOffsetInChunk.Init(~0u, NumTiles * NumLayers);
 		OutData.TileIndexPerChunk.Reserve(OutData.NumMips + 1);
 		OutData.TileIndexPerMip.Reserve(OutData.NumMips + 1);
@@ -430,14 +464,20 @@ void FVirtualTextureDataBuilder::BuildPagesMacroBlocks(bool bAllowAsync)
 		MipHeightInTiles = FMath::DivideAndRoundUp(SizeY, TileSize);
 		for (uint32 Mip = 0; Mip < OutData.NumMips; ++Mip)
 		{
+			FVirtualTextureTileOffsetData& OffsetData = OutData.TileOffsetData.AddDefaulted_GetRef();
+			OffsetData.Init(MipWidthInTiles, MipHeightInTiles);
+
+			OutData.ChunkIndexPerMip.Add(OutData.Chunks.Num());
+			OutData.TileIndexPerMip.Add(TileIndex);
+
 			const int32 MipBlockSizeInTilesX = FMath::Max(BlockSizeInTilesX >> Mip, 1);
 			const int32 MipBlockSizeInTilesY = FMath::Max(BlockSizeInTilesY >> Mip, 1);
 			const uint32 MaxTileInMip = FMath::MortonCode2(MipWidthInTiles - 1) | (FMath::MortonCode2(MipHeightInTiles - 1) << 1);
 
-			OutData.TileIndexPerMip.Add(TileIndex);
-
 			for (uint32 TileIndexInMip = 0u; TileIndexInMip <= MaxTileInMip; ++TileIndexInMip)
 			{
+				BuildTask.EnterProgressFrame();
+
 				const uint32 TileX = FMath::ReverseMortonCode2(TileIndexInMip);
 				const uint32 TileY = FMath::ReverseMortonCode2(TileIndexInMip >> 1);
 				if (TileX < MipWidthInTiles && TileY < MipHeightInTiles)
@@ -448,20 +488,23 @@ void FVirtualTextureDataBuilder::BuildPagesMacroBlocks(bool bAllowAsync)
 					const int32 BlockIndex = FindSourceBlockIndex(Mip, BlockX, BlockY);
 					if (BlockIndex != INDEX_NONE)
 					{
-						BuildTask.EnterProgressFrame();
-
 						const FTextureSourceBlockData& Block = SourceBlocks[BlockIndex];
 						FVTSourceTileEntry* TileEntry = new(TilesInChunk) FVTSourceTileEntry;
 						TileEntry->BlockIndex = BlockIndex;
 						TileEntry->TileIndex = TileIndex;
+						TileEntry->MipIndex = Mip;
 						TileEntry->MipIndexInBlock = Mip - Block.MipBias;
 						TileEntry->TileInBlockX = TileX - Block.BlockX * MipBlockSizeInTilesX;
 						TileEntry->TileInBlockY = TileY - Block.BlockY * MipBlockSizeInTilesY;
+
+						OffsetData.AddTile(TileIndexInMip);
 					}
 				}
 				TileIndex += NumLayers;
 			}
 
+			OffsetData.Finalize();
+
 			if (!bInFinalChunk && TilesInChunk.Num() >= (int32)MinTilesPerChunk)
 			{
 				OutData.TileIndexPerChunk.Add(TileIndex);
@@ -485,31 +528,49 @@ void FVirtualTextureDataBuilder::BuildPagesMacroBlocks(bool bAllowAsync)
 		{
 			BuildPagesForChunk(TilesInChunk, bAllowAsync);
 		}
+
+		check(OutData.BaseOffsetPerMip.Num() == OutData.NumMips);
 	}
 
-	// Patch holes left in offset array
-	for (int32 ChunkIndex = 0; ChunkIndex < OutData.Chunks.Num(); ++ChunkIndex)
+	// Use compact tile offsets if we have fixed tile sizes on every layer (raw GPU codecs).
+	// Otherwise use legacy data.
+	const bool bUseLegacyData = OutData.TileDataOffsetPerLayer.Num() != NumLayers;
+	if (bUseLegacyData)
 	{
-		uint32 CurrentOffset = OutData.Chunks[ChunkIndex].SizeInBytes;
-		for (int32 TileIndex = OutData.TileIndexPerChunk[ChunkIndex + 1] - 1u; TileIndex >= (int32)OutData.TileIndexPerChunk[ChunkIndex]; --TileIndex)
+		// Using legacy data from now on so remove the compact data.
+		OutData.TileOffsetData.Empty();
+
+		// Patch holes left in offset array
+		for (int32 ChunkIndex = 0; ChunkIndex < OutData.Chunks.Num(); ++ChunkIndex)
 		{
-			const uint32 TileOffset = OutData.TileOffsetInChunk[TileIndex];
-			if (TileOffset > CurrentOffset)
+			uint32 CurrentOffset = OutData.Chunks[ChunkIndex].SizeInBytes;
+			for (int32 TileIndex = OutData.TileIndexPerChunk[ChunkIndex + 1] - 1u; TileIndex >= (int32)OutData.TileIndexPerChunk[ChunkIndex]; --TileIndex)
 			{
-				check(TileOffset == ~0u);
-				OutData.TileOffsetInChunk[TileIndex] = CurrentOffset;
-			}
-			else
-			{
-				CurrentOffset = TileOffset;
+				const uint32 TileOffset = OutData.TileOffsetInChunk[TileIndex];
+				if (TileOffset > CurrentOffset)
+				{
+					check(TileOffset == ~0u);
+					OutData.TileOffsetInChunk[TileIndex] = CurrentOffset;
+				}
+				else
+				{
+					CurrentOffset = TileOffset;
+				}
 			}
 		}
-	}
 
-	for (int32 TileIndex = 0u; TileIndex < OutData.TileOffsetInChunk.Num(); ++TileIndex)
+		for (int32 TileIndex = 0u; TileIndex < OutData.TileOffsetInChunk.Num(); ++TileIndex)
+		{
+			const uint32 TileOffset = OutData.TileOffsetInChunk[TileIndex];
+			check(TileOffset != ~0u);
+		}
+	}
+	else
 	{
-		const uint32 TileOffset = OutData.TileOffsetInChunk[TileIndex];
-		check(TileOffset != ~0u);
+		// We can remove legacy data and only reference the compact data from now on.
+		OutData.TileIndexPerChunk.Empty();
+		OutData.TileIndexPerMip.Empty();
+		OutData.TileOffsetInChunk.Empty();
 	}
 }
 
@@ -786,7 +847,13 @@ void FVirtualTextureDataBuilder::PushDataToChunk(const TArray<FVTSourceTileEntry
 	for (int32 TileIdx = 0; TileIdx < Tiles.Num(); ++TileIdx)
 	{
 		const FVTSourceTileEntry& Tile = Tiles[TileIdx];
-		uint32 TileIndex = Tile.TileIndex;
+		const int32 MipIndex = Tile.MipIndex;
+		if (MipIndex >= OutData.BaseOffsetPerMip.Num())
+		{
+			check(MipIndex == OutData.BaseOffsetPerMip.Num());
+			OutData.BaseOffsetPerMip.Add(ChunkOffset);
+		}
+		int32 TileIndex = Tile.TileIndex;
 		for (int32 Layer = 0; Layer < NumLayers; ++Layer)
 		{
 			check(OutData.TileOffsetInChunk[TileIndex] == ~0u);
diff --git a/Engine/Source/Runtime/Engine/Private/VT/VirtualTextureDataBuilder.h b/Engine/Source/Runtime/Engine/Private/VT/VirtualTextureDataBuilder.h
--- a/Engine/Source/Runtime/Engine/Private/VT/VirtualTextureDataBuilder.h
+++ b/Engine/Source/Runtime/Engine/Private/VT/VirtualTextureDataBuilder.h
@@ -16,6 +16,7 @@ struct FVTSourceTileEntry
 {
 	int32 BlockIndex;
 	int32 TileIndex;
+	int32 MipIndex;
 	int32 MipIndexInBlock;
 	int32 TileInBlockX;
 	int32 TileInBlockY;
diff --git a/Engine/Source/Runtime/Engine/Private/VT/VirtualTextureTranscodeCache.cpp b/Engine/Source/Runtime/Engine/Private/VT/VirtualTextureTranscodeCache.cpp
--- a/Engine/Source/Runtime/Engine/Private/VT/VirtualTextureTranscodeCache.cpp
+++ b/Engine/Source/Runtime/Engine/Private/VT/VirtualTextureTranscodeCache.cpp
@@ -96,7 +96,6 @@ struct FTranscodeTask
 		const uint32 NumLayers = Params.VTData->GetNumLayers();
 		const uint8 vLevel = Params.vLevel;
 		const uint32 vAddress = Params.vAddress;
-		const uint32 TileIndex = Params.VTData->GetTileIndex(vLevel, vAddress);
 
 		// code must be fully loaded by the time we start transcoding
 		check(!Params.Codec || Params.Codec->IsCreationComplete());
@@ -116,8 +115,8 @@ struct FTranscodeTask
 				continue;
 			}
 
-			const uint32 TileLayerOffset = Params.VTData->GetTileOffset(ChunkIndex, TileIndex + LayerIndex);
-			const uint32 NextTileLayerOffset = Params.VTData->GetTileOffset(ChunkIndex, TileIndex + LayerIndex + 1u);
+			const uint32 TileLayerOffset = Params.VTData->GetTileOffset(vLevel, vAddress,  LayerIndex);
+			const uint32 NextTileLayerOffset = Params.VTData->GetTileOffset(vLevel, vAddress, LayerIndex + 1u);
 			if (TileBaseOffset == ~0u)
 			{
 				TileBaseOffset = TileLayerOffset;
@@ -406,7 +405,6 @@ FVTTranscodeTileHandle FVirtualTextureTranscodeCache::SubmitTask(FVirtualTexture
 	}
 
 	const uint32 ChunkIndex = InParams.ChunkIndex;
-	const uint32 TileIndex = InParams.VTData->GetTileIndex(InParams.vLevel, InParams.vAddress);
 	const FVirtualTextureDataChunk& Chunk = InParams.VTData->Chunks[ChunkIndex];
 	const uint32 TilePixelSize = InParams.VTData->GetPhysicalTileSize();
 	FVTUploadTileBuffer StagingBuffer[VIRTUALTEXTURE_SPACE_MAXLAYERS];
@@ -422,8 +420,8 @@ FVTTranscodeTileHandle FVirtualTextureTranscodeCache::SubmitTask(FVirtualTexture
 			const EPixelFormat LayerFormat = InParams.VTData->LayerTypes[LayerIndex];
 			FVTUploadTileBuffer& StagingBufferForLayer = StagingBuffer[LayerIndex];
 
-			const uint32 TileLayerOffset = InParams.VTData->GetTileOffset(ChunkIndex, TileIndex + LayerIndex);
-			const uint32 NextTileLayerOffset = InParams.VTData->GetTileOffset(ChunkIndex, TileIndex + LayerIndex + 1u);
+			const uint32 TileLayerOffset = InParams.VTData->GetTileOffset(InParams.vLevel, InParams.vAddress, LayerIndex);
+			const uint32 NextTileLayerOffset = InParams.VTData->GetTileOffset(InParams.vLevel, InParams.vAddress, LayerIndex + 1u);
 			if (TileBaseOffset == ~0u)
 			{
 				TileBaseOffset = TileLayerOffset;
