Disable mobile msaa if DepthFullPrepass is enabled.

#p4v-preflight-copy 17383727

diff --git a/Engine/Source/Runtime/Engine/Private/SceneUtils.cpp b/Engine/Source/Runtime/Engine/Private/SceneUtils.cpp
--- a/Engine/Source/Runtime/Engine/Private/SceneUtils.cpp
+++ b/Engine/Source/Runtime/Engine/Private/SceneUtils.cpp
@@ -141,19 +141,18 @@ ENGINE_API uint32 GetDefaultMSAACount(const FStaticFeatureLevel InFeatureLevel,
 
 			if (InFeatureLevel == ERHIFeatureLevel::ES3_1)
 			{
-				// Disable MSAA if we are using mobile pixel projected reflection, since we have to resolve the SceneColor and SceneDepth after opaque base pass
 				bool bMobilePixelProjectedReflection = IsUsingMobilePixelProjectedReflection(ShaderPlatform);
+				
+				static const auto MobileEarlyZPassCVar = IConsoleManager::Get().FindTConsoleVariableDataInt(TEXT("r.Mobile.EarlyZPass"));
+				bool bMobileAmbientOcclusion = UseMobileAmbientOcclusion(ShaderPlatform) && IsMobileHDR();
+				bool bMobileDistanceFieldEnabled = IsMobileDistanceFieldEnabled(ShaderPlatform);
+				bool bIsFullDepthPrepassEnabled = bMobileDistanceFieldEnabled || bMobileAmbientOcclusion || MobileEarlyZPassCVar->GetValueOnAnyThread() > 0;
 
-				// Disable MSAA if we are using mobile ambient occlusion, since we have to resolve the SceneColor and SceneDepth after opaque base pass
-				static const auto MobileAmbientOcclusionQualityCVar = IConsoleManager::Get().FindTConsoleVariableDataInt(TEXT("r.Mobile.AmbientOcclusionQuality"));
-
-				bool bMobileAmbientOcclusion = UseMobileAmbientOcclusion(ShaderPlatform) && IsMobileHDR() && MobileAmbientOcclusionQualityCVar->GetValueOnAnyThread() > 0;
-
-				bRendererSupportMSAA = bRHISupportsMSAA && !bMobilePixelProjectedReflection && !bMobileAmbientOcclusion;
+				bRendererSupportMSAA = bRHISupportsMSAA && !bMobilePixelProjectedReflection && !bIsFullDepthPrepassEnabled;
 
 				if (!bRendererSupportMSAA)
 				{
-					FailedReason = FString::Printf(TEXT("RHISupportsMSAA %d, MobilePixelProjectedReflection %d, MobileAmbientOcclusion %d"), bRHISupportsMSAA ? 1 : 0, bMobilePixelProjectedReflection ? 1 : 0, bMobileAmbientOcclusion ? 1 : 0);
+					FailedReason = FString::Printf(TEXT("RHISupportsMSAA %d, MobilePixelProjectedReflection %d, MobileAmbientOcclusion %d, MobileDistanceFieldEnabled %d, MobileEarlyZPass %d"), bRHISupportsMSAA ? 1 : 0, bMobilePixelProjectedReflection ? 1 : 0, bMobileAmbientOcclusion ? 1 : 0, bMobileDistanceFieldEnabled ? 1 : 0, MobileEarlyZPassCVar->GetValueOnAnyThread());
 				}
 			}
 			else
