Fixes the loading path of the global shader map on desktop cooked build.

diff --git a/Engine/Source/Runtime/Engine/Private/ShaderCompiler/ShaderCompiler.cpp b/Engine/Source/Runtime/Engine/Private/ShaderCompiler/ShaderCompiler.cpp
--- a/Engine/Source/Runtime/Engine/Private/ShaderCompiler/ShaderCompiler.cpp
+++ b/Engine/Source/Runtime/Engine/Private/ShaderCompiler/ShaderCompiler.cpp
@@ -6181,7 +6181,7 @@ static FString GGlobalShaderCacheOverrideDirectory;
 
 static FString GetGlobalShaderCacheOverrideFilename(EShaderPlatform Platform)
 {
-	FString DirectoryPrefix = FString(TEXT("Engine")) / TEXT("OverrideGlobalShaderCache-");
+	FString DirectoryPrefix = FPaths::EngineDir() / TEXT("OverrideGlobalShaderCache-");
 
 	if (!GGlobalShaderCacheOverrideDirectory.IsEmpty())
 	{
