Add Device Profile fragment selection to device profiles.
A device profile can can test for device features and capabilities and then include additional fragments when such conditions are satisfied.
added +MatchRule to DPs, this contains the if conditions and the resultant fragments to select.
Fragment cvars are applied after the device profile.
Selected fragments are applied in order, they will overwrite the cvars of prior fragments.
The pre-device profile cvar state is always recorded and used to 'undo' the DP before any reapplying or changing any DP state (such as SetOverrideDeviceProfile)

diff --git a/Engine/Source/Editor/UnrealEd/Private/EditorEngine.cpp b/Engine/Source/Editor/UnrealEd/Private/EditorEngine.cpp
--- a/Engine/Source/Editor/UnrealEd/Private/EditorEngine.cpp
+++ b/Engine/Source/Editor/UnrealEd/Private/EditorEngine.cpp
@@ -7472,7 +7472,7 @@ void UEditorEngine::SetPreviewPlatform(const FPreviewPlatformInfo& NewPreviewPla
 
 	Scalability::ChangeScalabilityPreviewPlatform(PreviewPlatform.GetEffectivePreviewPlatformName());
 
-	UDeviceProfileManager::Get().RestoreDefaultDeviceProfile();
+	UDeviceProfileManager::Get().RestorePreviewDeviceProfile();
 
 	UStaticMesh::OnLodStrippingQualityLevelChanged(nullptr);
 
@@ -7481,7 +7481,7 @@ void UEditorEngine::SetPreviewPlatform(const FPreviewPlatformInfo& NewPreviewPla
 	{
 		if (UDeviceProfile* DP = UDeviceProfileManager::Get().FindProfile(PreviewPlatform.DeviceProfileName.ToString(), false))
 		{
-			UDeviceProfileManager::Get().SetOverrideDeviceProfile(DP, true);
+			UDeviceProfileManager::Get().SetPreviewDeviceProfile(DP);
 		}
 	}
 
@@ -7510,17 +7510,17 @@ void UEditorEngine::ToggleFeatureLevelPreview()
 		{
 			if (UDeviceProfile* DP = UDeviceProfileManager::Get().FindProfile(PreviewPlatform.DeviceProfileName.ToString(), false))
 			{
-				UDeviceProfileManager::Get().SetOverrideDeviceProfile(DP, true);
+				UDeviceProfileManager::Get().SetPreviewDeviceProfile(DP);
 			}
 		}
 		else
 		{
-			UDeviceProfileManager::Get().RestoreDefaultDeviceProfile();
+			UDeviceProfileManager::Get().RestorePreviewDeviceProfile();
 		}
 	}
 	else
 	{
-		UDeviceProfileManager::Get().RestoreDefaultDeviceProfile();
+		UDeviceProfileManager::Get().RestorePreviewDeviceProfile();
 	}
 
 	PreviewPlatformChanged.Broadcast();
diff --git a/Engine/Source/Runtime/Engine/Classes/DeviceProfiles/DeviceProfile.h b/Engine/Source/Runtime/Engine/Classes/DeviceProfiles/DeviceProfile.h
--- a/Engine/Source/Runtime/Engine/Classes/DeviceProfiles/DeviceProfile.h
+++ b/Engine/Source/Runtime/Engine/Classes/DeviceProfiles/DeviceProfile.h
@@ -9,6 +9,7 @@
 #include "CoreMinimal.h"
 #include "UObject/ObjectMacros.h"
 #include "Engine/TextureLODSettings.h"
+#include "DeviceProfiles/DeviceProfileMatching.h"
 #include "DeviceProfile.generated.h"
 
 struct FPropertyChangedEvent;
@@ -89,6 +90,10 @@ private:
 	/** Delegate object fired when there has been any changes to the console variables */
 	FOnCVarsUpdated CVarsUpdatedDelegate;
 
+	// An array of conditions to test against and fragment names to select.
+	UPROPERTY(EditAnywhere, config, Category = "DeviceProfile Matching Rules")
+	TArray<FDPMatchingRulestruct> MatchingRules;
+
 public:
 	/* ValidateProfile()
 	* Validate the Profile after changes by loading it's config (.ini)
@@ -108,7 +113,9 @@ public:
 	bool ModifyCVarValue(const FString& CVarName, const FString& CVarValue, bool bAddIfNonExistant = false);
 	FString GetCVarValue(const FString& CVarName) const;
 
-	/** Lazily generate a consolidated list of CVars, recursing up the device profile hierarchy */
+	/** Lazily generate a consolidated list of CVars, recursing up the device profile hierarchy 
+	 *  This will not include any cvars from the device's selected fragments.
+	*/
 	const TMap<FString, FString>& GetConsolidatedCVars() const;
 
 	/** 
diff --git a/Engine/Source/Runtime/Engine/Classes/DeviceProfiles/DeviceProfileFragment.h b/Engine/Source/Runtime/Engine/Classes/DeviceProfiles/DeviceProfileFragment.h
--- a/Engine/Source/Runtime/Engine/Classes/DeviceProfiles/DeviceProfileFragment.h
+++ b/Engine/Source/Runtime/Engine/Classes/DeviceProfiles/DeviceProfileFragment.h
@@ -15,5 +15,6 @@ class ENGINE_API UDeviceProfileFragment : public UObject
 {
 	GENERATED_UCLASS_BODY()
 public:
+	UPROPERTY()
 	TArray<FString> CVars;
 };
diff --git a/Engine/Source/Runtime/Engine/Classes/DeviceProfiles/DeviceProfileManager.h b/Engine/Source/Runtime/Engine/Classes/DeviceProfiles/DeviceProfileManager.h
--- a/Engine/Source/Runtime/Engine/Classes/DeviceProfiles/DeviceProfileManager.h
+++ b/Engine/Source/Runtime/Engine/Classes/DeviceProfiles/DeviceProfileManager.h
@@ -9,6 +9,8 @@
 #include "CoreMinimal.h"
 #include "UObject/ObjectMacros.h"
 #include "UObject/Object.h"
+#include "DeviceProfileMatching.h"
+#include "CoreGlobals.h"
 #include "DeviceProfileManager.generated.h"
 
 class UDeviceProfile;
@@ -33,7 +35,7 @@ public:
 	 * Startup and select the active device profile
 	 * Then Init the CVars from this profile and it's Device profile parent tree.
 	 */
-	static void InitializeCVarsForActiveDeviceProfile(bool bPushSettings=false);
+	static void InitializeCVarsForActiveDeviceProfile();
 
 	/**
 	 * Reapplies the device profile. Useful when configs have changed (i.e. hotfix)
@@ -77,13 +79,26 @@ public:
 	/**
 	* Overrides the device profile. The original profile can be restored with RestoreDefaultDeviceProfile
 	*/
-	void SetOverrideDeviceProfile(UDeviceProfile* DeviceProfile, bool bIsDeviceProfilePreview = false);
+	void SetOverrideDeviceProfile(UDeviceProfile* DeviceProfile);
 
 	/**
 	* Restore the device profile to the default for this device
 	*/
 	void RestoreDefaultDeviceProfile();
 
+#if ALLOW_OTHER_PLATFORM_CONFIG
+	/**
+	* Applies the requested device profile's cvars onto the currently active DP.
+	* It does not change the actual DP. 
+	* Use RestorePreviewDeviceProfile to revert.
+	*/
+	void SetPreviewDeviceProfile(UDeviceProfile* DeviceProfile);
+
+	/**
+	* Revert the preview state.
+	*/
+	void RestorePreviewDeviceProfile();
+#endif
 
 	/**
 	 * Load the device profiles from the config file.
@@ -157,6 +172,19 @@ public:
 	static void ExpandDeviceProfileCVars(UDeviceProfile* DeviceProfile);
 #endif
 
+	/**
+	* Enable/Disable a tagged fragment of the active device profile.
+	* This unsets the entire cvar DP state, then re-sets the new DP+fragment state.
+	*/
+	void ChangeTaggedFragmentState(FName FragmentTag, bool bNewState);
+
+	/**
+	* Return the selected fragment property from the currently active device profile.
+	* null if the tag is not found.
+	*/
+	const FSelectedFragmentProperties* GetActiveDeviceProfileFragmentByTag(FName& FragmentTag) const;
+
+
 private:
 	/**
 	 * Set the active device profile - set via the device profile blueprint.
@@ -170,14 +198,9 @@ private:
 	*/
 	void HandleDeviceProfileOverrideChange();
 
-	/** Handle restoing CVars set in HandleDeviceProfileOverrideChange */
-	void HandleDeviceProfileOverridePop();
-
-
 	enum class EDeviceProfileMode : uint8
 	{
 		DPM_SetCVars,
-		DPM_PushCVars,
 		DPM_CacheValues,
 	};
 
@@ -187,6 +210,14 @@ private:
 	 */
 	static void ProcessDeviceProfileIniSettings(const FString& DeviceProfileName, EDeviceProfileMode Mode);
 
+
+	/** Read and process all of the fragment matching rules. Returns an array containing the names of fragments selected. */
+	static TArray<FSelectedFragmentProperties> FindMatchingFragments(const FString& ParentDP, class FConfigCacheIni* PreviewConfigSystem);
+
+	/** Get the current platform's the selector module. Can return null */
+	static class IDeviceProfileSelectorModule* GetDeviceProfileSelectorModule();
+	/** Get another platform's selector module. Can return null */
+	static class IDeviceProfileSelectorModule* GetPreviewDeviceProfileSelectorModule(class FConfigCacheIni* PreviewConfigSystem);
 public:
 
 	static class UDeviceProfileManager* DeviceProfileManagerSingleton;
@@ -218,12 +249,20 @@ private:
 	// Holds the device profile .ini location
 	static FString DeviceProfileFileName;
 
-	// values of CVars set in HandleDeviceProfileOverrideChange, to be popped later
-	TMap<FString, FString> PushedSettings;
+	// Original values of all the CVars modified by the DP.
+	// Used to undo the DP before applying new state.
+	static TMap<FString, FString> PushedSettings;
 
+#if ALLOW_OTHER_PLATFORM_CONFIG
+	// Original values of the CVars modified by SetPreviewDeviceProfile.
+	TMap<FString, FString> PreviewPushedSettings;
+#endif
 	// Holds the device profile that has been overridden, null no override active.
 	UDeviceProfile* BaseDeviceProfile = nullptr;
 
 	// Stores any scalability group settings set by the active device profile.
 	static TMap<FString, FString> DeviceProfileScalabilityCVars;
+	
+	// The list of fragments that have been selected by the active profile.
+	static TArray<FSelectedFragmentProperties> PlatformFragmentsSelected;
 };
diff --git a/Engine/Source/Runtime/Engine/Classes/DeviceProfiles/DeviceProfileMatching.h b/Engine/Source/Runtime/Engine/Classes/DeviceProfiles/DeviceProfileMatching.h
new file mode 100644
--- /dev/null
+++ b/Engine/Source/Runtime/Engine/Classes/DeviceProfiles/DeviceProfileMatching.h
@@ -0,0 +1,165 @@
+// Copyright Epic Games, Inc. All Rights Reserved.
+
+/*=============================================================================
+	DeviceProfileFragment.h: Declares the UDeviceProfileFragment class.
+=============================================================================*/
+
+#pragma once
+
+#include "CoreMinimal.h"
+#include "UObject/ObjectMacros.h"
+#include "DeviceProfileMatching.generated.h"
+
+USTRUCT()
+struct ENGINE_API FSelectedFragmentProperties
+{
+	GENERATED_USTRUCT_BODY()
+
+	bool operator == (const FSelectedFragmentProperties& rhs) const
+	{
+		return Tag == rhs.Tag && Fragment == rhs.Fragment && bEnabled == rhs.bEnabled;
+	}
+
+	bool operator != (const FSelectedFragmentProperties& rhs) const
+	{
+		return !(*this == rhs);
+	}
+
+	// user defined 'Tag' for this fragment.
+	UPROPERTY()
+	FName Tag = NAME_None;
+
+	// Actual name of the fragment to select
+	UPROPERTY()
+	FString Fragment;
+
+	// whether the fragment's cvars will be applied
+	UPROPERTY()
+	bool bEnabled = true;
+};
+
+USTRUCT()
+struct FDPMatchingIfCondition
+{
+	GENERATED_USTRUCT_BODY()
+
+	UPROPERTY()
+	FName Operator;
+
+	UPROPERTY()
+	FString Arg1;
+
+	UPROPERTY()
+	FString Arg2;
+};
+
+USTRUCT()
+struct FDPMatchingRulestructBase
+{
+	GENERATED_USTRUCT_BODY()
+	virtual ~FDPMatchingRulestructBase() {};
+
+	UPROPERTY()
+	FString RuleName;
+
+	UPROPERTY()
+	TArray<FDPMatchingIfCondition> IfConditions;
+
+	UPROPERTY()
+	FString AppendFragments;
+
+	UPROPERTY()
+	FString SetSrc;
+
+	virtual const FDPMatchingRulestructBase* GetOnTrue() const { return nullptr; };
+	virtual const FDPMatchingRulestructBase* GetOnFalse() const { return nullptr; };
+};
+
+USTRUCT()
+	struct FDPMatchingRulestructA : public FDPMatchingRulestructBase
+{
+	GENERATED_USTRUCT_BODY()
+
+	UPROPERTY(config)
+	TArray<struct FDPMatchingRulestructBase> OnTrue;
+
+	UPROPERTY(config)
+	TArray<struct FDPMatchingRulestructBase> OnFalse;
+
+	virtual const FDPMatchingRulestructBase* GetOnTrue() const override { return OnTrue.Num() ? &OnTrue[0] : nullptr; };
+	virtual const FDPMatchingRulestructBase* GetOnFalse() const override { return OnFalse.Num() ? &OnFalse[0] : nullptr; };
+};
+
+USTRUCT()
+	struct FDPMatchingRulestructB : public FDPMatchingRulestructBase
+{
+	GENERATED_USTRUCT_BODY()
+	UPROPERTY(config)
+	TArray<struct FDPMatchingRulestructA> OnTrue;
+
+	UPROPERTY(config)
+	TArray<struct FDPMatchingRulestructA> OnFalse;
+
+	virtual const FDPMatchingRulestructBase* GetOnTrue() const override { return OnTrue.Num() ? &OnTrue[0] : nullptr; };
+	virtual const FDPMatchingRulestructBase* GetOnFalse() const override { return OnFalse.Num() ? &OnFalse[0] : nullptr; };
+};
+
+USTRUCT()
+	struct FDPMatchingRulestructC : public FDPMatchingRulestructBase
+{
+	GENERATED_USTRUCT_BODY()
+
+	UPROPERTY(config)
+	TArray<struct FDPMatchingRulestructB> OnTrue;
+	UPROPERTY(config)
+	TArray<struct FDPMatchingRulestructB> OnFalse;
+
+	virtual const FDPMatchingRulestructBase* GetOnTrue() const override { return OnTrue.Num() ? &OnTrue[0] : nullptr; };
+	virtual const FDPMatchingRulestructBase* GetOnFalse() const override { return OnFalse.Num() ? &OnFalse[0] : nullptr; };
+};
+
+
+USTRUCT()
+	struct FDPMatchingRulestructD : public FDPMatchingRulestructBase
+{
+	GENERATED_USTRUCT_BODY()
+
+	UPROPERTY(config)
+	TArray<struct FDPMatchingRulestructC> OnTrue;
+
+	UPROPERTY(config)
+	TArray<struct FDPMatchingRulestructC> OnFalse;
+
+	virtual const FDPMatchingRulestructBase* GetOnTrue() const override { return OnTrue.Num() ? &OnTrue[0] : nullptr; };
+	virtual const FDPMatchingRulestructBase* GetOnFalse() const override { return OnFalse.Num() ? &OnFalse[0] : nullptr; };
+};
+
+USTRUCT()
+	struct FDPMatchingRulestructE : public FDPMatchingRulestructBase
+{
+	GENERATED_USTRUCT_BODY()
+
+	UPROPERTY(config)
+	TArray<struct FDPMatchingRulestructD> OnTrue;
+
+	UPROPERTY(config)
+	TArray<struct FDPMatchingRulestructD> OnFalse;
+
+	virtual const FDPMatchingRulestructBase* GetOnTrue() const override { return OnTrue.Num() ? &OnTrue[0] : nullptr; };
+	virtual const FDPMatchingRulestructBase* GetOnFalse() const override { return OnFalse.Num() ? &OnFalse[0] : nullptr; };
+};
+
+USTRUCT()
+struct FDPMatchingRulestruct : public FDPMatchingRulestructBase
+{
+	GENERATED_USTRUCT_BODY()
+
+	UPROPERTY(config)
+	TArray<struct FDPMatchingRulestructE> OnTrue;
+
+	UPROPERTY(config)
+	TArray<struct FDPMatchingRulestructE> OnFalse;
+
+	virtual const FDPMatchingRulestructBase* GetOnTrue() const override { return OnTrue.Num() ? &OnTrue[0] : nullptr; };
+	virtual const FDPMatchingRulestructBase* GetOnFalse() const override { return OnFalse.Num() ? &OnFalse[0] : nullptr; };
+};
diff --git a/Engine/Source/Runtime/Engine/Private/DeviceProfiles/DeviceProfileManager.cpp b/Engine/Source/Runtime/Engine/Private/DeviceProfiles/DeviceProfileManager.cpp
--- a/Engine/Source/Runtime/Engine/Private/DeviceProfiles/DeviceProfileManager.cpp
+++ b/Engine/Source/Runtime/Engine/Private/DeviceProfiles/DeviceProfileManager.cpp
@@ -39,6 +39,8 @@ static TAutoConsoleVariable<int32> CVarAllowScalabilityGroupsToChangeAtRuntime(
 	ECVF_Default);
 
 TMap<FString, FString> UDeviceProfileManager::DeviceProfileScalabilityCVars;
+TMap<FString, FString> UDeviceProfileManager::PushedSettings;
+TArray<FSelectedFragmentProperties> UDeviceProfileManager::PlatformFragmentsSelected;
 
 UDeviceProfileManager* UDeviceProfileManager::DeviceProfileManagerSingleton = nullptr;
 
@@ -75,7 +77,7 @@ UDeviceProfileManager& UDeviceProfileManager::Get(bool bFromPostCDOContruct)
 			TEXT("Restores any cvars set by dp.Override to their previous value"),
 			FConsoleCommandDelegate::CreateLambda([]()
 			{
-				UDeviceProfileManager::Get().HandleDeviceProfileOverridePop();
+				UDeviceProfileManager::Get().RestoreDefaultDeviceProfile();
 			}),
 			ECVF_Default
 		);
@@ -85,7 +87,27 @@ UDeviceProfileManager& UDeviceProfileManager::Get(bool bFromPostCDOContruct)
 	return *DeviceProfileManagerSingleton;
 }
 
-static void GetFragmentCvars(const FString& CurrentSectionName, const FString& CVarArrayName, TArray<FString>& FragmentCVarsINOUT, FConfigCacheIni* ConfigSystem)
+// Read the cvars from a [DeviceProfileFragment] section.
+static bool GetFragmentCVars(const FString& FragmentName, const FString& CVarArrayName, TArray<FString>& FragmentCVarsINOUT, FConfigCacheIni* ConfigSystem)
+{
+	FString FragmentSectionName = FString::Printf(TEXT("%s %s"), *FragmentName, *UDeviceProfileFragment::StaticClass()->GetName());
+	if (ConfigSystem->DoesSectionExist(*FragmentSectionName, GDeviceProfilesIni))
+	{
+		TArray<FString> FragmentCVars;
+		ConfigSystem->GetArray(*FragmentSectionName, *CVarArrayName, FragmentCVars, GDeviceProfilesIni);
+		UE_CLOG(FragmentCVars.Num() > 0, LogInit, Log, TEXT("Including %s from fragment: %s"), *CVarArrayName, *FragmentName);
+		FragmentCVarsINOUT += FragmentCVars;
+	}
+	else
+	{
+		UE_LOG(LogInit, Error, TEXT("Could not find device profile fragment %s."), *FragmentName);
+		return false;
+	}
+	return true;
+}
+
+// read the requested fragment from within the +FragmentIncludes= array of a DP.
+static void GetCVarsFromDPFragmentIncludes(const FString& CurrentSectionName, const FString& CVarArrayName, TArray<FString>& FragmentCVarsINOUT, FConfigCacheIni* ConfigSystem)
 {
 	FString FragmentIncludes = TEXT("FragmentIncludes");
 	TArray<FString> FragmentIncludeArray;
@@ -94,21 +116,7 @@ static void GetFragmentCvars(const FString& CurrentSectionName, const FString& C
 	for(const FString& FragmentInclude : FragmentIncludeArray)
 	{
 		FString FragmentSectionName = FString::Printf(TEXT("%s %s"), *FragmentInclude, *UDeviceProfileFragment::StaticClass()->GetName());
-		if (ConfigSystem->DoesSectionExist(*FragmentSectionName, GDeviceProfilesIni))
-		{
-			TArray<FString> FragmentCVars;
-			ConfigSystem->GetArray(*FragmentSectionName, *CVarArrayName, FragmentCVars, GDeviceProfilesIni);
-			UE_CLOG(FragmentCVars.Num()>0, LogDeviceProfileManager, Log, TEXT("Including %s from fragment: %s"), *CVarArrayName, *FragmentInclude);
-			FragmentCVarsINOUT += FragmentCVars;
-		}
-		else
-		{
-#if UE_BUILD_SHIPPING
-			UE_LOG(LogDeviceProfileManager, Error, TEXT("Could not find device profile fragment %s."), *FragmentInclude);
-#else
-			UE_LOG(LogDeviceProfileManager, Fatal, TEXT("Could not find device profile fragment %s."), *FragmentInclude);
-#endif
-		}
+		GetFragmentCVars(FragmentSectionName, CVarArrayName, FragmentCVarsINOUT, ConfigSystem);
 	}
 }
 
@@ -149,6 +157,11 @@ void UDeviceProfileManager::ProcessDeviceProfileIniSettings(const FString& Devic
 #endif
 	}
 
+	if (Mode == EDeviceProfileMode::DPM_SetCVars)
+	{
+		UE_LOG(LogDeviceProfileManager, Log, TEXT("Applying CVar settings loaded from the selected device profile: [%s]"), *DeviceProfileName);
+	}
+
 	check(ConfigSystem);
 
 	TArray< FString > AvailableProfiles;
@@ -162,15 +175,12 @@ void UDeviceProfileManager::ProcessDeviceProfileIniSettings(const FString& Devic
 	TMap<FString, FString> CVarsAlreadySetList;
 
 	// reset some global state for "active DP" mode
-	if (Mode != EDeviceProfileMode::DPM_CacheValues)
+	if (Mode == EDeviceProfileMode::DPM_SetCVars)
 	{
 		DeviceProfileScalabilityCVars.Empty();
 
-		// even if we aren't pushing new values, we should clear any old pushed values, as they are no longer valid after we run this loop
-		if (DeviceProfileManagerSingleton)
-		{
-			DeviceProfileManagerSingleton->PushedSettings.Empty();
-		}
+		// we should have always pushed away old values by the time we get here
+		check(PushedSettings.Num() == 0);
 
 #if !UE_BUILD_SHIPPING
 #if PLATFORM_ANDROID
@@ -237,7 +247,7 @@ void UDeviceProfileManager::ProcessDeviceProfileIniSettings(const FString& Devic
 	TSet<FString> PreviewAllowlistCVars;
 	TSet<FString> PreviewDenylistCVars;
 	bool bFoundAllowDeny = false;
-	if (Mode == EDeviceProfileMode::DPM_PushCVars)
+	if (Mode == EDeviceProfileMode::DPM_CacheValues)
 	{
 		// Walk up the device profile tree to find the most specific device profile with a Denylist or Allowlist of cvars to apply, and use those Allow/Denylists.
 		for(FString CurrentProfileName = DeviceProfileName, CurrentSectionName = DeviceProfileName + SectionSuffix;
@@ -260,6 +270,54 @@ void UDeviceProfileManager::ProcessDeviceProfileIniSettings(const FString& Devic
 		}
 	}
 #endif
+
+
+	TSet<FString> FragmentCVarKeys;
+	TArray<FString> SelectedFragmentCVars;
+	TArray<FSelectedFragmentProperties> FragmentsSelected;
+
+	// Process the fragment matching rules.
+	// Only perform the matching process for the base DP for simplicity.
+	// Perform the matching process if no fragments were selected or we are previewing. re-use the matched array if already present.
+	if (Mode == EDeviceProfileMode::DPM_CacheValues || PlatformFragmentsSelected.Num() == 0)
+	{
+		// run the matching rules.
+		FragmentsSelected = FindMatchingFragments(DeviceProfileName, ConfigSystem);
+		if (Mode != EDeviceProfileMode::DPM_CacheValues)
+		{
+			// Store the newly selected fragments.
+			PlatformFragmentsSelected = FragmentsSelected;
+		}
+	}
+	else if (Mode != EDeviceProfileMode::DPM_CacheValues)
+	{
+		// use the existing selected fragment state.
+		FragmentsSelected = PlatformFragmentsSelected;
+	}
+
+	// Here we gather the cvars from selected fragments in reverse order
+	for (int i = FragmentsSelected.Num() - 1; i >= 0; i--)
+	{
+		const FSelectedFragmentProperties& SelectedFragment = FragmentsSelected[i];
+		if (SelectedFragment.bEnabled)
+		{
+			TArray<FString> FragmentCVars;
+			GetFragmentCVars(*SelectedFragment.Fragment, TEXT("CVars"), FragmentCVars, ConfigSystem);
+			for (const FString& FragCVar : FragmentCVars)
+			{
+				FString CVarKey, CVarValue;
+				if (FragCVar.Split(TEXT("="), &CVarKey, &CVarValue))
+				{
+					if (!FragmentCVarKeys.Find(CVarKey))
+					{
+						FragmentCVarKeys.Add(CVarKey);
+						SelectedFragmentCVars.Add(FragCVar);
+					}
+				}
+			}
+		}
+	}
+
 	// For each device profile, starting with the selected and working our way up the BaseProfileName tree,
 	// Find all CVars and set them 
 	FString BaseDeviceProfileName = DeviceProfileName;
@@ -300,7 +358,7 @@ void UDeviceProfileManager::ProcessDeviceProfileIniSettings(const FString& Devic
 				}
 
 				TArray< FString > CurrentProfilesCVars, FragmentCVars;
-				GetFragmentCvars(*CurrentSectionName, *ArrayName, FragmentCVars, ConfigSystem);
+				GetCVarsFromDPFragmentIncludes(*CurrentSectionName, *ArrayName, FragmentCVars, ConfigSystem);
 				ConfigSystem->GetArray(*CurrentSectionName, *ArrayName, CurrentProfilesCVars, GDeviceProfilesIni);
 
 				if (FragmentCVars.Num())
@@ -310,6 +368,10 @@ void UDeviceProfileManager::ProcessDeviceProfileIniSettings(const FString& Devic
 					CurrentProfilesCVars += FragmentCVars;
 				}
 
+				// now add the selected fragments at the end so these override the DP.
+				CurrentProfilesCVars += SelectedFragmentCVars;
+				SelectedFragmentCVars.Empty();
+
 				// Iterate over the profile and make sure we do not have duplicate CVars
 				{
 					TMap< FString, FString > ValidCVars;
@@ -318,11 +380,6 @@ void UDeviceProfileManager::ProcessDeviceProfileIniSettings(const FString& Devic
 						FString CVarKey, CVarValue;
 						if ((*CVarIt).Split(TEXT("="), &CVarKey, &CVarValue))
 						{
-							if (ValidCVars.Find(CVarKey))
-							{
-								ValidCVars.Remove(CVarKey);
-							}
-
 							ValidCVars.Add(CVarKey, CVarValue);
 						}
 					}
@@ -346,7 +403,7 @@ void UDeviceProfileManager::ProcessDeviceProfileIniSettings(const FString& Devic
 						if (!CVarsAlreadySetList.Find(CVarKey))
 						{
 #if WITH_EDITOR
-							if (Mode == EDeviceProfileMode::DPM_PushCVars)
+							if (Mode == EDeviceProfileMode::DPM_CacheValues)
 							{
 								if (PreviewDenylistCVars.Contains(CVarKey))
 								{
@@ -362,20 +419,17 @@ void UDeviceProfileManager::ProcessDeviceProfileIniSettings(const FString& Devic
 							}
 #endif
 
-							if (Mode == EDeviceProfileMode::DPM_PushCVars)
+							if (Mode == EDeviceProfileMode::DPM_SetCVars)
 							{
 								IConsoleVariable* CVar = IConsoleManager::Get().FindConsoleVariable(*CVarKey);
 								if (CVar)
 								{
-									if (DeviceProfileManagerSingleton)
-									{
-										// remember the previous value
-										FString OldValue = CVar->GetString();
-										DeviceProfileManagerSingleton->PushedSettings.Add(CVarKey, OldValue);
-
-										// indicate we are pushing, not setting
-										UE_LOG(LogDeviceProfileManager, Log, TEXT("Pushing Device Profile CVar: [[%s:%s -> %s]]"), *CVarKey, *OldValue, *CVarValue);
-									}
+									// remember the previous value
+									FString OldValue = CVar->GetString();
+									PushedSettings.Add(CVarKey, OldValue);
+
+									// indicate we are pushing, not setting
+									UE_LOG(LogDeviceProfileManager, Log, TEXT("Pushing Device Profile CVar: [[%s:%s -> %s]]"), *CVarKey, *OldValue, *CVarValue);
 								}
 								else
 								{
@@ -408,9 +462,7 @@ void UDeviceProfileManager::ProcessDeviceProfileIniSettings(const FString& Devic
 								}
 
 								//If this is a dp preview then we set cvars with their existing priority so that we don't cause future issues when setting by scalability levels etc.
-								// @todo ini: preview should never be setting cvars via this function, should use the cached values
-								uint32 BaseCVarPriority = /*bIsDeviceProfilePreview ? ECVF_SetByMask : */ECVF_SetByDeviceProfile;
-								uint32 CVarPriority = bIsScalabilityBucket ? ECVF_SetByScalability : BaseCVarPriority;
+								uint32 CVarPriority = bIsScalabilityBucket ? ECVF_SetByScalability : ECVF_SetByDeviceProfile;
 								OnSetCVarFromIniEntry(*GDeviceProfilesIni, *CVarKey, *CVarValue, CVarPriority);
 								CVarsAlreadySetList.Add(CVarKey, CVarValue);
 							}
@@ -446,8 +498,60 @@ void UDeviceProfileManager::ProcessDeviceProfileIniSettings(const FString& Devic
 #endif
 }
 
+/**
+* Set the cvar state to PushedSettings.
+*/
+static void RestorePushedState(TMap<FString, FString>& PushedSettings)
+{
+	// restore pushed settings
+	for (TMap<FString, FString>::TIterator It(PushedSettings); It; ++It)
+	{
+		IConsoleVariable* CVar = IConsoleManager::Get().FindConsoleVariable(*It.Key());
+		if (CVar)
+		{
+			// restore it!
+			CVar->SetWithCurrentPriority(*It.Value());
+			UE_LOG(LogDeviceProfileManager, Log, TEXT("Popping Device Profile CVar: [[%s:%s]]"), *It.Key(), *It.Value());
+		}
+	}
+
+	PushedSettings.Reset();
+}
+
+const FSelectedFragmentProperties* UDeviceProfileManager::GetActiveDeviceProfileFragmentByTag(FName& FragmentTag) const
+{
+	for (const FSelectedFragmentProperties& SelectedFragment : PlatformFragmentsSelected)
+	{
+		if (SelectedFragment.Tag == FragmentTag)
+		{
+			return &SelectedFragment;
+		}
+	}
+	return nullptr;
+}
+
+// enable/disable a tagged fragment.
+void UDeviceProfileManager::ChangeTaggedFragmentState(FName FragmentTag, bool bNewState)
+{
+	for (FSelectedFragmentProperties& Fragment : PlatformFragmentsSelected)
+	{
+		if (Fragment.Tag == FragmentTag)
+		{
+			if (bNewState != Fragment.bEnabled)
+			{
+				UE_LOG(LogInit, Log, TEXT("ChangeTaggedFragmentState: %s=%d"), *FragmentTag.ToString(), bNewState);
+				// unset entire DP's cvar state.
+				RestorePushedState(PushedSettings);
+				// set the new state and reapply all fragments.
+				Fragment.bEnabled = bNewState;
+				ProcessDeviceProfileIniSettings(GetActiveDeviceProfileName(), EDeviceProfileMode::DPM_SetCVars);
+			}
+			break;
+		}
+	}
+}
 
-void UDeviceProfileManager::InitializeCVarsForActiveDeviceProfile(bool bPushSettings)
+void UDeviceProfileManager::InitializeCVarsForActiveDeviceProfile()
 {
 	FString ActiveProfileName;
 
@@ -460,8 +564,7 @@ void UDeviceProfileManager::InitializeCVarsForActiveDeviceProfile(bool bPushSett
 		ActiveProfileName = GetPlatformDeviceProfileName();
 	}
 
-	UE_LOG(LogDeviceProfileManager, Log, TEXT("Applying CVar settings loaded from the selected device profile: [%s]"), *ActiveProfileName);
-	ProcessDeviceProfileIniSettings(ActiveProfileName, bPushSettings ? EDeviceProfileMode::DPM_PushCVars : EDeviceProfileMode::DPM_SetCVars);
+	ProcessDeviceProfileIniSettings(ActiveProfileName, EDeviceProfileMode::DPM_SetCVars);
 }
 
 #if ALLOW_OTHER_PLATFORM_CONFIG
@@ -594,14 +697,7 @@ void UDeviceProfileManager::ReapplyDeviceProfile()
 	UDeviceProfile* OverrideProfile = DeviceProfileManagerSingleton->BaseDeviceProfile ? DeviceProfileManagerSingleton->GetActiveProfile() : nullptr;
 	UDeviceProfile* BaseProfile = DeviceProfileManagerSingleton->BaseDeviceProfile ? DeviceProfileManagerSingleton->BaseDeviceProfile : DeviceProfileManagerSingleton->GetActiveProfile();
 
-	UE_LOG(LogDeviceProfileManager, Log, TEXT("ReapplyDeviceProfile applying profile: [%s]"), *BaseProfile->GetName(), OverrideProfile ? *OverrideProfile->GetName() : TEXT("not set.") );
-
-	// pop any pushed settings
-	RestoreDefaultDeviceProfile();
-
-	// Set base profile and re-apply cvars.
-	SetActiveDeviceProfile(BaseProfile);
-	InitializeCVarsForActiveDeviceProfile();
+	UE_LOG(LogDeviceProfileManager, Log, TEXT("ReapplyDeviceProfile applying profile: [%s]"), *BaseProfile->GetName(), OverrideProfile ? *OverrideProfile->GetName() : TEXT("not set."));
 
 	if (OverrideProfile)
 	{
@@ -611,6 +707,14 @@ void UDeviceProfileManager::ReapplyDeviceProfile()
 	}
 	else
 	{
+		// reset any fragments, this will cause them to be rematched.
+		PlatformFragmentsSelected.Empty();
+		// restore to the pre-DP cvar state:
+		RestorePushedState(PushedSettings);
+
+		// Apply the active DP. 
+		InitializeCVarsForActiveDeviceProfile();
+
 		// broadcast cvar sinks now that we are done
 		IConsoleManager::Get().CallAllConsoleVariableSinks();
 	}
@@ -859,48 +963,74 @@ void UDeviceProfileManager::SaveProfiles(bool bSaveToDefaults)
 	}
 }
 
-/**
-* Overrides the device profile. The original profile can be restored with RestoreDefaultDeviceProfile
-*/
-void UDeviceProfileManager::SetOverrideDeviceProfile(UDeviceProfile* DeviceProfile, bool bIsDeviceProfilePreview)
+#if ALLOW_OTHER_PLATFORM_CONFIG
+void UDeviceProfileManager::SetPreviewDeviceProfile(UDeviceProfile* DeviceProfile)
 {
-	// pop any pushed settings
-	HandleDeviceProfileOverridePop();
+	// we're applying a preview mode on top of an overridden DP?
+	check(BaseDeviceProfile == nullptr);
 
-	// for preview, we assume this will be another platform's DP, so use the resolved cvars directly, bypassing the activate and set stuff
-	if (bIsDeviceProfilePreview)
+	RestorePreviewDeviceProfile();
+
+	UE_LOG(LogDeviceProfileManager, Log, TEXT("SetPreviewDeviceProfile preview to %s"), *DeviceProfile->GetName());
+	// apply the preview DP cvars.
+	for (const auto& Pair : DeviceProfile->GetAllExpandedCVars())
 	{
-#if ALLOW_OTHER_PLATFORM_CONFIG
-		for (const auto& Pair : DeviceProfile->GetAllExpandedCVars())
+		IConsoleVariable* CVar = IConsoleManager::Get().FindConsoleVariable(*Pair.Key);
+		// skip over scalability group cvars (maybe they shouldn't be left in the AllExpandedCVars?)
+		if (CVar != nullptr && !CVar->TestFlags(EConsoleVariableFlags::ECVF_ScalabilityGroup))
 		{
-			IConsoleVariable* CVar = IConsoleManager::Get().FindConsoleVariable(*Pair.Key);
-			// skip over scalability group cvars (maybe they shouldn't be left in the AllExpandedCVars?)
-			if (CVar != nullptr && !CVar->TestFlags(EConsoleVariableFlags::ECVF_ScalabilityGroup))
+			// remember the previous value so we can restore
+			FString OldValue = CVar->GetString();
+			PreviewPushedSettings.Add(Pair.Key, OldValue);
+			// cheat CVar can only be set in ConsoleVariables.ini
+			if (!CVar->TestFlags(EConsoleVariableFlags::ECVF_Cheat))
 			{
-				// remember the previous value so we can restore
-				FString OldValue = CVar->GetString();
-				PushedSettings.Add(Pair.Key, OldValue);
-
-				// cheat CVar can only be set in ConsoleVariables.ini
-				if (!CVar->TestFlags(EConsoleVariableFlags::ECVF_Cheat))
-				{
-					// set the cvar to the new value, with same priority that it was before (SetByMask means current priority)
-					CVar->SetWithCurrentPriority(ConvertValueFromHumanFriendlyValue(*Pair.Value));
-				}
+				// set the cvar to the new value, with same priority that it was before (SetByMask means current priority)
+				CVar->SetWithCurrentPriority(ConvertValueFromHumanFriendlyValue(*Pair.Value));
 			}
 		}
-#else
-		UE_LOG(LogDeviceProfileManager, Error, TEXT("SetOverrideDeviceProfile with bIsDeviceProfilePreview=true can only be used in a developer tool"));
+	}
+}
+
+
+void UDeviceProfileManager::RestorePreviewDeviceProfile()
+{
+	if (PreviewPushedSettings.Num())
+	{
+		checkf(BaseDeviceProfile == nullptr, TEXT("call to RestorePreviewDeviceProfile while both preview and BaseDeviceProfile has been set?"));
+
+		UE_LOG(LogDeviceProfileManager, Log, TEXT("Restoring Preview DP "));
+		// this sets us back to non-preview state.
+		RestorePushedState(PreviewPushedSettings);
+	}
+}
+#endif
+
+/**
+* Overrides the device profile. The original profile can be restored with RestoreDefaultDeviceProfile
+*/
+void UDeviceProfileManager::SetOverrideDeviceProfile(UDeviceProfile* DeviceProfile)
+{
+#if ALLOW_OTHER_PLATFORM_CONFIG
+	// we have an active preview running but we're changing the actual device's DP too?
+	check(PreviewPushedSettings.Num() == 0);
 #endif
-		return;
+
+	// If we're not already overriding record the BaseDeviceProfile
+	if(!BaseDeviceProfile)
+	{
+		BaseDeviceProfile = DeviceProfileManagerSingleton->GetActiveProfile();
 	}
+	UE_LOG(LogDeviceProfileManager, Log, TEXT("Overriding DeviceProfile to %s, base device profile %s"), *DeviceProfile->GetName(), *BaseDeviceProfile->GetName() );
 
-	// record the currently active profile, needed when we restore the default.
-	BaseDeviceProfile = DeviceProfileManagerSingleton->GetActiveProfile();
+	// reset any fragments, this will cause them to be rematched.
+	PlatformFragmentsSelected.Empty();
+	// restore to the pre-DP cvar state:
+	RestorePushedState(PushedSettings);
 
 	// activate new one!
 	DeviceProfileManagerSingleton->SetActiveDeviceProfile(DeviceProfile);
-	InitializeCVarsForActiveDeviceProfile(/*bPushSettings*/true);
+	InitializeCVarsForActiveDeviceProfile();
 
 	// broadcast cvar sinks now that we are done
 	IConsoleManager::Get().CallAllConsoleVariableSinks();
@@ -911,30 +1041,28 @@ void UDeviceProfileManager::SetOverrideDeviceProfile(UDeviceProfile* DeviceProfi
 */
 void UDeviceProfileManager::RestoreDefaultDeviceProfile()
 {
-	// restore pushed settings
-	for (TMap<FString, FString>::TIterator It(PushedSettings); It; ++It)
-	{
-		IConsoleVariable* CVar = IConsoleManager::Get().FindConsoleVariable(*It.Key());
-		if (CVar)
-		{
-			// restore it!
-			CVar->SetWithCurrentPriority(*It.Value());			
-			UE_LOG(LogDeviceProfileManager, Log, TEXT("Popping Device Profile CVar: [[%s:%s]]"), *It.Key(), *It.Value());
-		}
-	}
-
-	PushedSettings.Reset();
+#if ALLOW_OTHER_PLATFORM_CONFIG
+	// We're restoring overridden DP while a preview is active?
+	check(PreviewPushedSettings.Num() == 0);
+#endif
 
-	if(BaseDeviceProfile)
+	// have we been overridden?
+	if (BaseDeviceProfile)
 	{
+		UE_LOG(LogDeviceProfileManager, Log, TEXT("Restoring overridden DP back to %s"), *BaseDeviceProfile->GetName());
+		// this differs from previous behavior, we used to push only the cvar state that was modified by the override.
+		// But now we restore the entire CVar state to 'pre-DP' stage and reapply the currently active DP.
 		// reset the base profile as we are no longer overriding
-		DeviceProfileManagerSingleton->SetActiveDeviceProfile(BaseDeviceProfile);
+		RestorePushedState(PushedSettings);
+		// reset any fragments, this will cause them to be rematched.
+		PlatformFragmentsSelected.Empty();
+		SetActiveDeviceProfile(BaseDeviceProfile);
 		BaseDeviceProfile = nullptr;
+		//set the DP cvar state
+		InitializeCVarsForActiveDeviceProfile();
 	}
 }
 
-
-
 void UDeviceProfileManager::HandleDeviceProfileOverrideChange()
 {
 	FString CVarValue = CVarDeviceProfileOverride.GetValueOnGameThread();
@@ -949,9 +1077,38 @@ void UDeviceProfileManager::HandleDeviceProfileOverrideChange()
 	}
 }
 
-void UDeviceProfileManager::HandleDeviceProfileOverridePop()
+IDeviceProfileSelectorModule* UDeviceProfileManager::GetPreviewDeviceProfileSelectorModule(FConfigCacheIni* PreviewConfigSystemIn)
+{
+#if ALLOW_OTHER_PLATFORM_CONFIG
+	// If we're getting the selector for previewing the PIEPreviewDeviceProfileSelector device selector can be given a PreviewDeviceDesciption to return selector info for specific devices.
+	FString PreviewDeviceDesciption; 
+	if (PreviewConfigSystemIn->GetString(TEXT("DeviceProfileManager"), TEXT("PreviewDeviceDesciption"), PreviewDeviceDesciption, GEngineIni))
+	{
+		// this should only be specified when previewing.
+		if (IPIEPreviewDeviceModule* DPSelectorModule = FModuleManager::LoadModulePtr<IPIEPreviewDeviceModule>(TEXT("PIEPreviewDeviceProfileSelector")))
+		{
+			// TODO: check success.
+			DPSelectorModule->SetPreviewDevice(PreviewDeviceDesciption);
+			return DPSelectorModule;
+		}
+	}
+#else
+	checkNoEntry();
+#endif
+	return nullptr;
+}
+
+IDeviceProfileSelectorModule* UDeviceProfileManager::GetDeviceProfileSelectorModule()
 {
-	RestoreDefaultDeviceProfile();
+	FString DeviceProfileSelectionModule;
+	if (GConfig->GetString(TEXT("DeviceProfileManager"), TEXT("DeviceProfileSelectionModule"), DeviceProfileSelectionModule, GEngineIni))
+	{
+		if (IDeviceProfileSelectorModule* DPSelectorModule = FModuleManager::LoadModulePtr<IDeviceProfileSelectorModule>(*DeviceProfileSelectionModule))
+		{
+			return DPSelectorModule;
+		}
+	}
+	return nullptr;
 }
 
 const FString UDeviceProfileManager::GetPlatformDeviceProfileName()
@@ -972,14 +1129,9 @@ const FString UDeviceProfileManager::GetPlatformDeviceProfileName()
 		return OverrideProfileName;
 	}
 
-
-	FString DeviceProfileSelectionModule;
-	if (GConfig->GetString(TEXT("DeviceProfileManager"), TEXT("DeviceProfileSelectionModule"), DeviceProfileSelectionModule, GEngineIni))
+	if (IDeviceProfileSelectorModule* DPSelectorModule = GetDeviceProfileSelectorModule())
 	{
-		if (IDeviceProfileSelectorModule* DPSelectorModule = FModuleManager::LoadModulePtr<IDeviceProfileSelectorModule>(*DeviceProfileSelectionModule))
-		{
-			ActiveProfileName = DPSelectorModule->GetRuntimeDeviceProfileName();
-		}
+		ActiveProfileName = DPSelectorModule->GetRuntimeDeviceProfileName();
 	}
 
 #if WITH_EDITOR
@@ -1177,19 +1329,22 @@ public:
 			UDeviceProfile* DeviceProfile = UDeviceProfileManager::Get().FindProfile(Cmd, false);
 			if (DeviceProfile)
 			{
-				UDeviceProfileManager::Get().SetOverrideDeviceProfile(DeviceProfile, true);
+				UDeviceProfileManager::Get().SetPreviewDeviceProfile(DeviceProfile);
 			}
 		}
 		else if (FParse::Command(&Cmd, TEXT("dprestore")))
 		{
-			UDeviceProfileManager::Get().RestoreDefaultDeviceProfile();
+			UDeviceProfileManager::Get().RestorePreviewDeviceProfile();
 		}
 		else if (FParse::Command(&Cmd, TEXT("dpreload")))
 		{
 			FConfigCacheIni::ClearOtherPlatformConfigs();
 			// @todo ini clear out all DPs AllExpandedCVars
 		}
-
+		else if (FParse::Command(&Cmd, TEXT("dpreapply")))
+		{
+			UDeviceProfileManager::Get().ReapplyDeviceProfile();
+		}
 
 
 		return false;
diff --git a/Engine/Source/Runtime/Engine/Private/DeviceProfiles/DeviceProfileMatching.cpp b/Engine/Source/Runtime/Engine/Private/DeviceProfiles/DeviceProfileMatching.cpp
new file mode 100644
--- /dev/null
+++ b/Engine/Source/Runtime/Engine/Private/DeviceProfiles/DeviceProfileMatching.cpp
@@ -0,0 +1,605 @@
+// Copyright Epic Games, Inc. All Rights Reserved.
+
+#include "DeviceProfiles/DeviceProfileMatching.h"
+#include "DeviceProfiles/DeviceProfileManager.h"
+#include "Misc/ConfigCacheIni.h"
+#include "HAL/IConsoleManager.h"
+#include "Modules/ModuleManager.h"
+#include "Misc/CommandLine.h"
+#include "UObject/Package.h"
+#include "SceneManagement.h"
+#include "SystemSettings.h"
+#include "DeviceProfiles/DeviceProfile.h"
+#include "IDeviceProfileSelectorModule.h"
+#include "Misc/DataDrivenPlatformInfoRegistry.h"
+#include "Internationalization/Regex.h"
+#include "GenericPlatform/GenericPlatformDriver.h"
+
+// Platform independent source types
+static FName SRC_Chipset(TEXT("SRC_Chipset"));
+static FName SRC_MakeAndModel(TEXT("SRC_MakeAndModel"));
+static FName SRC_OSVersion(TEXT("SRC_OSVersion"));
+static FName SRC_CommandLine(TEXT("SRC_CommandLine"));
+static FName SRC_PrimaryGPUDesc(TEXT("SRC_PrimaryGPUDesc"));
+static FName SRC_False(TEXT("false"));
+static FName SRC_True(TEXT("true"));
+static FName SRC_PreviousRegexMatch(TEXT("SRC_PreviousRegexMatch"));
+static FName SRC_PreviousRegexMatch1(TEXT("SRC_PreviousRegexMatch1"));
+
+// comparison operators:
+static FName CMP_Equal(TEXT("=="));
+static FName CMP_Less(TEXT("<"));
+static FName CMP_LessEqual(TEXT("<="));
+static FName CMP_Greater(TEXT(">"));
+static FName CMP_GreaterEqual(TEXT(">="));
+static FName CMP_NotEqual(TEXT("!="));
+static FName CMP_Regex(TEXT("CMP_Regex"));
+static FName CMP_EqualIgnore(TEXT("CMP_EqualIgnore"));
+static FName CMP_LessIgnore(TEXT("CMP_LessIgnore"));
+static FName CMP_LessEqualIgnore(TEXT("CMP_LessEqualIgnore"));
+static FName CMP_GreaterIgnore(TEXT("CMP_GreaterIgnore"));
+static FName CMP_GreaterEqualIgnore(TEXT("CMP_GreaterEqualIgnore"));
+static FName CMP_NotEqualIgnore(TEXT("CMP_NotEqualIgnore"));
+static FName CMP_Hash(TEXT("CMP_Hash"));
+static FName CMP_Or(TEXT("OR"));
+static FName CMP_And(TEXT("AND"));
+
+class FRuleMatchRunner
+{
+	const TMap<FName, FString>* UserDefinedSrcs;
+	FOutputDevice* ErrorDevice;
+	FString PreviousRegexMatches[2];
+public:
+	struct FDeviceProfileMatchCriterion
+	{
+		FName SourceType;
+		FString SourceArg;
+		FName CompareType;
+		FString MatchString;
+	};
+
+	FRuleMatchRunner(const TMap<FName, FString>* UserDefinedSrcsIN, FOutputDevice* ErrorDeviceIn) : UserDefinedSrcs(UserDefinedSrcsIN), ErrorDevice(ErrorDeviceIn){}
+	bool ProcessRules(IDeviceProfileSelectorModule* DPSelectorModule, TArray<FDeviceProfileMatchCriterion>& MatchingCriteria, const FString& RuleName)
+	{
+		bool bResult = false;
+		bool bFoundMatch = true;
+		for (const FDeviceProfileMatchCriterion& DeviceProfileMatchCriterion : MatchingCriteria)
+		{
+			bool bCurrentMatch = true;
+			FName SourceType = DeviceProfileMatchCriterion.SourceType;
+			const FString& SourceArg = DeviceProfileMatchCriterion.SourceArg;
+			const FString& MatchString = DeviceProfileMatchCriterion.MatchString;
+			FName CompareType = DeviceProfileMatchCriterion.CompareType;
+			FString CommandLine = FCommandLine::Get();
+
+			FString SourceString("");
+			FString PlatformSourceValue("");
+
+			// Selector module gets first dibs on retrieving source data
+			if (DPSelectorModule && DPSelectorModule->GetSelectorPropertyValue(SourceType, PlatformSourceValue)) { SourceString = PlatformSourceValue; }
+			// universal properties
+			else if (SourceType == SRC_Chipset) { SourceString = FPlatformMisc::GetCPUChipset(); }
+			else if (SourceType == SRC_MakeAndModel) { SourceString = FPlatformMisc::GetDeviceMakeAndModel(); }
+			else if (SourceType == SRC_OSVersion) { SourceString = FPlatformMisc::GetOSVersion(); }
+			else if (SourceType == SRC_PrimaryGPUDesc) {SourceString = FPlatformMisc::GetGPUDriverInfo(FPlatformMisc::GetPrimaryGPUBrand()).DeviceDescription;}
+			else if (SourceType == SRC_PreviousRegexMatch) { SourceString = PreviousRegexMatches[0]; }
+			else if (SourceType == SRC_PreviousRegexMatch1) { SourceString = PreviousRegexMatches[1]; }
+			else if (SourceType == SRC_CommandLine) { SourceString = CommandLine; }
+			else if (SourceType == SRC_False) { SourceString = TEXT("false"); }
+			else if (SourceType == SRC_True) { SourceString = TEXT("true"); }
+			// SetSrc defined properties.
+			else if (const FString* Value = UserDefinedSrcs ? UserDefinedSrcs->Find(SourceType) : nullptr) { SourceString = *Value; }
+			else
+			{
+				// SourceType wasn't found. 
+				ErrorDevice->Logf(TEXT("source type '%s' was not defined for matchingrule %s. (%s, %s, %s)"), *SourceType.ToString(), *RuleName, *SourceType.ToString(), *CompareType.ToString(), *MatchString);
+			}
+
+			const bool bNumericOperands = SourceString.IsNumeric() && MatchString.IsNumeric();
+			if (CompareType == CMP_Equal)
+			{
+				if (SourceType == SRC_CommandLine)
+				{
+					if (!FParse::Param(*SourceString, *MatchString))
+					{
+						bCurrentMatch = false;
+					}
+				}
+				else
+				{
+					if (SourceString != MatchString)
+					{
+						bCurrentMatch = false;
+					}
+				}
+			}
+			else if (CompareType == CMP_Less)
+			{
+				if ((bNumericOperands && FCString::Atof(*SourceString) >= FCString::Atof(*MatchString)) || (!bNumericOperands && SourceString >= MatchString))
+				{
+					bCurrentMatch = false;
+				}
+			}
+			else if (CompareType == CMP_LessEqual)
+			{
+				if ((bNumericOperands && FCString::Atof(*SourceString) > FCString::Atof(*MatchString)) || (!bNumericOperands && SourceString > MatchString))
+				{
+					bCurrentMatch = false;
+				}
+			}
+			else if (CompareType == CMP_Greater)
+			{
+				if ((bNumericOperands && FCString::Atof(*SourceString) <= FCString::Atof(*MatchString)) || (!bNumericOperands && SourceString <= MatchString))
+				{
+					bCurrentMatch = false;
+				}
+			}
+			else if (CompareType == CMP_GreaterEqual)
+			{
+				if ((bNumericOperands && FCString::Atof(*SourceString) < FCString::Atof(*MatchString)) || (!bNumericOperands && SourceString < MatchString))
+				{
+					bCurrentMatch = false;
+				}
+			}
+			else if (CompareType == CMP_NotEqual)
+			{
+				if (SourceType == SRC_CommandLine)
+				{
+					if (FParse::Param(*CommandLine, *MatchString))
+					{
+						bCurrentMatch = false;
+					}
+				}
+				else
+				{
+					if (*SourceString == MatchString)
+					{
+						bCurrentMatch = false;
+					}
+				}
+			}
+			else if (CompareType == CMP_Or || CompareType == CMP_And)
+			{
+				bool bArg1, bArg2;
+
+				if (bNumericOperands)
+				{
+					bArg1 = FCString::Atoi(*SourceString) != 0;
+					bArg2 = FCString::Atoi(*MatchString) != 0;
+				}
+				else
+				{
+					static const FString TrueString("true");
+					bArg1 = SourceString == TrueString;
+					bArg2 = MatchString == TrueString;
+				}
+
+				if (CompareType == CMP_Or)
+				{
+					bCurrentMatch = (bArg1 || bArg2);
+				}
+				else
+				{
+					bCurrentMatch = (bArg1 && bArg2);
+				}
+			}
+			else if (CompareType == CMP_EqualIgnore)
+			{
+				if (SourceString.ToLower() != MatchString.ToLower())
+				{
+					bCurrentMatch = false;
+				}
+			}
+			else if (CompareType == CMP_LessIgnore)
+			{
+				if (SourceString.ToLower() >= MatchString.ToLower())
+				{
+					bCurrentMatch = false;
+				}
+			}
+			else if (CompareType == CMP_LessEqualIgnore)
+			{
+				if (SourceString.ToLower() > MatchString.ToLower())
+				{
+					bCurrentMatch = false;
+				}
+			}
+			else if (CompareType == CMP_GreaterIgnore)
+			{
+				if (SourceString.ToLower() <= MatchString.ToLower())
+				{
+					bCurrentMatch = false;
+				}
+				else if (CompareType == CMP_GreaterEqualIgnore)
+				{
+					if (SourceString.ToLower() < MatchString.ToLower())
+					{
+						bCurrentMatch = false;
+					}
+				}
+			}
+			else if (CompareType == CMP_NotEqualIgnore)
+			{
+				if (SourceString.ToLower() == MatchString.ToLower())
+				{
+					bCurrentMatch = false;
+				}
+			}
+			else if (CompareType == CMP_Regex)
+			{
+				const FRegexPattern RegexPattern(MatchString);
+				FRegexMatcher RegexMatcher(RegexPattern, SourceString);
+				if (RegexMatcher.FindNext())
+				{
+					PreviousRegexMatches[0] = RegexMatcher.GetCaptureGroup(1);
+					PreviousRegexMatches[1] = RegexMatcher.GetCaptureGroup(2);
+				}
+				else
+				{
+					PreviousRegexMatches[0].Empty();
+					PreviousRegexMatches[1].Empty();
+
+					bCurrentMatch = false;
+				}
+			}
+			else if (CompareType == CMP_Hash)
+			{
+				// Salt string is concatenated onto the end of the input text.
+				// For example the input string "PhoneModel" with salt "Salt" and pepper "Pepper" can be computed with
+				// % printf "PhoneModelSaltPepper" | openssl dgst -sha1 -hex
+				// resulting in d9e5cbd6b0e4dba00edd9de92cf64ee4c3f3a2db and would be stored in the matching rules as 
+				// "Salt|d9e5cbd6b0e4dba00edd9de92cf64ee4c3f3a2db". Salt is optional.
+				FString MatchHashString;
+				FString SaltString;
+				if (!MatchString.Split(TEXT("|"), &SaltString, &MatchHashString))
+				{
+					MatchHashString = MatchString;
+				}
+				FString HashInputString = SourceString + SaltString
+#ifdef HASH_PEPPER_SECRET_GUID
+					+ HASH_PEPPER_SECRET_GUID.ToString()
+#endif
+					;
+
+				FSHAHash SourceHash;
+				FSHA1::HashBuffer(TCHAR_TO_ANSI(*HashInputString), HashInputString.Len(), SourceHash.Hash);
+				if (SourceHash.ToString() != MatchHashString.ToUpper())
+				{
+					bCurrentMatch = false;
+				}
+			}
+			else
+			{
+				bCurrentMatch = false;
+			}
+
+			bFoundMatch = bCurrentMatch;
+		}
+		return bFoundMatch;
+	}
+};
+
+static bool EvaluateMatch(IDeviceProfileSelectorModule* DPSelectorModule, const FDPMatchingRulestructBase& rule, TMap<FName, FString>& UserDefinedSrc, FString& SelectedFragmentsOUT, FOutputDevice* Errors)
+{
+	static const FString NoName("");
+	const FString& RuleName = rule.RuleName.IsEmpty() ? NoName : rule.RuleName;
+
+	if(!rule.AppendFragments.IsEmpty())
+	{
+		if(!SelectedFragmentsOUT.IsEmpty())
+		{
+			SelectedFragmentsOUT += TEXT(",");
+		}
+		if (rule.AppendFragments.Contains(TEXT("[clear]")))
+		{
+			SelectedFragmentsOUT.Empty();
+			SelectedFragmentsOUT += rule.AppendFragments.Replace(TEXT("[clear]"), TEXT(""));
+		}
+		else
+		{
+			SelectedFragmentsOUT += rule.AppendFragments;
+		}
+	}
+
+	if (!rule.SetSrc.IsEmpty())
+	{
+		TArray<FString> NewSRCs;
+		rule.SetSrc.ParseIntoArray(NewSRCs, TEXT(","), true);
+		for (FString& SRCentry : NewSRCs)
+		{
+			FString CSRCType, CSRCValue;
+			if (SRCentry.Split(TEXT("="), &CSRCType, &CSRCValue))
+			{
+				UserDefinedSrc.Add(FName(CSRCType), CSRCValue);
+				UE_LOG(LogInit, Verbose, TEXT("MatchesRules: Adding source %s : %s"), *CSRCType, *CSRCValue);
+			}
+		}
+	}
+
+	TArray<FDPMatchingIfCondition> Expression = rule.IfConditions;
+	if (Expression.Num() == 0)
+	{
+		UE_LOG(LogInit, Verbose, TEXT("MatchesRules: %s, no match criteria."), *RuleName);
+		return true;
+	}
+
+	// if no operator specified, insert an implicit AND operator.
+	for (int i = 0; i < Expression.Num(); i++)
+	{
+		// if this and next are )( operations insert AND
+		if (i + 1 < Expression.Num() && (!Expression[i].Arg1.IsEmpty() || Expression[i].Operator == TEXT(")")))
+		{
+			if (!Expression[i + 1].Arg1.IsEmpty() || Expression[i + 1].Operator == TEXT("("))
+			{
+				FDPMatchingIfCondition ImplicitAnd;
+				ImplicitAnd.Operator = TEXT("AND");
+				Expression.Insert(ImplicitAnd, i + 1);
+				i++;
+			}
+		}
+	}
+
+	FString Line;
+	for (int i = 0; i < Expression.Num(); i++)
+	{
+		if (!(Expression[i].Arg1.IsEmpty()))
+		{
+			Line += FString::Printf(TEXT("(%s %s %s)"), *Expression[i].Arg1, *Expression[i].Operator.ToString(), *Expression[i].Arg2);
+		}
+		else
+		{
+			Line += FString::Printf(TEXT(" %s "), *Expression[i].Operator.ToString());
+		}
+	}
+	UE_LOG(LogInit, Verbose, TEXT("MatchesRules: rule %s : %s"), *RuleName, *Line);
+
+	struct FExpressionItem
+	{
+		FExpressionItem(FString InValue, bool bInOperator) : Value(InValue), bOperator(bInOperator) {}
+		FString Value;
+		bool bOperator;
+	};
+
+	TArray<FExpressionItem> RPNOutput;
+	TArray<FString> Operators;
+
+	auto PushOperand = [&RPNOutput](FString InValue) { RPNOutput.Push(FExpressionItem(InValue, false)); };
+	auto PushOperator = [&RPNOutput](FString InValue)
+	{
+		RPNOutput.Push(FExpressionItem(InValue, true));
+	};
+
+	for (const FDPMatchingIfCondition& element : Expression)
+	{
+		if (!element.Arg1.IsEmpty())
+		{
+			PushOperand(element.Arg1);
+			//if (!element.Arg2.IsEmpty())
+			{
+				PushOperand(element.Arg2);
+			}
+			if (!element.Operator.IsNone())
+			{
+				PushOperator(element.Operator.ToString());
+			}
+		}
+		else if (!element.Operator.IsNone())
+		{
+			if (element.Operator == FName(TEXT(")")))
+			{
+				while (Operators.Num())
+				{
+					FString PoppedOperator = Operators.Pop();
+					if (PoppedOperator != FString(TEXT("(")))
+					{
+						PushOperator(PoppedOperator);
+					}
+					else
+					{
+						break;
+					}
+				}
+			}
+			else
+			{
+				Operators.Push(element.Operator.ToString());
+			}
+		}
+	}
+	while (Operators.Num())
+	{
+		FString PoppedOperator = Operators.Pop();
+		if (PoppedOperator == FString(TEXT("(")) || PoppedOperator == FString(TEXT(")")))
+		{
+			Errors->Logf(TEXT("MatchesRules: rule %s failed due to mismatching parenthesis! %s"), *RuleName, *PoppedOperator);
+			return false;
+		}
+		PushOperator(PoppedOperator);
+	}
+	UE_LOG(LogInit, Verbose, TEXT("MatchesRules: rule %s : "), *RuleName);
+	for (int i = 0; i < RPNOutput.Num(); i++)
+	{
+		UE_LOG(LogInit, Verbose, TEXT("MatchesRules: (%d - %s)"), i, *RPNOutput[i].Value);
+	}
+
+	FRuleMatchRunner MatchMe(&UserDefinedSrc, Errors);
+	TArray<FExpressionItem> ExpressionStack;
+	for (int i = 0; i < RPNOutput.Num(); i++)
+	{
+		if (!RPNOutput[i].bOperator)
+		{
+			ExpressionStack.Push(RPNOutput[i]);
+		}
+		else
+		{
+			FExpressionItem B = ExpressionStack.Pop();
+			FExpressionItem A = ExpressionStack.Pop();
+			TArray<FRuleMatchRunner::FDeviceProfileMatchCriterion> MatchingCriteria;
+
+			FRuleMatchRunner::FDeviceProfileMatchCriterion crit;
+			crit.CompareType = FName(RPNOutput[i].Value);
+			crit.SourceType = FName(A.Value);
+			crit.MatchString = B.Value;
+
+			MatchingCriteria.Add(crit);
+			bool bResult = MatchMe.ProcessRules(DPSelectorModule, MatchingCriteria, RuleName);
+			FExpressionItem C(bResult ? TEXT("true") : TEXT("false"), false);
+			ExpressionStack.Push(C);
+
+			UE_LOG(LogInit, Verbose, TEXT("MatchesRules: rule %s evaluating (%s %s %s) = %s"), *RuleName, *A.Value, *RPNOutput[i].Value, *B.Value, *C.Value);
+		}
+	}
+	FExpressionItem Result = ExpressionStack.Pop();
+	UE_LOG(LogInit, Log, TEXT("MatchesRules: rule %s = %s"), *RuleName, *Result.Value);
+
+	const bool bMatched = Result.Value == TEXT("true");
+	const FDPMatchingRulestructBase* Next = bMatched ? rule.GetOnTrue() : rule.GetOnFalse();
+
+	bool bSuccess = true;
+	if (Next)
+	{
+		bSuccess = EvaluateMatch(DPSelectorModule, *Next, UserDefinedSrc, SelectedFragmentsOUT, Errors);
+	}
+
+	return bSuccess;
+}
+
+// Convert a string of fragment names to a FSelectedFragmentProperties array.
+// FragmentName1,FragmentName2,[optionaltag]FragmentName3, etc.
+static TArray<FSelectedFragmentProperties> FragmentStringToFragmentProperties(const FString& FragmentString)
+{
+	TArray<FSelectedFragmentProperties> FragmentPropertiesList;
+	TArray<FString> AppendedFragments;
+	FragmentString.ParseIntoArray(AppendedFragments, TEXT(","), true);
+	for (const FString& Fragment : AppendedFragments)
+	{
+		FSelectedFragmentProperties FragmentProperties;
+		int32 TagDeclStart = Fragment.Find(TEXT("["));
+		if (TagDeclStart >= 0)
+		{
+			int32 TagDeclEnd = Fragment.Find(TEXT("]"));
+			if (TagDeclEnd > TagDeclStart)
+			{
+				FName Tag(Fragment.Mid(TagDeclStart + 1, (TagDeclEnd - TagDeclStart)-1));
+				FragmentProperties.Tag = Tag;
+				FragmentProperties.bEnabled = false;
+				FragmentProperties.Fragment = Fragment.RightChop(TagDeclEnd + 1);
+				FragmentPropertiesList.Emplace(MoveTemp(FragmentProperties));
+			}
+		}
+		else
+		{
+			FragmentProperties.Fragment = Fragment;
+			FragmentProperties.bEnabled = true;
+			FragmentPropertiesList.Emplace(MoveTemp(FragmentProperties));
+		}
+	}
+	return FragmentPropertiesList;
+}
+
+class FDeviceProfileMatchingErrorContext: public FOutputDevice
+{
+public:
+	FString Stage;
+	int32 NumErrors;
+
+	FDeviceProfileMatchingErrorContext()
+		: FOutputDevice()
+		, NumErrors(0)
+	{}
+
+	virtual void Serialize(const TCHAR* V, ELogVerbosity::Type Verbosity, const class FName& Category) override
+	{
+		UE_LOG(LogInit, Error, TEXT("DeviceProfileMatching: Error while parsing Matching Rules (%s) : %s"), *Stage, V);
+		NumErrors++;
+	}
+};
+
+static bool DPHasMatchingRules(const FString& ParentDP, FConfigCacheIni* ConfigSystem)
+{
+	FString SectionSuffix = *FString::Printf(TEXT(" %s"), *UDeviceProfile::StaticClass()->GetName());
+	FString CurrentSectionName = ParentDP + SectionSuffix;
+	FString ArrayName(TEXT("MatchingRules"));
+	TArray<FString> MatchingRulesArray;
+	return ConfigSystem->GetArray(*CurrentSectionName, *ArrayName, MatchingRulesArray, GDeviceProfilesIni) != 0;
+}
+
+static FString LoadAndProcessMatchingRulesFromConfig(const FString& ParentDP, IDeviceProfileSelectorModule* DPSelector, FConfigCacheIni* ConfigSystem)
+{
+	FString SectionSuffix = *FString::Printf(TEXT(" %s"), *UDeviceProfile::StaticClass()->GetName());
+	FString CurrentSectionName = ParentDP + SectionSuffix;
+	FString ArrayName(TEXT("MatchingRules"));
+	TArray<FString> MatchingRulesArray;
+	ConfigSystem->GetArray(*CurrentSectionName, *ArrayName, MatchingRulesArray, GDeviceProfilesIni);
+
+	TMap<FName, FString> UserDefinedSrc;
+	FString SelectedFragments;
+	FDeviceProfileMatchingErrorContext DPMatchingErrorOutput;
+	int Count = 0;
+	for(const FString& RuleText : MatchingRulesArray)
+	{
+		DPMatchingErrorOutput.Stage = FString::Printf(TEXT("%s rule #%d"), *ParentDP, Count);
+		Count++;
+		FDPMatchingRulestruct RuleStruct;
+		FDPMatchingRulestruct::StaticStruct()->ImportText(*RuleText, &RuleStruct, nullptr, EPropertyPortFlags::PPF_None, &DPMatchingErrorOutput, FDPMatchingRulestruct::StaticStruct()->GetName(), true);
+		EvaluateMatch(DPSelector, RuleStruct, UserDefinedSrc, SelectedFragments, &DPMatchingErrorOutput);
+	}
+#if UE_BUILD_SHIPPING
+	UE_CLOG(DPMatchingErrorOutput.NumErrors > 0, LogInit, Error, TEXT("DeviceProfileMatching: %d Error(s) encountered while processing MatchedRules for %s"), DPMatchingErrorOutput.NumErrors, *ParentDP);
+#else
+	UE_CLOG(DPMatchingErrorOutput.NumErrors > 0, LogInit, Fatal, TEXT("DeviceProfileMatching: %d Error(s) encountered while processing MatchedRules for %s"), DPMatchingErrorOutput.NumErrors, *ParentDP);
+#endif
+
+	return SelectedFragments;
+}
+
+static FString RemoveAllWhiteSpace(const FString& StringIN)
+{
+	FString Ret;
+	Ret.Reserve(StringIN.Len());
+	for (TCHAR Character : StringIN)
+	{
+		if (!FChar::IsWhitespace(Character))
+		{
+			Ret += Character;
+		}
+	}
+	return Ret;
+}
+
+TArray<FSelectedFragmentProperties> UDeviceProfileManager::FindMatchingFragments(const FString& ParentDP, FConfigCacheIni* ConfigSystem)
+{
+#if !UE_BUILD_SHIPPING
+	// Override selected fragments with commandline specified list -DPFragments=fragmentname,fragmentname2,[taggedname]fragment,... 
+	FString SelectedFragments, DPFragmentString;
+	if (FParse::Value(FCommandLine::Get(), TEXT("DPFragments="), DPFragmentString, false))
+	{
+		SelectedFragments = DPFragmentString;
+	}
+	else
+#endif
+	{
+		bool bIsPreview = false;
+#if ALLOW_OTHER_PLATFORM_CONFIG
+		bIsPreview = ConfigSystem != GConfig;
+#endif
+		IDeviceProfileSelectorModule* DPSelector = bIsPreview ? GetPreviewDeviceProfileSelectorModule(ConfigSystem) : GetDeviceProfileSelectorModule();
+		SelectedFragments = LoadAndProcessMatchingRulesFromConfig(ParentDP, DPSelector, ConfigSystem);
+		
+		// previewing a DP with matching rules will run if conditions with the host device's data sources. It will likely not match the preview device's behavior.
+		UE_CLOG(bIsPreview && !DPSelector && DPHasMatchingRules(ParentDP, ConfigSystem), LogInit, Warning, TEXT("Preview DP %s contains fragment matching rules, but no preview profile selector was found. The selected fragments for %s will likely not match the behavior of the intended preview device."), *ParentDP, *ParentDP);
+	}
+	SelectedFragments = RemoveAllWhiteSpace(SelectedFragments);
+
+	UE_CLOG(!SelectedFragments.IsEmpty(), LogInit, Log, TEXT("MatchesRules:Fragment string %s"), *SelectedFragments);
+	TArray<FSelectedFragmentProperties> MatchedFragments = FragmentStringToFragmentProperties(SelectedFragments);
+
+	UE_CLOG(MatchedFragments.Num()>0, LogInit, Log, TEXT("MatchesRules: MatchedFragments:"));
+	for (FSelectedFragmentProperties& MatchedFrag : MatchedFragments)
+	{
+		UE_CLOG(MatchedFrag.Tag == NAME_None, LogInit, Log, TEXT("MatchesRules: %s, enabled %d"), *MatchedFrag.Fragment, MatchedFrag.bEnabled);
+		UE_CLOG(MatchedFrag.Tag != NAME_None, LogInit, Log, TEXT("MatchesRules: %s=%s, enabled %d"), *MatchedFrag.Tag.ToString(), *MatchedFrag.Fragment, MatchedFrag.bEnabled);
+	}
+
+	return MatchedFragments;
+}
