fix build issue with 17442791

diff --git a/Engine/Source/Runtime/Engine/Private/DeviceProfiles/DeviceProfileManager.cpp b/Engine/Source/Runtime/Engine/Private/DeviceProfiles/DeviceProfileManager.cpp
--- a/Engine/Source/Runtime/Engine/Private/DeviceProfiles/DeviceProfileManager.cpp
+++ b/Engine/Source/Runtime/Engine/Private/DeviceProfiles/DeviceProfileManager.cpp
@@ -1079,7 +1079,7 @@ void UDeviceProfileManager::HandleDeviceProfileOverrideChange()
 
 IDeviceProfileSelectorModule* UDeviceProfileManager::GetPreviewDeviceProfileSelectorModule(FConfigCacheIni* PreviewConfigSystemIn)
 {
-#if ALLOW_OTHER_PLATFORM_CONFIG
+#if ALLOW_OTHER_PLATFORM_CONFIG && WITH_EDITOR
 	// If we're getting the selector for previewing the PIEPreviewDeviceProfileSelector device selector can be given a PreviewDeviceDesciption to return selector info for specific devices.
 	FString PreviewDeviceDesciption; 
 	if (PreviewConfigSystemIn->GetString(TEXT("DeviceProfileManager"), TEXT("PreviewDeviceDesciption"), PreviewDeviceDesciption, GEngineIni))
