fix shipping build

diff --git a/Engine/Source/Runtime/Engine/Private/DeviceProfiles/DeviceProfileMatching.cpp b/Engine/Source/Runtime/Engine/Private/DeviceProfiles/DeviceProfileMatching.cpp
--- a/Engine/Source/Runtime/Engine/Private/DeviceProfiles/DeviceProfileMatching.cpp
+++ b/Engine/Source/Runtime/Engine/Private/DeviceProfiles/DeviceProfileMatching.cpp
@@ -569,9 +569,10 @@ static FString RemoveAllWhiteSpace(const FString& StringIN)
 
 TArray<FSelectedFragmentProperties> UDeviceProfileManager::FindMatchingFragments(const FString& ParentDP, FConfigCacheIni* ConfigSystem)
 {
+	FString SelectedFragments;
 #if !UE_BUILD_SHIPPING
 	// Override selected fragments with commandline specified list -DPFragments=fragmentname,fragmentname2,[taggedname]fragment,... 
-	FString SelectedFragments, DPFragmentString;
+	FString DPFragmentString;
 	if (FParse::Value(FCommandLine::Get(), TEXT("DPFragments="), DPFragmentString, false))
 	{
 		SelectedFragments = DPFragmentString;
