Morph target engine modifications for 3rd party plugin

diff --git a/Engine/Plugins/Experimental/AlembicImporter/Source/AlembicLibrary/Private/AbcImporter.cpp b/Engine/Plugins/Experimental/AlembicImporter/Source/AlembicLibrary/Private/AbcImporter.cpp
--- a/Engine/Plugins/Experimental/AlembicImporter/Source/AlembicLibrary/Private/AbcImporter.cpp
+++ b/Engine/Plugins/Experimental/AlembicImporter/Source/AlembicLibrary/Private/AbcImporter.cpp
@@ -683,7 +683,7 @@ TArray<UObject*> FAbcImporter::ImportAsSkeletalMesh(UObject* InParent, EObjectFl
 						const bool bCompareNormals = true;
 						MorphTarget->PopulateDeltas(MorphDeltas, 0, LODModel.Sections, bCompareNormals);
 
-						const float PercentageOfVerticesInfluences = ((float)MorphTarget->MorphLODModels[0].Vertices.Num() / (float)NumIndices) * 100.0f;
+						const float PercentageOfVerticesInfluences = ((float)MorphTarget->GetMorphLODModels()[0].Vertices.Num() / (float)NumIndices) * 100.0f;
 						if (PercentageOfVerticesInfluences > ImportSettings->CompressionSettings.MinimumNumberOfVertexInfluencePercentage)
 						{
 							SkeletalMesh->RegisterMorphTarget(MorphTarget);
diff --git a/Engine/Plugins/Importers/USDImporter/Source/USDUtilities/Private/USDSkeletalDataConversion.cpp b/Engine/Plugins/Importers/USDImporter/Source/USDUtilities/Private/USDSkeletalDataConversion.cpp
--- a/Engine/Plugins/Importers/USDImporter/Source/USDUtilities/Private/USDSkeletalDataConversion.cpp
+++ b/Engine/Plugins/Importers/USDImporter/Source/USDUtilities/Private/USDSkeletalDataConversion.cpp
@@ -967,7 +967,7 @@ namespace UnrealToUsdImpl
 	}
 
 	// Converts UE morph target deltas from DeltaArray into offsets, pointIndices and normalOffsets attributes of BlendShape
-	bool ConvertMorphTargetDeltas( FMorphTargetDelta* DeltaArray, int32 NumDeltas, pxr::UsdSkelBlendShape& BlendShape, pxr::UsdTimeCode TimeCode )
+	bool ConvertMorphTargetDeltas( const FMorphTargetDelta* DeltaArray, int32 NumDeltas, pxr::UsdSkelBlendShape& BlendShape, pxr::UsdTimeCode TimeCode )
 	{
 		if ( !DeltaArray || NumDeltas == 0 || !BlendShape )
 		{
@@ -2504,7 +2504,7 @@ bool UnrealToUsd::ConvertSkeletalMesh( const USkeletalMesh* SkeletalMesh, pxr::U
 			}
 
 			int32 NumDeltas = 0;
-			FMorphTargetDelta* DeltaArray = MorphTarget->GetMorphTargetDelta( LODIndex, NumDeltas );
+			const FMorphTargetDelta* DeltaArray = MorphTarget->GetMorphTargetDelta( LODIndex, NumDeltas );
 			if ( !DeltaArray || NumDeltas == 0 )
 			{
 				continue;
diff --git a/Engine/Source/Developer/MeshBuilder/Private/SkeletalMeshBuilder.cpp b/Engine/Source/Developer/MeshBuilder/Private/SkeletalMeshBuilder.cpp
--- a/Engine/Source/Developer/MeshBuilder/Private/SkeletalMeshBuilder.cpp
+++ b/Engine/Source/Developer/MeshBuilder/Private/SkeletalMeshBuilder.cpp
@@ -175,11 +175,12 @@ bool FSkeletalMeshBuilder::Build(const FSkeletalMeshBuildParameters& SkeletalMes
 						continue;
 					}
 					TArray<FMorphTargetDelta>& MorphDeltasArray = BaseLODMorphTargetData.FindOrAdd(MorphTarget->GetFullName());
-					const FMorphTargetLODModel& BaseMorphModel = MorphTarget->MorphLODModels[LODIndex];
 					//Iterate each original morph target source index to fill the NewMorphTargetDeltas array with the TargetMatchData.
-					for (const FMorphTargetDelta& MorphDelta : BaseMorphModel.Vertices)
+					int32 NumDeltas = 0;
+					const FMorphTargetDelta* BaseDeltaArray = MorphTarget->GetMorphTargetDelta(LODIndex, NumDeltas);
+					for (int32 DeltaIndex = 0; DeltaIndex < NumDeltas; DeltaIndex++)
 					{
-						MorphDeltasArray.Add(MorphDelta);
+						MorphDeltasArray.Add(BaseDeltaArray[DeltaIndex]);
 					}
 				}
 				//Copy the original SkeletalMesh LODModel
diff --git a/Engine/Source/Developer/MeshUtilities/Private/MeshUtilities.cpp b/Engine/Source/Developer/MeshUtilities/Private/MeshUtilities.cpp
--- a/Engine/Source/Developer/MeshUtilities/Private/MeshUtilities.cpp
+++ b/Engine/Source/Developer/MeshUtilities/Private/MeshUtilities.cpp
@@ -5873,7 +5873,7 @@ void FMeshUtilities::CreateImportDataFromLODModel(USkeletalMesh* SkeletalMesh) c
 					continue;
 				}
 				TArray<FMorphTargetDelta>& MorphDeltasArray = LODMorphTargetData.FindOrAdd(MorphTarget->GetName());
-				const FMorphTargetLODModel& BaseMorphModel = MorphTarget->MorphLODModels[LodIndex];
+				const FMorphTargetLODModel& BaseMorphModel = MorphTarget->GetMorphLODModels()[LodIndex];
 				//Iterate each original morph target source index to fill the NewMorphTargetDeltas array with the TargetMatchData.
 				for (const FMorphTargetDelta& MorphDelta : BaseMorphModel.Vertices)
 				{
diff --git a/Engine/Source/Developer/SkeletalMeshUtilitiesCommon/Private/LODUtilities.cpp b/Engine/Source/Developer/SkeletalMeshUtilitiesCommon/Private/LODUtilities.cpp
--- a/Engine/Source/Developer/SkeletalMeshUtilitiesCommon/Private/LODUtilities.cpp
+++ b/Engine/Source/Developer/SkeletalMeshUtilitiesCommon/Private/LODUtilities.cpp
@@ -418,7 +418,7 @@ void FLODUtilities::RemoveLOD(FSkeletalMeshUpdateContext& UpdateContext, int32 D
 		{
 			if (MorphTarget->HasDataForLOD(DesiredLOD))
 			{
-				MorphTarget->MorphLODModels.RemoveAt(DesiredLOD);
+				MorphTarget->GetMorphLODModels().RemoveAt(DesiredLOD);
 			}
 		}
 
@@ -872,7 +872,7 @@ void CreateLODMorphTarget(USkeletalMesh* SkeletalMesh, FReductionBaseSkeletalMes
 		TArray<FMorphTargetDelta> NewMorphTargetDeltas;
 		TSet<uint32> CreatedTargetIndex;
 		TMap<FVector3f, TArray<uint32>> MorphTargetPerPosition;
-		const FMorphTargetLODModel& BaseMorphModel = MorphTarget->MorphLODModels[SourceLOD];
+		const FMorphTargetLODModel& BaseMorphModel = MorphTarget->GetMorphLODModels()[SourceLOD];
 		//Iterate each original morph target source index to fill the NewMorphTargetDeltas array with the TargetMatchData.
 		const TArray<FMorphTargetDelta>& Vertices = bUseBaseMorphDelta ? *BaseMorphDeltas : BaseMorphModel.Vertices;
 		for (uint32 MorphDeltaIndex = 0; MorphDeltaIndex < (uint32)(Vertices.Num()); ++MorphDeltaIndex)
@@ -985,10 +985,10 @@ void FLODUtilities::ClearGeneratedMorphTarget(USkeletalMesh* SkeletalMesh, int32
 
 		//if (MorphTarget->MorphLODModels[TargetLOD].bGeneratedByEngine)
 		{
-			MorphTarget->MorphLODModels[TargetLOD].Reset();
+			MorphTarget->GetMorphLODModels()[TargetLOD].Reset();
 
 			// if this is the last one, we can remove empty ones
-			if (TargetLOD == MorphTarget->MorphLODModels.Num() - 1)
+			if (TargetLOD == MorphTarget->GetMorphLODModels().Num() - 1)
 			{
 				MorphTarget->RemoveEmptyMorphTargets();
 			}
@@ -1149,7 +1149,7 @@ void FLODUtilities::ApplyMorphTargetsToLOD(USkeletalMesh* SkeletalMesh, int32 So
 		}
 
 		TMap<uint32, uint32>& BaseIndexToMorphTargetDelta = PerMorphTargetBaseIndexToMorphTargetDelta.FindOrAdd(MorphTarget);
-		const FMorphTargetLODModel& BaseMorphModel = MorphTarget->MorphLODModels[SourceLOD];
+		const FMorphTargetLODModel& BaseMorphModel = MorphTarget->GetMorphLODModels()[SourceLOD];
 		const TArray<FMorphTargetDelta>& Vertices = bUseTempMorphDelta ? *TempMorphDeltas : BaseMorphModel.Vertices;
 		for (uint32 MorphDeltaIndex = 0; MorphDeltaIndex < (uint32)(Vertices.Num()); ++MorphDeltaIndex)
 		{
@@ -1261,11 +1261,12 @@ void FLODUtilities::SimplifySkeletalMeshLOD( USkeletalMesh* SkeletalMesh, int32
 						continue;
 					}
 					TArray<FMorphTargetDelta>& MorphDeltasArray = BaseLODMorphTargetData.FindOrAdd(MorphTarget->GetFullName());
-					const FMorphTargetLODModel& BaseMorphModel = MorphTarget->MorphLODModels[DesiredLOD];
 					//Iterate each original morph target source index to fill the NewMorphTargetDeltas array with the TargetMatchData.
-					for (const FMorphTargetDelta& MorphDelta : BaseMorphModel.Vertices)
+					int32 NumDeltas = 0;
+					const FMorphTargetDelta* BaseDeltaArray = MorphTarget->GetMorphTargetDelta(DesiredLOD, NumDeltas);
+					for (int32 DeltaIndex = 0; DeltaIndex < NumDeltas; DeltaIndex++)
 					{
-						MorphDeltasArray.Add(MorphDelta);
+						MorphDeltasArray.Add(BaseDeltaArray[DeltaIndex]);
 					}
 				}
 				
diff --git a/Engine/Source/Editor/Persona/Private/AnimationEditorViewportClient.cpp b/Engine/Source/Editor/Persona/Private/AnimationEditorViewportClient.cpp
--- a/Engine/Source/Editor/Persona/Private/AnimationEditorViewportClient.cpp
+++ b/Engine/Source/Editor/Persona/Private/AnimationEditorViewportClient.cpp
@@ -740,7 +740,7 @@ FText FAnimationViewportClient::GetDisplayInfo(bool bDisplayAllInfo) const
 			{
 				continue;
 			}
-			for (const FMorphTargetLODModel& MorphTargetLODModel : MorphTarget->MorphLODModels)
+			for (const FMorphTargetLODModel& MorphTargetLODModel : MorphTarget->GetMorphLODModels())
 			{
 				for (int32 SectionIndex : MorphTargetLODModel.SectionIndices)
 				{
diff --git a/Engine/Source/Editor/Persona/Private/SMorphTargetViewer.cpp b/Engine/Source/Editor/Persona/Private/SMorphTargetViewer.cpp
--- a/Engine/Source/Editor/Persona/Private/SMorphTargetViewer.cpp
+++ b/Engine/Source/Editor/Persona/Private/SMorphTargetViewer.cpp
@@ -476,7 +476,7 @@ void SMorphTargetViewer::CreateMorphTargetList( const FString& SearchText )
 				continue; // Skip items that don't match our filter
 			}
 
-			int32 NumberOfVerts = (MorphTargets[I]->MorphLODModels.Num() > 0)? MorphTargets[I]->MorphLODModels[0].Vertices.Num() : 0;
+			int32 NumberOfVerts = (MorphTargets[I]->GetMorphLODModels().Num() > 0)? MorphTargets[I]->GetMorphLODModels()[0].Vertices.Num() : 0;
 
 			TSharedRef<FDisplayedMorphTargetInfo> Info = FDisplayedMorphTargetInfo::Make( MorphTargets[I]->GetFName(), NumberOfVerts);
 			if(MeshComponent)
diff --git a/Engine/Source/Editor/UnrealEd/Private/Fbx/FbxSkeletalMeshExport.cpp b/Engine/Source/Editor/UnrealEd/Private/Fbx/FbxSkeletalMeshExport.cpp
--- a/Engine/Source/Editor/UnrealEd/Private/Fbx/FbxSkeletalMeshExport.cpp
+++ b/Engine/Source/Editor/UnrealEd/Private/Fbx/FbxSkeletalMeshExport.cpp
@@ -284,11 +284,11 @@ FbxNode* FFbxExporter::CreateMesh(const USkeletalMesh* SkelMesh, const TCHAR* Me
 					}
 				
 					int32 NumberOfDeltas = 0;
-					FMorphTargetDelta* MorphTargetDeltas = MorphTarget->GetMorphTargetDelta(LODIndex, NumberOfDeltas);
+					const FMorphTargetDelta* MorphTargetDeltas = MorphTarget->GetMorphTargetDelta(LODIndex, NumberOfDeltas);
 					for (int32 MorphTargetDeltaIndex = 0; MorphTargetDeltaIndex < NumberOfDeltas; ++MorphTargetDeltaIndex)
 					{
 						// Apply the morph target deltas to the control points.
-						FMorphTargetDelta& CurrentDelta = MorphTargetDeltas[MorphTargetDeltaIndex];
+						const FMorphTargetDelta& CurrentDelta = MorphTargetDeltas[MorphTargetDeltaIndex];
 						uint32 RemappedSourceIndex = CurrentDelta.SourceIdx;
 
 						if (VertexIndexOffsetPairArray.Num() > 1)
diff --git a/Engine/Source/Runtime/Apple/MetalRHI/Private/MetalRHI.cpp b/Engine/Source/Runtime/Apple/MetalRHI/Private/MetalRHI.cpp
--- a/Engine/Source/Runtime/Apple/MetalRHI/Private/MetalRHI.cpp
+++ b/Engine/Source/Runtime/Apple/MetalRHI/Private/MetalRHI.cpp
@@ -703,6 +703,13 @@ FMetalDynamicRHI::FMetalDynamicRHI(ERHIFeatureLevel::Type RequestedFeatureLevel)
 	GMetalBufferFormats[PF_ETC2_R11_EAC			] = { mtlpp::PixelFormat::Invalid, (uint8)EMetalBufferFormat::Unknown };
 	GMetalBufferFormats[PF_ETC2_RG11_EAC		] = { mtlpp::PixelFormat::Invalid, (uint8)EMetalBufferFormat::Unknown };
 		
+	GMetalBufferFormats[PF_G16R16_SNORM			] = { mtlpp::PixelFormat::RG16Snorm, (uint8)EMetalBufferFormat::RG16Snorm };
+	GMetalBufferFormats[PF_R8G8_UINT			] = { mtlpp::PixelFormat::RG8Uint, (uint8)EMetalBufferFormat::RG8Uint };
+	GMetalBufferFormats[PF_R32G32B32_UINT		] = { mtlpp::PixelFormat::Invalid, (uint8)EMetalBufferFormat::Unknown };
+	GMetalBufferFormats[PF_R32G32B32_SINT		] = { mtlpp::PixelFormat::Invalid, (uint8)EMetalBufferFormat::Unknown };
+	GMetalBufferFormats[PF_R32G32B32F			] = { mtlpp::PixelFormat::Invalid, (uint8)EMetalBufferFormat::Unknown };
+	GMetalBufferFormats[PF_R8_SINT				] = { mtlpp::PixelFormat::R8Sint, (uint8)EMetalBufferFormat::R8Sint };
+
 	// Initialize the platform pixel format map.
 	GPixelFormats[PF_Unknown			].PlatformFormat	= (uint32)mtlpp::PixelFormat::Invalid;
 	GPixelFormats[PF_A32B32G32R32F		].PlatformFormat	= (uint32)mtlpp::PixelFormat::RGBA32Float;
@@ -916,6 +923,16 @@ FMetalDynamicRHI::FMetalDynamicRHI(ERHIFeatureLevel::Type RequestedFeatureLevel)
 	GPixelFormats[PF_NV12				].PlatformFormat	= (uint32)mtlpp::PixelFormat::Invalid;
 	GPixelFormats[PF_NV12				].Supported			= false;
 	
+	GPixelFormats[PF_G16R16_SNORM		].PlatformFormat	= (uint32)mtlpp::PixelFormat::RG16Snorm;
+	GPixelFormats[PF_R8G8_UINT			].PlatformFormat	= (uint32)mtlpp::PixelFormat::RG8Uint;
+	GPixelFormats[PF_R32G32B32_UINT		].PlatformFormat	= (uint32)mtlpp::PixelFormat::Invalid;
+	GPixelFormats[PF_R32G32B32_UINT		].Supported			= false;
+	GPixelFormats[PF_R32G32B32_SINT		].PlatformFormat	= (uint32)mtlpp::PixelFormat::Invalid;
+	GPixelFormats[PF_R32G32B32_SINT		].Supported			= false;
+	GPixelFormats[PF_R32G32B32F			].PlatformFormat	= (uint32)mtlpp::PixelFormat::Invalid;
+	GPixelFormats[PF_R32G32B32F			].Supported			= false;
+	GPixelFormats[PF_R8_SINT			].PlatformFormat	= (uint32)mtlpp::PixelFormat::R8Sint;
+
 #if METAL_DEBUG_OPTIONS
 	for (uint32 i = 0; i < PF_MAX; i++)
 	{
diff --git a/Engine/Source/Runtime/Core/Public/PixelFormat.h b/Engine/Source/Runtime/Core/Public/PixelFormat.h
--- a/Engine/Source/Runtime/Core/Public/PixelFormat.h
+++ b/Engine/Source/Runtime/Core/Public/PixelFormat.h
@@ -90,8 +90,13 @@ enum EPixelFormat
 	PF_ASTC_8x8_HDR         =75,	
 	PF_ASTC_10x10_HDR       =76,	
 	PF_ASTC_12x12_HDR       =77,
-	
-	PF_MAX					=78,
+	PF_G16R16_SNORM			=78,
+	PF_R8G8_UINT			=79,
+	PF_R32G32B32_UINT		=80,
+	PF_R32G32B32_SINT		=81,
+	PF_R32G32B32F			=82,
+	PF_R8_SINT				=83,	
+	PF_MAX					=84,
 };
 #define FOREACH_ENUM_EPIXELFORMAT(op) \
 	op(PF_Unknown) \
@@ -171,7 +176,13 @@ enum EPixelFormat
 	op(PF_ASTC_6x6_HDR) \
 	op(PF_ASTC_8x8_HDR) \
 	op(PF_ASTC_10x10_HDR) \
-	op(PF_ASTC_12x12_HDR)
+	op(PF_ASTC_12x12_HDR) \
+	op(PF_G16R16_SNORM) \
+	op(PF_R8G8_UINT) \
+	op(PF_R32G32B32_UINT) \
+	op(PF_R32G32B32_SINT) \
+	op(PF_R32G32B32F) \
+	op(PF_R8_SINT)
 // Defines which channel is valid for each pixel format
 enum class EPixelFormatChannelFlags : uint8
 {
diff --git a/Engine/Source/Runtime/CoreUObject/Public/UObject/NoExportTypes.h b/Engine/Source/Runtime/CoreUObject/Public/UObject/NoExportTypes.h
--- a/Engine/Source/Runtime/CoreUObject/Public/UObject/NoExportTypes.h
+++ b/Engine/Source/Runtime/CoreUObject/Public/UObject/NoExportTypes.h
@@ -248,6 +248,12 @@ enum EPixelFormat
 	PF_ETC2_RG11_EAC,
 	PF_R8,
 	PF_B5G5R5A1_UNORM,
+	PF_G16R16_SNORM,
+	PF_R8G8_UINT,
+	PF_R32G32B32_UINT,
+	PF_R32G32B32_SINT,
+	PF_R32G32B32F,
+	PF_R8_SINT,
 	PF_MAX,
 };
 
diff --git a/Engine/Source/Runtime/D3D12RHI/Private/D3D12RHI.cpp b/Engine/Source/Runtime/D3D12RHI/Private/D3D12RHI.cpp
--- a/Engine/Source/Runtime/D3D12RHI/Private/D3D12RHI.cpp
+++ b/Engine/Source/Runtime/D3D12RHI/Private/D3D12RHI.cpp
@@ -175,6 +175,13 @@ FD3D12DynamicRHI::FD3D12DynamicRHI(const TArray<TSharedPtr<FD3D12Adapter>>& Chos
 	GPixelFormats[PF_NV12].PlatformFormat = DXGI_FORMAT_NV12;
 	GPixelFormats[PF_NV12].Supported = true;
 
+	GPixelFormats[PF_G16R16_SNORM	].PlatformFormat = DXGI_FORMAT_R16G16_SNORM;
+	GPixelFormats[PF_R8G8_UINT		].PlatformFormat = DXGI_FORMAT_R8G8_UINT;
+	GPixelFormats[PF_R32G32B32_UINT	].PlatformFormat = DXGI_FORMAT_R32G32B32_UINT;
+	GPixelFormats[PF_R32G32B32_SINT	].PlatformFormat = DXGI_FORMAT_R32G32B32_SINT;
+	GPixelFormats[PF_R32G32B32F		].PlatformFormat = DXGI_FORMAT_R32G32B32_FLOAT;
+	GPixelFormats[PF_R8_SINT		].PlatformFormat = DXGI_FORMAT_R8_SINT;
+
 	// MS - Not doing any feature level checks. D3D12 currently supports these limits.
 	// However this may need to be revisited if new feature levels are introduced with different HW requirement
 	GSupportsSeparateRenderTargetBlendState = true;
diff --git a/Engine/Source/Runtime/Engine/Classes/Animation/MorphTarget.h b/Engine/Source/Runtime/Engine/Classes/Animation/MorphTarget.h
--- a/Engine/Source/Runtime/Engine/Classes/Animation/MorphTarget.h
+++ b/Engine/Source/Runtime/Engine/Classes/Animation/MorphTarget.h
@@ -8,7 +8,7 @@
 #include "PackedNormal.h"
 #include "UObject/EditorObjectVersion.h"
 #include "UObject/FortniteMainBranchObjectVersion.h"
-
+#include "Serialization/MemoryArchive.h"
 #include "MorphTarget.generated.h"
 
 class USkeletalMesh;
@@ -66,7 +66,7 @@ struct FMorphTargetLODModel
 	bool bGeneratedByEngine;
 
 	/** Get Resource Size */
-	void GetResourceSizeEx(FResourceSizeEx& CumulativeResourceSize) const;
+	ENGINE_API void GetResourceSizeEx(FResourceSizeEx& CumulativeResourceSize) const;
 
 	FMorphTargetLODModel()
 		: NumVertices(0)
@@ -140,32 +140,45 @@ public:
 	UPROPERTY(AssetRegistrySearchable)
 	TObjectPtr<class USkeletalMesh> BaseSkelMesh;
 
+	/** morph mesh vertex data for each LOD */
+	virtual const TArray<FMorphTargetLODModel>& GetMorphLODModels() const { return MorphLODModels; }
+
+	/** morph mesh vertex data for each LOD */
+	virtual TArray<FMorphTargetLODModel>& GetMorphLODModels() { return MorphLODModels; }
+
+protected:
 	/** morph mesh vertex data for each LOD */
 	TArray<FMorphTargetLODModel>	MorphLODModels;
 
 public:
 
 	/** Get Morphtarget Delta array for the given input Index */
-	ENGINE_API FMorphTargetDelta* GetMorphTargetDelta(int32 LODIndex, int32& OutNumDeltas);
-	ENGINE_API bool HasDataForLOD(int32 LODIndex);
+	ENGINE_API const FMorphTargetDelta* GetMorphTargetDelta(int32 LODIndex, int32& OutNumDeltas) const;
+	ENGINE_API virtual bool HasDataForLOD(int32 LODIndex) const;
+	/** return true if this morphtarget contains data for section within LOD */
+	ENGINE_API virtual bool HasDataForSection(int32 LODIndex, int32 SectionIndex) const;
 	/** return true if this morphtarget contains valid vertices */
-	ENGINE_API bool HasValidData() const;
+	ENGINE_API virtual bool HasValidData() const;
+	ENGINE_API virtual void EmptyMorphLODModels();
 
 	/** Discard CPU Buffers after render resources have been created. */
 	ENGINE_API void DiscardVertexData();
 
 #if WITH_EDITOR
 	/** Populates the given morph target LOD model with the provided deltas */
-	ENGINE_API void PopulateDeltas(const TArray<FMorphTargetDelta>& Deltas, const int32 LODIndex, const TArray<struct FSkelMeshSection>& Sections, const bool bCompareNormal = false, const bool bGeneratedByReductionSetting = false, const float PositionThreshold = THRESH_POINTS_ARE_NEAR);
+	ENGINE_API virtual void PopulateDeltas(const TArray<FMorphTargetDelta>& Deltas, const int32 LODIndex, const TArray<struct FSkelMeshSection>& Sections, const bool bCompareNormal = false, const bool bGeneratedByReductionSetting = false, const float PositionThreshold = THRESH_POINTS_ARE_NEAR);
 	/** Remove empty LODModels */
-	ENGINE_API void RemoveEmptyMorphTargets();
+	ENGINE_API virtual void RemoveEmptyMorphTargets();
 #endif // WITH_EDITOR
 
 public:
 
 	//~ UObject interface
 
-	virtual void Serialize(FArchive& Ar) override;
-	virtual void GetResourceSizeEx(FResourceSizeEx& CumulativeResourceSize) override;
-	virtual void PostLoad() override;
+	ENGINE_API virtual void Serialize(FArchive& Ar) override;
+	ENGINE_API virtual void GetResourceSizeEx(FResourceSizeEx& CumulativeResourceSize) override;
+	ENGINE_API virtual void PostLoad() override;
+
+	/** UObject does not support serialization via FMemoryArchive, so manually handle separately */
+	ENGINE_API virtual void SerializeMemoryArchive(FMemoryArchive& Ar);
 };
diff --git a/Engine/Source/Runtime/Engine/Classes/Components/SkinnedMeshComponent.h b/Engine/Source/Runtime/Engine/Classes/Components/SkinnedMeshComponent.h
--- a/Engine/Source/Runtime/Engine/Classes/Components/SkinnedMeshComponent.h
+++ b/Engine/Source/Runtime/Engine/Classes/Components/SkinnedMeshComponent.h
@@ -690,6 +690,12 @@ public:
 	/** Object responsible for sending bone transforms, morph target state etc. to render thread. */
 	class FSkeletalMeshObject*	MeshObject;
 
+	/** Supports user-defined FSkeletalMeshObjects */
+	class FSkeletalMeshObject* (*MeshObjectFactory)(void* UserData, USkinnedMeshComponent* InMeshComponent, FSkeletalMeshRenderData* InSkelMeshRenderData, ERHIFeatureLevel::Type InFeatureLevel);
+
+	/** Passed into MeshObjectFactory */
+	void* MeshObjectFactoryUserData;
+
 	/** Gets the skeletal mesh resource used for rendering the component. */
 	FSkeletalMeshRenderData* GetSkeletalMeshRenderData() const;
 
diff --git a/Engine/Source/Runtime/Engine/Private/Components/SkinnedMeshComponent.cpp b/Engine/Source/Runtime/Engine/Private/Components/SkinnedMeshComponent.cpp
--- a/Engine/Source/Runtime/Engine/Private/Components/SkinnedMeshComponent.cpp
+++ b/Engine/Source/Runtime/Engine/Private/Components/SkinnedMeshComponent.cpp
@@ -321,6 +321,8 @@ namespace FAnimUpdateRateManager
 
 USkinnedMeshComponent::USkinnedMeshComponent(const FObjectInitializer& ObjectInitializer)
 	: Super(ObjectInitializer)
+	, MeshObjectFactory(nullptr)
+	, MeshObjectFactoryUserData(nullptr)
 	, AnimUpdateRateParams(nullptr)
 {
 	bAutoActivate = true;
@@ -395,7 +397,7 @@ void USkinnedMeshComponent::UpdateMorphMaterialUsageOnProxy()
 				{
 					continue;
 				}
-				for (const FMorphTargetLODModel& MorphTargetLODModel : MorphTarget->MorphLODModels)
+				for (const FMorphTargetLODModel& MorphTargetLODModel : MorphTarget->GetMorphLODModels())
 				{
 					for (int32 SectionIndex : MorphTargetLODModel.SectionIndices)
 					{
@@ -552,29 +554,37 @@ void USkinnedMeshComponent::CreateRenderState_Concurrent(FRegisterComponentConte
 #endif
 	
 			// Also check if skeletal mesh has too many bones/chunk for GPU skinning.
-			if (bRenderStatic)
+			if (MeshObjectFactory)
 			{
-				// GPU skin vertex buffer + LocalVertexFactory
-				MeshObject = ::new FSkeletalMeshObjectStatic(this, SkelMeshRenderData, SceneFeatureLevel); 
+				MeshObject = MeshObjectFactory(MeshObjectFactoryUserData, this, SkelMeshRenderData, SceneFeatureLevel);
 			}
-			else if(ShouldCPUSkin())
+			if (!MeshObject)
 			{
-				MeshObject = ::new FSkeletalMeshObjectCPUSkin(this, SkelMeshRenderData, SceneFeatureLevel);
-			}
-			// don't silently enable CPU skinning for unsupported meshes, just do not render them, so their absence can be noticed and fixed
-			else if (!SkelMeshRenderData->RequiresCPUSkinning(SceneFeatureLevel, MinLODIndex)) 
-			{
-				MeshObject = ::new FSkeletalMeshObjectGPUSkin(this, SkelMeshRenderData, SceneFeatureLevel);
-			}
-			else
-			{
-				int32 MaxBonesPerChunk = SkelMeshRenderData->GetMaxBonesPerSection(MinLODIndex);
-				int32 MaxSupportedGPUSkinBones = FGPUBaseSkinVertexFactory::GetMaxGPUSkinBones();
-				int32 NumBoneInfluences = SkelMeshRenderData->GetNumBoneInfluences(MinLODIndex);
-				FString FeatureLevelName; GetFeatureLevelName(SceneFeatureLevel, FeatureLevelName);
+				// Also check if skeletal mesh has too many bones/chunk for GPU skinning.
+				if (bRenderStatic)
+				{
+					// GPU skin vertex buffer + LocalVertexFactory
+					MeshObject = ::new FSkeletalMeshObjectStatic(this, SkelMeshRenderData, SceneFeatureLevel);
+				}
+				else if (ShouldCPUSkin())
+				{
+					MeshObject = ::new FSkeletalMeshObjectCPUSkin(this, SkelMeshRenderData, SceneFeatureLevel);
+				}
+				// don't silently enable CPU skinning for unsupported meshes, just do not render them, so their absence can be noticed and fixed
+				else if (!SkelMeshRenderData->RequiresCPUSkinning(SceneFeatureLevel, MinLODIndex))
+				{
+					MeshObject = ::new FSkeletalMeshObjectGPUSkin(this, SkelMeshRenderData, SceneFeatureLevel);
+				}
+				else
+				{
+					int32 MaxBonesPerChunk = SkelMeshRenderData->GetMaxBonesPerSection(MinLODIndex);
+					int32 MaxSupportedGPUSkinBones = FGPUBaseSkinVertexFactory::GetMaxGPUSkinBones();
+					int32 NumBoneInfluences = SkelMeshRenderData->GetNumBoneInfluences(MinLODIndex);
+					FString FeatureLevelName; GetFeatureLevelName(SceneFeatureLevel, FeatureLevelName);
 
-				UE_LOG(LogSkinnedMeshComp, Warning, TEXT("SkeletalMesh %s, is not supported for current feature level (%s) and will not be rendered. MinLOD %d, NumBones %d (supported %d), NumBoneInfluences: %d"), 
-					*GetNameSafe(SkeletalMesh), *FeatureLevelName, MinLODIndex, MaxBonesPerChunk, MaxSupportedGPUSkinBones, NumBoneInfluences);
+					UE_LOG(LogSkinnedMeshComp, Warning, TEXT("SkeletalMesh %s, is not supported for current feature level (%s) and will not be rendered. MinLOD %d, NumBones %d (supported %d), NumBoneInfluences: %d"),
+						*GetNameSafe(SkeletalMesh), *FeatureLevelName, MinLODIndex, MaxBonesPerChunk, MaxSupportedGPUSkinBones, NumBoneInfluences);
+				}
 			}
 
 			//Allow the editor a chance to manipulate it before its added to the scene
diff --git a/Engine/Source/Runtime/Engine/Private/MorphMesh.cpp b/Engine/Source/Runtime/Engine/Private/MorphMesh.cpp
--- a/Engine/Source/Runtime/Engine/Private/MorphMesh.cpp
+++ b/Engine/Source/Runtime/Engine/Private/MorphMesh.cpp
@@ -36,6 +36,25 @@ void UMorphTarget::Serialize( FArchive& Ar )
 	}
 }
 
+void UMorphTarget::SerializeMemoryArchive(FMemoryArchive& Ar)
+{
+	int32 MorphLODModelNumber = 0;
+	if (Ar.IsLoading())
+	{
+		Ar << MorphLODModelNumber;
+		MorphLODModels.AddDefaulted(MorphLODModelNumber);
+	}
+	else
+	{
+		MorphLODModelNumber = MorphLODModels.Num();
+		Ar << MorphLODModelNumber;
+	}
+
+	for (int32 MorphIndex = 0; MorphIndex < MorphLODModelNumber; ++MorphIndex)
+	{
+		Ar << MorphLODModels[MorphIndex];
+	}
+}
 
 void UMorphTarget::GetResourceSizeEx(FResourceSizeEx& CumulativeResourceSize)
 {
diff --git a/Engine/Source/Runtime/Engine/Private/MorphTools.cpp b/Engine/Source/Runtime/Engine/Private/MorphTools.cpp
--- a/Engine/Source/Runtime/Engine/Private/MorphTools.cpp
+++ b/Engine/Source/Runtime/Engine/Private/MorphTools.cpp
@@ -20,11 +20,12 @@ struct FCompareMorphTargetDeltas
 	}
 };
 
-FMorphTargetDelta* UMorphTarget::GetMorphTargetDelta(int32 LODIndex, int32& OutNumDeltas)
+const FMorphTargetDelta* UMorphTarget::GetMorphTargetDelta(int32 LODIndex, int32& OutNumDeltas) const
 {
 	if(LODIndex < MorphLODModels.Num())
 	{
-		FMorphTargetLODModel& MorphModel = MorphLODModels[LODIndex];
+		// Calling GetMorphLODModels to potentially get from subclasses
+		const FMorphTargetLODModel& MorphModel = GetMorphLODModels()[LODIndex];
 		OutNumDeltas = MorphModel.Vertices.Num();
 		return MorphModel.Vertices.GetData();
 	}
@@ -33,7 +34,7 @@ FMorphTargetDelta* UMorphTarget::GetMorphTargetDelta(int32 LODIndex, int32& OutN
 	return NULL;
 }
 
-bool UMorphTarget::HasDataForLOD(int32 LODIndex) 
+bool UMorphTarget::HasDataForLOD(int32 LODIndex) const
 {
 	// If we have an entry for this LOD, and it has verts
 #if WITH_EDITOR
@@ -57,6 +58,16 @@ bool UMorphTarget::HasValidData() const
 	return false;
 }
 
+bool UMorphTarget::HasDataForSection(int32 LODIndex, int32 SectionIndex) const
+{
+	return HasDataForLOD(LODIndex) && MorphLODModels[LODIndex].SectionIndices.Contains(SectionIndex);
+}
+
+void UMorphTarget::EmptyMorphLODModels()
+{
+	MorphLODModels.Empty();
+}
+
 void UMorphTarget::DiscardVertexData()
 {
 	for (FMorphTargetLODModel& Model : MorphLODModels)
diff --git a/Engine/Source/Runtime/Engine/Private/SkeletalMesh.cpp b/Engine/Source/Runtime/Engine/Private/SkeletalMesh.cpp
--- a/Engine/Source/Runtime/Engine/Private/SkeletalMesh.cpp
+++ b/Engine/Source/Runtime/Engine/Private/SkeletalMesh.cpp
@@ -1994,14 +1994,14 @@ void FFinishBuildInternalData::ApplyMorphTargetsEditorData(USkeletalMesh* Skelet
 		{
 			ToDeleteMorphTargets.Remove(MorphTarget);
 		}
-		MorphTarget->MorphLODModels.Empty();
+		MorphTarget->EmptyMorphLODModels();
 		SkeletalMesh->GetMorphTargets().Add(MorphTarget);
 		const TArray<FMorphTargetLODModel>& MorphTargetLODModels = TargetNameAndMorphLODModels.Value;
 		int32 MorphLODModelNumber = MorphTargetLODModels.Num();
-		MorphTarget->MorphLODModels.AddDefaulted(MorphLODModelNumber);
+		MorphTarget->GetMorphLODModels().AddDefaulted(MorphLODModelNumber);
 		for (int32 MorphDataIndex = 0; MorphDataIndex < MorphLODModelNumber; ++MorphDataIndex)
 		{
-			MorphTarget->MorphLODModels[MorphDataIndex] = MorphTargetLODModels[MorphDataIndex];
+			MorphTarget->GetMorphLODModels()[MorphDataIndex] = MorphTargetLODModels[MorphDataIndex];
 		}
 	}
 	//Rebuild the mapping and rehook the curve data
@@ -2010,7 +2010,7 @@ void FFinishBuildInternalData::ApplyMorphTargetsEditorData(USkeletalMesh* Skelet
 	for (UMorphTarget* ToDeleteMorphTarget : ToDeleteMorphTargets)
 	{
 		ToDeleteMorphTarget->BaseSkelMesh = nullptr;
-		ToDeleteMorphTarget->MorphLODModels.Empty();
+		ToDeleteMorphTarget->EmptyMorphLODModels();
 
 		//Move the unused asset in the transient package and mark it pending kill
 		ToDeleteMorphTarget->Rename(nullptr, GetTransientPackage(), REN_ForceNoResetLoaders | REN_DoNotDirty | REN_DontCreateRedirectors | REN_NonTransactional);
diff --git a/Engine/Source/Runtime/Engine/Private/SkeletalMeshLODRenderData.cpp b/Engine/Source/Runtime/Engine/Private/SkeletalMeshLODRenderData.cpp
--- a/Engine/Source/Runtime/Engine/Private/SkeletalMeshLODRenderData.cpp
+++ b/Engine/Source/Runtime/Engine/Private/SkeletalMeshLODRenderData.cpp
@@ -310,7 +310,7 @@ void FSkeletalMeshLODRenderData::InitResources(bool bNeedsVertexColors, int32 LO
 			float MinimumValues[4] = { +FLT_MAX, +FLT_MAX, +FLT_MAX, +FLT_MAX };
 			UMorphTarget* MorphTarget = InMorphTargets[AnimIdx];
 			int32 NumSrcDeltas = 0;
-			FMorphTargetDelta* MorphDeltas = MorphTarget->GetMorphTargetDelta(LODIndex, NumSrcDeltas);
+			const FMorphTargetDelta* MorphDeltas = MorphTarget->GetMorphTargetDelta(LODIndex, NumSrcDeltas);
 
 			if (NumSrcDeltas == 0)
 			{
diff --git a/Engine/Source/Runtime/Engine/Private/SkeletalMeshRenderData.cpp b/Engine/Source/Runtime/Engine/Private/SkeletalMeshRenderData.cpp
--- a/Engine/Source/Runtime/Engine/Private/SkeletalMeshRenderData.cpp
+++ b/Engine/Source/Runtime/Engine/Private/SkeletalMeshRenderData.cpp
@@ -485,12 +485,7 @@ void FSkeletalMeshRenderData::Cache(const ITargetPlatform* TargetPlatform, USkel
 				{
 					FName MorphTargetName = Owner->GetMorphTargets()[MorphTargetIndex]->GetFName();
 					Ar << MorphTargetName;
-					int32 MorphLODModelNumber = Owner->GetMorphTargets()[MorphTargetIndex]->MorphLODModels.Num();
-					Ar << MorphLODModelNumber;
-					for (int32 MorphIndex = 0; MorphIndex < MorphLODModelNumber; ++MorphIndex)
-					{
-						Ar << Owner->GetMorphTargets()[MorphTargetIndex]->MorphLODModels[MorphIndex];
-					}
+					Owner->GetMorphTargets()[MorphTargetIndex]->SerializeMemoryArchive(Ar);
 				}
 				//No need to serialize the morph target mapping since we will rebuild the mapping when loading a ddc
 
diff --git a/Engine/Source/Runtime/Engine/Private/SkeletalRenderCPUSkin.cpp b/Engine/Source/Runtime/Engine/Private/SkeletalRenderCPUSkin.cpp
--- a/Engine/Source/Runtime/Engine/Private/SkeletalRenderCPUSkin.cpp
+++ b/Engine/Source/Runtime/Engine/Private/SkeletalRenderCPUSkin.cpp
@@ -565,7 +565,7 @@ struct FMorphTargetInfo
 	/** Index of next delta to try applying. This prevents us looking at every delta for every vertex. */
 	int32						NextDeltaIndex;
 	/** Array of deltas to apply to mesh, sorted based on the index of the base mesh vert that they affect. */
-	FMorphTargetDelta*			Deltas;
+	const FMorphTargetDelta*	Deltas;
 	/** How many deltas are in array */
 	int32						NumDeltas;
 };
@@ -1147,10 +1147,11 @@ static void CalculateMorphTargetWeights(FFinalSkinVertex* DestVertex, FSkeletalM
 
 	for (const UMorphTarget* Morphtarget : InMorphTargetOfInterest)
 	{
-		const FMorphTargetLODModel& MTLOD = Morphtarget->MorphLODModels[LODIndex];
-		for (int32 MorphVertexIndex = 0; MorphVertexIndex < MTLOD.Vertices.Num(); ++MorphVertexIndex)
+		int32 NumDeltas;
+		const FMorphTargetDelta* MTLODVertices = Morphtarget->GetMorphTargetDelta(LODIndex, NumDeltas);
+		for (int32 MorphVertexIndex = 0; MorphVertexIndex < NumDeltas; ++MorphVertexIndex)
 		{
-			FFinalSkinVertex* SetVert = DestVertex + MTLOD.Vertices[MorphVertexIndex].SourceIdx;
+			FFinalSkinVertex* SetVert = DestVertex + MTLODVertices[MorphVertexIndex].SourceIdx;
 			SetVert->U = 1.0f;
 			SetVert->V = 1.0f;
 		}
diff --git a/Engine/Source/Runtime/Engine/Private/SkeletalRenderGPUSkin.cpp b/Engine/Source/Runtime/Engine/Private/SkeletalRenderGPUSkin.cpp
--- a/Engine/Source/Runtime/Engine/Private/SkeletalRenderGPUSkin.cpp
+++ b/Engine/Source/Runtime/Engine/Private/SkeletalRenderGPUSkin.cpp
@@ -154,6 +154,7 @@ FSkeletalMeshObjectGPUSkin::FSkeletalMeshObjectGPUSkin(USkinnedMeshComponent* In
 	,	bMorphNeedsUpdateDeferred(false)
 	,	bMorphResourcesInitialized(false)
 	, 	LastBoneTransformRevisionNumber(0)
+	,	bAlwaysUpdateMorphVertexBuffer(false)
 {
 	// create LODs to match the base mesh
 	LODs.Empty(SkeletalMeshRenderData->LODRenderData.Num());
@@ -388,6 +389,7 @@ void FSkeletalMeshObjectGPUSkin::UpdateDynamicData_RenderThread(FGPUSkinCache* G
 	bool bMorphNeedsUpdate=false;
 	// figure out if the morphing vertex buffer needs to be updated. compare old vs new active morphs
 	bMorphNeedsUpdate = 
+		bAlwaysUpdateMorphVertexBuffer ||
 		(bMorphNeedsUpdateDeferred && bNeedsUpdateDeferred) || // the need for an update sticks
 		(DynamicData ? (DynamicData->LODIndex != InDynamicData->LODIndex ||
 		!DynamicData->ActiveMorphTargetsEqual(InDynamicData->ActiveMorphTargets, InDynamicData->MorphTargetWeights))
@@ -499,12 +501,13 @@ void FSkeletalMeshObjectGPUSkin::ProcessUpdatedDynamicData(EGPUSkinCacheEntryMod
 	// if hasn't been updated, force update again
 	bMorphNeedsUpdate = LOD.MorphVertexBuffer.bHasBeenUpdated ? bMorphNeedsUpdate : true;
 	bMorphNeedsUpdate |= (GForceUpdateMorphTargets != 0);
+	bMorphNeedsUpdate |= bAlwaysUpdateMorphVertexBuffer;
 
 	const FSkeletalMeshLODRenderData& LODData = SkeletalMeshRenderData->LODRenderData[LODIndex];
 	const TArray<FSkelMeshRenderSection>& Sections = GetRenderSections(LODIndex);
 
 	// Only consider morphs with active curves and data to deform.
-	const bool bMorph = DynamicData->NumWeightedActiveMorphTargets > 0 && LODData.GetNumVertices() > 0;
+	const bool bMorph = bAlwaysUpdateMorphVertexBuffer || (DynamicData->NumWeightedActiveMorphTargets > 0 && LODData.GetNumVertices() > 0);
 
 	// use correct vertex factories based on alternate weights usage
 	FVertexFactoryData& VertexFactoryData = LOD.GPUSkinVertexFactories;
@@ -579,18 +582,7 @@ void FSkeletalMeshObjectGPUSkin::ProcessUpdatedDynamicData(EGPUSkinCacheEntryMod
 		// only update if the morph data changed and there are weighted morph targets
 		if(bMorphNeedsUpdate)
 		{
-			QUICK_SCOPE_CYCLE_COUNTER(STAT_FSkeletalMeshObjectGPUSkin_ProcessUpdatedDynamicData_UpdateMorphBuffer);
-			if (UseGPUMorphTargets(FeatureLevel))
-			{
-				// sometimes this goes out of bound, we'll ensure here
-				ensureAlways(DynamicData->MorphTargetWeights.Num() == LODData.MorphTargetVertexInfoBuffers.GetNumMorphs());
-				LOD.UpdateMorphVertexBufferGPU(RHICmdList, DynamicData->MorphTargetWeights, LODData.MorphTargetVertexInfoBuffers, DynamicData->SectionIdsUseByActiveMorphTargets, GetDebugName(), Mode);
-			}
-			else
-			{
-				// update the morph data for the lod (before SkinCache)
-				LOD.UpdateMorphVertexBufferCPU(DynamicData->ActiveMorphTargets, DynamicData->MorphTargetWeights);
-			}		
+			UpdateMorphVertexBuffer(RHICmdList, Mode, LOD, LODData);
 		}
 	}
 	else
@@ -617,7 +609,7 @@ void FSkeletalMeshObjectGPUSkin::ProcessUpdatedDynamicData(EGPUSkinCacheEntryMod
 				}
 				else
 				{
-					if(DynamicData->NumWeightedActiveMorphTargets > 0 && DynamicData->SectionIdsUseByActiveMorphTargets.Contains(SectionIdx))
+					if(bAlwaysUpdateMorphVertexBuffer || (DynamicData->NumWeightedActiveMorphTargets > 0 && DynamicData->SectionIdsUseByActiveMorphTargets.Contains(SectionIdx)))
 					{
 						VertexFactory = VertexFactoryData.MorphVertexFactories[SectionIdx].Get();
 					}
@@ -725,6 +717,22 @@ void FSkeletalMeshObjectGPUSkin::ProcessUpdatedDynamicData(EGPUSkinCacheEntryMod
 	}
 }
 
+void FSkeletalMeshObjectGPUSkin::UpdateMorphVertexBuffer(FRHICommandListImmediate& RHICmdList, EGPUSkinCacheEntryMode Mode, FSkeletalMeshObjectLOD& LOD, const FSkeletalMeshLODRenderData& LODData)
+{
+	QUICK_SCOPE_CYCLE_COUNTER(STAT_FSkeletalMeshObjectGPUSkin_ProcessUpdatedDynamicData_UpdateMorphBuffer);
+	if (UseGPUMorphTargets(FeatureLevel))
+	{
+		// sometimes this goes out of bound, we'll ensure here
+		ensureAlways(DynamicData->MorphTargetWeights.Num() == LODData.MorphTargetVertexInfoBuffers.GetNumMorphs());
+		LOD.UpdateMorphVertexBufferGPU(RHICmdList, DynamicData->MorphTargetWeights, LODData.MorphTargetVertexInfoBuffers, DynamicData->SectionIdsUseByActiveMorphTargets, GetDebugName(), Mode);
+	}
+	else
+	{
+		// update the morph data for the lod (before SkinCache)
+		LOD.UpdateMorphVertexBufferCPU(DynamicData->ActiveMorphTargets, DynamicData->MorphTargetWeights);
+	}
+}
+
 TArray<float> FSkeletalMeshObjectGPUSkin::FSkeletalMeshObjectLOD::MorphAccumulatedWeightArray;
 
 void FGPUMorphUpdateCS::SetParameters(FRHICommandList& RHICmdList, const FVector4& LocalScale, float WeightScale, const FMorphTargetVertexInfoBuffers& MorphTargetVertexInfoBuffers, FMorphVertexBuffer& MorphVertexBuffer)
@@ -1073,7 +1081,7 @@ void FSkeletalMeshObjectGPUSkin::FSkeletalMeshObjectLOD::UpdateMorphVertexBuffer
 
 				// Get deltas
 				int32 NumDeltas;
-				FMorphTargetDelta* Deltas = MorphTarget.MorphTarget->GetMorphTargetDelta(LODIndex, NumDeltas);
+				const FMorphTargetDelta* Deltas = MorphTarget.MorphTarget->GetMorphTargetDelta(LODIndex, NumDeltas);
 
 				// iterate over the vertices that this lod model has changed
 				for (int32 MorphVertIdx = 0; MorphVertIdx < NumDeltas; MorphVertIdx++)
@@ -1174,12 +1182,16 @@ const FVertexFactory* FSkeletalMeshObjectGPUSkin::GetSkinVertexFactory(const FSc
 	{
 		for (FActiveMorphTarget& ActiveMorphTarget : DynamicData->ActiveMorphTargets)
 		{
-			if (ActiveMorphTarget.MorphTarget->MorphLODModels.IsValidIndex(LODIndex) && ActiveMorphTarget.MorphTarget->MorphLODModels[LODIndex].SectionIndices.Contains(ChunkIdx))
+			if (ActiveMorphTarget.MorphTarget->HasDataForSection(LODIndex, ChunkIdx))
+			{
+				return LOD.GPUSkinVertexFactories.MorphVertexFactories[ChunkIdx].Get();
+			}
+		}
+	}
+	else if (bAlwaysUpdateMorphVertexBuffer)
 	{
 		return LOD.GPUSkinVertexFactories.MorphVertexFactories[ChunkIdx].Get();
 	}
-		}
-	}
 
 	// use the default gpu skin vertex factory
 	return LOD.GPUSkinVertexFactories.VertexFactories[ChunkIdx].Get();
@@ -2098,9 +2110,10 @@ void FDynamicSkelMeshObjectDataGPUSkin::InitDynamicSkelMeshObjectDataGPUSkin(
 			Morph.MorphTarget->HasDataForLOD(LODIndex) ) 
 		{
 			NumWeightedActiveMorphTargets++;
-			for (int32 SecId = 0; SecId < Morph.MorphTarget->MorphLODModels[LODIndex].SectionIndices.Num(); ++SecId)
+			const TArray<int32>& MorphSectionIndices = Morph.MorphTarget->GetMorphLODModels()[LODIndex].SectionIndices;
+			for (int32 SecId = 0; SecId < MorphSectionIndices.Num(); ++SecId)
 			{
-				SectionIdsUseByActiveMorphTargets.AddUnique(Morph.MorphTarget->MorphLODModels[LODIndex].SectionIndices[SecId]);
+				SectionIdsUseByActiveMorphTargets.AddUnique(MorphSectionIndices[SecId]);
 			}
 		}
 		else
diff --git a/Engine/Source/Runtime/Engine/Private/SkeletalRenderGPUSkin.h b/Engine/Source/Runtime/Engine/Private/SkeletalRenderGPUSkin.h
--- a/Engine/Source/Runtime/Engine/Private/SkeletalRenderGPUSkin.h
+++ b/Engine/Source/Runtime/Engine/Private/SkeletalRenderGPUSkin.h
@@ -25,7 +25,7 @@ class FGPUSkinCache;
 * Stores the updated matrices needed to skin the verts.
 * Created by the game thread and sent to the rendering thread as an update 
 */
-class FDynamicSkelMeshObjectDataGPUSkin
+class ENGINE_API FDynamicSkelMeshObjectDataGPUSkin
 {
 	/**
 	* Constructor, these are recycled, so you never use a constructor
@@ -363,7 +363,7 @@ public:
 /**
  * Render data for a GPU skinned mesh
  */
-class FSkeletalMeshObjectGPUSkin : public FSkeletalMeshObject
+class ENGINE_API FSkeletalMeshObjectGPUSkin : public FSkeletalMeshObject
 {
 public:
 	/** @param	InSkeletalMeshComponent - skeletal mesh primitive we want to render */
@@ -469,7 +469,7 @@ public:
 
 	virtual void RefreshClothingTransforms(const FMatrix& InNewLocalToWorld, uint32 FrameNumber) override;
 	virtual void UpdateSkinWeightBuffer(USkinnedMeshComponent* InMeshComponent) override;
-private:
+protected:
 
 	/**
 	 * Vertex factories and their matrix arrays
@@ -660,6 +660,8 @@ private:
 
 	void ProcessUpdatedDynamicData(EGPUSkinCacheEntryMode Mode, FGPUSkinCache* GPUSkinCache, FRHICommandListImmediate& RHICmdList, uint32 FrameNumberToPrepare, uint32 RevisionNumber, bool bMorphNeedsUpdate, int32 LODIndex);
 
+	virtual void UpdateMorphVertexBuffer(FRHICommandListImmediate& RHICmdList, EGPUSkinCacheEntryMode Mode, FSkeletalMeshObjectLOD& LOD, const FSkeletalMeshLODRenderData& LODData);
+
 	void WaitForRHIThreadFenceForDynamicData();
 
 	/** Render data for each LOD */
@@ -682,6 +684,13 @@ private:
 
 	/** last updated bone transform revision number */
 	uint32 LastBoneTransformRevisionNumber;
+
+	/** true to indicate that a subclass is handling the update of the morph vertex buffer and that it should always be called */
+	bool bAlwaysUpdateMorphVertexBuffer;
+
+private:
+	FSkeletalMeshObjectGPUSkin(const FSkeletalMeshObjectGPUSkin&);
+	FSkeletalMeshObjectGPUSkin& operator=(const FSkeletalMeshObjectGPUSkin&);
 };
 
 
diff --git a/Engine/Source/Runtime/Engine/Public/Rendering/SkeletalMeshRenderData.h b/Engine/Source/Runtime/Engine/Public/Rendering/SkeletalMeshRenderData.h
--- a/Engine/Source/Runtime/Engine/Public/Rendering/SkeletalMeshRenderData.h
+++ b/Engine/Source/Runtime/Engine/Public/Rendering/SkeletalMeshRenderData.h
@@ -82,10 +82,10 @@ public:
 	ENGINE_API bool RequiresCPUSkinning(ERHIFeatureLevel::Type FeatureLevel, int32 MinLODIdx) const;
 
 	/** Returns the number of bone influences per vertex. */
-	uint32 GetNumBoneInfluences() const;
+	ENGINE_API uint32 GetNumBoneInfluences() const;
 
 	/** Returns the number of bone influences per vertex starting at MinLODIdx. */
-	uint32 GetNumBoneInfluences(int32 MinLODIdx) const;
+	ENGINE_API uint32 GetNumBoneInfluences(int32 MinLODIdx) const;
 	
 	/**
 	* Computes the maximum number of bones per section used to render this mesh.
diff --git a/Engine/Source/Runtime/OpenGLDrv/Private/OpenGLDevice.cpp b/Engine/Source/Runtime/OpenGLDrv/Private/OpenGLDevice.cpp
--- a/Engine/Source/Runtime/OpenGLDrv/Private/OpenGLDevice.cpp
+++ b/Engine/Source/Runtime/OpenGLDrv/Private/OpenGLDevice.cpp
@@ -956,6 +956,12 @@ static void InitRHICapabilitiesForGL()
 	SetupTextureFormat( PF_FloatR11G11B10,		FOpenGLTextureFormat( GL_R11F_G11F_B10F,		GL_R11F_G11F_B10F,		GL_RGB,				GL_UNSIGNED_INT_10F_11F_11F_REV,	false,  false));
 	SetupTextureFormat( PF_FloatRGBA,			FOpenGLTextureFormat( GL_RGBA16F,				GL_RGBA16F,				GL_RGBA,			GL_HALF_FLOAT,						false,	false));
 
+	SetupTextureFormat( PF_R8G8_UINT,			FOpenGLTextureFormat( GL_RG8UI,					GL_RG8UI,				GL_RG_INTEGER,		GL_UNSIGNED_BYTE,					false, false));
+	SetupTextureFormat( PF_R32G32B32_UINT,		FOpenGLTextureFormat( GL_RGB32UI,				GL_RGB32UI,				GL_RGB_INTEGER,		GL_UNSIGNED_INT,					false, false));
+	SetupTextureFormat( PF_R32G32B32_SINT,		FOpenGLTextureFormat( GL_RGB32I,				GL_RGB32I,				GL_RGB_INTEGER,		GL_INT,								false, false));
+	SetupTextureFormat( PF_R32G32B32F,			FOpenGLTextureFormat( GL_RGB32F,				GL_RGB32F,				GL_RGB,				GL_FLOAT,							false, false));
+	SetupTextureFormat( PF_R8_SINT,				FOpenGLTextureFormat( GL_R8I,					GL_R8I,					GL_RED_INTEGER,		GL_BYTE,							false, false));
+
 #if PLATFORM_DESKTOP
 	CA_SUPPRESS(6286);
 	if (PLATFORM_DESKTOP || FOpenGL::GetFeatureLevel() >= ERHIFeatureLevel::SM5)
@@ -970,6 +976,7 @@ static void InitRHICapabilitiesForGL()
 		SetupTextureFormat(PF_R8G8B8A8_UINT, FOpenGLTextureFormat(GL_RGBA8UI, GL_RGBA8UI, GL_RGBA_INTEGER, GL_UNSIGNED_BYTE, false, false));
 
 		SetupTextureFormat(PF_G16R16, FOpenGLTextureFormat(GL_RG16, GL_RG16, GL_RG, GL_UNSIGNED_SHORT, false, false));
+		SetupTextureFormat(PF_G16R16_SNORM, FOpenGLTextureFormat(GL_RG16_SNORM, GL_RG16_SNORM, GL_RG, GL_SHORT, false, false));
 	}
 	else
 #endif // PLATFORM_DESKTOP
diff --git a/Engine/Source/Runtime/RHI/Private/RHI.cpp b/Engine/Source/Runtime/RHI/Private/RHI.cpp
--- a/Engine/Source/Runtime/RHI/Private/RHI.cpp
+++ b/Engine/Source/Runtime/RHI/Private/RHI.cpp
@@ -2490,6 +2490,13 @@ FPixelFormatInfo    GPixelFormats[PF_MAX] =
 	FPixelFormatInfo(PF_ASTC_8x8_HDR,           TEXT("ASTC_8x8_HDR"),      8,          8,          1,          16,         4,              0),
 	FPixelFormatInfo(PF_ASTC_10x10_HDR,         TEXT("ASTC_10x10_HDR"),    10,         10,         1,          16,         4,              0),
 	FPixelFormatInfo(PF_ASTC_12x12_HDR,         TEXT("ASTC_12x12_HDR"),    12,         12,         1,          16,         4,              0),
+
+	FPixelFormatInfo(PF_G16R16_SNORM,       TEXT("G16R16_SNORM"),       1,          1,          1,          4,          2,              1),
+	FPixelFormatInfo(PF_R8G8_UINT,          TEXT("R8G8_UINT"),          1,          1,          1,          2,          2,              1),
+	FPixelFormatInfo(PF_R32G32B32_UINT,     TEXT("R32G32B32_UINT"),     1,          1,          1,          12,         3,              1),
+	FPixelFormatInfo(PF_R32G32B32_SINT,     TEXT("R32G32B32_SINT"),     1,          1,          1,          12,         3,              1),
+	FPixelFormatInfo(PF_R32G32B32F,         TEXT("R32G32B32F"),         1,          1,          1,          12,         3,              1),
+	FPixelFormatInfo(PF_R8_SINT,            TEXT("R8_SINT"),            1,          1,          1,          1,          1,              1),
 };
 
 static struct FValidatePixelFormats
diff --git a/Engine/Source/Runtime/RenderCore/Private/RenderUtils.cpp b/Engine/Source/Runtime/RenderCore/Private/RenderUtils.cpp
--- a/Engine/Source/Runtime/RenderCore/Private/RenderUtils.cpp
+++ b/Engine/Source/Runtime/RenderCore/Private/RenderUtils.cpp
@@ -943,6 +943,12 @@ EPixelFormatChannelFlags GetPixelFormatValidChannels(EPixelFormat InPixelFormat)
 		EPixelFormatChannelFlags::RGB,		// PF_ASTC_8x8_HDR
 		EPixelFormatChannelFlags::RGB,		// PF_ASTC_10x10_HDR
 		EPixelFormatChannelFlags::RGB,		// PF_ASTC_12x12_HDR
+		EPixelFormatChannelFlags::RG,		// PF_G16R16_SNORM
+		EPixelFormatChannelFlags::RG,		// PF_R8G8_UINT
+		EPixelFormatChannelFlags::RGB,		// PF_R32G32B32_UINT
+		EPixelFormatChannelFlags::RGB,		// PF_R32G32B32_SINT
+		EPixelFormatChannelFlags::RGB,		// PF_R32G32B32F
+		EPixelFormatChannelFlags::R,		// PF_R8_SINT
 	};
 	static_assert(UE_ARRAY_COUNT(PixelFormatToChannelFlags) == (uint8)PF_MAX, "Missing pixel format");
 	return (InPixelFormat < PF_MAX) ? PixelFormatToChannelFlags[(uint8)InPixelFormat] : EPixelFormatChannelFlags::None;
diff --git a/Engine/Source/Runtime/Renderer/Private/SystemTextures.cpp b/Engine/Source/Runtime/Renderer/Private/SystemTextures.cpp
--- a/Engine/Source/Runtime/Renderer/Private/SystemTextures.cpp
+++ b/Engine/Source/Runtime/Renderer/Private/SystemTextures.cpp
@@ -1062,13 +1062,16 @@ template<typename TInType>
 void InitializeData(const TInType& InData, EPixelFormat InFormat, uint8* OutData, uint32& OutByteCount)
 {
 	// If a new format is added insure that it is either supported here, or at least flagged as not supported
-	static_assert(PF_MAX == 78);
+	static_assert(PF_MAX == 84);
 
 	switch (InFormat)
 	{
 		// 32bits
 		case PF_R32G32B32A32_UINT:		{ FormatData<EDefaultInputType::Typed, TInType,	uint32,  0, 1, 2, 3>	(InData, OutData, OutByteCount); } break;
 		case PF_A32B32G32R32F:			{ FormatData<EDefaultInputType::Typed, TInType,	float,   3, 2, 1, 0>	(InData, OutData, OutByteCount); } break;
+		case PF_R32G32B32_UINT:			{ FormatData<EDefaultInputType::Typed, TInType, uint32,  0, 1, 2>		(InData, OutData, OutByteCount); } break;
+		case PF_R32G32B32_SINT:			{ FormatData<EDefaultInputType::Typed, TInType, int32,   0, 1, 2>		(InData, OutData, OutByteCount); } break;
+		case PF_R32G32B32F:				{ FormatData<EDefaultInputType::Typed, TInType, float,   0, 1, 2>		(InData, OutData, OutByteCount); } break;
 		case PF_R32G32_UINT:			{ FormatData<EDefaultInputType::Typed, TInType,	uint32,  0, 1>			(InData, OutData, OutByteCount); } break;
 		case PF_G32R32F:				{ FormatData<EDefaultInputType::Typed, TInType,	float,   1, 0>			(InData, OutData, OutByteCount); } break;
 		case PF_R32_UINT:				{ FormatData<EDefaultInputType::Typed, TInType,	uint32>					(InData, OutData, OutByteCount); } break;
@@ -1084,6 +1087,7 @@ void InitializeData(const TInType& InData, EPixelFormat InFormat, uint8* OutData
 		case PF_FloatRGBA:				{ FormatData<EDefaultInputType::Typed, TInType,	FFloat16, 0, 1, 2, 3>	(InData, OutData, OutByteCount); } break;
 		case PF_R16G16_UINT:			{ FormatData<EDefaultInputType::Typed, TInType,	uint16,   0, 1>			(InData, OutData, OutByteCount); } break;
 		case PF_G16R16:					{ FormatData<EDefaultInputType::UNorm, TInType,	uint16,   1, 0>			(InData, OutData, OutByteCount); } break;
+		case PF_G16R16_SNORM:			{ FormatData<EDefaultInputType::SNorm, TInType,	int16,    1, 0>			(InData, OutData, OutByteCount); } break;
 		case PF_G16R16F:				{ FormatData<EDefaultInputType::Typed, TInType,	FFloat16, 0, 1>			(InData, OutData, OutByteCount); } break;
 		case PF_G16R16F_FILTER:			{ FormatData<EDefaultInputType::Typed, TInType,	FFloat16, 0, 1>			(InData, OutData, OutByteCount); } break;
 		case PF_R16F_FILTER:			{ FormatData<EDefaultInputType::Typed, TInType,	FFloat16>				(InData, OutData, OutByteCount); } break;
@@ -1099,7 +1103,9 @@ void InitializeData(const TInType& InData, EPixelFormat InFormat, uint8* OutData
 		case PF_R8G8B8A8_UINT:			{ FormatData<EDefaultInputType::Typed, TInType,	uint8,    0, 1, 2, 3>	(InData, OutData, OutByteCount); } break;
 		case PF_R8G8B8A8_SNORM:			{ FormatData<EDefaultInputType::SNorm, TInType,	int8,     0, 1, 2, 3>	(InData, OutData, OutByteCount); } break;
 		case PF_R8G8:					{ FormatData<EDefaultInputType::UNorm, TInType,	uint8,    0, 1>			(InData, OutData, OutByteCount); } break;
+		case PF_R8G8_UINT:				{ FormatData<EDefaultInputType::Typed, TInType, uint8,    0, 1>			(InData, OutData, OutByteCount); } break;
 		case PF_R8_UINT:				{ FormatData<EDefaultInputType::Typed, TInType,	uint8>					(InData, OutData, OutByteCount); } break;
+		case PF_R8_SINT:				{ FormatData<EDefaultInputType::Typed, TInType, int8>					(InData, OutData, OutByteCount); } break;
 		case PF_R8:						{ FormatData<EDefaultInputType::UNorm, TInType,	uint8>					(InData, OutData, OutByteCount); } break;
 		case PF_G8:						{ FormatData<EDefaultInputType::UNorm, TInType,	uint8>					(InData, OutData, OutByteCount); } break;
 		case PF_L8:						{ FormatData<EDefaultInputType::UNorm, TInType,	uint8>					(InData, OutData, OutByteCount); } break;
diff --git a/Engine/Source/Runtime/VulkanRHI/Private/VulkanDevice.cpp b/Engine/Source/Runtime/VulkanRHI/Private/VulkanDevice.cpp
--- a/Engine/Source/Runtime/VulkanRHI/Private/VulkanDevice.cpp
+++ b/Engine/Source/Runtime/VulkanRHI/Private/VulkanDevice.cpp
@@ -719,6 +719,24 @@ void FVulkanDevice::SetupFormats()
 	MapFormatSupport(PF_R8, VK_FORMAT_R8_UNORM);
 	SetComponentMapping(PF_R8, VK_COMPONENT_SWIZZLE_R, VK_COMPONENT_SWIZZLE_ZERO, VK_COMPONENT_SWIZZLE_ZERO, VK_COMPONENT_SWIZZLE_ONE);
 
+	MapFormatSupport(PF_G16R16_SNORM, VK_FORMAT_R16G16_SNORM);
+	SetComponentMapping(PF_G16R16_SNORM, VK_COMPONENT_SWIZZLE_R, VK_COMPONENT_SWIZZLE_G, VK_COMPONENT_SWIZZLE_ZERO, VK_COMPONENT_SWIZZLE_ONE);
+
+	MapFormatSupport(PF_R8G8_UINT, VK_FORMAT_R8G8_UINT);
+	SetComponentMapping(PF_R8G8_UINT, VK_COMPONENT_SWIZZLE_R, VK_COMPONENT_SWIZZLE_G, VK_COMPONENT_SWIZZLE_ZERO, VK_COMPONENT_SWIZZLE_ONE);
+
+	MapFormatSupport(PF_R32G32B32_UINT, VK_FORMAT_R32G32B32_UINT);
+	SetComponentMapping(PF_R32G32B32_UINT, VK_COMPONENT_SWIZZLE_R, VK_COMPONENT_SWIZZLE_G, VK_COMPONENT_SWIZZLE_B, VK_COMPONENT_SWIZZLE_ONE);
+
+	MapFormatSupport(PF_R32G32B32_SINT, VK_FORMAT_R32G32B32_SINT);
+	SetComponentMapping(PF_R32G32B32_SINT, VK_COMPONENT_SWIZZLE_R, VK_COMPONENT_SWIZZLE_G, VK_COMPONENT_SWIZZLE_B, VK_COMPONENT_SWIZZLE_ONE);
+
+	MapFormatSupport(PF_R32G32B32F, VK_FORMAT_R32G32B32_SFLOAT);
+	SetComponentMapping(PF_R32G32B32F, VK_COMPONENT_SWIZZLE_R, VK_COMPONENT_SWIZZLE_G, VK_COMPONENT_SWIZZLE_B, VK_COMPONENT_SWIZZLE_ONE);
+
+	MapFormatSupport(PF_R8_SINT, VK_FORMAT_R8_SINT);
+	SetComponentMapping(PF_R8_SINT, VK_COMPONENT_SWIZZLE_R, VK_COMPONENT_SWIZZLE_ZERO, VK_COMPONENT_SWIZZLE_ZERO, VK_COMPONENT_SWIZZLE_ONE);
+
 	if (FVulkanPlatform::SupportsBCTextureFormats())
 	{
 		MapFormatSupport(PF_DXT1, VK_FORMAT_BC1_RGB_UNORM_BLOCK);	// Also what OpenGL expects (RGBA instead RGB, but not SRGB)
diff --git a/Engine/Source/Runtime/Windows/D3D11RHI/Private/D3D11Device.cpp b/Engine/Source/Runtime/Windows/D3D11RHI/Private/D3D11Device.cpp
--- a/Engine/Source/Runtime/Windows/D3D11RHI/Private/D3D11Device.cpp
+++ b/Engine/Source/Runtime/Windows/D3D11RHI/Private/D3D11Device.cpp
@@ -194,6 +194,13 @@ FD3D11DynamicRHI::FD3D11DynamicRHI(IDXGIFactory1* InDXGIFactory1, D3D_FEATURE_LE
 	GPixelFormats[PF_NV12].PlatformFormat = DXGI_FORMAT_NV12;
 	GPixelFormats[PF_NV12].Supported = true;
 
+	GPixelFormats[PF_G16R16_SNORM	].PlatformFormat = DXGI_FORMAT_R16G16_SNORM;
+	GPixelFormats[PF_R8G8_UINT		].PlatformFormat = DXGI_FORMAT_R8G8_UINT;
+	GPixelFormats[PF_R32G32B32_UINT	].PlatformFormat = DXGI_FORMAT_R32G32B32_UINT;
+	GPixelFormats[PF_R32G32B32_SINT	].PlatformFormat = DXGI_FORMAT_R32G32B32_SINT;
+	GPixelFormats[PF_R32G32B32F		].PlatformFormat = DXGI_FORMAT_R32G32B32_FLOAT;
+	GPixelFormats[PF_R8_SINT		].PlatformFormat = DXGI_FORMAT_R8_SINT;
+
 	GSupportsSeparateRenderTargetBlendState = true;
 	GMaxTextureDimensions = D3D11_REQ_TEXTURE2D_U_OR_V_DIMENSION;
 	GMaxCubeTextureDimensions = D3D11_REQ_TEXTURECUBE_DIMENSION;
