PR #8313: FVisualLogger: GetTimeStampFunc, pass in function to sync time between server and client (Contributed by [REDACTED_PERSON])

diff --git a/Engine/Source/Runtime/Engine/Private/VisualLogger/VisualLogger.cpp b/Engine/Source/Runtime/Engine/Private/VisualLogger/VisualLogger.cpp
--- a/Engine/Source/Runtime/Engine/Private/VisualLogger/VisualLogger.cpp
+++ b/Engine/Source/Runtime/Engine/Private/VisualLogger/VisualLogger.cpp
@@ -72,7 +72,7 @@ bool FVisualLogger::CheckVisualLogInputInternal(const UObject* Object, const FNa
 		return false;
 	}
 
-	*CurrentEntry = VisualLogger.GetEntryToWrite(Object, (*World)->TimeSeconds);
+	*CurrentEntry = VisualLogger.GetEntryToWrite(Object, VisualLogger.GetTimeStampForObject(Object));
 	if (*CurrentEntry == nullptr)
 	{
 		return false;
@@ -81,6 +81,26 @@ bool FVisualLogger::CheckVisualLogInputInternal(const UObject* Object, const FNa
 	return true;
 }
 
+float FVisualLogger::GetTimeStampForObject(const UObject* Object) const
+{
+	if (GetTimeStampFunc)
+	{
+		return GetTimeStampFunc(Object);
+	}
+
+	if (const UWorld* World = GEngine->GetWorldFromContextObject(Object, EGetWorldErrorMode::ReturnNull))
+	{
+		return World->TimeSeconds;
+	}
+
+	return 0.0f;
+}
+
+void FVisualLogger::SetGetTimeStampFunc(const TFunction<float(const UObject*)> Function)
+{
+	GetTimeStampFunc = Function;
+}
+
 void FVisualLogger::AddWhitelistedClass(UClass& InClass)
 {
 	ClassWhitelist.AddUnique(&InClass);
diff --git a/Engine/Source/Runtime/Engine/Public/VisualLogger/VisualLogger.h b/Engine/Source/Runtime/Engine/Public/VisualLogger/VisualLogger.h
--- a/Engine/Source/Runtime/Engine/Public/VisualLogger/VisualLogger.h
+++ b/Engine/Source/Runtime/Engine/Public/VisualLogger/VisualLogger.h
@@ -532,6 +532,12 @@ public:
 
 	/** internal check for each usage of visual logger */
 	static bool CheckVisualLogInputInternal(const UObject* Object, const FName& CategoryName, ELogVerbosity::Type Verbosity, UWorld **World, FVisualLogEntry **CurrentEntry);
+	
+	/** Returns time stamp for object */
+	float GetTimeStampForObject(const UObject* Object) const;
+
+	/** Sets function to call to get a timestamp instead of the default implementation (i.e. world time) */
+	void SetGetTimeStampFunc(TFunction<float(const UObject*)> Function);
 
 	typedef TMap<UObject*, TArray<TWeakObjectPtr<const UObject> > > FOwnerToChildrenRedirectionMap;
 	static FOwnerToChildrenRedirectionMap& GetRedirectionMap(const UObject* InObject);
@@ -604,6 +610,8 @@ protected:
 	FVisualLogFilenameGetterDelegate LogFileNameGetter;
 	/** Handle for timer used to serialize all waiting logs */
 	FTimerHandle VisualLoggerCleanupTimerHandle;
+	/** function to call when getting the time stamp */
+	TFunction<float(const UObject*)> GetTimeStampFunc;
 
 	// if set we are recording and collecting all vlog data
 	static int32 bIsRecording;
