Fixing test config
#rnx

diff --git a/Engine/Source/Runtime/Engine/Private/UnrealEngine.cpp b/Engine/Source/Runtime/Engine/Private/UnrealEngine.cpp
--- a/Engine/Source/Runtime/Engine/Private/UnrealEngine.cpp
+++ b/Engine/Source/Runtime/Engine/Private/UnrealEngine.cpp
@@ -6839,7 +6839,7 @@ bool UEngine::HandleSkeletalMeshReportCommand(const TCHAR* Cmd, FOutputDevice& A
 	auto GetResourceSizeKB = [&](const FRHIResource* Resource) -> uint32
 	{
 		int32 SizeInKB = 0;
-#if RHI_WANT_RESOURCE_INFO
+#if RHI_ENABLE_RESOURCE_INFO
 		FRHIResourceInfo ResourceInfo;
 		if (Resource && Resource->GetResourceInfo(ResourceInfo))
 		{
