Change default behavior of "fasttextureencode" in cook/ddcfill commandlets
previously would select fast encode for "WithEditor" target platforms
change makes these commandlets default to shipping quality only
commandlets can explicitly use the -fasttextureencode option if desired

diff --git a/Engine/Source/Runtime/Engine/Private/TextureDerivedData.cpp b/Engine/Source/Runtime/Engine/Private/TextureDerivedData.cpp
--- a/Engine/Source/Runtime/Engine/Private/TextureDerivedData.cpp
+++ b/Engine/Source/Runtime/Engine/Private/TextureDerivedData.cpp
@@ -419,28 +419,12 @@ static ETextureFastEncode GetDesiredFastTextureEncode(	const ITargetPlatform& Cu
 		check( CachedFastTextureEncodeOption == 1 );
 		// no option specified, default behavior :
 
-		#if 1
-
-		// control by is target platform "with editor" ?
-		
-		// can't store this in CachedFastTextureEncodeOption if it depends on CurrentPlatform
-		//	because we can get called multiple times for different platform in cook/ddc operations
-
-		// this is the old behavior
-		
-		if ( CurrentPlatform.HasEditorOnlyData() )
-			return ETextureFastEncode::TryOffEncodeFast;
-		else
-			return ETextureFastEncode::Off;
-
-		#else
-
 		// control by is in interactive editor (not cook / ddc fill commandlets)
+		//  cook/ddcfill will only make "notfast" ("shipping") by default
+		//  that behavior can be changed with an explicit -fasttextureencode argument
 
 		// does not depend on currentplatform so we could store in CachedFastTextureEncodeOption
 
-		// @todo Oodle consider to changing to this behavior
-
 		if ( GIsEditor && ! IsRunningCommandlet() )
 		{
 			return ETextureFastEncode::TryOffEncodeFast;
@@ -449,8 +433,6 @@ static ETextureFastEncode GetDesiredFastTextureEncode(	const ITargetPlatform& Cu
 		{
 			return ETextureFastEncode::Off;
 		}
-
-		#endif
 	}
 }
 
