Shader compiler improvements.

Local shader compiler:
 - Not making 1 job batches (for High prio jobs) on startup
 - Randomizing pending job selection to reduce chance that multiple 60+ sec jobs get into a single batch

 (the above have the most effect on speed up, below is optional/misc)
 - Pushing the completed jobs back to manager before, and not after, pulling new ones to reduce idle time
 - Processing jobs in FIFO rather than LIFO order (change to LIFO seems like an ODSC regression? Hard to say definitively if it's a regression from the numbers, but seems odd to have the earliest jobs processed last)
 - Parallel processing of input and output files (starts sequential by default to reduce the CPU overhead, but is enabled if we ever see too long write/read)
 - More insights scopes

Distributed:
  - Not avoiding local machine for XGE on startup

Both:
 - Input file compression (disabled by default, need to better figure out when it's beneficial in a general case. Mostly for people with really slow I/O and XGE over VPN).
 - More stats, also more dense stat output. Removed an unimportant one, added stats about the batches.

#[review] [REDACTED_HANDLE], [REDACTED_HANDLE]

diff --git a/Engine/Plugins/XGEController/Source/Private/XGEControllerModule.cpp b/Engine/Plugins/XGEController/Source/Private/XGEControllerModule.cpp
--- a/Engine/Plugins/XGEController/Source/Private/XGEControllerModule.cpp
+++ b/Engine/Plugins/XGEController/Source/Private/XGEControllerModule.cpp
@@ -12,6 +12,7 @@
 #include "Modules/ModuleManager.h"
 #include "Serialization/MemoryWriter.h"
 #include "Modules/ModuleManager.h"
+#include "ShaderCompiler.h"
 
 // Comma separated list of executable file names which should be intercepted by XGE.
 // Update this list if adding new tasks.
@@ -384,7 +385,7 @@ void FXGEControllerModule::WriteOutThreadProc()
 		// and set the working directory of xgConsole.exe to the engine binaries folder below.
 		FString XGConsoleArgs = FString::Printf(TEXT("/VIRTUALIZEDIRECTX /allowremote=\"%s\" %s /allowintercept=\"%s\" /title=\"Unreal Engine XGE Tasks\" /monitordirs=\"%s\" /command=\"%s -xgecontroller %s\""),
 			XGE_INTERCEPT_EXE_NAMES,
-			XGEController::AvoidUsingLocalMachine() ? TEXT("/avoidlocal=ON") : TEXT(""),
+			(XGEController::AvoidUsingLocalMachine() && !GShaderCompilingManager->IgnoreAllThrottling()) ? TEXT("/avoidlocal=ON") : TEXT(""),
 			XGE_CONTROL_WORKER_NAME,
 			*WorkingDirectory,
 			XGE_CONTROL_WORKER_FILENAME,
diff --git a/Engine/Plugins/XGEController/Source/XGEController.Build.cs b/Engine/Plugins/XGEController/Source/XGEController.Build.cs
--- a/Engine/Plugins/XGEController/Source/XGEController.Build.cs
+++ b/Engine/Plugins/XGEController/Source/XGEController.Build.cs
@@ -9,9 +9,13 @@ public class XGEController : ModuleRules
 	{
 		PrivateDependencyModuleNames.AddRange(new string[]
 		{
-			"Core"
+			"Core",
+
+			// all these below are only to get access to GShaderCompilingManager
+			"Slate",
+			"RHI",
+			"RenderCore",
+			"Engine",
 		});
-		
-		PrivateIncludePathModuleNames.Add("Engine");
 	}
 }
diff --git a/Engine/Source/Programs/ShaderCompileWorker/Private/ShaderCompileWorker.cpp b/Engine/Source/Programs/ShaderCompileWorker/Private/ShaderCompileWorker.cpp
--- a/Engine/Source/Programs/ShaderCompileWorker/Private/ShaderCompileWorker.cpp
+++ b/Engine/Source/Programs/ShaderCompileWorker/Private/ShaderCompileWorker.cpp
@@ -13,6 +13,7 @@
 #include "Interfaces/ITargetPlatformManagerModule.h"
 #include "RHIShaderFormatDefinitions.inl"
 #include "ShaderCompilerCommon.h"
+#include "Serialization/MemoryReader.h"
 
 #define DEBUG_USING_CONSOLE	0
 
@@ -346,14 +347,45 @@ private:
 
 	void ProcessInputFromArchive(FArchive* InputFilePtr, TArray<FJobResult>& OutSingleJobResults, TArray<FPipelineJobResult>& OutPipelineJobResults)
 	{
-		FArchive& InputFile = *InputFilePtr;
 		int32 InputVersion;
-		InputFile << InputVersion;
+		*InputFilePtr << InputVersion;
 		if (ShaderCompileWorkerInputVersion != InputVersion)
 		{
 			ExitWithoutCrash(ESCWErrorCode::BadInputVersion, FString::Printf(TEXT("Exiting due to ShaderCompilerWorker expecting input version %d, got %d instead! Did you forget to build ShaderCompilerWorker?"), ShaderCompileWorkerInputVersion, InputVersion));
 		}
 
+		FString CompressionFormatString;
+		*InputFilePtr << CompressionFormatString;
+		FName CompressionFormat(*CompressionFormatString);
+
+		bool bWasCompressed = (CompressionFormat != NAME_None);
+
+		TArray<uint8> UncompressedData;
+		if (bWasCompressed)
+		{
+			int32 UncompressedDataSize = 0;
+			*InputFilePtr << UncompressedDataSize;
+
+			if (UncompressedDataSize == 0)
+			{
+				ExitWithoutCrash(ESCWErrorCode::BadInputFile, TEXT("Exiting due to bad input file to ShaderCompilerWorker (uncompressed size is 0)! Did you forget to build ShaderCompilerWorker?"));
+				// unreachable
+				return;
+			}
+
+			UncompressedData.SetNumUninitialized(UncompressedDataSize);
+			TArray<uint8> CompressedData;
+			*InputFilePtr << CompressedData;
+			if (!FCompression::UncompressMemory(CompressionFormat, UncompressedData.GetData(), UncompressedDataSize, CompressedData.GetData(), CompressedData.Num()))
+			{
+				ExitWithoutCrash(ESCWErrorCode::BadInputFile, FString::Printf(TEXT("Exiting due to bad input file to ShaderCompilerWorker (cannot uncompress with the format %s)! Did you forget to build ShaderCompilerWorker?"), *CompressionFormatString));
+				// unreachable
+				return;
+			}
+		}
+		FMemoryReader InputMemory(UncompressedData);
+		FArchive& InputFile = bWasCompressed ? InputMemory : *InputFilePtr;
+
 		TMap<FString, uint32> ReceivedFormatVersionMap;
 		InputFile << ReceivedFormatVersionMap;
 
diff --git a/Engine/Source/Runtime/Engine/Private/ShaderCompiler/ShaderCompiler.cpp b/Engine/Source/Runtime/Engine/Private/ShaderCompiler/ShaderCompiler.cpp
--- a/Engine/Source/Runtime/Engine/Private/ShaderCompiler/ShaderCompiler.cpp
+++ b/Engine/Source/Runtime/Engine/Private/ShaderCompiler/ShaderCompiler.cpp
@@ -78,6 +78,9 @@ DEFINE_LOG_CATEGORY(LogShaderCompilers);
 // Switch to Verbose after initial testing
 #define UE_SHADERCACHE_LOG_LEVEL		VeryVerbose
 
+// whether to parallelize writing/reading task files
+#define UE_SHADERCOMPILER_FIFO_JOB_EXECUTION  1
+
 int32 GShaderCompilerJobCache = 1;
 static FAutoConsoleVariableRef CVarShaderCompilerJobCache(
 	TEXT("r.ShaderCompiler.JobCache"),
@@ -137,6 +140,17 @@ static FAutoConsoleVariableRef CVarGSShaderCheckLevel(
 	ECVF_Default
 );
 
+float GShaderCompilerTooLongIOThresholdSeconds = 0.3;
+static FAutoConsoleVariableRef CVarShaderCompilerTooLongIOThresholdSeconds(
+	TEXT("r.ShaderCompiler.TooLongIOThresholdSeconds"),
+	GShaderCompilerTooLongIOThresholdSeconds,
+	TEXT("By default, task files for SCW will be read/written sequentially, but if we ever spend more than this time (0.3s by default) doing that, we'll switch to parallel.") \
+	TEXT("We don't default to parallel writes as it increases the CPU overhead from the shader compiler."),
+	ECVF_Default
+);
+
+
+
 /** Helper functions for logging more debug info */
 namespace ShaderCompiler
 {
@@ -531,7 +545,26 @@ void FShaderCompileJobCollection::SubmitJobs(const TArray<FShaderCommonCompileJo
 				{
 					GShaderCompilerStats->RegisterNewPendingJob(*Job);
 					ensure(!ShaderCompiler::IsJobCacheEnabled() || Job->bInputHashSet);
-					Job->LinkHead(PendingJobs[PriorityIndex]);
+#if UE_SHADERCOMPILER_FIFO_JOB_EXECUTION
+					// link the job at the end of pending list, so we're executing them in a FIFO and not LIFO order
+					if (PendingJobs[PriorityIndex])
+					{
+						for (FShaderCommonCompileJob::TIterator It(PendingJobs[PriorityIndex]); It;)
+						{
+							FShaderCommonCompileJob& ExistingJob = *It;
+							if (ExistingJob.GetNextLink() == nullptr)
+							{
+								Job->LinkAfter(&ExistingJob);
+								break;
+							}
+						}
+					}
+					else
+#endif // UE_SHADERCOMPILER_FIFO_JOB_EXECUTION
+					{
+						Job->LinkHead(PendingJobs[PriorityIndex]);
+					}
+					
 
 					NumPendingJobs[PriorityIndex]++;
 					NumSubmittedJobs[PriorityIndex]++;
@@ -679,10 +712,29 @@ int32 FShaderCompileJobCollection::GetPendingJobs(EShaderCompilerWorkerType InWo
 		FWriteScopeLock Locker(Lock);
 		NumJobs = FMath::Min(MaxNumJobs, NumPendingJobs[PriorityIndex]);
 		FShaderCommonCompileJob::TIterator It(PendingJobs[PriorityIndex]);
+		// Randomize job selection by randomly skipping over jobs while traversing the list.
+		// Say, we need to pick 3 jobs out of 5 total. We can skip over 2 jobs in total, e.g. like this:
+		// pick one (4 more to go and we need to get 2 of 4), skip one (3 more to go, picking 2 out of 3), pick one (2 more to go, picking 1 of 2), skip one, pick one.
+		// It is possible that we won't skip at all and instead pick consequential jobs
+		int32 MaxJobsWeCanSkipOver = NumPendingJobs[PriorityIndex] - NumJobs;
 		for (int32 i = 0; i < NumJobs; ++i)
 		{
 			FShaderCommonCompileJob& Job = *It;
 
+			// get a random number of jobs to skip (if we can)
+			if (MaxJobsWeCanSkipOver > 0)
+			{
+				int32 NumJobsToSkipOver = FMath::RandHelper(MaxJobsWeCanSkipOver + 1);
+				while (NumJobsToSkipOver > 0)
+				{
+					It.Next();
+					--NumJobsToSkipOver;
+					--MaxJobsWeCanSkipOver;
+				}
+				checkf(MaxJobsWeCanSkipOver >= 0, TEXT("We skipped over too many jobs"));
+				checkf(MaxJobsWeCanSkipOver <= NumPendingJobs[PriorityIndex] - i, TEXT("Number of jobs to skip should stay less or equal than the number of nodes to go"));
+			}
+
 			GShaderCompilerStats->RegisterAssignedJob(Job);
 			// Temporary commented out until r.ShaderDevelopmentMode=1 shader error retry crash gets fixed
 			//check(Job.CurrentWorker == EShaderCompilerWorkerType::None);
@@ -1125,10 +1177,20 @@ static void SplitJobsByType(const TArray<FShaderCommonCompileJobPtr>& QueuedJobs
 }
 
 // Serialize Queued Job information
-bool FShaderCompileUtilities::DoWriteTasks(const TArray<FShaderCommonCompileJobPtr>& QueuedJobs, FArchive& TransferFile, bool bUseRelativePaths)
+bool FShaderCompileUtilities::DoWriteTasks(const TArray<FShaderCommonCompileJobPtr>& QueuedJobs, FArchive& InTransferFile, bool bUseRelativePaths, bool bCompressTaskFile)
 {
 	int32 InputVersion = ShaderCompileWorkerInputVersion;
-	TransferFile << InputVersion;
+	InTransferFile << InputVersion;
+
+	TArray<uint8> UncompressedArray;
+	FMemoryWriter TransferMemory(UncompressedArray);
+	FArchive& TransferFile = bCompressTaskFile ? TransferMemory : InTransferFile;
+	if (!bCompressTaskFile)
+	{
+		// still write NAME_None as string
+		FString FormatNone = FName(NAME_None).ToString();
+		TransferFile << FormatNone;
+	}
 
 	static TMap<FString, uint32> FormatVersionMap = GetFormatVersionMap();
 
@@ -1300,7 +1362,36 @@ bool FShaderCompileUtilities::DoWriteTasks(const TArray<FShaderCommonCompileJobP
 		}
 	}
 
-	return TransferFile.Close();
+	if (bCompressTaskFile)
+	{
+		TransferFile.Close();
+
+		FName CompressionFormatToUse = NAME_LZ4;
+
+		FString FormatName = CompressionFormatToUse.ToString();
+		InTransferFile << FormatName;
+
+		// serialize uncompressed data size
+		int32 UncompressedDataSize = UncompressedArray.Num();
+		checkf(UncompressedDataSize != 0, TEXT("Did not write any data to the task file for the compression."));
+		InTransferFile << UncompressedDataSize;
+
+		// not using SerializeCompressed because it splits into smaller chunks
+		int32 CompressedSizeBound = FCompression::CompressMemoryBound(CompressionFormatToUse, static_cast<int32>(UncompressedDataSize));
+		TArray<uint8> CompressedBuffer;
+		CompressedBuffer.SetNumUninitialized(CompressedSizeBound);
+
+		int32 ActualCompressedSize = CompressedSizeBound;
+		bool bSucceeded = FCompression::CompressMemory(CompressionFormatToUse, CompressedBuffer.GetData(), ActualCompressedSize, UncompressedArray.GetData(), UncompressedDataSize, COMPRESS_BiasSpeed);
+		checkf(ActualCompressedSize <= CompressedSizeBound, TEXT("Compressed size was larger than the bound - we stomped the memory."));
+		CompressedBuffer.SetNum(ActualCompressedSize, false);
+
+		InTransferFile << CompressedBuffer;
+		UE_LOG(LogShaderCompilers, Verbose, TEXT("Compressed the task file from %d bytes to %d bytes (%.2f%% savings)"), UncompressedDataSize, ActualCompressedSize,
+			100.0 * (UncompressedDataSize - ActualCompressedSize) / static_cast<double>(UncompressedDataSize));
+	}
+
+	return InTransferFile.Close();
 }
 
 static void ProcessErrors(const FShaderCompileJob& CurrentJob, TArray<FString>& UniqueErrors, FString& ErrorString)
@@ -1906,6 +1997,8 @@ uint32 FShaderCompileThreadRunnableBase::Run()
 
 int32 FShaderCompileThreadRunnable::PullTasksFromQueue()
 {
+	TRACE_CPUPROFILER_EVENT_SCOPE(FShaderCompileThreadRunnable::PullTasksFromQueue);
+
 	int32 NumActiveThreads = 0;
 	int32 NumJobsStarted[NumShaderCompileJobPriorities] = { 0 };
 	{
@@ -1914,7 +2007,7 @@ int32 FShaderCompileThreadRunnable::PullTasksFromQueue()
 
 		const int32 NumWorkersToFeed = Manager->bCompilingDuringGame ? Manager->NumShaderCompilingThreadsDuringGame : WorkerInfos.Num();
 
-		for (int32 PriorityIndex = MaxPriorityIndex; PriorityIndex >= MinPriorityIndex ; --PriorityIndex)
+		for (int32 PriorityIndex = MaxPriorityIndex; PriorityIndex >= MinPriorityIndex; --PriorityIndex)
 		{
 			int32 NumPendingJobs = Manager->AllJobs.GetNumPendingJobs((EShaderCompileJobPriority)PriorityIndex);
 			// Try to distribute the work evenly between the workers
@@ -1935,7 +2028,8 @@ int32 FShaderCompileThreadRunnable::PullTasksFromQueue()
 						UE_LOG(LogShaderCompilers, Verbose, TEXT("Worker (%d/%d): shaders left to compile %i"), WorkerIndex + 1, WorkerInfos.Num(), NumPendingJobs);
 
 						int32 MaxNumJobs = 1;
-						if (PriorityIndex < (int32)EShaderCompileJobPriority::High)
+						// high priority jobs go in 1 per "batch", unless the engine is still starting up
+						if (PriorityIndex < (int32)EShaderCompileJobPriority::High || Manager->IgnoreAllThrottling())
 						{
 							MaxNumJobs = FMath::Min3(NumJobsPerWorker, NumPendingJobs, Manager->MaxShaderJobBatchSize);
 						}
@@ -1962,91 +2056,121 @@ int32 FShaderCompileThreadRunnable::PullTasksFromQueue()
 				}
 			}
 		}
+	}
 
-		for (int32 WorkerIndex = 0; WorkerIndex < WorkerInfos.Num(); WorkerIndex++)
+	for (int32 WorkerIndex = 0; WorkerIndex < WorkerInfos.Num(); WorkerIndex++)
+	{
+		if (WorkerInfos[WorkerIndex]->QueuedJobs.Num() > 0)
 		{
-			FShaderCompileWorkerInfo& CurrentWorkerInfo = *WorkerInfos[WorkerIndex];
+			NumActiveThreads++;
+		}
+	}
 
-			if (CurrentWorkerInfo.QueuedJobs.Num() > 0)
-			{
-				NumActiveThreads++;
-			}
+	for (int32 PriorityIndex = 0; PriorityIndex < NumShaderCompileJobPriorities; ++PriorityIndex)
+	{
+		if (NumJobsStarted[PriorityIndex] > 0)
+		{
+			UE_LOG(LogShaderCompilers, Verbose, TEXT("Started %d 'Local' shader compile jobs with '%s' priority"),
+				NumJobsStarted[PriorityIndex],
+				ShaderCompileJobPriorityToString((EShaderCompileJobPriority)PriorityIndex));
+		}
+	}
+
+	return NumActiveThreads;
+}
+
+void FShaderCompileThreadRunnable::PushCompletedJobsToManager()
+{
+	for (int32 WorkerIndex = 0; WorkerIndex < WorkerInfos.Num(); WorkerIndex++)
+	{
+		FShaderCompileWorkerInfo& CurrentWorkerInfo = *WorkerInfos[WorkerIndex];
+
+		// Add completed jobs to the output queue, which is ShaderMapJobs
+		if (CurrentWorkerInfo.bComplete)
+		{
+			// Enter the critical section so we can access the input and output queues
+			FScopeLock Lock(&Manager->CompileQueueSection);
 
-			// Add completed jobs to the output queue, which is ShaderMapJobs
-			if (CurrentWorkerInfo.bComplete)
+			for (int32 JobIndex = 0; JobIndex < CurrentWorkerInfo.QueuedJobs.Num(); JobIndex++)
 			{
-				for (int32 JobIndex = 0; JobIndex < CurrentWorkerInfo.QueuedJobs.Num(); JobIndex++)
-				{
-					auto& Job = CurrentWorkerInfo.QueuedJobs[JobIndex];
+				auto& Job = CurrentWorkerInfo.QueuedJobs[JobIndex];
 
-					Manager->ProcessFinishedJob(Job.GetReference());
-				}
+				Manager->ProcessFinishedJob(Job.GetReference());
+			}
 
-				const float ElapsedTime = FPlatformTime::Seconds() - CurrentWorkerInfo.StartTime;
+			const float ElapsedTime = FPlatformTime::Seconds() - CurrentWorkerInfo.StartTime;
 
-				Manager->WorkersBusyTime += ElapsedTime;
-				COOK_STAT(ShaderCompilerCookStats::AsyncCompileTimeSec += ElapsedTime);
+			Manager->WorkersBusyTime += ElapsedTime;
+			COOK_STAT(ShaderCompilerCookStats::AsyncCompileTimeSec += ElapsedTime);
 
-				// Log if requested or if there was an exceptionally slow batch, to see the offender easily
-				if (Manager->bLogJobCompletionTimes || ElapsedTime > 60.0f)
-				{
-					FString JobNames;
+			// Log if requested or if there was an exceptionally slow batch, to see the offender easily
+			if (Manager->bLogJobCompletionTimes || ElapsedTime > 60.0f)
+			{
+				FString JobNames;
 
-					for (int32 JobIndex = 0; JobIndex < CurrentWorkerInfo.QueuedJobs.Num(); JobIndex++)
+				for (int32 JobIndex = 0; JobIndex < CurrentWorkerInfo.QueuedJobs.Num(); JobIndex++)
+				{
+					const FShaderCommonCompileJob& Job = *CurrentWorkerInfo.QueuedJobs[JobIndex];
+					if (auto* SingleJob = Job.GetSingleShaderJob())
 					{
-						const FShaderCommonCompileJob& Job = *CurrentWorkerInfo.QueuedJobs[JobIndex];
-						if (auto* SingleJob = Job.GetSingleShaderJob())
-						{
-							const TCHAR* JobName = Manager->bLogJobCompletionTimes ? *SingleJob->Input.DebugGroupName : SingleJob->Key.ShaderType->GetName();
-							JobNames += FString::Printf(TEXT("%s [Instructions=%d WorkerTime=%.3fs]"), JobName, SingleJob->Output.NumInstructions, SingleJob->Output.CompileTime);
-						}
-						else
-						{
-							auto* PipelineJob = Job.GetShaderPipelineJob();
-							JobNames += FString(PipelineJob->Key.ShaderPipeline->GetName());
-							if (PipelineJob->bFailedRemovingUnused)
-							{
-								JobNames += FString(TEXT("(failed to optimize)"));
-							}
-						}
-						if (JobIndex < CurrentWorkerInfo.QueuedJobs.Num() - 1)
+						const TCHAR* JobName = Manager->bLogJobCompletionTimes ? *SingleJob->Input.DebugGroupName : SingleJob->Key.ShaderType->GetName();
+						JobNames += FString::Printf(TEXT("%s [Instructions=%d WorkerTime=%.3fs]"), JobName, SingleJob->Output.NumInstructions, SingleJob->Output.CompileTime);
+					}
+					else
+					{
+						auto* PipelineJob = Job.GetShaderPipelineJob();
+						JobNames += FString(PipelineJob->Key.ShaderPipeline->GetName());
+						if (PipelineJob->bFailedRemovingUnused)
 						{
-							JobNames += TEXT(", ");
+							JobNames += FString(TEXT("(failed to optimize)"));
 						}
 					}
-
-					UE_LOG(LogShaders, Display, TEXT("Worker (%d/%d) finished batch of %u jobs in %.3fs, %s"), WorkerIndex + 1, WorkerInfos.Num(), CurrentWorkerInfo.QueuedJobs.Num(), ElapsedTime, *JobNames);
+					if (JobIndex < CurrentWorkerInfo.QueuedJobs.Num() - 1)
+					{
+						JobNames += TEXT(", ");
+					}
 				}
 
-				CurrentWorkerInfo.FinishTime = FPlatformTime::Seconds();
-				CurrentWorkerInfo.bComplete = false;
-				CurrentWorkerInfo.QueuedJobs.Empty();
+				UE_LOG(LogShaders, Display, TEXT("Worker (%d/%d) finished batch of %u jobs in %.3fs, %s"), WorkerIndex + 1, WorkerInfos.Num(), CurrentWorkerInfo.QueuedJobs.Num(), ElapsedTime, *JobNames);
 			}
+
+			CurrentWorkerInfo.FinishTime = FPlatformTime::Seconds();
+			CurrentWorkerInfo.bComplete = false;
+			CurrentWorkerInfo.QueuedJobs.Empty();
 		}
 	}
+}
 
-	for (int32 PriorityIndex = 0; PriorityIndex < NumShaderCompileJobPriorities; ++PriorityIndex)
+void FShaderCompileThreadRunnable::WriteNewTasks()
+{
+	TRACE_CPUPROFILER_EVENT_SCOPE(FShaderCompileThreadRunnable::WriteNewTasks);
+
+	// first, a quick check if anything is needed just to avoid hammering the task graph
+	bool bHasTasksToWrite = false;
+	for (int32 WorkerIndex = 0, NumWorkers = WorkerInfos.Num(); WorkerIndex < NumWorkers; ++WorkerIndex)
 	{
-		if (NumJobsStarted[PriorityIndex] > 0)
+		FShaderCompileWorkerInfo& CurrentWorkerInfo = *WorkerInfos[WorkerIndex];
+		if (!CurrentWorkerInfo.bIssuedTasksToWorker && CurrentWorkerInfo.QueuedJobs.Num() > 0)
 		{
-			UE_LOG(LogShaderCompilers, Verbose, TEXT("Started %d 'Local' shader compile jobs with '%s' priority"),
-				NumJobsStarted[PriorityIndex],
-				ShaderCompileJobPriorityToString((EShaderCompileJobPriority)PriorityIndex));
+			bHasTasksToWrite = true;
+			break;
 		}
 	}
 
-	return NumActiveThreads;
-}
+	if (!bHasTasksToWrite)
+	{
+		return;
+	}
 
-void FShaderCompileThreadRunnable::WriteNewTasks()
-{
-	for (int32 WorkerIndex = 0; WorkerIndex < WorkerInfos.Num(); WorkerIndex++)
+
+	auto LoopBody = [this](int32 WorkerIndex)
 	{
 		FShaderCompileWorkerInfo& CurrentWorkerInfo = *WorkerInfos[WorkerIndex];
 
 		// Only write tasks once
 		if (!CurrentWorkerInfo.bIssuedTasksToWorker && CurrentWorkerInfo.QueuedJobs.Num() > 0)
 		{
+			TRACE_CPUPROFILER_EVENT_SCOPE(FShaderCompileThreadRunnable::WriteNewTasksForWorker);
 			CurrentWorkerInfo.bIssuedTasksToWorker = true;
 
 			const FString WorkingDirectory = Manager->AbsoluteShaderBaseWorkingDirectory + FString::FromInt(WorkerIndex);
@@ -2088,6 +2212,7 @@ void FShaderCompileThreadRunnable::WriteNewTasks()
 			}
 			check(TransferFile);
 
+			GShaderCompilerStats->RegisterJobBatch(CurrentWorkerInfo.QueuedJobs.Num(), FShaderCompilerStats::EExecutionType::Local);
 			if (!FShaderCompileUtilities::DoWriteTasks(CurrentWorkerInfo.QueuedJobs, *TransferFile))
 			{
 				uint64 TotalDiskSpace = 0;
@@ -2124,11 +2249,33 @@ void FShaderCompileThreadRunnable::WriteNewTasks()
 				UE_LOG(LogShaderCompilers, Error, TEXT("Could not rename the shader compiler transfer filename to '%s' from '%s' (Free Disk Space: %llu)."), *ProperTransferFileName, *TransferFileName, FreeDiskSpace);
 			}
 		}
+	};
+
+	if (bParallelizeIO)
+	{
+		ParallelFor(WorkerInfos.Num(), LoopBody, EParallelForFlags::Unbalanced);
+	}
+	else
+	{
+		double StartIOWork = FPlatformTime::Seconds();
+		for (int32 WorkerIndex = 0; WorkerIndex < WorkerInfos.Num(); WorkerIndex++)
+		{
+			LoopBody(WorkerIndex);
+		}
+
+		double IODuration = FPlatformTime::Seconds() - StartIOWork;
+		if (IODuration > GShaderCompilerTooLongIOThresholdSeconds)
+		{
+			UE_LOG(LogShaderCompilers, Display, TEXT("FShaderCompileThreadRunnable::WriteNewTasks()() took too long (%.3f seconds, threshold is %.3f s), will parallelize next time."), IODuration, GShaderCompilerTooLongIOThresholdSeconds);
+			bParallelizeIO = true;
+		}
 	}
 }
 
 bool FShaderCompileThreadRunnable::LaunchWorkersIfNeeded()
 {
+	TRACE_CPUPROFILER_EVENT_SCOPE(FShaderCompileThreadRunnable::LaunchWorkersIfNeeded);
+
 	const double CurrentTime = FPlatformTime::Seconds();
 	// Limit how often we check for workers running since IsApplicationRunning eats up some CPU time on Windows
 	const bool bCheckForWorkerRunning = (CurrentTime - LastCheckForWorkersTime > .1f);
@@ -2157,6 +2304,8 @@ bool FShaderCompileThreadRunnable::LaunchWorkersIfNeeded()
 
 		if (!CurrentWorkerInfo.WorkerProcess.IsValid() || (bCheckForWorkerRunning && !FShaderCompilingManager::IsShaderCompilerWorkerRunning(CurrentWorkerInfo.WorkerProcess)))
 		{
+			TRACE_CPUPROFILER_EVENT_SCOPE(FShaderCompileThreadRunnable::LaunchingWorkers);
+
 			// @TODO: dubious design - worker should not be launched unless we know there's more work to do.
 			bool bLaunchAgain = true;
 
@@ -2215,26 +2364,43 @@ bool FShaderCompileThreadRunnable::LaunchWorkersIfNeeded()
 
 int32 FShaderCompileThreadRunnable::ReadAvailableResults()
 {
+	TRACE_CPUPROFILER_EVENT_SCOPE(FShaderCompileThreadRunnable::ReadAvailableResults);
 	int32 NumProcessed = 0;
 
-	for (int32 WorkerIndex = 0; WorkerIndex < WorkerInfos.Num(); WorkerIndex++)
+	// first, a quick check if anything is needed just to avoid hammering the task graph
+	bool bHasQueuedJobs = false;
+	for (int32 WorkerIndex = 0, NumWorkers = WorkerInfos.Num(); WorkerIndex < NumWorkers; ++WorkerIndex)
+	{
+		if (WorkerInfos[WorkerIndex]->QueuedJobs.Num() > 0)
+		{
+			bHasQueuedJobs = true;
+			break;
+		}
+	}
+
+	if (!bHasQueuedJobs)
+	{
+		return NumProcessed;
+	}
+
+	auto LoopBody = [this, &NumProcessed](int32 WorkerIndex)
 	{
 		FShaderCompileWorkerInfo& CurrentWorkerInfo = *WorkerInfos[WorkerIndex];
 
 		// Check for available result files
 		if (CurrentWorkerInfo.QueuedJobs.Num() > 0)
 		{
-			TRACE_CPUPROFILER_EVENT_SCOPE(FShaderCompileThreadRunnable::ReadAvailableResults);
-
 			// Distributed compiles always use the same directory
 			// 'Only' indicates to the worker that it should log and continue checking for the input file after the first one is processed
 			TStringBuilder<512> OutputFileNameAndPath;
 			OutputFileNameAndPath << Manager->AbsoluteShaderBaseWorkingDirectory << WorkerIndex << TEXT("/WorkerOutputOnly.out");
-			
+
 			// In the common case the output file will not exist, so check for existence before opening
 			// This is only a win if FileExists is faster than CreateFileReader, which it is on Windows
 			if (FPlatformFileManager::Get().GetPlatformFile().FileExists(*OutputFileNameAndPath))
 			{
+				TRACE_CPUPROFILER_EVENT_SCOPE(FShaderCompileThreadRunnable::ProcessOutputFile);
+
 				FArchive* OutputFilePtr = IFileManager::Get().CreateFileReader(*OutputFileNameAndPath, FILEREAD_Silent);
 
 				if (OutputFilePtr)
@@ -2252,6 +2418,8 @@ int32 FShaderCompileThreadRunnable::ReadAvailableResults()
 					// Retry over the next two seconds if we couldn't delete it
 					while (!bDeletedOutput && RetryCount < 200)
 					{
+						TRACE_CPUPROFILER_EVENT_SCOPE(FShaderCompileThreadRunnable::DeleteOutputFile);
+
 						FPlatformProcess::Sleep(0.01f);
 						bDeletedOutput = IFileManager::Get().Delete(*OutputFileNameAndPath, true, true);
 						RetryCount++;
@@ -2261,9 +2429,29 @@ int32 FShaderCompileThreadRunnable::ReadAvailableResults()
 					CurrentWorkerInfo.bComplete = true;
 				}
 
-				++NumProcessed;
+				FPlatformAtomics::InterlockedIncrement(&NumProcessed);
 			}
 		}
+	};
+
+	if (bParallelizeIO)
+	{
+		ParallelFor(WorkerInfos.Num(), LoopBody, EParallelForFlags::Unbalanced);
+	}
+	else 
+	{
+		double StartIOWork = FPlatformTime::Seconds();
+		for (int32 WorkerIndex = 0; WorkerIndex < WorkerInfos.Num(); WorkerIndex++)
+		{
+			LoopBody(WorkerIndex);
+		}
+
+		double IODuration = FPlatformTime::Seconds() - StartIOWork;
+		if (IODuration > 0.3)
+		{
+			UE_LOG(LogShaderCompilers, Display, TEXT("FShaderCompileThreadRunnable::WriteNewTasks() took too long (%.3f seconds, threshold is %.3f s), will parallelize next time."), IODuration, GShaderCompilerTooLongIOThresholdSeconds);
+			bParallelizeIO = true;
+		}
 	}
 
 	return NumProcessed;
@@ -2457,7 +2645,10 @@ void FShaderCompileUtilities::DeleteFileHelper(const FString& Filename)
 
 int32 FShaderCompileThreadRunnable::CompilingLoop()
 {
-	// Grab more shader compile jobs from the input queue, and move completed jobs to Manager->ShaderMapJobs
+	// push completed jobs to Manager->ShaderMapJobs before asking for new ones, so we can free the workers now and avoid them waiting a cycle
+	PushCompletedJobsToManager();
+
+	// Grab more shader compile jobs from the input queue
 	const int32 NumActiveThreads = PullTasksFromQueue();
 
 	if (NumActiveThreads == 0 && Manager->bAllowAsynchronousShaderCompiling)
@@ -2702,8 +2893,7 @@ void FShaderCompilerStats::WriteStatSummary()
 	UE_LOG(LogShaderCompilers, Display, TEXT("Shaders Compiled: %u"), TotalCompiled);
 
 	FScopeLock Lock(&CompileStatsLock);	// make a local copy for all the stats?
-	UE_LOG(LogShaderCompilers, Display, TEXT("Jobs assigned: %.0f"), JobsAssigned);
-	UE_LOG(LogShaderCompilers, Display, TEXT("Jobs completed: %.0f (%.2f%%)"), JobsCompleted, (JobsAssigned > 0.0) ? 100.0 * JobsCompleted / JobsAssigned : 0.0);
+	UE_LOG(LogShaderCompilers, Display, TEXT("Jobs assigned %.0f, completed %.0f (%.2f%%)"), JobsAssigned, JobsCompleted, (JobsAssigned > 0.0) ? 100.0 * JobsCompleted / JobsAssigned : 0.0);
 
 	if (TimesLocalWorkersWereIdle > 0.0)
 	{
@@ -2712,13 +2902,13 @@ void FShaderCompilerStats::WriteStatSummary()
 
 	if (JobsAssigned > 0.0)
 	{
-		UE_LOG(LogShaderCompilers, Display, TEXT("Average time job spent in pending queue: %.2f s"), AccumulatedPendingTime / JobsAssigned);
+		UE_LOG(LogShaderCompilers, Display, TEXT("Time job spent in pending queue: average %.2f s, longest %.2f s"), AccumulatedPendingTime / JobsAssigned, MaxPendingTime);
 	}
 
 	if (JobsCompleted > 0.0)
 	{
-		UE_LOG(LogShaderCompilers, Display, TEXT("Average job execution time: %.2f s"), AccumulatedJobExecutionTime / JobsCompleted);
-		UE_LOG(LogShaderCompilers, Display, TEXT("Average job life time (pending + execution): %.2f s"), AccumulatedJobLifeTime / JobsCompleted);
+		UE_LOG(LogShaderCompilers, Display, TEXT("Job execution time: average %.2f s, max %.2f s"), AccumulatedJobExecutionTime / JobsCompleted, MaxJobExecutionTime);
+		UE_LOG(LogShaderCompilers, Display, TEXT("Job life time (pending + execution): average %.2f s, max %.2f"), AccumulatedJobLifeTime / JobsCompleted, MaxJobLifeTime);
 	}
 
 	auto SumUpInterval = [](TArray<TInterval<double>>& Accumulator) -> double
@@ -2732,11 +2922,33 @@ void FShaderCompilerStats::WriteStatSummary()
 		return Sum;
 	};
 
-	double TotalTimeExecutingAtLeastOneJob = SumUpInterval(JobPendingTimeIntervals);
-	UE_LOG(LogShaderCompilers, Display, TEXT("Time at least one job was waiting to be executed: %.2f s"), TotalTimeExecutingAtLeastOneJob);
-
 	double TotalTimeAtLeastOneJobWasInFlight = SumUpInterval(JobLifeTimeIntervals);
 	UE_LOG(LogShaderCompilers, Display, TEXT("Time at least one job was in flight (either pending or executed): %.2f s"), TotalTimeAtLeastOneJobWasInFlight);
+
+	// print stats about the batches
+	if (LocalJobBatchesSeen > 0 && DistributedJobBatchesSeen > 0)
+	{
+		int64 JobBatchesSeen = LocalJobBatchesSeen + DistributedJobBatchesSeen;
+		double TotalJobsReportedInJobBatches = TotalJobsReportedInLocalJobBatches + TotalJobsReportedInDistributedJobBatches;
+
+		UE_LOG(LogShaderCompilers, Display, TEXT("Jobs were issued in %lld batches (%lld local, %lld distributed), average %.2f jobs/batch (%.2f jobs/local batch. %.2f jobs/distributed batch)"),
+			JobBatchesSeen, LocalJobBatchesSeen, DistributedJobBatchesSeen,
+			static_cast<double>(TotalJobsReportedInJobBatches) / static_cast<double>(JobBatchesSeen),
+			static_cast<double>(TotalJobsReportedInLocalJobBatches) / static_cast<double>(LocalJobBatchesSeen),
+			static_cast<double>(TotalJobsReportedInDistributedJobBatches) / static_cast<double>(DistributedJobBatchesSeen)
+		);
+	}
+	else if (LocalJobBatchesSeen > 0)
+	{
+		UE_LOG(LogShaderCompilers, Display, TEXT("Jobs were issued in %lld batches (only local compilation was used), average %.2f jobs/batch"), 
+			LocalJobBatchesSeen, static_cast<double>(TotalJobsReportedInLocalJobBatches) / static_cast<double>(LocalJobBatchesSeen));
+	}
+	else if (DistributedJobBatchesSeen > 0)
+	{
+		UE_LOG(LogShaderCompilers, Display, TEXT("Jobs were issued in %lld batches (only distributed compilation was used), average %.2f jobs/batch"),
+			DistributedJobBatchesSeen, static_cast<double>(TotalJobsReportedInDistributedJobBatches) / static_cast<double>(DistributedJobBatchesSeen));
+	}
+
 	if (TotalTimeAtLeastOneJobWasInFlight > 0.0)
 	{
 		UE_LOG(LogShaderCompilers, Display, TEXT("Average processing rate: %.2f jobs/sec"), JobsCompleted / TotalTimeAtLeastOneJobWasInFlight);
@@ -2818,7 +3030,9 @@ void FShaderCompilerStats::RegisterAssignedJob(FShaderCommonCompileJob& Job)
 
 	FScopeLock Lock(&CompileStatsLock);
 	JobsAssigned++;
-	AccumulatedPendingTime += (Job.TimeAssignedToExecution - Job.TimeAddedToPendingQueue);
+	double TimeSpendPending = (Job.TimeAssignedToExecution - Job.TimeAddedToPendingQueue);
+	AccumulatedPendingTime += TimeSpendPending;
+	MaxPendingTime = FMath::Max(TimeSpendPending, MaxPendingTime);
 }
 
 void FShaderCompilerStats::RegisterFinishedJob(FShaderCommonCompileJob& Job)
@@ -2828,8 +3042,14 @@ void FShaderCompilerStats::RegisterFinishedJob(FShaderCommonCompileJob& Job)
 
 	FScopeLock Lock(&CompileStatsLock);
 	JobsCompleted++;
-	AccumulatedJobExecutionTime += (Job.TimeExecutionCompleted - Job.TimeAssignedToExecution);
-	AccumulatedJobLifeTime += (Job.TimeExecutionCompleted - Job.TimeAddedToPendingQueue);
+	
+	double ExecutionTime = (Job.TimeExecutionCompleted - Job.TimeAssignedToExecution);
+	AccumulatedJobExecutionTime += ExecutionTime;
+	MaxJobExecutionTime = FMath::Max(ExecutionTime, MaxJobExecutionTime);
+
+	double LifeTime = (Job.TimeExecutionCompleted - Job.TimeAddedToPendingQueue);
+	AccumulatedJobLifeTime += LifeTime;
+	MaxJobLifeTime = FMath::Max(LifeTime, MaxJobLifeTime);
 
 	auto AddToInterval = [](TArray<TInterval<double>>& Accumulator, const TInterval<double>& NewInterval)
 	{
@@ -2871,10 +3091,7 @@ void FShaderCompilerStats::RegisterFinishedJob(FShaderCommonCompileJob& Job)
 	};
 
 	// estimate lifetime without an overlap
-	ensure(Job.TimeAddedToPendingQueue != 0.0 && Job.TimeAddedToPendingQueue <= Job.TimeAssignedToExecution);
-	AddToInterval(JobPendingTimeIntervals, TInterval<double>(Job.TimeAddedToPendingQueue, Job.TimeAssignedToExecution));
-
-	ensure(Job.TimeAddedToPendingQueue != 0.0 && Job.TimeAssignedToExecution <= Job.TimeExecutionCompleted);
+	ensure(Job.TimeAddedToPendingQueue != 0.0 && Job.TimeAddedToPendingQueue <= Job.TimeExecutionCompleted);
 	AddToInterval(JobLifeTimeIntervals, TInterval<double>(Job.TimeAddedToPendingQueue, Job.TimeExecutionCompleted));
 
 	if (const FShaderCompileJob* SingleJob = Job.GetSingleShaderJob())
@@ -2903,6 +3120,26 @@ void FShaderCompilerStats::RegisterFinishedJob(FShaderCommonCompileJob& Job)
 
 }
 
+void FShaderCompilerStats::RegisterJobBatch(int32 NumJobs, EExecutionType ExecType)
+{
+	if (ExecType == EExecutionType::Local)
+	{
+		FScopeLock Lock(&CompileStatsLock);
+		++LocalJobBatchesSeen;
+		TotalJobsReportedInLocalJobBatches += NumJobs;
+	}
+	else if (ExecType == EExecutionType::Distributed)
+	{
+		FScopeLock Lock(&CompileStatsLock);
+		++DistributedJobBatchesSeen;
+		TotalJobsReportedInDistributedJobBatches += NumJobs;
+	}
+	else
+	{
+		checkNoEntry();
+	}
+}
+
 void FShaderCompilerStats::RegisterCookedShaders(uint32 NumCooked, float CompileTime, EShaderPlatform Platform, const FString MaterialPath, FString PermutationString)
 {
 	FScopeLock Lock(&CompileStatsLock);
@@ -3079,7 +3316,13 @@ FShaderCompilingManager::FShaderCompilingManager() :
 	}
 
 	bool bForceUseSCWMemoryPressureLimits = false;
-	
+
+	bIsEngineLoopInitialized = false;
+	FCoreDelegates::OnFEngineLoopInitComplete.AddLambda([&]() 
+		{ 
+			bIsEngineLoopInitialized = true; 
+		}
+	);
 
 	WorkersBusyTime = 0;
 
diff --git a/Engine/Source/Runtime/Engine/Private/ShaderCompiler/ShaderCompilerDistributed.cpp b/Engine/Source/Runtime/Engine/Private/ShaderCompiler/ShaderCompilerDistributed.cpp
--- a/Engine/Source/Runtime/Engine/Private/ShaderCompiler/ShaderCompilerDistributed.cpp
+++ b/Engine/Source/Runtime/Engine/Private/ShaderCompiler/ShaderCompilerDistributed.cpp
@@ -75,6 +75,7 @@ void FShaderCompileDistributedThreadRunnable_Interface::DispatchShaderCompileJob
 	const FString WorkingDirectory = FPaths::GetPath(InputFilePath);
 
 	// Serialize the jobs to the input file
+	GShaderCompilerStats->RegisterJobBatch(JobsToSerialize.Num(), FShaderCompilerStats::EExecutionType::Distributed);
 	FArchive* InputFileAr = IFileManager::Get().CreateFileWriter(*InputFilePath, FILEWRITE_EvenIfReadOnly | FILEWRITE_NoFail);
 	FShaderCompileUtilities::DoWriteTasks(JobsToSerialize, *InputFileAr, CachedController.RequiresRelativePaths());
 	delete InputFileAr;
diff --git a/Engine/Source/Runtime/Engine/Public/ShaderCompiler.h b/Engine/Source/Runtime/Engine/Public/ShaderCompiler.h
--- a/Engine/Source/Runtime/Engine/Public/ShaderCompiler.h
+++ b/Engine/Source/Runtime/Engine/Public/ShaderCompiler.h
@@ -600,6 +600,9 @@ private:
 	/** Tracks the last time that this thread checked if the workers were still active. */
 	double LastCheckForWorkersTime;
 
+	/** Whether to read/write files for SCW in parallel (can help situations when this takes too long for a number of reasons) */
+	bool bParallelizeIO = false;
+
 public:
 	/** Initialization constructor. */
 	FShaderCompileThreadRunnable(class FShaderCompilingManager* InManager);
@@ -609,10 +612,14 @@ private:
 
 	/** 
 	 * Grabs tasks from Manager->CompileQueue in a thread safe way and puts them into QueuedJobs of available workers. 
-	 * Also writes completed jobs to Manager->ShaderMapJobs.
 	 */
 	int32 PullTasksFromQueue();
 
+	/**
+	 * Writes completed jobs to Manager->ShaderMapJobs.
+	 */
+	void PushCompletedJobsToManager();
+
 	/** Used when compiling through workers, writes out the worker inputs for any new tasks in WorkerInfos.QueuedJobs. */
 	void WriteNewTasks();
 
@@ -631,7 +638,7 @@ private:
 
 namespace FShaderCompileUtilities
 {
-	bool DoWriteTasks(const TArray<FShaderCommonCompileJobPtr>& QueuedJobs, FArchive& TransferFile, bool bUseRelativePaths = false);
+	bool DoWriteTasks(const TArray<FShaderCommonCompileJobPtr>& QueuedJobs, FArchive& TransferFile, bool bUseRelativePaths = false, bool bCompressTaskFile = false);
 	void DoReadTaskResults(const TArray<FShaderCommonCompileJobPtr>& QueuedJobs, FArchive& OutputFile);
 
 	/** Execute the specified (single or pipeline) shader compile job. */
@@ -742,6 +749,12 @@ public:
 	ENGINE_API void WriteStats(class FOutputDevice* Ar = nullptr);
 	ENGINE_API void WriteStatSummary();
 
+	enum class EExecutionType
+	{
+		Local,
+		Distributed
+	};
+
 	/** Informs statistics about a time a local ShaderCompileWorker spent idle. */
 	void RegisterLocalWorkerIdleTime(double IdleTime);
 
@@ -754,6 +767,9 @@ public:
 	/** Marks the job as finished for the stats purpose. Job will be modified to include the current timestamp. */
 	void RegisterFinishedJob(FShaderCommonCompileJob& InOutJob);
 
+	/** Informs statistics about a new job batch, so we can tally up batches. */
+	void RegisterJobBatch(int32 NumJobs, EExecutionType ExecType);
+
 private:
 	FCriticalSection CompileStatsLock;
 	TSparseArray<ShaderCompilerStats> CompileStats;
@@ -773,17 +789,35 @@ private:
 	/** Amount of time a job had to spent in pending queue (i.e. waiting to be assigned to a worker). */
 	double AccumulatedPendingTime = 0;
 
+	/** Max amount of time any single job was pending (waiting to be assigned to a worker). */
+	double MaxPendingTime = 0;
+
 	/** Total number jobs completed. */
 	double JobsCompleted = 0;
 
 	/** Amount of time job spent being processed by the worker. */
 	double AccumulatedJobExecutionTime = 0;
 
-	/** Amount of time job spent being processed by the worker. */
+	/** Max amount of time any single job spent being processed by the worker. */
+	double MaxJobExecutionTime = 0;
+
+	/** Amount of time job spent being processed overall. */
 	double AccumulatedJobLifeTime = 0;
 
-	/** Accumulates the job pending times without overlaps */
-	TArray<TInterval<double>> JobPendingTimeIntervals;
+	/** Max amount of time any single job spent being processed overall. */
+	double MaxJobLifeTime = 0;
+
+	/** Number of local job batches seen. */
+	int64 LocalJobBatchesSeen = 0;
+
+	/** Total jobs in local job batches. */
+	int64 TotalJobsReportedInLocalJobBatches = 0;
+
+	/** Number of distributed job batches seen. */
+	int64 DistributedJobBatchesSeen = 0;
+
+	/** Total jobs in local job batches. */
+	int64 TotalJobsReportedInDistributedJobBatches = 0;
 
 	/** Accumulates the job lifetimes without overlaps */
 	TArray<TInterval<double>> JobLifeTimeIntervals;
@@ -968,6 +1002,15 @@ public:
 		return GetNumOutstandingJobs() > 0 || HasShaderJobs() || GetNumPendingJobs() > 0 || NumExternalJobs > 0;
 	}
 
+	/**
+	 * Returns whether normal throttling settings should be ignored because shader compilation is at the moment the only action blocking the critical path.
+	 * An example of such situation is startup compilation of global shaders and default materials, but there may be more cases like that.
+	 */
+	bool IgnoreAllThrottling() const
+	{
+		return !bIsEngineLoopInitialized;
+	}
+
 	/**
 	 * return true if we have shader jobs in any state
 	 * shader jobs are removed when they are applied to the gamethreadshadermap
diff --git a/Engine/Source/Runtime/RenderCore/Public/ShaderCompilerCore.h b/Engine/Source/Runtime/RenderCore/Public/ShaderCompilerCore.h
--- a/Engine/Source/Runtime/RenderCore/Public/ShaderCompilerCore.h
+++ b/Engine/Source/Runtime/RenderCore/Public/ShaderCompilerCore.h
@@ -16,7 +16,7 @@
 class Error;
 
 // this is for the protocol, not the data, bump if FShaderCompilerInput or ProcessInputFromArchive changes.
-const int32 ShaderCompileWorkerInputVersion = 15;
+const int32 ShaderCompileWorkerInputVersion = 16;
 // this is for the protocol, not the data, bump if FShaderCompilerOutput or WriteToOutputArchive changes.
 const int32 ShaderCompileWorkerOutputVersion = 7;
 // this is for the protocol, not the data.
@@ -482,6 +482,7 @@ enum class ESCWErrorCode
 	NoTargetShaderFormatsFound,
 	CantCompileForSpecificFormat,
 	CrashInsidePlatformCompiler,
+	BadInputFile
 };
 
 
