Fix for virtualization issue causing sticky sounds introduced with last AudioComponent change

diff --git a/Engine/Source/Runtime/Engine/Private/AudioDevice.cpp b/Engine/Source/Runtime/Engine/Private/AudioDevice.cpp
--- a/Engine/Source/Runtime/Engine/Private/AudioDevice.cpp
+++ b/Engine/Source/Runtime/Engine/Private/AudioDevice.cpp
@@ -4955,6 +4955,8 @@ void FAudioDevice::AddVirtualLoop(const FAudioVirtualLoop& InVirtualLoop)
 	FActiveSound& ActiveSound = VirtualLoop.GetActiveSound();
 	check(!VirtualLoops.Contains(&ActiveSound));
 
+	// If associated with an AudioComponent, add the virtualizing ActiveSound pointer to the VirtualLoop system, 
+	// and ensure it is in the AudioComponentIDToActiveSoundMap so updates from the AudioComponent are still tracked.
 	const int64 ComponentID = ActiveSound.GetAudioComponentID();
 	if (ComponentID > 0)
 	{
@@ -4969,16 +4971,22 @@ void FAudioDevice::AddVirtualLoop(const FAudioVirtualLoop& InVirtualLoop)
 					const FActiveSound* ExistingSound = (*ExistingSounds)[i];
 					if (ensure(ExistingSound))
 					{
-						UE_LOG(LogAudio, Warning, TEXT("Attempting to add Sound '%s' to ComponentID when map already contains ID for Sound '%s', and ID can only reference one active sound."),
+						UE_LOG(LogAudio, Warning, TEXT("Attempting to add Sound '%s' to ComponentID when map already contains ID for Sound '%s'. "
+						"Associated AudioComponent can only play a single sound instance at one time."),
 							ActiveSound.Sound ? *ActiveSound.Sound->GetName() : TEXT("N/A"),
 							ExistingSound->Sound ? *ExistingSound->Sound->GetName() : TEXT("N/A")
 						);
 						ExistingSounds->RemoveAtSwap(i, 1, false /* bAllowShrinking */);
 					}
 				}
-			}	
+			}
 			ExistingSounds->AddUnique(&ActiveSound);
 		}
+		else
+		{
+			TArray<FActiveSound*>& NewActiveSoundArray = AudioComponentIDToActiveSoundMap.Add(ComponentID);
+			NewActiveSoundArray.Add(&ActiveSound);
+		}
 	}
 
 #if !(UE_BUILD_SHIPPING || UE_BUILD_TEST)
diff --git a/Engine/Source/Runtime/Engine/Private/AudioVirtualLoop.cpp b/Engine/Source/Runtime/Engine/Private/AudioVirtualLoop.cpp
--- a/Engine/Source/Runtime/Engine/Private/AudioVirtualLoop.cpp
+++ b/Engine/Source/Runtime/Engine/Private/AudioVirtualLoop.cpp
@@ -79,6 +79,11 @@ bool FAudioVirtualLoop::Virtualize(const FActiveSound& InActiveSound, FAudioDevi
 		return false;
 	}
 
+	if (InAudioDevice.CanHaveMultipleActiveSounds(InActiveSound.GetAudioComponentID()))
+	{
+		return false;
+	}
+
 	if (bDoRangeCheck && IsInAudibleRange(InActiveSound, &InAudioDevice))
 	{
 		return false;
