Fix the list traversal.

- Also disable the code for a while to make sure it's correct.

diff --git a/Engine/Source/Runtime/Engine/Private/ShaderCompiler/ShaderCompiler.cpp b/Engine/Source/Runtime/Engine/Private/ShaderCompiler/ShaderCompiler.cpp
--- a/Engine/Source/Runtime/Engine/Private/ShaderCompiler/ShaderCompiler.cpp
+++ b/Engine/Source/Runtime/Engine/Private/ShaderCompiler/ShaderCompiler.cpp
@@ -79,7 +79,7 @@ DEFINE_LOG_CATEGORY(LogShaderCompilers);
 #define UE_SHADERCACHE_LOG_LEVEL		VeryVerbose
 
 // whether to parallelize writing/reading task files
-#define UE_SHADERCOMPILER_FIFO_JOB_EXECUTION  1
+#define UE_SHADERCOMPILER_FIFO_JOB_EXECUTION  0
 
 int32 GShaderCompilerJobCache = 1;
 static FAutoConsoleVariableRef CVarShaderCompilerJobCache(
@@ -549,7 +549,7 @@ void FShaderCompileJobCollection::SubmitJobs(const TArray<FShaderCommonCompileJo
 					// link the job at the end of pending list, so we're executing them in a FIFO and not LIFO order
 					if (PendingJobs[PriorityIndex])
 					{
-						for (FShaderCommonCompileJob::TIterator It(PendingJobs[PriorityIndex]); It;)
+						for (FShaderCommonCompileJob::TIterator It(PendingJobs[PriorityIndex]); It; ++It)
 						{
 							FShaderCommonCompileJob& ExistingJob = *It;
 							if (ExistingJob.GetNextLink() == nullptr)
