[VisualLogger] prevented all entries from being removed when cleaning up a world that was not registered in the WorldToRedirectionMap.
- removed unused VisualLoggerCleanupTimerHandle
- misc code cleanup (added 'const' + removed 'auto')

diff --git a/Engine/Source/Runtime/Engine/Private/VisualLogger/VisualLogger.cpp b/Engine/Source/Runtime/Engine/Private/VisualLogger/VisualLogger.cpp
--- a/Engine/Source/Runtime/Engine/Private/VisualLogger/VisualLogger.cpp
+++ b/Engine/Source/Runtime/Engine/Private/VisualLogger/VisualLogger.cpp
@@ -7,7 +7,6 @@
 #include "Misc/CommandLine.h"
 #include "Serialization/CustomVersion.h"
 #include "Modules/ModuleManager.h"
-#include "EngineGlobals.h"
 #include "AI/NavDataGenerator.h"
 #include "AI/NavigationSystemBase.h"
 #include "Framework/Docking/TabManager.h"
@@ -15,7 +14,6 @@
 #include "VisualLogger/VisualLoggerTraceDevice.h"
 #include "VisualLogger/VisualLoggerDebugSnapshotInterface.h"
 #include "Engine/Engine.h"
-#include "TimerManager.h"
 
 #if WITH_EDITOR
 #include "Editor/EditorEngine.h"
@@ -23,7 +21,7 @@
 
 
 DEFINE_LOG_CATEGORY(LogVisual);
-#if ENABLE_VISUAL_LOG 
+#if ENABLE_VISUAL_LOG
 DEFINE_STAT(STAT_VisualLog);
 
 namespace
@@ -158,7 +156,7 @@ bool FVisualLogger::IsObjectWhitelisted(const UObject* InObject) const
 
 FVisualLogEntry* FVisualLogger::GetLastEntryForObject(const UObject* Object)
 {
-	UObject* LogOwner = FVisualLogger::FindRedirection(Object);
+	const UObject* LogOwner = FVisualLogger::FindRedirection(Object);
 	return CurrentEntryPerObject.Contains(LogOwner) ? &CurrentEntryPerObject[LogOwner] : nullptr;
 }
 
@@ -182,7 +180,6 @@ FVisualLogEntry* FVisualLogger::GetEntryToWrite(const UObject* Object, float Tim
 			bInitializeNewEntry = (TimeStamp > CurrentEntry->TimeStamp) && (ShouldCreate == ECreateIfNeeded::Create);
 			if (World && IsInGameThread())
 			{
-				World->GetTimerManager().ClearTimer(VisualLoggerCleanupTimerHandle);
 				for (auto& CurrentPair : CurrentEntryPerObject)
 				{
 					FVisualLogEntry* Entry = &CurrentPair.Value;
@@ -210,8 +207,8 @@ FVisualLogEntry* FVisualLogger::GetEntryToWrite(const UObject* Object, float Tim
 		ObjectToClassNameMap.Add(LogOwner, *(LogOwner->GetClass()->GetName()));
 		ObjectToPointerMap.Add(LogOwner, LogOwner);
 		ObjectToWorldMap.Add(LogOwner, World);
-		
-		// IsClassWhitelisted isn't super fast, but this gets calculated only once for every 
+
+		// IsClassWhitelisted isn't super fast, but this gets calculated only once for every
 		// object trying to log something
 		CurrentEntry->bIsClassWhitelisted = (ClassWhitelist.Num() == 0) || IsClassWhitelisted(*LogOwner->GetClass()) || IsClassWhitelisted(*Object->GetClass());
 		CurrentEntry->bIsObjectWhitelisted = IsObjectWhitelisted(LogOwner);
@@ -231,7 +228,7 @@ FVisualLogEntry* FVisualLogger::GetEntryToWrite(const UObject* Object, float Tim
 			CurrentEntry->Location = ObjectAsActor->GetActorLocation();
 		}
 
-		auto& RedirectionMap = GetRedirectionMap(LogOwner);
+		FOwnerToChildrenRedirectionMap& RedirectionMap = GetRedirectionMap(LogOwner);
 		if (RedirectionMap.Contains(LogOwner))
 		{
 			if (ObjectToPointerMap.Contains(LogOwner) && ObjectToPointerMap[LogOwner].IsValid())
@@ -242,7 +239,7 @@ FVisualLogEntry* FVisualLogger::GetEntryToWrite(const UObject* Object, float Tim
 					DebugSnapshotInterface->GrabDebugSnapshot(CurrentEntry);
 				}
 			}
-			for (auto Child : RedirectionMap[LogOwner])
+			for (TWeakObjectPtr<const UObject>& Child : RedirectionMap[LogOwner])
 			{
 				if (Child.IsValid())
 				{
@@ -366,12 +363,12 @@ void FVisualLogger::NavigationDataDump(const UObject* Object, const FName& Categ
 
 	UWorld* World = nullptr;
 	FVisualLogEntry* CurrentEntry = nullptr;
-	if (CheckVisualLogInputInternal(Object, CategoryName, Verbosity, &World, &CurrentEntry) == false 
+	if (CheckVisualLogInputInternal(Object, CategoryName, Verbosity, &World, &CurrentEntry) == false
 		|| CurrentEntry == nullptr)
 	{
 		return;
 	}
-	
+
 	NavigationDataDumpDelegate.Broadcast(Object, CategoryName, Verbosity, Box, *World, *CurrentEntry);
 }
 
@@ -415,43 +412,45 @@ void FVisualLogger::Cleanup(UWorld* OldWorld, bool bReleaseMemory)
 		Device->Cleanup(bReleaseMemory);
 	}
 
-	if (OldWorld != nullptr && WorldToRedirectionMap.Find(OldWorld))
+	if (OldWorld != nullptr)
 	{
-		WorldToRedirectionMap.Remove(OldWorld);
-
-		if (WorldToRedirectionMap.Num() == 0)
+		// perform cleanup only if provided world is valid and was registered
+		if (WorldToRedirectionMap.Remove(OldWorld))
 		{
-			WorldToRedirectionMap.Reset();
-			ObjectToWorldMap.Reset();
-			ChildToOwnerMap.Reset();
-			CurrentEntryPerObject.Reset();
-			ObjectToNameMap.Reset();
-			ObjectToClassNameMap.Reset();
-			ObjectToPointerMap.Reset();
-		}
-		else
-		{
-			for (auto It = ObjectToWorldMap.CreateIterator(); It; ++It)
-			{
-				if (It.Value() == OldWorld)
-				{
-					const UObject* Obj = It.Key();
-					ObjectToWorldMap.Remove(Obj);
-					CurrentEntryPerObject.Remove(Obj);
-					ObjectToNameMap.Remove(Obj);
-					ObjectToClassNameMap.Remove(Obj);
-					ObjectToPointerMap.Remove(Obj);
-				}
-			}
-
-			for (FChildToOwnerRedirectionMap::TIterator It = ChildToOwnerMap.CreateIterator(); It; ++It)
-			{
-				if (It->Key.IsValid() == false
-					|| It->Key->GetWorld() == OldWorld)
-				{
-					It.RemoveCurrent();
-				}
-			}
+		    if (WorldToRedirectionMap.Num() == 0)
+            {
+                WorldToRedirectionMap.Reset();
+                ObjectToWorldMap.Reset();
+                ChildToOwnerMap.Reset();
+                CurrentEntryPerObject.Reset();
+                ObjectToNameMap.Reset();
+                ObjectToClassNameMap.Reset();
+                ObjectToPointerMap.Reset();
+            }
+            else
+            {
+                for (auto It = ObjectToWorldMap.CreateIterator(); It; ++It)
+                {
+                    if (It.Value() == OldWorld)
+                    {
+                        const UObject* Obj = It.Key();
+                        ObjectToWorldMap.Remove(Obj);
+                        CurrentEntryPerObject.Remove(Obj);
+                        ObjectToNameMap.Remove(Obj);
+                        ObjectToClassNameMap.Remove(Obj);
+                        ObjectToPointerMap.Remove(Obj);
+                    }
+                }
+
+                for (FChildToOwnerRedirectionMap::TIterator It = ChildToOwnerMap.CreateIterator(); It; ++It)
+                {
+                    if (It->Key.IsValid() == false
+                        || It->Key->GetWorld() == OldWorld)
+                    {
+                        It.RemoveCurrent();
+                    }
+                }
+            }
 		}
 	}
 	else
@@ -501,7 +500,7 @@ void FVisualLogger::Redirect(const UObject* FromObject, const UObject* ToObject)
 		return;
 	}
 
-	TWeakObjectPtr<const UObject> FromWeakPtr(FromObject);
+	const TWeakObjectPtr<const UObject> FromWeakPtr(FromObject);
 	UObject* OldRedirection = FindRedirection(FromObject);
 	UObject* NewRedirection = FindRedirection(ToObject);
 
@@ -509,7 +508,7 @@ void FVisualLogger::Redirect(const UObject* FromObject, const UObject* ToObject)
 	{
 		FOwnerToChildrenRedirectionMap& OwnerToChildrenMap = GetRedirectionMap(FromObject);
 
-		auto OldArray = OwnerToChildrenMap.Find(OldRedirection);
+		TArray<TWeakObjectPtr<const UObject>>* OldArray = OwnerToChildrenMap.Find(OldRedirection);
 		if (OldArray)
 		{
 			OldArray->RemoveSingleSwap(FromWeakPtr);
@@ -551,8 +550,8 @@ UObject* FVisualLogger::FindRedirection(const UObject* Object)
 	return const_cast<UObject*>(TargetWeakPtr.Get());
 }
 
-void FVisualLogger::SetIsRecording(bool InIsRecording) 
-{ 
+void FVisualLogger::SetIsRecording(bool InIsRecording)
+{
 	if (InIsRecording == false && InIsRecording != !!bIsRecording && FParse::Param(FCommandLine::Get(), TEXT("LogNavOctree")))
 	{
 		FVisualLogger::NavigationDataDump(GetWorldForVisualLogger(nullptr), LogNavigation, ELogVerbosity::Log, FBox());
@@ -561,7 +560,7 @@ void FVisualLogger::SetIsRecording(bool InIsRecording)
 	{
 		SetIsRecordingToFile(false);
 	}
-	bIsRecording = InIsRecording; 
+	bIsRecording = InIsRecording;
 };
 
 void FVisualLogger::SetIsRecordingToFile(bool InIsRecording)
@@ -576,7 +575,7 @@ void FVisualLogger::SetIsRecordingToFile(bool InIsRecording)
 	const FString BaseFileName = LogFileNameGetter.IsBound() ? LogFileNameGetter.Execute() : TEXT("VisualLog");
 	const FString MapName = World ? World->GetMapName() : TEXT("");
 
-	FString OutputFileName = FString::Printf(TEXT("%s_%s"), *BaseFileName, *MapName);
+	const FString OutputFileName = FString::Printf(TEXT("%s_%s"), *BaseFileName, *MapName);
 
 	if (bIsRecordingToFile && !InIsRecording)
 	{
@@ -666,10 +665,10 @@ FCustomVersionRegistration GVisualLoggerVersion(EVisualLoggerVersion::GUID, EVis
 
 #if !(UE_BUILD_SHIPPING || UE_BUILD_TEST)
 
-static class FLogVisualizerExec : private FSelfRegisteringExec
+class FLogVisualizerExec : private FSelfRegisteringExec
 {
 public:
-	/** Console commands, see embeded usage statement **/
+	/** Console commands, see embedded usage statement **/
 	virtual bool Exec(UWorld* InWorld, const TCHAR* Cmd, FOutputDevice& Ar) override
 	{
 		if (FParse::Command(&Cmd, TEXT("VISLOG")))
@@ -677,7 +676,7 @@ public:
 			if (FModuleManager::Get().LoadModulePtr<IModuleInterface>("LogVisualizer") != nullptr)
 			{
 #if ENABLE_VISUAL_LOG
-				FString Command = FParse::Token(Cmd, 0);
+				const FString Command = FParse::Token(Cmd, /*UseEscape*/ false);
 				if (Command == TEXT("record"))
 				{
 					FVisualLogger::Get().SetIsRecording(true);
@@ -690,7 +689,7 @@ public:
 				}
 				else if (Command == TEXT("disableallbut"))
 				{
-					FString Category = FParse::Token(Cmd, 1);
+					const FString Category = FParse::Token(Cmd, /*UseEscape*/ true);
 					FVisualLogger::Get().BlockAllCategories(true);
 					FVisualLogger::Get().AddCategoryToWhitelist(*Category);
 					return true;
@@ -717,5 +716,4 @@ public:
 	}
 } LogVisualizerExec;
 
-
 #endif // !(UE_BUILD_SHIPPING || UE_BUILD_TEST)
diff --git a/Engine/Source/Runtime/Engine/Public/VisualLogger/VisualLogger.h b/Engine/Source/Runtime/Engine/Public/VisualLogger/VisualLogger.h
--- a/Engine/Source/Runtime/Engine/Public/VisualLogger/VisualLogger.h
+++ b/Engine/Source/Runtime/Engine/Public/VisualLogger/VisualLogger.h
@@ -608,8 +608,7 @@ protected:
 	float StartRecordingToFileTime;
 	/** Delegate to set project specific file name for vlogs */
 	FVisualLogFilenameGetterDelegate LogFileNameGetter;
-	/** Handle for timer used to serialize all waiting logs */
-	FTimerHandle VisualLoggerCleanupTimerHandle;
+
 	/** function to call when getting the time stamp */
 	TFunction<float(const UObject*)> GetTimeStampFunc;
 
