WorldPartition: Newly created world now add the grid size to partitioned actors

diff --git a/Engine/Source/Runtime/Engine/Classes/GameFramework/WorldSettings.h b/Engine/Source/Runtime/Engine/Classes/GameFramework/WorldSettings.h
--- a/Engine/Source/Runtime/Engine/Classes/GameFramework/WorldSettings.h
+++ b/Engine/Source/Runtime/Engine/Classes/GameFramework/WorldSettings.h
@@ -522,6 +522,21 @@ public:
 	UPROPERTY(config, EditAnywhere, Category = Broadphase)
 	uint8 bOverrideDefaultBroadphaseSettings:1;
 
+#if WITH_EDITORONLY_DATA
+	/** 
+	 * Whether Foliage actors of this world contain their grid size in their name. This should only be changed by UWorldPartitionFoliageBuilder or when creating new worlds so that older worlds are unaffected
+	 * and is used by the UActorPartitionSubsystem to find existing foliage actors by name.
+	 */
+	UPROPERTY()
+	uint8 bIncludeGridSizeInNameForFoliageActors : 1;
+
+	/**
+	 * Whether partitioned actors of this world contain their grid size in their name. This should only be changed when creating new worlds so that older worlds are unaffected
+	 * and is used by the UActorPartitionSubsystem to find existing foliage actors by name.
+	 */
+	UPROPERTY()
+	uint8 bIncludeGridSizeInNameForPartitionedActors : 1;
+#endif
 protected:
 	/** Holds parameters for NavigationSystem's creation. Set to Null will result
 	 *	in NavigationSystem instance not being created for this world. Note that
@@ -537,6 +552,7 @@ protected:
 	TObjectPtr<UWorldPartition> WorldPartition;
 
 public:
+#if WITH_EDITORONLY_DATA
 	/** Size of the grid for instanced foliage actors*/
 	UPROPERTY(EditAnywhere, Category = Foliage)
 	uint32 InstancedFoliageGridSize;
@@ -544,6 +560,7 @@ public:
 	/** Default size of the grid for placed elements from the editor */
 	UPROPERTY(EditAnywhere, Category=Editor)
 	uint32 DefaultPlacementGridSize;
+#endif
 
 	/** scale of 1uu to 1m in real world measurements, for HMD and other physically tracked devices (e.g. 1uu = 1cm would be 100.0) */
 	UPROPERTY(EditAnywhere, BlueprintReadOnly, Category=VR)
diff --git a/Engine/Source/Runtime/Engine/Private/ActorPartition/ActorPartitionSubsystem.cpp b/Engine/Source/Runtime/Engine/Private/ActorPartition/ActorPartitionSubsystem.cpp
--- a/Engine/Source/Runtime/Engine/Private/ActorPartition/ActorPartitionSubsystem.cpp
+++ b/Engine/Source/Runtime/Engine/Private/ActorPartition/ActorPartitionSubsystem.cpp
@@ -250,6 +250,11 @@ public:
 				ActorNameBuilder += TEXT("_");
 			}
 
+			if (InActorPartitionId.GetClass()->GetDefaultObject<APartitionActor>()->ShouldIncludeGridSizeInName(World))
+			{
+				ActorNameBuilder += FString::Printf(TEXT("%d_"), InGridSize);
+			}
+
 			ActorNameBuilder += FString::Printf(TEXT("%d_%d_%d"), InCellCoord.X, InCellCoord.Y, InCellCoord.Z);
 
 			if (InActorPartitionId.GetDataLayerEditorContextHash() != FDataLayerEditorContext::EmptyHash)
diff --git a/Engine/Source/Runtime/Engine/Private/ActorPartition/PartitionActor.cpp b/Engine/Source/Runtime/Engine/Private/ActorPartition/PartitionActor.cpp
--- a/Engine/Source/Runtime/Engine/Private/ActorPartition/PartitionActor.cpp
+++ b/Engine/Source/Runtime/Engine/Private/ActorPartition/PartitionActor.cpp
@@ -42,6 +42,11 @@ bool APartitionActor::IsUserManaged() const
 
 	return true;
 }
+
+bool APartitionActor::ShouldIncludeGridSizeInName(UWorld* InWorld) const
+{
+	return InWorld->GetWorldSettings()->bIncludeGridSizeInNameForPartitionedActors;
+}
 #endif
 
 #undef LOCTEXT_NAMESPACE
\ No newline at end of file
diff --git a/Engine/Source/Runtime/Engine/Private/WorldPartition/WorldPartition.cpp b/Engine/Source/Runtime/Engine/Private/WorldPartition/WorldPartition.cpp
--- a/Engine/Source/Runtime/Engine/Private/WorldPartition/WorldPartition.cpp
+++ b/Engine/Source/Runtime/Engine/Private/WorldPartition/WorldPartition.cpp
@@ -710,6 +710,11 @@ UWorldPartition* UWorldPartition::CreateWorldPartition(AWorldSettings* WorldSett
 
 	UWorldPartition* WorldPartition = NewObject<UWorldPartition>(WorldSettings);
 	WorldSettings->SetWorldPartition(WorldPartition);
+	
+	// New maps should include GridSize in name
+	WorldSettings->bIncludeGridSizeInNameForFoliageActors = true;
+	WorldSettings->bIncludeGridSizeInNameForPartitionedActors = true;
+
 	WorldSettings->MarkPackageDirty();
 
 	WorldPartition->EditorHash = NewObject<UWorldPartitionEditorHash>(WorldPartition, EditorHashClass);
diff --git a/Engine/Source/Runtime/Engine/Private/WorldSettings.cpp b/Engine/Source/Runtime/Engine/Private/WorldSettings.cpp
--- a/Engine/Source/Runtime/Engine/Private/WorldSettings.cpp
+++ b/Engine/Source/Runtime/Engine/Private/WorldSettings.cpp
@@ -89,8 +89,11 @@ PRAGMA_ENABLE_DEPRECATION_WARNINGS
 	KillZ = -HALF_WORLD_MAX1;
 	KillZDamageType = ConstructorStatics.DmgType_Environmental_Object.Object;
 
+#if WITH_EDITORONLY_DATA
 	InstancedFoliageGridSize = DefaultPlacementGridSize = 25600;
-
+	bIncludeGridSizeInNameForFoliageActors = false;
+	bIncludeGridSizeInNameForPartitionedActors = false;
+#endif
 	WorldToMeters = 100.f;
 
 	DefaultPhysicsVolumeClass = ADefaultPhysicsVolume::StaticClass();
diff --git a/Engine/Source/Runtime/Engine/Public/ActorPartition/PartitionActor.h b/Engine/Source/Runtime/Engine/Public/ActorPartition/PartitionActor.h
--- a/Engine/Source/Runtime/Engine/Public/ActorPartition/PartitionActor.h
+++ b/Engine/Source/Runtime/Engine/Public/ActorPartition/PartitionActor.h
@@ -23,6 +23,7 @@ public:
 	virtual EActorGridPlacement GetDefaultGridPlacement() const override { return EActorGridPlacement::Location; }
 	virtual TUniquePtr<class FWorldPartitionActorDesc> CreateClassActorDesc() const override;
 	virtual uint32 GetDefaultGridSize(UWorld* InWorld) const PURE_VIRTUAL(APartitionActor, return 0;)
+	virtual bool ShouldIncludeGridSizeInName(UWorld * InWorld) const;
 	virtual FGuid GetGridGuid() const { return FGuid(); }
 	virtual bool IsUserManaged() const override;
 #endif
diff --git a/Engine/Source/Runtime/Foliage/Private/InstancedFoliage.cpp b/Engine/Source/Runtime/Foliage/Private/InstancedFoliage.cpp
--- a/Engine/Source/Runtime/Foliage/Private/InstancedFoliage.cpp
+++ b/Engine/Source/Runtime/Foliage/Private/InstancedFoliage.cpp
@@ -4947,6 +4947,11 @@ uint32 AInstancedFoliageActor::GetDefaultGridSize(UWorld* InWorld) const
 	return InWorld->GetWorldSettings()->InstancedFoliageGridSize;
 }
 
+bool AInstancedFoliageActor::ShouldIncludeGridSizeInName(UWorld* InWorld) const
+{
+	return InWorld->GetWorldSettings()->bIncludeGridSizeInNameForFoliageActors;
+}
+
 void AInstancedFoliageActor::OnApplyLevelTransform(const FTransform& InTransform)
 {
 #if WITH_EDITORONLY_DATA
diff --git a/Engine/Source/Runtime/Foliage/Public/InstancedFoliageActor.h b/Engine/Source/Runtime/Foliage/Public/InstancedFoliageActor.h
--- a/Engine/Source/Runtime/Foliage/Public/InstancedFoliageActor.h
+++ b/Engine/Source/Runtime/Foliage/Public/InstancedFoliageActor.h
@@ -96,6 +96,7 @@ public:
 	FOLIAGE_API void DetectFoliageTypeChangeAndUpdate();
 
 	virtual uint32 GetDefaultGridSize(UWorld* InWorld) const override;
+	virtual bool ShouldIncludeGridSizeInName(UWorld* InWorld) const override;
 
 	// Delegate type for selection change events
 	DECLARE_MULTICAST_DELEGATE_TwoParams(FOnSelectionChanged, bool, const TArray<AActor*>&);
