Fixes crash in nanite mesh when reloading the global shader map

diff --git a/Engine/Source/Runtime/Engine/Private/ShaderCompiler/ShaderCompiler.cpp b/Engine/Source/Runtime/Engine/Private/ShaderCompiler/ShaderCompiler.cpp
--- a/Engine/Source/Runtime/Engine/Private/ShaderCompiler/ShaderCompiler.cpp
+++ b/Engine/Source/Runtime/Engine/Private/ShaderCompiler/ShaderCompiler.cpp
@@ -6816,6 +6816,8 @@ void ReloadGlobalShaders()
 	{
 		BeginUpdateResourceRHI(*It);
 	}
+
+	PropagateGlobalShadersToAllPrimitives();
 }
 
 static FAutoConsoleCommand CCmdReloadGlobalShaders = FAutoConsoleCommand(
