ShaderCompiler: re-enable FIFO task mode.

- Also some changes for clarity/correctness.

diff --git a/Engine/Source/Runtime/Engine/Private/ShaderCompiler/ShaderCompiler.cpp b/Engine/Source/Runtime/Engine/Private/ShaderCompiler/ShaderCompiler.cpp
--- a/Engine/Source/Runtime/Engine/Private/ShaderCompiler/ShaderCompiler.cpp
+++ b/Engine/Source/Runtime/Engine/Private/ShaderCompiler/ShaderCompiler.cpp
@@ -79,7 +79,7 @@ DEFINE_LOG_CATEGORY(LogShaderCompilers);
 #define UE_SHADERCACHE_LOG_LEVEL		VeryVerbose
 
 // whether to parallelize writing/reading task files
-#define UE_SHADERCOMPILER_FIFO_JOB_EXECUTION  0
+#define UE_SHADERCOMPILER_FIFO_JOB_EXECUTION  1
 
 int32 GShaderCompilerJobCache = 1;
 static FAutoConsoleVariableRef CVarShaderCompilerJobCache(
@@ -706,51 +706,57 @@ int32 FShaderCompileJobCollection::GetPendingJobs(EShaderCompilerWorkerType InWo
 		return 0;
 	}
 
+	FWriteScopeLock Locker(Lock);
+
+	// there was a time window before we checked and then acquired the write lock - make sure the number is still sufficient
+	NumPendingJobsOfPriority = NumPendingJobs[PriorityIndex];
+	if (NumPendingJobsOfPriority < MinNumJobs)
+	{
+		// Not enough jobs
+		return 0;
+	}
+	
 	OutJobs.Reserve(OutJobs.Num() + FMath::Min(MaxNumJobs, NumPendingJobsOfPriority));
-	int32 NumJobs = 0;
+	int32 NumJobs = FMath::Min(MaxNumJobs, NumPendingJobs[PriorityIndex]);
+	FShaderCommonCompileJob::TIterator It(PendingJobs[PriorityIndex]);
+	// Randomize job selection by randomly skipping over jobs while traversing the list.
+	// Say, we need to pick 3 jobs out of 5 total. We can skip over 2 jobs in total, e.g. like this:
+	// pick one (4 more to go and we need to get 2 of 4), skip one (3 more to go, picking 2 out of 3), pick one (2 more to go, picking 1 of 2), skip one, pick one.
+	// It is possible that we won't skip at all and instead pick consequential jobs
+	int32 MaxJobsWeCanSkipOver = NumPendingJobs[PriorityIndex] - NumJobs;
+	for (int32 i = 0; i < NumJobs; ++i)
 	{
-		FWriteScopeLock Locker(Lock);
-		NumJobs = FMath::Min(MaxNumJobs, NumPendingJobs[PriorityIndex]);
-		FShaderCommonCompileJob::TIterator It(PendingJobs[PriorityIndex]);
-		// Randomize job selection by randomly skipping over jobs while traversing the list.
-		// Say, we need to pick 3 jobs out of 5 total. We can skip over 2 jobs in total, e.g. like this:
-		// pick one (4 more to go and we need to get 2 of 4), skip one (3 more to go, picking 2 out of 3), pick one (2 more to go, picking 1 of 2), skip one, pick one.
-		// It is possible that we won't skip at all and instead pick consequential jobs
-		int32 MaxJobsWeCanSkipOver = NumPendingJobs[PriorityIndex] - NumJobs;
-		for (int32 i = 0; i < NumJobs; ++i)
-		{
-			FShaderCommonCompileJob& Job = *It;
-
-			// get a random number of jobs to skip (if we can)
-			if (MaxJobsWeCanSkipOver > 0)
-			{
-				int32 NumJobsToSkipOver = FMath::RandHelper(MaxJobsWeCanSkipOver + 1);
-				while (NumJobsToSkipOver > 0)
-				{
-					It.Next();
-					--NumJobsToSkipOver;
-					--MaxJobsWeCanSkipOver;
-				}
-				checkf(MaxJobsWeCanSkipOver >= 0, TEXT("We skipped over too many jobs"));
-				checkf(MaxJobsWeCanSkipOver <= NumPendingJobs[PriorityIndex] - i, TEXT("Number of jobs to skip should stay less or equal than the number of nodes to go"));
-			}
+		FShaderCommonCompileJob& Job = *It;
 
-			GShaderCompilerStats->RegisterAssignedJob(Job);
-			// Temporary commented out until r.ShaderDevelopmentMode=1 shader error retry crash gets fixed
-			//check(Job.CurrentWorker == EShaderCompilerWorkerType::None);
-			//check(Job.PendingPriority == InPriority);
-			ensure(!ShaderCompiler::IsJobCacheEnabled() || Job.bInputHashSet);
+		GShaderCompilerStats->RegisterAssignedJob(Job);
+		// Temporary commented out until r.ShaderDevelopmentMode=1 shader error retry crash gets fixed
+		//check(Job.CurrentWorker == EShaderCompilerWorkerType::None);
+		//check(Job.PendingPriority == InPriority);
+		ensure(!ShaderCompiler::IsJobCacheEnabled() || Job.bInputHashSet);
 
-			It.Next();
-			Job.Unlink();
+		It.Next();
+		Job.Unlink();
 
-			Job.PendingPriority = EShaderCompileJobPriority::None;
-			Job.CurrentWorker = InWorkerType;
-			OutJobs.Add(&Job);
-		}
+		Job.PendingPriority = EShaderCompileJobPriority::None;
+		Job.CurrentWorker = InWorkerType;
+		OutJobs.Add(&Job);
 
-		NumPendingJobs[PriorityIndex] -= NumJobs;
+		// get a random number of jobs to skip (if we can). We're skipping after taking the first job so we can ensure that we always take the latest job into the batch
+		if (MaxJobsWeCanSkipOver > 0)
+		{
+			int32 NumJobsToSkipOver = FMath::RandHelper(MaxJobsWeCanSkipOver + 1);
+			while (NumJobsToSkipOver > 0 && It)
+			{
+				It.Next();
+				--NumJobsToSkipOver;
+				--MaxJobsWeCanSkipOver;
+			}
+			checkf(MaxJobsWeCanSkipOver >= 0, TEXT("We skipped over too many jobs"));
+			checkf(MaxJobsWeCanSkipOver <= NumPendingJobs[PriorityIndex] - i, TEXT("Number of jobs to skip should stay less or equal than the number of nodes to go"));
+		}
 	}
+
+	NumPendingJobs[PriorityIndex] -= NumJobs;
 	return NumJobs;
 }
 
@@ -2965,7 +2971,14 @@ void FShaderCompilerStats::WriteStatSummary()
 		if (TotalTimeAtLeastOneJobWasInFlight > 0.0)
 		{
 			double EffectiveParallelization = TotalTimeForAllShaders / TotalTimeAtLeastOneJobWasInFlight;
-			UE_LOG(LogShaderCompilers, Display, TEXT("Effective parallelization: %.2f (times faster than compiling all shaders on one thread)."), EffectiveParallelization);
+			if (DistributedJobBatchesSeen == 0)
+			{
+				UE_LOG(LogShaderCompilers, Display, TEXT("Effective parallelization: %.2f (times faster than compiling all shaders on one thread). Compare with number of workers: %d"), EffectiveParallelization, GShaderCompilingManager->GetNumLocalWorkers());
+			}
+			else
+			{
+				UE_LOG(LogShaderCompilers, Display, TEXT("Effective parallelization: %.2f (times faster than compiling all shaders on one thread). Distributed compilation was used."), EffectiveParallelization);
+			}
 		}
 
 
diff --git a/Engine/Source/Runtime/Engine/Public/ShaderCompiler.h b/Engine/Source/Runtime/Engine/Public/ShaderCompiler.h
--- a/Engine/Source/Runtime/Engine/Public/ShaderCompiler.h
+++ b/Engine/Source/Runtime/Engine/Public/ShaderCompiler.h
@@ -1030,6 +1030,14 @@ public:
 		return GetNumOutstandingJobs() + NumExternalJobs;
 	}
 
+	/**
+	 * Returns the (current) number of local workers.
+	 */
+	int32 GetNumLocalWorkers() const
+	{
+		return bCompilingDuringGame ? NumShaderCompilingThreadsDuringGame : NumShaderCompilingThreads;
+	}
+
 	void SetExternalJobs(int32 NumJobs)
 	{
 		NumExternalJobs = NumJobs;
