Increase the XGE batch size on the build machine to help with the pool oversubscription.

diff --git a/Engine/Source/Runtime/Engine/Private/ShaderCompiler/ShaderCompilerDistributed.cpp b/Engine/Source/Runtime/Engine/Private/ShaderCompiler/ShaderCompilerDistributed.cpp
--- a/Engine/Source/Runtime/Engine/Private/ShaderCompiler/ShaderCompilerDistributed.cpp
+++ b/Engine/Source/Runtime/Engine/Private/ShaderCompiler/ShaderCompilerDistributed.cpp
@@ -198,9 +198,12 @@ int32 FShaderCompileDistributedThreadRunnable_Interface::CompilingLoop()
 	{
 		// Increase the batch size when more jobs are queued/in flight.
 
+		// Build farm is much more prone to pool oversubscription, so make sure the jobs are submitted in batches of at least MinBatchSize
+		int MinJobsPerBatch = GIsBuildMachine ? DistributedShaderCompilerVariables::MinBatchSize : 1;
+
 		// Just to provide typical numbers: the number of total jobs is usually in tens of thousands at most, oftentimes in low thousands. Thus JobsPerBatch when calculated as a log2 rarely reaches the value of 16,
-		// and that seems to be a sweet spot: lowering it does not result in faster completion, while increasing the number of jobs per batch slows it down .
-		const uint32 JobsPerBatch = FMath::Max(1, FMath::FloorToInt(FMath::LogX(2.f, PendingJobs.Num() + NumDispatchedJobs)));
+		// and that seems to be a sweet spot: lowering it does not result in faster completion, while increasing the number of jobs per batch slows it down.
+		const uint32 JobsPerBatch = FMath::Max(MinJobsPerBatch, FMath::FloorToInt(FMath::LogX(2.f, PendingJobs.Num() + NumDispatchedJobs)));
 		UE_LOG(LogShaderCompilers, Verbose, TEXT("Current jobs: %d, Batch size: %d, Num Already Dispatched: %d"), PendingJobs.Num(), JobsPerBatch, NumDispatchedJobs);
 
 
