Static analysis fixes

diff --git a/Engine/Source/Developer/AudioFormatBink/Private/AudioFormatBink.cpp b/Engine/Source/Developer/AudioFormatBink/Private/AudioFormatBink.cpp
--- a/Engine/Source/Developer/AudioFormatBink/Private/AudioFormatBink.cpp
+++ b/Engine/Source/Developer/AudioFormatBink/Private/AudioFormatBink.cpp
@@ -122,7 +122,8 @@ public:
 		Current += sizeof(uint16) * Header->seek_table_entry_count;
 
 		// We must be on the first frame.
-		ensure( (*((uint16*)Current)) == BLOCK_HEADER_MAGIC);
+		uint16 CurrentFrameValue = *(uint16*)Current;
+		ensure(BLOCK_HEADER_MAGIC == CurrentFrameValue);
 
 		int32 ChunkLimitBytes = InMaxInitialChunkSize;
 		for (;;)
diff --git a/Engine/Source/Runtime/Engine/Private/ActiveSound.cpp b/Engine/Source/Runtime/Engine/Private/ActiveSound.cpp
--- a/Engine/Source/Runtime/Engine/Private/ActiveSound.cpp
+++ b/Engine/Source/Runtime/Engine/Private/ActiveSound.cpp
@@ -49,6 +49,7 @@ FActiveSound::FActiveSound()
 	, SourceEffectChain(nullptr)
 	, AudioComponentID(0)
 	, OwnerID(0)
+	, PlayOrder(INDEX_NONE)
 	, AudioDevice(nullptr)
 	, SoundClassOverride(nullptr)
 	, bHasCheckedOcclusion(false)
diff --git a/Engine/Source/Runtime/Engine/Private/AudioDevice.cpp b/Engine/Source/Runtime/Engine/Private/AudioDevice.cpp
--- a/Engine/Source/Runtime/Engine/Private/AudioDevice.cpp
+++ b/Engine/Source/Runtime/Engine/Private/AudioDevice.cpp
@@ -5149,11 +5149,17 @@ void FAudioDevice::ProcessingPendingActiveSoundStops(bool bForceDelete)
 			if (bDeleteActiveSound)
 			{
 				NotifyPendingDelete(*ActiveSound);
+
+				// Remove from the list of pending sounds to stop
+				PendingSoundsToStop.Remove(ActiveSound);
+	
 				delete ActiveSound;
 			}
-		
-			// Remove from the list of pending sounds to stop
-			PendingSoundsToStop.Remove(ActiveSound);
+			else
+			{
+				// Remove from the list of pending sounds to stop
+				PendingSoundsToStop.Remove(ActiveSound);
+			}	
 		}
 	}
 }
