Re-expose anim runtime methods to engine api

diff --git a/Engine/Source/Runtime/Engine/Public/AnimationRuntime.h b/Engine/Source/Runtime/Engine/Public/AnimationRuntime.h
--- a/Engine/Source/Runtime/Engine/Public/AnimationRuntime.h
+++ b/Engine/Source/Runtime/Engine/Public/AnimationRuntime.h
@@ -43,7 +43,7 @@ namespace ETransformBlendMode
 }
 
 template<int32>
-FORCEINLINE void BlendTransform(const FTransform& Source, FTransform& Dest, const float BlendWeight);
+ENGINE_API void BlendTransform(const FTransform& Source, FTransform& Dest, const float BlendWeight);
 
 template<>
 FORCEINLINE void BlendTransform<ETransformBlendMode::Overwrite>(const FTransform& Source, FTransform& Dest, const float BlendWeight)
@@ -59,8 +59,8 @@ FORCEINLINE void BlendTransform<ETransformBlendMode::Accumulate>(const FTransfor
 	Dest.AccumulateWithShortestRotation(Source, VBlendWeight);
 }
 
-FORCEINLINE void BlendCurves(const TArrayView<const FBlendedCurve> SourceCurves, const TArrayView<const float> SourceWeights, FBlendedCurve& OutCurve);
-void BlendCurves(const TArrayView<const FBlendedCurve* const> SourceCurves, const TArrayView<const float> SourceWeights, FBlendedCurve& OutCurve);
+ENGINE_API void BlendCurves(const TArrayView<const FBlendedCurve> SourceCurves, const TArrayView<const float> SourceWeights, FBlendedCurve& OutCurve);
+void ENGINE_API BlendCurves(const TArrayView<const FBlendedCurve* const> SourceCurves, const TArrayView<const float> SourceWeights, FBlendedCurve& OutCurve);
 
 
 /////////////////////////////////////////////////////////
