Update ISMPartitionActor to use ModifyActor and prevent multiple Modify calls from placement subsystem.

diff --git a/Engine/Source/Runtime/Engine/Private/ISMPartition/ISMPartitionActor.cpp b/Engine/Source/Runtime/Engine/Private/ISMPartition/ISMPartitionActor.cpp
--- a/Engine/Source/Runtime/Engine/Private/ISMPartition/ISMPartitionActor.cpp
+++ b/Engine/Source/Runtime/Engine/Private/ISMPartition/ISMPartitionActor.cpp
@@ -97,7 +97,7 @@ void AISMPartitionActor::PostEditUndo()
 
 FISMClientHandle AISMPartitionActor::RegisterClient(const FGuid& ClientGuid)
 {
-	Modify();
+	ModifyActor();
 	int32 Index = Clients.Find(ClientGuid);
 	if (Index == INDEX_NONE)
 	{
@@ -117,7 +117,7 @@ FISMClientHandle AISMPartitionActor::RegisterClient(const FGuid& ClientGuid)
 
 void AISMPartitionActor::UnregisterClient(FISMClientHandle& Handle)
 {
-	Modify();
+	ModifyActor();
 	check(Handle.Guid == Clients[Handle.Index]);
 	RemoveISMInstances(Handle);
 	
