Adding extra LOD overrides in FBuildMeshDescriptionsParams. Currently set :
- bUseHighPrecisionTangentBasis
- bUseFullPrecisionUVs

diff --git a/Engine/Source/Runtime/Engine/Classes/Engine/StaticMesh.h b/Engine/Source/Runtime/Engine/Classes/Engine/StaticMesh.h
--- a/Engine/Source/Runtime/Engine/Classes/Engine/StaticMesh.h
+++ b/Engine/Source/Runtime/Engine/Classes/Engine/StaticMesh.h
@@ -1383,6 +1383,19 @@ public:
 	UFUNCTION(BlueprintCallable, Category="StaticMesh")
 	ENGINE_API UStaticMeshDescription* GetStaticMeshDescription(int32 LODIndex);
 
+	struct FBuildMeshDescriptionsLODParams
+	{
+		/**
+		 * If true, Tangents will be stored at 16 bit vs 8 bit precision.
+		 */
+		bool bUseHighPrecisionTangentBasis = false;
+
+		/**
+		 * If true, UVs will be stored at full floating point precision.
+		 */
+		bool bUseFullPrecisionUVs = false;
+	};
+
 	 /** Structure that defines parameters passed into the build mesh description function */
 	struct FBuildMeshDescriptionsParams
 	{
@@ -1427,6 +1440,11 @@ public:
 		 * Ored with the value of bAllowCpuAccess on the static mesh. Set to false by default.
 		 */
 		bool bAllowCpuAccess;
+
+		/**
+		 * Extra optional LOD params. Overrides any previous settings from source model build settings.
+		 */
+		TArray<FBuildMeshDescriptionsLODParams> PerLODOverrides;
 	};
 
 	/**
diff --git a/Engine/Source/Runtime/Engine/Private/StaticMesh.cpp b/Engine/Source/Runtime/Engine/Private/StaticMesh.cpp
--- a/Engine/Source/Runtime/Engine/Private/StaticMesh.cpp
+++ b/Engine/Source/Runtime/Engine/Private/StaticMesh.cpp
@@ -5995,6 +5995,8 @@ bool UStaticMesh::BuildFromMeshDescriptions(const TArray<const FMeshDescription*
 {
 	TRACE_CPUPROFILER_EVENT_SCOPE(UStaticMesh::BuildFromMeshDescriptions);
 
+	const int32 NewNumLODs = MeshDescriptions.Num();
+
 	// Set up
 	NeverStream = true;
 
@@ -6029,9 +6031,16 @@ bool UStaticMesh::BuildFromMeshDescriptions(const TArray<const FMeshDescription*
 		CommitParams.bMarkPackageDirty = Params.bMarkPackageDirty;
 		CommitParams.bUseHashAsGuid = Params.bUseHashAsGuid;
 
-		SetNumSourceModels(MeshDescriptions.Num());
-		for (int32 LODIndex = 0; LODIndex < MeshDescriptions.Num(); LODIndex++)
+		SetNumSourceModels(NewNumLODs);
+		for (int32 LODIndex = 0; LODIndex < NewNumLODs; LODIndex++)
 		{
+			FMeshBuildSettings& LODBuildSettings = GetSourceModel(LODIndex).BuildSettings;
+			if (Params.PerLODOverrides.IsValidIndex(LODIndex))
+			{
+				const FBuildMeshDescriptionsLODParams& LODParams = Params.PerLODOverrides[LODIndex];
+				LODBuildSettings.bUseHighPrecisionTangentBasis = LODParams.bUseHighPrecisionTangentBasis;
+				LODBuildSettings.bUseFullPrecisionUVs = LODParams.bUseFullPrecisionUVs;
+			}
 			CreateMeshDescription(LODIndex, *MeshDescriptions[LODIndex]);
 			CommitMeshDescription(LODIndex, CommitParams);
 		}
@@ -6039,12 +6048,18 @@ bool UStaticMesh::BuildFromMeshDescriptions(const TArray<const FMeshDescription*
 	#endif
 
 	SetRenderData(MakeUnique<FStaticMeshRenderData>());
-	GetRenderData()->AllocateLODResources(MeshDescriptions.Num());
+	GetRenderData()->AllocateLODResources(NewNumLODs);
 
 	FStaticMeshLODResourcesArray& LODResourcesArray = GetRenderData()->LODResources;
 	for (int32 LODIndex = 0; LODIndex < LODResourcesArray.Num(); ++LODIndex)
 	{
 		LODResourcesArray[LODIndex].IndexBuffer.TrySetAllowCPUAccess(bAllowCPUAccess || Params.bAllowCpuAccess);
+		if (Params.PerLODOverrides.IsValidIndex(LODIndex))
+		{
+			const FBuildMeshDescriptionsLODParams& LODParams = Params.PerLODOverrides[LODIndex];
+			LODResourcesArray[LODIndex].VertexBuffers.StaticMeshVertexBuffer.SetUseHighPrecisionTangentBasis(LODParams.bUseHighPrecisionTangentBasis);
+			LODResourcesArray[LODIndex].VertexBuffers.StaticMeshVertexBuffer.SetUseFullPrecisionUVs(LODParams.bUseFullPrecisionUVs);
+		}
 	}
 
 	// Build render data from each mesh description
@@ -6053,7 +6068,7 @@ bool UStaticMesh::BuildFromMeshDescriptions(const TArray<const FMeshDescription*
 	if (Params.bFastBuild)
 #endif
 	{
-		for (int32 LODIndex = 0; LODIndex < MeshDescriptions.Num(); LODIndex++)
+		for (int32 LODIndex = 0; LODIndex < NewNumLODs; LODIndex++)
 		{
 			check(MeshDescriptions[LODIndex] != nullptr);
 			FStaticMeshLODResources& LODResources = GetRenderData()->LODResources[LODIndex];
@@ -6067,7 +6082,7 @@ bool UStaticMesh::BuildFromMeshDescriptions(const TArray<const FMeshDescription*
 		GetRenderData()->Bounds = MeshDescriptions[0]->GetBounds();
 		CalculateExtendedBounds();
 
-		for (int32 LOD = 0; LOD < MeshDescriptions.Num(); ++LOD)
+		for (int32 LOD = 0; LOD < NewNumLODs; ++LOD)
 		{
 			// @todo: some way of customizing LOD screen size and/or calculate it based on mesh bounds
 			if (true)
@@ -6114,7 +6129,7 @@ bool UStaticMesh::BuildFromMeshDescriptions(const TArray<const FMeshDescription*
 		Build(true);
 	}
 
-	for (int32 LODIndex = 0; LODIndex < MeshDescriptions.Num(); LODIndex++)
+	for (int32 LODIndex = 0; LODIndex < NewNumLODs; LODIndex++)
 	{
 		FStaticMeshLODResources& LODResources = GetRenderData()->LODResources[LODIndex];
 		for (int32 SectionIndex = 0; SectionIndex < LODResources.Sections.Num(); SectionIndex++)
