Change output of DDC2 texture build function to better organize structured data and bulk data into separate payloads.
Remove older prototype DDC2 texture build code path and route the commandline switch it used to be synonymous with the newer path.
Fixed a bug where DDC2 produced texture mips didn't clear the "PagedToDerivedData" flag when they were streamed in causing the code to then go get the DDC1 texture payload.
Fixed a bug where the mip info provider wasn't accounting for the fact that DDC2 mip payloads aren't prefixed with a MipSize.

diff --git a/Engine/Source/Developer/TextureBuild/Private/TextureBuildFunction.cpp b/Engine/Source/Developer/TextureBuild/Private/TextureBuildFunction.cpp
--- a/Engine/Source/Developer/TextureBuild/Private/TextureBuildFunction.cpp
+++ b/Engine/Source/Developer/TextureBuild/Private/TextureBuildFunction.cpp
@@ -15,6 +15,7 @@
 #include "Modules/ModuleManager.h"
 #include "PixelFormat.h"
 #include "Serialization/CompactBinary.h"
+#include "Serialization/CompactBinaryWriter.h"
 #include "Serialization/FileRegions.h"
 #include "Serialization/MemoryWriter.h"
 #include "TextureCompressorModule.h"
@@ -26,7 +27,7 @@ DEFINE_LOG_CATEGORY_STATIC(LogTextureBuildFunction, Log, All);
 // MUST have a corresponding change to this version. Individual texture formats have a version to
 // change that is specific to the format. A merge conflict affecting the version MUST be resolved
 // by generating a new version.
-static const FGuid TextureDerivedDataVersion(TEXT("392f5367-8112-4305-8cef-ae7897844779"));
+static const FGuid TextureDerivedDataVersion(TEXT("4af5c0e4-46cf-466b-a6f7-7b2b3dec8b10"));
 
 #ifndef CASE_ENUM_TO_TEXT
 #define CASE_ENUM_TO_TEXT(txt) case txt: return TEXT(#txt);
@@ -151,10 +152,9 @@ static FTextureBuildSettings ReadBuildSettingsFromCompactBinary(const FCbObjectV
 	return BuildSettings;
 }
 
-static void ReadOutputSettingsFromCompactBinary(const FCbObjectView& Object, int32& NumInlineMips, FString& MipKeyPrefix)
+static void ReadOutputSettingsFromCompactBinary(const FCbObjectView& Object, int32& NumInlineMips)
 {
 	NumInlineMips = Object["NumInlineMips"].AsInt32();
-	MipKeyPrefix = FUTF8ToTCHAR(Object["MipKeyPrefix"].AsString());
 }
 
 static ERawImageFormat::Type ComputeRawImageFormat(ETextureSourceFormat SourceFormat)
@@ -326,8 +326,7 @@ void FTextureBuildFunction::Build(UE::DerivedData::FBuildContext& Context) const
 	}
 
 	int32 NumInlineMips;
-	FString MipKeyPrefix;
-	ReadOutputSettingsFromCompactBinary(Settings["Output"].AsObjectView(), NumInlineMips, MipKeyPrefix);
+	ReadOutputSettingsFromCompactBinary(Settings["Output"].AsObjectView(), NumInlineMips);
 
 	TArray<FImage> SourceMips;
 	if (!TryReadTextureSourceFromCompactBinary(Settings["Source"], Context, SourceMips))
@@ -365,108 +364,82 @@ void FTextureBuildFunction::Build(UE::DerivedData::FBuildContext& Context) const
 	const bool bForceAllMipsToBeInlined = BuildSettings.bCubemap || (BuildSettings.bVolume && !BuildSettings.bStreamable) || (BuildSettings.bTextureArray && !BuildSettings.bStreamable);
 	const int32 FirstInlineMip = bForceAllMipsToBeInlined ? 0 : FMath::Max(0, MipCount - FMath::Max(NumInlineMips, (int32)NumMipsInTail));
 
-	// Write out texture platform data and non-streaming mip tail as a single payload
-	// This is meant to match the behavior in SerializePlatformData in TextureDerivedData.cpp.
-	// TODO: The texture header/footer information here should not be serialized from the worker but could be left to the editor to fill in.
-	//		 Doing it here causes more code duplication than I would like.
 	{
-		bool bHasOptData = (NumMipsInTail != 0) || (ExtData != 0);
-		TArray<FSharedBuffer> OrderedBuffers;
-
-		TArray<uint8> TextureHeaderArray;
-		FMemoryWriter TextureHeaderWriter(TextureHeaderArray, /*bIsPersistent=*/ true);
-		TextureHeaderWriter << CompressedMips[0].SizeX;
-		TextureHeaderWriter << CompressedMips[0].SizeY;
+		FCbWriter DescriptionWriter;
+		DescriptionWriter.BeginObject();
 
+		DescriptionWriter.BeginArray("Size"_ASV);
+		DescriptionWriter.AddInteger(CompressedMips[0].SizeX);
+		DescriptionWriter.AddInteger(CompressedMips[0].SizeY);
 		const int32 NumSlices = (BuildSettings.bVolume || BuildSettings.bTextureArray) ? CompressedMips[0].SizeZ : BuildSettings.bCubemap ? 6 : 1;
-		static constexpr uint32 BitMask_CubeMap = 1u << 31u;
-		static constexpr uint32 BitMask_HasOptData = 1u << 30u;
-		static constexpr uint32 BitMask_NumSlices = BitMask_HasOptData - 1u;
-		uint32 PackedData = (NumSlices & BitMask_NumSlices) | (BuildSettings.bCubemap ? BitMask_CubeMap : 0) | (bHasOptData ? BitMask_HasOptData : 0);
-		TextureHeaderWriter << PackedData;
-
-		FString PixelFormatString(GetPixelFormatString((EPixelFormat)CompressedMips[0].PixelFormat));
-		TextureHeaderWriter << PixelFormatString;
-
-		if (bHasOptData)
-		{
-			TextureHeaderWriter << ExtData;
-			TextureHeaderWriter << NumMipsInTail;
-		}
-
-		TextureHeaderWriter << MipCount;
-
-		OrderedBuffers.Add(FSharedBuffer::MakeView(TextureHeaderArray.GetData(), TextureHeaderArray.Num()));
-
-		int64 CurrentOffset = TextureHeaderArray.Num();
+		DescriptionWriter.AddInteger(NumSlices);
+		DescriptionWriter.EndArray();
+
+		DescriptionWriter.AddString("PixelFormat"_ASV, GetPixelFormatString((EPixelFormat)CompressedMips[0].PixelFormat));
+		DescriptionWriter.AddBool("bCubeMap"_ASV, BuildSettings.bCubemap);
+		DescriptionWriter.AddInteger("ExtData"_ASV, ExtData);
+		DescriptionWriter.AddInteger("NumMips"_ASV, MipCount);
+		DescriptionWriter.AddInteger("NumStreamingMips"_ASV, FirstInlineMip);
+		DescriptionWriter.AddInteger("NumMipsInTail"_ASV, NumMipsInTail);
+
+		DescriptionWriter.BeginArray("Mips"_ASV);
+		int64 MipPayloadOffset = 0;
 		for (int32 MipIndex = 0; MipIndex < MipCount; ++MipIndex)
 		{
-			bool bIsInlineMip = MipIndex >= FirstInlineMip;
 			FCompressedImage2D& CompressedMip = CompressedMips[MipIndex];
 
-			TArray<uint8> MipHeader;
-			FMemoryWriter MipHeaderWriter(MipHeader, /*bIsPersistent=*/ true);
-			bool bCooked = false;
-			MipHeaderWriter << bCooked;
+			DescriptionWriter.BeginObject();
+			
+			DescriptionWriter.BeginArray("Size"_ASV);
+			DescriptionWriter.AddInteger(CompressedMip.SizeX);
+			DescriptionWriter.AddInteger(CompressedMip.SizeY);
+			DescriptionWriter.AddInteger(CompressedMip.SizeZ);
+			DescriptionWriter.EndArray();
 
-			uint32 BulkDataFlags = 0;
-			MipHeaderWriter << BulkDataFlags;
-
-			int32 BulkDataElementCount = bIsInlineMip ? CompressedMip.RawData.Num() : 0;
-			MipHeaderWriter << BulkDataElementCount;
-
-			int32 BulkDataBytesOnDisk = bIsInlineMip ? CompressedMip.RawData.Num() : 0;
-			MipHeaderWriter << BulkDataBytesOnDisk;
-
-			int64 BulkDataOffset = CurrentOffset + MipHeader.Num() + sizeof(int64);
-			MipHeaderWriter << BulkDataOffset;
-
-			CurrentOffset += MipHeader.Num();
-			OrderedBuffers.Add(MakeSharedBufferFromArray(MoveTemp(MipHeader)));
+			const bool bIsInlineMip = MipIndex >= FirstInlineMip;
 
+			EFileRegionType FileRegion = FFileRegion::SelectType(EPixelFormat(CompressedMip.PixelFormat));
+			DescriptionWriter.AddInteger("FileRegion"_ASV, static_cast<int32>(FileRegion));
 			if (bIsInlineMip)
 			{
-				OrderedBuffers.Add(FSharedBuffer::MakeView(CompressedMip.RawData.GetData(), CompressedMip.RawData.Num()));
-				CurrentOffset += CompressedMip.RawData.Num();
+				DescriptionWriter.AddInteger("PayloadOffset"_ASV, MipPayloadOffset);
 			}
+			DescriptionWriter.AddInteger("NumBytes"_ASV, CompressedMip.RawData.Num());
 
-			TArray<uint8> MipFooter;
-			FMemoryWriter MipFooterWriter(MipFooter, /*bIsPersistent=*/ true);
-
-			MipFooterWriter << CompressedMip.SizeX;
-			MipFooterWriter << CompressedMip.SizeY;
-			MipFooterWriter << CompressedMip.SizeZ;
+			DescriptionWriter.EndObject();
 
-			EFileRegionType FileRegion = FFileRegion::SelectType(EPixelFormat(CompressedMip.PixelFormat));
-			MipFooterWriter << FileRegion;
-
-			bool bExistsInDerivedData = !bIsInlineMip;
-			MipFooterWriter << bExistsInDerivedData;
-
-			CurrentOffset += MipFooter.Num();
-			OrderedBuffers.Add(MakeSharedBufferFromArray(MoveTemp(MipFooter)));
+			if (bIsInlineMip)
+			{
+				MipPayloadOffset += CompressedMip.RawData.Num();
+			}
 		}
+		DescriptionWriter.EndArray();
 
-		// Bool serialized as 32-bit int is the footer
-		uint32 IsVirtual = 0;
-		OrderedBuffers.Add(FSharedBuffer::MakeView(&IsVirtual, sizeof(IsVirtual)));
-
-		FCompositeBuffer CompositeResult(OrderedBuffers);
-		Context.AddPayload(UE::DerivedData::FPayloadId::FromName(TEXT("Texture"_SV)), CompositeResult);
+		DescriptionWriter.EndObject();
+		FCbObject DescriptionObject = DescriptionWriter.Save().AsObject();
+		Context.AddPayload(UE::DerivedData::FPayloadId::FromName("Description"_ASV), DescriptionObject);
 	}
 
-	// Write out streaming mips as payloads.
-	const int32 WritableMipCount = MipCount - ((NumMipsInTail > 0) ? ((int32)NumMipsInTail - 1) : 0);
-	for (int32 MipIndex = 0; MipIndex < WritableMipCount; ++MipIndex)
+	// Streaming mips
+	for (int32 MipIndex = 0; MipIndex < FirstInlineMip; ++MipIndex)
 	{
-		if (MipIndex < FirstInlineMip)
-		{
-			FCompressedImage2D& CompressedMip = CompressedMips[MipIndex];
-			TStringBuilder<32> PayloadName;
-			PayloadName << "Mip" << MipIndex;
+		TAnsiStringBuilder<16> MipName;
+		MipName << "Mip"_ASV << MipIndex;
 
-			FSharedBuffer MipData = MakeSharedBufferFromArray(MoveTemp(CompressedMip.RawData));
-			Context.AddPayload(UE::DerivedData::FPayloadId::FromName(*PayloadName), MipData);
-		}
+		FSharedBuffer MipData = MakeSharedBufferFromArray(MoveTemp(CompressedMips[MipIndex].RawData));
+		Context.AddPayload(UE::DerivedData::FPayloadId::FromName(MipName), MipData);
+	}
+
+	// Mip tail
+	TArray<FSharedBuffer> MipTailComponents;
+	for (int32 MipIndex = FirstInlineMip; MipIndex < MipCount; ++MipIndex)
+	{
+		FSharedBuffer MipData = MakeSharedBufferFromArray(MoveTemp(CompressedMips[MipIndex].RawData));
+		MipTailComponents.Add(MipData);
+	}
+	FCompositeBuffer MipTail(MipTailComponents);
+	if (MipTail.GetSize() > 0)
+	{
+		Context.AddPayload(UE::DerivedData::FPayloadId::FromName("MipTail"_ASV), MipTail);
 	}
 }
diff --git a/Engine/Source/Runtime/Engine/Private/Streaming/Texture2DMipDataProvider_DDC.cpp b/Engine/Source/Runtime/Engine/Private/Streaming/Texture2DMipDataProvider_DDC.cpp
--- a/Engine/Source/Runtime/Engine/Private/Streaming/Texture2DMipDataProvider_DDC.cpp
+++ b/Engine/Source/Runtime/Engine/Private/Streaming/Texture2DMipDataProvider_DDC.cpp
@@ -123,11 +123,8 @@ void FTexture2DMipDataProvider_DDC::Init(const FTextureUpdateContext& Context, c
 	}
 }
 
-bool FTexture2DMipDataProvider_DDC::SerializeMipInfo(const FTextureUpdateContext& Context, FArchive& Ar, int32 MipIndex, const FTextureMipInfo& OutMipInfo)
+bool FTexture2DMipDataProvider_DDC::SerializeMipInfo(const FTextureUpdateContext& Context, FArchive& Ar, int32 MipIndex, int32 MipSize, const FTextureMipInfo& OutMipInfo)
 {
-	int32 MipSize = 0;
-	Ar << MipSize;
-
 	const uint32 DepthOrArraySize = FMath::Max<uint32>(OutMipInfo.ArraySize, OutMipInfo.SizeZ);
 	if (MipSize == OutMipInfo.DataSize)
 	{
@@ -187,7 +184,10 @@ int32 FTexture2DMipDataProvider_DDC::GetMips(
 
 						// The result must be read from a memory reader!
 						FMemoryReader Ar(DerivedMipData, true);
-						if (SerializeMipInfo(Context, Ar, MipIndex, MipInfo))
+						int32 MipSize = 0;
+						Ar << MipSize;
+						check(MipSize == (DerivedMipData.Num() - sizeof(int32)));
+						if (SerializeMipInfo(Context, Ar, MipIndex, MipSize, MipInfo))
 						{
 							bSuccess = true;
 						}
@@ -210,7 +210,7 @@ int32 FTexture2DMipDataProvider_DDC::GetMips(
 				{
 					// The result must be read from a memory reader!
 					FMemoryReaderView Ar(DDCBuffers[MipIndex], true);
-					if (SerializeMipInfo(Context, Ar, MipIndex, MipInfos[MipIndex]))
+					if (SerializeMipInfo(Context, Ar, MipIndex, DDCBuffers[MipIndex].GetSize(), MipInfos[MipIndex]))
 					{
 						bSuccess = true;
 					}
diff --git a/Engine/Source/Runtime/Engine/Private/Streaming/Texture2DMipDataProvider_DDC.h b/Engine/Source/Runtime/Engine/Private/Streaming/Texture2DMipDataProvider_DDC.h
--- a/Engine/Source/Runtime/Engine/Private/Streaming/Texture2DMipDataProvider_DDC.h
+++ b/Engine/Source/Runtime/Engine/Private/Streaming/Texture2DMipDataProvider_DDC.h
@@ -40,7 +40,7 @@ protected:
 	// Helper to route incomplete DDC requests to the FAbandonedDDCHandleManager.
 	void ReleaseDDCResources();
 
-	bool SerializeMipInfo(const FTextureUpdateContext& Context, FArchive& Ar, int32 MipIndex, const FTextureMipInfo& OutMipInfo);
+	bool SerializeMipInfo(const FTextureUpdateContext& Context, FArchive& Ar, int32 MipIndex, int32 MipSize, const FTextureMipInfo& OutMipInfo);
 
 private:
 
diff --git a/Engine/Source/Runtime/Engine/Private/TextureDerivedData.cpp b/Engine/Source/Runtime/Engine/Private/TextureDerivedData.cpp
--- a/Engine/Source/Runtime/Engine/Private/TextureDerivedData.cpp
+++ b/Engine/Source/Runtime/Engine/Private/TextureDerivedData.cpp
@@ -894,7 +894,7 @@ void FTexturePlatformData::CancelCache()
 
 bool FTexturePlatformData::IsUsingNewDerivedData()
 {
-	static const bool bUseNewDerivedData = FParse::Param(FCommandLine::Get(), TEXT("DDC2AsyncTextureBuilds"));
+	static const bool bUseNewDerivedData = FParse::Param(FCommandLine::Get(), TEXT("DDC2AsyncTextureBuilds")) || FParse::Param(FCommandLine::Get(), TEXT("DDC2TextureBuilds"));
 	return bUseNewDerivedData;
 }
 
@@ -1056,6 +1056,7 @@ bool FTexturePlatformData::TryInlineMipData(int32 FirstMipToLoad, UTexture* Text
 			void* MipData = Mip.BulkData.Realloc(int64(MipBuffer.GetSize()));
 			FMemory::Memcpy(MipData, MipBuffer.GetData(), MipBuffer.GetSize());
 			Mip.BulkData.Unlock();
+			Mip.SetPagedToDerivedData(false);
 		}))
 	{
 		return false;
diff --git a/Engine/Source/Runtime/Engine/Private/TextureDerivedDataBuildUtils.cpp b/Engine/Source/Runtime/Engine/Private/TextureDerivedDataBuildUtils.cpp
--- a/Engine/Source/Runtime/Engine/Private/TextureDerivedDataBuildUtils.cpp
+++ b/Engine/Source/Runtime/Engine/Private/TextureDerivedDataBuildUtils.cpp
@@ -158,22 +158,11 @@ static void WriteBuildSettings(FCbWriter& Writer, const FTextureBuildSettings& B
 	Writer.EndObject();
 }
 
-static void WriteOutputSettings(FCbWriter& Writer, int32 NumInlineMips, const FString& KeySuffix)
+static void WriteOutputSettings(FCbWriter& Writer, int32 NumInlineMips)
 {
 	Writer.BeginObject();
 
 	Writer.AddInteger("NumInlineMips", NumInlineMips);
-	
-	FString MipDerivedDataKey;
-	FTexture2DMipMap DummyMip;
-	DummyMip.SizeX = 0;
-	DummyMip.SizeY = 0;
-	GetTextureDerivedMipKey(0, DummyMip, KeySuffix, MipDerivedDataKey);
-	int32 PrefixEndIndex = MipDerivedDataKey.Find(TEXT("_MIP0_"), ESearchCase::CaseSensitive);
-	check(PrefixEndIndex != -1);
-	MipDerivedDataKey.LeftInline(PrefixEndIndex);
-	check(!MipDerivedDataKey.IsEmpty());
-	Writer.AddString("MipKeyPrefix",*MipDerivedDataKey);
 
 	Writer.EndObject();
 }
@@ -237,7 +226,7 @@ bool TryFindTextureBuildFunction(FStringBuilderBase& OutFunctionName, const FTex
 	return UE::DerivedData::GetBuild().GetFunctionRegistry().FindFunctionVersion(FunctionName).IsValid();
 }
 
-FCbObject SaveTextureBuildSettings(const FString& KeySuffix, const UTexture& Texture, const FTextureBuildSettings& BuildSettings, int32 LayerIndex, int32 NumInlineMips)
+FCbObject SaveTextureBuildSettings(const UTexture& Texture, const FTextureBuildSettings& BuildSettings, int32 LayerIndex, int32 NumInlineMips)
 {
 	const ITextureFormat* TextureFormat = nullptr;
 	if (ITextureFormatManagerModule* TFM = GetTextureFormatManager())
@@ -265,7 +254,7 @@ FCbObject SaveTextureBuildSettings(const FString& KeySuffix, const UTexture& Tex
 	WriteBuildSettings(Writer, BuildSettings, TextureFormat);
 
 	Writer.SetName("Output");
-	WriteOutputSettings(Writer, NumInlineMips, KeySuffix);
+	WriteOutputSettings(Writer, NumInlineMips);
 
 	Writer.SetName("Source");
 	WriteSource(Writer, Texture, LayerIndex);
diff --git a/Engine/Source/Runtime/Engine/Private/TextureDerivedDataBuildUtils.h b/Engine/Source/Runtime/Engine/Private/TextureDerivedDataBuildUtils.h
--- a/Engine/Source/Runtime/Engine/Private/TextureDerivedDataBuildUtils.h
+++ b/Engine/Source/Runtime/Engine/Private/TextureDerivedDataBuildUtils.h
@@ -13,6 +13,6 @@ class UTexture;
 struct FTextureBuildSettings;
 
 bool TryFindTextureBuildFunction(FStringBuilderBase& OutFunctionName, const FTextureBuildSettings& BuildSettings);
-FCbObject SaveTextureBuildSettings(const FString& KeySuffix, const UTexture& Texture, const FTextureBuildSettings& BuildSettings, int32 LayerIndex, int32 NumInlineMips);
+FCbObject SaveTextureBuildSettings(const UTexture& Texture, const FTextureBuildSettings& BuildSettings, int32 LayerIndex, int32 NumInlineMips);
 
 #endif // WITH_EDITOR
diff --git a/Engine/Source/Runtime/Engine/Private/TextureDerivedDataTask.cpp b/Engine/Source/Runtime/Engine/Private/TextureDerivedDataTask.cpp
--- a/Engine/Source/Runtime/Engine/Private/TextureDerivedDataTask.cpp
+++ b/Engine/Source/Runtime/Engine/Private/TextureDerivedDataTask.cpp
@@ -388,80 +388,6 @@ private:
 
 } // UE::TextureDerivedData
 
-void FTextureCacheDerivedDataWorker::ConsumeBuildFunctionOutput(const UE::DerivedData::FBuildOutput& BuildOutput, const FString& TexturePath, bool bReplaceExistingDDC)
-{
-	using namespace UE::DerivedData;
-
-	FSharedBuffer PrimaryData;
-	{
-		FPayloadId PrimaryPayloadId = FPayloadId::FromName(TEXT("Texture"));
-		const FPayload& PrimaryPayload = BuildOutput.GetPayload(PrimaryPayloadId);
-		if (PrimaryPayload.IsNull())
-		{
-			UE_LOG(LogTexture, Warning, TEXT("Texture build function missing primary output payload when building %s derived data for %s"), *BuildSettingsPerLayer[0].TextureFormatName.GetPlainNameString(), *Texture.GetPathName());
-			return;
-		}
-		PrimaryData = PrimaryPayload.GetData().Decompress();
-		TArrayView<const uint8> PrimaryDataView((const uint8*)PrimaryData.GetData(), PrimaryData.GetSize());
-
-		FMemoryReaderView Ar(PrimaryDataView, /*bIsPersistent=*/ true);
-		DerivedData->Serialize(Ar, NULL);
-	}
-
-	if (DerivedData->Mips.Num())
-	{
-		const bool bInlineMips = EnumHasAnyFlags(CacheFlags, ETextureCacheFlags::InlineMips);
-
-		for (int32 MipIndex = 0; MipIndex < DerivedData->Mips.Num(); ++MipIndex)
-		{
-			FTexture2DMipMap& Mip = DerivedData->Mips[MipIndex];
-			if (!Mip.IsPagedToDerivedData())
-			{
-				break;
-			}
-
-			FPayloadId MipPayloadId = FPayloadId::FromName(WriteToString<8>(TEXT("Mip"), MipIndex));
-			const FPayload& MipPayload = BuildOutput.GetPayload(MipPayloadId);
-			if (MipPayload.IsNull())
-			{
-				UE_LOG(LogTexture, Warning, TEXT("Texture build function missing Mip%d output payload when building %s derived data for %s"), MipIndex, *BuildSettingsPerLayer[0].TextureFormatName.GetPlainNameString(), *Texture.GetPathName());
-				return;
-			}
-
-			FSharedBuffer MipData = MipPayload.GetData().Decompress();
-
-			{
-				TArray<uint8> MipDataCopy;
-				FMemoryWriter MipWriter(MipDataCopy, /*bIsPersistent=*/ true);
-				int32 MipSize = TCheckValueCast<int32>(MipData.GetSize());
-				MipWriter << MipSize;
-				MipWriter.Serialize(const_cast<void*>(MipData.GetData()), MipSize);
-				GetDerivedDataCacheRef().Put(*DerivedData->GetDerivedDataMipKeyString(MipIndex, Mip), MipDataCopy, TexturePath, bReplaceExistingDDC);
-			}
-
-			if (bInlineMips && (MipIndex >= (int32)BuildSettingsPerLayer[0].LODBiasWithCinematicMips))
-			{
-				const uint64 MipSize = MipData.GetSize();
-
-				Mip.BulkData.Lock(LOCK_READ_WRITE);
-				void* MipAllocData = Mip.BulkData.Realloc(int64(MipSize));
-				MakeMemoryView(MipAllocData, MipSize).CopyFrom(MipData);
-				Mip.BulkData.Unlock();
-				Mip.SetPagedToDerivedData(false);
-			}
-		}
-
-		TArrayView<const uint8> PrimaryDataView((const uint8*)PrimaryData.GetData(), PrimaryData.GetSize());
-		GetDerivedDataCacheRef().Put(*DerivedData->DerivedDataKey.Get<FString>(), PrimaryDataView, TexturePath, bReplaceExistingDDC);
-
-		bSucceeded = true;
-	}
-	else
-	{
-		UE_LOG(LogTexture, Warning, TEXT("Failed to build %s derived data for %s"), *BuildSettingsPerLayer[0].TextureFormatName.GetPlainNameString(), *Texture.GetPathName());
-	}
-}
-
 void FTextureCacheDerivedDataWorker::BuildTexture(bool bReplaceExistingDDC)
 {
 	TRACE_CPUPROFILER_EVENT_SCOPE(FTextureCacheDerivedDataWorker::BuildTexture);
@@ -554,132 +480,82 @@ void FTextureCacheDerivedDataWorker::BuildTexture(bool bReplaceExistingDDC)
 
 		FOptTexturePlatformData OptData;
 
-		if (!BuildFunctionName.IsEmpty())
+		// Compress the texture by calling texture compressor directly.
+		TArray<FCompressedImage2D> CompressedMips;
+		if (Compressor->BuildTexture(TextureData.Blocks[0].MipsPerLayer[0],
+			((bool)Texture.CompositeTexture && CompositeTextureData.Blocks.Num() && CompositeTextureData.Blocks[0].MipsPerLayer.Num()) ? CompositeTextureData.Blocks[0].MipsPerLayer[0] : TArray<FImage>(),
+			BuildSettingsPerLayer[0],
+			CompressedMips,
+			OptData.NumMipsInTail,
+			OptData.ExtData))
 		{
-			using namespace UE::DerivedData;
-			// Compress the texture by calling the DDC2 build function.  May be executed locally or remotely.
-			// The FBuildDefinition in the future will replace both the cache get/put as well as execution.
-			IBuild& Build = GetBuild();
-			FString TexturePath = Texture.GetPathName();
-			FBuildDefinitionBuilder DefinitionBuilder = Build.CreateDefinition(TexturePath, BuildFunctionName);
-
-			DefinitionBuilder.AddConstant(TEXT("Settings"_SV),
-				SaveTextureBuildSettings(KeySuffix, Texture, BuildSettingsPerLayer[0], 0, NUM_INLINE_DERIVED_MIPS));
-			DefinitionBuilder.AddInputBulkData(TEXT("Source"_SV), Texture.Source.GetPersistentId());
-			if (Texture.CompositeTexture)
-			{
-				DefinitionBuilder.AddInputBulkData(TEXT("CompositeSource"_SV), Texture.CompositeTexture->Source.GetPersistentId());
-			}
+			check(CompressedMips.Num());
 
-			TOptional<UE::TextureDerivedData::FTextureBuildInputResolver> PlaceholderResolver;
-			UE::DerivedData::IBuildInputResolver* InputResolver = GetGlobalBuildInputResolver();
-			if (!InputResolver)
+			// Build the derived data.
+			const int32 MipCount = CompressedMips.Num();
+			for (int32 MipIndex = 0; MipIndex < MipCount; ++MipIndex)
 			{
-				PlaceholderResolver.Emplace(Texture);
-				InputResolver = &PlaceholderResolver.GetValue();
-			}
-			FRequestOwner Owner(EPriority::Blocking);
-			FBuildSession Session = Build.CreateSession(TexturePath, InputResolver);
-			Session.Build(DefinitionBuilder.Build(), EBuildPolicy::Default, Owner,
-				[this, &TexturePath, bReplaceExistingDDC] (FBuildCompleteParams&& Params)
+				const FCompressedImage2D& CompressedImage = CompressedMips[MipIndex];
+				FTexture2DMipMap* NewMip = new FTexture2DMipMap();
+				DerivedData->Mips.Add(NewMip);
+				NewMip->SizeX = CompressedImage.SizeX;
+				NewMip->SizeY = CompressedImage.SizeY;
+				NewMip->SizeZ = CompressedImage.SizeZ;
+				NewMip->FileRegionType = FFileRegion::SelectType(EPixelFormat(CompressedImage.PixelFormat));
+				check(NewMip->SizeZ == 1 || BuildSettingsPerLayer[0].bVolume || BuildSettingsPerLayer[0].bTextureArray); // Only volume & arrays can have SizeZ != 1
+				NewMip->BulkData.Lock(LOCK_READ_WRITE);
+				check(CompressedImage.RawData.GetTypeSize() == 1);
+				void* NewMipData = NewMip->BulkData.Realloc(CompressedImage.RawData.Num());
+				FMemory::Memcpy(NewMipData, CompressedImage.RawData.GetData(), CompressedImage.RawData.Num());
+				NewMip->BulkData.Unlock();
+
+				if (MipIndex == 0)
 				{
-					#if !NO_LOGGING
-					Params.Output.IterateDiagnostics([this, &TexturePath] (const FBuildDiagnostic& Diagnostic)
-						{
-							if (GWarn)
-							{
-								FName CategoryName(Diagnostic.Category);
-								GWarn->Log(CategoryName, Diagnostic.Level == EBuildDiagnosticLevel::Error ? ELogVerbosity::Error : ELogVerbosity::Warning, *FString(Diagnostic.Message));
-							}
-						});
-					#endif
-
-					if (Params.Status == EStatus::Ok)
+					DerivedData->SizeX = CompressedImage.SizeX;
+					DerivedData->SizeY = CompressedImage.SizeY;
+					DerivedData->PixelFormat = (EPixelFormat)CompressedImage.PixelFormat;
+					if (BuildSettingsPerLayer[0].bVolume || BuildSettingsPerLayer[0].bTextureArray)
 					{
-						ConsumeBuildFunctionOutput(Params.Output, TexturePath, bReplaceExistingDDC);
+						DerivedData->SetNumSlices(CompressedImage.SizeZ);
 					}
-				});
-			Owner.Wait();
-		}
-		else
-		{
-			// Compress the texture by calling texture compressor directly.
-			TArray<FCompressedImage2D> CompressedMips;
-			if (Compressor->BuildTexture(TextureData.Blocks[0].MipsPerLayer[0],
-				((bool)Texture.CompositeTexture && CompositeTextureData.Blocks.Num() && CompositeTextureData.Blocks[0].MipsPerLayer.Num()) ? CompositeTextureData.Blocks[0].MipsPerLayer[0] : TArray<FImage>(),
-				BuildSettingsPerLayer[0],
-				CompressedMips,
-				OptData.NumMipsInTail,
-				OptData.ExtData))
-			{
-				check(CompressedMips.Num());
-
-				// Build the derived data.
-				const int32 MipCount = CompressedMips.Num();
-				for (int32 MipIndex = 0; MipIndex < MipCount; ++MipIndex)
-				{
-					const FCompressedImage2D& CompressedImage = CompressedMips[MipIndex];
-					FTexture2DMipMap* NewMip = new FTexture2DMipMap();
-					DerivedData->Mips.Add(NewMip);
-					NewMip->SizeX = CompressedImage.SizeX;
-					NewMip->SizeY = CompressedImage.SizeY;
-					NewMip->SizeZ = CompressedImage.SizeZ;
-					NewMip->FileRegionType = FFileRegion::SelectType(EPixelFormat(CompressedImage.PixelFormat));
-					check(NewMip->SizeZ == 1 || BuildSettingsPerLayer[0].bVolume || BuildSettingsPerLayer[0].bTextureArray); // Only volume & arrays can have SizeZ != 1
-					NewMip->BulkData.Lock(LOCK_READ_WRITE);
-					check(CompressedImage.RawData.GetTypeSize() == 1);
-					void* NewMipData = NewMip->BulkData.Realloc(CompressedImage.RawData.Num());
-					FMemory::Memcpy(NewMipData, CompressedImage.RawData.GetData(), CompressedImage.RawData.Num());
-					NewMip->BulkData.Unlock();
-
-					if (MipIndex == 0)
+					else if (BuildSettingsPerLayer[0].bCubemap)
 					{
-						DerivedData->SizeX = CompressedImage.SizeX;
-						DerivedData->SizeY = CompressedImage.SizeY;
-						DerivedData->PixelFormat = (EPixelFormat)CompressedImage.PixelFormat;
-						if (BuildSettingsPerLayer[0].bVolume || BuildSettingsPerLayer[0].bTextureArray)
-						{
-							DerivedData->SetNumSlices(CompressedImage.SizeZ);
-						}
-						else if (BuildSettingsPerLayer[0].bCubemap)
-						{
-							DerivedData->SetNumSlices(6);
-						}
-						else
-						{
-							DerivedData->SetNumSlices(1);
-						}
-						DerivedData->SetIsCubemap(BuildSettingsPerLayer[0].bCubemap);
+						DerivedData->SetNumSlices(6);
 					}
 					else
 					{
-						check(CompressedImage.PixelFormat == DerivedData->PixelFormat);
+						DerivedData->SetNumSlices(1);
 					}
+					DerivedData->SetIsCubemap(BuildSettingsPerLayer[0].bCubemap);
 				}
-
-				DerivedData->SetOptData(OptData);
-				
-				// Store it in the cache.
-				// @todo: This will remove the streaming bulk data, which we immediately reload below!
-				// Should ideally avoid this redundant work, but it only happens when we actually have 
-				// to build the texture, which should only ever be once.
-				this->BytesCached = PutDerivedDataInCache(DerivedData, KeySuffix, Texture.GetPathName(), BuildSettingsPerLayer[0].bCubemap || (BuildSettingsPerLayer[0].bVolume && !GSupportsVolumeTextureStreaming) || (BuildSettingsPerLayer[0].bTextureArray && !GSupportsTexture2DArrayStreaming), bReplaceExistingDDC);
-			}
-
-			if (DerivedData->Mips.Num())
-			{
-				const bool bInlineMips = EnumHasAnyFlags(CacheFlags, ETextureCacheFlags::InlineMips);
-				bSucceeded = !bInlineMips || DerivedData->TryInlineMipData(BuildSettingsPerLayer[0].LODBiasWithCinematicMips, &Texture);
-				if (!bSucceeded)
+				else
 				{
-					UE_LOG(LogTexture, Display, TEXT("Failed to put and then read back mipmap data from DDC for %s"), *Texture.GetPathName());
+					check(CompressedImage.PixelFormat == DerivedData->PixelFormat);
 				}
 			}
-			else
+
+			DerivedData->SetOptData(OptData);
+				
+			// Store it in the cache.
+			// @todo: This will remove the streaming bulk data, which we immediately reload below!
+			// Should ideally avoid this redundant work, but it only happens when we actually have 
+			// to build the texture, which should only ever be once.
+			this->BytesCached = PutDerivedDataInCache(DerivedData, KeySuffix, Texture.GetPathName(), BuildSettingsPerLayer[0].bCubemap || (BuildSettingsPerLayer[0].bVolume && !GSupportsVolumeTextureStreaming) || (BuildSettingsPerLayer[0].bTextureArray && !GSupportsTexture2DArrayStreaming), bReplaceExistingDDC);
+		}
+
+		if (DerivedData->Mips.Num())
+		{
+			const bool bInlineMips = EnumHasAnyFlags(CacheFlags, ETextureCacheFlags::InlineMips);
+			bSucceeded = !bInlineMips || DerivedData->TryInlineMipData(BuildSettingsPerLayer[0].LODBiasWithCinematicMips, &Texture);
+			if (!bSucceeded)
 			{
-				UE_LOG(LogTexture, Warning, TEXT("Failed to build %s derived data for %s"), *BuildSettingsPerLayer[0].TextureFormatName.GetPlainNameString(), *Texture.GetPathName());
+				UE_LOG(LogTexture, Display, TEXT("Failed to put and then read back mipmap data from DDC for %s"), *Texture.GetPathName());
 			}
 		}
+		else
+		{
+			UE_LOG(LogTexture, Warning, TEXT("Failed to build %s derived data for %s"), *BuildSettingsPerLayer[0].TextureFormatName.GetPlainNameString(), *Texture.GetPathName());
+		}
 	}
 }
 
@@ -727,16 +603,6 @@ FTextureCacheDerivedDataWorker::FTextureCacheDerivedDataWorker(
 	const bool bAllowAsyncLoading = EnumHasAnyFlags(CacheFlags, ETextureCacheFlags::AllowAsyncLoading);
 	const bool bForVirtualTextureStreamingBuild = EnumHasAnyFlags(CacheFlags, ETextureCacheFlags::ForVirtualTextureStreamingBuild);
 
-	static const bool bBuildFunctionEnabled = FParse::Param(FCommandLine::Get(), TEXT("DDC2TextureBuilds"));
-	if (bBuildFunctionEnabled && !bForVirtualTextureStreamingBuild && (BuildSettingsPerLayer.Num() == 1))
-	{
-		TStringBuilder<64> FunctionName;
-		if (TryFindTextureBuildFunction(FunctionName, BuildSettingsPerLayer[0]))
-		{
-			BuildFunctionName = FunctionName;
-		}
-	}
-
 	// FVirtualTextureDataBuilder always wants to load ImageWrapper module
 	// This is not strictly necessary, used only for debug output, but seems simpler to just always load this here, doesn't seem like it should be too expensive
 	if (bAllowAsyncLoading || bForVirtualTextureStreamingBuild)
@@ -1070,9 +936,6 @@ public:
 
 		static bool bLoadedModules = LoadModules();
 
-		FString KeySuffix;
-		GetTextureDerivedDataKeySuffix(Texture, &Settings, KeySuffix);
-
 		TStringBuilder<256> TexturePath;
 		Texture.GetPathName(nullptr, TexturePath);
 
@@ -1094,13 +957,13 @@ public:
 			StatusMessage.Emplace(ComposeTextureBuildText(Texture, Settings, Texture.GetBuildRequiredMemory(), EnumHasAnyFlags(Flags, ETextureCacheFlags::ForVirtualTextureStreamingBuild)));
 		}
 
-		FBuildDefinition Definition = CreateDefinition(Build, Texture, TexturePath, FunctionName, KeySuffix, Settings, bUseCompositeTexture);
+		FBuildDefinition Definition = CreateDefinition(Build, Texture, TexturePath, FunctionName, Settings, bUseCompositeTexture);
 		
 		if (!EnumHasAnyFlags(Flags, ETextureCacheFlags::ForceRebuild) && Settings.FastTextureEncode == ETextureFastEncode::TryOffEncodeFast)
 		{
 			FTextureBuildSettings ShippingSettings = Settings;
 			ShippingSettings.FastTextureEncode = ETextureFastEncode::Off;
-			FBuildDefinition ShippingDefinition = CreateDefinition(Build, Texture, TexturePath, FunctionName, KeySuffix, ShippingSettings, bUseCompositeTexture);
+			FBuildDefinition ShippingDefinition = CreateDefinition(Build, Texture, TexturePath, FunctionName, ShippingSettings, bUseCompositeTexture);
 			BuildSession.Get().Build(ShippingDefinition, EBuildPolicy::Cache, *Owner,
 				[this, Definition = MoveTemp(Definition), Flags](FBuildCompleteParams&& Params)
 				{
@@ -1130,13 +993,12 @@ public:
 		UTexture& Texture,
 		FStringView TexturePath,
 		FStringView FunctionName,
-		const FString& KeySuffix,
 		const FTextureBuildSettings& Settings,
 		const bool bUseCompositeTexture)
 	{
 		UE::DerivedData::FBuildDefinitionBuilder DefinitionBuilder = Build.CreateDefinition(TexturePath, FunctionName);
 		DefinitionBuilder.AddConstant(TEXT("Settings"_SV),
-			SaveTextureBuildSettings(KeySuffix, Texture, Settings, 0, NUM_INLINE_DERIVED_MIPS));
+			SaveTextureBuildSettings(Texture, Settings, 0, NUM_INLINE_DERIVED_MIPS));
 		DefinitionBuilder.AddInputBulkData(TEXT("Source"_SV), Texture.Source.GetPersistentId());
 		if (Texture.CompositeTexture && bUseCompositeTexture)
 		{
@@ -1335,6 +1197,119 @@ public:
 	}
 
 private:
+
+	static bool DeserializeTextureFromPayloads(FTexturePlatformData& DerivedData, const UE::DerivedData::FBuildOutput& Output, int32 FirstMipToLoad, bool bInlineMips)
+	{
+		using namespace UE::DerivedData;
+		const FPayload& Payload = Output.GetPayload(FPayloadId::FromName("Description"_ASV));
+		if (!Payload)
+		{
+			UE_LOG(LogTexture, Error, TEXT("Missing texture description for build of '%s' by %s."),
+				*WriteToString<128>(Output.GetName()), *WriteToString<32>(Output.GetFunction()));
+			return false;
+		}
+
+		FCbObject TextureDescription(Payload.GetData().Decompress());
+
+		FCbFieldViewIterator SizeIt = TextureDescription["Size"_ASV].AsArrayView().CreateViewIterator();
+		DerivedData.SizeX = SizeIt++->AsInt32();
+		DerivedData.SizeY = SizeIt++->AsInt32();
+		int32 NumSlices = SizeIt++->AsInt32();
+
+		UEnum* PixelFormatEnum = UTexture::GetPixelFormatEnum();
+		FUtf8StringView PixelFormatStringView = TextureDescription["PixelFormat"_ASV].AsString();
+		FName PixelFormatName(PixelFormatStringView.Len(), PixelFormatStringView.GetData());
+		DerivedData.PixelFormat = (EPixelFormat)PixelFormatEnum->GetValueByName(PixelFormatName);
+
+		const bool bCubeMap = TextureDescription["bCubeMap"_ASV].AsBool();
+		DerivedData.OptData.ExtData = TextureDescription["ExtData"_ASV].AsUInt32();
+		DerivedData.OptData.NumMipsInTail = TextureDescription["NumMipsInTail"_ASV].AsUInt32();
+		const bool bHasOptData = (DerivedData.OptData.NumMipsInTail != 0) || (DerivedData.OptData.ExtData != 0);
+		static constexpr uint32 BitMask_CubeMap = 1u << 31u;
+		static constexpr uint32 BitMask_HasOptData = 1u << 30u;
+		static constexpr uint32 BitMask_NumSlices = BitMask_HasOptData - 1u;
+		DerivedData.PackedData = (NumSlices & BitMask_NumSlices) | (bCubeMap ? BitMask_CubeMap : 0) | (bHasOptData ? BitMask_HasOptData : 0);
+
+		int32 NumMips = TextureDescription["NumMips"_ASV].AsInt32();
+		int32 NumStreamingMips = TextureDescription["NumStreamingMips"_ASV].AsInt32();
+
+		FCbArrayView MipArrayView = TextureDescription["Mips"_ASV].AsArrayView();
+		if (NumMips != MipArrayView.Num())
+		{
+			UE_LOG(LogTexture, Error, TEXT("Mismatched mip quantity (%d and %d) for build of '%s' by %s."),
+				NumMips, MipArrayView.Num(), *WriteToString<128>(Output.GetName()), *WriteToString<32>(Output.GetFunction()));
+			return false;
+		}
+		check(NumMips >= (int32)DerivedData.OptData.NumMipsInTail);
+		check(NumMips >= NumStreamingMips);
+
+		FSharedBuffer MipTailData;
+		if (NumMips > NumStreamingMips)
+		{
+			const FPayload& MipTailPayload = Output.GetPayload(FPayloadId::FromName("MipTail"_ASV));
+			if (!MipTailPayload)
+			{
+				UE_LOG(LogTexture, Error, TEXT("Missing texture mip tail for build of '%s' by %s."),
+					*WriteToString<128>(Output.GetName()), *WriteToString<32>(Output.GetFunction()));
+				return false;
+			}
+			MipTailData = MipTailPayload.GetData().Decompress();
+		}
+
+		int32 MipIndex = 0;
+		DerivedData.Mips.Empty(NumMips);
+		for (FCbFieldView MipFieldView : MipArrayView)
+		{
+			FCbObjectView MipObjectView = MipFieldView.AsObjectView();
+			FTexture2DMipMap* NewMip = new FTexture2DMipMap();
+
+			FCbFieldViewIterator MipSizeIt = MipObjectView["Size"_ASV].AsArrayView().CreateViewIterator();
+			NewMip->SizeX = MipSizeIt++->AsInt32();
+			NewMip->SizeY = MipSizeIt++->AsInt32();
+			NewMip->SizeZ = MipSizeIt++->AsInt32();
+			NewMip->FileRegionType = static_cast<EFileRegionType>(MipObjectView["FileRegion"_ASV].AsInt32());
+			
+			if (MipIndex >= NumStreamingMips)
+			{
+				uint64 MipSize = MipObjectView["NumBytes"_ASV].AsUInt64();
+				FMemoryView MipView = MipTailData.GetView().Mid(MipObjectView["PayloadOffset"_ASV].AsUInt64(), MipSize);
+
+				NewMip->BulkData.Lock(LOCK_READ_WRITE);
+				void* MipAllocData = NewMip->BulkData.Realloc(int64(MipSize));
+				MakeMemoryView(MipAllocData, MipSize).CopyFrom(MipView);
+				NewMip->BulkData.Unlock();
+				NewMip->SetPagedToDerivedData(false);
+			}
+			else if (bInlineMips && (MipIndex >= FirstMipToLoad))
+			{
+				const FPayload& StreamingMipPayload = Output.GetPayload(FPayloadId::FromName(WriteToString<8>(TEXT("Mip"), MipIndex)));
+				if (!StreamingMipPayload)
+				{
+					UE_LOG(LogTexture, Error, TEXT("Missing texture streaming mip '%s' for build of '%s' by %s."),
+						*WriteToString<8>(TEXT("Mip"), MipIndex), *WriteToString<128>(Output.GetName()), *WriteToString<32>(Output.GetFunction()));
+					return false;
+				}
+				FSharedBuffer StreamingMipData = StreamingMipPayload.GetData().Decompress();
+				uint64 MipSize = StreamingMipData.GetSize();
+
+				NewMip->BulkData.Lock(LOCK_READ_WRITE);
+				void* MipAllocData = NewMip->BulkData.Realloc(int64(MipSize));
+				MakeMemoryView(MipAllocData, MipSize).CopyFrom(StreamingMipData.GetView());
+				NewMip->BulkData.Unlock();
+				NewMip->SetPagedToDerivedData(false);
+			}
+			else
+			{
+				NewMip->SetPagedToDerivedData(true);
+			}
+
+			DerivedData.Mips.Add(NewMip);
+			++MipIndex;
+		}
+
+		return true;
+	}
+
 	void WriteDerivedData(UE::DerivedData::FBuildOutput&& Output)
 	{
 		using namespace UE::DerivedData;
@@ -1362,48 +1337,7 @@ private:
 			return;
 		}
 
-		if (const FPayload& Payload = Output.GetPayload(FPayloadId::FromName(TEXT("Texture"))))
-		{
-			FSharedBuffer Data = Payload.GetData().Decompress();
-			FMemoryReaderView Ar(Data, /*bIsPersistent*/ true);
-			DerivedData.Serialize(Ar, nullptr);
-		}
-		else
-		{
-			UE_LOG(LogTexture, Warning, TEXT("Missing output for build of '%s' by %s."),
-				*WriteToString<128>(Output.GetName()), *WriteToString<32>(Output.GetFunction()));
-			return;
-		}
-
-		for (int32 MipIndex = 0, MipCount = DerivedData.Mips.Num(); MipIndex < MipCount; ++MipIndex)
-		{
-			FTexture2DMipMap& Mip = DerivedData.Mips[MipIndex];
-			if (!Mip.IsPagedToDerivedData())
-			{
-				break;
-			}
-
-			if (const FPayload& MipPayload = Output.GetPayload(FPayloadId::FromName(WriteToString<8>(TEXT("Mip"), MipIndex))))
-			{
-				if (bInlineMips && MipIndex >= FirstMipToLoad)
-				{
-					FSharedBuffer MipData = MipPayload.GetData().Decompress();
-					const uint64 MipSize = MipData.GetSize();
-
-					Mip.BulkData.Lock(LOCK_READ_WRITE);
-					void* MipAllocData = Mip.BulkData.Realloc(int64(MipSize));
-					MakeMemoryView(MipAllocData, MipSize).CopyFrom(MipData);
-					Mip.BulkData.Unlock();
-					Mip.SetPagedToDerivedData(false);
-				}
-			}
-			else
-			{
-				UE_LOG(LogTexture, Warning, TEXT("Missing output for mip %d for build of '%s' by %s."),
-					MipIndex, *WriteToString<128>(Output.GetName()), *WriteToString<32>(Output.GetFunction()));
-				return;
-			}
-		}
+		DeserializeTextureFromPayloads(DerivedData, Output, FirstMipToLoad, bInlineMips);
 	}
 
 	static UE::DerivedData::EPriority ConvertPriority(EQueuedWorkPriority SourcePriority)
@@ -1485,9 +1419,6 @@ UE::DerivedData::FCacheKeyProxy CreateTextureCacheKeyProxy(
 		TStringBuilder<256> TexturePath;
 		Texture.GetPathName(nullptr, TexturePath);
 
-		FString KeySuffix;
-		GetTextureDerivedDataKeySuffix(Texture, &Settings, KeySuffix);
-
 		bool bUseCompositeTexture = false;
 		if (FTextureBuildTask::IsTextureValidForBuilding(Texture, CacheFlags, bUseCompositeTexture))
 		{
@@ -1499,7 +1430,7 @@ UE::DerivedData::FCacheKeyProxy CreateTextureCacheKeyProxy(
 				InputResolver = &PlaceholderResolver.GetValue();
 			}
 			FBuildSession Session = Build.CreateSession(TexturePath, InputResolver);
-			FBuildDefinition Definition = FTextureBuildTask::CreateDefinition(Build, Texture, TexturePath, FunctionName, KeySuffix, Settings, bUseCompositeTexture);
+			FBuildDefinition Definition = FTextureBuildTask::CreateDefinition(Build, Texture, TexturePath, FunctionName, Settings, bUseCompositeTexture);
 			
 			FCacheKey GeneratedKey;
 			FRequestOwner CacheKeyGenerationOwner(EPriority::Blocking);
diff --git a/Engine/Source/Runtime/Engine/Private/TextureDerivedDataTask.h b/Engine/Source/Runtime/Engine/Private/TextureDerivedDataTask.h
--- a/Engine/Source/Runtime/Engine/Private/TextureDerivedDataTask.h
+++ b/Engine/Source/Runtime/Engine/Private/TextureDerivedDataTask.h
@@ -124,8 +124,6 @@ class FTextureCacheDerivedDataWorker : public FNonAbandonableTask
 	FTextureSourceData TextureData;
 	/** Source mip images of the composite texture (e.g. normal map for compute roughness). Not necessarily in RGBA32F, usually only top mip as other mips need to be generated first */
 	FTextureSourceData CompositeTextureData;
-	/** DDC2 build function name to use to build this texture (if DDC2 is enabled and the target texture type has a DDC2 build function, empty otherwise) */
-	FString BuildFunctionName;
 	/** Texture cache flags. */
 	ETextureCacheFlags CacheFlags;
 	/** Have many bytes were loaded from DDC or built (for telemetry) */
@@ -140,8 +138,6 @@ class FTextureCacheDerivedDataWorker : public FNonAbandonableTask
 
 	/** Build the texture. This function is safe to call from any thread. */
 	void BuildTexture(bool bReplaceExistingDDC = false);
-
-	void ConsumeBuildFunctionOutput(const UE::DerivedData::FBuildOutput& BuildOutput, const FString& TexturePath, bool bReplaceExistingDDC);
 public:
 
 	/** Initialization constructor. */
