Fix a bug where comparing actor handles can return the wrong result if one or both is represented by an actor.

diff --git a/Engine/Source/Runtime/Engine/Private/EngineTypes.cpp b/Engine/Source/Runtime/Engine/Private/EngineTypes.cpp
--- a/Engine/Source/Runtime/Engine/Private/EngineTypes.cpp
+++ b/Engine/Source/Runtime/Engine/Private/EngineTypes.cpp
@@ -479,7 +479,17 @@ FActorInstanceHandle& FActorInstanceHandle::operator=(AActor* OtherActor)
 
 bool FActorInstanceHandle::operator==(const FActorInstanceHandle& Other) const
 {
-	return Manager == Other.Manager && InstanceIndex == Other.InstanceIndex;
+	// try to compare managers and indices first if we have them
+	if (Manager.IsValid() && Other.Manager.IsValid() && InstanceIndex != INDEX_NONE && Other.InstanceIndex != INDEX_NONE)
+	{
+		return Manager == Other.Manager && InstanceIndex == Other.InstanceIndex;
+	}
+
+	// try to compare the actors
+	const AActor* MyActor = FetchActor();
+	const AActor* OtherActor = Other.FetchActor();
+
+	return MyActor == OtherActor;
 }
 
 bool FActorInstanceHandle::operator!=(const FActorInstanceHandle& Other) const
