Remove SerializeWithConditionalBulkData which is unused, incorrect, and to be made redundant by the upcoming introduction of FDerivedData.

diff --git a/Engine/Source/Runtime/Engine/Classes/Engine/Texture.h b/Engine/Source/Runtime/Engine/Classes/Engine/Texture.h
--- a/Engine/Source/Runtime/Engine/Classes/Engine/Texture.h
+++ b/Engine/Source/Runtime/Engine/Classes/Engine/Texture.h
@@ -674,8 +674,6 @@ public:
 	void Serialize(FArchive& Ar, class UTexture* Owner);
 
 #if WITH_EDITORONLY_DATA
-	void SerializeWithConditionalBulkData(FArchive& Ar, class UTexture* Owner);
-
 	FString GetDerivedDataMipKeyString(int32 MipIndex, const FTexture2DMipMap& Mip) const;
 	UE::DerivedData::FCachePayloadKeyProxy GetDerivedDataMipKeyProxy(int32 MipIndex, const FTexture2DMipMap& Mip) const;
 #endif // WITH_EDITORONLY_DATA
diff --git a/Engine/Source/Runtime/Engine/Private/Texture2D.cpp b/Engine/Source/Runtime/Engine/Private/Texture2D.cpp
--- a/Engine/Source/Runtime/Engine/Private/Texture2D.cpp
+++ b/Engine/Source/Runtime/Engine/Private/Texture2D.cpp
@@ -138,21 +138,13 @@ extern bool TrackTextureEvent( FStreamingRenderAsset* StreamingTexture, UStreama
 	FTexture2DMipMap
 -----------------------------------------------------------------------------*/
 
-void FTexture2DMipMap::SerializeCommon(FArchive& Ar, UObject* Owner, int32 MipIdx, bool bIgnoreBulkDataForStreamingMips)
+void FTexture2DMipMap::Serialize(FArchive& Ar, UObject* Owner, int32 MipIdx)
 {
 	bool bCooked = Ar.IsCooking();
 	Ar << bCooked;
 
 #if WITH_EDITORONLY_DATA
-	if (bIgnoreBulkDataForStreamingMips && IsPagedToDerivedData())
-	{
-		FByteBulkData EmptyData;
-		EmptyData.Serialize(Ar, Owner, MipIdx, false, FileRegionType);
-	}
-	else
-	{
-		BulkData.Serialize(Ar, Owner, MipIdx, false, FileRegionType);
-	}
+	BulkData.Serialize(Ar, Owner, MipIdx, false, FileRegionType);
 #else
 	BulkData.Serialize(Ar, Owner, MipIdx, false);
 #endif
@@ -170,17 +162,7 @@ void FTexture2DMipMap::SerializeCommon(FArchive& Ar, UObject* Owner, int32 MipId
 #endif // #if WITH_EDITORONLY_DATA
 }
 
-void FTexture2DMipMap::Serialize(FArchive& Ar, UObject* Owner, int32 MipIdx)
-{
-	SerializeCommon(Ar, Owner, MipIdx, false);
-}
-
 #if WITH_EDITORONLY_DATA
-void FTexture2DMipMap::SerializeWithConditionalBulkData(FArchive& Ar, UObject* Owner, int32 MipIdx)
-{
-	SerializeCommon(Ar, Owner, MipIdx, true);
-}
-
 uint32 FTexture2DMipMap::StoreInDerivedDataCache(const FString& InDerivedDataKey, const FStringView& TextureName, bool bReplaceExistingDDC)
 {
 	int32 BulkDataSizeInBytes = BulkData.GetBulkDataSize();
diff --git a/Engine/Source/Runtime/Engine/Private/TextureDerivedData.cpp b/Engine/Source/Runtime/Engine/Private/TextureDerivedData.cpp
--- a/Engine/Source/Runtime/Engine/Private/TextureDerivedData.cpp
+++ b/Engine/Source/Runtime/Engine/Private/TextureDerivedData.cpp
@@ -1535,10 +1535,6 @@ enum class EPlatformDataSerializationFlags : uint8
 	None = 0,
 	Cooked = 1<<0,
 	Streamable = 1<<1,
-
-#if WITH_EDITORONLY_DATA
-	IgnoreBulkDataForStreamingMips = 1<<2, // These are transient, so this can be changed if more runtime flags are needed above
-#endif // WITH_EDITORONLY_DATA
 };
 ENUM_CLASS_FLAGS(EPlatformDataSerializationFlags);
 
@@ -1706,21 +1702,9 @@ static void SerializePlatformData(
 		}
 	}
 
-#if WITH_EDITORONLY_DATA
-	if ((Flags & EPlatformDataSerializationFlags::IgnoreBulkDataForStreamingMips) == EPlatformDataSerializationFlags::IgnoreBulkDataForStreamingMips)
+	for (int32 MipIndex = 0; MipIndex < NumMips; ++MipIndex)
 	{
-		for (int32 MipIndex = 0; MipIndex < NumMips; ++MipIndex)
-		{
-			PlatformData->Mips[FirstMipToSerialize + MipIndex].SerializeWithConditionalBulkData(Ar, Texture, MipIndex);
-		}
-	}
-	else
-#endif // WITH_EDITORONLY_DATA
-	{
-		for (int32 MipIndex = 0; MipIndex < NumMips; ++MipIndex)
-		{
-			PlatformData->Mips[FirstMipToSerialize + MipIndex].Serialize(Ar, Texture, MipIndex);
-		}
+		PlatformData->Mips[FirstMipToSerialize + MipIndex].Serialize(Ar, Texture, MipIndex);
 	}
 
 	Ar << bIsVirtual;
@@ -1762,11 +1746,6 @@ void FTexturePlatformData::Serialize(FArchive& Ar, UTexture* Owner)
 }
 
 #if WITH_EDITORONLY_DATA
-void FTexturePlatformData::SerializeWithConditionalBulkData(FArchive& Ar, UTexture* Owner)
-{
-	SerializePlatformData(Ar, this, Owner, EPlatformDataSerializationFlags::IgnoreBulkDataForStreamingMips);
-}
-
 FString FTexturePlatformData::GetDerivedDataMipKeyString(int32 MipIndex, const FTexture2DMipMap& Mip) const
 {
 	const FString& KeyString = DerivedDataKey.Get<FString>();
diff --git a/Engine/Source/Runtime/Engine/Public/TextureResource.h b/Engine/Source/Runtime/Engine/Public/TextureResource.h
--- a/Engine/Source/Runtime/Engine/Public/TextureResource.h
+++ b/Engine/Source/Runtime/Engine/Public/TextureResource.h
@@ -62,8 +62,6 @@ struct FTexture2DMipMap
 	ENGINE_API void Serialize(FArchive& Ar, UObject* Owner, int32 MipIndex);
 
 #if WITH_EDITORONLY_DATA
-	void SerializeWithConditionalBulkData(FArchive& Ar, UObject* Owner, int32 MipIndex);
-
 	/** The file region type appropriate for this mip's pixel format. */
 	EFileRegionType FileRegionType = EFileRegionType::None;
 
@@ -80,8 +78,6 @@ struct FTexture2DMipMap
 	uint32 StoreInDerivedDataCache(const FString& InDerivedDataKey, const FStringView& TextureName, bool bReplaceExistingDDC);
 
 #endif // #if WITH_EDITORONLY_DATA
-private:
-	void SerializeCommon(FArchive& Ar, UObject* Owner, int32 MipIndex, bool bBulkDataConditionalOnKey);
 };
 
 /** 
