Don't crash if a texture reference is null.

There might be valid reasons for this (material being destroyed), so all resource types should behave the same way as 2D textures - use a black default instead of crashing with an assertion failure.

#rnx

diff --git a/Engine/Source/Runtime/Engine/Private/Materials/MaterialUniformExpressions.cpp b/Engine/Source/Runtime/Engine/Private/Materials/MaterialUniformExpressions.cpp
--- a/Engine/Source/Runtime/Engine/Private/Materials/MaterialUniformExpressions.cpp
+++ b/Engine/Source/Runtime/Engine/Private/Materials/MaterialUniformExpressions.cpp
@@ -893,9 +893,8 @@ void FUniformExpressionSet::FillUniformBuffer(const FMaterialRenderContext& Mate
 			BufferCursor = ((uint8*)BufferCursor) + (SHADER_PARAMETER_POINTER_ALIGNMENT * 2);
 			check(BufferCursor <= TempBuffer + TempBufferSize);
 
-			if(Value && Value->GetResource() && (Value->GetMaterialType() & MCT_TextureCube) != 0u)
+			if(Value && Value->GetResource() && Value->TextureReference.TextureReferenceRHI && (Value->GetMaterialType() & MCT_TextureCube) != 0u)
 			{
-				check(Value->TextureReference.TextureReferenceRHI);
 				*ResourceTableTexturePtr = Value->TextureReference.TextureReferenceRHI;
 				const FSamplerStateRHIRef* SamplerSource = &Value->GetResource()->SamplerStateRHI;
 
@@ -933,9 +932,8 @@ void FUniformExpressionSet::FillUniformBuffer(const FMaterialRenderContext& Mate
 			void** ResourceTableSamplerPtr = (void**)((uint8*)BufferCursor + 1 * SHADER_PARAMETER_POINTER_ALIGNMENT);
 			BufferCursor = ((uint8*)BufferCursor) + (SHADER_PARAMETER_POINTER_ALIGNMENT * 2);
 
-			if (Value && Value->GetResource() && (Value->GetMaterialType() & MCT_Texture2DArray) != 0u)
+			if (Value && Value->GetResource() && Value->TextureReference.TextureReferenceRHI && (Value->GetMaterialType() & MCT_Texture2DArray) != 0u)
 			{
-				check(Value->TextureReference.TextureReferenceRHI);
 				*ResourceTableTexturePtr = Value->TextureReference.TextureReferenceRHI;
 				const FSamplerStateRHIRef* SamplerSource = &Value->GetResource()->SamplerStateRHI;
 				ESamplerSourceMode SourceMode = Parameter.SamplerSource;
@@ -972,9 +970,8 @@ void FUniformExpressionSet::FillUniformBuffer(const FMaterialRenderContext& Mate
 			void** ResourceTableSamplerPtr = (void**)((uint8*)BufferCursor + 1 * SHADER_PARAMETER_POINTER_ALIGNMENT);
 			BufferCursor = ((uint8*)BufferCursor) + (SHADER_PARAMETER_POINTER_ALIGNMENT * 2);
 
-			if (Value && Value->GetResource() && (Value->GetMaterialType() & MCT_TextureCubeArray) != 0u)
+			if (Value && Value->GetResource() && Value->TextureReference.TextureReferenceRHI && (Value->GetMaterialType() & MCT_TextureCubeArray) != 0u)
 			{
-				check(Value->TextureReference.TextureReferenceRHI);
 				*ResourceTableTexturePtr = Value->TextureReference.TextureReferenceRHI;
 				const FSamplerStateRHIRef* SamplerSource = &Value->GetResource()->SamplerStateRHI;
 				ESamplerSourceMode SourceMode = Parameter.SamplerSource;
@@ -1012,9 +1009,8 @@ void FUniformExpressionSet::FillUniformBuffer(const FMaterialRenderContext& Mate
 			BufferCursor = ((uint8*)BufferCursor) + (SHADER_PARAMETER_POINTER_ALIGNMENT * 2);
 			check(BufferCursor <= TempBuffer + TempBufferSize);
 
-			if(Value && Value->GetResource() && (Value->GetMaterialType() & MCT_VolumeTexture) != 0u)
+			if(Value && Value->GetResource() && Value->TextureReference.TextureReferenceRHI && (Value->GetMaterialType() & MCT_VolumeTexture) != 0u)
 			{
-				check(Value->TextureReference.TextureReferenceRHI);
 				*ResourceTableTexturePtr = Value->TextureReference.TextureReferenceRHI;
 				const FSamplerStateRHIRef* SamplerSource = &Value->GetResource()->SamplerStateRHI;
 
