Fix an assert (check(!Task->StaticMesh->IsCompiling());) seen when cooking from a cold DDC.

Note the code fix was by [REDACTED_PERSON], my machine just happened to be in a good state to reproduce the problem.

#rnx

- Make sure that StaticMesh is not compiling as well as the existing check that GenerateSource is not compiling.

diff --git a/Engine/Source/Runtime/Engine/Private/DistanceFieldAtlas.cpp b/Engine/Source/Runtime/Engine/Private/DistanceFieldAtlas.cpp
--- a/Engine/Source/Runtime/Engine/Private/DistanceFieldAtlas.cpp
+++ b/Engine/Source/Runtime/Engine/Private/DistanceFieldAtlas.cpp
@@ -456,8 +456,7 @@ void FDistanceFieldAsyncQueue::ProcessPendingTasks()
 	for (auto It = PendingTasks.CreateIterator(); It; ++It)
 	{
 		FAsyncDistanceFieldTask* Task = *It;
-		if (Task->GenerateSource == nullptr || 
-			Task->GenerateSource->IsCompiling() == false)
+		if ((Task->GenerateSource == nullptr || !Task->GenerateSource->IsCompiling()) && (Task->StaticMesh == nullptr || !Task->StaticMesh ->IsCompiling()))
 		{
 			StartBackgroundTask(Task);
 			It.RemoveCurrent();
@@ -482,7 +481,8 @@ void FDistanceFieldAsyncQueue::AddTask(FAsyncDistanceFieldTask* Task)
 	}
 	
 	const bool bUseAsyncBuild = GUseAsyncDistanceFieldBuildQueue || !IsInGameThread();
-	const bool bIsCompiling = Task->GenerateSource->IsCompiling();
+	const bool bIsCompiling = (Task->GenerateSource && Task->GenerateSource->IsCompiling()) || (Task->StaticMesh && Task->StaticMesh->IsCompiling());
+
 	{
 		// Array protection when called from multiple threads
 		FScopeLock Lock(&CriticalSection);
