- ActiveSound ensure removal to unblock while the underlying pointer leak is tracked down.
- Only allow setting 'Multiple Sound Instances' per AudioComponent to false to be observed while sounds are not already playing to avoid mis-tracking looping sound instances that may be virtualized.

diff --git a/Engine/Source/Runtime/Engine/Classes/Components/AudioComponent.h b/Engine/Source/Runtime/Engine/Classes/Components/AudioComponent.h
--- a/Engine/Source/Runtime/Engine/Classes/Components/AudioComponent.h
+++ b/Engine/Source/Runtime/Engine/Classes/Components/AudioComponent.h
@@ -218,7 +218,8 @@ class ENGINE_API UAudioComponent : public USceneComponent, public ISoundGenerato
 	uint8 bSuppressSubtitles:1;
 
 	/** If true, the Audio Component will play multiple sound instances at once. Switching sounds or calling play while already playing
-	  * will not stop already active instances. */
+	  * will not stop already active instances. Virtualization for all played sounds will be disabled. Disabling while sound(s) are playing
+	  * will not take effect until the AudioComponent is stopped and restarted. */
 	UPROPERTY(EditAnywhere, BlueprintReadWrite, Category = Sound, meta = (DisplayName = "Play Multiple Instances", DisplayAfter = "Priority"))
 	uint8 bCanPlayMultipleInstances:1;
 
diff --git a/Engine/Source/Runtime/Engine/Private/AudioDevice.cpp b/Engine/Source/Runtime/Engine/Private/AudioDevice.cpp
--- a/Engine/Source/Runtime/Engine/Private/AudioDevice.cpp
+++ b/Engine/Source/Runtime/Engine/Private/AudioDevice.cpp
@@ -4904,7 +4904,7 @@ void FAudioDevice::AddNewActiveSoundInternal(const FActiveSound& InNewActiveSoun
 	if (ActiveSound->GetAudioComponentID() > 0)
 	{
 		TArray<FActiveSound*>& ActiveSoundArray = AudioComponentIDToActiveSoundMap.FindOrAdd(ActiveSound->GetAudioComponentID());
-		ActiveSoundArray.Add(ActiveSound);
+		ActiveSoundArray.AddUnique(ActiveSound);
 	}
 }
 
@@ -4980,22 +4980,18 @@ void FAudioDevice::AddVirtualLoop(const FAudioVirtualLoop& InVirtualLoop)
 	{
 		if (TArray<FActiveSound*>* ExistingSounds = AudioComponentIDToActiveSoundMap.Find(ComponentID))
 		{
-			if (!CanHaveMultipleActiveSounds(ComponentID))
+			// Only components playing a single looped sound at a time can virtualize.
+			for (int32 i = ExistingSounds->Num() - 1; i >= 0; --i)
 			{
-				// if we can't have more than one sound, we should have only one element (or less) here
-				ensure(ExistingSounds->Num() <= 1);
-				for (int32 i = ExistingSounds->Num() - 1; i >= 0; --i)
+				const FActiveSound* ExistingSound = (*ExistingSounds)[i];
+				if (ensure(ExistingSound))
 				{
-					const FActiveSound* ExistingSound = (*ExistingSounds)[i];
-					if (ensure(ExistingSound))
-					{
-						UE_LOG(LogAudio, Warning, TEXT("Attempting to add Sound '%s' to ComponentID when map already contains ID for Sound '%s'. "
-						"Associated AudioComponent can only play a single sound instance at one time."),
-							ActiveSound.Sound ? *ActiveSound.Sound->GetName() : TEXT("N/A"),
-							ExistingSound->Sound ? *ExistingSound->Sound->GetName() : TEXT("N/A")
-						);
-						ExistingSounds->RemoveAtSwap(i, 1, false /* bAllowShrinking */);
-					}
+					UE_LOG(LogAudio, Warning, TEXT("Attempting to add Sound '%s' to ComponentID when map already contains ID for Sound '%s'. "
+					"Associated AudioComponent can only play a single sound instance at one time. This may indicate a leak of FActiveSounds."),
+						ActiveSound.Sound ? *ActiveSound.Sound->GetName() : TEXT("N/A"),
+						ExistingSound->Sound ? *ExistingSound->Sound->GetName() : TEXT("N/A")
+					);
+					ExistingSounds->RemoveAtSwap(i, 1, false /* bAllowShrinking */);
 				}
 			}
 			ExistingSounds->AddUnique(&ActiveSound);
@@ -5327,7 +5323,11 @@ void FAudioDevice::SetCanHaveMultipleActiveSounds(uint64 InAudioComponentID, boo
 		return;
 	}
 
-	AudioComponentIDToCanHaveMultipleActiveSoundsMap.FindOrAdd(InAudioComponentID, InCanHaveMultipleActiveSounds) = InCanHaveMultipleActiveSounds;
+	bool& bCanHaveMultipleActiveSounds = AudioComponentIDToCanHaveMultipleActiveSoundsMap.FindOrAdd(InAudioComponentID, InCanHaveMultipleActiveSounds);
+
+	// This must be or'ed with existing value as disabling multiple active sounds while playing can potentially cause ActiveSound
+	// instances to get lost while virtualizing.
+	bCanHaveMultipleActiveSounds |= InCanHaveMultipleActiveSounds;
 }
 
 void FAudioDevice::RemoveActiveSound(FActiveSound* ActiveSound)
diff --git a/Engine/Source/Runtime/Engine/Private/Components/AudioComponent.cpp b/Engine/Source/Runtime/Engine/Private/Components/AudioComponent.cpp
--- a/Engine/Source/Runtime/Engine/Private/Components/AudioComponent.cpp
+++ b/Engine/Source/Runtime/Engine/Private/Components/AudioComponent.cpp
@@ -442,17 +442,19 @@ void UAudioComponent::PlayInternal(const PlayInternalRequestData& InPlayRequestD
 	// Reset our fading out flag in case this is a reused audio component and we are replaying after previously fading out
 	bIsFadingOut = false;
 
+	const bool bIsSoundLooping = InSound && InSound->IsLooping();
+
 	if (IsActive())
 	{
-		// If this is an auto destroy component we need to prevent it from being auto-destroyed since we're really just restarting it
-		bool bCurrentAutoDestroy = bAutoDestroy;
-		bAutoDestroy = false;
 		//Only stop if this component is limited to one active sound
-		if (!bCanPlayMultipleInstances)
+		if (!bCanPlayMultipleInstances || bIsSoundLooping)
 		{
+			// If this is an auto destroy component we need to prevent it from being auto-destroyed since we're really just restarting it
+			bool bCurrentAutoDestroy = bAutoDestroy;
+			bAutoDestroy = false;
 			Stop();
+			bAutoDestroy = bCurrentAutoDestroy;
 		}
-		bAutoDestroy = bCurrentAutoDestroy;
 	}
 
 	if (Sound && (World == nullptr || World->bAllowAudioPlayback))
@@ -615,8 +617,22 @@ void UAudioComponent::PlayInternal(const PlayInternalRequestData& InPlayRequestD
 			// Bump ActiveCount... this is used to determine if an audio component is still active after a sound reports back as completed
 			++ActiveCount;
 
-			// Pass along whether or not component is setup to support multiple active sounds. This is to ensure virtualization will function accordingly.
-			AudioDevice->SetCanHaveMultipleActiveSounds(AudioComponentID, bCanPlayMultipleInstances);
+			// Pass along whether or not component is setup to support multiple active sounds.
+			// This is to ensure virtualization will function accordingly. Disable the feature
+			// if the sound is looping as a safety mechanism to avoid stuck loops.
+			if (bIsSoundLooping)
+			{
+				if (bCanPlayMultipleInstances)
+				{
+					UE_LOG(LogAudio, Warning, TEXT("'Can Play Multiple Instances' disabled: Sound '%s' set to looping"), *InSound->GetName());
+				}
+
+				AudioDevice->SetCanHaveMultipleActiveSounds(AudioComponentID, false);
+			}
+			else
+			{
+				AudioDevice->SetCanHaveMultipleActiveSounds(AudioComponentID, bCanPlayMultipleInstances);
+			}
 
 			AudioDevice->AddNewActiveSound(NewActiveSound, &InstanceParameters);
 
