Fix for crash using UTexture2DArray in headless commandlets: Use UTexture2D::Source rather than UTexture2D::PlatformData (which isn't valid in headless commandlets).

diff --git a/Engine/Source/Runtime/Engine/Classes/Engine/Texture2DArray.cpp b/Engine/Source/Runtime/Engine/Classes/Engine/Texture2DArray.cpp
--- a/Engine/Source/Runtime/Engine/Classes/Engine/Texture2DArray.cpp
+++ b/Engine/Source/Runtime/Engine/Classes/Engine/Texture2DArray.cpp
@@ -172,12 +172,11 @@ ENGINE_API bool UTexture2DArray::UpdateSourceFromSourceTextures(bool bCreatingNe
 	{
 		FTextureSource& InitialSource = SourceTextures[0]->Source;
 		// Format and format size.
-		EPixelFormat PixelFormat = SourceTextures[0]->GetPixelFormat();
 		ETextureSourceFormat Format = InitialSource.GetFormat();
 		int32 FormatDataSize = InitialSource.GetBytesPerPixel();
-		// X,Y,Z size of the array.
-		int32 SizeX = SourceTextures[0]->GetSizeX();
-		int32 SizeY = SourceTextures[0]->GetSizeY();
+		// X,Y,Z size of the array. Switched to use Source sizes - previously would use PlatformData sizes which caused an error in headless commandlet MHC export.
+		int32 SizeX = InitialSource.GetSizeX();
+		int32 SizeY = InitialSource.GetSizeY();
 		uint32 ArraySize = SourceTextures.Num();
 		// Only copy the first mip from the source textures to array texture.
 		uint32 NumMips = 1;
