Fix Replicated component destruction not reaching client

diff --git a/Engine/Source/Runtime/Engine/Private/DataChannel.cpp b/Engine/Source/Runtime/Engine/Private/DataChannel.cpp
--- a/Engine/Source/Runtime/Engine/Private/DataChannel.cpp
+++ b/Engine/Source/Runtime/Engine/Private/DataChannel.cpp
@@ -2899,7 +2899,10 @@ void UActorChannel::ProcessBunch( FInBunch & Bunch )
 		TSharedRef<FObjectReplicator>& ObjectReplicator = RepComp.Value();
 		if (ObjectReplicator->GetObject() == nullptr)
 		{
-			RepComp.RemoveCurrent();
+			if (!Connection->Driver->IsServer())
+			{
+				RepComp.RemoveCurrent();
+			}
 			continue;
 		}
 
