Integrating CIS warning fix to UE5/Main

diff --git a/Engine/Source/Runtime/Engine/Private/SoundWave.cpp b/Engine/Source/Runtime/Engine/Private/SoundWave.cpp
--- a/Engine/Source/Runtime/Engine/Private/SoundWave.cpp
+++ b/Engine/Source/Runtime/Engine/Private/SoundWave.cpp
@@ -1239,7 +1239,7 @@ uint32 USoundWave::GetNumChunks() const
 	}
 	else if (GetOutermost()->HasAnyPackageFlags(PKG_ReloadingForCooker) || !FApp::CanEverRenderAudio())
 	{
-		UE_LOG(LogAudio, Warning, TEXT("USoundWave::GetNumChunks called in the cook commandlet."))
+		UE_LOG(LogAudio, Verbose, TEXT("USoundWave::GetNumChunks called either during reloading for cooker or when the application can't render audio."))
 		return 0;
 	}
 	else
