- Fix bug with checking wrong sound object in AudioComponent 'PlayInternal' call for looping/logging info
- Minor comment/logging clean-up
- Reset InstanceParameters on UAudioComponent when stopping sound

diff --git a/Engine/Source/Runtime/Engine/Classes/Components/AudioComponent.h b/Engine/Source/Runtime/Engine/Classes/Components/AudioComponent.h
--- a/Engine/Source/Runtime/Engine/Classes/Components/AudioComponent.h
+++ b/Engine/Source/Runtime/Engine/Classes/Components/AudioComponent.h
@@ -491,6 +491,8 @@ public:
 		, EAudioFaderCurve InFadeCurve = EAudioFaderCurve::Linear
 	);
 
+	virtual void ResetParameters() override;
+
 private:
 
 	static uint64 AudioComponentIDCounter;
@@ -809,7 +811,7 @@ private:
 protected:
 
 	/** Utility function called by Play and FadeIn to start a sound playing. */
-	void PlayInternal(const PlayInternalRequestData& InPlayRequestData, USoundBase * InSound = nullptr);
+	void PlayInternal(const PlayInternalRequestData& InPlayRequestData, USoundBase * InSoundOverride = nullptr);
 
 #if WITH_EDITORONLY_DATA
 	/** Utility function that updates which texture is displayed on the sprite dependent on the properties of the Audio Component. */
diff --git a/Engine/Source/Runtime/Engine/Private/Components/AudioComponent.cpp b/Engine/Source/Runtime/Engine/Private/Components/AudioComponent.cpp
--- a/Engine/Source/Runtime/Engine/Private/Components/AudioComponent.cpp
+++ b/Engine/Source/Runtime/Engine/Private/Components/AudioComponent.cpp
@@ -431,25 +431,28 @@ void UAudioComponent::PlayQuantized(
 	PlayInternal(Data);
 }
 
-void UAudioComponent::PlayInternal(const PlayInternalRequestData& InPlayRequestData, USoundBase* InSound)
+void UAudioComponent::PlayInternal(const PlayInternalRequestData& InPlayRequestData, USoundBase* InSoundOverride)
 {
 	SCOPE_CYCLE_COUNTER(STAT_AudioComp_Play);
 
 	UWorld* World = GetWorld();
+	USoundBase* SoundToPlay = InSoundOverride ? InSoundOverride : Sound.Get();
+	if (SoundToPlay)
+	{
+		UE_LOG(LogAudio, Verbose, TEXT("%g: Playing AudioComponent : '%s' with Sound: '%s'"), World ? World->GetAudioTimeSeconds() : 0.0f, *GetFullName(), *SoundToPlay->GetName());
+	}
 
-	UE_LOG(LogAudio, Verbose, TEXT("%g: Playing AudioComponent : '%s' with Sound: '%s'"), World ? World->GetAudioTimeSeconds() : 0.0f, *GetFullName(), Sound ? *Sound->GetName() : TEXT("nullptr"));
-
-	// Reset our fading out flag in case this is a reused audio component and we are replaying after previously fading out
+	// Reset fading out flag in case this is a reused audio component and replaying after previously fading out.
 	bIsFadingOut = false;
 
-	const bool bIsSoundLooping = InSound && InSound->IsLooping();
-
+	// Stop sound if active & not set to play multiple instances, irrespective of whether or not a valid sound is set to play.
+	const bool bIsSoundLooping = SoundToPlay && SoundToPlay->IsLooping();
 	if (IsActive())
 	{
-		//Only stop if this component is limited to one active sound
+		// Stop if this component is limited to one active sound that is not looping.
 		if (!bCanPlayMultipleInstances || bIsSoundLooping)
 		{
-			// If this is an auto destroy component we need to prevent it from being auto-destroyed since we're really just restarting it
+			// Prevent sound from being auto-destroyed since its just being restarted.
 			bool bCurrentAutoDestroy = bAutoDestroy;
 			bAutoDestroy = false;
 			Stop();
@@ -457,192 +460,193 @@ void UAudioComponent::PlayInternal(const PlayInternalRequestData& InPlayRequestD
 		}
 	}
 
-	if (Sound && (World == nullptr || World->bAllowAudioPlayback))
+	if (!SoundToPlay)
 	{
-		if (FAudioDevice* AudioDevice = GetAudioDevice())
-		{
-			// Store the time that this audio component played
-			if (World)
-			{
-				TimeAudioComponentPlayed = World->GetAudioTimeSeconds();
-			}
-			else
-			{
-				TimeAudioComponentPlayed = 0.0f;
-			}
-			FadeInTimeDuration = InPlayRequestData.FadeInDuration;
-
-			// Auto attach if requested
-			const bool bWasAutoAttached = bDidAutoAttach;
-			bDidAutoAttach = false;
-			if (bAutoManageAttachment && World->IsGameWorld())
-			{
-				USceneComponent* NewParent = AutoAttachParent.Get();
-				if (NewParent)
-				{
-					const bool bAlreadyAttached = GetAttachParent() && (GetAttachParent() == NewParent) && (GetAttachSocketName() == AutoAttachSocketName) && GetAttachParent()->GetAttachChildren().Contains(this);
-					if (!bAlreadyAttached)
-					{
-						bDidAutoAttach = bWasAutoAttached;
-						CancelAutoAttachment(true, World);
-						SavedAutoAttachRelativeLocation = GetRelativeLocation();
-						SavedAutoAttachRelativeRotation = GetRelativeRotation();
-						SavedAutoAttachRelativeScale3D = GetRelativeScale3D();
-						AttachToComponent(NewParent, FAttachmentTransformRules(AutoAttachLocationRule, AutoAttachRotationRule, AutoAttachScaleRule, false), AutoAttachSocketName);
-					}
-
-					bDidAutoAttach = true;
-				}
-				else
-				{
-					CancelAutoAttachment(true, World);
-				}
-			}
-
-			// Create / configure new ActiveSound
-			const FSoundAttenuationSettings* AttenuationSettingsToApply = bAllowSpatialization ? GetAttenuationSettingsToApply() : nullptr;
+		UE_LOG(LogAudio, Verbose, TEXT("%g: AudioComponent '%s' failed to play sound: No sound to play specified.'"), World ? World->GetAudioTimeSeconds() : 0.0f, *GetFullName());
+		return;
+	}
 
-			float MaxDistance = 0.0f;
-			float FocusFactor = 1.0f;
-			FVector Location = GetComponentTransform().GetLocation();
+	if (World && !World->bAllowAudioPlayback)
+	{
+		return;
+	}
 
-			AudioDevice->GetMaxDistanceAndFocusFactor(Sound, World, Location, AttenuationSettingsToApply, MaxDistance, FocusFactor);
+	FAudioDevice* AudioDevice = GetAudioDevice();
+	if (!AudioDevice)
+	{
+		return;
+	}
 
-			FActiveSound NewActiveSound;
-			NewActiveSound.SetAudioComponent(*this);
-			NewActiveSound.SetWorld(GetWorld());
+	// Store the time that this audio component played
+	if (World)
+	{
+		TimeAudioComponentPlayed = World->GetAudioTimeSeconds();
+	}
+	else
+	{
+		TimeAudioComponentPlayed = 0.0f;
+	}
+	FadeInTimeDuration = InPlayRequestData.FadeInDuration;
 
-			//If the InSound pointer is non-null, use its data. Otherwise, use the base Sound
-			if (InSound == nullptr)
-			{
-				NewActiveSound.SetSound(Sound);
-			}
-			else
+	// Auto attach if requested
+	const bool bWasAutoAttached = bDidAutoAttach;
+	bDidAutoAttach = false;
+	if (bAutoManageAttachment && World->IsGameWorld())
+	{
+		USceneComponent* NewParent = AutoAttachParent.Get();
+		if (NewParent)
+		{
+			const bool bAlreadyAttached = GetAttachParent() && (GetAttachParent() == NewParent) && (GetAttachSocketName() == AutoAttachSocketName) && GetAttachParent()->GetAttachChildren().Contains(this);
+			if (!bAlreadyAttached)
 			{
-				NewActiveSound.SetSound(InSound);
+				bDidAutoAttach = bWasAutoAttached;
+				CancelAutoAttachment(true, World);
+				SavedAutoAttachRelativeLocation = GetRelativeLocation();
+				SavedAutoAttachRelativeRotation = GetRelativeRotation();
+				SavedAutoAttachRelativeScale3D = GetRelativeScale3D();
+				AttachToComponent(NewParent, FAttachmentTransformRules(AutoAttachLocationRule, AutoAttachRotationRule, AutoAttachScaleRule, false), AutoAttachSocketName);
 			}
 
-			NewActiveSound.SetSourceEffectChain(SourceEffectChain);
-			NewActiveSound.SetSoundClass(SoundClassOverride);
-			NewActiveSound.ConcurrencySet = ConcurrencySet;
+			bDidAutoAttach = true;
+		}
+		else
+		{
+			CancelAutoAttachment(true, World);
+		}
+	}
 
-			const float Volume = (VolumeModulationMax + ((VolumeModulationMin - VolumeModulationMax) * RandomStream.FRand())) * VolumeMultiplier;
-			NewActiveSound.SetVolume(Volume);
+	// Create / configure new ActiveSound
+	const FSoundAttenuationSettings* AttenuationSettingsToApply = bAllowSpatialization ? GetAttenuationSettingsToApply() : nullptr;
 
-			// The priority used for the active sound is the audio component's priority scaled with the sound's priority
-			if (bOverridePriority)
-			{
-				NewActiveSound.Priority = Priority;
-			}
-			else
-			{
-				NewActiveSound.Priority = Sound->Priority;
-			}
+	float MaxDistance = 0.0f;
+	float FocusFactor = 1.0f;
+	FVector Location = GetComponentTransform().GetLocation();
 
-			const float Pitch = (PitchModulationMax + ((PitchModulationMin - PitchModulationMax) * RandomStream.FRand())) * PitchMultiplier;
-			NewActiveSound.SetPitch(Pitch);
+	AudioDevice->GetMaxDistanceAndFocusFactor(SoundToPlay, World, Location, AttenuationSettingsToApply, MaxDistance, FocusFactor);
 
-			NewActiveSound.bEnableLowPassFilter = bEnableLowPassFilter;
-			NewActiveSound.LowPassFilterFrequency = LowPassFilterFrequency;
-			NewActiveSound.RequestedStartTime = FMath::Max(0.f, InPlayRequestData.StartTime);
+	FActiveSound NewActiveSound;
+	NewActiveSound.SetAudioComponent(*this);
+	NewActiveSound.SetWorld(GetWorld());
+	NewActiveSound.SetSound(SoundToPlay);
+	NewActiveSound.SetSourceEffectChain(SourceEffectChain);
+	NewActiveSound.SetSoundClass(SoundClassOverride);
+	NewActiveSound.ConcurrencySet = ConcurrencySet;
 
-			if (bOverrideSubtitlePriority)
-			{
-				NewActiveSound.SubtitlePriority = SubtitlePriority;
-			}
-			else
-			{
-				NewActiveSound.SubtitlePriority = Sound->GetSubtitlePriority();
-			}
+	const float Volume = (VolumeModulationMax + ((VolumeModulationMin - VolumeModulationMax) * RandomStream.FRand())) * VolumeMultiplier;
+	NewActiveSound.SetVolume(Volume);
 
-			NewActiveSound.bShouldRemainActiveIfDropped = bShouldRemainActiveIfDropped;
-			NewActiveSound.bHandleSubtitles = (!bSuppressSubtitles || OnQueueSubtitles.IsBound());
-			NewActiveSound.bIgnoreForFlushing = bIgnoreForFlushing;
+	// The priority used for the active sound is the audio component's priority scaled with the sound's priority
+	if (bOverridePriority)
+	{
+		NewActiveSound.Priority = Priority;
+	}
+	else
+	{
+		NewActiveSound.Priority = SoundToPlay->Priority;
+	}
 
-			NewActiveSound.bIsUISound = bIsUISound;
-			NewActiveSound.bIsMusic = bIsMusic;
-			NewActiveSound.bAlwaysPlay = bAlwaysPlay;
-			NewActiveSound.bReverb = bReverb;
-			NewActiveSound.bCenterChannelOnly = bCenterChannelOnly;
-			NewActiveSound.bIsPreviewSound = bIsPreviewSound;
-			NewActiveSound.bLocationDefined = !bPreviewComponent;
-			NewActiveSound.bIsPaused = bIsPaused;
+	const float Pitch = (PitchModulationMax + ((PitchModulationMin - PitchModulationMax) * RandomStream.FRand())) * PitchMultiplier;
+	NewActiveSound.SetPitch(Pitch);
 
-			if (NewActiveSound.bLocationDefined)
-			{
-				NewActiveSound.Transform = GetComponentTransform();
-			}
+	NewActiveSound.bEnableLowPassFilter = bEnableLowPassFilter;
+	NewActiveSound.LowPassFilterFrequency = LowPassFilterFrequency;
+	NewActiveSound.RequestedStartTime = FMath::Max(0.f, InPlayRequestData.StartTime);
 
-			NewActiveSound.bAllowSpatialization = bAllowSpatialization;
-			NewActiveSound.bHasAttenuationSettings = (AttenuationSettingsToApply != nullptr);
-			if (NewActiveSound.bHasAttenuationSettings)
-			{
-				NewActiveSound.AttenuationSettings = *AttenuationSettingsToApply;
-				NewActiveSound.FocusData.PriorityScale = AttenuationSettingsToApply->GetFocusPriorityScale(AudioDevice->GetGlobalFocusSettings(), FocusFactor);
-			}
+	if (bOverrideSubtitlePriority)
+	{
+		NewActiveSound.SubtitlePriority = SubtitlePriority;
+	}
+	else
+	{
+		NewActiveSound.SubtitlePriority = SoundToPlay->GetSubtitlePriority();
+	}
 
-			NewActiveSound.EnvelopeFollowerAttackTime = FMath::Max(EnvelopeFollowerAttackTime, 0);
-			NewActiveSound.EnvelopeFollowerReleaseTime = FMath::Max(EnvelopeFollowerReleaseTime, 0);
+	NewActiveSound.bShouldRemainActiveIfDropped = bShouldRemainActiveIfDropped;
+	NewActiveSound.bHandleSubtitles = (!bSuppressSubtitles || OnQueueSubtitles.IsBound());
+	NewActiveSound.bIgnoreForFlushing = bIgnoreForFlushing;
 
-			NewActiveSound.bUpdatePlayPercentage = OnAudioPlaybackPercentNative.IsBound() || OnAudioPlaybackPercent.IsBound();
-			NewActiveSound.bUpdateSingleEnvelopeValue = OnAudioSingleEnvelopeValue.IsBound() || OnAudioSingleEnvelopeValueNative.IsBound();
-			NewActiveSound.bUpdateMultiEnvelopeValue = OnAudioMultiEnvelopeValue.IsBound() || OnAudioMultiEnvelopeValueNative.IsBound();
+	NewActiveSound.bIsUISound = bIsUISound;
+	NewActiveSound.bIsMusic = bIsMusic;
+	NewActiveSound.bAlwaysPlay = bAlwaysPlay;
+	NewActiveSound.bReverb = bReverb;
+	NewActiveSound.bCenterChannelOnly = bCenterChannelOnly;
+	NewActiveSound.bIsPreviewSound = bIsPreviewSound;
+	NewActiveSound.bLocationDefined = !bPreviewComponent;
+	NewActiveSound.bIsPaused = bIsPaused;
 
-			NewActiveSound.ModulationRouting = ModulationRouting;
+	if (NewActiveSound.bLocationDefined)
+	{
+		NewActiveSound.Transform = GetComponentTransform();
+	}
 
-			// Setup audio component cooked analysis data playback data set
-			if (AudioDevice->IsBakedAnalaysisQueryingEnabled())
-			{
-				TArray<USoundWave*> SoundWavesWithCookedData;
-				NewActiveSound.bUpdatePlaybackTime = Sound->GetSoundWavesWithCookedAnalysisData(SoundWavesWithCookedData);
+	NewActiveSound.bAllowSpatialization = bAllowSpatialization;
+	NewActiveSound.bHasAttenuationSettings = (AttenuationSettingsToApply != nullptr);
+	if (NewActiveSound.bHasAttenuationSettings)
+	{
+		NewActiveSound.AttenuationSettings = *AttenuationSettingsToApply;
+		NewActiveSound.FocusData.PriorityScale = AttenuationSettingsToApply->GetFocusPriorityScale(AudioDevice->GetGlobalFocusSettings(), FocusFactor);
+	}
 
-				// Reset the audio component's soundwave playback times
-				SoundWavePlaybackTimes.Reset();
-				for (USoundWave* SoundWave : SoundWavesWithCookedData)
-				{
-					SoundWavePlaybackTimes.Add(SoundWave->GetUniqueID(), FSoundWavePlaybackTimeData(SoundWave));
-				}
-			}
+	NewActiveSound.EnvelopeFollowerAttackTime = FMath::Max(EnvelopeFollowerAttackTime, 0);
+	NewActiveSound.EnvelopeFollowerReleaseTime = FMath::Max(EnvelopeFollowerReleaseTime, 0);
 
+	NewActiveSound.bUpdatePlayPercentage = OnAudioPlaybackPercentNative.IsBound() || OnAudioPlaybackPercent.IsBound();
+	NewActiveSound.bUpdateSingleEnvelopeValue = OnAudioSingleEnvelopeValue.IsBound() || OnAudioSingleEnvelopeValueNative.IsBound();
+	NewActiveSound.bUpdateMultiEnvelopeValue = OnAudioMultiEnvelopeValue.IsBound() || OnAudioMultiEnvelopeValueNative.IsBound();
 
-			// Pass quantization data to the active sound
-			NewActiveSound.QuantizedRequestData = InPlayRequestData.QuantizedRequestData;
+	NewActiveSound.ModulationRouting = ModulationRouting;
 
-			NewActiveSound.MaxDistance = MaxDistance;
+	// Setup audio component cooked analysis data playback data set
+	if (AudioDevice->IsBakedAnalaysisQueryingEnabled())
+	{
+		TArray<USoundWave*> SoundWavesWithCookedData;
+		NewActiveSound.bUpdatePlaybackTime = Sound->GetSoundWavesWithCookedAnalysisData(SoundWavesWithCookedData);
 
-			Audio::FVolumeFader& Fader = NewActiveSound.ComponentVolumeFader;
-			Fader.SetVolume(0.0f); // Init to 0.0f to fade as default is 1.0f
-			Fader.StartFade(InPlayRequestData.FadeVolumeLevel, InPlayRequestData.FadeInDuration, static_cast<Audio::EFaderCurve>(InPlayRequestData.FadeCurve));
+		// Reset the audio component's soundwave playback times
+		SoundWavePlaybackTimes.Reset();
+		for (USoundWave* SoundWave : SoundWavesWithCookedData)
+		{
+			SoundWavePlaybackTimes.Add(SoundWave->GetUniqueID(), FSoundWavePlaybackTimeData(SoundWave));
+		}
+	}
 
-			// Bump ActiveCount... this is used to determine if an audio component is still active after a sound reports back as completed
-			++ActiveCount;
 
-			// Pass along whether or not component is setup to support multiple active sounds.
-			// This is to ensure virtualization will function accordingly. Disable the feature
-			// if the sound is looping as a safety mechanism to avoid stuck loops.
-			if (bIsSoundLooping)
-			{
-				if (bCanPlayMultipleInstances)
-				{
-					UE_LOG(LogAudio, Warning, TEXT("'Can Play Multiple Instances' disabled: Sound '%s' set to looping"), *InSound->GetName());
-				}
+	// Pass quantization data to the active sound
+	NewActiveSound.QuantizedRequestData = InPlayRequestData.QuantizedRequestData;
 
-				AudioDevice->SetCanHaveMultipleActiveSounds(AudioComponentID, false);
-			}
-			else
-			{
-				AudioDevice->SetCanHaveMultipleActiveSounds(AudioComponentID, bCanPlayMultipleInstances);
-			}
+	NewActiveSound.MaxDistance = MaxDistance;
 
-			AudioDevice->AddNewActiveSound(NewActiveSound, &InstanceParameters);
+	Audio::FVolumeFader& Fader = NewActiveSound.ComponentVolumeFader;
+	Fader.SetVolume(0.0f); // Init to 0.0f to fade as default is 1.0f
+	Fader.StartFade(InPlayRequestData.FadeVolumeLevel, InPlayRequestData.FadeInDuration, static_cast<Audio::EFaderCurve>(InPlayRequestData.FadeCurve));
 
-			// In editor, the audio thread is not run separate from the game thread, and can result in calling PlaybackComplete prior
-			// to bIsActive being set. Therefore, we assign to the current state of ActiveCount as opposed to just setting to true.
-			SetActiveFlag(ActiveCount > 0);
+	// Bump ActiveCount... this is used to determine if an audio component is still active after a sound reports back as completed
+	++ActiveCount;
 
-			BroadcastPlayState();
+	// Pass along whether or not component is setup to support multiple active sounds.
+	// This is to ensure virtualization will function accordingly. Disable the feature
+	// if the sound is looping as a safety mechanism to avoid stuck loops.
+	if (bIsSoundLooping)
+	{
+		if (bCanPlayMultipleInstances)
+		{
+			UE_LOG(LogAudio, Warning, TEXT("'Can Play Multiple Instances' disabled: Sound '%s' set to looping"), *SoundToPlay->GetName());
 		}
+
+		AudioDevice->SetCanHaveMultipleActiveSounds(AudioComponentID, false);
+	}
+	else
+	{
+		AudioDevice->SetCanHaveMultipleActiveSounds(AudioComponentID, bCanPlayMultipleInstances);
 	}
+
+	AudioDevice->AddNewActiveSound(NewActiveSound, &InstanceParameters);
+
+	// In editor, the audio thread is not run separate from the game thread, and can result in calling PlaybackComplete prior
+	// to bIsActive being set. Therefore, we assign to the current state of ActiveCount as opposed to just setting to true.
+	SetActiveFlag(ActiveCount > 0);
+
+	BroadcastPlayState();
 }
 
 FAudioDevice* UAudioComponent::GetAudioDevice() const
@@ -766,6 +770,12 @@ void UAudioComponent::AdjustVolumeInternal(float AdjustVolumeDuration, float Adj
 	}, GET_STATID(STAT_AudioAdjustVolume));
 }
 
+void UAudioComponent::ResetParameters()
+{
+	InstanceParameters.Reset();
+	ISoundGeneratorParameterInterface::ResetParameters();
+}
+
 void UAudioComponent::Stop()
 {
 	if (!IsActive())
@@ -927,6 +937,11 @@ void UAudioComponent::PlaybackCompleted(bool bFailedToStart)
 		CancelAutoAttachment(true, MyWorld);
 	}
 
+	if (bIsPreviewSound)
+	{
+		ResetParameters();
+	}
+
 	BroadcastPlayState();
 }
 
