Default to heap allocations for all TArrays within TAttributeContainer - addressing checks when running in debug configuration

diff --git a/Engine/Source/Runtime/Engine/Classes/Animation/AttributesContainer.h b/Engine/Source/Runtime/Engine/Classes/Animation/AttributesContainer.h
--- a/Engine/Source/Runtime/Engine/Classes/Animation/AttributesContainer.h
+++ b/Engine/Source/Runtime/Engine/Classes/Animation/AttributesContainer.h
@@ -49,26 +49,8 @@ namespace UE
 			template <typename OtherAllocator>
 			void CopyFrom(const TAttributeContainer<BoneIndexType, OtherAllocator>& Other)
 			{
-				AttributeIdentifiers.Empty(Other.AttributeIdentifiers.Num());
-				for (const TArray<FAttributeId, OtherAllocator> Array : Other.AttributeIdentifiers)
-				{
-					TArray<FAttributeId, InAllocator>& IdsArray = AttributeIdentifiers.AddDefaulted_GetRef();
-					for (const FAttributeId& Index : Array)
-					{
-						IdsArray.Add(Index);
-					}
-				}
-
-				UniqueTypedBoneIndices.Empty(Other.UniqueTypedBoneIndices.Num());
-				for (const TArray<int32, OtherAllocator> Array : Other.UniqueTypedBoneIndices)
-				{
-					TArray<int32, InAllocator>& IndicesArray = UniqueTypedBoneIndices.AddDefaulted_GetRef();
-					for (const int32& Index : Array)
-					{
-						IndicesArray.Add(Index);
-					}
-				}
-
+				AttributeIdentifiers = Other.AttributeIdentifiers;
+				UniqueTypedBoneIndices = Other.UniqueTypedBoneIndices;
 				UniqueTypes = Other.UniqueTypes;
 				Values.Empty(Other.Values.Num());
 				const int32 NumTypes = UniqueTypes.Num();
@@ -229,7 +211,7 @@ namespace UE
 			{
 				const int32 TypeIndex = FindOrAddTypeIndex(InScriptStruct);
 
-				TArray<FAttributeId, InAllocator>& AttributeIds = AttributeIdentifiers[TypeIndex];
+				TArray<FAttributeId>& AttributeIds = AttributeIdentifiers[TypeIndex];
 
 				const int32 ExistingAttributeIndex = AttributeIds.IndexOfByKey(InAttributeId);
 
@@ -280,7 +262,7 @@ namespace UE
 			uint8* FindOrAdd(const UScriptStruct* InScriptStruct, const FAttributeId& InAttributeId)
 			{
 				const int32 TypeIndex = FindOrAddTypeIndex(InScriptStruct);
-				TArray<FAttributeId, InAllocator>& AttributeIds = AttributeIdentifiers[TypeIndex];
+				TArray<FAttributeId>& AttributeIds = AttributeIdentifiers[TypeIndex];
 
 				int32 AttributeIndex = AttributeIds.IndexOfByKey(InAttributeId);
 
@@ -336,7 +318,7 @@ namespace UE
 				const int32 TypeIndex = FindTypeIndex(InScriptStruct);
 				if (TypeIndex != INDEX_NONE)
 				{
-					const TArray<FAttributeId, InAllocator>& AttributeIds = AttributeIdentifiers[TypeIndex];
+					const TArray<FAttributeId>& AttributeIds = AttributeIdentifiers[TypeIndex];
 
 					const int32 AttributeIndex = AttributeIds.IndexOfByKey(InAttributeId);
 					ensure(Values.IsValidIndex(TypeIndex));
@@ -379,7 +361,7 @@ namespace UE
 				const int32 TypeIndex = FindTypeIndex(InScriptStruct);
 				if (TypeIndex != INDEX_NONE)
 				{
-					const TArray<FAttributeId, InAllocator>& AttributeIds = AttributeIdentifiers[TypeIndex];
+					const TArray<FAttributeId>& AttributeIds = AttributeIdentifiers[TypeIndex];
 					const int32 AttributeIndex = AttributeIds.IndexOfByKey(InAttributeId);
 
 					ensure(Values.IsValidIndex(TypeIndex));
@@ -410,7 +392,7 @@ namespace UE
 				const int32 TypeIndex = FindTypeIndex(AttributeType::StaticStruct());
 				if (TypeIndex != INDEX_NONE)
 				{
-					const TArray<FAttributeId, InAllocator>& AttributeIds = AttributeIdentifiers[TypeIndex];
+					const TArray<FAttributeId>& AttributeIds = AttributeIdentifiers[TypeIndex];
 					const int32 AttributeIndex = AttributeIds.IndexOfByKey(InAttributeId);
 
 					ensure(Values.IsValidIndex(TypeIndex));
@@ -446,7 +428,7 @@ namespace UE
 				const int32 TypeIndex = FindTypeIndex(AttributeType::StaticStruct());
 				if (TypeIndex != INDEX_NONE)
 				{
-					const TArray<FAttributeId, InAllocator>& AttributeIds = AttributeIdentifiers[TypeIndex];
+					const TArray<FAttributeId>& AttributeIds = AttributeIdentifiers[TypeIndex];
 					const int32 AttributeIndex = AttributeIds.IndexOfByKey(InAttributeId);
 
 					ensure(Values.IsValidIndex(TypeIndex));
@@ -477,12 +459,12 @@ namespace UE
 				const int32 TypeIndex = FindTypeIndex(AttributeType::StaticStruct());
 				if (TypeIndex != INDEX_NONE)
 				{
-					const TArray<int32, InAllocator>& UniqueBoneIndices = UniqueTypedBoneIndices[TypeIndex];
+					const TArray<int32>& UniqueBoneIndices = UniqueTypedBoneIndices[TypeIndex];
 
 					// Early out if for this bone index no attributes are currently contained
 					if (UniqueBoneIndices.Contains(InAttributeId.GetIndex()))
 					{
-						const TArray<FAttributeId, InAllocator>& AttributeIds = AttributeIdentifiers[TypeIndex];
+						const TArray<FAttributeId>& AttributeIds = AttributeIdentifiers[TypeIndex];
 						const int32 AttributeIndex = AttributeIds.IndexOfByKey(InAttributeId);
 						
 						return AttributeIndex;
@@ -517,7 +499,7 @@ namespace UE
 				const int32 TypeIndex = FindTypeIndex(InScriptStruct);
 				if (TypeIndex != INDEX_NONE)
 				{
-					TArray<FAttributeId, InAllocator>& AttributeIds = AttributeIdentifiers[TypeIndex];
+					TArray<FAttributeId>& AttributeIds = AttributeIdentifiers[TypeIndex];
 
 					const int32 AttributeIndex = AttributeIds.IndexOfByKey(InAttributeId);
 
@@ -602,7 +584,7 @@ namespace UE
 			/*
 			* @return Array of all the contained keys for the provided TypeIndex
 			*/
-			const TArray<FAttributeId, InAllocator>& GetKeys(int32 TypeIndex) const { return AttributeIdentifiers[TypeIndex]; }
+			const TArray<FAttributeId>& GetKeys(int32 TypeIndex) const { return AttributeIdentifiers[TypeIndex]; }
 
 			/*
 			* @return Array of all the contained values for the provided TypeIndex
@@ -612,7 +594,7 @@ namespace UE
 			/*
 			* @return Array of all the contained attribute types (UScriptStruct)
 			*/
-			const TArray<TWeakObjectPtr<const UScriptStruct>, InAllocator>& GetUniqueTypes() const { return UniqueTypes; }
+			const TArray<TWeakObjectPtr<const UScriptStruct>>& GetUniqueTypes() const { return UniqueTypes; }
 
 			/*
 			* Populate out parameter OutAttributeKeyNames for all contained attribute key names
@@ -624,7 +606,7 @@ namespace UE
 			{
 				if (!AttributeIdentifiers.IsEmpty())
 				{
-					for (const TArray<FAttributeId, InAllocator>& AttributeTypeEntry : AttributeIdentifiers)
+					for (const TArray<FAttributeId>& AttributeTypeEntry : AttributeIdentifiers)
 					{
 						for (const FAttributeId& AttributeId : AttributeTypeEntry)
 						{
@@ -641,13 +623,13 @@ namespace UE
 		/*
 		* @return Unique bone indices for all contained entries of a specific attribute type
 		*/
-		const TArray<int32, InAllocator>& GetUniqueTypedBoneIndices(int32 TypeIndex) const { return UniqueTypedBoneIndices[TypeIndex]; }
+		const TArray<int32>& GetUniqueTypedBoneIndices(int32 TypeIndex) const { return UniqueTypedBoneIndices[TypeIndex]; }
 
 		public:
 			/* Deprecated API */
 			template<typename DataType> 
 			UE_DEPRECATED(5.0, "Deprecated behaviour, see new API")
-			TArray<DataType, InAllocator>& GetValuesArray()
+			TArray<DataType>& GetValuesArray()
 			{
 				static TArray<DataType, InAllocator> Array;
 				return Array;
@@ -684,17 +666,17 @@ namespace UE
 
 			template<typename DataType>
 			UE_DEPRECATED(5.0, "Deprecated behaviour, see new API")
-			const TArray<FAttributeId, InAllocator>& GetAttributeInfo() const
+			const TArray<FAttributeId>& GetAttributeInfo() const
 			{
-				static TArray<FAttributeId, InAllocator> Array;
+				static TArray<FAttributeId> Array;
 				return Array;
 			}
 
 			template<typename DataType>
 			UE_DEPRECATED(5.0, "Deprecated behaviour, see new API")
-			const TArray<int32, InAllocator>& GetUniqueBoneIndices() const
+			const TArray<int32>& GetUniqueBoneIndices() const
 			{
-				static TArray<int32, InAllocator> Array;
+				static TArray<int32> Array;
 				return Array;
 			}
 		protected:
@@ -707,7 +689,7 @@ namespace UE
 			/*
 			* @return Array of all the contained keys for the provided TypeIndex
 			*/
-			TArray<FAttributeId, InAllocator>& GetKeysInternal(int32 TypeIndex) { return AttributeIdentifiers[TypeIndex]; }
+			TArray<FAttributeId>& GetKeysInternal(int32 TypeIndex) { return AttributeIdentifiers[TypeIndex]; }
 			
 			/** Find or add a new root-level entry for the provided attribute data type, returning the index into the arrays representing the type */
 			int32 FindOrAddTypeIndex(const UScriptStruct* InScriptStruct)
@@ -729,10 +711,10 @@ namespace UE
 			}
 		protected:
 			/** Unique bone indices for all contained entries of a specific attribute type */
-			TArray<TArray<int32, InAllocator>, InAllocator> UniqueTypedBoneIndices;
-			TArray<TArray<FAttributeId, InAllocator>, InAllocator> AttributeIdentifiers;
+			TArray<TArray<int32>> UniqueTypedBoneIndices;
+			TArray<TArray<FAttributeId>> AttributeIdentifiers;
 			TArray<TArray<TWrappedAttribute<InAllocator>>> Values;
-			TArray<TWeakObjectPtr<const UScriptStruct>, InAllocator> UniqueTypes;
+			TArray<TWeakObjectPtr<const UScriptStruct>> UniqueTypes;
 
 			template <typename OtherBoneIndexType, typename OtherAllocator>
 			friend struct TAttributeContainer;
diff --git a/Engine/Source/Runtime/Engine/Private/Animation/AttributeBlendData.cpp b/Engine/Source/Runtime/Engine/Private/Animation/AttributeBlendData.cpp
--- a/Engine/Source/Runtime/Engine/Private/Animation/AttributeBlendData.cpp
+++ b/Engine/Source/Runtime/Engine/Private/Animation/AttributeBlendData.cpp
@@ -231,7 +231,7 @@ namespace UE { namespace Anim {
 		if (TypeIndex != INDEX_NONE)
 		{
 			const TArray<TWrappedAttribute<FAnimStackAllocator>>& ValuesArray = AttributeContainers.GetValues(TypeIndex);
-			const TArray<FAttributeId, FAnimStackAllocator>& AttributeIdentifiers = AttributeContainers.GetKeys(TypeIndex);
+			const TArray<FAttributeId>& AttributeIdentifiers = AttributeContainers.GetKeys(TypeIndex);
 			
 			for (int32 AttributeIndex = 0; AttributeIndex < AttributeIdentifiers.Num(); ++AttributeIndex)
 			{
diff --git a/Engine/Source/Runtime/Engine/Private/Animation/AttributesRuntime.cpp b/Engine/Source/Runtime/Engine/Private/Animation/AttributesRuntime.cpp
--- a/Engine/Source/Runtime/Engine/Private/Animation/AttributesRuntime.cpp
+++ b/Engine/Source/Runtime/Engine/Private/Animation/AttributesRuntime.cpp
@@ -481,8 +481,8 @@ void Attributes::MirrorAttributes(FStackAttributeContainer& CustomAttributes, co
 		if (TypeIndex != INDEX_NONE)
 		{
 			MirrorSets.Reset();
-			const TArray<int32, FAnimStackAllocator>& UniqueBoneIndices = CustomAttributes.GetUniqueTypedBoneIndices(TypeIndex);
-			TArray<int32, FAnimStackAllocator> SortedBoneIndices(UniqueBoneIndices); 
+			const TArray<int32>& UniqueBoneIndices = CustomAttributes.GetUniqueTypedBoneIndices(TypeIndex);
+			TArray<int32> SortedBoneIndices(UniqueBoneIndices); 
 			SortedBoneIndices.Sort();
 
 			for (int32 CurBoneIndex : UniqueBoneIndices)
@@ -499,7 +499,7 @@ void Attributes::MirrorAttributes(FStackAttributeContainer& CustomAttributes, co
 			}
 
 			TArray<TWrappedAttribute<FAnimStackAllocator>>& ValuesArray = FStackAttributeContainerAccessor::GetValues(CustomAttributes, TypeIndex);
-			const TArray<FAttributeId, FAnimStackAllocator>& AttributeIdentifiers = CustomAttributes.GetKeys(TypeIndex);
+			const TArray<FAttributeId>& AttributeIdentifiers = CustomAttributes.GetKeys(TypeIndex);
 
 			// gather attributes that are on mirrored bones 
 			for (int32 AttributeIndex = 0; AttributeIndex < AttributeIdentifiers.Num(); ++AttributeIndex)
diff --git a/Engine/Source/Runtime/Engine/Public/Animation/AttributesRuntime.h b/Engine/Source/Runtime/Engine/Public/Animation/AttributesRuntime.h
--- a/Engine/Source/Runtime/Engine/Public/Animation/AttributesRuntime.h
+++ b/Engine/Source/Runtime/Engine/Public/Animation/AttributesRuntime.h
@@ -45,7 +45,7 @@ namespace UE
         struct TAttributeContainerAccessor
 		{
 			static TArray<TWrappedAttribute<InAllocator>>& GetValues(TAttributeContainer<BoneIndexType, InAllocator>& Attributes,  int32 TypeIndex) {return Attributes.GetValuesInternal(TypeIndex);}
-			static TArray<FAttributeId, InAllocator>& GetKeys(TAttributeContainer<BoneIndexType, InAllocator>& Attributes,  int32 TypeIndex) {return Attributes.GetKeysInternal(TypeIndex);}
+			static TArray<FAttributeId>& GetKeys(TAttributeContainer<BoneIndexType, InAllocator>& Attributes, int32 TypeIndex) {return Attributes.GetKeysInternal(TypeIndex);}
 		};
 
 		struct ENGINE_API FStackAttributeContainerAccessor : public TAttributeContainerAccessor<FCompactPoseBoneIndex, FAnimStackAllocator> {};
