Fix DDC2 texture bugs encountered in editor when running with -DDC=Cold:
-One issue caused texture builds to be cancelled and re-issued repeatedly as the texture was using an empty comparison key to decide if it should rebuild or not.
-One issue caused an assert when we attempted to stream in mips for a texture that wasn't finished building yet.

diff --git a/Engine/Source/Runtime/Engine/Private/Streaming/Texture2DStreamIn_DDC.cpp b/Engine/Source/Runtime/Engine/Private/Streaming/Texture2DStreamIn_DDC.cpp
--- a/Engine/Source/Runtime/Engine/Private/Streaming/Texture2DStreamIn_DDC.cpp
+++ b/Engine/Source/Runtime/Engine/Private/Streaming/Texture2DStreamIn_DDC.cpp
@@ -173,7 +173,7 @@ void FTexture2DStreamIn_DDC::DoCreateAsyncDDCRequests(const FContext& Context)
 			if (MipKeys.Num())
 			{
 				GetCache().GetPayload(MipKeys, Context.Texture->GetPathName(), ECachePolicy::Default, DDCRequestOwner,
-					[this, &Context, LODBias](FCacheGetPayloadCompleteParams&& Params)
+					[this, LODBias](FCacheGetPayloadCompleteParams&& Params)
 					{
 						if (Params.Status == EStatus::Ok)
 						{
@@ -273,7 +273,6 @@ void FTexture2DStreamIn_DDC::DoLoadNewMipsFromDDC(const FContext& Context)
 					}
 					else
 					{
-						// UE_LOG(LogTexture, Warning, TEXT("Failed to stream mip data from the derived data cache for %s. Streaming mips will be recached."), Context.Texture->GetPathName() );
 						MarkAsCancelled();
 					}
 				}
@@ -293,19 +292,24 @@ void FTexture2DStreamIn_DDC::DoLoadNewMipsFromDDC(const FContext& Context)
 
 				if (MipMap.IsPagedToDerivedData())
 				{
-					check(DDCBuffers[MipIndex]);
-
-					const int32 ExpectedMipSize = CalcTextureMipMapSize(MipMap.SizeX, MipMap.SizeY, Context.Resource->GetPixelFormat(), 0);
-					if (DDCBuffers[MipIndex].GetSize() == ExpectedMipSize)
+					if (DDCBuffers[MipIndex])
 					{
-						FMemory::Memcpy(MipData[MipIndex], DDCBuffers[MipIndex].GetData(), DDCBuffers[MipIndex].GetSize());
+						const int32 ExpectedMipSize = CalcTextureMipMapSize(MipMap.SizeX, MipMap.SizeY, Context.Resource->GetPixelFormat(), 0);
+						if (DDCBuffers[MipIndex].GetSize() == ExpectedMipSize)
+						{
+							FMemory::Memcpy(MipData[MipIndex], DDCBuffers[MipIndex].GetData(), DDCBuffers[MipIndex].GetSize());
+						}
+						else
+						{
+							UE_LOG(LogTexture, Error, TEXT("DDC mip size (%d) not as expected (%d) for mip %d of %s."), static_cast<int32>(DDCBuffers[MipIndex].GetSize()), ExpectedMipSize, MipIndex, *Context.Texture->GetPathName());
+							MarkAsCancelled();
+						}
+						DDCBuffers[MipIndex].Reset();
 					}
 					else
 					{
-						UE_LOG(LogTexture, Error, TEXT("DDC mip size (%d) not as expected (%d) for mip %d of %s."), static_cast<int32>(DDCBuffers[MipIndex].GetSize()), ExpectedMipSize, MipIndex, *Context.Texture->GetPathName());
 						MarkAsCancelled();
 					}
-					DDCBuffers[MipIndex].Reset();
 				}
 				else
 				{
diff --git a/Engine/Source/Runtime/Engine/Private/TextureDerivedData.cpp b/Engine/Source/Runtime/Engine/Private/TextureDerivedData.cpp
--- a/Engine/Source/Runtime/Engine/Private/TextureDerivedData.cpp
+++ b/Engine/Source/Runtime/Engine/Private/TextureDerivedData.cpp
@@ -1874,18 +1874,21 @@ void UTexture::CachePlatformData(bool bAsyncCache, bool bAllowAsyncBuild, bool b
 			{
 				if (FTexturePlatformData::IsUsingNewDerivedData())
 				{
-					using namespace UE::DerivedData;
-					if (const FCacheKeyProxy* ExistingDerivedDataKey = PlatformDataLink->ComparisonDerivedDataKey.TryGet<FCacheKeyProxy>())
+					if (PlatformDataLink->IsAsyncWorkComplete())
 					{
-						if (*ExistingDerivedDataKey->AsCacheKey() != *CreateTextureCacheKeyProxy(*this, CacheFlags, BuildSettings[0]).AsCacheKey())
+						using namespace UE::DerivedData;
+						if (const FCacheKeyProxy* ExistingDerivedDataKey = PlatformDataLink->ComparisonDerivedDataKey.TryGet<FCacheKeyProxy>())
+						{
+							if (*ExistingDerivedDataKey->AsCacheKey() != *CreateTextureCacheKeyProxy(*this, CacheFlags, BuildSettings[0]).AsCacheKey())
+							{
+								bPerformCache = true;
+							}
+						}
+						else
 						{
 							bPerformCache = true;
 						}
 					}
-					else
-					{
-						bPerformCache = true;
-					}
 				}
 				else
 				{
