Turned the TypedElement.EnableSMInstanceElements cvar to be on by default.

Static mesh instance can now be selected and edited using a gizmo by default.

diff --git a/Engine/Source/Runtime/Engine/Private/Elements/SMInstance/SMInstanceElementData.cpp b/Engine/Source/Runtime/Engine/Private/Elements/SMInstance/SMInstanceElementData.cpp
--- a/Engine/Source/Runtime/Engine/Private/Elements/SMInstance/SMInstanceElementData.cpp
+++ b/Engine/Source/Runtime/Engine/Private/Elements/SMInstance/SMInstanceElementData.cpp
@@ -10,7 +10,7 @@ UE_DEFINE_TYPED_ELEMENT_DATA_RTTI(FSMInstanceElementData);
 namespace SMInstanceElementDataUtil
 {
 
-static int32 GEnableSMInstanceElements = 0;
+static int32 GEnableSMInstanceElements = 1;
 static FAutoConsoleVariableRef CVarEnableSMInstanceElements(
 	TEXT("TypedElements.EnableSMInstanceElements"),
 	GEnableSMInstanceElements,
