- Fix mip calculation for none PoT textures

#rnx

diff --git a/Engine/Source/Runtime/Engine/Private/TextureRenderTarget2D.cpp b/Engine/Source/Runtime/Engine/Private/TextureRenderTarget2D.cpp
--- a/Engine/Source/Runtime/Engine/Private/TextureRenderTarget2D.cpp
+++ b/Engine/Source/Runtime/Engine/Private/TextureRenderTarget2D.cpp
@@ -45,7 +45,7 @@ FTextureResource* UTextureRenderTarget2D::CreateResource()
 
 	if (bAutoGenerateMips)
 	{
-		NumMips = FGenericPlatformMath::CeilToInt(FGenericPlatformMath::Log2(static_cast<float>(FGenericPlatformMath::Max(SizeX, SizeY)))) + 1;
+		NumMips = FMath::FloorLog2(FMath::Max(SizeX, SizeY)) + 1;
 
 		if (RHIRequiresComputeGenerateMips())
 		{
@@ -127,7 +127,7 @@ void UTextureRenderTarget2D::ResizeTarget(uint32 InSizeX, uint32 InSizeY)
 		SizeY = InSizeY;
 		if (bAutoGenerateMips)
 		{
-			NumMips = FGenericPlatformMath::CeilToInt(FGenericPlatformMath::Log2(static_cast<float>(FGenericPlatformMath::Max(SizeX, SizeY)))) + 1;
+			NumMips = FMath::FloorLog2(FMath::Max(SizeX, SizeY)) + 1;
 		}
 
 		if (GetResource())
