Further bug fixing for texture build flow for issues encountered when testing Reverb dataset in editor with -DDC=Cold:
-Fixed a bug where virtual textures were having their comparison key compared against the DDC2 expected comparison key even though they're still built with the legacy path.
-Fixed a bug with the cleanup of slow tasks when instigating blocking texture builds on the game thread.
-Fixed a bugs caused by late construction of comparison keys for texture derived data.  Comparison keys have been changed and will now always be computed serially at the time a task is constructed.

diff --git a/Engine/Source/Runtime/Engine/Classes/Engine/Texture.h b/Engine/Source/Runtime/Engine/Classes/Engine/Texture.h
--- a/Engine/Source/Runtime/Engine/Classes/Engine/Texture.h
+++ b/Engine/Source/Runtime/Engine/Classes/Engine/Texture.h
@@ -637,8 +637,24 @@ struct FTexturePlatformData
 	/** The key associated with this derived data. */
 	TVariant<FString, UE::DerivedData::FCacheKeyProxy> DerivedDataKey;
 
+	struct FStructuredDerivedDataKey
+	{
+		FIoHash BuildDefinitionKey;
+		FGuid SourceGuid;
+		FGuid CompositeSourceGuid;
+
+		bool operator==(const FStructuredDerivedDataKey& Other) const
+		{
+			return BuildDefinitionKey == Other.BuildDefinitionKey && SourceGuid == Other.SourceGuid && CompositeSourceGuid == Other.CompositeSourceGuid;
+		}
+
+		bool operator!=(const FStructuredDerivedDataKey& Other) const
+		{
+			return BuildDefinitionKey != Other.BuildDefinitionKey || SourceGuid != Other.SourceGuid || CompositeSourceGuid != Other.CompositeSourceGuid;
+		}
+	};
 	/** The key used when comparing this derived data key to another reference key. */
-	TVariant<FString, UE::DerivedData::FCacheKeyProxy> ComparisonDerivedDataKey;
+	TVariant<FString, FStructuredDerivedDataKey> ComparisonDerivedDataKey;
 
 	/** Async cache task if one is outstanding. */
 	struct FTextureAsyncCacheDerivedDataTask* AsyncTask;
diff --git a/Engine/Source/Runtime/Engine/Private/TextureDerivedData.cpp b/Engine/Source/Runtime/Engine/Private/TextureDerivedData.cpp
--- a/Engine/Source/Runtime/Engine/Private/TextureDerivedData.cpp
+++ b/Engine/Source/Runtime/Engine/Private/TextureDerivedData.cpp
@@ -1872,23 +1872,21 @@ void UTexture::CachePlatformData(bool bAsyncCache, bool bAllowAsyncBuild, bool b
 			}
 			else
 			{
-				if (FTexturePlatformData::IsUsingNewDerivedData())
+				if (FTexturePlatformData::IsUsingNewDerivedData() && (Source.GetNumLayers() == 1) && !BuildSettings[0].bVirtualStreamable)
 				{
-					if (PlatformDataLink->IsAsyncWorkComplete())
+					
+					using namespace UE::DerivedData;
+					if (const FTexturePlatformData::FStructuredDerivedDataKey* ExistingDerivedDataKey = PlatformDataLink->ComparisonDerivedDataKey.TryGet<FTexturePlatformData::FStructuredDerivedDataKey>())
 					{
-						using namespace UE::DerivedData;
-						if (const FCacheKeyProxy* ExistingDerivedDataKey = PlatformDataLink->ComparisonDerivedDataKey.TryGet<FCacheKeyProxy>())
-						{
-							if (*ExistingDerivedDataKey->AsCacheKey() != *CreateTextureCacheKeyProxy(*this, CacheFlags, BuildSettings[0]).AsCacheKey())
-							{
-								bPerformCache = true;
-							}
-						}
-						else
+						if (*ExistingDerivedDataKey != CreateTextureDerivedDataKey(*this, CacheFlags, BuildSettings[0]))
 						{
 							bPerformCache = true;
 						}
 					}
+					else
+					{
+						bPerformCache = true;
+					}
 				}
 				else
 				{
diff --git a/Engine/Source/Runtime/Engine/Private/TextureDerivedDataTask.cpp b/Engine/Source/Runtime/Engine/Private/TextureDerivedDataTask.cpp
--- a/Engine/Source/Runtime/Engine/Private/TextureDerivedDataTask.cpp
+++ b/Engine/Source/Runtime/Engine/Private/TextureDerivedDataTask.cpp
@@ -590,6 +590,12 @@ FTextureCacheDerivedDataWorker::FTextureCacheDerivedDataWorker(
 	}
 	check(Texture.Source.GetId().IsValid());
 
+	FString LocalDerivedDataKeySuffix;
+	FString LocalDerivedDataKey;
+	GetTextureDerivedDataKeySuffix(Texture, BuildSettingsPerLayer.GetData(), LocalDerivedDataKeySuffix);
+	GetTextureDerivedDataKeyFromSuffix(LocalDerivedDataKeySuffix, LocalDerivedDataKey);
+	DerivedData->ComparisonDerivedDataKey.Emplace<FString>(LocalDerivedDataKey);
+
 	// Dump any existing mips.
 	DerivedData->Mips.Empty();
 	if (DerivedData->VTData)
@@ -672,10 +678,8 @@ void FTextureCacheDerivedDataWorker::DoWork()
 
 	FString LocalDerivedDataKeySuffix;
 	FString LocalDerivedDataKey;
-	FString LocalComparisonDerivedDataKey;
 	GetTextureDerivedDataKeySuffix(Texture, BuildSettingsPerLayer.GetData(), LocalDerivedDataKeySuffix);
 	GetTextureDerivedDataKeyFromSuffix(LocalDerivedDataKeySuffix, LocalDerivedDataKey);
-	LocalComparisonDerivedDataKey = LocalDerivedDataKey;
 	if (!bForceRebuild)
 	{
 		// First try to load a texture generated for the shipping build from the cache.
@@ -711,7 +715,6 @@ void FTextureCacheDerivedDataWorker::DoWork()
 	}
 	KeySuffix = LocalDerivedDataKeySuffix;
 	DerivedData->DerivedDataKey.Emplace<FString>(LocalDerivedDataKey);
-	DerivedData->ComparisonDerivedDataKey.Emplace<FString>(LocalComparisonDerivedDataKey);
 
 	if (bLoadedFromDDC)
 	{
@@ -958,7 +961,8 @@ public:
 		}
 
 		FBuildDefinition Definition = CreateDefinition(Build, Texture, TexturePath, FunctionName, Settings, bUseCompositeTexture);
-		
+		DerivedData.ComparisonDerivedDataKey.Emplace<FTexturePlatformData::FStructuredDerivedDataKey>(GetKey(Definition, Texture, bUseCompositeTexture));
+
 		if (!EnumHasAnyFlags(Flags, ETextureCacheFlags::ForceRebuild) && Settings.FastTextureEncode == ETextureFastEncode::TryOffEncodeFast)
 		{
 			FTextureBuildSettings ShippingSettings = Settings;
@@ -971,11 +975,6 @@ public:
 					{
 					default:
 					case EStatus::Ok:
-						BuildSession.Get().Build(Definition, EBuildPolicy::None, *Owner,
-							[this](FBuildCompleteParams&& Params)
-							{
-								DerivedData.ComparisonDerivedDataKey.Emplace<FCacheKeyProxy>(Params.CacheKey);
-							});
 						return EndBuild(Params.CacheKey, MoveTemp(Params.Output), Params.BuildStatus);
 					case EStatus::Error:
 						return BeginBuild(Definition, Flags);
@@ -1018,7 +1017,6 @@ public:
 		BuildSession.Get().Build(Definition, BuildPolicy, *Owner,
 			[this](FBuildCompleteParams&& Params)
 			{
-				DerivedData.ComparisonDerivedDataKey.Emplace<FCacheKeyProxy>(Params.CacheKey);
 				EndBuild(Params.CacheKey, MoveTemp(Params.Output), Params.BuildStatus);
 			});
 	}
@@ -1031,6 +1029,7 @@ public:
 		BuildOutputSize = Algo::TransformAccumulate(Output.GetPayloads(),
 			[](const FPayload& Payload) { return Payload.GetData().GetRawSize(); }, uint64(0));
 		WriteDerivedData(MoveTemp(Output));
+		StatusMessage.Reset();
 	}
 
 	void Finalize(bool& bOutFoundInCache, uint64& OutProcessedByteCount) final
@@ -1196,6 +1195,18 @@ public:
 		return true;
 	}
 
+	static FTexturePlatformData::FStructuredDerivedDataKey GetKey(const UE::DerivedData::FBuildDefinition& BuildDefinition, const UTexture& Texture, bool bUseCompositeTexture)
+	{
+		FTexturePlatformData::FStructuredDerivedDataKey Key;
+		Key.BuildDefinitionKey = BuildDefinition.GetKey().Hash;
+		Key.SourceGuid = Texture.Source.GetId();
+		if (bUseCompositeTexture && Texture.CompositeTexture)
+		{
+			Key.CompositeSourceGuid = Texture.CompositeTexture->Source.GetId();
+		}
+		return Key;
+	}
+
 private:
 
 	static bool DeserializeTextureFromPayloads(FTexturePlatformData& DerivedData, const UE::DerivedData::FBuildOutput& Output, int32 FirstMipToLoad, bool bInlineMips)
@@ -1404,7 +1415,7 @@ FTextureAsyncCacheDerivedDataTask* CreateTextureBuildTask(
 	return nullptr;
 }
 
-UE::DerivedData::FCacheKeyProxy CreateTextureCacheKeyProxy(
+FTexturePlatformData::FStructuredDerivedDataKey CreateTextureDerivedDataKey(
 	UTexture& Texture,
 	ETextureCacheFlags CacheFlags,
 	const FTextureBuildSettings& Settings)
@@ -1422,26 +1433,12 @@ UE::DerivedData::FCacheKeyProxy CreateTextureCacheKeyProxy(
 		bool bUseCompositeTexture = false;
 		if (FTextureBuildTask::IsTextureValidForBuilding(Texture, CacheFlags, bUseCompositeTexture))
 		{
-			TOptional<UE::TextureDerivedData::FTextureBuildInputResolver> PlaceholderResolver;
-			UE::DerivedData::IBuildInputResolver* InputResolver = GetGlobalBuildInputResolver();
-			if (!InputResolver)
-			{
-				PlaceholderResolver.Emplace(Texture);
-				InputResolver = &PlaceholderResolver.GetValue();
-			}
-			FBuildSession Session = Build.CreateSession(TexturePath, InputResolver);
 			FBuildDefinition Definition = FTextureBuildTask::CreateDefinition(Build, Texture, TexturePath, FunctionName, Settings, bUseCompositeTexture);
-			
-			FCacheKey GeneratedKey;
-			FRequestOwner CacheKeyGenerationOwner(EPriority::Blocking);
-			Session.Build(Definition, EBuildPolicy::None, CacheKeyGenerationOwner,
-				[&GeneratedKey](FBuildCompleteParams&& Params) { GeneratedKey = Params.CacheKey; });
-			CacheKeyGenerationOwner.Wait();
 
-			return GeneratedKey;
+			return FTextureBuildTask::GetKey(Definition, Texture, bUseCompositeTexture);
 		}
 	}
-	return FCacheKey::Empty;
+	return FTexturePlatformData::FStructuredDerivedDataKey();
 }
 
 #endif // WITH_EDITOR
diff --git a/Engine/Source/Runtime/Engine/Private/TextureDerivedDataTask.h b/Engine/Source/Runtime/Engine/Private/TextureDerivedDataTask.h
--- a/Engine/Source/Runtime/Engine/Private/TextureDerivedDataTask.h
+++ b/Engine/Source/Runtime/Engine/Private/TextureDerivedDataTask.h
@@ -260,7 +260,7 @@ FTextureAsyncCacheDerivedDataTask* CreateTextureBuildTask(
 	EQueuedWorkPriority Priority,
 	ETextureCacheFlags Flags);
 
-UE::DerivedData::FCacheKeyProxy CreateTextureCacheKeyProxy(
+FTexturePlatformData::FStructuredDerivedDataKey CreateTextureDerivedDataKey(
 	UTexture& Texture,
 	ETextureCacheFlags CacheFlags,
 	const FTextureBuildSettings& Settings);
