Fix bug involving duplication of actor components not causing their TObjectPtr fields to be duplicated in a manner equivalent to raw pointers. Fixed by ensuring ithat the custom archive (not derived from FArchiveUObject) used during serialization in this actor component code path handles FObjectPtr fields in the standard manner.

diff --git a/Engine/Source/Runtime/Engine/Private/ComponentInstanceDataCache.cpp b/Engine/Source/Runtime/Engine/Private/ComponentInstanceDataCache.cpp
--- a/Engine/Source/Runtime/Engine/Private/ComponentInstanceDataCache.cpp
+++ b/Engine/Source/Runtime/Engine/Private/ComponentInstanceDataCache.cpp
@@ -150,6 +150,11 @@ public:
 		return *this;
 	}
 
+	virtual FArchive& operator<<(FObjectPtr& ObjectPtr) override
+	{
+		return FArchiveUObject::SerializeObjectPtr(*this, ObjectPtr);
+	}
+
 private:
 	const UActorComponent* Component;
 	FActorComponentInstanceData& ActorInstanceData;
@@ -203,6 +208,11 @@ public:
 		return *this;
 	}
 
+	virtual FArchive& operator<<(FObjectPtr& ObjectPtr) override
+	{
+		return FArchiveUObject::SerializeObjectPtr(*this, ObjectPtr);
+	}
+
 	FActorComponentInstanceData& ActorInstanceData;
 	TSet<const FProperty*> PropertiesToSkip;
 };
