Make UPrimitiveComponent::OverlapComponent const correct.
Left the old version of the method but added the 'final' keyword to make sure that class that overloaded the old method signature will fail to compile.
Had to comment out the UE_DEPRECATE macros to avoid false positives.

Reported on UDN: https://epicgames.lightning.force.com/lightning/r/Case/5004z00001YR9JqAAL/view

diff --git a/Engine/Source/Runtime/Engine/Classes/Components/PrimitiveComponent.h b/Engine/Source/Runtime/Engine/Classes/Components/PrimitiveComponent.h
--- a/Engine/Source/Runtime/Engine/Classes/Components/PrimitiveComponent.h
+++ b/Engine/Source/Runtime/Engine/Classes/Components/PrimitiveComponent.h
@@ -2490,7 +2490,13 @@ public:
 	 *  @param  CollisionShape 	Shape of collision of PrimComp geometry
 	 *  @return true if PrimComp overlaps this component at the specified location/rotation
 	 */
-	virtual bool OverlapComponent(const FVector& Pos, const FQuat& Rot, const FCollisionShape& CollisionShape);
+
+    //UE_DEPRECATED(5.0, "Use the const version of OverlapComponent. This deprecation cannot be uncommented as it would produce false positives." )
+	virtual bool OverlapComponent(const FVector& Pos, const FQuat& Rot, const FCollisionShape& CollisionShape) final 
+	{
+		return const_cast<const UPrimitiveComponent*>(this)->OverlapComponent(Pos, Rot, CollisionShape);
+	}
+	virtual bool OverlapComponent(const FVector& Pos, const FQuat& Rot, const FCollisionShape& CollisionShape) const;
 
 	/**
 	 * Computes the minimum translation direction (MTD) when an overlap exists between the component and the given shape.
diff --git a/Engine/Source/Runtime/Engine/Classes/Components/SkeletalMeshComponent.h b/Engine/Source/Runtime/Engine/Classes/Components/SkeletalMeshComponent.h
--- a/Engine/Source/Runtime/Engine/Classes/Components/SkeletalMeshComponent.h
+++ b/Engine/Source/Runtime/Engine/Classes/Components/SkeletalMeshComponent.h
@@ -1759,7 +1759,7 @@ public:
 	 */
 	 virtual bool SweepComponent( FHitResult& OutHit, const FVector Start, const FVector End, const FQuat& ShapRotation, const FCollisionShape& CollisionShape, bool bTraceComplex=false) override;
 	
-	virtual bool OverlapComponent(const FVector& Pos, const FQuat& Rot, const FCollisionShape& CollisionShape) override;
+	virtual bool OverlapComponent(const FVector& Pos, const FQuat& Rot, const FCollisionShape& CollisionShape) const override;
 	virtual void SetSimulatePhysics(bool bEnabled) override;
 	virtual void AddRadialImpulse(FVector Origin, float Radius, float Strength, ERadialImpulseFalloff Falloff, bool bVelChange=false) override;
 	virtual void AddRadialForce(FVector Origin, float Radius, float Strength, ERadialImpulseFalloff Falloff, bool bAccelChange=false) override;
diff --git a/Engine/Source/Runtime/Engine/Private/Components/PrimitiveComponent.cpp b/Engine/Source/Runtime/Engine/Private/Components/PrimitiveComponent.cpp
--- a/Engine/Source/Runtime/Engine/Private/Components/PrimitiveComponent.cpp
+++ b/Engine/Source/Runtime/Engine/Private/Components/PrimitiveComponent.cpp
@@ -2731,7 +2731,7 @@ bool UPrimitiveComponent::ComponentOverlapComponentImpl(class UPrimitiveComponen
 }
 
 
-bool UPrimitiveComponent::OverlapComponent(const FVector& Pos, const FQuat& Rot, const struct FCollisionShape& CollisionShape)
+bool UPrimitiveComponent::OverlapComponent(const FVector& Pos, const FQuat& Rot, const struct FCollisionShape& CollisionShape) const
 {
 	return BodyInstance.OverlapTest(Pos, Rot, CollisionShape);
 }
diff --git a/Engine/Source/Runtime/Engine/Private/SkeletalMeshComponentPhysics.cpp b/Engine/Source/Runtime/Engine/Private/SkeletalMeshComponentPhysics.cpp
--- a/Engine/Source/Runtime/Engine/Private/SkeletalMeshComponentPhysics.cpp
+++ b/Engine/Source/Runtime/Engine/Private/SkeletalMeshComponentPhysics.cpp
@@ -2611,7 +2611,7 @@ bool USkeletalMeshComponent::ComponentOverlapComponentImpl(class UPrimitiveCompo
 	return false;
 }
 
-bool USkeletalMeshComponent::OverlapComponent(const FVector& Pos, const FQuat& Rot, const FCollisionShape& CollisionShape)
+bool USkeletalMeshComponent::OverlapComponent(const FVector& Pos, const FQuat& Rot, const FCollisionShape& CollisionShape) const
 {
 	for (FBodyInstance* Body : Bodies)
 	{
