Duplicate non-asset material override instances when generating instanced HLODs to avoid creating inter-level references.

diff --git a/Engine/Source/Runtime/Engine/Private/ISMPartition/ISMComponentDescriptor.cpp b/Engine/Source/Runtime/Engine/Private/ISMPartition/ISMComponentDescriptor.cpp
--- a/Engine/Source/Runtime/Engine/Private/ISMPartition/ISMComponentDescriptor.cpp
+++ b/Engine/Source/Runtime/Engine/Private/ISMPartition/ISMComponentDescriptor.cpp
@@ -160,7 +160,20 @@ UInstancedStaticMeshComponent* FISMComponentDescriptor::CreateComponent(UObject*
 void FISMComponentDescriptor::InitComponent(UInstancedStaticMeshComponent* ISMComponent) const
 {
 	ISMComponent->SetStaticMesh(StaticMesh);
-	ISMComponent->OverrideMaterials = OverrideMaterials;
+
+	ISMComponent->OverrideMaterials.Empty(OverrideMaterials.Num());
+	for (UMaterialInterface* OverrideMaterial : OverrideMaterials)
+	{
+		if (OverrideMaterial && !OverrideMaterial->IsAsset())
+		{
+			// As override materials are normally outered to their owner component, we need to duplicate them here to make sure we don't create
+			// references to actors in other levels (for packed level instances or HLOD actors).
+			OverrideMaterial = DuplicateObject<UMaterialInterface>(OverrideMaterial, ISMComponent);
+		}
+
+		ISMComponent->OverrideMaterials.Add(OverrideMaterial);
+	}
+
 	ISMComponent->Mobility = Mobility;
 	ISMComponent->RuntimeVirtualTextures = RuntimeVirtualTextures;
 	ISMComponent->VirtualTextureRenderPassType = VirtualTextureRenderPassType;
