AudioThread task priority set to BackgroundNormal (was Normal by mistake, originally AudioThread was below normal priority before it was replaced by a pipe), added a cvar to control it in runtime

diff --git a/Engine/Source/Runtime/Core/Public/Async/Fundamental/Task.h b/Engine/Source/Runtime/Core/Public/Async/Fundamental/Task.h
--- a/Engine/Source/Runtime/Core/Public/Async/Fundamental/Task.h
+++ b/Engine/Source/Runtime/Core/Public/Async/Fundamental/Task.h
@@ -23,6 +23,59 @@ namespace LowLevelTasks
 		Count
 	};
 
+	inline const TCHAR* ToString(ETaskPriority Priority)
+	{
+		if (Priority < ETaskPriority::High || Priority >= ETaskPriority::Count)
+		{
+			return nullptr;
+		}
+
+		const TCHAR* TaskPriorityToStr[] =
+		{
+			TEXT("High"),
+			TEXT("Normal"),
+			TEXT("BackgroundHigh"),
+			TEXT("BackgroundNormal"),
+			TEXT("BackgroundLow")
+		};
+		return TaskPriorityToStr[(int32)Priority];
+	}
+
+	inline bool ToTaskPriority(const TCHAR* PriorityStr, ETaskPriority& OutPriority)
+	{
+		if (FCString::Stricmp(PriorityStr, ToString(ETaskPriority::High)) == 0)
+		{
+			OutPriority = ETaskPriority::High;
+			return true;
+		}
+
+		if (FCString::Stricmp(PriorityStr, ToString(ETaskPriority::Normal)) == 0)
+		{
+			OutPriority = ETaskPriority::Normal;
+			return true;
+		}
+
+		if (FCString::Stricmp(PriorityStr, ToString(ETaskPriority::BackgroundHigh)) == 0)
+		{
+			OutPriority = ETaskPriority::BackgroundHigh;
+			return true;
+		}
+
+		if (FCString::Stricmp(PriorityStr, ToString(ETaskPriority::BackgroundNormal)) == 0)
+		{
+			OutPriority = ETaskPriority::BackgroundNormal;
+			return true;
+		}
+
+		if (FCString::Stricmp(PriorityStr, ToString(ETaskPriority::BackgroundLow)) == 0)
+		{
+			OutPriority = ETaskPriority::BackgroundLow;
+			return true;
+		}
+
+		return false;
+	}
+
 	enum class ETaskState : int32
 	{
 		CanceledAndCompleted,	//means the task is completed with execution of it's continuation but the runnable was cancelled
diff --git a/Engine/Source/Runtime/Engine/Private/AudioThread.cpp b/Engine/Source/Runtime/Engine/Private/AudioThread.cpp
--- a/Engine/Source/Runtime/Engine/Private/AudioThread.cpp
+++ b/Engine/Source/Runtime/Engine/Private/AudioThread.cpp
@@ -84,6 +84,36 @@ struct FAudioThreadInteractor
 	}
 };
 
+UE::Tasks::ETaskPriority GAudioTaskPriority = UE::Tasks::ETaskPriority::BackgroundNormal;
+
+static void SetAudioTaskPriority(const TArray<FString>& Args)
+{
+	UE_LOG(LogConsoleResponse, Display, TEXT("AudioTaskPriority was %s."), LowLevelTasks::ToString(GAudioTaskPriority));
+
+	if (Args.Num() > 1)
+	{
+		UE_LOG(LogConsoleResponse, Display, TEXT("WARNING: This command requires a single argument while %d were provided, all extra arguments will be ignored."), Args.Num());
+	}
+	else if (Args.IsEmpty())
+	{
+		UE_LOG(LogConsoleResponse, Display, TEXT("ERROR: Please provide a new priority value."));
+		return;
+	}
+
+	if (!LowLevelTasks::ToTaskPriority(*Args[0], GAudioTaskPriority))
+	{
+		UE_LOG(LogConsoleResponse, Display, TEXT("ERROR: Invalid priority: %s."), *Args[0]);
+	}
+
+	UE_LOG(LogConsoleResponse, Display, TEXT("Audio Task Priority was set to %s."), LowLevelTasks::ToString(GAudioTaskPriority));
+}
+
+static FAutoConsoleCommand AudioThreadPriorityConsoleCommand(
+	TEXT("AudioThread.TaskPriority"),
+	TEXT("Takes a single parameter of value `High`, `Normal`, `BackgroundHigh`, `BackgroundNormal` or `BackgroundLow`."),
+	FConsoleCommandWithArgsDelegate::CreateStatic(&SetAudioTaskPriority)
+);
+
 static FAutoConsoleVariableSink CVarUseAudioThreadSink(FConsoleCommandDelegate::CreateStatic(&FAudioThreadInteractor::UseAudioThreadCVarSinkFunction));
 
 bool FAudioThread::bUseThreadedAudio = false;
@@ -464,7 +494,7 @@ void FAudioThread::RunCommandOnAudioThread(TUniqueFunction<void()> InFunction, c
 	// we are on an unknown thread
 	else if (IsUsingThreadedAudio())
 	{
-		GAudioPipe.Launch(TEXT("AudioCommand"), MoveTemp(CommandWrapper));
+		GAudioPipe.Launch(TEXT("AudioCommand"), MoveTemp(CommandWrapper), GAudioTaskPriority);
 	}
 	else
 	{
