UE-125091 - Blueprint did not automatically create float>double conversion nodes on load

Any old content that's attempting to create a quat from float data will run into a BP compilation error since FQuat expects doubles now. The easiet fix is to create custom Make/Break functions for FQuat in the KismetMathLibrary, which will perform the implicit cast for us.

Once BPs are using doubles everywhere, we can forego this option. UE-126734 has been created as a follow-up task that will ensure we're using doubles and avoiding unnecessary casts in KismetMathLibrary.

diff --git a/Engine/Source/Runtime/CoreUObject/Public/UObject/NoExportTypes.h b/Engine/Source/Runtime/CoreUObject/Public/UObject/NoExportTypes.h
--- a/Engine/Source/Runtime/CoreUObject/Public/UObject/NoExportTypes.h
+++ b/Engine/Source/Runtime/CoreUObject/Public/UObject/NoExportTypes.h
@@ -638,7 +638,7 @@ struct FQuat4d
  * Quaternion.
  * @note The full C++ class is located here: Engine\Source\Runtime\Core\Public\Math\Quat.h
  */
-USTRUCT(immutable, noexport, BlueprintType)
+USTRUCT(immutable, noexport, BlueprintType, meta=(HasNativeMake ="Engine.KismetMathLibrary.MakeQuat", HasNativeBreak="Engine.KismetMathLibrary.BreakQuat"))
 struct FQuat
 {
 	UPROPERTY(EditAnywhere, BlueprintReadWrite, Category = Quat, SaveGame)
diff --git a/Engine/Source/Runtime/Engine/Classes/Kismet/KismetMathLibrary.h b/Engine/Source/Runtime/Engine/Classes/Kismet/KismetMathLibrary.h
--- a/Engine/Source/Runtime/Engine/Classes/Kismet/KismetMathLibrary.h
+++ b/Engine/Source/Runtime/Engine/Classes/Kismet/KismetMathLibrary.h
@@ -2744,6 +2744,14 @@ class ENGINE_API UKismetMathLibrary : public UBlueprintFunctionLibrary
 	UFUNCTION(BlueprintPure, meta = (DisplayName = "Quat - Quat", CompactNodeTitle = "-", ScriptMethod = "Subtract", ScriptOperator = "-;-=", Keywords = "- subtract minus"), Category = "Math|Quat")
 	static FQuat Subtract_QuatQuat(const FQuat& A, const FQuat& B);
 
+	/** Makes a quat {X, Y, Z, W} */
+	UFUNCTION(BlueprintPure, Category = "Math|Quat", meta = (Keywords = "construct build", NativeMakeFunc))
+	static FQuat MakeQuat(float X, float Y, float Z, float W);
+
+	/** Breaks a quat apart into X, Y, Z, W */
+	UFUNCTION(BlueprintPure, Category = "Math|Quat", meta = (NativeBreakFunc))
+	static void BreakQuat(const FQuat& InQuat, float& X, float& Y, float& Z, float& W);
+
 	/**
 	 * Gets the result of multiplying two quaternions (A * B).
 	 *
diff --git a/Engine/Source/Runtime/Engine/Classes/Kismet/KismetMathLibrary.inl b/Engine/Source/Runtime/Engine/Classes/Kismet/KismetMathLibrary.inl
--- a/Engine/Source/Runtime/Engine/Classes/Kismet/KismetMathLibrary.inl
+++ b/Engine/Source/Runtime/Engine/Classes/Kismet/KismetMathLibrary.inl
@@ -3202,6 +3202,20 @@ FQuat UKismetMathLibrary::Multiply_QuatQuat(const FQuat& A, const FQuat& B)
 	return A * B;
 }
 
+KISMET_MATH_FORCEINLINE
+FQuat UKismetMathLibrary::MakeQuat(float X, float Y, float Z, float W)
+{
+	return FQuat(X, Y, Z, W);
+}
+
+KISMET_MATH_FORCEINLINE
+void UKismetMathLibrary::BreakQuat(const FQuat& InQuat, float& X, float& Y, float& Z, float& W)
+{
+	X = InQuat.X;
+	Y = InQuat.Y;
+	Z = InQuat.Z;
+	W = InQuat.W;
+}
 
 KISMET_MATH_FORCEINLINE
 bool UKismetMathLibrary::Quat_IsIdentity(const FQuat& Q, float Tolerance)
