Add PathTracingQualitySwitch node

This shading node works very similarly to the RayTracingQualitySwitch, but only picks up its alternate port when compiling the material's path tracing permutation. This allows artists to customize the behavior of materials for the path tracer in cases where hacks are necessary for regular materials that are not needed on the path tracing side (glass in particular).

diff --git a/Engine/Shaders/Private/Common.ush b/Engine/Shaders/Private/Common.ush
--- a/Engine/Shaders/Private/Common.ush
+++ b/Engine/Shaders/Private/Common.ush
@@ -1684,6 +1684,15 @@ bool GetRayTracingQualitySwitch()
 #endif
 }
 
+bool GetPathTracingQualitySwitch()
+{
+#if RAYHITGROUPSHADER && PATH_TRACING
+	return true;
+#else
+	return false;
+#endif
+}
+
 // Used for the RuntimeVirtualTextureReplace material expression
 // Returns true if rendering to a runtime virtual texture page
 bool GetRuntimeVirtualTextureOutputSwitch()
diff --git a/Engine/Source/Runtime/Engine/Classes/Materials/MaterialExpressionPathTracingQualitySwitch.h b/Engine/Source/Runtime/Engine/Classes/Materials/MaterialExpressionPathTracingQualitySwitch.h
new file mode 100644
--- /dev/null
+++ b/Engine/Source/Runtime/Engine/Classes/Materials/MaterialExpressionPathTracingQualitySwitch.h
@@ -0,0 +1,33 @@
+// Copyright Epic Games, Inc. All Rights Reserved.
+
+#pragma once
+
+#include "CoreMinimal.h"
+#include "UObject/ObjectMacros.h"
+#include "MaterialExpressionIO.h"
+#include "Materials/MaterialExpression.h"
+#include "MaterialExpressionPathTracingQualitySwitch.generated.h"
+
+UCLASS()
+class UMaterialExpressionPathTracingQualitySwitch : public UMaterialExpression
+{
+	GENERATED_UCLASS_BODY()
+
+		/** Used for non-path traced modes */
+		UPROPERTY()
+		FExpressionInput Normal;
+
+	/** Used by the path tracer */
+	UPROPERTY()
+		FExpressionInput PathTraced;
+
+	//~ Begin UMaterialExpression Interface
+#if WITH_EDITOR
+	virtual int32 Compile(class FMaterialCompiler* Compiler, int32 OutputIndex) override;
+	virtual void GetCaption(TArray<FString>& OutCaptions) const override;
+	virtual bool IsResultMaterialAttributes(int32 OutputIndex) override;
+	virtual uint32 GetInputType(int32 InputIndex) override;
+	virtual uint32 GetOutputType(int32 OutputIndex) override { return MCT_Unknown; }
+#endif
+	//~ End UMaterialExpression Interface
+};
diff --git a/Engine/Source/Runtime/Engine/Private/Materials/HLSLMaterialTranslator.cpp b/Engine/Source/Runtime/Engine/Private/Materials/HLSLMaterialTranslator.cpp
--- a/Engine/Source/Runtime/Engine/Private/Materials/HLSLMaterialTranslator.cpp
+++ b/Engine/Source/Runtime/Engine/Private/Materials/HLSLMaterialTranslator.cpp
@@ -8137,6 +8137,12 @@ int32 FHLSLMaterialTranslator::RayTracingQualitySwitchReplace(int32 Normal, int3
 	return GenericSwitch(TEXT("GetRayTracingQualitySwitch()"), RayTraced, Normal);
 }
 
+int32 FHLSLMaterialTranslator::PathTracingQualitySwitchReplace(int32 Normal, int32 PathTraced)
+{
+	return GenericSwitch(TEXT("GetPathTracingQualitySwitch()"), PathTraced, Normal);
+}
+
+
 int32 FHLSLMaterialTranslator::VirtualTextureOutputReplace(int32 Default, int32 VirtualTexture)
 {
 	return GenericSwitch(TEXT("GetRuntimeVirtualTextureOutputSwitch()"), VirtualTexture, Default);
diff --git a/Engine/Source/Runtime/Engine/Private/Materials/HLSLMaterialTranslator.h b/Engine/Source/Runtime/Engine/Private/Materials/HLSLMaterialTranslator.h
--- a/Engine/Source/Runtime/Engine/Private/Materials/HLSLMaterialTranslator.h
+++ b/Engine/Source/Runtime/Engine/Private/Materials/HLSLMaterialTranslator.h
@@ -889,7 +889,8 @@ protected:
 	virtual int32 ShadowReplace(int32 Default, int32 Shadow) override;
 	virtual int32 ReflectionCapturePassSwitch(int32 Default, int32 Reflection) override;
 
-	virtual int32 RayTracingQualitySwitchReplace(int32 Normal, int32 RayTraced);
+	virtual int32 RayTracingQualitySwitchReplace(int32 Normal, int32 RayTraced) override;
+	virtual int32 PathTracingQualitySwitchReplace(int32 Normal, int32 PathTraced) override;
 
 	virtual int32 VirtualTextureOutputReplace(int32 Default, int32 VirtualTexture) override;
 
diff --git a/Engine/Source/Runtime/Engine/Private/Materials/MaterialExpressions.cpp b/Engine/Source/Runtime/Engine/Private/Materials/MaterialExpressions.cpp
--- a/Engine/Source/Runtime/Engine/Private/Materials/MaterialExpressions.cpp
+++ b/Engine/Source/Runtime/Engine/Private/Materials/MaterialExpressions.cpp
@@ -110,6 +110,7 @@
 #include "Materials/MaterialExpressionFunctionOutput.h"
 #include "Materials/MaterialExpressionGIReplace.h"
 #include "Materials/MaterialExpressionRayTracingQualitySwitch.h"
+#include "Materials/MaterialExpressionPathTracingQualitySwitch.h"
 #include "Materials/MaterialExpressionGetMaterialAttributes.h"
 #include "Materials/MaterialExpressionHairAttributes.h"
 #include "Materials/MaterialExpressionHairColor.h"
@@ -15684,7 +15685,7 @@ int32 UMaterialExpressionRayTracingQualitySwitch::Compile(class FMaterialCompile
 	}
 	else if (!RayTraced.GetTracedInput().Expression)
 	{
-		return Compiler->Errorf(TEXT("Missing RayTracingQualitySwitch input 'RayTracing'"));
+		return Compiler->Errorf(TEXT("Missing RayTracingQualitySwitch input 'RayTraced'"));
 	}
 	else
 	{
@@ -15729,6 +15730,83 @@ uint32 UMaterialExpressionRayTracingQualitySwitch::GetInputType(int32 InputIndex
 }
 #endif // WITH_EDITOR
 
+//
+// UMaterialExpressionPathTracingQualitySwitch
+//
+UMaterialExpressionPathTracingQualitySwitch::UMaterialExpressionPathTracingQualitySwitch(const FObjectInitializer& ObjectInitializer)
+	: Super(ObjectInitializer)
+{
+	// Structure to hold one-time initialization
+	struct FConstructorStatics
+	{
+		FText NAME_Utility;
+		FConstructorStatics()
+			: NAME_Utility(LOCTEXT("Utility", "Utility"))
+		{
+		}
+	};
+	static FConstructorStatics ConstructorStatics;
+
+#if WITH_EDITORONLY_DATA
+	MenuCategories.Add(ConstructorStatics.NAME_Utility);
+#endif
+}
+
+#if WITH_EDITOR
+int32 UMaterialExpressionPathTracingQualitySwitch::Compile(class FMaterialCompiler* Compiler, int32 OutputIndex)
+{
+	if (!Normal.GetTracedInput().Expression)
+	{
+		return Compiler->Errorf(TEXT("Missing PathTracingQualitySwitch input 'Normal'"));
+	}
+	else if (!PathTraced.GetTracedInput().Expression)
+	{
+		return Compiler->Errorf(TEXT("Missing PathTracingQualitySwitch input 'PathTraced'"));
+	}
+	else
+	{
+		int32 Arg1 = Normal.Compile(Compiler);
+		int32 Arg2 = PathTraced.Compile(Compiler);
+
+		//only when both of these are real expressions do the actual code.  otherwise various output pins will
+		//end up considered 'set' when really we just want a default.  This can cause us to force depth output when we don't want it for example.
+		if (Arg1 != INDEX_NONE && Arg2 != INDEX_NONE)
+		{
+			return Compiler->PathTracingQualitySwitchReplace(Arg1, Arg2);
+		}
+		else if (Arg1 != INDEX_NONE)
+		{
+			return Arg1;
+		}
+		else if (Arg2 != INDEX_NONE)
+		{
+			return Arg2;
+		}
+		return INDEX_NONE;
+	}
+}
+
+bool UMaterialExpressionPathTracingQualitySwitch::IsResultMaterialAttributes(int32 OutputIndex)
+{
+	if (Normal.Expression)
+	{
+		return Normal.Expression->IsResultMaterialAttributes(OutputIndex);
+	}
+	return false;
+}
+
+void UMaterialExpressionPathTracingQualitySwitch::GetCaption(TArray<FString>& OutCaptions) const
+{
+	OutCaptions.Add(TEXT("PathTracingQualitySwitchReplace"));
+}
+
+uint32 UMaterialExpressionPathTracingQualitySwitch::GetInputType(int32 InputIndex)
+{
+	return MCT_Unknown;
+}
+#endif // WITH_EDITOR
+
+
 UMaterialExpressionObjectOrientation::UMaterialExpressionObjectOrientation(const FObjectInitializer& ObjectInitializer)
 	: Super(ObjectInitializer)
 {
diff --git a/Engine/Source/Runtime/Engine/Public/MaterialCompiler.h b/Engine/Source/Runtime/Engine/Public/MaterialCompiler.h
--- a/Engine/Source/Runtime/Engine/Public/MaterialCompiler.h
+++ b/Engine/Source/Runtime/Engine/Public/MaterialCompiler.h
@@ -340,6 +340,7 @@ public:
 	virtual int32 GIReplace(int32 Direct, int32 StaticIndirect, int32 DynamicIndirect) = 0;
 	virtual int32 ShadowReplace(int32 Default, int32 Shadow) = 0;
 	virtual int32 RayTracingQualitySwitchReplace(int32 Normal, int32 RayTraced) = 0;
+	virtual int32 PathTracingQualitySwitchReplace(int32 Normal, int32 PathTraced) = 0;
 	virtual int32 VirtualTextureOutputReplace(int32 Default, int32 VirtualTexture) = 0;
 	virtual int32 ReflectionCapturePassSwitch(int32 Default, int32 Reflection) = 0;
 
@@ -687,6 +688,7 @@ public:
 	virtual int32 GIReplace(int32 Direct, int32 StaticIndirect, int32 DynamicIndirect) override { return Compiler->GIReplace(Direct, StaticIndirect, DynamicIndirect); }
 	virtual int32 ShadowReplace(int32 Default, int32 Shadow) override { return Compiler->ShadowReplace(Default, Shadow); }
 	virtual int32 RayTracingQualitySwitchReplace(int32 Normal, int32 RayTraced) override { return Compiler->RayTracingQualitySwitchReplace(Normal, RayTraced); }
+	virtual int32 PathTracingQualitySwitchReplace(int32 Normal, int32 PathTraced) override { return Compiler->PathTracingQualitySwitchReplace(Normal, PathTraced); }
 	virtual int32 VirtualTextureOutputReplace(int32 Default, int32 VirtualTexture) override { return Compiler->VirtualTextureOutputReplace(Default, VirtualTexture); }
 	virtual int32 ReflectionCapturePassSwitch(int32 Default, int32 Reflection) override { return Compiler->ReflectionCapturePassSwitch(Default, Reflection); }
 
