WorldPartition: LevelInstanceActorDesc lazy container registration fixes broken embedding when LevelInstance actor was saved in wp level

diff --git a/Engine/Source/Runtime/Engine/Private/WorldPartition/LevelInstance/LevelInstanceActorDesc.cpp b/Engine/Source/Runtime/Engine/Private/WorldPartition/LevelInstance/LevelInstanceActorDesc.cpp
--- a/Engine/Source/Runtime/Engine/Private/WorldPartition/LevelInstance/LevelInstanceActorDesc.cpp
+++ b/Engine/Source/Runtime/Engine/Private/WorldPartition/LevelInstance/LevelInstanceActorDesc.cpp
@@ -14,6 +14,7 @@
 
 FLevelInstanceActorDesc::FLevelInstanceActorDesc()
 	: DesiredRuntimeBehavior(ELevelInstanceRuntimeBehavior::Embedded)
+	, LevelInstanceContainer(nullptr)
 {
 
 }
@@ -36,33 +37,27 @@ void FLevelInstanceActorDesc::Init(UActorDescContainer* InContainer, const FWorl
 	FWorldPartitionActorDesc::Init(InContainer, DescData);
 }
 
-void FLevelInstanceActorDesc::OnRegister()
-{
-	FWorldPartitionActorDesc::OnRegister();
-	RegisterContainer();
-}
-
 void FLevelInstanceActorDesc::OnUnregister()
 {
 	FWorldPartitionActorDesc::OnUnregister();
 	LevelInstanceContainer = nullptr;
 }
 
-void FLevelInstanceActorDesc::RegisterContainer()
+bool FLevelInstanceActorDesc::GetContainerInstance(const UActorDescContainer*& OutLevelContainer, FTransform& OutLevelTransform, EContainerClusterMode& OutClusterMode) const
 {
-	check(!LevelInstanceContainer);
-	if (Container && (DesiredRuntimeBehavior == ELevelInstanceRuntimeBehavior::Embedded || DesiredRuntimeBehavior == ELevelInstanceRuntimeBehavior::Partitioned))
+	// Lazy register
+	if (!LevelInstanceContainer)
 	{
-		if (!LevelPackage.IsNone() && ULevel::GetIsLevelUsingExternalActorsFromPackage(LevelPackage) && !ULevel::GetIsLevelPartitionedFromPackage(LevelPackage))
+		if (Container && (DesiredRuntimeBehavior == ELevelInstanceRuntimeBehavior::Embedded || DesiredRuntimeBehavior == ELevelInstanceRuntimeBehavior::Partitioned))
 		{
-			UWorldPartition* WorldPartition = Container->GetWorld()->GetWorldPartition();
-			LevelInstanceContainer = WorldPartition->RegisterActorDescContainer(LevelPackage);
+			if (!LevelPackage.IsNone() && ULevel::GetIsLevelUsingExternalActorsFromPackage(LevelPackage) && !ULevel::GetIsLevelPartitionedFromPackage(LevelPackage))
+			{
+				UWorldPartition* WorldPartition = Container->GetWorld()->GetWorldPartition();
+				LevelInstanceContainer = WorldPartition->RegisterActorDescContainer(LevelPackage);
+			}
 		}
 	}
-}
 
-bool FLevelInstanceActorDesc::GetContainerInstance(const UActorDescContainer*& OutLevelContainer, FTransform& OutLevelTransform, EContainerClusterMode& OutClusterMode) const
-{
 	if (LevelInstanceContainer)
 	{
 		OutLevelContainer = LevelInstanceContainer;
diff --git a/Engine/Source/Runtime/Engine/Public/WorldPartition/LevelInstance/LevelInstanceActorDesc.h b/Engine/Source/Runtime/Engine/Public/WorldPartition/LevelInstance/LevelInstanceActorDesc.h
--- a/Engine/Source/Runtime/Engine/Public/WorldPartition/LevelInstance/LevelInstanceActorDesc.h
+++ b/Engine/Source/Runtime/Engine/Public/WorldPartition/LevelInstance/LevelInstanceActorDesc.h
@@ -5,6 +5,7 @@
 #include "CoreMinimal.h"
 #include "UObject/ObjectMacros.h"
 #include "UObject/GCObject.h"
+#include "UObject/ObjectPtr.h"
 #include "Containers/Map.h"
 #include "WorldPartition/WorldPartitionActorDesc.h"
 
@@ -30,14 +31,12 @@ protected:
 	virtual void Init(UActorDescContainer* InContainer, const FWorldPartitionActorDescInitData& DescData) override;
 	virtual void Serialize(FArchive& Ar) override;
 	virtual void AddReferencedObjects(FReferenceCollector& Collector);
-	virtual void RegisterContainer();
-	virtual void OnRegister();
 	virtual void OnUnregister();
 
 	FName LevelPackage;
 	FTransform LevelInstanceTransform;
 	ELevelInstanceRuntimeBehavior DesiredRuntimeBehavior;
 
-	const UActorDescContainer* LevelInstanceContainer = nullptr;
+	mutable TObjectPtr<const UActorDescContainer> LevelInstanceContainer;
 #endif
 };
