WorldPartition: Non WP ActorDescs (Level Instance) shouldn't set/find their ActorPtr. (Fix potiential UDN issue when committing a level instance)

diff --git a/Engine/Source/Runtime/Engine/Private/WorldPartition/WorldPartitionActorDesc.cpp b/Engine/Source/Runtime/Engine/Private/WorldPartition/WorldPartitionActorDesc.cpp
--- a/Engine/Source/Runtime/Engine/Private/WorldPartition/WorldPartitionActorDesc.cpp
+++ b/Engine/Source/Runtime/Engine/Private/WorldPartition/WorldPartitionActorDesc.cpp
@@ -42,8 +42,6 @@ void FWorldPartitionActorDesc::Init(const AActor* InActor)
 	ActorClass = GetParentNativeClass(InActor->GetClass());
 	Class = ActorClass->GetFName();
 
-	ActorPtr = (AActor*)InActor;
-
 	const FBox StreamingBounds = InActor->GetStreamingBounds();
 	StreamingBounds.GetCenterAndExtents(BoundsLocation, BoundsExtent);
 
@@ -80,7 +78,12 @@ void FWorldPartitionActorDesc::Init(const AActor* InActor)
 
 	ActorLabel = *InActor->GetActorLabel(false);
 
-	Container = InActor->GetLevel()->GetWorldPartition();
+	// Only set ActorPtr/Container on WorldPartition ActorDescs
+	if (UWorldPartition* WorldPartition = InActor->GetLevel()->GetWorldPartition())
+	{
+		ActorPtr = (AActor*)InActor;
+		Container = WorldPartition;
+	}
 }
 
 void FWorldPartitionActorDesc::Init(UActorDescContainer* InContainer, const FWorldPartitionActorDescInitData& DescData)
@@ -108,9 +111,12 @@ void FWorldPartitionActorDesc::Init(UActorDescContainer* InContainer, const FWor
 		GridPlacement = DefaultGridPlacement;
 	}
 
-	Container = InContainer;
-
-	ActorPtr = FindObject<AActor>(nullptr, *ActorPath.ToString());
+	// Only set ActorPtr/Container on WorldPartition ActorDescs
+	if (UWorldPartition* WorldPartition = Cast<UWorldPartition>(InContainer))
+	{
+		Container = InContainer;
+		ActorPtr = FindObject<AActor>(nullptr, *ActorPath.ToString());
+	}
 }
 
 void FWorldPartitionActorDesc::SerializeTo(TArray<uint8>& OutData)
@@ -246,6 +252,11 @@ UHLODLayer* FWorldPartitionActorDesc::GetHLODLayer() const
 
 bool FWorldPartitionActorDesc::IsLoaded(bool bEvenIfPendingKill) const
 {
+	if (!Container)
+	{
+		return false;
+	}
+
 #if WITH_DEV_AUTOMATION_TESTS
 	if (GIsAutomationTesting)
 	{
@@ -263,6 +274,7 @@ AActor* FWorldPartitionActorDesc::GetActor(bool bEvenIfPendingKill, bool bEvenIf
 
 AActor* FWorldPartitionActorDesc::Load() const
 {
+	check(Container);
 	if (ActorPtr.IsExplicitlyNull())
 	{
 		// First, try to find the existing actor which could have been loaded by another actor (through standard serialization)
@@ -330,6 +342,7 @@ void FWorldPartitionActorDesc::Unload()
 
 void FWorldPartitionActorDesc::RegisterActor()
 {
+	check(Container);
 	if (AActor* Actor = GetActor())
 	{
 		Container->OnActorDescRegistered(*this);
@@ -338,6 +351,7 @@ void FWorldPartitionActorDesc::RegisterActor()
 
 void FWorldPartitionActorDesc::UnregisterActor()
 {
+	check(Container);
 	if (AActor* Actor = GetActor())
 	{
 		Container->OnActorDescUnregistered(*this);
