Disable p.ClothPhysics.WaitForParallelClothTask so we do not block on cloth.

diff --git a/Engine/Source/Runtime/Engine/Private/SkeletalMeshComponentPhysics.cpp b/Engine/Source/Runtime/Engine/Private/SkeletalMeshComponentPhysics.cpp
--- a/Engine/Source/Runtime/Engine/Private/SkeletalMeshComponentPhysics.cpp
+++ b/Engine/Source/Runtime/Engine/Private/SkeletalMeshComponentPhysics.cpp
@@ -65,7 +65,7 @@ CSV_DECLARE_CATEGORY_MODULE_EXTERN(CORE_API, Basic);
 
 TAutoConsoleVariable<int32> CVarEnableClothPhysics(TEXT("p.ClothPhysics"), 1, TEXT("If 1, physics cloth will be used for simulation."));
 static TAutoConsoleVariable<int32> CVarEnableClothPhysicsUseTaskThread(TEXT("p.ClothPhysics.UseTaskThread"), 1, TEXT("If 1, run cloth on the task thread. If 0, run on game thread."));
-static TAutoConsoleVariable<int32> CVarClothPhysicsTickWaitForParallelClothTask(TEXT("p.ClothPhysics.WaitForParallelClothTask"), 1, TEXT("If 1, always wait for cloth task completion in the Cloth Tick function. If 0, wait at end-of-frame updates instead if allowed by component settings"));
+static TAutoConsoleVariable<int32> CVarClothPhysicsTickWaitForParallelClothTask(TEXT("p.ClothPhysics.WaitForParallelClothTask"), 0, TEXT("If 1, always wait for cloth task completion in the Cloth Tick function. If 0, wait at end-of-frame updates instead if allowed by component settings"));
 
 static TAutoConsoleVariable<bool> CVarClothTeleportOverride(TEXT("p.Cloth.TeleportOverride"), false, TEXT("Force console variable teleport override values over skeletal mesh properties.\n Default: false."));
 static TAutoConsoleVariable<bool> CVarClothResetAfterTeleport(TEXT("p.Cloth.ResetAfterTeleport"), true, TEXT("Require p.Cloth.TeleportOverride. Reset the clothing after moving the clothing position (called teleport).\n Default: true."));
