Expose path tracer ShowFlag to shipping builds

diff --git a/Engine/Source/Runtime/Engine/Public/ShowFlagsValues.inl b/Engine/Source/Runtime/Engine/Public/ShowFlagsValues.inl
--- a/Engine/Source/Runtime/Engine/Public/ShowFlagsValues.inl
+++ b/Engine/Source/Runtime/Engine/Public/ShowFlagsValues.inl
@@ -363,7 +363,7 @@ SHOWFLAG_FIXED_IN_SHIPPING(0, VREditing, SFG_Hidden, NSLOCTEXT("UnrealEd", "VREd
 SHOWFLAG_FIXED_IN_SHIPPING(0, OcclusionMeshes, SFG_Visualize, NSLOCTEXT("UnrealEd", "VisualizeOcclusionQueries", "Visualize Occlusion Queries"))
 
 // RHI_RAYTRACING begin
-SHOWFLAG_FIXED_IN_SHIPPING(0, PathTracing, SFG_Hidden, NSLOCTEXT("UnrealEd", "PathTracing", "Path tracing"))
+SHOWFLAG_ALWAYS_ACCESSIBLE(PathTracing, SFG_Hidden, NSLOCTEXT("UnrealEd", "PathTracing", "Path tracing"))
 SHOWFLAG_FIXED_IN_SHIPPING(0, RayTracingDebug, SFG_Hidden, NSLOCTEXT("UnrealEd", "RayTracingDebug", "Ray tracing debug"))
 // RHI_RAYTRACING end
 
