OFPA: ULevel::PreDuplicate -> Fix bad package name being assigned to Actors

diff --git a/Engine/Source/Runtime/Engine/Private/Level.cpp b/Engine/Source/Runtime/Engine/Private/Level.cpp
--- a/Engine/Source/Runtime/Engine/Private/Level.cpp
+++ b/Engine/Source/Runtime/Engine/Private/Level.cpp
@@ -979,7 +979,7 @@ void ULevel::PreDuplicate(FObjectDuplicationParameters& DupParams)
 			UPackage* ActorPackage = Actor ? Actor->GetExternalPackage() : nullptr;
 			if (ActorPackage)
 			{
-				FString ActorPath = GetPathName();
+				FString ActorPath = Actor->GetPathName();
 
 				if (DstPackage != SrcPackage)
 				{
