Replaced strange tangent blending logic with a straightforward implementation.
Removed code and buffers for accumulating morph target tangent weights

diff --git a/Engine/Shaders/Private/MorphTargets.usf b/Engine/Shaders/Private/MorphTargets.usf
--- a/Engine/Shaders/Private/MorphTargets.usf
+++ b/Engine/Shaders/Private/MorphTargets.usf
@@ -19,7 +19,6 @@
 #define TANGENTZ_MAX_BITS				16
 
 RWBuffer<uint> MorphVertexBuffer;
-RWBuffer<int> MorphNormalizationBuffer;
 ByteAddressBuffer MorphDataBuffer;
 
 //TODO: use a buffer instead of packing into constant buffer?
@@ -28,7 +27,6 @@ uint4 MorphTargetGroupOffsets[GMorphTargetDispatchBatchSize / 4];
 float4 MorphTargetWeights[GMorphTargetDispatchBatchSize / 4];
 
 float4 PositionScale;
-float WeightScale;
 float2 Precision;
 
 groupshared uint LocalGroupOffsets[GMorphTargetDispatchBatchSize];
@@ -151,10 +149,6 @@ void GPUMorphUpdateCS(uint3 GroupId : SV_GroupID, uint GroupThreadIndex : SV_Gro
 			InterlockedAdd(MorphVertexBuffer[DestVertexAddress + 4], asuint(QuantizedTangentZOffset.y));
 			InterlockedAdd(MorphVertexBuffer[DestVertexAddress + 5], asuint(QuantizedTangentZOffset.z));
 		}
-
-		// Update weight sum
-		int QuantizedWeight = (int)round(Weight * WeightScale);
-		InterlockedAdd(MorphNormalizationBuffer[DestVertexIndex], QuantizedWeight);	//TODO: Only used to normalize tangent weights. Are we sure we even want to do this?
 	}
 }
 
@@ -170,13 +164,7 @@ void GPUMorphNormalizeCS(uint3 GroupId : SV_GroupID, uint GroupThreadIndex : SV_
 	MorphVertexBuffer[DestVertexAddress + 1] = asuint(Position.y);
 	MorphVertexBuffer[DestVertexAddress + 2] = asuint(Position.z);
 	
-	// TODO: Does this tangent normalization code make sense?
-	const float Weight = float(MorphNormalizationBuffer[DestVertexIndex]) * WeightScale;
 	float3 TangentZ = float3(int3(MorphVertexBuffer[DestVertexAddress + 3], MorphVertexBuffer[DestVertexAddress + 4], MorphVertexBuffer[DestVertexAddress + 5])) * PositionScale.w;
-	if (Weight > 1.0f)
-	{
-		TangentZ /= Weight;
-	}
 	MorphVertexBuffer[DestVertexAddress + 3] = asuint(TangentZ.x);
 	MorphVertexBuffer[DestVertexAddress + 4] = asuint(TangentZ.y);
 	MorphVertexBuffer[DestVertexAddress + 5] = asuint(TangentZ.z);
diff --git a/Engine/Source/Runtime/Engine/Private/SkeletalRenderGPUSkin.cpp b/Engine/Source/Runtime/Engine/Private/SkeletalRenderGPUSkin.cpp
--- a/Engine/Source/Runtime/Engine/Private/SkeletalRenderGPUSkin.cpp
+++ b/Engine/Source/Runtime/Engine/Private/SkeletalRenderGPUSkin.cpp
@@ -125,10 +125,6 @@ void FMorphVertexBuffer::InitDynamicRHI()
 	{
 		UAVValue = RHICreateUnorderedAccessView(VertexBufferRHI, PF_R32_UINT);
 		bNeedsInitialClear = true;
-
-		uint32 NormalizationBufferSize = LodData.GetNumVertices() * sizeof(int);
-		NormalizationBufferRHI = RHICreateVertexBuffer(NormalizationBufferSize, Flags, CreateInfo);
-		NormalizationBufferUAV = RHICreateUnorderedAccessView(NormalizationBufferRHI, PF_R32_SINT);
 	}
 
 	// hasn't been updated yet
@@ -140,8 +136,6 @@ void FMorphVertexBuffer::ReleaseDynamicRHI()
 	UAVValue.SafeRelease();
 	VertexBufferRHI.SafeRelease();
 	SRVValue.SafeRelease();
-	NormalizationBufferRHI.SafeRelease();
-	NormalizationBufferUAV.SafeRelease();
 }
 
 /*-----------------------------------------------------------------------------
@@ -735,15 +729,13 @@ void FSkeletalMeshObjectGPUSkin::UpdateMorphVertexBuffer(FRHICommandListImmediat
 
 TArray<float> FSkeletalMeshObjectGPUSkin::FSkeletalMeshObjectLOD::MorphAccumulatedWeightArray;
 
-void FGPUMorphUpdateCS::SetParameters(FRHICommandList& RHICmdList, const FVector4& LocalScale, float WeightScale, const FMorphTargetVertexInfoBuffers& MorphTargetVertexInfoBuffers, FMorphVertexBuffer& MorphVertexBuffer)
+void FGPUMorphUpdateCS::SetParameters(FRHICommandList& RHICmdList, const FVector4& LocalScale, const FMorphTargetVertexInfoBuffers& MorphTargetVertexInfoBuffers, FMorphVertexBuffer& MorphVertexBuffer)
 {
 	FRHIComputeShader* CS = RHICmdList.GetBoundComputeShader();
 
 	SetUAVParameter(RHICmdList, CS, MorphVertexBufferParameter, MorphVertexBuffer.GetUAV());
-	SetUAVParameter(RHICmdList, CS, MorphNormalizationBufferParameter, MorphVertexBuffer.GetNormalizationUAV());
 
 	SetShaderValue(RHICmdList, CS, PositionScaleParameter, LocalScale);
-	SetShaderValue(RHICmdList, CS, WeightScaleParameter, WeightScale);
 	FVector2D Precision = { MorphTargetVertexInfoBuffers.GetPositionPrecision(), MorphTargetVertexInfoBuffers.GetTangentZPrecision() };
 	SetShaderValue(RHICmdList, CS, PrecisionParameter, Precision);
 
@@ -773,13 +765,11 @@ void FGPUMorphUpdateCS::EndAllDispatches(FRHICommandList& RHICmdList)
 
 IMPLEMENT_SHADER_TYPE(, FGPUMorphUpdateCS, TEXT("/Engine/Private/MorphTargets.usf"), TEXT("GPUMorphUpdateCS"), SF_Compute);
 
-void FGPUMorphNormalizeCS::SetParameters(FRHICommandList& RHICmdList, const FVector4& InvLocalScale, float InvWeightScale, const FMorphTargetVertexInfoBuffers& MorphTargetVertexInfoBuffers, FMorphVertexBuffer& MorphVertexBuffer)
+void FGPUMorphNormalizeCS::SetParameters(FRHICommandList& RHICmdList, const FVector4& InvLocalScale, const FMorphTargetVertexInfoBuffers& MorphTargetVertexInfoBuffers, FMorphVertexBuffer& MorphVertexBuffer)
 {
 	FRHIComputeShader* CS = RHICmdList.GetBoundComputeShader();
 	SetUAVParameter(RHICmdList, CS, MorphVertexBufferParameter, MorphVertexBuffer.GetUAV());
-	SetUAVParameter(RHICmdList, CS, MorphNormalizationBufferParameter, MorphVertexBuffer.GetNormalizationUAV());
 	SetShaderValue(RHICmdList, CS, PositionScaleParameter, InvLocalScale);
-	SetShaderValue(RHICmdList, CS, WeightScaleParameter, InvWeightScale);
 }
 
 void FGPUMorphNormalizeCS::Dispatch(FRHICommandList& RHICmdList, uint32 NumVertices)
@@ -796,16 +786,14 @@ void FGPUMorphNormalizeCS::EndAllDispatches(FRHICommandList& RHICmdList)
 
 IMPLEMENT_SHADER_TYPE(, FGPUMorphNormalizeCS, TEXT("/Engine/Private/MorphTargets.usf"), TEXT("GPUMorphNormalizeCS"), SF_Compute);
 
-static void CalculateMorphDeltaBounds(const TArray<float>& MorphTargetWeights, const FMorphTargetVertexInfoBuffers& MorphTargetVertexInfoBuffers, FVector4& MorphScale, FVector4& InvMorphScale, float& OutWeightScale, float& OutInvWeightScale)
+static void CalculateMorphDeltaBounds(const TArray<float>& MorphTargetWeights, const FMorphTargetVertexInfoBuffers& MorphTargetVertexInfoBuffers, FVector4& MorphScale, FVector4& InvMorphScale)
 {
 	double MinAccumScale[4] = { 0, 0, 0, 0 };
 	double MaxAccumScale[4] = { 0, 0, 0, 0 };
 	double MaxScale[4] = { 0, 0, 0, 0 };
-	double WeightScale = 0.0;
 
 	for (uint32 i = 0; i < MorphTargetVertexInfoBuffers.GetNumMorphs(); i++)
 	{
-		WeightScale += FMath::Abs(MorphTargetWeights[i]);
 		FVector4 MinMorphScale = MorphTargetVertexInfoBuffers.GetMinimumMorphScale(i);
 		FVector4 MaxMorphScale = MorphTargetVertexInfoBuffers.GetMaximumMorphScale(i);
 
@@ -821,15 +809,12 @@ static void CalculateMorphDeltaBounds(const TArray<float>& MorphTargetWeights, c
 		}
 	}
 
-	WeightScale = FMath::Max<double>(WeightScale, 1.0);
 	MaxScale[0] = FMath::Max<double>(MaxScale[0], 1.0);
 	MaxScale[1] = FMath::Max<double>(MaxScale[1], 1.0);
 	MaxScale[2] = FMath::Max<double>(MaxScale[2], 1.0);
 	MaxScale[3] = FMath::Max<double>(MaxScale[3], 1.0);
 
 	const double ScaleToInt24 = 16777216.0;
-	OutWeightScale = float(ScaleToInt24 / WeightScale);
-	OutInvWeightScale = float(WeightScale / ScaleToInt24);
 
 	MorphScale = FVector4
 	(
@@ -871,33 +856,22 @@ void FSkeletalMeshObjectGPUSkin::FSkeletalMeshObjectLOD::UpdateMorphVertexBuffer
 			LodData.GetNumVertices(),
 			MorphTargetVertexInfoBuffers.GetNumBatches());
 
-		RHICmdList.Transition(
-		{
-			FRHITransitionInfo(MorphVertexBuffer.GetUAV(), ERHIAccess::Unknown, ERHIAccess::UAVCompute),
-			FRHITransitionInfo(MorphVertexBuffer.GetNormalizationUAV(), ERHIAccess::Unknown, ERHIAccess::UAVCompute)
-		});
+		RHICmdList.Transition(FRHITransitionInfo(MorphVertexBuffer.GetUAV(), ERHIAccess::Unknown, ERHIAccess::UAVCompute));
 		RHICmdList.ClearUAVUint(MorphVertexBuffer.GetUAV(), FUintVector4(0, 0, 0, 0));
-		RHICmdList.ClearUAVUint(MorphVertexBuffer.GetNormalizationUAV(), FUintVector4(0, 0, 0, 0));
 
 		{
 			FVector4 MorphScale; 
 			FVector4 InvMorphScale; 
-			float WeightScale;
-			float InvWeightScale;
 			{
 				SCOPE_CYCLE_COUNTER(STAT_MorphVertexBuffer_ApplyDelta);
-				CalculateMorphDeltaBounds(MorphTargetWeights, MorphTargetVertexInfoBuffers, MorphScale, InvMorphScale, WeightScale, InvWeightScale);
+				CalculateMorphDeltaBounds(MorphTargetWeights, MorphTargetVertexInfoBuffers, MorphScale, InvMorphScale);
 			}
 
 			{
 				SCOPED_DRAW_EVENTF(RHICmdList, MorphUpdateScatter, TEXT("Scatter"));
 
-				RHICmdList.Transition(
-				{
-					FRHITransitionInfo(MorphVertexBuffer.GetUAV(), ERHIAccess::UAVCompute, ERHIAccess::UAVCompute),
-					FRHITransitionInfo(MorphVertexBuffer.GetNormalizationUAV(), ERHIAccess::UAVCompute, ERHIAccess::UAVCompute)
-				});
-				RHICmdList.BeginUAVOverlap({ MorphVertexBuffer.GetUAV(), MorphVertexBuffer.GetNormalizationUAV() });
+				RHICmdList.Transition(FRHITransitionInfo(MorphVertexBuffer.GetUAV(), ERHIAccess::UAVCompute, ERHIAccess::UAVCompute));
+				RHICmdList.BeginUAVOverlap(MorphVertexBuffer.GetUAV());
 
 				//the first pass scatters all morph targets into the vertexbuffer using atomics
 				//multiple morph targets can be batched by a single shader where the shader will rely on
@@ -935,23 +909,19 @@ void FSkeletalMeshObjectGPUSkin::FSkeletalMeshObjectLOD::UpdateMorphVertexBuffer
 					}
 
 					RHICmdList.SetComputeShader(GPUMorphUpdateCS.GetComputeShader());
-					GPUMorphUpdateCS->SetParameters(RHICmdList, MorphScale, WeightScale, MorphTargetVertexInfoBuffers, MorphVertexBuffer);
+					GPUMorphUpdateCS->SetParameters(RHICmdList, MorphScale, MorphTargetVertexInfoBuffers, MorphVertexBuffer);
 					GPUMorphUpdateCS->SetMorphOffsetsAndWeights(RHICmdList, BatchOffsets, GroupOffsets, Weights);
 					GPUMorphUpdateCS->Dispatch(RHICmdList, NumBatches);
 				}
 				
 				GPUMorphUpdateCS->EndAllDispatches(RHICmdList);
-				RHICmdList.EndUAVOverlap({ MorphVertexBuffer.GetUAV(), MorphVertexBuffer.GetNormalizationUAV() });
+				RHICmdList.EndUAVOverlap(MorphVertexBuffer.GetUAV());
 			}
 
 			{
 				SCOPED_DRAW_EVENTF(RHICmdList, MorphUpdateNormalize, TEXT("Normalize"));
 
-				RHICmdList.Transition(
-				{
-					FRHITransitionInfo(MorphVertexBuffer.GetUAV(), ERHIAccess::UAVCompute, ERHIAccess::UAVCompute),
-					FRHITransitionInfo(MorphVertexBuffer.GetNormalizationUAV(), ERHIAccess::UAVCompute, ERHIAccess::UAVCompute)
-				});
+				RHICmdList.Transition(FRHITransitionInfo(MorphVertexBuffer.GetUAV(), ERHIAccess::UAVCompute, ERHIAccess::UAVCompute));
 
 				//The second pass normalizes the scattered result and converts it back into floats.
 				//The dispatches are split by morph permutation (and their accumulated weight) .
@@ -960,7 +930,7 @@ void FSkeletalMeshObjectGPUSkin::FSkeletalMeshObjectLOD::UpdateMorphVertexBuffer
 				//binary search to find the correct target weight within the batch.
 				TShaderMapRef<FGPUMorphNormalizeCS> GPUMorphNormalizeCS(GetGlobalShaderMap(FeatureLevel));
 				RHICmdList.SetComputeShader(GPUMorphNormalizeCS.GetComputeShader());
-				GPUMorphNormalizeCS->SetParameters(RHICmdList, InvMorphScale, InvWeightScale, MorphTargetVertexInfoBuffers, MorphVertexBuffer);
+				GPUMorphNormalizeCS->SetParameters(RHICmdList, InvMorphScale, MorphTargetVertexInfoBuffers, MorphVertexBuffer);
 				GPUMorphNormalizeCS->Dispatch(RHICmdList, MorphVertexBuffer.GetNumVerticies());
 				GPUMorphNormalizeCS->EndAllDispatches(RHICmdList);
 				RHICmdList.Transition(FRHITransitionInfo(MorphVertexBuffer.GetUAV(), ERHIAccess::UAVCompute, ERHIAccess::VertexOrIndexBuffer | ERHIAccess::SRVMask));
diff --git a/Engine/Source/Runtime/Engine/Private/SkeletalRenderGPUSkin.h b/Engine/Source/Runtime/Engine/Private/SkeletalRenderGPUSkin.h
--- a/Engine/Source/Runtime/Engine/Private/SkeletalRenderGPUSkin.h
+++ b/Engine/Source/Runtime/Engine/Private/SkeletalRenderGPUSkin.h
@@ -211,15 +211,6 @@ public:
 			ResourceSize += LodData.GetNumVertices() * sizeof(FMorphGPUSkinVertex);
 		}
 
-		if (NormalizationBufferRHI)
-		{
-			// LOD of the skel mesh is used to find number of vertices in buffer
-			FSkeletalMeshLODRenderData& LodData = SkelMeshRenderData->LODRenderData[LODIdx];
-
-			// Create the buffer rendering resource
-			ResourceSize += LodData.GetNumVertices() * sizeof(int);
-		}
-
 		return ResourceSize;
 	}
 	
@@ -249,12 +240,6 @@ public:
 		return UAVValue;
 	}
 
-	// @param guaranteed only to be valid if the Normalization Buffer is valid
-	FRHIUnorderedAccessView* GetNormalizationUAV() const
-	{
-		return NormalizationBufferUAV;
-	}
-
 	FSkeletalMeshLODRenderData* GetLODRenderData() const { return &SkelMeshRenderData->LODRenderData[LODIdx]; }
 	
 	// section ids that are using this Morph buffer
@@ -276,12 +261,6 @@ private:
 
 	// parent mesh containing the source data, never 0
 	FSkeletalMeshRenderData* SkelMeshRenderData;
-
-	//temporary buffer used for normalization
-	FBufferRHIRef NormalizationBufferRHI;
-
-	// guaranteed only to be valid if the Normalization Buffer is valid
-	FUnorderedAccessViewRHIRef NormalizationBufferUAV;
 };
 
 /**
@@ -705,13 +684,11 @@ public:
 		: FGlobalShader(Initializer)
 	{
 		MorphVertexBufferParameter.Bind(Initializer.ParameterMap, TEXT("MorphVertexBuffer"));
-		MorphNormalizationBufferParameter.Bind(Initializer.ParameterMap, TEXT("MorphNormalizationBuffer"));
 
 		MorphTargetWeightsParameter.Bind(Initializer.ParameterMap, TEXT("MorphTargetWeights"));
 		MorphTargetBatchOffsetsParameter.Bind(Initializer.ParameterMap, TEXT("MorphTargetBatchOffsets"));
 		MorphTargetGroupOffsetsParameter.Bind(Initializer.ParameterMap, TEXT("MorphTargetGroupOffsets"));
 		PositionScaleParameter.Bind(Initializer.ParameterMap, TEXT("PositionScale"));
-		WeightScaleParameter.Bind(Initializer.ParameterMap, TEXT("WeightScale"));
 		PrecisionParameter.Bind(Initializer.ParameterMap, TEXT("Precision"));
 
 		MorphDataBufferParameter.Bind(Initializer.ParameterMap, TEXT("MorphDataBuffer"));
@@ -719,7 +696,7 @@ public:
 
 	static const uint32 MorphTargetDispatchBatchSize = 128;
 
-	void SetParameters(FRHICommandList& RHICmdList, const FVector4& LocalScale, float WeightScale, const FMorphTargetVertexInfoBuffers& MorphTargetVertexInfoBuffers, FMorphVertexBuffer& MorphVertexBuffer);
+	void SetParameters(FRHICommandList& RHICmdList, const FVector4& LocalScale, const FMorphTargetVertexInfoBuffers& MorphTargetVertexInfoBuffers, FMorphVertexBuffer& MorphVertexBuffer);
 	void SetMorphOffsetsAndWeights(FRHICommandList& RHICmdList, uint32 BatchOffsets[MorphTargetDispatchBatchSize], uint32 GroupOffsets[MorphTargetDispatchBatchSize], float Weights[MorphTargetDispatchBatchSize]);
 
 	void Dispatch(FRHICommandList& RHICmdList, uint32 Size);
@@ -732,14 +709,12 @@ public:
 
 protected:
 	LAYOUT_FIELD(FShaderResourceParameter, MorphVertexBufferParameter);
-	LAYOUT_FIELD(FShaderResourceParameter, MorphNormalizationBufferParameter);
 
 	LAYOUT_FIELD(FShaderParameter, MorphTargetWeightsParameter);
 	LAYOUT_FIELD(FShaderParameter, OffsetAndSizeParameter);
 	LAYOUT_FIELD(FShaderParameter, MorphTargetBatchOffsetsParameter);
 	LAYOUT_FIELD(FShaderParameter, MorphTargetGroupOffsetsParameter);
 	LAYOUT_FIELD(FShaderParameter, PositionScaleParameter);
-	LAYOUT_FIELD(FShaderParameter, WeightScaleParameter);
 	LAYOUT_FIELD(FShaderParameter, PrecisionParameter);
 
 	LAYOUT_FIELD(FShaderResourceParameter, MorphDataBufferParameter);
@@ -756,9 +731,6 @@ public:
 		: FGlobalShader(Initializer)
 	{
 		MorphVertexBufferParameter.Bind(Initializer.ParameterMap, TEXT("MorphVertexBuffer"));
-		MorphNormalizationBufferParameter.Bind(Initializer.ParameterMap, TEXT("MorphNormalizationBuffer"));
-
-		WeightScaleParameter.Bind(Initializer.ParameterMap, TEXT("WeightScale"));
 		PositionScaleParameter.Bind(Initializer.ParameterMap, TEXT("PositionScale"));
 	}
 
@@ -767,15 +739,13 @@ public:
 		return IsFeatureLevelSupported(Parameters.Platform, ERHIFeatureLevel::SM5);
 	}
 
-	void SetParameters(FRHICommandList& RHICmdList, const FVector4& LocalScale, float InvWeightScale, const FMorphTargetVertexInfoBuffers& MorphTargetVertexInfoBuffers, FMorphVertexBuffer& MorphVertexBuffer);
+	void SetParameters(FRHICommandList& RHICmdList, const FVector4& LocalScale, const FMorphTargetVertexInfoBuffers& MorphTargetVertexInfoBuffers, FMorphVertexBuffer& MorphVertexBuffer);
 
 	void Dispatch(FRHICommandList& RHICmdList, uint32 NumVerticies);
 	void EndAllDispatches(FRHICommandList& RHICmdList);
 
 protected:
 	LAYOUT_FIELD(FShaderResourceParameter, MorphVertexBufferParameter);
-	LAYOUT_FIELD(FShaderResourceParameter, MorphNormalizationBufferParameter);
 
-	LAYOUT_FIELD(FShaderParameter, WeightScaleParameter);
 	LAYOUT_FIELD(FShaderParameter, PositionScaleParameter);
 };
