Add MaxLODSize_VT option to FTextureLODGroup to clamp cooked VT sizes per group.
Removed some small transient state from FTextureLODGroup and calculate it in the places it's needed.

diff --git a/Engine/Source/Editor/DetailCustomizations/Private/TextureLODSettingsDetails.cpp b/Engine/Source/Editor/DetailCustomizations/Private/TextureLODSettingsDetails.cpp
--- a/Engine/Source/Editor/DetailCustomizations/Private/TextureLODSettingsDetails.cpp
+++ b/Engine/Source/Editor/DetailCustomizations/Private/TextureLODSettingsDetails.cpp
@@ -137,6 +137,25 @@ void FTextureLODGroupLayout::GenerateChildContent(IDetailChildrenBuilder& Childr
 			];
 	}
 
+	{
+		ChildrenBuilder.AddCustomRow(LOCTEXT("MaxLODSizeVT", "Max LOD Size VT"))
+			.NameContent()
+			[
+				SNew(STextBlock)
+				.Font(IDetailLayoutBuilder::GetDetailFont())
+				.Text(LOCTEXT("MaxLODSizeVT", "Max LOD Size VT"))
+			]
+			.ValueContent()
+			[
+				SNew(SSpinBox<uint32>)
+				.Font(IDetailLayoutBuilder::GetDetailFont())
+				.MinValue(0)
+				.MaxValue(32768)
+				.Value(this, &FTextureLODGroupLayout::GetMaxLODSizeVT)
+				.OnValueChanged(this, &FTextureLODGroupLayout::OnMaxLODSizeVTChanged)
+				.OnValueCommitted(this, &FTextureLODGroupLayout::OnMaxLODSizeVTCommitted)
+			];
+	}
 
 	// LOD Bias
 	{
@@ -292,6 +311,22 @@ void FTextureLODGroupLayout::OnMaxLODSizeCommitted(uint32 NewValue, ETextCommit:
 }
 
 
+uint32 FTextureLODGroupLayout::GetMaxLODSizeVT() const
+{
+	return LodGroup->MaxLODSize_VT;
+}
+
+void FTextureLODGroupLayout::OnMaxLODSizeVTChanged(uint32 NewValue)
+{
+	LodGroup->MaxLODSize_VT = NewValue;
+}
+
+void FTextureLODGroupLayout::OnMaxLODSizeVTCommitted(uint32 NewValue, ETextCommit::Type TextCommitType)
+{
+	OnMaxLODSizeVTChanged(NewValue);
+}
+
+
 int32 FTextureLODGroupLayout::GetLODBias() const
 {
 	return LodGroup->LODBias;
diff --git a/Engine/Source/Editor/DetailCustomizations/Private/TextureLODSettingsDetails.h b/Engine/Source/Editor/DetailCustomizations/Private/TextureLODSettingsDetails.h
--- a/Engine/Source/Editor/DetailCustomizations/Private/TextureLODSettingsDetails.h
+++ b/Engine/Source/Editor/DetailCustomizations/Private/TextureLODSettingsDetails.h
@@ -47,6 +47,11 @@ private:
 	void OnMaxLODSizeChanged(uint32 NewValue);
 	void OnMaxLODSizeCommitted(uint32 NewValue, ETextCommit::Type TextCommitType);
 
+	// Controls for the Max LOD Size VT property editing
+	uint32 GetMaxLODSizeVT() const;
+	void OnMaxLODSizeVTChanged(uint32 NewValue);
+	void OnMaxLODSizeVTCommitted(uint32 NewValue, ETextCommit::Type TextCommitType);
+
 	// Controls for the LOD Bias property editing
 	int32 GetLODBias() const;
 	void OnLODBiasChanged(int32 NewValue);
diff --git a/Engine/Source/Runtime/Engine/Classes/Engine/TextureLODSettings.h b/Engine/Source/Runtime/Engine/Classes/Engine/TextureLODSettings.h
--- a/Engine/Source/Runtime/Engine/Classes/Engine/TextureLODSettings.h
+++ b/Engine/Source/Runtime/Engine/Classes/Engine/TextureLODSettings.h
@@ -18,8 +18,6 @@ struct ENGINE_API FTextureLODGroup
 
 	FTextureLODGroup()
 		: Group(TEXTUREGROUP_World)
-		, MinLODMipCount(0)
-		, MaxLODMipCount(32)
 		, LODBias(0)
 		, LODBias_Smaller(-1)
 		, LODBias_Smallest(-1)
@@ -29,9 +27,9 @@ struct ENGINE_API FTextureLODGroup
 		, MaxLODSize(4096)
 		, MaxLODSize_Smaller(-1)
 		, MaxLODSize_Smallest(-1)
+		, MaxLODSize_VT(0)
 		, OptionalLODBias(0)
 		, OptionalMaxLODSize(4096)
-		, OptionalMaxLODMipCount(32)
 		, MinMagFilter(NAME_Aniso)
 		, MipFilter(NAME_Point)
 		, MipLoadOptions(ETextureMipLoadOptions::AllMips)
@@ -46,16 +44,10 @@ struct ENGINE_API FTextureLODGroup
 		SetupGroup();
 	}
 
-	/** Minimum LOD mip count below which the code won't bias.						*/
+	/** Group ID.																	*/
 	UPROPERTY()
 	TEnumAsByte<TextureGroup> Group;
 
-	/** Minimum LOD mip count below which the code won't bias.						*/
-	int32 MinLODMipCount;
-	
-	/** Maximum LOD mip count. Bias will be adjusted so texture won't go above.		*/
-	int32 MaxLODMipCount;
-	
 	/** Group LOD bias.																*/
 	UPROPERTY()
 	int32 LODBias;
@@ -89,6 +81,9 @@ struct ENGINE_API FTextureLODGroup
 	UPROPERTY()
 	int32 MaxLODSize_Smallest;
 
+	UPROPERTY()
+	int32 MaxLODSize_VT;
+
 	/** If this is greater then 0 will put that number of mips into an optional bulkdata package */
 	UPROPERTY()
 	int32 OptionalLODBias;
@@ -97,8 +92,6 @@ struct ENGINE_API FTextureLODGroup
 	UPROPERTY()
 	int32 OptionalMaxLODSize;
 
-	int32 OptionalMaxLODMipCount;
-
 	UPROPERTY()
 	FName MinMagFilter;
 
diff --git a/Engine/Source/Runtime/Engine/Private/Streaming/StreamingManagerTexture.cpp b/Engine/Source/Runtime/Engine/Private/Streaming/StreamingManagerTexture.cpp
--- a/Engine/Source/Runtime/Engine/Private/Streaming/StreamingManagerTexture.cpp
+++ b/Engine/Source/Runtime/Engine/Private/Streaming/StreamingManagerTexture.cpp
@@ -2622,7 +2622,7 @@ bool FRenderAssetStreamingManager::HandleInvestigateRenderAssetCommand(const TCH
 						const int32 MipCountBeforeMaxRes = ResourceState.MaxNumLODs - RenderAsset->NumCinematicMipLevels -
 							(StreamingRenderAsset.LODGroup == TEXTUREGROUP_UI ? GUITextureLODBias : 0) - 
 							(FPlatformProperties::RequiresCookedData() ? 0 : (LODGroupInfo.LODBias + (Texture ? Texture->LODBias : 0)));
-						const int32 MaxResBias = MipCountBeforeMaxRes - (LODGroupInfo.MaxLODMipCount + 1);
+						const int32 MaxResBias = MipCountBeforeMaxRes - (FMath::CeilLogTwo(LODGroupInfo.MaxLODSize) + 1);
 						if (MaxResBias > 0)
 						{
 							BiasDesc += FString::Printf(TEXT(" [LODGroup.MaxRes:%d]"), MaxResBias);
diff --git a/Engine/Source/Runtime/Engine/Private/TextureLODSettings.cpp b/Engine/Source/Runtime/Engine/Private/TextureLODSettings.cpp
--- a/Engine/Source/Runtime/Engine/Private/TextureLODSettings.cpp
+++ b/Engine/Source/Runtime/Engine/Private/TextureLODSettings.cpp
@@ -59,11 +59,6 @@ void FTextureLODGroup::SetupGroup()
 		}
 	}
 
-	MinLODMipCount = FMath::CeilLogTwo(MinLODSize);
-	MaxLODMipCount = FMath::CeilLogTwo(MaxLODSize);
-
-	OptionalMaxLODMipCount = FMath::CeilLogTwo(OptionalMaxLODSize);
-
 	// Linear filtering
 	if (MinMagFilter == NAME_Linear)
 	{
@@ -182,9 +177,15 @@ int32 UTextureLODSettings::CalculateLODBias(int32 Width, int32 Height, int32 Max
 		UsedLODBias += GUITextureLODBias;
 	}
 
-	// Don't apply LOD group mip clamping to VT, since VT max size isn't constrained by any GPU capabilities
-	const int32 MinLOD	= bVirtualTexture ? 0 : LODGroupInfo.MinLODMipCount;
-	const int32 MaxLOD	= bVirtualTexture ? TextureMaxLOD : LODGroupInfo.MaxLODMipCount;
+	int32 MinLOD = FMath::CeilLogTwo(LODGroupInfo.MinLODSize);
+	int32 MaxLOD = FMath::CeilLogTwo(LODGroupInfo.MaxLODSize);
+
+	if (bVirtualTexture)
+	{
+		// Virtual textures have no MinLOD and a specific MaxLOD setting.
+		MinLOD = 0;
+		MaxLOD = LODGroupInfo.MaxLODSize_VT > 0 ? FMath::CeilLogTwo(LODGroupInfo.MaxLODSize_VT) : TextureMaxLOD;
+	}
 
 	int32 WantedMaxLOD	= FMath::Clamp( TextureMaxLOD - UsedLODBias, MinLOD, MaxLOD );
 	WantedMaxLOD		= FMath::Clamp( WantedMaxLOD, 0, TextureMaxLOD );
@@ -206,7 +207,7 @@ int32 UTextureLODSettings::CalculateNumOptionalMips(int32 LODGroup, const int32
 		return 0;
 	}
 
-	int32 OptionalLOD = FMath::Min(LODGroupInfo.OptionalMaxLODMipCount+1, NumMips);
+	int32 OptionalLOD = FMath::Min<int32>(FMath::CeilLogTwo(LODGroupInfo.OptionalMaxLODSize) + 1, NumMips);
 
 	int32 NumOptionalMips = FMath::Min(NumMips - (OptionalLOD - LODGroupInfo.OptionalLODBias), MinMipToInline);
 	return NumOptionalMips;
