Motion Trajectory Component debugging functionality. Includes minor Horde compiler suite fixes.

#trivial

diff --git a/Engine/Plugins/Experimental/Animation/MotionTrajectory/Source/MotionTrajectory/Private/MotionTrajectory.cpp b/Engine/Plugins/Experimental/Animation/MotionTrajectory/Source/MotionTrajectory/Private/MotionTrajectory.cpp
--- a/Engine/Plugins/Experimental/Animation/MotionTrajectory/Source/MotionTrajectory/Private/MotionTrajectory.cpp
+++ b/Engine/Plugins/Experimental/Animation/MotionTrajectory/Source/MotionTrajectory/Private/MotionTrajectory.cpp
@@ -27,14 +27,14 @@ FTrajectorySample UMotionTrajectoryComponent::GetPresentTrajectory() const
 {
 	const APawn* Pawn = TryGetOwnerPawn();
 	checkf(false, TEXT("UMotionTrajectoryComponent::GetPresentTrajectory for Pawn: %s requires implementation."), Pawn ? *(Pawn->GetHumanReadableName()) : TEXT("NULL"));
-	return {};
+	return FTrajectorySample();
 }
 
 FTransform UMotionTrajectoryComponent::GetPresentWorldTransform() const
 {
 	const APawn* Pawn = TryGetOwnerPawn();
 	checkf(false, TEXT("UMotionTrajectoryComponent::GetPresentWorldTransform for Pawn: %s requires implementation."), Pawn ? *(Pawn->GetHumanReadableName()) : TEXT("NULL"));
-	return {};
+	return FTransform::Identity;
 }
 
 const APawn* UMotionTrajectoryComponent::TryGetOwnerPawn() const
@@ -62,15 +62,18 @@ FTrajectorySampleRange UMotionTrajectoryComponent::CombineHistoryPresentPredicti
 	Trajectory.Samples.Add(PresentTrajectory);
 	Trajectory.Samples.Append(Prediction.Samples);
 	Trajectory.SampleRate = SampleRate;
-
+#if WITH_EDITORONLY_DATA
+	Trajectory.DebugDrawTrajectory(bDebugDrawTrajectory, GetOwner()->GetWorld(), PreviousWorldTransform);
+#endif
 	return Trajectory;
 }
 
 void UMotionTrajectoryComponent::TickHistoryEvictionPolicy()
 {
 	// World space transform and time are used for determining local-relative historical samples
+	const UWorld* World = GetWorld();
 	const FTransform WorldTransform = GetPresentWorldTransform();
-	const float WorldGameTime = UKismetSystemLibrary::GetGameTimeInSeconds(GetWorld());
+	const float WorldGameTime = UKismetSystemLibrary::GetGameTimeInSeconds(World);
 
 	// Skip on the very first sample
 	if (PreviousWorldGameTime != 0.f)
@@ -115,6 +118,7 @@ void UMotionTrajectoryComponent::TickHistoryEvictionPolicy()
 		// Remove all trajectory samples which are outside the boundaries of the enabled domain horizons
 		SampleHistory.RemoveAll([&](const FTrajectorySample& Sample)
 			{
+				// Remove superfluous zero motion samples
 				if (Sample.IsZeroSample() && PresentTrajectory.IsZeroSample())
 				{
 					return true;
diff --git a/Engine/Plugins/Experimental/Animation/MotionTrajectory/Source/MotionTrajectory/Private/MotionTrajectoryLibrary.cpp b/Engine/Plugins/Experimental/Animation/MotionTrajectory/Source/MotionTrajectory/Private/MotionTrajectoryLibrary.cpp
--- a/Engine/Plugins/Experimental/Animation/MotionTrajectory/Source/MotionTrajectory/Private/MotionTrajectoryLibrary.cpp
+++ b/Engine/Plugins/Experimental/Animation/MotionTrajectory/Source/MotionTrajectory/Private/MotionTrajectoryLibrary.cpp
@@ -170,4 +170,28 @@ FTrajectorySampleRange UMotionTrajectoryBlueprintLibrary::ClampTrajectoryDirecti
 	}
 
 	return Trajectory;
+}
+
+void UMotionTrajectoryBlueprintLibrary::DebugDrawTrajectory(const AActor* Actor
+	, const FTransform& WorldTransform
+	, const FTrajectorySampleRange& Trajectory
+	, const FLinearColor PredictionColor
+	, const FLinearColor HistoryColor
+	, float ArrowScale
+	, float ArrowSize
+	, float ArrowThickness
+)
+{
+	if (Actor)
+	{
+		Trajectory.DebugDrawTrajectory(true
+			, Actor->GetWorld()
+			, WorldTransform.IsValid() ? WorldTransform : FTransform::Identity
+			, PredictionColor
+			, HistoryColor
+			, ArrowScale
+			, ArrowSize
+			, ArrowThickness
+		);
+	}
 }
\ No newline at end of file
diff --git a/Engine/Plugins/Experimental/Animation/MotionTrajectory/Source/MotionTrajectory/Public/MotionTrajectory.h b/Engine/Plugins/Experimental/Animation/MotionTrajectory/Source/MotionTrajectory/Public/MotionTrajectory.h
--- a/Engine/Plugins/Experimental/Animation/MotionTrajectory/Source/MotionTrajectory/Public/MotionTrajectory.h
+++ b/Engine/Plugins/Experimental/Animation/MotionTrajectory/Source/MotionTrajectory/Public/MotionTrajectory.h
@@ -29,7 +29,7 @@ struct MOTIONTRAJECTORY_API FMotionTrajectorySettings
 	// Trajectory mask for specifying a time and/or distance domain
 	// Default: Time Domain of 2.0 seconds
 	UPROPERTY(EditAnywhere, BlueprintReadWrite, Category="Settings", meta=(Bitmask, BitmaskEnum=ETrajectorySampleDomain))
-	int32 Domain = 1; // ETrajectorySampleDomain::TIME
+	int32 Domain = 1; // ETrajectorySampleDomain::Time
 
 	// Sample time horizon
 	UPROPERTY(EditAnywhere, BlueprintReadWrite, Category="Settings", meta=(ClampMin="0.0"))
@@ -106,15 +106,15 @@ public:
 	UFUNCTION(BlueprintCallable, Category="Motion Trajectory")
 	void SetSampleRate(int32 Rate);
 
-	// Retrieves the trajectory history as a range
+	// Retrieves the historical trajectory
 	UFUNCTION(BlueprintCallable, Category="Motion Trajectory")
 	FTrajectorySampleRange GetHistory() const;
 
-	// Prediction simulation settings
+	// Prediction trajectory simulation settings
 	UPROPERTY(EditDefaultsOnly, BlueprintReadOnly, Category="Settings")
 	FMotionTrajectorySettings PredictionSettings;
 
-	// History retention settings
+	// Historical trajectory settings
 	UPROPERTY(EditDefaultsOnly, BlueprintReadOnly, Category="Settings")
 	FMotionTrajectorySettings HistorySettings;
 
@@ -141,4 +141,13 @@ public:
 	// Consider enabling this option when accurate trajectory sample speed at low tick rates (high interpolation usage) is important
 	UPROPERTY(EditAnywhere, BlueprintReadWrite, Category="Settings")
 	bool bSmoothInterpolation = false;
+
+#if WITH_EDITORONLY_DATA
+	// Debug drawing options for the complete trajectory range
+	// a.MotionTrajectory.Debug (0/1) to enable or disable
+	// a.MotionTrajectory.Options (0-6) to enable additional sample information
+	// a.MotionTrajectory.Stride (*) to stride sample information display by a specified modulo
+	UPROPERTY(EditAnywhere, Category="Debug")
+	bool bDebugDrawTrajectory = false;
+#endif
 };
\ No newline at end of file
diff --git a/Engine/Plugins/Experimental/Animation/MotionTrajectory/Source/MotionTrajectory/Public/MotionTrajectoryLibrary.h b/Engine/Plugins/Experimental/Animation/MotionTrajectory/Source/MotionTrajectory/Public/MotionTrajectoryLibrary.h
--- a/Engine/Plugins/Experimental/Animation/MotionTrajectory/Source/MotionTrajectory/Public/MotionTrajectoryLibrary.h
+++ b/Engine/Plugins/Experimental/Animation/MotionTrajectory/Source/MotionTrajectory/Public/MotionTrajectoryLibrary.h
@@ -48,6 +48,29 @@ public:
 	*
 	* @return					Direction clamped, modified trajectory range
 	*/
-	UFUNCTION(BlueprintPure, Category = "Motion Trajectory", meta=(BlueprintThreadSafe, AutoCreateRefTerm="Directions"))
+	UFUNCTION(BlueprintPure, Category="Motion Trajectory", meta=(BlueprintThreadSafe, AutoCreateRefTerm="Directions"))
 	static FTrajectorySampleRange ClampTrajectoryDirection(FTrajectorySampleRange Trajectory, const TArray<FTrajectoryDirectionClamp>& Directions);
+
+	/**
+	* Projects trajectory samples onto a defined set of allowed directions
+	*
+	* @param WorldTransform		Input world transform to establish a world-space basis for drawing the trajectory
+	* @param Trajectory			Input trajectory range
+	* @param PredictionColor	Input prediction color to interpolate sample drawing towards
+	* @param HistoryColor		Input historical color to interpolate sample drawing towards
+	* @param ArrowScale			Input sample velocity draw scale
+	* @param ArrowSize			Input sample arrow draw size
+	* @param ArrowThickness		Input sample arrow draw thickness
+	* @param bDrawText			Input include drawing of per-sample trajectory information
+	*/
+	UFUNCTION(BlueprintPure, Category="Motion Trajectory", meta=(BlueprintThreadSafe, AutoCreateRefTerm="WorldTransform"))
+	static void DebugDrawTrajectory(const AActor* Actor
+		, const FTransform& WorldTransform
+		, const FTrajectorySampleRange& Trajectory
+		, const FLinearColor PredictionColor = FLinearColor(0.f, 1.f, 0.f)
+		, const FLinearColor HistoryColor = FLinearColor(0.f, 0.f, 1.f)
+		, float ArrowScale = 0.025f
+		, float ArrowSize = 40.f
+		, float ArrowThickness = 2.f
+	);
 };
\ No newline at end of file
diff --git a/Engine/Plugins/Experimental/Animation/PoseSearch/Source/Runtime/Private/PoseSearch.cpp b/Engine/Plugins/Experimental/Animation/PoseSearch/Source/Runtime/Private/PoseSearch.cpp
--- a/Engine/Plugins/Experimental/Animation/PoseSearch/Source/Runtime/Private/PoseSearch.cpp
+++ b/Engine/Plugins/Experimental/Animation/PoseSearch/Source/Runtime/Private/PoseSearch.cpp
@@ -37,8 +37,8 @@ DEFINE_LOG_CATEGORY(LogPoseSearch);
 
 constexpr float DrawDebugLineThickness = 2.0f;
 constexpr float DrawDebugPointSize = 3.0f;
-constexpr float DrawDebugVelocityScale = 0.1f;
-constexpr float DrawDebugArrowSize = 5.0f;
+constexpr float DrawDebugVelocityScale = 0.08f;
+constexpr float DrawDebugArrowSize = 30.0f;
 constexpr float DrawDebugSphereSize = 3.0f;
 constexpr int32 DrawDebugSphereSegments = 8;
 constexpr float DrawDebugSphereLineThickness = 0.5f;
@@ -74,10 +74,16 @@ static inline float CompareFeatureVectors(int32 NumValues, const float* A, const
 
 FLinearColor GetColorForFeature(FPoseSearchFeatureDesc Feature, const FPoseSearchFeatureVectorLayout* Layout)
 {
-	int32 FeatureIdx = Layout->Features.IndexOfByKey(Feature);
+	const float FeatureIdx = Layout->Features.IndexOfByKey(Feature);
+	const float FeatureCountIdx = Layout->Features.Num() - 1;
+	const float FeatureCountIdxHalf = FeatureCountIdx / 2.f;
 	check(FeatureIdx != INDEX_NONE);
-	float Lerp = (float)(FeatureIdx) / (Layout->Features.Num() - 1);
-	FLinearColor ColorHSV(Lerp * 360.0f, 0.8f, 0.5f, 1.0f);
+
+	const float Hue = FeatureIdx < FeatureCountIdxHalf
+		? FMath::GetMappedRangeValueUnclamped({ 0.f, FeatureCountIdxHalf }, FVector2D(60.f, 0.f), FeatureIdx)
+		: FMath::GetMappedRangeValueUnclamped({ FeatureCountIdxHalf, FeatureCountIdx }, FVector2D(280.f, 220.f), FeatureIdx);
+
+	const FLinearColor ColorHSV(Hue, 1.f, 1.f);
 	return ColorHSV.HSVToLinearRGB();
 }
 
diff --git a/Engine/Plugins/Experimental/Animation/PoseSearch/Source/Runtime/Private/PoseSearchLibrary.cpp b/Engine/Plugins/Experimental/Animation/PoseSearch/Source/Runtime/Private/PoseSearchLibrary.cpp
--- a/Engine/Plugins/Experimental/Animation/PoseSearch/Source/Runtime/Private/PoseSearchLibrary.cpp
+++ b/Engine/Plugins/Experimental/Animation/PoseSearch/Source/Runtime/Private/PoseSearchLibrary.cpp
@@ -226,7 +226,7 @@ void UPoseSearchLibrary::UpdateMotionMatchingForSequencePlayer(const FAnimUpdate
 			if ((InOutMotionMatchingState.Flags & EMotionMatchingFlags::JumpedToPose) == EMotionMatchingFlags::JumpedToPose)
 			{
 				const FPoseSearchDatabaseSequence& ResultDbSequence = Database->Sequences[InOutMotionMatchingState.DbSequenceIdx];
-				SequencePlayerNode->SetSequence(ResultDbSequence.Sequence);
+				SequencePlayerNode->SetSequence(Cast<UAnimSequenceBase>(ResultDbSequence.Sequence));
 				SequencePlayerNode->SetAccumulatedTime(InOutMotionMatchingState.AssetPlayerTime);
 				SequencePlayerNode->SetLoopAnimation(ResultDbSequence.bLoopAnimation);
 				SequencePlayerNode->SetPlayRate(1.f);
diff --git a/Engine/Source/Runtime/Engine/Private/Animation/MotionTrajectoryTypes.cpp b/Engine/Source/Runtime/Engine/Private/Animation/MotionTrajectoryTypes.cpp
--- a/Engine/Source/Runtime/Engine/Private/Animation/MotionTrajectoryTypes.cpp
+++ b/Engine/Source/Runtime/Engine/Private/Animation/MotionTrajectoryTypes.cpp
@@ -2,6 +2,22 @@
 
 #include "Animation/MotionTrajectoryTypes.h"
 #include "Algo/AllOf.h"
+#include "DrawDebugHelpers.h"
+
+#if ENABLE_ANIM_DEBUG
+constexpr int32 DebugTrajectorySampleDisable = 0;
+constexpr int32 DebugTrajectorySampleCount = 1;
+constexpr int32 DebugTrajectorySampleTime = 2;
+constexpr int32 DebugTrajectorySampleDistance = 3;
+constexpr int32 DebugTrajectorySamplePosition = 4;
+constexpr int32 DebugTrajectorySampleVelocity = 5;
+constexpr int32 DebugTrajectorySampleAccel = 6;
+TAutoConsoleVariable<int32> CVarMotionTrajectoryDebug(TEXT("a.MotionTrajectory.Debug"), 0, TEXT("Turn on debug drawing for motion trajectory"));
+TAutoConsoleVariable<int32> CVarMotionTrajectoryDebugStride(TEXT("a.MotionTrajectory.Stride"), 1, TEXT("Configure the sample stride when displaying information"));
+TAutoConsoleVariable<int32> CVarMotionTrajectoryDebugOptions(TEXT("a.MotionTrajectory.Options"), 0,
+	TEXT("Toggle motion trajectory sample information:\n 0. Disable Text\n 1. Index\n2. Accumulated Time\n 3. Accumulated Distance\n 4. Position\n 5. Velocity\n 6. Acceleration")
+);
+#endif
 
 namespace
 {
@@ -86,4 +102,71 @@ void FTrajectorySampleRange::RemoveHistory()
 		{
 			return Sample.AccumulatedSeconds < 0.f;
 		});
+}
+
+void FTrajectorySampleRange::DebugDrawTrajectory(bool bEnable
+	, const UWorld* World
+	, const FTransform& WorldTransform
+	, const FLinearColor PredictionColor
+	, const FLinearColor HistoryColor
+	, float ArrowScale
+	, float ArrowSize
+	, float ArrowThickness) const
+{
+	if (bEnable
+#if ENABLE_ANIM_DEBUG
+		|| CVarMotionTrajectoryDebug.GetValueOnAnyThread()
+#endif
+		)
+	{
+		if (World)
+		{
+			const int32 DebugSampleStride = CVarMotionTrajectoryDebugStride.GetValueOnAnyThread();
+			const int32 DebugSampleOptions = CVarMotionTrajectoryDebugOptions.GetValueOnAnyThread();
+
+			for (int32 Idx = 0, Num = Samples.Num(); Idx < Num; Idx++)
+			{
+				const FVector WorldPosition = WorldTransform.TransformPosition(Samples[Idx].Position);
+				const FVector WorldVelocity = (WorldTransform.TransformVector(Samples[Idx].LocalLinearVelocity) * ArrowScale) + WorldPosition;
+
+				// Interpolate the history and prediction color over the entire trajectory range
+				const float ColorLerp = static_cast<float>(Idx) / static_cast<float>(Num);
+				const FLinearColor Color = FLinearColor::LerpUsingHSV(PredictionColor, HistoryColor, ColorLerp);
+
+				DrawDebugDirectionalArrow(World, WorldPosition, WorldVelocity, ArrowSize, Color.ToFColor(true), false, 0.f, 0, ArrowThickness);
+#if ENABLE_ANIM_DEBUG
+				FString DebugSampleString;
+				switch (DebugSampleOptions)
+				{
+				case 1: // Sample Index
+					DebugSampleString = DebugSampleString.Format(TEXT("{0}"), { Idx });
+;					break;
+				case 2: // Sample Accumulated Time
+					DebugSampleString = DebugSampleString.Format(TEXT("{0}"), { Samples[Idx].AccumulatedSeconds });
+					break;
+				case 3: // Sample Accumulated Distance
+					DebugSampleString = DebugSampleString.Format(TEXT("{0}"), { Samples[Idx].AccumulatedDistance });
+					break;
+				case 4: // Sample Position
+					DebugSampleString = DebugSampleString.Format(TEXT("{0}"), { Samples[Idx].Position.ToCompactString() });
+					break;
+				case 5: // Sample Velocity
+					DebugSampleString = DebugSampleString.Format(TEXT("{0}"), { Samples[Idx].LocalLinearVelocity.ToCompactString() });
+					break;
+				case 6: // Sample Acceleration
+					DebugSampleString = DebugSampleString.Format(TEXT("{0}"), { Samples[Idx].LocalLinearAcceleration.ToCompactString() });
+					break;
+				default:
+					break;
+				}
+
+				// Conditionally display per-sample information against a specified stride
+				if (!DebugSampleString.IsEmpty() && !!DebugSampleStride && (Idx % DebugSampleStride == 0))
+				{
+					DrawDebugString(World, WorldVelocity + FVector(0.f, 0.f, 10.f), DebugSampleString, nullptr, FColor::White, 0.f, false, 1.f);
+				}
+#endif
+			}
+		}
+	}
 }
\ No newline at end of file
diff --git a/Engine/Source/Runtime/Engine/Public/Animation/MotionTrajectoryTypes.h b/Engine/Source/Runtime/Engine/Public/Animation/MotionTrajectoryTypes.h
--- a/Engine/Source/Runtime/Engine/Public/Animation/MotionTrajectoryTypes.h
+++ b/Engine/Source/Runtime/Engine/Public/Animation/MotionTrajectoryTypes.h
@@ -63,6 +63,13 @@ struct ENGINE_API FTrajectorySampleRange
 {
 	GENERATED_BODY()
 
+	// Debug rendering contants
+	static constexpr FLinearColor DebugDefaultPredictionColor{ 0.f, 1.f, 0.f };
+	static constexpr FLinearColor DebugDefaultHistoryColor{ 0.f, 0.f, 1.f };
+	static constexpr float DebugDefaultArrowScale = 0.025f;
+	static constexpr float DebugDefaultArrowSize = 40.f;
+	static constexpr float DebugDefaultArrowThickness = 2.f;
+
 	// Linearly ordered (Time or Distance domain) container for past, present, and future trajectory samples
 	UPROPERTY(EditAnywhere, BlueprintReadWrite, Category="Motion Trajectory")
 	TArray<FTrajectorySample> Samples;
@@ -85,6 +92,16 @@ struct ENGINE_API FTrajectorySampleRange
 	// Determine if all trajectory samples are default values
 	bool HasOnlyZeroSamples() const;
 
+	// Debug draw in-world trajectory samples and optional corresponding information
+	void DebugDrawTrajectory(bool bEnable
+		, const UWorld* World
+		, const FTransform& WorldTransform
+		, const FLinearColor PredictionColor = DebugDefaultPredictionColor
+		, const FLinearColor HistoryColor = DebugDefaultHistoryColor
+		, float ArrowScale = DebugDefaultArrowScale
+		, float ArrowSize = DebugDefaultArrowSize
+		, float ArrowThickness = DebugDefaultArrowThickness) const;
+
 	// Iterator for precise subsampling of the trajectory
 	template <typename Container> static FTrajectorySample IterSampleTrajectory(const Container& Samples, ETrajectorySampleDomain DomainType, float DomainValue, int32& InitialIdx, bool bSmoothInterp = false)
 	{
diff --git a/Engine/Source/Runtime/Experimental/Animation/DynamicPlayRate/Private/DynamicPlayRate/DynamicPlayRateLibrary.cpp b/Engine/Source/Runtime/Experimental/Animation/DynamicPlayRate/Private/DynamicPlayRate/DynamicPlayRateLibrary.cpp
--- a/Engine/Source/Runtime/Experimental/Animation/DynamicPlayRate/Private/DynamicPlayRate/DynamicPlayRateLibrary.cpp
+++ b/Engine/Source/Runtime/Experimental/Animation/DynamicPlayRate/Private/DynamicPlayRate/DynamicPlayRateLibrary.cpp
@@ -11,8 +11,8 @@
 DEFINE_LOG_CATEGORY_STATIC(LogDynamicPlayRateLibrary, Verbose, All);
 
 #if ENABLE_ANIM_DEBUG
-TAutoConsoleVariable<int32> CVarAnimNodeDynamicPlayRateDebug(TEXT("a.DynamicPlayRate.Debug"), 0, TEXT("Turn on debug for dynamic play rate adjustment"));
-TAutoConsoleVariable<int32> CVarAnimNodeDynamicPlayRateEnable(TEXT("a.DynamicPlayRate.Enable"), 1, TEXT("Toggle dynamic play rate adjustment"));
+TAutoConsoleVariable<int32> CVarDynamicPlayRateDebug(TEXT("a.DynamicPlayRate.Debug"), 0, TEXT("Turn on debug for dynamic play rate adjustment"));
+TAutoConsoleVariable<int32> CVarDynamicPlayRateEnable(TEXT("a.DynamicPlayRate.Enable"), 1, TEXT("Toggle dynamic play rate adjustment"));
 #endif
 
 float FDynamicPlayRateSettings::ComputePlayRate(float PlayRate, float DeltaTime) const
@@ -37,12 +37,12 @@ float DynamicPlayRateAdjustment(const FAnimationUpdateContext& Context
 
 #if ENABLE_ANIM_DEBUG
 	// Debug enable/disable toggle for play rate scaling
-	if (!CVarAnimNodeDynamicPlayRateEnable.GetValueOnAnyThread())
+	if (!CVarDynamicPlayRateEnable.GetValueOnAnyThread())
 	{
 		return PlayRate;
 	}
 #if WITH_EDITORONLY_DATA
-	bool bDebugDraw = Settings.bDebugDraw || CVarAnimNodeDynamicPlayRateDebug.GetValueOnAnyThread();
+	bool bDebugDraw = Settings.bDebugDraw || CVarDynamicPlayRateDebug.GetValueOnAnyThread();
 #endif
 #else
 #if WITH_EDITORONLY_DATA
