Motion Trajectory Component debugging functionality minor improvement.

#trivial

diff --git a/Engine/Source/Runtime/Engine/Private/Animation/MotionTrajectoryTypes.cpp b/Engine/Source/Runtime/Engine/Private/Animation/MotionTrajectoryTypes.cpp
--- a/Engine/Source/Runtime/Engine/Private/Animation/MotionTrajectoryTypes.cpp
+++ b/Engine/Source/Runtime/Engine/Private/Animation/MotionTrajectoryTypes.cpp
@@ -135,25 +135,32 @@ void FTrajectorySampleRange::DebugDrawTrajectory(bool bEnable
 
 				DrawDebugDirectionalArrow(World, WorldPosition, WorldVelocity, ArrowSize, Color.ToFColor(true), false, 0.f, 0, ArrowThickness);
 #if ENABLE_ANIM_DEBUG
+				FString DebugString;
 				FString DebugSampleString;
 				switch (DebugSampleOptions)
 				{
 				case 1: // Sample Index
+					DebugString = "Sample Index:";
 					DebugSampleString = DebugSampleString.Format(TEXT("{0}"), { Idx });
 ;					break;
 				case 2: // Sample Accumulated Time
+					DebugString = "Sample Time:";
 					DebugSampleString = DebugSampleString.Format(TEXT("{0}"), { Samples[Idx].AccumulatedSeconds });
 					break;
 				case 3: // Sample Accumulated Distance
+					DebugString = "Sample Distance:";
 					DebugSampleString = DebugSampleString.Format(TEXT("{0}"), { Samples[Idx].AccumulatedDistance });
 					break;
 				case 4: // Sample Position
+					DebugString = "Sample Position:";
 					DebugSampleString = DebugSampleString.Format(TEXT("{0}"), { Samples[Idx].Position.ToCompactString() });
 					break;
 				case 5: // Sample Velocity
+					DebugString = "Sample Velocity:";
 					DebugSampleString = DebugSampleString.Format(TEXT("{0}"), { Samples[Idx].LocalLinearVelocity.ToCompactString() });
 					break;
 				case 6: // Sample Acceleration
+					DebugString = "Sample Acceleration:";
 					DebugSampleString = DebugSampleString.Format(TEXT("{0}"), { Samples[Idx].LocalLinearAcceleration.ToCompactString() });
 					break;
 				default:
@@ -163,6 +170,12 @@ void FTrajectorySampleRange::DebugDrawTrajectory(bool bEnable
 				// Conditionally display per-sample information against a specified stride
 				if (!DebugSampleString.IsEmpty() && !!DebugSampleStride && (Idx % DebugSampleStride == 0))
 				{
+					// One time debug drawing of the per-sample type description
+					if (!DebugString.IsEmpty() && Idx == 0)
+					{
+						DrawDebugString(World, WorldTransform.GetLocation() + FVector(0.f, 0.f, 50.f), DebugString, nullptr, FColor::White, 0.f, false, 1.f);
+					}
+
 					DrawDebugString(World, WorldVelocity + FVector(0.f, 0.f, 10.f), DebugSampleString, nullptr, FColor::White, 0.f, false, 1.f);
 				}
 #endif
