Motion Trajectory Horde compiler fixes

diff --git a/Engine/Source/Runtime/Engine/Private/Animation/MotionTrajectoryTypes.cpp b/Engine/Source/Runtime/Engine/Private/Animation/MotionTrajectoryTypes.cpp
--- a/Engine/Source/Runtime/Engine/Private/Animation/MotionTrajectoryTypes.cpp
+++ b/Engine/Source/Runtime/Engine/Private/Animation/MotionTrajectoryTypes.cpp
@@ -121,9 +121,10 @@ void FTrajectorySampleRange::DebugDrawTrajectory(bool bEnable
 	{
 		if (World)
 		{
+#if ENABLE_ANIM_DEBUG
 			const int32 DebugSampleStride = CVarMotionTrajectoryDebugStride.GetValueOnAnyThread();
 			const int32 DebugSampleOptions = CVarMotionTrajectoryDebugOptions.GetValueOnAnyThread();
-
+#endif
 			for (int32 Idx = 0, Num = Samples.Num(); Idx < Num; Idx++)
 			{
 				const FVector WorldPosition = WorldTransform.TransformPosition(Samples[Idx].Position);
