Fixed PIE bug caused by template actor container (used to handle dirty actors in a PIE session) that was outered to the PIE world including all its template actors.

diff --git a/Engine/Source/Runtime/Engine/Private/WorldPartition/WorldPartitionRuntimeSpatialHashCell.cpp b/Engine/Source/Runtime/Engine/Private/WorldPartition/WorldPartitionRuntimeSpatialHashCell.cpp
--- a/Engine/Source/Runtime/Engine/Private/WorldPartition/WorldPartitionRuntimeSpatialHashCell.cpp
+++ b/Engine/Source/Runtime/Engine/Private/WorldPartition/WorldPartitionRuntimeSpatialHashCell.cpp
@@ -29,6 +29,9 @@ void UWorldPartitionRuntimeSpatialHashCell::PostDuplicate(bool bDuplicateForPIE)
 
 	if (ActorContainer)
 	{
+		// Make sure actor container isn't under PIE World so those template actors will never be considered part of the world.
+		ActorContainer->Rename(nullptr, GetPackage());
+
 		for (auto& ActorPair : ActorContainer->Actors)
 		{
 			// Don't use AActor::Rename here since the actor is not par of the world, it's only a duplication template.
