Build fix after 17553079.
df

diff --git a/Engine/Source/Runtime/Engine/Classes/Engine/TextureLODSettings.h b/Engine/Source/Runtime/Engine/Classes/Engine/TextureLODSettings.h
--- a/Engine/Source/Runtime/Engine/Classes/Engine/TextureLODSettings.h
+++ b/Engine/Source/Runtime/Engine/Classes/Engine/TextureLODSettings.h
@@ -18,6 +18,7 @@ struct ENGINE_API FTextureLODGroup
 
 	FTextureLODGroup()
 		: Group(TEXTUREGROUP_World)
+		, MaxLODMipCount(32)
 		, LODBias(0)
 		, LODBias_Smaller(-1)
 		, LODBias_Smallest(-1)
@@ -48,6 +49,9 @@ struct ENGINE_API FTextureLODGroup
 	UPROPERTY()
 	TEnumAsByte<TextureGroup> Group;
 
+	/** Maximum LOD mip count. Bias will be adjusted so texture won't go above.		*/
+	int32 MaxLODMipCount;
+	
 	/** Group LOD bias.																*/
 	UPROPERTY()
 	int32 LODBias;
diff --git a/Engine/Source/Runtime/Engine/Private/Streaming/StreamingManagerTexture.cpp b/Engine/Source/Runtime/Engine/Private/Streaming/StreamingManagerTexture.cpp
--- a/Engine/Source/Runtime/Engine/Private/Streaming/StreamingManagerTexture.cpp
+++ b/Engine/Source/Runtime/Engine/Private/Streaming/StreamingManagerTexture.cpp
@@ -2622,7 +2622,7 @@ bool FRenderAssetStreamingManager::HandleInvestigateRenderAssetCommand(const TCH
 						const int32 MipCountBeforeMaxRes = ResourceState.MaxNumLODs - RenderAsset->NumCinematicMipLevels -
 							(StreamingRenderAsset.LODGroup == TEXTUREGROUP_UI ? GUITextureLODBias : 0) - 
 							(FPlatformProperties::RequiresCookedData() ? 0 : (LODGroupInfo.LODBias + (Texture ? Texture->LODBias : 0)));
-						const int32 MaxResBias = MipCountBeforeMaxRes - (FMath::CeilLogTwo(LODGroupInfo.MaxLODSize) + 1);
+						const int32 MaxResBias = MipCountBeforeMaxRes - (LODGroupInfo.MaxLODMipCount + 1);
 						if (MaxResBias > 0)
 						{
 							BiasDesc += FString::Printf(TEXT(" [LODGroup.MaxRes:%d]"), MaxResBias);
diff --git a/Engine/Source/Runtime/Engine/Private/TextureLODSettings.cpp b/Engine/Source/Runtime/Engine/Private/TextureLODSettings.cpp
--- a/Engine/Source/Runtime/Engine/Private/TextureLODSettings.cpp
+++ b/Engine/Source/Runtime/Engine/Private/TextureLODSettings.cpp
@@ -59,6 +59,8 @@ void FTextureLODGroup::SetupGroup()
 		}
 	}
 
+	MaxLODMipCount = FMath::CeilLogTwo(MaxLODSize);
+
 	// Linear filtering
 	if (MinMagFilter == NAME_Linear)
 	{
