Respect FirstMipOffset in places where it wasn't being respected.
This value handles LodBias emulation without a texture recook, so when it's not respected texture LodBias setting doesn't work correctly in editor.

diff --git a/Engine/Source/Runtime/Engine/Private/VT/UploadingVirtualTexture.cpp b/Engine/Source/Runtime/Engine/Private/VT/UploadingVirtualTexture.cpp
--- a/Engine/Source/Runtime/Engine/Private/VT/UploadingVirtualTexture.cpp
+++ b/Engine/Source/Runtime/Engine/Private/VT/UploadingVirtualTexture.cpp
@@ -96,6 +96,8 @@ FUploadingVirtualTexture::~FUploadingVirtualTexture()
 
 uint32 FUploadingVirtualTexture::GetLocalMipBias(uint8 vLevel, uint32 vAddress) const
 {
+	vLevel += FirstMipOffset;
+
 	const uint32 NumMips = Data->NumMips;
 	uint32 Current_vLevel = vLevel;
 	uint32 Current_vAddress = vAddress;
@@ -123,8 +125,10 @@ uint32 FUploadingVirtualTexture::GetLocalMipBias(uint8 vLevel, uint32 vAddress)
 
 FVTRequestPageResult FUploadingVirtualTexture::RequestPageData(const FVirtualTextureProducerHandle& ProducerHandle, uint8 LayerMask, uint8 vLevel, uint64 vAddress, EVTRequestPagePriority Priority)
 {
-	check(vAddress <= MAX_uint32); // Not supporting 64 bit vAddress here. Only currrently supported for adaptive runtime virtual texture.
-	return StreamingManager->RequestTile(this, ProducerHandle, LayerMask, FirstMipOffset + vLevel, (uint32)vAddress, Priority);
+	vLevel += FirstMipOffset;
+
+	check(vAddress <= MAX_uint32); // Not supporting 64 bit vAddress here. Only currently supported for adaptive runtime virtual texture.
+	return StreamingManager->RequestTile(this, ProducerHandle, LayerMask, vLevel, (uint32)vAddress, Priority);
 }
 
 IVirtualTextureFinalizer* FUploadingVirtualTexture::ProducePageData(FRHICommandListImmediate& RHICmdList,
