Color Management: Fixed inconsistent capitalization for RED Log encodings.

diff --git a/Engine/Source/Runtime/ColorManagement/Private/TransferFunctions.cpp b/Engine/Source/Runtime/ColorManagement/Private/TransferFunctions.cpp
--- a/Engine/Source/Runtime/ColorManagement/Private/TransferFunctions.cpp
+++ b/Engine/Source/Runtime/ColorManagement/Private/TransferFunctions.cpp
@@ -18,7 +18,7 @@ FORCEINLINE TFunction<float(float)> GetTransferFunction(EEncoding SourceEncoding
 	case EEncoding::Gamma22:     return bIsEncode ? &EncodeGamma22 : &DecodeGamma22;
 	case EEncoding::BT1886:      return bIsEncode ? &EncodeBT1886 : &DecodeBT1886;
 	case EEncoding::Cineon:      return bIsEncode ? &EncodeCineon : &DecodeCineon;
-	case EEncoding::RedLog:      return bIsEncode ? &EncodeRedLog : &DecodeRedLog;
+	case EEncoding::REDLog:      return bIsEncode ? &EncodeREDLog : &DecodeREDLog;
 	case EEncoding::REDLog3G10:  return bIsEncode ? &EncodeREDLog3G10 : &DecodeREDLog3G10;
 	case EEncoding::SLog1:       return bIsEncode ? &EncodeSLog1 : &DecodeSLog1;
 	case EEncoding::SLog2:       return bIsEncode ? &EncodeSLog2 : &DecodeSLog2;
diff --git a/Engine/Source/Runtime/ColorManagement/Public/ColorManagementDefines.h b/Engine/Source/Runtime/ColorManagement/Public/ColorManagementDefines.h
--- a/Engine/Source/Runtime/ColorManagement/Public/ColorManagementDefines.h
+++ b/Engine/Source/Runtime/ColorManagement/Public/ColorManagementDefines.h
@@ -19,7 +19,7 @@ enum class EEncoding : uint8 {
 	Gamma22,
 	BT1886,
 	Cineon,
-	RedLog,
+	REDLog,
 	REDLog3G10,
 	SLog1,
 	SLog2,
diff --git a/Engine/Source/Runtime/ColorManagement/Public/TransferFunctions.h b/Engine/Source/Runtime/ColorManagement/Public/TransferFunctions.h
--- a/Engine/Source/Runtime/ColorManagement/Public/TransferFunctions.h
+++ b/Engine/Source/Runtime/ColorManagement/Public/TransferFunctions.h
@@ -166,7 +166,7 @@ FORCEINLINE float DecodeCineon(float Value)
 *
 * @return float encoded value.
 */
-FORCEINLINE float EncodeRedLog(float Value)
+FORCEINLINE float EncodeREDLog(float Value)
 {
 	const float BlackOffset = FGenericPlatformMath::Pow(10.0f, (0.0f - 1023.0f) / 511.0f);
 	return (1023.0f + 511.0f * FGenericPlatformMath::LogX(10.0f, Value * (1.0f - BlackOffset) + BlackOffset)) / 1023.0f;
@@ -177,7 +177,7 @@ FORCEINLINE float EncodeRedLog(float Value)
 *
 * @return float decoded value.
 */
-FORCEINLINE float DecodeRedLog(float Value)
+FORCEINLINE float DecodeREDLog(float Value)
 {
 	const float BlackOffset = FGenericPlatformMath::Pow(10.0f, (0.0f - 1023.0f) / 511.0f);
 	return (FGenericPlatformMath::Pow(10.0f, (1023.0f * Value - 1023.0f) / 511.0f) - BlackOffset) / (1.0f - BlackOffset);
diff --git a/Engine/Source/Runtime/Engine/Classes/Engine/TextureDefines.h b/Engine/Source/Runtime/Engine/Classes/Engine/TextureDefines.h
--- a/Engine/Source/Runtime/Engine/Classes/Engine/TextureDefines.h
+++ b/Engine/Source/Runtime/Engine/Classes/Engine/TextureDefines.h
@@ -307,8 +307,8 @@ enum ETextureSourceEncoding
 	TSE_Gamma22 = 4 UMETA(DisplayName = "Gamma 2.2", ToolTip = "Gamma 2.2 source encoding to be linearized (before optional sRGB encoding is applied)."),
 	TSE_BT1886 = 5 UMETA(DisplayName = "BT1886", ToolTip = "BT1886 source encoding to be linearized (before optional sRGB encoding is applied)."),
 	TSE_Cineon = 6 UMETA(DisplayName = "Cineon", ToolTip = "Cineon source encoding to be linearized (before optional sRGB encoding is applied)."),
-	TSE_RedLog = 7 UMETA(DisplayName = "RedLog", ToolTip = "RedLog source encoding to be linearized (before optional sRGB encoding is applied)."),
-	TSE_REDLog3G10 = 8 UMETA(DisplayName = "REDLog3G10", ToolTip = "REDLog3G10 source encoding to be linearized (before optional sRGB encoding is applied)."),
+	TSE_REDLog = 7 UMETA(DisplayName = "REDLog", ToolTip = "RED Log source encoding to be linearized (before optional sRGB encoding is applied)."),
+	TSE_REDLog3G10 = 8 UMETA(DisplayName = "REDLog3G10", ToolTip = "RED Log3G10 source encoding to be linearized (before optional sRGB encoding is applied)."),
 	TSE_SLog1 = 9 UMETA(DisplayName = "SLog1", ToolTip = "Sony SLog1 source encoding to be linearized (before optional sRGB encoding is applied)."),
 	TSE_SLog2 = 10 UMETA(DisplayName = "SLog2", ToolTip = "Sony SLog2 source encoding to be linearized (before optional sRGB encoding is applied)."),
 	TSE_SLog3 = 11 UMETA(DisplayName = "SLog3", ToolTip = "Sony SLog3 source encoding to be linearized (before optional sRGB encoding is applied)."),
diff --git a/Engine/Source/Runtime/Engine/Private/TextureDerivedData.cpp b/Engine/Source/Runtime/Engine/Private/TextureDerivedData.cpp
--- a/Engine/Source/Runtime/Engine/Private/TextureDerivedData.cpp
+++ b/Engine/Source/Runtime/Engine/Private/TextureDerivedData.cpp
@@ -64,7 +64,7 @@
 
 // Version number for the texture encoding enum types (UE::Color::EEncoding/ETextureSourceEncoding).
 // Make sure to increment upon changes to the list.
-#define TEXTURE_ENCODING_TYPES_VER 1
+#define TEXTURE_ENCODING_TYPES_VER 2
 
 #if ENABLE_COOK_STATS
 namespace TextureCookStats
