- Improved world partition cells streaming priority computation (better handling of multiple streaming source, angle importance).
- Added streaming source priority parameter to UWorldPartitionStreamingSourceComponent.
- Added debug view mode wp.Runtime.ShowRuntimeSpatialHashCellStreamingPriotity to show cell priority using a heatmap colorization.

diff --git a/Engine/Source/Runtime/Engine/Classes/Components/WorldPartitionStreamingSourceComponent.h b/Engine/Source/Runtime/Engine/Classes/Components/WorldPartitionStreamingSourceComponent.h
--- a/Engine/Source/Runtime/Engine/Classes/Components/WorldPartitionStreamingSourceComponent.h
+++ b/Engine/Source/Runtime/Engine/Classes/Components/WorldPartitionStreamingSourceComponent.h
@@ -60,6 +60,9 @@ class ENGINE_API UWorldPartitionStreamingSourceComponent : public UActorComponen
 	UPROPERTY(BlueprintReadWrite, EditAnywhere, Category = "Streaming")
 	TArray<FStreamingSourceShape> Shapes;
 
+	UPROPERTY(BlueprintReadWrite, EditAnywhere, Category = "Streaming")
+	EStreamingSourcePriority Priority;
+
 #if WITH_EDITOR
 	virtual bool CanEditChange(const FProperty* InProperty) const override;
 #endif
diff --git a/Engine/Source/Runtime/Engine/Private/Components/WorldPartitionStreamingSourceComponent.cpp b/Engine/Source/Runtime/Engine/Private/Components/WorldPartitionStreamingSourceComponent.cpp
--- a/Engine/Source/Runtime/Engine/Private/Components/WorldPartitionStreamingSourceComponent.cpp
+++ b/Engine/Source/Runtime/Engine/Private/Components/WorldPartitionStreamingSourceComponent.cpp
@@ -16,6 +16,7 @@ UWorldPartitionStreamingSourceComponent::UWorldPartitionStreamingSourceComponent
 	, DefaultVisualizerLoadingRange(10000.f)
 #endif
 	, TargetGrid(NAME_None)
+	, Priority(EStreamingSourcePriority::Low)
 	, bStreamingSourceEnabled(true)
 	, TargetState(EStreamingSourceTargetState::Activated)
 {
@@ -73,6 +74,7 @@ bool UWorldPartitionStreamingSourceComponent::GetStreamingSource(FWorldPartition
 		OutStreamingSource.TargetGrid = TargetGrid;
 		OutStreamingSource.TargetHLODLayer = TargetHLODLayer;
 		OutStreamingSource.Shapes = Shapes;
+		OutStreamingSource.Priority = Priority;
 		return true;
 	}
 	return false;
diff --git a/Engine/Source/Runtime/Engine/Private/WorldPartition/WorldPartitionDebugHelper.cpp b/Engine/Source/Runtime/Engine/Private/WorldPartition/WorldPartitionDebugHelper.cpp
--- a/Engine/Source/Runtime/Engine/Private/WorldPartition/WorldPartitionDebugHelper.cpp
+++ b/Engine/Source/Runtime/Engine/Private/WorldPartition/WorldPartitionDebugHelper.cpp
@@ -134,7 +134,7 @@ FAutoConsoleCommand FWorldPartitionDebugHelper::DebugFilterByCellNameCommand(
 	})
 );
 
-bool FWorldPartitionDebugHelper::IsDebugCellNameShow(const FString& CellName)
+bool FWorldPartitionDebugHelper::IsDebugCellNameShown(const FString& CellName)
 {
 	return !DebugCellNameFilter.Num() || Algo::AllOf(DebugCellNameFilter, [&CellName](const FString& NameFilter) { return CellName.Contains(NameFilter); });
 }
@@ -170,4 +170,44 @@ void FWorldPartitionDebugHelper::DrawLegendItem(UCanvas* Canvas, const FString&
 	}
 
 	Pos.Y = TextPos.Y;
+}
+
+bool FWorldPartitionDebugHelper::bShowRuntimeSpatialHashCellStreamingPriotity = false;
+FAutoConsoleVariableRef FWorldPartitionDebugHelper::ShowRuntimeSpatialHashCellStreamingPriotityCommand(
+	TEXT("wp.Runtime.ShowRuntimeSpatialHashCellStreamingPriotity"),
+	FWorldPartitionDebugHelper::bShowRuntimeSpatialHashCellStreamingPriotity,
+	TEXT("Enable to show a heatmap of the runtime spatial hash grid cells based on their priority."));
+
+FLinearColor FWorldPartitionDebugHelper::GetHeatMapColor(float ValueNormalized)
+{
+	const int32 NUM_COLORS = 4;
+	static FLinearColor Colors[NUM_COLORS] =
+	{
+		FLinearColor::Blue,
+		FLinearColor::Green,
+		FLinearColor::Yellow,
+		FLinearColor::Red
+	};
+
+	int32 idx1 = 0;
+	int32 idx2 = 0;
+	float Fraction = 0.f;
+
+	if (ValueNormalized >= 1.f)
+	{
+		idx1 = idx2 = NUM_COLORS - 1;
+	}
+	else if (ValueNormalized > 0)
+	{
+		// Find the 2 color indexes in which value lies.
+		ValueNormalized = ValueNormalized * (NUM_COLORS - 1);
+		idx1 = FMath::Floor(ValueNormalized);
+		idx2 = idx1 + 1;
+		// Calculate the fraction between the two indexes [0, 1].
+		Fraction = ValueNormalized - float(idx1);
+	}
+
+	// Interpolate between the 2 colors
+	FLinearColor Result = Colors[idx1] + Fraction * (Colors[idx2] - Colors[idx1]);
+	return Result;
 }
\ No newline at end of file
diff --git a/Engine/Source/Runtime/Engine/Private/WorldPartition/WorldPartitionRuntimeCell.cpp b/Engine/Source/Runtime/Engine/Private/WorldPartition/WorldPartitionRuntimeCell.cpp
--- a/Engine/Source/Runtime/Engine/Private/WorldPartition/WorldPartitionRuntimeCell.cpp
+++ b/Engine/Source/Runtime/Engine/Private/WorldPartition/WorldPartitionRuntimeCell.cpp
@@ -15,8 +15,11 @@ UWorldPartitionRuntimeCell::UWorldPartitionRuntimeCell(const FObjectInitializer&
 : Super(ObjectInitializer)
 , bIsAlwaysLoaded(false)
 , Priority(0)
-, CachedSourcePriority(0)
+, CachedMinSourcePriority((uint8)EStreamingSourcePriority::Lowest)
 , CachedSourceInfoEpoch(INT_MIN)
+#if !UE_BUILD_SHIPPING
+, DebugStreamingPriority(-1.f)
+#endif
 {}
 
 #if WITH_EDITOR
@@ -74,17 +77,21 @@ bool UWorldPartitionRuntimeCell::CacheStreamingSourceInfo(const UWorldPartitionR
 	{
 		CachedSourceInfoEpoch = UWorldPartitionRuntimeCell::StreamingSourceCacheEpoch;
 		bWasCacheDirtied = true;
+		CachedSourcePriorityWeights.Reset();
 	}
 
+	static_assert((uint8)EStreamingSourcePriority::Lowest == 255);
+	CachedSourcePriorityWeights.Add(1.f - ((float)Info.Source.Priority / 255.f));
+
 	// If cache was dirtied, use value, else use minimum with existing cached value
-	CachedSourcePriority = bWasCacheDirtied ? (int32)Info.Source.Priority : FMath::Min((int32)Info.Source.Priority, CachedSourcePriority);
+	CachedMinSourcePriority = bWasCacheDirtied ? (uint8)Info.Source.Priority : FMath::Min((uint8)Info.Source.Priority, CachedMinSourcePriority);
 	return bWasCacheDirtied;
 }
 
 int32 UWorldPartitionRuntimeCell::SortCompare(const UWorldPartitionRuntimeCell* Other) const
 {
 	// Source priority (lower value is higher prio)
-	const int32 Comparison = CachedSourcePriority - Other->CachedSourcePriority;
+	const int32 Comparison = (int32)CachedMinSourcePriority - (int32)Other->CachedMinSourcePriority;
 	// Cell priority (lower value is higher prio)
 	return (Comparison != 0) ? Comparison : (Priority - Other->Priority);
 }
@@ -94,5 +101,17 @@ bool UWorldPartitionRuntimeCell::IsDebugShown() const
 	return FWorldPartitionDebugHelper::IsDebugRuntimeHashGridShown(GridName) &&
 		   FWorldPartitionDebugHelper::IsDebugStreamingStatusShown(GetStreamingStatus()) &&
 		   FWorldPartitionDebugHelper::AreDebugDataLayersShown(DataLayers) &&
-		   FWorldPartitionDebugHelper::IsDebugCellNameShow(DebugName);
+		   FWorldPartitionDebugHelper::IsDebugCellNameShown(DebugName);
+}
+
+FLinearColor UWorldPartitionRuntimeCell::GetDebugStreamingPriorityColor() const
+{
+#if !UE_BUILD_SHIPPING
+	if ((CachedSourceInfoEpoch == UWorldPartitionRuntimeCell::StreamingSourceCacheEpoch) &&
+		(DebugStreamingPriority >= 0.f && DebugStreamingPriority <= 1.f))
+	{
+		return FWorldPartitionDebugHelper::GetHeatMapColor(1.f - DebugStreamingPriority);
+	}
+#endif
+	return FLinearColor::Transparent;
 }
\ No newline at end of file
diff --git a/Engine/Source/Runtime/Engine/Private/WorldPartition/WorldPartitionRuntimeLevelStreamingCell.cpp b/Engine/Source/Runtime/Engine/Private/WorldPartition/WorldPartitionRuntimeLevelStreamingCell.cpp
--- a/Engine/Source/Runtime/Engine/Private/WorldPartition/WorldPartitionRuntimeLevelStreamingCell.cpp
+++ b/Engine/Source/Runtime/Engine/Private/WorldPartition/WorldPartitionRuntimeLevelStreamingCell.cpp
@@ -7,6 +7,7 @@
 
 #include "WorldPartition/HLOD/HLODSubsystem.h"
 #include "WorldPartition/HLOD/HLODActor.h"
+#include "WorldPartition/WorldPartitionDebugHelper.h"
 
 #if WITH_EDITOR
 #include "WorldPartition/WorldPartitionLevelStreamingPolicy.h"
@@ -64,11 +65,26 @@ bool UWorldPartitionRuntimeLevelStreamingCell::IsLoading() const
 	return Super::IsLoading();
 }
 
-FLinearColor UWorldPartitionRuntimeLevelStreamingCell::GetDebugColor() const
+FLinearColor UWorldPartitionRuntimeLevelStreamingCell::GetDebugColor(EWorldPartitionRuntimeCellVisualizeMode VisualizeMode) const
 {
-	FLinearColor Color = LevelStreaming ? ULevelStreaming::GetLevelStreamingStatusColor(GetStreamingStatus()) : FLinearColor::Black;
-	Color.A = 0.25f / (Level + 1);
-	return Color;
+	switch (VisualizeMode)
+	{
+		case EWorldPartitionRuntimeCellVisualizeMode::StreamingPriority:
+		{
+			return GetDebugStreamingPriorityColor();
+		}
+		case EWorldPartitionRuntimeCellVisualizeMode::StreamingStatus:
+		{
+			// Return streaming status color
+			FLinearColor Color = LevelStreaming ? ULevelStreaming::GetLevelStreamingStatusColor(GetStreamingStatus()) : FLinearColor::Black;
+			Color.A = 0.25f / (Level + 1);
+			return Color;
+		}
+		default:
+		{
+			return Super::GetDebugColor(VisualizeMode);
+		}
+	}
 }
 
 void UWorldPartitionRuntimeLevelStreamingCell::SetIsAlwaysLoaded(bool bInIsAlwaysLoaded)
diff --git a/Engine/Source/Runtime/Engine/Private/WorldPartition/WorldPartitionRuntimeSpatialHash.cpp b/Engine/Source/Runtime/Engine/Private/WorldPartition/WorldPartitionRuntimeSpatialHash.cpp
--- a/Engine/Source/Runtime/Engine/Private/WorldPartition/WorldPartitionRuntimeSpatialHash.cpp
+++ b/Engine/Source/Runtime/Engine/Private/WorldPartition/WorldPartitionRuntimeSpatialHash.cpp
@@ -274,8 +274,15 @@
 	}
 }
 
+EWorldPartitionRuntimeCellVisualizeMode FSpatialHashStreamingGrid::GetStreamingCellVisualizeMode() const
+{
+	const EWorldPartitionRuntimeCellVisualizeMode VisualizeMode = FWorldPartitionDebugHelper::IsRuntimeSpatialHashCellStreamingPriotityShown() ? EWorldPartitionRuntimeCellVisualizeMode::StreamingPriority : EWorldPartitionRuntimeCellVisualizeMode::StreamingStatus;
+	return VisualizeMode;
+}
+
 void FSpatialHashStreamingGrid::Draw3D(UWorld* World, const TArray<FWorldPartitionStreamingSource>& Sources, const FTransform& Transform) const
 {
+	const EWorldPartitionRuntimeCellVisualizeMode VisualizeMode = GetStreamingCellVisualizeMode();
 	const UDataLayerSubsystem* DataLayerSubsystem = World->GetSubsystem<UDataLayerSubsystem>();
 	TMap<FName, FColor> DataLayerDebugColors;
 	if (DataLayerSubsystem)
@@ -335,7 +342,7 @@
 					}
 
 					// Draw Cell using its debug color
-					FColor CellColor = Cell->GetDebugColor().ToFColor(false).WithAlpha(16);
+					FColor CellColor = Cell->GetDebugColor(VisualizeMode).ToFColor(false).WithAlpha(64);
 					DrawDebugSolidBox(World, CellBox, CellColor, Transform, false, -1.f, 255);
 					FVector CellPos = Transform.TransformPosition(CellBox.GetCenter());
 					DrawDebugBox(World, CellPos, BoundsExtent, Transform.GetRotation(), CellColor.WithAlpha(255), false, -1.f, 255, 10.f);
@@ -400,6 +407,7 @@
 
 void FSpatialHashStreamingGrid::Draw2D(UCanvas* Canvas, UWorld* World, const TArray<FWorldPartitionStreamingSource>& Sources, const FBox& Region, const FBox2D& GridScreenBounds, TFunctionRef<FVector2D(const FVector2D&)> WorldToScreen) const
 {
+	const EWorldPartitionRuntimeCellVisualizeMode VisualizeMode = GetStreamingCellVisualizeMode();
 	const UDataLayerSubsystem* DataLayerSubsystem = World->GetSubsystem<UDataLayerSubsystem>();
 	TMap<FName, FColor> DataLayerDebugColors;
 	if (DataLayerSubsystem)
@@ -417,36 +425,9 @@
 	float MaxCellCoordTextWidth, MaxCellCoordTextHeight;
 	Canvas->StrLen(GEngine->GetTinyFont(), CellCoordString, MaxCellCoordTextWidth, MaxCellCoordTextHeight);
 
-	const UWorldPartitionRuntimeCell* DefaultEmptyCell = UWorldPartitionRuntimeCell::StaticClass()->GetDefaultObject<UWorldPartitionRuntimeCell>();
 	TArray<const UWorldPartitionRuntimeCell*> FilteredCells;
 	for (int32 GridLevel = MinGridLevel; GridLevel <= MaxGridLevel; ++GridLevel)
 	{
-		// Draw X/Y Axis
-		{
-			FCanvasLineItem Axis;
-			Axis.LineThickness = 3;
-			{
-				Axis.SetColor(FLinearColor::Red);
-				FVector2D LineStart = WorldToScreen(FVector2D(-1638400.f, 0.f));
-				FVector2D LineEnd = WorldToScreen(FVector2D(1638400.f, 0.f));
-				LineStart.X = FMath::Clamp(LineStart.X, GridScreenBounds.Min.X, GridScreenBounds.Max.X);
-				LineStart.Y = FMath::Clamp(LineStart.Y, GridScreenBounds.Min.Y, GridScreenBounds.Max.Y);
-				LineEnd.X = FMath::Clamp(LineEnd.X, GridScreenBounds.Min.X, GridScreenBounds.Max.X);
-				LineEnd.Y = FMath::Clamp(LineEnd.Y, GridScreenBounds.Min.Y, GridScreenBounds.Max.Y);
-				Axis.Draw(CanvasObject, LineStart, LineEnd);
-			}
-			{
-				Axis.SetColor(FLinearColor::Green);
-				FVector2D LineStart = WorldToScreen(FVector2D(0.f, -1638400.f));
-				FVector2D LineEnd = WorldToScreen(FVector2D(0.f, 1638400.f));
-				LineStart.X = FMath::Clamp(LineStart.X, GridScreenBounds.Min.X, GridScreenBounds.Max.X);
-				LineStart.Y = FMath::Clamp(LineStart.Y, GridScreenBounds.Min.Y, GridScreenBounds.Max.Y);
-				LineEnd.X = FMath::Clamp(LineEnd.X, GridScreenBounds.Min.X, GridScreenBounds.Max.X);
-				LineEnd.Y = FMath::Clamp(LineEnd.Y, GridScreenBounds.Min.Y, GridScreenBounds.Max.Y);
-				Axis.Draw(CanvasObject, LineStart, LineEnd);
-			}
-		}
-
 		// Draw Grid cells at desired grid level
 		const FSquare2DGridHelper& Helper = GetGridHelper();
 		Helper.Levels[GridLevel].ForEachIntersectingCells(Region, [&](const FIntVector2& Coords)
@@ -478,13 +459,6 @@
 				}
 			}
 
-			// Draw cell bounds
-			static FLinearColor CellDefaultColor = DefaultEmptyCell->GetDebugColor().CopyWithNewOpacity(0.25f);
-			FCanvasBoxItem Box(CellScreenBounds.Min, CellScreenBounds.GetSize());
-			Box.SetColor(CellDefaultColor);
-			Box.BlendMode = SE_BLEND_Translucent;
-			Canvas->DrawItem(Box);
-
 			const int32* LayerCellIndexPtr = GridLevels[GridLevel].LayerCellsMapping.Find(Coords.Y * Helper.Levels[GridLevel].GridSize + Coords.X);
 			const FSpatialHashStreamingGridLayerCell* LayerCell = LayerCellIndexPtr ? &GridLevels[GridLevel].LayerCells[*LayerCellIndexPtr] : nullptr;
 			GetFilteredCellsForDebugDraw(LayerCell, DataLayerSubsystem, FilteredCells);
@@ -497,7 +471,7 @@
 			{
 				// Draw Cell using its debug color
 				FVector2D StartPos = CellScreenBounds.Min + CellOffset;
-				FCanvasTileItem Item(StartPos, GWhiteTexture, CellBoundsSize, Cell->GetDebugColor().CopyWithNewOpacity(0.25f));
+				FCanvasTileItem Item(StartPos, GWhiteTexture, CellBoundsSize, Cell->GetDebugColor(VisualizeMode));
 				Item.BlendMode = SE_BLEND_Translucent;
 				Canvas->DrawItem(Item);
 				CellOffset.Y += CellBoundsSize.Y;
@@ -517,8 +491,40 @@
 					}
 				}
 			}
+
+			// Draw cell bounds
+			FCanvasBoxItem Box(CellScreenBounds.Min, CellScreenBounds.GetSize());
+			Box.SetColor(FLinearColor::Black);
+			Box.BlendMode = SE_BLEND_Translucent;
+			Canvas->DrawItem(Box);
 		});
 
+		// Draw X/Y Axis
+		{
+			FCanvasLineItem Axis;
+			Axis.LineThickness = 3;
+			{
+				Axis.SetColor(FLinearColor::Red);
+				FVector2D LineStart = WorldToScreen(FVector2D(-1638400.f, 0.f));
+				FVector2D LineEnd = WorldToScreen(FVector2D(1638400.f, 0.f));
+				LineStart.X = FMath::Clamp(LineStart.X, GridScreenBounds.Min.X, GridScreenBounds.Max.X);
+				LineStart.Y = FMath::Clamp(LineStart.Y, GridScreenBounds.Min.Y, GridScreenBounds.Max.Y);
+				LineEnd.X = FMath::Clamp(LineEnd.X, GridScreenBounds.Min.X, GridScreenBounds.Max.X);
+				LineEnd.Y = FMath::Clamp(LineEnd.Y, GridScreenBounds.Min.Y, GridScreenBounds.Max.Y);
+				Axis.Draw(CanvasObject, LineStart, LineEnd);
+			}
+			{
+				Axis.SetColor(FLinearColor::Green);
+				FVector2D LineStart = WorldToScreen(FVector2D(0.f, -1638400.f));
+				FVector2D LineEnd = WorldToScreen(FVector2D(0.f, 1638400.f));
+				LineStart.X = FMath::Clamp(LineStart.X, GridScreenBounds.Min.X, GridScreenBounds.Max.X);
+				LineStart.Y = FMath::Clamp(LineStart.Y, GridScreenBounds.Min.Y, GridScreenBounds.Max.Y);
+				LineEnd.X = FMath::Clamp(LineEnd.X, GridScreenBounds.Min.X, GridScreenBounds.Max.X);
+				LineEnd.Y = FMath::Clamp(LineEnd.Y, GridScreenBounds.Min.Y, GridScreenBounds.Max.Y);
+				Axis.Draw(CanvasObject, LineStart, LineEnd);
+			}
+		}
+
 		// Draw Streaming Sources
 		const float GridLoadingRange = GetLoadingRange();
 		for (const FWorldPartitionStreamingSource& Source : Sources)
@@ -949,6 +955,7 @@
 				FBox2D Bounds;
 				verify(TempLevel.GetCellBounds(FIntVector2(CellCoordX, CellCoordY), Bounds));
 				StreamingCell->Position = FVector(Bounds.GetCenter(), 0.f);
+				StreamingCell->Extent = Bounds.GetExtent().X;
 				StreamingCell->SetDebugInfo(CellGlobalCoords, CurrentStreamingGrid.GridName);
 				StreamingCell->SetClientOnlyVisible(CurrentStreamingGrid.bClientOnlyVisible);
 				StreamingCell->SetBlockOnSlowLoading(CurrentStreamingGrid.bBlockOnSlowStreaming);
@@ -1214,7 +1221,7 @@
 
 	if (StreamingCell->CachedIsBlockingSource)
 	{
-		const float Distance = FMath::Sqrt(StreamingCell->CachedBlockingMinDistanceSquare) - ((float)(*StreamingGrid)->GetCellSize(StreamingCell->Level) / 2.f);
+		const float Distance = FMath::Sqrt(StreamingCell->CachedMinSquareDistanceToSource) - ((float)(*StreamingGrid)->GetCellSize(StreamingCell->Level) / 2.f);
 
 		const float Ratio = Distance / LoadingRange;
 
diff --git a/Engine/Source/Runtime/Engine/Private/WorldPartition/WorldPartitionRuntimeSpatialHashCell.cpp b/Engine/Source/Runtime/Engine/Private/WorldPartition/WorldPartitionRuntimeSpatialHashCell.cpp
--- a/Engine/Source/Runtime/Engine/Private/WorldPartition/WorldPartitionRuntimeSpatialHashCell.cpp
+++ b/Engine/Source/Runtime/Engine/Private/WorldPartition/WorldPartitionRuntimeSpatialHashCell.cpp
@@ -18,8 +18,8 @@ UWorldPartitionRuntimeSpatialHashCell::UWorldPartitionRuntimeSpatialHashCell(con
 : Super(ObjectInitializer)
 , Level(0)
 , CachedIsBlockingSource(false)
-, CachedBlockingMinDistanceSquare(MAX_FLT)
-, CachedSourceMinSortDistance(0.f)
+, CachedMinSquareDistanceToSource(MAX_FLT)
+, CachedSourceSortingDistance(0.f)
 {}
 
 #if WITH_EDITOR
@@ -47,33 +47,89 @@ bool UWorldPartitionRuntimeSpatialHashCell::CacheStreamingSourceInfo(const UWorl
 	if (bWasCacheDirtied)
 	{
 		CachedIsBlockingSource = false;
-		CachedBlockingMinDistanceSquare = MAX_FLT;
+		CachedMinSquareDistanceToSource = MAX_FLT;
+		CachedSourceModulatedDistances.Reset();
 	}
 
-	const float AngleContribution = FMath::Clamp(GRuntimeSpatialHashCellToSourceAngleContributionToCellImportance, 0.f, 1.f);
-	const float SqrDistance = FVector::DistSquared(Info.SourceShape.GetCenter(), Position);
+	float AngleContribution = FMath::Clamp(GRuntimeSpatialHashCellToSourceAngleContributionToCellImportance, 0.f, 1.f);
+	const float SquareDistance = FVector::DistSquared2D(Info.SourceShape.GetCenter(), Position);
 	float AngleFactor = 1.f;
 	if (!FMath::IsNearlyZero(AngleContribution))
 	{
-		const FVector2D SourceForward(Info.SourceShape.GetAxis());
-		const FVector2D SourceToCell(Position - Info.SourceShape.GetCenter());
-		const float Dot = FVector2D::DotProduct(SourceForward.GetSafeNormal(), SourceToCell.GetSafeNormal());
-		const float NormalizedAngle = FMath::Clamp(FMath::Abs(FMath::Acos(Dot) / PI), 0.f, 1.f);
+		const FBox Box(FVector(Position.X - Extent, Position.Y - Extent, 0.f), FVector(Position.X + Extent, Position.Y + Extent, 0.f));
+		const FVector2D SourcePos(Info.SourceShape.GetCenter());
+		const FVector StartVert(SourcePos, 0.f);
+		const FVector EndVert(SourcePos + FVector2D(Info.SourceShape.GetScaledAxis()), 0.f);
+
+		float Angle = 0.f;
+		if (!FMath::LineBoxIntersection(Box, StartVert, EndVert, EndVert - StartVert))
+		{
+			// Find smallest angle using 4 corners and center of cell bounds
+			const FVector2D Position2D(Position);
+			float MaxDot = 0.f;
+			FVector2D SourceForward(Info.SourceShape.GetAxis());
+			SourceForward.Normalize();
+			FVector2D CellPoints[5];
+			CellPoints[0] = Position2D + FVector2D(-Extent, -Extent);
+			CellPoints[1] = Position2D + FVector2D(-Extent, Extent);
+			CellPoints[2] = Position2D + FVector2D(Extent, -Extent);
+			CellPoints[3] = Position2D + FVector2D(Extent, Extent);
+			CellPoints[4] = Position2D;
+			for (const FVector2D& CellPoint : CellPoints)
+			{
+				const FVector2D SourceToCell(CellPoint - SourcePos);
+				const float Dot = FVector2D::DotProduct(SourceForward, SourceToCell.GetSafeNormal());
+				MaxDot = FMath::Max(MaxDot, Dot);
+			}
+			Angle = FMath::Abs(FMath::Acos(MaxDot) / PI);
+		}
+		const float NormalizedAngle = FMath::Clamp(Angle, PI/180.f, 1.f);
 		AngleFactor = FMath::Pow(NormalizedAngle, AngleContribution);
 	}
 	// Modulate distance to cell by angle relative to source forward vector (to prioritize cells in front)
-	float ModulatedDistance = SqrDistance * AngleFactor;
-
-	// If cache was dirtied, use value, else use minimum with existing cached value
-	CachedSourceMinSortDistance = bWasCacheDirtied ? ModulatedDistance : FMath::Min(ModulatedDistance, CachedSourceMinSortDistance);
+	const float SquareAngleFactor = AngleFactor * AngleFactor;
+	const float ModulatedSquareDistance = SquareDistance * SquareAngleFactor;
+	CachedSourceModulatedDistances.Add(ModulatedSquareDistance);
+	int32 Count = CachedSourceModulatedDistances.Num();
+	check(Count == CachedSourcePriorityWeights.Num());
+	if (Count == 1)
+	{
+		CachedSourceSortingDistance = ModulatedSquareDistance;
+	}
+	else
+	{
+		float TotalSourcePriorityWeight = 0.f;
+		for (int32 i = 0; i < Count; ++i)
+		{
+			TotalSourcePriorityWeight += CachedSourcePriorityWeights[i];
+		}
+		int32 HighestPrioMinDistIndex = 0;
+		float WeightedModulatedDistance = 0.f;
+		for (int32 i = 0; i < Count; ++i)
+		{
+			WeightedModulatedDistance += CachedSourceModulatedDistances[i] * CachedSourcePriorityWeights[i] / TotalSourcePriorityWeight;
+			
+			// Find highest priority source with the minimum modulated distance
+			if ((i != 0) &&
+				(CachedSourceModulatedDistances[i] < CachedSourceModulatedDistances[HighestPrioMinDistIndex]) &&
+				(CachedSourcePriorityWeights[i] >= CachedSourcePriorityWeights[HighestPrioMinDistIndex]))
+			{
+				HighestPrioMinDistIndex = i;
+			}
+		}
+		// Sorting distance is the minimum between these:
+		// - the highest priority source with the minimum modulated distance 
+		// - the weighted modulated distance
+		CachedSourceSortingDistance = FMath::Min(CachedSourceModulatedDistances[HighestPrioMinDistIndex], WeightedModulatedDistance);
+	}
 	
 	// Only consider blocking sources
 	if (Info.Source.bBlockOnSlowLoading)
 	{
 		CachedIsBlockingSource = true;
-		CachedBlockingMinDistanceSquare = FMath::Min(SqrDistance, CachedBlockingMinDistanceSquare);
 	}
 
+	CachedMinSquareDistanceToSource = FMath::Min(SquareDistance, CachedMinSquareDistanceToSource);
 	return bWasCacheDirtied;
 }
 
@@ -90,8 +146,16 @@ int32 UWorldPartitionRuntimeSpatialHashCell::SortCompare(const UWorldPartitionRu
 		if (Result == 0)
 		{
 			// Closest distance (lower value is higher prio)
-			const float Diff = CachedSourceMinSortDistance - Other->CachedSourceMinSortDistance;
-			Result = Diff < 0.f ? -1 : (Diff > 0.f ? 1 : 0);
+			const float Diff = CachedSourceSortingDistance - Other->CachedSourceSortingDistance;
+			if (FMath::IsNearlyZero(Diff))
+			{
+				const float RawDistanceDiff = CachedMinSquareDistanceToSource - Other->CachedMinSquareDistanceToSource;
+				Result = RawDistanceDiff < 0.f ? -1 : (RawDistanceDiff > 0.f ? 1 : 0);
+			}
+			else
+			{
+				Result = Diff < 0.f ? -1 : (Diff > 0.f ? 1 : 0);
+			}
 		}
 	}
 	return Result;
diff --git a/Engine/Source/Runtime/Engine/Private/WorldPartition/WorldPartitionStreamingPolicy.cpp b/Engine/Source/Runtime/Engine/Private/WorldPartition/WorldPartitionStreamingPolicy.cpp
--- a/Engine/Source/Runtime/Engine/Private/WorldPartition/WorldPartitionStreamingPolicy.cpp
+++ b/Engine/Source/Runtime/Engine/Private/WorldPartition/WorldPartitionStreamingPolicy.cpp
@@ -13,6 +13,7 @@
 #include "WorldPartition/DataLayer/WorldDataLayers.h"
 #include "WorldPartition/HLOD/HLODActor.h"
 #include "WorldPartition/HLOD/HLODSubsystem.h"
+#include "WorldPartition/WorldPartitionDebugHelper.h"
 #include "GameFramework/PlayerController.h"
 #include "Engine/LocalPlayer.h"
 #include "Engine/World.h"
@@ -111,7 +112,7 @@ void UWorldPartitionStreamingPolicy::UpdateStreamingSources()
 			const FVector ViewLocationLocal = WorldToLocal.TransformPosition(ViewLocation);
 			const FRotator ViewRotationLocal = WorldToLocal.TransformRotation(ViewRotation.Quaternion()).Rotator();
 			static const FName NAME_SIE(TEXT("SIE"));
-			StreamingSources.Add(FWorldPartitionStreamingSource(NAME_SIE, ViewLocationLocal, ViewRotationLocal, EStreamingSourceTargetState::Activated, /*bBlockOnSlowLoading=*/false));
+			StreamingSources.Add(FWorldPartitionStreamingSource(NAME_SIE, ViewLocationLocal, ViewRotationLocal, EStreamingSourceTargetState::Activated, /*bBlockOnSlowLoading=*/false, EStreamingSourcePriority::Default));
 		}
 		else
 #endif
@@ -134,7 +135,7 @@ void UWorldPartitionStreamingPolicy::UpdateStreamingSources()
 						ViewLocation = WorldToLocal.TransformPosition(ViewLocation);
 						ViewRotation = WorldToLocal.TransformRotation(ViewRotation.Quaternion()).Rotator();						
 						EStreamingSourceTargetState TargetState = Player->PlayerController->StreamingSourceShouldActivate() ? EStreamingSourceTargetState::Activated : EStreamingSourceTargetState::Loaded;
-						StreamingSources.Add(FWorldPartitionStreamingSource(Player->GetFName(), ViewLocation, ViewRotation, TargetState, /*bBlockOnSlowLoading=*/true));
+						StreamingSources.Add(FWorldPartitionStreamingSource(Player->GetFName(), ViewLocation, ViewRotation, TargetState, /*bBlockOnSlowLoading=*/true, EStreamingSourcePriority::Default));
 					}
 				}
 			}
@@ -143,17 +144,13 @@ void UWorldPartitionStreamingPolicy::UpdateStreamingSources()
 		for (IWorldPartitionStreamingSourceProvider* StreamingSourceProvider : WorldPartition->StreamingSourceProviders)
 		{
 			FWorldPartitionStreamingSource StreamingSource;
+			// Default Streaming Source provider's priority to be less than those based on player controllers
+			StreamingSource.Priority = EStreamingSourcePriority::Low;
 			if (StreamingSourceProvider->GetStreamingSource(StreamingSource))
 			{
 				// Transform to Local
 				StreamingSource.Location = WorldToLocal.TransformPosition(StreamingSource.Location);
 				StreamingSource.Rotation = WorldToLocal.TransformRotation(StreamingSource.Rotation.Quaternion()).Rotator();
-				// If none is provided, default Streaming Source provider's priority to be less than those based on player controllers
-				if (StreamingSource.Priority == EStreamingSourcePriority::Default)
-				{
-					StreamingSource.Priority = EStreamingSourcePriority::Low;
-				}
-
 				StreamingSources.Add(StreamingSource);
 			}
 		}
@@ -228,7 +225,7 @@ void UWorldPartitionStreamingPolicy::UpdateStreamingState()
 			// Activation superseeds Loading
 			LoadStreamingCells = LoadStreamingCells.Difference(ActivateStreamingCells);
 		}
-	} 
+	}
 	else 
 	{
 		// Server will activate all non data layer cells at first and then load/activate/unload data layer cells only when the data layer states change
@@ -275,6 +272,10 @@ void UWorldPartitionStreamingPolicy::UpdateStreamingState()
 	TSet<const UWorldPartitionRuntimeCell*> ToLoadCells = LoadStreamingCells.Difference(LoadedCells);
 	TSet<const UWorldPartitionRuntimeCell*> ToUnloadCells = ActivatedCells.Union(LoadedCells).Difference(ActivateStreamingCells.Union(LoadStreamingCells));
 
+#if !UE_BUILD_SHIPPING
+	UpdateDebugCellsStreamingPriority(ActivateStreamingCells, LoadStreamingCells);
+#endif
+
 	if(World->bMatchStarted)
 	{
 		WORLDPARTITION_LOG_UPDATESTREAMINGSTATE(Verbose);
@@ -308,6 +309,27 @@ void UWorldPartitionStreamingPolicy::UpdateStreamingState()
 	FrameLoadCells.Reset();
 }
 
+#if !UE_BUILD_SHIPPING
+void UWorldPartitionStreamingPolicy::UpdateDebugCellsStreamingPriority(const TSet<const UWorldPartitionRuntimeCell*>& ActivateStreamingCells, const TSet<const UWorldPartitionRuntimeCell*>& LoadStreamingCells)
+{
+	if (FWorldPartitionDebugHelper::IsRuntimeSpatialHashCellStreamingPriotityShown())
+	{
+		TSet<const UWorldPartitionRuntimeCell*> CellsToSort;
+		CellsToSort.Append(ActivateStreamingCells);
+		CellsToSort.Append(LoadStreamingCells);
+
+		TArray<const UWorldPartitionRuntimeCell*, TInlineAllocator<256>> SortedCells;
+		WorldPartition->RuntimeHash->SortStreamingCellsByImportance(CellsToSort, StreamingSources, SortedCells);
+		const int32 CellCount = SortedCells.Num();
+		int32 CellPrio = 0;
+		for (const UWorldPartitionRuntimeCell* SortedCell : SortedCells)
+		{
+			const_cast<UWorldPartitionRuntimeCell*>(SortedCell)->SetDebugStreamingPriority(float(CellPrio++) / CellCount);
+		}
+	}
+}
+#endif
+
 void UWorldPartitionStreamingPolicy::UpdateStreamingPerformance(const TSet<const UWorldPartitionRuntimeCell*>& CellsToActivate)
 {		
 	TRACE_CPUPROFILER_EVENT_SCOPE(UWorldPartitionStreamingPolicy::UpdateStreamingPerformance);
diff --git a/Engine/Source/Runtime/Engine/Public/WorldPartition/WorldPartitionDebugHelper.h b/Engine/Source/Runtime/Engine/Public/WorldPartition/WorldPartitionDebugHelper.h
--- a/Engine/Source/Runtime/Engine/Public/WorldPartition/WorldPartitionDebugHelper.h
+++ b/Engine/Source/Runtime/Engine/Public/WorldPartition/WorldPartitionDebugHelper.h
@@ -5,10 +5,12 @@
 #include "CoreTypes.h"
 #include "UObject/NameTypes.h"
 #include "ProfilingDebugging/ProfilingHelpers.h"
+#include "HAL/IConsoleManager.h"
 
 class FAutoConsoleCommand;
 class UCanvas;
 class UFont;
+struct FLinearColor;
 
 class FWorldPartitionDebugHelper
 {
@@ -17,9 +19,11 @@ public:
 	static bool AreDebugDataLayersShown(const TArray<FName>& DataLayerNames);
 	static bool IsDebugStreamingStatusShown(EStreamingStatus StreamingStatus);
 	static bool IsDebugRuntimeHashGridShown(FName Name);
-	static bool IsDebugCellNameShow(const FString& Name);
+	static bool IsDebugCellNameShown(const FString& Name);
 	static void DrawText(UCanvas* Canvas, const FString& Text, const UFont* Font, const FColor& Color, FVector2D& Pos, float* MaxTextWidth = nullptr);
 	static void DrawLegendItem(UCanvas* Canvas, const FString& Text, const UFont* Font, const FColor& Color, const FColor& TextColor, FVector2D& Pos, float* MaxItemWidth = nullptr);
+	static FLinearColor GetHeatMapColor(float Value);
+	FORCEINLINE static bool IsRuntimeSpatialHashCellStreamingPriotityShown() { return bShowRuntimeSpatialHashCellStreamingPriotity; }
 
 private:
 	static FAutoConsoleCommand DebugFilterByStreamingStatusCommand;
@@ -33,4 +37,7 @@ private:
 
 	static FAutoConsoleCommand DebugFilterByCellNameCommand;
 	static TArray<FString> DebugCellNameFilter;
+
+	static FAutoConsoleVariableRef ShowRuntimeSpatialHashCellStreamingPriotityCommand;
+	static bool bShowRuntimeSpatialHashCellStreamingPriotity;
 };
\ No newline at end of file
diff --git a/Engine/Source/Runtime/Engine/Public/WorldPartition/WorldPartitionRuntimeCell.h b/Engine/Source/Runtime/Engine/Public/WorldPartition/WorldPartitionRuntimeCell.h
--- a/Engine/Source/Runtime/Engine/Public/WorldPartition/WorldPartitionRuntimeCell.h
+++ b/Engine/Source/Runtime/Engine/Public/WorldPartition/WorldPartitionRuntimeCell.h
@@ -10,6 +10,12 @@
 #include "Algo/AnyOf.h"
 #include "WorldPartitionRuntimeCell.generated.h"
 
+enum class EWorldPartitionRuntimeCellVisualizeMode
+{
+	StreamingPriority,
+	StreamingStatus
+};
+
 USTRUCT()
 struct FWorldPartitionRuntimeCellObjectMapping
 {
@@ -122,7 +128,7 @@ class UWorldPartitionRuntimeCell : public UObject
 	virtual bool CanAddToWorld() const PURE_VIRTUAL(UWorldPartitionRuntimeCell::CanAddToWorld, return false;);
 	virtual ULevel* GetLevel() const PURE_VIRTUAL(UWorldPartitionRuntimeCell::GetLevel, return nullptr;);
 	virtual EWorldPartitionRuntimeCellState GetCurrentState() const PURE_VIRTUAL(UWorldPartitionRuntimeCell::GetCurrentState, return EWorldPartitionRuntimeCellState::Unloaded;);
-	virtual FLinearColor GetDebugColor() const { return FLinearColor::Black; }
+	virtual FLinearColor GetDebugColor(EWorldPartitionRuntimeCellVisualizeMode VisualizeMode) const { static const FLinearColor DefaultColor = FLinearColor::Black.CopyWithNewOpacity(0.25f); return DefaultColor; }
 	virtual bool IsAlwaysLoaded() const { return bIsAlwaysLoaded; }
 	virtual void SetIsAlwaysLoaded(bool bInIsAlwaysLoaded) { bIsAlwaysLoaded = bInIsAlwaysLoaded; }
 	virtual void SetPriority(int32 InPriority) { Priority = InPriority; }
@@ -167,7 +173,13 @@ class UWorldPartitionRuntimeCell : public UObject
 	
 	bool GetIsHLOD() const { return bIsHLOD; }
 
+#if !UE_BUILD_SHIPPING
+	void SetDebugStreamingPriority(float InDebugStreamingPriority) { DebugStreamingPriority = InDebugStreamingPriority; }
+#endif
+
 protected:
+	FLinearColor GetDebugStreamingPriorityColor() const;
+
 #if WITH_EDITOR
 	void UpdateDebugName();
 #endif
@@ -202,11 +214,20 @@ private:
 	UPROPERTY()
 	bool bBlockOnSlowLoading;
 
+protected:
 	// Source Priority
-	mutable int32 CachedSourcePriority;
+	mutable uint8 CachedMinSourcePriority;
+
+	// Source Priorities
+	mutable TArray<float> CachedSourcePriorityWeights;
 
 	// Epoch used to dirty cache
 	mutable int32 CachedSourceInfoEpoch;
 
 	static int32 StreamingSourceCacheEpoch;
+
+#if !UE_BUILD_SHIPPING
+	// Represents the streaming priority relative to other cells
+	float DebugStreamingPriority;
+#endif
 };
diff --git a/Engine/Source/Runtime/Engine/Public/WorldPartition/WorldPartitionRuntimeSpatialHashCell.h b/Engine/Source/Runtime/Engine/Public/WorldPartition/WorldPartitionRuntimeSpatialHashCell.h
--- a/Engine/Source/Runtime/Engine/Public/WorldPartition/WorldPartitionRuntimeSpatialHashCell.h
+++ b/Engine/Source/Runtime/Engine/Public/WorldPartition/WorldPartitionRuntimeSpatialHashCell.h
@@ -21,17 +21,27 @@ class UWorldPartitionRuntimeSpatialHashCell : public UWorldPartitionRuntimeCell
 	UPROPERTY()
 	FVector Position;
 
+	UPROPERTY()
+	float Extent;
+
 	UPROPERTY()
 	int32 Level;
 		
 	// Used to determine if cell was requested by blocking source
 	mutable bool CachedIsBlockingSource;
 
-	// Used to determine if the streaming of cells is lagging
-	mutable float CachedBlockingMinDistanceSquare;
+	// Represents the square distance from cell to closest streaming source
+	mutable float CachedMinSquareDistanceToSource;
+
+	// Modulated distance to the different streaming sources used to sort relative priority amongst streaming cells
+	// The value is affected by :
+	// - All sources intersecting the cell
+	// - The priority of each source
+	// - The distance between the cell and each source
+	// - The angle between the cell and each source orientation
+	mutable float CachedSourceSortingDistance;
 
-	// Computed and cached value used by SortCompare to sort Cells
-	mutable float CachedSourceMinSortDistance;
+	mutable TArray<float> CachedSourceModulatedDistances;
 
 #if WITH_EDITORONLY_DATA
 	UPROPERTY()
diff --git a/Engine/Source/Runtime/Engine/Public/WorldPartition/WorldPartitionStreamingPolicy.h b/Engine/Source/Runtime/Engine/Public/WorldPartition/WorldPartitionStreamingPolicy.h
--- a/Engine/Source/Runtime/Engine/Public/WorldPartition/WorldPartitionStreamingPolicy.h
+++ b/Engine/Source/Runtime/Engine/Public/WorldPartition/WorldPartitionStreamingPolicy.h
@@ -101,6 +101,8 @@ protected:
 
 	EWorldPartitionStreamingPerformance StreamingPerformance;
 #if !UE_BUILD_SHIPPING
+	void UpdateDebugCellsStreamingPriority(const TSet<const UWorldPartitionRuntimeCell*>& ActivateStreamingCells, const TSet<const UWorldPartitionRuntimeCell*>& LoadStreamingCells);
+
 	double OnScreenMessageStartTime;
 	EWorldPartitionStreamingPerformance  OnScreenMessageStreamingPerformance;
 #endif
diff --git a/Engine/Source/Runtime/Engine/Public/WorldPartition/WorldPartitionStreamingSource.h b/Engine/Source/Runtime/Engine/Public/WorldPartition/WorldPartitionStreamingSource.h
--- a/Engine/Source/Runtime/Engine/Public/WorldPartition/WorldPartitionStreamingSource.h
+++ b/Engine/Source/Runtime/Engine/Public/WorldPartition/WorldPartitionStreamingSource.h
@@ -302,13 +302,14 @@ struct FWorldPartitionStreamingQuerySource
 /**
  * Streaming Source Priority
  */
-enum class EStreamingSourcePriority : int32
+UENUM(BlueprintType)
+enum class EStreamingSourcePriority : uint8
 {
-	Highest = INT_MIN,
-	High = -4096,
-	Normal = 0,
-	Low = 4096,
-	Lowest = INT_MAX,
+	Highest = 0,
+	High = 64,
+	Normal = 128,
+	Low = 192,
+	Lowest = 255,
 	Default = Normal
 };
 
@@ -324,7 +325,7 @@ struct ENGINE_API FWorldPartitionStreamingSource
 		, TargetGrid(NAME_None)
 	{}
 
-	FWorldPartitionStreamingSource(FName InName, const FVector& InLocation, const FRotator& InRotation, EStreamingSourceTargetState InTargetState, bool bInBlockOnSlowLoading, EStreamingSourcePriority InPriority = EStreamingSourcePriority::Default, float InVelocity = 0.f)
+	FWorldPartitionStreamingSource(FName InName, const FVector& InLocation, const FRotator& InRotation, EStreamingSourceTargetState InTargetState, bool bInBlockOnSlowLoading, EStreamingSourcePriority InPriority, float InVelocity = 0.f)
 		: Name(InName)
 		, Location(InLocation)
 		, Rotation(InRotation)
