Fixed a bug where pooled FX attached to actors in hidden sublevels are not hidden properly

diff --git a/Engine/Plugins/FX/Niagara/Source/Niagara/Private/NiagaraComponent.cpp b/Engine/Plugins/FX/Niagara/Source/Niagara/Private/NiagaraComponent.cpp
--- a/Engine/Plugins/FX/Niagara/Source/Niagara/Private/NiagaraComponent.cpp
+++ b/Engine/Plugins/FX/Niagara/Source/Niagara/Private/NiagaraComponent.cpp
@@ -1550,6 +1550,13 @@ void UNiagaraComponent::OnRegister()
 #endif
 
 	Super::OnRegister();
+
+	// Deal with the case where the particle component is attached to an actor in a hidden sublevel. Without this, the component will be visible instead of being hidden as well.
+	if (CachedLevelCollection == nullptr && GetOwner() == nullptr && IsValid(GetAttachParent()))
+	{
+		const ULevel* const AttachParentLevel = GetAttachParent()->GetComponentLevel();
+		CachedLevelCollection = AttachParentLevel ? AttachParentLevel->GetCachedLevelCollection() : nullptr;
+	}
 }
 
 bool UNiagaraComponent::IsReadyForOwnerToAutoDestroy() const
diff --git a/Engine/Source/Runtime/Engine/Private/Particles/ParticleComponents.cpp b/Engine/Source/Runtime/Engine/Private/Particles/ParticleComponents.cpp
--- a/Engine/Source/Runtime/Engine/Private/Particles/ParticleComponents.cpp
+++ b/Engine/Source/Runtime/Engine/Private/Particles/ParticleComponents.cpp
@@ -3924,6 +3924,13 @@ void UParticleSystemComponent::OnRegister()
 		// Force it to LODLevel 0
 		LODLevel = 0;
 	}
+
+	// Deal with the case where the particle component is attached to an actor in a hidden sublevel. Without this, the component will be visible instead of being hidden as well.
+	if (CachedLevelCollection == nullptr && GetOwner() == nullptr && IsValid(GetAttachParent()))
+	{
+		const ULevel* const AttachParentLevel = GetAttachParent()->GetComponentLevel();
+		CachedLevelCollection = AttachParentLevel ? AttachParentLevel->GetCachedLevelCollection() : nullptr;
+	}
 }
 
 void UParticleSystemComponent::OnUnregister()
