Submitted missing files from CL 17570855

diff --git a/Engine/Source/Runtime/Engine/Public/WorldPartition/WorldPartitionRuntimeLevelStreamingCell.h b/Engine/Source/Runtime/Engine/Public/WorldPartition/WorldPartitionRuntimeLevelStreamingCell.h
--- a/Engine/Source/Runtime/Engine/Public/WorldPartition/WorldPartitionRuntimeLevelStreamingCell.h
+++ b/Engine/Source/Runtime/Engine/Public/WorldPartition/WorldPartitionRuntimeLevelStreamingCell.h
@@ -19,7 +19,7 @@ class UWorldPartitionRuntimeLevelStreamingCell : public UWorldPartitionRuntimeSp
 	virtual bool CanAddToWorld() const override;
 	virtual ULevel* GetLevel() const override;
 	virtual EWorldPartitionRuntimeCellState GetCurrentState() const override;
-	virtual FLinearColor GetDebugColor() const override;
+	virtual FLinearColor GetDebugColor(EWorldPartitionRuntimeCellVisualizeMode VisualizeMode) const override;
 	virtual void SetIsAlwaysLoaded(bool bInIsAlwaysLoaded) override;
 	virtual EStreamingStatus GetStreamingStatus() const override;
 	virtual bool IsLoading() const override;
diff --git a/Engine/Source/Runtime/Engine/Public/WorldPartition/WorldPartitionRuntimeSpatialHash.h b/Engine/Source/Runtime/Engine/Public/WorldPartition/WorldPartitionRuntimeSpatialHash.h
--- a/Engine/Source/Runtime/Engine/Public/WorldPartition/WorldPartitionRuntimeSpatialHash.h
+++ b/Engine/Source/Runtime/Engine/Public/WorldPartition/WorldPartitionRuntimeSpatialHash.h
@@ -146,6 +146,7 @@ private:
 	void DrawStreamingSource3D(UWorld* World, const FSphericalSector& Shape, const FTransform& Transform, const FColor& Color) const;
 	void GetFilteredCellsForDebugDraw(const FSpatialHashStreamingGridLayerCell* LayerCell, const UDataLayerSubsystem* DataLayerSubsystem, TArray<const UWorldPartitionRuntimeCell*>& FilteredCells) const;
 	const FSquare2DGridHelper& GetGridHelper() const;
+	EWorldPartitionRuntimeCellVisualizeMode GetStreamingCellVisualizeMode() const;
 	mutable FSquare2DGridHelper* GridHelper;
 };
 
