Make it clear that PlatformData can't be null when creating texture cube array, fix static analysis warning

diff --git a/Engine/Source/Runtime/Engine/Private/TextureCubeArray.cpp b/Engine/Source/Runtime/Engine/Private/TextureCubeArray.cpp
--- a/Engine/Source/Runtime/Engine/Private/TextureCubeArray.cpp
+++ b/Engine/Source/Runtime/Engine/Private/TextureCubeArray.cpp
@@ -27,6 +27,7 @@ public:
 		, TextureSize(0)
 		, ProxiedResource(nullptr)
 	{
+		check(Owner);
 		check(Owner->GetNumMips() > 0);
 
 		TIndirectArray<FTexture2DMipMap>& Mips = InOwner->PlatformData->Mips;
@@ -71,6 +72,7 @@ public:
 		, TextureSize(0)
 		, ProxiedResource(InProxiedResource)
 	{
+		check(Owner);
 	}
 
 	/**
@@ -108,11 +110,13 @@ public:
 
 		const uint32 ArraySize = NumSlices / 6u;
 
+		check(Owner->GetNumMips() > 0);
+
 		// Create the RHI texture.
 		ETextureCreateFlags TexCreateFlags = (Owner->SRGB ? TexCreate_SRGB : TexCreate_None) | (Owner->bNotOfflineProcessed ? TexCreate_None : TexCreate_OfflineProcessed);
 		FString Name = Owner->GetPathName();
 		FRHIResourceCreateInfo CreateInfo(*Name);
-		CreateInfo.ExtData = Owner->PlatformData ? Owner->PlatformData->GetExtData() : 0;
+		CreateInfo.ExtData = Owner->PlatformData->GetExtData();
 		TextureCubeRHI = RHICreateTextureCubeArray(Owner->GetSizeX(), ArraySize, Owner->GetPixelFormat(), Owner->GetNumMips(), TexCreateFlags, CreateInfo);
 		TextureRHI = TextureCubeRHI;
 		TextureRHI->SetName(Owner->GetFName());
