PR #7527: In CapsuleComponent, we have functions like UCapsuleComponent::GetScaledCapsuleHalfHeight, which return component properties multiplied by the component scale. Similar functionality is useful when it comes to TimelineComponent. I added a UTimelineComponent::GetScaledTimelineLength function, which returns the timeline's length divided by the play rate.

#p4v-preflight-copy 17571336

diff --git a/Engine/Source/Runtime/Engine/Classes/Components/TimelineComponent.h b/Engine/Source/Runtime/Engine/Classes/Components/TimelineComponent.h
--- a/Engine/Source/Runtime/Engine/Classes/Components/TimelineComponent.h
+++ b/Engine/Source/Runtime/Engine/Classes/Components/TimelineComponent.h
@@ -316,6 +316,9 @@ public:
 	/** Get length of the timeline */
 	ENGINE_API float GetTimelineLength() const;
 
+	/** Get length of the timeline divided by the play rate */
+	ENGINE_API float GetScaledTimelineLength() const;
+
 	/** Sets the timeline length mode */
 	ENGINE_API void SetTimelineLengthMode(ETimelineLengthMode NewMode);
 
@@ -462,6 +465,10 @@ public:
 	UFUNCTION(BlueprintCallable, Category="Components|Timeline")
 	ENGINE_API float GetTimelineLength() const;
 
+	/** Get length of the timeline divided by the play rate */
+	UFUNCTION(BlueprintCallable, Category="Components|Timeline")
+	ENGINE_API float GetScaledTimelineLength() const;
+
 	/** Set length of the timeline */
 	UFUNCTION(BlueprintCallable, Category="Components|Timeline")
 	ENGINE_API void SetTimelineLength(float NewLength);
diff --git a/Engine/Source/Runtime/Engine/Private/Timeline.cpp b/Engine/Source/Runtime/Engine/Private/Timeline.cpp
--- a/Engine/Source/Runtime/Engine/Private/Timeline.cpp
+++ b/Engine/Source/Runtime/Engine/Private/Timeline.cpp
@@ -558,6 +558,20 @@ float FTimeline::GetTimelineLength() const
 	}
 }
 
+float FTimeline::GetScaledTimelineLength() const
+{
+	const float CurrentPlayRate = GetPlayRate();
+	if(CurrentPlayRate != 0.0f)
+	{
+		return GetTimelineLength() / CurrentPlayRate;	
+	}
+	else
+	{
+		UE_LOG(LogTimeline, Error, TEXT("Invalid timeline PlayRate!"));
+		return 0.0f;
+	}
+}
+
 /** Sets the timeline length mode */
 void FTimeline::SetTimelineLengthMode(ETimelineLengthMode NewMode)
 {
@@ -812,6 +826,11 @@ float UTimelineComponent::GetTimelineLength() const
 	return TheTimeline.GetTimelineLength();
 }
 
+float UTimelineComponent::GetScaledTimelineLength() const
+{
+	return TheTimeline.GetScaledTimelineLength();
+}
+
 void UTimelineComponent::SetTimelineLength(float NewLength)
 {
 	return TheTimeline.SetTimelineLength(NewLength);
