Type truncation warning fixes in a few AI classes

diff --git a/Engine/Source/Runtime/Engine/Classes/AI/Navigation/NavRelevantInterface.h b/Engine/Source/Runtime/Engine/Classes/AI/Navigation/NavRelevantInterface.h
--- a/Engine/Source/Runtime/Engine/Classes/AI/Navigation/NavRelevantInterface.h
+++ b/Engine/Source/Runtime/Engine/Classes/AI/Navigation/NavRelevantInterface.h
@@ -89,8 +89,8 @@ struct ENGINE_API FNavigationRelevantData : public TSharedFromThis<FNavigationRe
 	FORCEINLINE bool NeedAnyPendingLazyModifiersGathering() const { return bPendingLazyModifiersGathering || bPendingChildLazyModifiersGathering; }
 	FORCEINLINE bool SupportsGatheringGeometrySlices() const { return bSupportsGatheringGeometrySlices; }
 	FORCEINLINE bool IsEmpty() const { return !HasGeometry() && !HasModifiers(); }
-	FORCEINLINE uint32 GetAllocatedSize() const { return CollisionData.GetAllocatedSize() + VoxelData.GetAllocatedSize() + Modifiers.GetAllocatedSize(); }
-	FORCEINLINE uint32 GetGeometryAllocatedSize() const { return CollisionData.GetAllocatedSize() + VoxelData.GetAllocatedSize(); }
+	FORCEINLINE SIZE_T GetAllocatedSize() const { return CollisionData.GetAllocatedSize() + VoxelData.GetAllocatedSize() + Modifiers.GetAllocatedSize(); }
+	FORCEINLINE SIZE_T GetGeometryAllocatedSize() const { return CollisionData.GetAllocatedSize() + VoxelData.GetAllocatedSize(); }
 	FORCEINLINE int32 GetDirtyFlag() const
 	{
 		return ((HasGeometry() || IsPendingLazyGeometryGathering() || Modifiers.GetMaskFillCollisionUnderneathForNavmesh()) ? ENavigationDirtyFlag::Geometry : 0) |
diff --git a/Engine/Source/Runtime/Engine/Classes/AI/Navigation/NavigationTypes.h b/Engine/Source/Runtime/Engine/Classes/AI/Navigation/NavigationTypes.h
--- a/Engine/Source/Runtime/Engine/Classes/AI/Navigation/NavigationTypes.h
+++ b/Engine/Source/Runtime/Engine/Classes/AI/Navigation/NavigationTypes.h
@@ -327,7 +327,7 @@ struct FNavigationPortalEdge
 		return ((FVector*)&Left)[Index];
 	}
 
-	FORCEINLINE float GetLength() const { return FVector::Dist(Left, Right); }
+	FORCEINLINE FVector::FReal GetLength() const { return FVector::Dist(Left, Right); }
 
 	FORCEINLINE FVector GetMiddlePoint() const { return Left + (Right - Left) / 2; }
 };
diff --git a/Engine/Source/Runtime/Engine/Classes/AI/NavigationSystemBase.h b/Engine/Source/Runtime/Engine/Classes/AI/NavigationSystemBase.h
--- a/Engine/Source/Runtime/Engine/Classes/AI/NavigationSystemBase.h
+++ b/Engine/Source/Runtime/Engine/Classes/AI/NavigationSystemBase.h
@@ -54,7 +54,7 @@ class ENGINE_API FNavigationLockContext
 {
 public:
 	FNavigationLockContext(ENavigationLockReason::Type Reason = ENavigationLockReason::Unknown, bool bApplyLock = true)
-		: MyWorld(NULL), LockReason(Reason), bSingleWorld(false), bIsLocked(false)
+		: MyWorld(NULL), LockReason((uint8)Reason), bSingleWorld(false), bIsLocked(false)
 	{
 		if (bApplyLock)
 		{
@@ -63,7 +63,7 @@ public:
 	}
 
 	FNavigationLockContext(UWorld* InWorld, ENavigationLockReason::Type Reason = ENavigationLockReason::Unknown, bool bApplyLock = true)
-		: MyWorld(InWorld), LockReason(Reason), bSingleWorld(true), bIsLocked(false)
+		: MyWorld(InWorld), LockReason((uint8)Reason), bSingleWorld(true), bIsLocked(false)
 	{
 		if (bApplyLock)
 		{
diff --git a/Engine/Source/Runtime/Engine/Private/VisualLogger/VisualLoggerAutomationTests.cpp b/Engine/Source/Runtime/Engine/Private/VisualLogger/VisualLoggerAutomationTests.cpp
--- a/Engine/Source/Runtime/Engine/Private/VisualLogger/VisualLoggerAutomationTests.cpp
+++ b/Engine/Source/Runtime/Engine/Private/VisualLogger/VisualLoggerAutomationTests.cpp
@@ -143,7 +143,7 @@ bool FVisualLogTest::RunTest(const FString& Parameters)
 			CHECK_SUCCESS(CurrentEntry->LogLines[0].Category == LogVisual.GetCategoryName());
 			CHECK_SUCCESS(CurrentEntry->LogLines[0].Line == TextToLog);
 
-			const float NewTimestamp = CurrentTimestamp + 0.1;
+			const double NewTimestamp = CurrentTimestamp + 0.1;
 			FVisualLogEntry* NewEntry = FVisualLogger::Get().GetEntryToWrite(World, NewTimestamp); //generate new entry and serialize old one
 			CurrentEntry = &Context.Device.LastEntry;
 			CHECK_NOT_NULL(CurrentEntry);
@@ -199,7 +199,7 @@ bool FVisualLogSegmentsTest::RunTest(const FString& Parameters)
 			CHECK_SUCCESS(CurrentEntry->ElementsToDraw[0].Points[0] == StartPoint);
 			CHECK_SUCCESS(CurrentEntry->ElementsToDraw[0].Points[1] == EndPoint);
 
-			const float NewTimestamp = CurrentTimestamp + 0.1;
+			const double NewTimestamp = CurrentTimestamp + 0.1;
 			FVisualLogEntry* NewEntry = FVisualLogger::Get().GetEntryToWrite(World, NewTimestamp); //generate new entry and serialize old one
 			CurrentEntry = &Context.Device.LastEntry;
 			CHECK_NOT_NULL(CurrentEntry);
@@ -277,7 +277,7 @@ bool FVisualLogEventsTest::RunTest(const FString& Parameters)
 		CHECK_SUCCESS(CurrentEntry->Events[1].UserFriendlyDesc == TEXT("Second simple event for vlog tests"));
 		CHECK_SUCCESS(CurrentEntry->Events[2].UserFriendlyDesc == TEXT("Third simple event for vlog tests"));
 
-		const float NewTimestamp = CurrentTimestamp + 0.1;
+		const double NewTimestamp = CurrentTimestamp + 0.1;
 		FVisualLogEntry* NewEntry = FVisualLogger::Get().GetEntryToWrite(World, NewTimestamp); //generate new entry and serialize old one
 		CurrentEntry = &Context.Device.LastEntry;
 		CHECK_NOT_NULL(CurrentEntry);
diff --git a/Engine/Source/Runtime/Engine/Private/VisualLogger/VisualLoggerKismetLibrary.cpp b/Engine/Source/Runtime/Engine/Private/VisualLogger/VisualLoggerKismetLibrary.cpp
--- a/Engine/Source/Runtime/Engine/Private/VisualLogger/VisualLoggerKismetLibrary.cpp
+++ b/Engine/Source/Runtime/Engine/Private/VisualLogger/VisualLoggerKismetLibrary.cpp
@@ -76,7 +76,7 @@ void UVisualLoggerKismetLibrary::LogSegment(UObject* WorldContextObject, const F
 {
 #if ENABLE_VISUAL_LOG
 	const ELogVerbosity::Type DefaultVerbosity = ELogVerbosity::Log;
-	FVisualLogger::GeometryShapeLogf(WorldContextObject, CategoryName, DefaultVerbosity, SegmentStart, SegmentEnd, ObjectColor.ToFColor(true), Thickness, TEXT("%s"), *Text);
+	FVisualLogger::GeometryShapeLogf(WorldContextObject, CategoryName, DefaultVerbosity, SegmentStart, SegmentEnd, ObjectColor.ToFColor(true), (uint16)Thickness, TEXT("%s"), *Text);
 #endif
 	if (bAddToMessageLog)
 	{
diff --git a/Engine/Source/Runtime/Engine/Public/VisualLogger/VisualLogger.h b/Engine/Source/Runtime/Engine/Public/VisualLogger/VisualLogger.h
--- a/Engine/Source/Runtime/Engine/Public/VisualLogger/VisualLogger.h
+++ b/Engine/Source/Runtime/Engine/Public/VisualLogger/VisualLogger.h
@@ -711,13 +711,13 @@ inline void FVisualLogger::GeometryShapeLogfImpl(const UObject* Object, const FL
 {
 	const FName CategoryName = Category.GetCategoryName(); 
 	COLLAPSED_LOGF(
-		CurrentEntry->AddElement(Location, CategoryName, Verbosity, Color, Buffer, Radius);
+		CurrentEntry->AddElement(Location, CategoryName, Verbosity, Color, Buffer, (uint16)Radius);
 	);
 }
 inline void FVisualLogger::GeometryShapeLogfImpl(const UObject* Object, const FName& CategoryName, ELogVerbosity::Type Verbosity, const FVector& Location, float Radius, const FColor& Color, const TCHAR* Fmt, ...)
 {
 	COLLAPSED_LOGF(
-		CurrentEntry->AddElement(Location, CategoryName, Verbosity, Color, Buffer, Radius);
+		CurrentEntry->AddElement(Location, CategoryName, Verbosity, Color, Buffer, (uint16)Radius);
 	);
 }
 
diff --git a/Engine/Source/Runtime/Engine/Public/VisualLogger/VisualLoggerTypes.h b/Engine/Source/Runtime/Engine/Public/VisualLogger/VisualLoggerTypes.h
--- a/Engine/Source/Runtime/Engine/Public/VisualLogger/VisualLoggerTypes.h
+++ b/Engine/Source/Runtime/Engine/Public/VisualLogger/VisualLoggerTypes.h
@@ -396,7 +396,7 @@ FVisualLogShapeElement::FVisualLogShapeElement(const FString& InDescription, con
 inline
 void FVisualLogShapeElement::SetColor(const FColor& InColor)
 {
-	Color = ((InColor.DWColor() >> 30) << 6)	| (((InColor.DWColor() & 0x00ff0000) >> 22) << 4)	| (((InColor.DWColor() & 0x0000ff00) >> 14) << 2)	| ((InColor.DWColor() & 0x000000ff) >> 6);
+	Color = (uint8)(((InColor.DWColor() >> 30) << 6)	| (((InColor.DWColor() & 0x00ff0000) >> 22) << 4)	| (((InColor.DWColor() & 0x0000ff00) >> 14) << 2)	| ((InColor.DWColor() & 0x000000ff) >> 6));
 }
 
 inline
