Type truncation warning fixes in a few Anim classes (partial even within these files)

diff --git a/Engine/Source/Runtime/Engine/Classes/Animation/AnimSequence.h b/Engine/Source/Runtime/Engine/Classes/Animation/AnimSequence.h
--- a/Engine/Source/Runtime/Engine/Classes/Animation/AnimSequence.h
+++ b/Engine/Source/Runtime/Engine/Classes/Animation/AnimSequence.h
@@ -451,7 +451,7 @@ public:
 
 	// Returns the framerate of the animation
 	UE_DEPRECATED(5.0, "GetFrameRate is deprecated see UAnimDataModel::GetFrameRate for the source frame rate, or GetSamplingFrameRate for the target frame rate instead")
-	float GetFrameRate() const { return GetSamplingFrameRate().AsDecimal(); }
+	float GetFrameRate() const { return (float)GetSamplingFrameRate().AsDecimal(); }
 
 	// Extract Root Motion transform from the animation
 	FTransform ExtractRootMotion(float StartTime, float DeltaTime, bool bAllowLooping) const;
diff --git a/Engine/Source/Runtime/Engine/Classes/Animation/SmartName.h b/Engine/Source/Runtime/Engine/Classes/Animation/SmartName.h
--- a/Engine/Source/Runtime/Engine/Classes/Animation/SmartName.h
+++ b/Engine/Source/Runtime/Engine/Classes/Animation/SmartName.h
@@ -158,7 +158,7 @@ struct ENGINE_API FSmartNameMapping
 	void InitializeCurveMetaData(class USkeleton* Skeleton);
 
 	/** Get the maximum in use UID */
-	SmartName::UID_Type GetMaxUID() const { return CurveNameList.Num() - 1; }
+	SmartName::UID_Type GetMaxUID() const { return (SmartName::UID_Type)(CurveNameList.Num() - 1); }
 
 private:
 	// List of curve names, indexed by UID
diff --git a/Engine/Source/Runtime/Engine/Public/AnimationCompression.h b/Engine/Source/Runtime/Engine/Public/AnimationCompression.h
--- a/Engine/Source/Runtime/Engine/Public/AnimationCompression.h
+++ b/Engine/Source/Runtime/Engine/Public/AnimationCompression.h
@@ -114,9 +114,9 @@ public:
 		}
 		Temp.Normalize();
 
-		X = (int32)(Temp.X * Quant16BitFactor) + Quant16BitOffs;
-		Y = (int32)(Temp.Y * Quant16BitFactor) + Quant16BitOffs;
-		Z = (int32)(Temp.Z * Quant16BitFactor) + Quant16BitOffs;
+		X = (uint16)((int32)(Temp.X * Quant16BitFactor) + Quant16BitOffs);
+		Y = (uint16)((int32)(Temp.Y * Quant16BitFactor) + Quant16BitOffs);
+		Z = (uint16)((int32)(Temp.Z * Quant16BitFactor) + Quant16BitOffs);
 	}
 
 	void ToQuat(FQuat& Out) const
@@ -293,9 +293,9 @@ public:
 	{
 		FVector3f Temp( Vec / 128.0f );
 
-		X = (int32)(Temp.X * Quant16BitFactor) + Quant16BitOffs;
-		Y = (int32)(Temp.Y * Quant16BitFactor) + Quant16BitOffs;
-		Z = (int32)(Temp.Z * Quant16BitFactor) + Quant16BitOffs;
+		X = (uint16)((int32)(Temp.X * Quant16BitFactor) + Quant16BitOffs);
+		Y = (uint16)((int32)(Temp.Y * Quant16BitFactor) + Quant16BitOffs);
+		Z = (uint16)((int32)(Temp.Z * Quant16BitFactor) + Quant16BitOffs);
 	}
 
 	void ToVector(FVector3f& Out) const
