Type truncation warning fixes in a few audio files

diff --git a/Engine/Source/Runtime/Engine/Classes/Sound/SoundEffectPreset.h b/Engine/Source/Runtime/Engine/Classes/Sound/SoundEffectPreset.h
--- a/Engine/Source/Runtime/Engine/Classes/Sound/SoundEffectPreset.h
+++ b/Engine/Source/Runtime/Engine/Classes/Sound/SoundEffectPreset.h
@@ -37,7 +37,7 @@ public:
 	virtual bool HasAssetActions() const { return false; }
 	virtual void Init() PURE_VIRTUAL(USoundEffectPreset::Init, );
 	virtual void OnInit() {};
-	virtual FColor GetPresetColor() const { return FColor(200.0f, 100.0f, 100.0f); }
+	virtual FColor GetPresetColor() const { return FColor(200, 100, 100); }
 
 	void Update();
 	void AddEffectInstance(TSoundEffectPtr& InEffectPtr);
diff --git a/Engine/Source/Runtime/Engine/Classes/Sound/SoundWave.h b/Engine/Source/Runtime/Engine/Classes/Sound/SoundWave.h
--- a/Engine/Source/Runtime/Engine/Classes/Sound/SoundWave.h
+++ b/Engine/Source/Runtime/Engine/Classes/Sound/SoundWave.h
@@ -976,7 +976,7 @@ public:
 	 * @param CompressionOverrides Optional argument for compression overrides.
 	 * @return			compressed data size, or zero if it could not be obtained
 	 */
-	int32 GetCompressedDataSize(FName Format, const FPlatformAudioCookOverrides* CompressionOverrides = GetPlatformCompressionOverridesForCurrentPlatform())
+	SIZE_T GetCompressedDataSize(FName Format, const FPlatformAudioCookOverrides* CompressionOverrides = GetPlatformCompressionOverridesForCurrentPlatform())
 	{
 		FByteBulkData* Data = GetCompressedData(Format, CompressionOverrides);
 		return Data ? Data->GetBulkDataSize() : 0;
