Some more type truncation fixes in various headers that typical game modules might include

diff --git a/Engine/Source/Runtime/Engine/Classes/Components/BrushComponent.h b/Engine/Source/Runtime/Engine/Classes/Components/BrushComponent.h
--- a/Engine/Source/Runtime/Engine/Classes/Components/BrushComponent.h
+++ b/Engine/Source/Runtime/Engine/Classes/Components/BrushComponent.h
@@ -50,7 +50,7 @@ public:
 	virtual FPrimitiveSceneProxy* CreateSceneProxy() override;
 	virtual class UBodySetup* GetBodySetup() override { return BrushBodySetup; };
 	virtual void GetUsedMaterials( TArray<UMaterialInterface*>& OutMaterials, bool bGetDebugMaterials = false ) const override;
-	virtual uint8 GetStaticDepthPriorityGroup() const override;
+	virtual ESceneDepthPriorityGroup GetStaticDepthPriorityGroup() const override;
 	virtual bool IsEditorOnly() const override;
 #if WITH_EDITOR
 	virtual bool ComponentIsTouchingSelectionBox(const FBox& InSelBBox, const FEngineShowFlags& ShowFlags, const bool bConsiderOnlyBSP, const bool bMustEncompassEntireComponent) const override;
diff --git a/Engine/Source/Runtime/Engine/Classes/Components/PrimitiveComponent.h b/Engine/Source/Runtime/Engine/Classes/Components/PrimitiveComponent.h
--- a/Engine/Source/Runtime/Engine/Classes/Components/PrimitiveComponent.h
+++ b/Engine/Source/Runtime/Engine/Classes/Components/PrimitiveComponent.h
@@ -1843,7 +1843,7 @@ public:
 	 * Even if the primitive's elements are drawn in multiple DPGs, a primary DPG is needed for occlusion culling and shadow projection.
 	 * @return The DPG the primitive's primary elements will be drawn in.
 	 */
-	virtual uint8 GetStaticDepthPriorityGroup() const { return DepthPriorityGroup; }
+	virtual ESceneDepthPriorityGroup GetStaticDepthPriorityGroup() const { return DepthPriorityGroup; }
 
 	/** 
 	 * Retrieves the materials used in this component 
diff --git a/Engine/Source/Runtime/Engine/Classes/Engine/EngineBaseTypes.h b/Engine/Source/Runtime/Engine/Classes/Engine/EngineBaseTypes.h
--- a/Engine/Source/Runtime/Engine/Classes/Engine/EngineBaseTypes.h
+++ b/Engine/Source/Runtime/Engine/Classes/Engine/EngineBaseTypes.h
@@ -1020,7 +1020,7 @@ struct FExposureSettings
 		TCHAR Buffer[BUFFER_SIZE];
 		check((Comma-In)+1 < BUFFER_SIZE);
 		
-		FCString::Strncpy(Buffer, In, (Comma-In)+1);
+		FCString::Strncpy(Buffer, In, UE_PTRDIFF_TO_INT32((Comma-In)+1));
 		FixedEV100 = FCString::Atof(Buffer);
 		bFixed = !!FCString::Atoi(Comma+1);
 	}
diff --git a/Engine/Source/Runtime/Engine/Classes/Engine/EngineTypes.h b/Engine/Source/Runtime/Engine/Classes/Engine/EngineTypes.h
--- a/Engine/Source/Runtime/Engine/Classes/Engine/EngineTypes.h
+++ b/Engine/Source/Runtime/Engine/Classes/Engine/EngineTypes.h
@@ -3543,7 +3543,7 @@ struct ENGINE_API FRepMovement
 		RBState.Quaternion = Rotation.Quaternion();
 		RBState.LinVel = LinearVelocity;
 		RBState.AngVel = AngularVelocity;
-		RBState.Flags = (bSimulatedPhysicSleep ? ERigidBodyFlags::Sleeping : ERigidBodyFlags::None) | ERigidBodyFlags::NeedsUpdate;
+		RBState.Flags = (decltype(FRigidBodyState::Flags))(bSimulatedPhysicSleep ? ERigidBodyFlags::Sleeping : ERigidBodyFlags::None) | ERigidBodyFlags::NeedsUpdate;
 	}
 
 	bool operator==(const FRepMovement& Other) const
@@ -4071,7 +4071,7 @@ private:
 public:
 	/** Max Samples to record. */
 	UPROPERTY(EditAnywhere, BlueprintReadWrite, Category="DebugFloatHistory")
-	float MaxSamples;
+	int32 MaxSamples;
 
 	/** Min value to record. */
 	UPROPERTY(EditAnywhere, BlueprintReadWrite, Category = "DebugFloatHistory")
@@ -4092,7 +4092,7 @@ public:
 		, bAutoAdjustMinMax(true)
 	{ }
 
-	FDebugFloatHistory(float const & InMaxSamples, float const & InMinValue, float const & InMaxValue, bool const & InbAutoAdjustMinMax)
+	FDebugFloatHistory(int32 const & InMaxSamples, float const & InMinValue, float const & InMaxValue, bool const & InbAutoAdjustMinMax)
 		: MaxSamples(InMaxSamples)
 		, MinValue(InMinValue)
 		, MaxValue(InMaxValue)
diff --git a/Engine/Source/Runtime/Engine/Classes/Engine/Scene.h b/Engine/Source/Runtime/Engine/Classes/Engine/Scene.h
--- a/Engine/Source/Runtime/Engine/Classes/Engine/Scene.h
+++ b/Engine/Source/Runtime/Engine/Classes/Engine/Scene.h
@@ -1498,7 +1498,7 @@ struct FPostProcessSettings
 
 #if WITH_EDITORONLY_DATA
 	UPROPERTY()
-	FVector BloomConvolutionPreFilter_DEPRECATED;
+	FVector3f BloomConvolutionPreFilter_DEPRECATED;
 #endif
 	
 	/** Boost intensity of select pixels  prior to computing bloom convolution (Min, Max, Multiplier).  Max < Min disables */
diff --git a/Engine/Source/Runtime/Engine/Classes/PhysicsEngine/BoxElem.h b/Engine/Source/Runtime/Engine/Classes/PhysicsEngine/BoxElem.h
--- a/Engine/Source/Runtime/Engine/Classes/PhysicsEngine/BoxElem.h
+++ b/Engine/Source/Runtime/Engine/Classes/PhysicsEngine/BoxElem.h
@@ -109,7 +109,7 @@ struct FKBoxElem : public FKShapeElem
 		Center = InTransform.GetLocation();
 	}
 
-	FORCEINLINE float GetVolume(const FVector& Scale3D) const { float MinScale = Scale3D.GetMin(); return (X * MinScale) * (Y * MinScale) * (Z * MinScale); }
+	FORCEINLINE FVector::FReal GetVolume(const FVector& Scale3D) const { FVector::FReal MinScale = Scale3D.GetMin(); return (X * MinScale) * (Y * MinScale) * (Z * MinScale); }
 
 	ENGINE_API void	DrawElemWire(class FPrimitiveDrawInterface* PDI, const FTransform& ElemTM, const FVector& Scale3D, const FColor Color) const;
 	ENGINE_API void	DrawElemSolid(class FPrimitiveDrawInterface* PDI, const FTransform& ElemTM, const FVector& Scale3D, const FMaterialRenderProxy* MaterialRenderProxy) const;
diff --git a/Engine/Source/Runtime/Engine/Classes/PhysicsEngine/SphereElem.h b/Engine/Source/Runtime/Engine/Classes/PhysicsEngine/SphereElem.h
--- a/Engine/Source/Runtime/Engine/Classes/PhysicsEngine/SphereElem.h
+++ b/Engine/Source/Runtime/Engine/Classes/PhysicsEngine/SphereElem.h
@@ -72,7 +72,7 @@ struct FKSphereElem : public FKShapeElem
 		Center = InTransform.GetLocation();
 	}
 
-	FORCEINLINE float GetVolume(const FVector& Scale) const { return 1.3333f * PI * FMath::Pow(Radius * Scale.GetMin(), 3); }
+	FORCEINLINE FVector::FReal GetVolume(const FVector& Scale) const { return 1.3333f * PI * FMath::Pow(Radius * Scale.GetMin(), 3); }
 	
 	ENGINE_API void	DrawElemWire(class FPrimitiveDrawInterface* PDI, const FTransform& ElemTM, const FVector& Scale3D, const FColor Color) const;
 
diff --git a/Engine/Source/Runtime/Engine/Classes/PhysicsEngine/SphylElem.h b/Engine/Source/Runtime/Engine/Classes/PhysicsEngine/SphylElem.h
--- a/Engine/Source/Runtime/Engine/Classes/PhysicsEngine/SphylElem.h
+++ b/Engine/Source/Runtime/Engine/Classes/PhysicsEngine/SphylElem.h
@@ -90,7 +90,7 @@ struct FKSphylElem : public FKShapeElem
 		Center = InTransform.GetLocation();
 	}
 
-	FORCEINLINE float GetVolume(const FVector& Scale) const { float ScaledRadius = Radius * Scale.GetMin(); return PI * FMath::Square(ScaledRadius) * ( 1.3333f * ScaledRadius + (Length * Scale.GetMin())); }
+	FORCEINLINE FVector::FReal GetVolume(const FVector& Scale) const { FVector::FReal ScaledRadius = Radius * Scale.GetMin(); return PI * FMath::Square(ScaledRadius) * ( 1.3333f * ScaledRadius + (Length * Scale.GetMin())); }
 
 	ENGINE_API void	DrawElemWire(class FPrimitiveDrawInterface* PDI, const FTransform& ElemTM, const FVector& Scale3D, const FColor Color) const;
 	ENGINE_API void GetElemSolid(const FTransform& ElemTM, const FVector& Scale3D, const FMaterialRenderProxy* MaterialRenderProxy, int32 ViewIndex, FMeshElementCollector& Collector) const;
diff --git a/Engine/Source/Runtime/Engine/Private/Components/BrushComponent.cpp b/Engine/Source/Runtime/Engine/Private/Components/BrushComponent.cpp
--- a/Engine/Source/Runtime/Engine/Private/Components/BrushComponent.cpp
+++ b/Engine/Source/Runtime/Engine/Private/Components/BrushComponent.cpp
@@ -527,7 +527,7 @@ void UBrushComponent::PostLoad()
 #endif
 }
 
-uint8 UBrushComponent::GetStaticDepthPriorityGroup() const
+ESceneDepthPriorityGroup UBrushComponent::GetStaticDepthPriorityGroup() const
 {
 	ABrush* BrushOwner = Cast<ABrush>(GetOwner());
 
diff --git a/Engine/Source/Runtime/Engine/Private/ModelRender.cpp b/Engine/Source/Runtime/Engine/Private/ModelRender.cpp
--- a/Engine/Source/Runtime/Engine/Private/ModelRender.cpp
+++ b/Engine/Source/Runtime/Engine/Private/ModelRender.cpp
@@ -516,7 +516,7 @@ public:
 		if (!HasViewDependentDPG())
 		{
 			// Determine the DPG the primitive should be drawn in.
-			uint8 PrimitiveDPG = GetStaticDepthPriorityGroup();
+			ESceneDepthPriorityGroup PrimitiveDPG = GetStaticDepthPriorityGroup();
 
 			PDI->ReserveMemoryForMeshes(Elements.Num());
 
@@ -537,7 +537,7 @@ public:
 					BatchElement.MaxVertexIndex = ModelElement.MaxVertexIndex;
 					BatchElement.VertexFactoryUserData = Elements[ElementIndex].GetVertexFactoryUniformBuffer();
 					MeshElement.Type = PT_TriangleList;
-					MeshElement.DepthPriorityGroup = PrimitiveDPG;
+					MeshElement.DepthPriorityGroup = (uint8)PrimitiveDPG;
 					MeshElement.LODIndex = 0;
 					const bool bValidIndexBuffer = !BatchElement.IndexBuffer || (BatchElement.IndexBuffer && BatchElement.IndexBuffer->IsInitialized() && BatchElement.IndexBuffer->IndexBufferRHI);
 					ensure(bValidIndexBuffer);
diff --git a/Engine/Source/Runtime/Engine/Private/PrimitiveSceneProxy.cpp b/Engine/Source/Runtime/Engine/Private/PrimitiveSceneProxy.cpp
--- a/Engine/Source/Runtime/Engine/Private/PrimitiveSceneProxy.cpp
+++ b/Engine/Source/Runtime/Engine/Private/PrimitiveSceneProxy.cpp
@@ -122,7 +122,7 @@ FPrimitiveSceneProxy::FPrimitiveSceneProxy(const UPrimitiveComponent* InComponen
 ,	bHovered(false)
 ,	bUseViewOwnerDepthPriorityGroup(InComponent->bUseViewOwnerDepthPriorityGroup)
 ,	bHasMotionBlurVelocityMeshes(InComponent->bHasMotionBlurVelocityMeshes)
-,	StaticDepthPriorityGroup(InComponent->GetStaticDepthPriorityGroup())
+,	StaticDepthPriorityGroup((uint8)InComponent->GetStaticDepthPriorityGroup())
 ,	ViewOwnerDepthPriorityGroup(InComponent->ViewOwnerDepthPriorityGroup)
 ,	bStaticLighting(InComponent->HasStaticLighting())
 ,	bVisibleInReflectionCaptures(InComponent->bVisibleInReflectionCaptures)
diff --git a/Engine/Source/Runtime/Engine/Private/SkeletalMesh.cpp b/Engine/Source/Runtime/Engine/Private/SkeletalMesh.cpp
--- a/Engine/Source/Runtime/Engine/Private/SkeletalMesh.cpp
+++ b/Engine/Source/Runtime/Engine/Private/SkeletalMesh.cpp
@@ -6045,7 +6045,7 @@ void FSkeletalMeshSceneProxy::DrawStaticElements(FStaticPrimitiveDrawInterface*
 
 	if (!HasViewDependentDPG())
 	{
-		uint8 PrimitiveDPG = GetStaticDepthPriorityGroup();
+		ESceneDepthPriorityGroup PrimitiveDPG = GetStaticDepthPriorityGroup();
 		bool bUseSelectedMaterial = false;
 
 		int32 NumLODs = SkeletalMeshRenderData->LODRenderData.Num();
diff --git a/Engine/Source/Runtime/Engine/Private/StaticMeshRender.cpp b/Engine/Source/Runtime/Engine/Private/StaticMeshRender.cpp
--- a/Engine/Source/Runtime/Engine/Private/StaticMeshRender.cpp
+++ b/Engine/Source/Runtime/Engine/Private/StaticMeshRender.cpp
@@ -1095,7 +1095,7 @@ void FStaticMeshSceneProxy::DrawStaticElements(FStaticPrimitiveDrawInterface* PD
 	if (!HasViewDependentDPG())
 	{
 		// Determine the DPG the primitive should be drawn in.
-		uint8 PrimitiveDPG = GetStaticDepthPriorityGroup();
+		ESceneDepthPriorityGroup PrimitiveDPG = GetStaticDepthPriorityGroup();
 		int32 NumLODs = RenderData->LODResources.Num();
 		//Never use the dynamic path in this path, because only unselected elements will use DrawStaticElements
 		bool bIsMeshElementSelected = false;
@@ -1696,7 +1696,7 @@ void FStaticMeshSceneProxy::GetDynamicRayTracingInstances(FRayTracingMaterialGat
 		return;
 	}
 
-	uint8 PrimitiveDPG = GetStaticDepthPriorityGroup();
+	ESceneDepthPriorityGroup PrimitiveDPG = GetStaticDepthPriorityGroup();
 	const uint32 LODIndex = FMath::Max(GetLOD(Context.ReferenceView), (int32)GetCurrentFirstLODIdx_RenderThread());
 	const FStaticMeshLODResources& LODModel = RenderData->LODResources[LODIndex];
 
diff --git a/Engine/Source/Runtime/Engine/Private/WorldPartition/HLOD/HLODLayer.cpp b/Engine/Source/Runtime/Engine/Private/WorldPartition/HLOD/HLODLayer.cpp
--- a/Engine/Source/Runtime/Engine/Private/WorldPartition/HLOD/HLODLayer.cpp
+++ b/Engine/Source/Runtime/Engine/Private/WorldPartition/HLOD/HLODLayer.cpp
@@ -113,7 +113,7 @@ void UHLODLayer::PostEditChangeProperty(FPropertyChangedEvent& PropertyChangedEv
 
 #if WITH_EDITORONLY_DATA
 
-FName UHLODLayer::GetRuntimeGridName(uint32 InLODLevel, int32 InCellSize, float InLoadingRange)
+FName UHLODLayer::GetRuntimeGridName(uint32 InLODLevel, int32 InCellSize, double InLoadingRange)
 {
 	return *FString::Format(TEXT("HLOD{0}_{1}m_{2}m"), { InLODLevel, int32(InCellSize * 0.01f), int32(InLoadingRange * 0.01f)});
 }
diff --git a/Engine/Source/Runtime/Engine/Public/InGamePerformanceTracker.h b/Engine/Source/Runtime/Engine/Public/InGamePerformanceTracker.h
--- a/Engine/Source/Runtime/Engine/Public/InGamePerformanceTracker.h
+++ b/Engine/Source/Runtime/Engine/Public/InGamePerformanceTracker.h
@@ -86,7 +86,7 @@ public:
 		return History.AddCycles(Cycles);
 	}
 
-	FORCEINLINE float GetAverageTimeSeconds()const
+	FORCEINLINE double GetAverageTimeSeconds() const
 	{
 		return FPlatformTime::ToSeconds(History.GetAverageCycles());
 	}
diff --git a/Engine/Source/Runtime/Engine/Public/Physics/Experimental/PhysScene_Chaos.h b/Engine/Source/Runtime/Engine/Public/Physics/Experimental/PhysScene_Chaos.h
--- a/Engine/Source/Runtime/Engine/Public/Physics/Experimental/PhysScene_Chaos.h
+++ b/Engine/Source/Runtime/Engine/Public/Physics/Experimental/PhysScene_Chaos.h
@@ -246,7 +246,7 @@ private:
 
 		friend inline uint32 GetTypeHash(FUniqueContactPairKey const& P)
 		{
-			return (PTRINT)P.Body0 ^ ((PTRINT)P.Body1 << 18);
+			return (uint32)((PTRINT)P.Body0 ^ ((PTRINT)P.Body1 << 18));
 		}
 	};
 
diff --git a/Engine/Source/Runtime/Engine/Public/PrimitiveSceneProxy.h b/Engine/Source/Runtime/Engine/Public/PrimitiveSceneProxy.h
--- a/Engine/Source/Runtime/Engine/Public/PrimitiveSceneProxy.h
+++ b/Engine/Source/Runtime/Engine/Public/PrimitiveSceneProxy.h
@@ -449,10 +449,10 @@ public:
 	 * Determines the DPG to render the primitive in regardless of view.
 	 * Should only be called if HasViewDependentDPG()==true.
 	 */
-	virtual uint8 GetStaticDepthPriorityGroup() const
+	virtual ESceneDepthPriorityGroup GetStaticDepthPriorityGroup() const
 	{
 		check(!HasViewDependentDPG());
-		return StaticDepthPriorityGroup;
+		return (ESceneDepthPriorityGroup)StaticDepthPriorityGroup;
 	}
 
 	/**
@@ -461,16 +461,16 @@ public:
 	 * @param View - The view to determine the primitive's DPG for.
 	 * @return The DPG the primitive should be rendered in for the given view.
 	 */
-	uint8 GetDepthPriorityGroup(const FSceneView* View) const
+	ESceneDepthPriorityGroup GetDepthPriorityGroup(const FSceneView* View) const
 	{
 		return (bUseViewOwnerDepthPriorityGroup && IsOwnedBy(View->ViewActor)) ?
-			ViewOwnerDepthPriorityGroup :
-			StaticDepthPriorityGroup;
+			(ESceneDepthPriorityGroup)ViewOwnerDepthPriorityGroup :
+			(ESceneDepthPriorityGroup)StaticDepthPriorityGroup;
 	}
 
 	/** Every derived class should override these functions */
 	virtual uint32 GetMemoryFootprint( void ) const = 0;
-	uint32 GetAllocatedSize( void ) const { return( Owners.GetAllocatedSize() ); }
+	SIZE_T GetAllocatedSize( void ) const { return( Owners.GetAllocatedSize() ); }
 
 	/**
 	 * Set the collision flag on the scene proxy to enable/disable collision drawing
diff --git a/Engine/Source/Runtime/Engine/Public/PrimitiveUniformShaderParameters.h b/Engine/Source/Runtime/Engine/Public/PrimitiveUniformShaderParameters.h
--- a/Engine/Source/Runtime/Engine/Public/PrimitiveUniformShaderParameters.h
+++ b/Engine/Source/Runtime/Engine/Public/PrimitiveUniformShaderParameters.h
@@ -258,9 +258,9 @@ public:
 			FVector4 WorldX = FVector4(Parameters.LocalToWorld.M[0][0], Parameters.LocalToWorld.M[0][1], Parameters.LocalToWorld.M[0][2], 0);
 			FVector4 WorldY = FVector4(Parameters.LocalToWorld.M[1][0], Parameters.LocalToWorld.M[1][1], Parameters.LocalToWorld.M[1][2], 0);
 			FVector4 WorldZ = FVector4(Parameters.LocalToWorld.M[2][0], Parameters.LocalToWorld.M[2][1], Parameters.LocalToWorld.M[2][2], 0);
-			float ScaleX = FVector(WorldX).Size();
-			float ScaleY = FVector(WorldY).Size();
-			float ScaleZ = FVector(WorldZ).Size();
+			FVector::FReal ScaleX = FVector(WorldX).Size();
+			FVector::FReal ScaleY = FVector(WorldY).Size();
+			FVector::FReal ScaleZ = FVector(WorldZ).Size();
 			Parameters.NonUniformScale = FVector4(ScaleX, ScaleY, ScaleZ, FMath::Max3(FMath::Abs(ScaleX), FMath::Abs(ScaleY), FMath::Abs(ScaleZ)));
 			Parameters.InvNonUniformScale = FVector3f(
 				ScaleX > KINDA_SMALL_NUMBER ? 1.0f / ScaleX : 0.0f,
diff --git a/Engine/Source/Runtime/Engine/Public/RawIndexBuffer.h b/Engine/Source/Runtime/Engine/Public/RawIndexBuffer.h
--- a/Engine/Source/Runtime/Engine/Public/RawIndexBuffer.h
+++ b/Engine/Source/Runtime/Engine/Public/RawIndexBuffer.h
@@ -261,7 +261,7 @@ public:
 	/**
 	 * Computes the amount of memory allocated to store the indices.
 	 */
-	FORCEINLINE uint32 GetAllocatedSize() const
+	FORCEINLINE SIZE_T GetAllocatedSize() const
 	{
 		return IndexStorage.GetAllocatedSize();
 	}
@@ -447,7 +447,7 @@ public:
 	virtual int32 AddItem(uint32 Val) override
 	{
 		++CachedNumIndices;
-		return Indices.Add(Val);
+		return Indices.Add((INDEX_TYPE)Val);
 	}
 
 	virtual uint32 Get(uint32 Idx) const override
diff --git a/Engine/Source/Runtime/Engine/Public/SceneManagement.h b/Engine/Source/Runtime/Engine/Public/SceneManagement.h
--- a/Engine/Source/Runtime/Engine/Public/SceneManagement.h
+++ b/Engine/Source/Runtime/Engine/Public/SceneManagement.h
@@ -3125,8 +3125,8 @@ struct FLODMask
 
 	void SetLOD(int32 LODIndex)
 	{
-		DitheredLODIndices[0] = LODIndex;
-		DitheredLODIndices[1] = LODIndex;
+		DitheredLODIndices[0] = (int8)LODIndex;
+		DitheredLODIndices[1] = (int8)LODIndex;
 	}
 	void SetLODSample(int32 LODIndex, int32 SampleIndex)
 	{
diff --git a/Engine/Source/Runtime/Engine/Public/SceneView.h b/Engine/Source/Runtime/Engine/Public/SceneView.h
--- a/Engine/Source/Runtime/Engine/Public/SceneView.h
+++ b/Engine/Source/Runtime/Engine/Public/SceneView.h
@@ -471,10 +471,11 @@ public:
 		VUp.Normalize();
 		VRight.Normalize();
 
-		return FVector2D(FMath::Acos(VCenter | VRight), FMath::Acos(VCenter | VUp));
+		using ResultType = decltype(FVector2D::X);
+		return FVector2D((ResultType)FMath::Acos(VCenter | VRight), (ResultType)FMath::Acos(VCenter | VUp));
 	}
 
-	float ComputeNearPlane() const
+	FMatrix::FReal ComputeNearPlane() const
 	{
 		return ( ProjectionMatrix.M[3][3] - ProjectionMatrix.M[3][2] ) / ( ProjectionMatrix.M[2][2] - ProjectionMatrix.M[2][3] );
 	}
diff --git a/Engine/Source/Runtime/Engine/Public/SkeletalMeshTypes.h b/Engine/Source/Runtime/Engine/Public/SkeletalMeshTypes.h
--- a/Engine/Source/Runtime/Engine/Public/SkeletalMeshTypes.h
+++ b/Engine/Source/Runtime/Engine/Public/SkeletalMeshTypes.h
@@ -182,7 +182,7 @@ public:
 	void DebugDrawSkeleton(int32 ViewIndex, FMeshElementCollector& Collector, const FEngineShowFlags& EngineShowFlags) const;
 
 	virtual uint32 GetMemoryFootprint( void ) const override { return( sizeof( *this ) + GetAllocatedSize() ); }
-	uint32 GetAllocatedSize( void ) const { return( FPrimitiveSceneProxy::GetAllocatedSize() + LODSections.GetAllocatedSize() ); }
+	SIZE_T GetAllocatedSize( void ) const { return( FPrimitiveSceneProxy::GetAllocatedSize() + LODSections.GetAllocatedSize() ); }
 
 	/**
 	* Updates morph material usage for materials referenced by each LOD entry
diff --git a/Engine/Source/Runtime/Engine/Public/StaticMeshResources.h b/Engine/Source/Runtime/Engine/Public/StaticMeshResources.h
--- a/Engine/Source/Runtime/Engine/Public/StaticMeshResources.h
+++ b/Engine/Source/Runtime/Engine/Public/StaticMeshResources.h
@@ -959,7 +959,7 @@ public:
 	virtual bool HasDistanceFieldRepresentation() const override;
 	virtual bool HasDynamicIndirectShadowCasterRepresentation() const override;
 	virtual uint32 GetMemoryFootprint( void ) const override { return( sizeof( *this ) + GetAllocatedSize() ); }
-	uint32 GetAllocatedSize( void ) const { return( FPrimitiveSceneProxy::GetAllocatedSize() + LODs.GetAllocatedSize() ); }
+	SIZE_T GetAllocatedSize( void ) const { return( FPrimitiveSceneProxy::GetAllocatedSize() + LODs.GetAllocatedSize() ); }
 
 	virtual void GetMeshDescription(int32 LODIndex, TArray<FMeshBatch>& OutMeshElements) const override;
 
@@ -1727,10 +1727,10 @@ private:
 			&& ensure((void*)((&ElementData[InstanceIndex]) + 0) >= (void*)(InstanceLightmapDataPtr)))
 		{
 
-			ElementData[InstanceIndex].InstanceLightmapAndShadowMapUVBias[0] = FMath::Clamp<int32>(FMath::TruncToInt(LightmapData.X * 32767.0f), MIN_int16, MAX_int16);
-			ElementData[InstanceIndex].InstanceLightmapAndShadowMapUVBias[1] = FMath::Clamp<int32>(FMath::TruncToInt(LightmapData.Y * 32767.0f), MIN_int16, MAX_int16);
-			ElementData[InstanceIndex].InstanceLightmapAndShadowMapUVBias[2] = FMath::Clamp<int32>(FMath::TruncToInt(LightmapData.Z * 32767.0f), MIN_int16, MAX_int16);
-			ElementData[InstanceIndex].InstanceLightmapAndShadowMapUVBias[3] = FMath::Clamp<int32>(FMath::TruncToInt(LightmapData.W * 32767.0f), MIN_int16, MAX_int16);
+			ElementData[InstanceIndex].InstanceLightmapAndShadowMapUVBias[0] = (int16)FMath::Clamp<int32>(FMath::TruncToInt(LightmapData.X * 32767.0f), MIN_int16, MAX_int16);
+			ElementData[InstanceIndex].InstanceLightmapAndShadowMapUVBias[1] = (int16)FMath::Clamp<int32>(FMath::TruncToInt(LightmapData.Y * 32767.0f), MIN_int16, MAX_int16);
+			ElementData[InstanceIndex].InstanceLightmapAndShadowMapUVBias[2] = (int16)FMath::Clamp<int32>(FMath::TruncToInt(LightmapData.Z * 32767.0f), MIN_int16, MAX_int16);
+			ElementData[InstanceIndex].InstanceLightmapAndShadowMapUVBias[3] = (int16)FMath::Clamp<int32>(FMath::TruncToInt(LightmapData.W * 32767.0f), MIN_int16, MAX_int16);
 		}
 
 	}
diff --git a/Engine/Source/Runtime/Engine/Public/UnrealEngine.h b/Engine/Source/Runtime/Engine/Public/UnrealEngine.h
--- a/Engine/Source/Runtime/Engine/Public/UnrealEngine.h
+++ b/Engine/Source/Runtime/Engine/Public/UnrealEngine.h
@@ -361,12 +361,12 @@ inline void CalculateFPSTimings()
 	extern ENGINE_API float GAverageFPS;
 	extern ENGINE_API float GAverageMS;
 	// Calculate the average frame time via continued averaging.
-	static double LastTime	= 0;
-	double CurrentTime		= FPlatformTime::Seconds();
-	float FrameTime			= (CurrentTime - LastTime) * 1000;
+	static double LastTime = 0.0;
+	double CurrentTime = FPlatformTime::Seconds();
+	const float FrameTimeMS = (float)((CurrentTime - LastTime) * 1000.0);
 	// A 3/4, 1/4 split gets close to a simple 10 frame moving average
-	GAverageMS				= GAverageMS * 0.75f + FrameTime * 0.25f;
-	LastTime				= CurrentTime;
+	GAverageMS = GAverageMS * 0.75f + FrameTimeMS * 0.25f;
+	LastTime = CurrentTime;
 	// Calculate average framerate.
 	GAverageFPS = 1000.f / GAverageMS;
 }
diff --git a/Engine/Source/Runtime/Engine/Public/WorldPartition/HLOD/HLODLayer.h b/Engine/Source/Runtime/Engine/Public/WorldPartition/HLOD/HLODLayer.h
--- a/Engine/Source/Runtime/Engine/Public/WorldPartition/HLOD/HLODLayer.h
+++ b/Engine/Source/Runtime/Engine/Public/WorldPartition/HLOD/HLODLayer.h
@@ -53,7 +53,7 @@ public:
 	const TSoftObjectPtr<UHLODLayer>& GetParentLayer() const;
 	bool IsAlwaysLoaded() const { return bAlwaysLoaded; }
 
-	static FName GetRuntimeGridName(uint32 InLODLevel, int32 InCellSize, float InLoadingRange);
+	static FName GetRuntimeGridName(uint32 InLODLevel, int32 InCellSize, double InLoadingRange);
 #endif
 
 #if WITH_EDITOR
@@ -87,7 +87,7 @@ private:
 
 	/** Loading range of the runtime grid created to encompass HLOD actors generated for this HLOD Layer */
 	UPROPERTY(EditAnywhere, Config, Category=HLOD, meta = (EditConditionHides, EditCondition = "!bAlwaysLoaded"))
-	float LoadingRange;
+	double LoadingRange;
 
 	/** HLOD Layer to assign to the generated HLOD actors */
 	UPROPERTY(EditAnywhere, Config, Category=HLOD, meta = (EditConditionHides, EditCondition = "!bAlwaysLoaded"))
diff --git a/Engine/Source/Runtime/Net/Core/Public/Net/Core/Misc/GuidReferences.h b/Engine/Source/Runtime/Net/Core/Public/Net/Core/Misc/GuidReferences.h
--- a/Engine/Source/Runtime/Net/Core/Public/Net/Core/Misc/GuidReferences.h
+++ b/Engine/Source/Runtime/Net/Core/Public/Net/Core/Misc/GuidReferences.h
@@ -53,7 +53,7 @@ public:
 		MappedDynamicGUIDs(InMappedDynamicGUIDs),
 		Array(nullptr)
 	{
-		NumBufferBits = InReader.GetPosBits() - InMark.GetPos();
+		NumBufferBits = UE_PTRDIFF_TO_INT32(InReader.GetPosBits() - InMark.GetPos());
 		InMark.Copy(InReader, Buffer);
 	}
 
diff --git a/Engine/Source/Runtime/PakFile/Public/IPlatformFilePak.h b/Engine/Source/Runtime/PakFile/Public/IPlatformFilePak.h
--- a/Engine/Source/Runtime/PakFile/Public/IPlatformFilePak.h
+++ b/Engine/Source/Runtime/PakFile/Public/IPlatformFilePak.h
@@ -262,7 +262,7 @@ struct FPakInfo
 		}
 	}
 
-	uint8 GetCompressionMethodIndex(FName CompressionMethod)
+	int32 GetCompressionMethodIndex(FName CompressionMethod)
 	{
 		// look for existing method
 		for (uint8 Index = 0; Index < CompressionMethods.Num(); Index++)
