Fixed blocking sources : distinguish minimum distance between cell and blocking sources from distance used to sort cells.

diff --git a/Engine/Source/Runtime/Engine/Private/WorldPartition/WorldPartitionRuntimeSpatialHash.cpp b/Engine/Source/Runtime/Engine/Private/WorldPartition/WorldPartitionRuntimeSpatialHash.cpp
--- a/Engine/Source/Runtime/Engine/Private/WorldPartition/WorldPartitionRuntimeSpatialHash.cpp
+++ b/Engine/Source/Runtime/Engine/Private/WorldPartition/WorldPartitionRuntimeSpatialHash.cpp
@@ -1221,7 +1221,7 @@ EWorldPartitionStreamingPerformance UWorldPartitionRuntimeSpatialHash::GetStream
 
 	if (StreamingCell->CachedIsBlockingSource)
 	{
-		const float Distance = FMath::Sqrt(StreamingCell->CachedMinSquareDistanceToSource) - ((float)(*StreamingGrid)->GetCellSize(StreamingCell->Level) / 2.f);
+		const float Distance = FMath::Sqrt(StreamingCell->CachedMinSquareDistanceToBlockingSource) - ((float)(*StreamingGrid)->GetCellSize(StreamingCell->Level) / 2.f);
 
 		const float Ratio = Distance / LoadingRange;
 
diff --git a/Engine/Source/Runtime/Engine/Private/WorldPartition/WorldPartitionRuntimeSpatialHashCell.cpp b/Engine/Source/Runtime/Engine/Private/WorldPartition/WorldPartitionRuntimeSpatialHashCell.cpp
--- a/Engine/Source/Runtime/Engine/Private/WorldPartition/WorldPartitionRuntimeSpatialHashCell.cpp
+++ b/Engine/Source/Runtime/Engine/Private/WorldPartition/WorldPartitionRuntimeSpatialHashCell.cpp
@@ -18,6 +18,7 @@ UWorldPartitionRuntimeSpatialHashCell::UWorldPartitionRuntimeSpatialHashCell(con
 : Super(ObjectInitializer)
 , Level(0)
 , CachedIsBlockingSource(false)
+, CachedMinSquareDistanceToBlockingSource(MAX_FLT)
 , CachedMinSquareDistanceToSource(MAX_FLT)
 , CachedSourceSortingDistance(0.f)
 {}
@@ -47,6 +48,7 @@ bool UWorldPartitionRuntimeSpatialHashCell::CacheStreamingSourceInfo(const UWorl
 	if (bWasCacheDirtied)
 	{
 		CachedIsBlockingSource = false;
+		CachedMinSquareDistanceToBlockingSource = MAX_FLT;
 		CachedMinSquareDistanceToSource = MAX_FLT;
 		CachedSourceModulatedDistances.Reset();
 	}
@@ -127,6 +129,7 @@ bool UWorldPartitionRuntimeSpatialHashCell::CacheStreamingSourceInfo(const UWorl
 	if (Info.Source.bBlockOnSlowLoading)
 	{
 		CachedIsBlockingSource = true;
+		CachedMinSquareDistanceToBlockingSource = FMath::Min(SquareDistance, CachedMinSquareDistanceToBlockingSource);
 	}
 
 	CachedMinSquareDistanceToSource = FMath::Min(SquareDistance, CachedMinSquareDistanceToSource);
diff --git a/Engine/Source/Runtime/Engine/Public/WorldPartition/WorldPartitionRuntimeSpatialHashCell.h b/Engine/Source/Runtime/Engine/Public/WorldPartition/WorldPartitionRuntimeSpatialHashCell.h
--- a/Engine/Source/Runtime/Engine/Public/WorldPartition/WorldPartitionRuntimeSpatialHashCell.h
+++ b/Engine/Source/Runtime/Engine/Public/WorldPartition/WorldPartitionRuntimeSpatialHashCell.h
@@ -30,7 +30,10 @@ class UWorldPartitionRuntimeSpatialHashCell : public UWorldPartitionRuntimeCell
 	// Used to determine if cell was requested by blocking source
 	mutable bool CachedIsBlockingSource;
 
-	// Represents the square distance from cell to closest streaming source
+	// Represents the square distance from cell to the closest blocking streaming source
+	mutable float CachedMinSquareDistanceToBlockingSource;
+
+	// Represents the square distance from cell to the closest streaming source
 	mutable float CachedMinSquareDistanceToSource;
 
 	// Modulated distance to the different streaming sources used to sort relative priority amongst streaming cells
