Fix an issue where the skeleton remapping was going wrong after adding or removing virtual bones.
This was because this triggered some different code path being taken after this edit. That code path did not do any skeleton remapping.

diff --git a/Engine/Source/Runtime/Engine/Private/Animation/AnimSequenceHelpers.cpp b/Engine/Source/Runtime/Engine/Private/Animation/AnimSequenceHelpers.cpp
--- a/Engine/Source/Runtime/Engine/Private/Animation/AnimSequenceHelpers.cpp
+++ b/Engine/Source/Runtime/Engine/Private/Animation/AnimSequenceHelpers.cpp
@@ -64,7 +64,7 @@ FTransform ExtractTransformForKey(int32 Key, const FRawAnimSequenceTrack& TrackT
 }
 
 template<bool bInterpolateT>
-void ExtractPoseFromModel(const TArray<FBoneAnimationTrack>& BoneAnimationTracks, const TMap<FName, FTransform>& OverrideBoneTransforms, FCompactPose& InOutPose, const int32 KeyIndex1, const int32 KeyIndex2, float Alpha)
+void ExtractPoseFromModel(const TArray<FBoneAnimationTrack>& BoneAnimationTracks, const TMap<FName, FTransform>& OverrideBoneTransforms, FCompactPose& InOutPose, const int32 KeyIndex1, const int32 KeyIndex2, float Alpha, const USkeleton* SourceSkeleton)
 {
 	const int32 NumAnimationTracks = BoneAnimationTracks.Num();
 	const FBoneContainer& RequiredBones = InOutPose.GetBoneContainer();
@@ -81,9 +81,14 @@ void ExtractPoseFromModel(const TArray<FBoneAnimationTrack>& BoneAnimationTracks
 		Key2Pose.CopyBonesFrom(InOutPose);
 	}
 
+	const USkeleton* TargetSkeleton = RequiredBones.GetSkeletonAsset();
+	const FSkeletonRemapping* SkeletonRemapping = TargetSkeleton->GetSkeletonRemapping(SourceSkeleton);
+
 	for (const FBoneAnimationTrack& AnimationTrack : BoneAnimationTracks)
 	{
-		const int32 SkeletonBoneIndex = AnimationTrack.BoneTreeIndex;
+		const int32 SourceSkeletonBoneIndex = AnimationTrack.BoneTreeIndex;
+		const int32 SkeletonBoneIndex = SkeletonRemapping ? SkeletonRemapping->GetTargetSkeletonBoneIndex(SourceSkeletonBoneIndex) : SourceSkeletonBoneIndex;
+
 		// not sure it's safe to assume that SkeletonBoneIndex can never be INDEX_NONE
 		if ((SkeletonBoneIndex != INDEX_NONE) && (SkeletonBoneIndex < MAX_BONES))
 		{
@@ -213,19 +218,19 @@ void BuildPoseFromModel(const UAnimDataModel* Model, FCompactPose& OutPose, cons
 	}
 
 	TMap<FName, FTransform> ActiveCurves;
-
 	if (!OutPose.GetBoneContainer().ShouldUseSourceData())
 	{
 		EvaluateTransformCurvesFromModel(Model, ActiveCurves, Time, 1.f);
 	}
 
+	const USkeleton* SourceSkeleton = Model->GetAnimationSequence()->GetSkeleton();
 	if (bShouldInterpolate)
 	{
-		ExtractPoseFromModel<true>(Model->GetBoneAnimationTracks(), ActiveCurves, OutPose, KeyIndex1, KeyIndex2, Alpha);
+		ExtractPoseFromModel<true>(Model->GetBoneAnimationTracks(), ActiveCurves, OutPose, KeyIndex1, KeyIndex2, Alpha, SourceSkeleton);
 	}
 	else
 	{
-		ExtractPoseFromModel<false>(Model->GetBoneAnimationTracks(), ActiveCurves, OutPose, KeyIndex1, KeyIndex2, Alpha);
+		ExtractPoseFromModel<false>(Model->GetBoneAnimationTracks(), ActiveCurves, OutPose, KeyIndex1, KeyIndex2, Alpha, SourceSkeleton);
 	}
 
 	Retargeting::RetargetPose(OutPose, RetargetSource, RetargetTransforms);
