Clean up some more material parameter related code
- Remove UMaterialExpression::SetValueToMatchingExpression, replace with SetParameterValue, add Flags to SetParameterValue to allow matching behavior
- Add UMaterialExpression::GetParameterValue
- Fix propogating curve parameter values between expressions

diff --git a/Engine/Source/Runtime/Engine/Classes/Materials/MaterialExpression.h b/Engine/Source/Runtime/Engine/Classes/Materials/MaterialExpression.h
--- a/Engine/Source/Runtime/Engine/Classes/Materials/MaterialExpression.h
+++ b/Engine/Source/Runtime/Engine/Classes/Materials/MaterialExpression.h
@@ -149,6 +149,14 @@ enum class EMaterialGenerateHLSLStatus : uint8
 	Error,
 };
 
+enum class EMaterialExpressionSetParameterValueFlags : uint32
+{
+	None = 0u,
+	SendPostEditChangeProperty = (1u << 0), // Send PostEditChangeProperty events for all properties that are modified
+	NoUpdateExpressionGuid = (1u << 1), // By default ExpressionGUI will be updated for static parameters
+};
+ENUM_CLASS_FLAGS(EMaterialExpressionSetParameterValueFlags);
+
 UCLASS(abstract, BlueprintType, hidecategories=Object)
 class ENGINE_API UMaterialExpression : public UObject
 {
@@ -515,11 +523,13 @@ class ENGINE_API UMaterialExpression : public UObject
 	virtual bool HasAParameterName() const { return false; }
 	virtual void ValidateParameterName(const bool bAllowDuplicateName = true);
 	virtual bool HasClassAndNameCollision(UMaterialExpression* OtherExpression) const;
-	virtual void SetValueToMatchingExpression(UMaterialExpression* OtherExpression) {};
+
+	EMaterialParameterType GetParameterType() const;
 
 	virtual FName GetParameterName() const { return NAME_None; }
 	virtual void SetParameterName(const FName& Name) {}
-	virtual bool SetParameterValue(const FName& Name, const FMaterialParameterMetadata& Meta) { return false; }
+	virtual bool GetParameterValue(FMaterialParameterMetadata& OutMeta) const { return false; }
+	virtual bool SetParameterValue(const FName& Name, const FMaterialParameterMetadata& Meta, EMaterialExpressionSetParameterValueFlags Flags = EMaterialExpressionSetParameterValueFlags::None) { return false; }
 
 	/**
 	 * Called after a node copy, once the Material and Function properties are set correctly and that all new expressions are added to Material->Expressions
diff --git a/Engine/Source/Runtime/Engine/Classes/Materials/MaterialExpressionChannelMaskParameter.h b/Engine/Source/Runtime/Engine/Classes/Materials/MaterialExpressionChannelMaskParameter.h
--- a/Engine/Source/Runtime/Engine/Classes/Materials/MaterialExpressionChannelMaskParameter.h
+++ b/Engine/Source/Runtime/Engine/Classes/Materials/MaterialExpressionChannelMaskParameter.h
@@ -33,7 +33,17 @@ class UMaterialExpressionChannelMaskParameter : public UMaterialExpressionVector
 #endif
 
 #if WITH_EDITOR
-	virtual bool SetParameterValue(FName InParameterName, FLinearColor InValue) override;
+	virtual bool GetParameterValue(FMaterialParameterMetadata& OutMeta) const override
+	{
+		if (Super::GetParameterValue(OutMeta))
+		{
+			OutMeta.bUsedAsChannelMask = true;
+			return true;
+		}
+		return false;
+	}
+
+	virtual bool SetParameterValue(FName InParameterName, FLinearColor InValue, EMaterialExpressionSetParameterValueFlags Flags) override;
 	virtual void PostEditChangeProperty(FPropertyChangedEvent& PropertyChangedEvent) override;
 
 	virtual int32 Compile(class FMaterialCompiler* Compiler, int32 OutputIndex) override;
diff --git a/Engine/Source/Runtime/Engine/Classes/Materials/MaterialExpressionCurveAtlasRowParameter.h b/Engine/Source/Runtime/Engine/Classes/Materials/MaterialExpressionCurveAtlasRowParameter.h
--- a/Engine/Source/Runtime/Engine/Classes/Materials/MaterialExpressionCurveAtlasRowParameter.h
+++ b/Engine/Source/Runtime/Engine/Classes/Materials/MaterialExpressionCurveAtlasRowParameter.h
@@ -28,6 +28,9 @@ class UMaterialExpressionCurveAtlasRowParameter : public UMaterialExpressionScal
 	virtual bool CanReferenceTexture() const { return true; }
 
 #if WITH_EDITOR
+	virtual bool GetParameterValue(FMaterialParameterMetadata& OutMeta) const override;
+	virtual bool SetParameterValue(const FName& Name, const FMaterialParameterMetadata& Meta, EMaterialExpressionSetParameterValueFlags Flags) override;
+
 	virtual uint32 GetInputType(int32 InputIndex) override 
 	{
 		return MCT_Float;
diff --git a/Engine/Source/Runtime/Engine/Classes/Materials/MaterialExpressionFontSampleParameter.h b/Engine/Source/Runtime/Engine/Classes/Materials/MaterialExpressionFontSampleParameter.h
--- a/Engine/Source/Runtime/Engine/Classes/Materials/MaterialExpressionFontSampleParameter.h
+++ b/Engine/Source/Runtime/Engine/Classes/Materials/MaterialExpressionFontSampleParameter.h
@@ -48,14 +48,19 @@ class UMaterialExpressionFontSampleParameter : public UMaterialExpressionFontSam
 	virtual FName GetParameterName() const override { return ParameterName; }
 	virtual void SetParameterName(const FName& Name) override { ParameterName = Name; }
 	virtual void ValidateParameterName(const bool bAllowDuplicateName) override;
-	virtual void SetValueToMatchingExpression(UMaterialExpression* OtherExpression) override;
-	virtual bool SetParameterValue(const FName& Name, const FMaterialParameterMetadata& Meta) override
+	virtual bool GetParameterValue(FMaterialParameterMetadata& OutMeta) const override
+	{
+		OutMeta.Value = FMaterialParameterValue(Font, FontTexturePage);
+		OutMeta.ExpressionGuid = ExpressionGUID;
+		return true;
+	}
+	virtual bool SetParameterValue(const FName& Name, const FMaterialParameterMetadata& Meta, EMaterialExpressionSetParameterValueFlags Flags) override
 	{
 		if (Meta.Value.Type == EMaterialParameterType::Font)
 		{
-			return SetParameterValue(Name, Meta.Value.Font.Value, Meta.Value.Font.Page);
+			return SetParameterValue(Name, Meta.Value.Font.Value, Meta.Value.Font.Page, Flags);
 		}
-		return Super::SetParameterValue(Name, Meta);
+		return Super::SetParameterValue(Name, Meta, Flags);
 	}
 #endif
 	//~ End UMaterialExpression Interface
@@ -64,7 +69,7 @@ class UMaterialExpressionFontSampleParameter : public UMaterialExpressionFontSam
 	bool IsNamedParameter(const FHashedMaterialParameterInfo& ParameterInfo, UFont*& OutFontValue, int32& OutFontPage) const;
 
 #if WITH_EDITOR
-	bool SetParameterValue(FName InParameterName, UFont* InFontValue, int32 InFontPage);
+	bool SetParameterValue(FName InParameterName, UFont* InFontValue, int32 InFontPage, EMaterialExpressionSetParameterValueFlags Flags = EMaterialExpressionSetParameterValueFlags::None);
 #endif
 
 	/**
diff --git a/Engine/Source/Runtime/Engine/Classes/Materials/MaterialExpressionRuntimeVirtualTextureSampleParameter.h b/Engine/Source/Runtime/Engine/Classes/Materials/MaterialExpressionRuntimeVirtualTextureSampleParameter.h
--- a/Engine/Source/Runtime/Engine/Classes/Materials/MaterialExpressionRuntimeVirtualTextureSampleParameter.h
+++ b/Engine/Source/Runtime/Engine/Classes/Materials/MaterialExpressionRuntimeVirtualTextureSampleParameter.h
@@ -31,7 +31,7 @@ class UMaterialExpressionRuntimeVirtualTextureSampleParameter : public UMaterial
 
 #if WITH_EDITOR
 	/** If this is the named parameter from this material expression, then set its value. */
-	bool SetParameterValue(FName InParameterName, URuntimeVirtualTexture* InValue);
+	bool SetParameterValue(FName InParameterName, URuntimeVirtualTexture* InValue, EMaterialExpressionSetParameterValueFlags Flags = EMaterialExpressionSetParameterValueFlags::None);
 #endif
 
 	/** Return whether this is the named parameter from this material expression, and if it is then return its value. */
@@ -50,14 +50,19 @@ class UMaterialExpressionRuntimeVirtualTextureSampleParameter : public UMaterial
 	virtual void ValidateParameterName(const bool bAllowDuplicateName) override;
 	virtual void GetCaption(TArray<FString>& OutCaptions) const override;
 	virtual bool MatchesSearchQuery(const TCHAR* SearchQuery) override;
-	virtual void SetValueToMatchingExpression(UMaterialExpression* OtherExpression) override;
-	virtual bool SetParameterValue(const FName& Name, const FMaterialParameterMetadata& Meta) override
+	virtual bool GetParameterValue(FMaterialParameterMetadata& OutMeta) const override
+	{
+		OutMeta.Value = VirtualTexture;
+		OutMeta.ExpressionGuid = ExpressionGUID;
+		return true;
+	}
+	virtual bool SetParameterValue(const FName& Name, const FMaterialParameterMetadata& Meta, EMaterialExpressionSetParameterValueFlags Flags) override
 	{
 		if (Meta.Value.Type == EMaterialParameterType::RuntimeVirtualTexture)
 		{
-			return SetParameterValue(Name, Meta.Value.RuntimeVirtualTexture);
+			return SetParameterValue(Name, Meta.Value.RuntimeVirtualTexture, Flags);
 		}
-		return Super::SetParameterValue(Name, Meta);
+		return Super::SetParameterValue(Name, Meta, Flags);
 	}
 #endif
 	virtual FGuid& GetParameterExpressionId() override { return ExpressionGUID; }
diff --git a/Engine/Source/Runtime/Engine/Classes/Materials/MaterialExpressionScalarParameter.h b/Engine/Source/Runtime/Engine/Classes/Materials/MaterialExpressionScalarParameter.h
--- a/Engine/Source/Runtime/Engine/Classes/Materials/MaterialExpressionScalarParameter.h
+++ b/Engine/Source/Runtime/Engine/Classes/Materials/MaterialExpressionScalarParameter.h
@@ -44,13 +44,21 @@ class UMaterialExpressionScalarParameter : public UMaterialExpressionParameter
 	virtual int32 Compile(class FMaterialCompiler* Compiler, int32 OutputIndex) override;
 	virtual void GetCaption(TArray<FString>& OutCaptions) const override;
 	EMaterialGenerateHLSLStatus GenerateHLSLExpression(FMaterialHLSLGenerator& Generator, UE::HLSLTree::FScope& Scope, int32 OutputIndex, UE::HLSLTree::FExpression*& OutExpression) override;
-	virtual bool SetParameterValue(const FName& Name, const FMaterialParameterMetadata& Meta) override
+	virtual bool GetParameterValue(FMaterialParameterMetadata& OutMeta) const override
+	{
+		OutMeta.Value = DefaultValue;
+		OutMeta.ExpressionGuid = ExpressionGUID;
+		OutMeta.ScalarMin = SliderMin;
+		OutMeta.ScalarMax = SliderMax;
+		return true;
+	}
+	virtual bool SetParameterValue(const FName& Name, const FMaterialParameterMetadata& Meta, EMaterialExpressionSetParameterValueFlags Flags) override
 	{
 		if (Meta.Value.Type == EMaterialParameterType::Scalar)
 		{
-			return SetParameterValue(Name, Meta.Value.AsScalar());
+			return SetParameterValue(Name, Meta.Value.AsScalar(), Flags);
 		}
-		return Super::SetParameterValue(Name, Meta);
+		return Super::SetParameterValue(Name, Meta, Flags);
 	}
 #endif
 	//~ End UMaterialExpression Interface
@@ -59,12 +67,11 @@ class UMaterialExpressionScalarParameter : public UMaterialExpressionParameter
 	bool IsNamedParameter(const FHashedMaterialParameterInfo& ParameterInfo, float& OutValue) const;
 
 #if WITH_EDITOR
-	bool SetParameterValue(FName InParameterName, float InValue);
+	bool SetParameterValue(FName InParameterName, float InValue, EMaterialExpressionSetParameterValueFlags Flags = EMaterialExpressionSetParameterValueFlags::None);
 
 	virtual void PostEditChangeProperty(FPropertyChangedEvent& PropertyChangedEvent) override;
 	virtual void ValidateParameterName(const bool bAllowDuplicateName) override;
 	virtual bool HasClassAndNameCollision(UMaterialExpression* OtherExpression) const override;
-	virtual void SetValueToMatchingExpression(UMaterialExpression* OtherExpression) override;
 #endif
 
 	virtual bool IsUsedAsAtlasPosition() const { return false; }
diff --git a/Engine/Source/Runtime/Engine/Classes/Materials/MaterialExpressionStaticBoolParameter.h b/Engine/Source/Runtime/Engine/Classes/Materials/MaterialExpressionStaticBoolParameter.h
--- a/Engine/Source/Runtime/Engine/Classes/Materials/MaterialExpressionStaticBoolParameter.h
+++ b/Engine/Source/Runtime/Engine/Classes/Materials/MaterialExpressionStaticBoolParameter.h
@@ -23,15 +23,20 @@ public:
 	virtual int32 CompilePreview(class FMaterialCompiler* Compiler, int32 OutputIndex) override;
 	virtual void GetCaption(TArray<FString>& OutCaptions) const override;
 	virtual uint32 GetOutputType(int32 OutputIndex) override {return MCT_StaticBool;}
-	virtual void SetValueToMatchingExpression(UMaterialExpression* OtherExpression) override;
 	virtual EMaterialGenerateHLSLStatus GenerateHLSLExpression(FMaterialHLSLGenerator& Generator, UE::HLSLTree::FScope& Scope, int32 OutputIndex, UE::HLSLTree::FExpression*& OutExpression) override;
-	virtual bool SetParameterValue(const FName& Name, const FMaterialParameterMetadata& Meta) override
+	virtual bool GetParameterValue(FMaterialParameterMetadata& OutMeta) const override
+	{
+		OutMeta.Value = (bool)DefaultValue;
+		OutMeta.ExpressionGuid = ExpressionGUID;
+		return true;
+	}
+	virtual bool SetParameterValue(const FName& Name, const FMaterialParameterMetadata& Meta, EMaterialExpressionSetParameterValueFlags Flags) override
 	{
 		if (Meta.Value.Type == EMaterialParameterType::StaticSwitch)
 		{
-			return SetParameterValue(Name, Meta.Value.AsStaticSwitch(), Meta.ExpressionGuid);
+			return SetParameterValue(Name, Meta.Value.AsStaticSwitch(), Meta.ExpressionGuid, Flags);
 		}
-		return Super::SetParameterValue(Name, Meta);
+		return Super::SetParameterValue(Name, Meta, Flags);
 	}
 #endif
 	//~ End UMaterialExpression Interface
@@ -40,7 +45,7 @@ public:
 	bool IsNamedParameter(const FHashedMaterialParameterInfo& ParameterInfo, bool& OutValue, FGuid& OutExpressionGuid) const;
 
 #if WITH_EDITOR
-	bool SetParameterValue(FName InParameterName, bool OutValue, FGuid InExpressionGuid);
+	bool SetParameterValue(FName InParameterName, bool OutValue, FGuid InExpressionGuid, EMaterialExpressionSetParameterValueFlags Flags = EMaterialExpressionSetParameterValueFlags::None);
 #endif
 };
 
diff --git a/Engine/Source/Runtime/Engine/Classes/Materials/MaterialExpressionStaticComponentMaskParameter.h b/Engine/Source/Runtime/Engine/Classes/Materials/MaterialExpressionStaticComponentMaskParameter.h
--- a/Engine/Source/Runtime/Engine/Classes/Materials/MaterialExpressionStaticComponentMaskParameter.h
+++ b/Engine/Source/Runtime/Engine/Classes/Materials/MaterialExpressionStaticComponentMaskParameter.h
@@ -38,8 +38,13 @@ public:
 #if WITH_EDITOR
 	virtual int32 Compile(class FMaterialCompiler* Compiler, int32 OutputIndex) override;
 	virtual void GetCaption(TArray<FString>& OutCaptions) const override;
-	virtual void SetValueToMatchingExpression(UMaterialExpression* OtherExpression) override;
-	virtual bool SetParameterValue(const FName& Name, const FMaterialParameterMetadata& Meta) override
+	virtual bool GetParameterValue(FMaterialParameterMetadata& OutMeta) const override
+	{
+		OutMeta.Value = FMaterialParameterValue(DefaultR, DefaultG, DefaultB, DefaultA);
+		OutMeta.ExpressionGuid = ExpressionGUID;
+		return true;
+	}
+	virtual bool SetParameterValue(const FName& Name, const FMaterialParameterMetadata& Meta, EMaterialExpressionSetParameterValueFlags Flags) override
 	{
 		if (Meta.Value.Type == EMaterialParameterType::StaticComponentMask)
 		{
@@ -48,9 +53,10 @@ public:
 				Meta.Value.Bool[1],
 				Meta.Value.Bool[2],
 				Meta.Value.Bool[3],
-				Meta.ExpressionGuid);
+				Meta.ExpressionGuid,
+				Flags);
 		}
-		return Super::SetParameterValue(Name, Meta);
+		return Super::SetParameterValue(Name, Meta, Flags);
 	}
 #endif
 	//~ End UMaterialExpression Interface
@@ -59,7 +65,7 @@ public:
 	bool IsNamedParameter(const FHashedMaterialParameterInfo& ParameterInfo, bool& OutR, bool& OutG, bool& OutB, bool& OutA, FGuid&OutExpressionGuid) const;
 
 #if WITH_EDITOR
-	bool SetParameterValue(FName InParameterName, bool InR, bool InG, bool InB, bool InA, FGuid InExpressionGuid);
+	bool SetParameterValue(FName InParameterName, bool InR, bool InG, bool InB, bool InA, FGuid InExpressionGuid, EMaterialExpressionSetParameterValueFlags Flags = EMaterialExpressionSetParameterValueFlags::None);
 #endif
 };
 
diff --git a/Engine/Source/Runtime/Engine/Classes/Materials/MaterialExpressionTextureSampleParameter.h b/Engine/Source/Runtime/Engine/Classes/Materials/MaterialExpressionTextureSampleParameter.h
--- a/Engine/Source/Runtime/Engine/Classes/Materials/MaterialExpressionTextureSampleParameter.h
+++ b/Engine/Source/Runtime/Engine/Classes/Materials/MaterialExpressionTextureSampleParameter.h
@@ -50,15 +50,24 @@ class ENGINE_API UMaterialExpressionTextureSampleParameter : public UMaterialExp
 	virtual FName GetParameterName() const override { return ParameterName; }
 	virtual void SetParameterName(const FName& Name) override { ParameterName = Name; }
 	virtual void ValidateParameterName(const bool bAllowDuplicateName) override;
-	virtual void SetValueToMatchingExpression(UMaterialExpression* OtherExpression) override;
 	virtual EMaterialGenerateHLSLStatus GenerateHLSLExpression(FMaterialHLSLGenerator& Generator, UE::HLSLTree::FScope& Scope, int32 OutputIndex, UE::HLSLTree::FExpression*& OutExpression) override;
-	virtual bool SetParameterValue(const FName& Name, const FMaterialParameterMetadata& Meta) override
+	virtual bool GetParameterValue(FMaterialParameterMetadata& OutMeta) const override
+	{
+		OutMeta.Value = Texture;
+		OutMeta.ExpressionGuid = ExpressionGUID;
+		OutMeta.ChannelName[0] = ChannelNames.R;
+		OutMeta.ChannelName[1] = ChannelNames.G;
+		OutMeta.ChannelName[2] = ChannelNames.B;
+		OutMeta.ChannelName[3] = ChannelNames.A;
+		return true;
+	}
+	virtual bool SetParameterValue(const FName& Name, const FMaterialParameterMetadata& Meta, EMaterialExpressionSetParameterValueFlags Flags) override
 	{
 		if (Meta.Value.Type == EMaterialParameterType::Texture)
 		{
-			return SetParameterValue(Name, Meta.Value.Texture);
+			return SetParameterValue(Name, Meta.Value.Texture, Flags);
 		}
-		return Super::SetParameterValue(Name, Meta);
+		return Super::SetParameterValue(Name, Meta, Flags);
 	}
 #endif
 	//~ End UMaterialExpression Interface
@@ -67,7 +76,7 @@ class ENGINE_API UMaterialExpressionTextureSampleParameter : public UMaterialExp
 	bool IsNamedParameter(const FHashedMaterialParameterInfo& ParameterInfo, UTexture*& OutValue) const;
 
 #if WITH_EDITOR
-	bool SetParameterValue(FName InParameterName, UTexture* InValue);
+	bool SetParameterValue(FName InParameterName, UTexture* InValue, EMaterialExpressionSetParameterValueFlags Flags = EMaterialExpressionSetParameterValueFlags::None);
 	virtual void PostEditChangeProperty(FPropertyChangedEvent& PropertyChangedEvent) override;
 	void ApplyChannelNames();
 
diff --git a/Engine/Source/Runtime/Engine/Classes/Materials/MaterialExpressionVectorParameter.h b/Engine/Source/Runtime/Engine/Classes/Materials/MaterialExpressionVectorParameter.h
--- a/Engine/Source/Runtime/Engine/Classes/Materials/MaterialExpressionVectorParameter.h
+++ b/Engine/Source/Runtime/Engine/Classes/Materials/MaterialExpressionVectorParameter.h
@@ -33,13 +33,23 @@ class UMaterialExpressionVectorParameter : public UMaterialExpressionParameter
 	virtual int32 Compile(class FMaterialCompiler* Compiler, int32 OutputIndex) override;
 	virtual void GetCaption(TArray<FString>& OutCaptions) const override;
 	EMaterialGenerateHLSLStatus GenerateHLSLExpression(FMaterialHLSLGenerator& Generator, UE::HLSLTree::FScope& Scope, int32 OutputIndex, UE::HLSLTree::FExpression*& OutExpression) override;
-	virtual bool SetParameterValue(const FName& Name, const FMaterialParameterMetadata& Meta) override
+	virtual bool GetParameterValue(FMaterialParameterMetadata& OutMeta) const override
+	{
+		OutMeta.Value = DefaultValue;
+		OutMeta.ExpressionGuid = ExpressionGUID;
+		OutMeta.ChannelName[0] = ChannelNames.R;
+		OutMeta.ChannelName[1] = ChannelNames.G;
+		OutMeta.ChannelName[2] = ChannelNames.B;
+		OutMeta.ChannelName[3] = ChannelNames.A;
+		return true;
+	}
+	virtual bool SetParameterValue(const FName& Name, const FMaterialParameterMetadata& Meta, EMaterialExpressionSetParameterValueFlags Flags) override
 	{
 		if (Meta.Value.Type == EMaterialParameterType::Vector)
 		{
-			return SetParameterValue(Name, Meta.Value.AsLinearColor());
+			return SetParameterValue(Name, Meta.Value.AsLinearColor(), Flags);
 		}
-		return Super::SetParameterValue(Name, Meta);
+		return Super::SetParameterValue(Name, Meta, Flags);
 	}
 #endif
 	//~ End UMaterialExpression Interface
@@ -48,7 +58,7 @@ class UMaterialExpressionVectorParameter : public UMaterialExpressionParameter
 	bool IsNamedParameter(const FHashedMaterialParameterInfo& ParameterInfo, FLinearColor& OutValue) const;
 
 #if WITH_EDITOR
-	virtual bool SetParameterValue(FName InParameterName, FLinearColor InValue);
+	virtual bool SetParameterValue(FName InParameterName, FLinearColor InValue, EMaterialExpressionSetParameterValueFlags Flags = EMaterialExpressionSetParameterValueFlags::None);
 
 	virtual void PostEditChangeProperty(FPropertyChangedEvent& PropertyChangedEvent) override;	
 
@@ -57,7 +67,6 @@ class UMaterialExpressionVectorParameter : public UMaterialExpressionParameter
 
 	virtual void ValidateParameterName(const bool bAllowDuplicateName) override;
 	virtual bool HasClassAndNameCollision(UMaterialExpression* OtherExpression) const override;
-	virtual void SetValueToMatchingExpression(UMaterialExpression* OtherExpression) override;
 #endif
 
 	virtual bool IsUsedAsChannelMask() const {return false;}
diff --git a/Engine/Source/Runtime/Engine/Private/Materials/Material.cpp b/Engine/Source/Runtime/Engine/Private/Materials/Material.cpp
--- a/Engine/Source/Runtime/Engine/Private/Materials/Material.cpp
+++ b/Engine/Source/Runtime/Engine/Private/Materials/Material.cpp
@@ -1336,25 +1336,9 @@ bool UMaterial::IsCompilingOrHadCompileError(ERHIFeatureLevel::Type InFeatureLev
 bool UMaterial::SetParameterValueEditorOnly(const FName& ParameterName, const FMaterialParameterMetadata& Meta)
 {
 	// Warning: in the case of duplicate parameters with different default values, this will find the first in the expression array, not necessarily the one that's used for rendering
-	// Lambda which calls SetParameterValue on a given parameter and triggers a PostEditChange event if successful
-	auto TrySetParameterValue = [ParameterName, Meta](UMaterialExpression* Expression) -> bool
-	{
-		if (Expression && Expression->SetParameterValue(ParameterName, Meta))
-		{
-			FProperty* ParamProperty = Expression->GetClass()->FindPropertyByName(TEXT("DefaultValue"));
-			if (ParamProperty)
-			{
-				FPropertyChangedEvent PropertyChangedEvent(ParamProperty);
-				Expression->PostEditChangeProperty(PropertyChangedEvent);
-				return true;
-			}
-		}
-		return false;
-	};
-
 	for (TObjectPtr<UMaterialExpression>& Expression : Expressions)
 	{
-		if (TrySetParameterValue(Expression))
+		if (Expression && Expression->SetParameterValue(ParameterName, Meta, EMaterialExpressionSetParameterValueFlags::SendPostEditChangeProperty))
 		{
 			return true;
 		}
@@ -1373,7 +1357,7 @@ bool UMaterial::SetParameterValueEditorOnly(const FName& ParameterName, const FM
 					{
 						for (const TObjectPtr<UMaterialExpression>& FunctionExpression : *ExpressionPtr)
 						{
-							if (TrySetParameterValue(FunctionExpression))
+							if (FunctionExpression && FunctionExpression->SetParameterValue(ParameterName, Meta, EMaterialExpressionSetParameterValueFlags::SendPostEditChangeProperty))
 							{
 								return true;
 							}
@@ -4460,23 +4444,11 @@ void UMaterial::RebuildShadingModelField()
 bool UMaterial::GetExpressionParameterName(const UMaterialExpression* Expression, FName& OutName)
 {
 	bool bRet = false;
-
-	if(Expression->IsA(UMaterialExpressionParameter::StaticClass()))
-	{
-		OutName = ((UMaterialExpressionParameter*)Expression)->ParameterName;
-		bRet = true;
-	}
-	else if(Expression->IsA(UMaterialExpressionTextureSampleParameter::StaticClass()))
-	{
-		OutName = ((UMaterialExpressionTextureSampleParameter*)Expression)->ParameterName;
-		bRet = true;
-	}
-	else if(Expression->IsA(UMaterialExpressionFontSampleParameter::StaticClass()))
+	if (Expression->HasAParameterName())
 	{
-		OutName = ((UMaterialExpressionFontSampleParameter*)Expression)->ParameterName;
+		OutName = Expression->GetParameterName();
 		bRet = true;
 	}
-
 	return bRet;
 }
 
@@ -4488,111 +4460,63 @@ bool UMaterial::CopyExpressionParameters(UMaterialExpression* Source, UMaterialE
 		return false;
 	}
 
-	bool bRet = true;
-
-	if(Source->IsA(UMaterialExpressionTextureSampleParameter::StaticClass()))
-	{
-		UMaterialExpressionTextureSampleParameter *SourceTex = (UMaterialExpressionTextureSampleParameter*)Source;
-		UMaterialExpressionTextureSampleParameter *DestTex = (UMaterialExpressionTextureSampleParameter*)Destination;
-
-		DestTex->Modify();
-		DestTex->Texture = SourceTex->Texture;
-		DestTex->Group = SourceTex->Group; // Group isn't editor only for textures
-#if WITH_EDITOR
-		DestTex->SortPriority = SourceTex->SortPriority;
-#endif
-	}
-	else if (Source->IsA(UMaterialExpressionRuntimeVirtualTextureSampleParameter::StaticClass()))
-	{
-		UMaterialExpressionRuntimeVirtualTextureSampleParameter* SourceRVT = (UMaterialExpressionRuntimeVirtualTextureSampleParameter*)Source;
-		UMaterialExpressionRuntimeVirtualTextureSampleParameter* DestRVT = (UMaterialExpressionRuntimeVirtualTextureSampleParameter*)Destination;
-
-		DestRVT->Modify();
-		DestRVT->VirtualTexture = SourceRVT->VirtualTexture;
-		DestRVT->Group = SourceRVT->Group;
-#if WITH_EDITORONLY_DATA
-		DestRVT->SortPriority = SourceRVT->SortPriority;
-#endif
-	}
-	else if(Source->IsA(UMaterialExpressionVectorParameter::StaticClass()))
-	{
-		UMaterialExpressionVectorParameter *SourceVec = (UMaterialExpressionVectorParameter*)Source;
-		UMaterialExpressionVectorParameter *DestVec = (UMaterialExpressionVectorParameter*)Destination;
-
-		DestVec->Modify();
-		DestVec->DefaultValue = SourceVec->DefaultValue;
-#if WITH_EDITOR
-		DestVec->Group = SourceVec->Group;
-		DestVec->SortPriority = SourceVec->SortPriority;
-#endif
-	}
-	else if(Source->IsA(UMaterialExpressionStaticBoolParameter::StaticClass()))
+	bool bResult = false;
+	FMaterialParameterMetadata Meta;
+	if (Source->GetParameterValue(Meta))
 	{
-		UMaterialExpressionStaticBoolParameter *SourceVec = (UMaterialExpressionStaticBoolParameter*)Source;
-		UMaterialExpressionStaticBoolParameter *DestVec = (UMaterialExpressionStaticBoolParameter*)Destination;
+		const FName ParameterName = Destination->GetParameterName();
+		Destination->Modify();
+		verify(Destination->SetParameterValue(ParameterName, Meta, EMaterialExpressionSetParameterValueFlags::NoUpdateExpressionGuid));
 
-		DestVec->Modify();
-		DestVec->DefaultValue = SourceVec->DefaultValue;
+		if (Source->IsA(UMaterialExpressionParameter::StaticClass()))
+		{
+			UMaterialExpressionParameter* SourceParam = (UMaterialExpressionParameter*)Source;
+			UMaterialExpressionParameter* DestParam = (UMaterialExpressionParameter*)Destination;
 #if WITH_EDITOR
-		DestVec->Group = SourceVec->Group;
-		DestVec->SortPriority = SourceVec->SortPriority;
+			DestParam->Group = SourceParam->Group;
+			DestParam->SortPriority = SourceParam->SortPriority;
 #endif
-	}
-	else if(Source->IsA(UMaterialExpressionStaticComponentMaskParameter::StaticClass()))
-	{
-		UMaterialExpressionStaticComponentMaskParameter *SourceVec = (UMaterialExpressionStaticComponentMaskParameter*)Source;
-		UMaterialExpressionStaticComponentMaskParameter *DestVec = (UMaterialExpressionStaticComponentMaskParameter*)Destination;
-
-		DestVec->Modify();
-		DestVec->DefaultR = SourceVec->DefaultR;
-		DestVec->DefaultG = SourceVec->DefaultG;
-		DestVec->DefaultB = SourceVec->DefaultB;
-		DestVec->DefaultA = SourceVec->DefaultA;
+		}
+		else if (Source->IsA(UMaterialExpressionTextureSampleParameter::StaticClass()))
+		{
+			UMaterialExpressionTextureSampleParameter* SourceTex = (UMaterialExpressionTextureSampleParameter*)Source;
+			UMaterialExpressionTextureSampleParameter* DestTex = (UMaterialExpressionTextureSampleParameter*)Destination;
+			DestTex->Group = SourceTex->Group; // Group isn't editor only for textures
 #if WITH_EDITOR
-		DestVec->Group = SourceVec->Group;
-		DestVec->SortPriority = SourceVec->SortPriority;
+			DestTex->SortPriority = SourceTex->SortPriority;
 #endif
-	}
-	else if(Source->IsA(UMaterialExpressionScalarParameter::StaticClass()))
-	{
-		UMaterialExpressionScalarParameter *SourceVec = (UMaterialExpressionScalarParameter*)Source;
-		UMaterialExpressionScalarParameter *DestVec = (UMaterialExpressionScalarParameter*)Destination;
-
-		DestVec->Modify();
-		DestVec->DefaultValue = SourceVec->DefaultValue;
-#if WITH_EDITOR
-		DestVec->Group = SourceVec->Group;
-		DestVec->SortPriority = SourceVec->SortPriority;
+		}
+		else if (Source->IsA(UMaterialExpressionRuntimeVirtualTextureSampleParameter::StaticClass()))
+		{
+			UMaterialExpressionRuntimeVirtualTextureSampleParameter* SourceRVT = (UMaterialExpressionRuntimeVirtualTextureSampleParameter*)Source;
+			UMaterialExpressionRuntimeVirtualTextureSampleParameter* DestRVT = (UMaterialExpressionRuntimeVirtualTextureSampleParameter*)Destination;
+			DestRVT->Group = SourceRVT->Group;
+#if WITH_EDITORONLY_DATA
+			DestRVT->SortPriority = SourceRVT->SortPriority;
 #endif
-	}
-	else if(Source->IsA(UMaterialExpressionFontSampleParameter::StaticClass()))
-	{
-		UMaterialExpressionFontSampleParameter *SourceFont = (UMaterialExpressionFontSampleParameter*)Source;
-		UMaterialExpressionFontSampleParameter *DestFont = (UMaterialExpressionFontSampleParameter*)Destination;
-
-		DestFont->Modify();
-		DestFont->Font = SourceFont->Font;
-		DestFont->FontTexturePage = SourceFont->FontTexturePage;
-		DestFont->Group = SourceFont->Group; // Group isn't editor only for fonts
+		}
+		else if (Source->IsA(UMaterialExpressionFontSampleParameter::StaticClass()))
+		{
+			UMaterialExpressionFontSampleParameter* SourceFont = (UMaterialExpressionFontSampleParameter*)Source;
+			UMaterialExpressionFontSampleParameter* DestFont = (UMaterialExpressionFontSampleParameter*)Destination;
+			DestFont->Group = SourceFont->Group; // Group isn't editor only for fonts
 #if WITH_EDITOR
-		DestFont->SortPriority = SourceFont->SortPriority;
+			DestFont->SortPriority = SourceFont->SortPriority;
 #endif
-	}
-	else
-	{
-		bRet = false;
-	}
-
+		}
+		else
+		{
+			ensureMsgf(false, TEXT("Unexpected expression parameter type %s"), *Source->GetClass()->GetName());
+		}
 
 #if WITH_EDITOR
-	if (bRet)
-	{
 		Destination->Desc = Source->Desc;
 		Destination->GraphNode->OnUpdateCommentText(Destination->Desc);
+#endif
+		bResult = true;
 	}
-#endif // WITH_EDITOR
 
-	return bRet;
+	return bResult;
 }
 
 void UMaterial::BeginDestroy()
diff --git a/Engine/Source/Runtime/Engine/Private/Materials/MaterialExpressions.cpp b/Engine/Source/Runtime/Engine/Private/Materials/MaterialExpressions.cpp
--- a/Engine/Source/Runtime/Engine/Private/Materials/MaterialExpressions.cpp
+++ b/Engine/Source/Runtime/Engine/Private/Materials/MaterialExpressions.cpp
@@ -278,6 +278,7 @@
 #endif //WITH_EDITOR
 #include "Materials/MaterialInstanceConstant.h"
 #include "Curves/CurveLinearColorAtlas.h"
+#include "Curves/CurveLinearColor.h"
 #include "Interfaces/ITargetPlatform.h"
 #include "Interfaces/ITargetPlatformManagerModule.h"
 
@@ -518,7 +519,14 @@ void ValidateParameterNameInternal(class UMaterialExpression* ExpressionToValida
 					// Name are unique per class type
 					if (Expression != ExpressionToValidate && Expression->GetParameterName() == PotentialName && Expression->GetClass() == ExpressionToValidate->GetClass())
 					{
-						ExpressionToValidate->SetValueToMatchingExpression(Expression);
+						FMaterialParameterMetadata Meta;
+						if (OwningMaterial->GetParameterValue(Expression->GetParameterType(), PotentialName, Meta, EMaterialGetParameterValueFlags::CheckAll))
+						{
+							const EMaterialExpressionSetParameterValueFlags Flags =
+								EMaterialExpressionSetParameterValueFlags::SendPostEditChangeProperty |
+								EMaterialExpressionSetParameterValueFlags::NoUpdateExpressionGuid;
+							verify(ExpressionToValidate->SetParameterValue(PotentialName, Meta, Flags));
+						}
 						break;
 					}
 				}
@@ -1477,6 +1485,16 @@ void UMaterialExpression::SetEditableName(const FString& NewName)
 	check(false);
 }
 
+EMaterialParameterType UMaterialExpression::GetParameterType() const
+{
+	FMaterialParameterMetadata Meta;
+	if (GetParameterValue(Meta))
+	{
+		return Meta.Value.Type;
+	}
+	return EMaterialParameterType::None;
+}
+
 void UMaterialExpression::ValidateParameterName(const bool bAllowDuplicateName)
 {
 	// Incrementing the name is now handled in UMaterialExpressionParameter::ValidateParameterName
@@ -2694,11 +2712,17 @@ void UMaterialExpressionRuntimeVirtualTextureSampleParameter::GetAllParameterInf
 
 #if WITH_EDITOR
 
-bool UMaterialExpressionRuntimeVirtualTextureSampleParameter::SetParameterValue(FName InParameterName, URuntimeVirtualTexture* InValue)
+bool UMaterialExpressionRuntimeVirtualTextureSampleParameter::SetParameterValue(FName InParameterName, URuntimeVirtualTexture* InValue, EMaterialExpressionSetParameterValueFlags Flags)
 {
 	if (InParameterName == ParameterName)
 	{
 		VirtualTexture = InValue;
+		if (EnumHasAnyFlags(Flags, EMaterialExpressionSetParameterValueFlags::SendPostEditChangeProperty))
+		{
+			FProperty* ParamProperty = FindFProperty<FProperty>(UMaterialExpressionRuntimeVirtualTextureSampleParameter::StaticClass(), GET_MEMBER_NAME_STRING_CHECKED(UMaterialExpressionRuntimeVirtualTextureSampleParameter, VirtualTexture));
+			FPropertyChangedEvent PropertyChangedEvent(ParamProperty);
+			PostEditChangeProperty(PropertyChangedEvent);
+		}
 		return true;
 	}
 
@@ -2735,19 +2759,6 @@ bool UMaterialExpressionRuntimeVirtualTextureSampleParameter::MatchesSearchQuery
 
 	return Super::MatchesSearchQuery(SearchQuery);
 }
-
-void UMaterialExpressionRuntimeVirtualTextureSampleParameter::SetValueToMatchingExpression(UMaterialExpression* OtherExpression)
-{
-	URuntimeVirtualTexture* Value = nullptr;
-	if (Material->GetRuntimeVirtualTextureParameterValue(FMaterialParameterInfo(OtherExpression->GetParameterName()), Value))
-	{
-		VirtualTexture = Value;
-		FProperty* ParamProperty = FindFProperty<FProperty>(UMaterialExpressionRuntimeVirtualTextureSampleParameter::StaticClass(), GET_MEMBER_NAME_STRING_CHECKED(UMaterialExpressionRuntimeVirtualTextureSampleParameter, VirtualTexture));
-		FPropertyChangedEvent PropertyChangedEvent(ParamProperty);
-		PostEditChangeProperty(PropertyChangedEvent);
-	}
-}
-
 #endif
 
 UMaterialExpressionRuntimeVirtualTextureReplace::UMaterialExpressionRuntimeVirtualTextureReplace(const FObjectInitializer& ObjectInitializer)
@@ -2958,18 +2969,6 @@ void UMaterialExpressionTextureSampleParameter::ValidateParameterName(const bool
 {
 	ValidateParameterNameInternal(this, Material, bAllowDuplicateName);
 }
-
-void UMaterialExpressionTextureSampleParameter::SetValueToMatchingExpression(UMaterialExpression* OtherExpression)
-{
-	UTexture* ExistingValue;
-	if (Material->GetTextureParameterValue(OtherExpression->GetParameterName(), ExistingValue))
-	{
-		Texture = ExistingValue;
-		FProperty* ParamProperty = FindFProperty<FProperty>(UMaterialExpressionTextureSampleParameter::StaticClass(), GET_MEMBER_NAME_STRING_CHECKED(UMaterialExpressionTextureSampleParameter, Texture));
-		FPropertyChangedEvent PropertyChangedEvent(ParamProperty);
-		PostEditChangeProperty(PropertyChangedEvent);
-	}
-}
 #endif // WITH_EDITOR
 
 bool UMaterialExpressionTextureSampleParameter::IsNamedParameter(const FHashedMaterialParameterInfo& ParameterInfo, UTexture*& OutValue) const
@@ -2984,11 +2983,17 @@ bool UMaterialExpressionTextureSampleParameter::IsNamedParameter(const FHashedMa
 }
 
 #if WITH_EDITOR
-bool UMaterialExpressionTextureSampleParameter::SetParameterValue(FName InParameterName, UTexture* InValue)
+bool UMaterialExpressionTextureSampleParameter::SetParameterValue(FName InParameterName, UTexture* InValue, EMaterialExpressionSetParameterValueFlags Flags)
 {
 	if (InParameterName == ParameterName)
 	{
 		Texture = InValue;
+		if (EnumHasAnyFlags(Flags, EMaterialExpressionSetParameterValueFlags::SendPostEditChangeProperty))
+		{
+			FProperty* ParamProperty = FindFProperty<FProperty>(UMaterialExpressionTextureSampleParameter::StaticClass(), GET_MEMBER_NAME_STRING_CHECKED(UMaterialExpressionTextureSampleParameter, Texture));
+			FPropertyChangedEvent PropertyChangedEvent(ParamProperty);
+			PostEditChangeProperty(PropertyChangedEvent);
+		}
 		return true;
 	}
 	return false;
@@ -4904,32 +4909,6 @@ void UMaterialExpressionStaticComponentMaskParameter::GetCaption(TArray<FString>
 	OutCaptions.Add(TEXT("Mask Param")); 
 	OutCaptions.Add(FString::Printf(TEXT("'%s'"), *ParameterName.ToString()));
 }
-
-void UMaterialExpressionStaticComponentMaskParameter::SetValueToMatchingExpression(UMaterialExpression* OtherExpression)
-{
-	bool R = false;
-	bool G = false;
-	bool B = false;
-	bool A = false;
-	FGuid Guid;
-	Material->GetStaticComponentMaskParameterValue(OtherExpression->GetParameterName(), R, G, B, A, Guid);
-	DefaultR = R;
-	DefaultG = G;
-	DefaultB = B;
-	DefaultA = A;
-	FProperty* ParamProperty = FindFProperty<FProperty>(UMaterialExpressionStaticComponentMaskParameter::StaticClass(), GET_MEMBER_NAME_STRING_CHECKED(UMaterialExpressionStaticComponentMaskParameter, DefaultR));
-	FPropertyChangedEvent RChangedEvent(ParamProperty);
-	PostEditChangeProperty(RChangedEvent);
-	ParamProperty = FindFProperty<FProperty>(UMaterialExpressionStaticComponentMaskParameter::StaticClass(), GET_MEMBER_NAME_STRING_CHECKED(UMaterialExpressionStaticComponentMaskParameter, DefaultG));
-	FPropertyChangedEvent GChangedEvent(ParamProperty);
-	PostEditChangeProperty(GChangedEvent);
-	ParamProperty = FindFProperty<FProperty>(UMaterialExpressionStaticComponentMaskParameter::StaticClass(), GET_MEMBER_NAME_STRING_CHECKED(UMaterialExpressionStaticComponentMaskParameter, DefaultB));
-	FPropertyChangedEvent BChangedEvent(ParamProperty);
-	PostEditChangeProperty(BChangedEvent);
-	ParamProperty = FindFProperty<FProperty>(UMaterialExpressionStaticComponentMaskParameter::StaticClass(), GET_MEMBER_NAME_STRING_CHECKED(UMaterialExpressionStaticComponentMaskParameter, DefaultA));
-	FPropertyChangedEvent AChangedEvent(ParamProperty);
-	PostEditChangeProperty(AChangedEvent);
-}
 #endif // WITH_EDITOR
 
 bool UMaterialExpressionStaticComponentMaskParameter::IsNamedParameter(const FHashedMaterialParameterInfo& ParameterInfo, bool& OutR, bool& OutG, bool& OutB, bool& OutA, FGuid&OutExpressionGuid) const
@@ -4948,7 +4927,7 @@ bool UMaterialExpressionStaticComponentMaskParameter::IsNamedParameter(const FHa
 }
 
 #if WITH_EDITOR
-bool  UMaterialExpressionStaticComponentMaskParameter::SetParameterValue(FName InParameterName, bool InR, bool InG, bool InB, bool InA, FGuid InExpressionGuid)
+bool  UMaterialExpressionStaticComponentMaskParameter::SetParameterValue(FName InParameterName, bool InR, bool InG, bool InB, bool InA, FGuid InExpressionGuid, EMaterialExpressionSetParameterValueFlags Flags)
 {
 	if (InParameterName == ParameterName)
 	{
@@ -4956,7 +4935,25 @@ bool  UMaterialExpressionStaticComponentMaskParameter::SetParameterValue(FName I
 		DefaultG = InG;
 		DefaultB = InB;
 		DefaultA = InA;
-		ExpressionGUID = InExpressionGuid;
+		if (!EnumHasAnyFlags(Flags, EMaterialExpressionSetParameterValueFlags::NoUpdateExpressionGuid))
+		{
+			ExpressionGUID = InExpressionGuid;
+		}
+		if (EnumHasAnyFlags(Flags, EMaterialExpressionSetParameterValueFlags::SendPostEditChangeProperty))
+		{
+			FProperty* ParamProperty = FindFProperty<FProperty>(UMaterialExpressionStaticComponentMaskParameter::StaticClass(), GET_MEMBER_NAME_STRING_CHECKED(UMaterialExpressionStaticComponentMaskParameter, DefaultR));
+			FPropertyChangedEvent RChangedEvent(ParamProperty);
+			PostEditChangeProperty(RChangedEvent);
+			ParamProperty = FindFProperty<FProperty>(UMaterialExpressionStaticComponentMaskParameter::StaticClass(), GET_MEMBER_NAME_STRING_CHECKED(UMaterialExpressionStaticComponentMaskParameter, DefaultG));
+			FPropertyChangedEvent GChangedEvent(ParamProperty);
+			PostEditChangeProperty(GChangedEvent);
+			ParamProperty = FindFProperty<FProperty>(UMaterialExpressionStaticComponentMaskParameter::StaticClass(), GET_MEMBER_NAME_STRING_CHECKED(UMaterialExpressionStaticComponentMaskParameter, DefaultB));
+			FPropertyChangedEvent BChangedEvent(ParamProperty);
+			PostEditChangeProperty(BChangedEvent);
+			ParamProperty = FindFProperty<FProperty>(UMaterialExpressionStaticComponentMaskParameter::StaticClass(), GET_MEMBER_NAME_STRING_CHECKED(UMaterialExpressionStaticComponentMaskParameter, DefaultA));
+			FPropertyChangedEvent AChangedEvent(ParamProperty);
+			PostEditChangeProperty(AChangedEvent);
+		}
 		return true;
 	}
 
@@ -7860,11 +7857,17 @@ void UMaterialExpressionVectorParameter::GetAllParameterInfo(TArray<FMaterialPar
 }
 
 #if WITH_EDITOR
-bool UMaterialExpressionVectorParameter::SetParameterValue(FName InParameterName, FLinearColor InValue)
+bool UMaterialExpressionVectorParameter::SetParameterValue(FName InParameterName, FLinearColor InValue, EMaterialExpressionSetParameterValueFlags Flags)
 {
 	if (InParameterName == ParameterName)
 	{
 		DefaultValue = InValue;
+		if (EnumHasAnyFlags(Flags, EMaterialExpressionSetParameterValueFlags::SendPostEditChangeProperty))
+		{
+			FProperty* ParamProperty = FindFProperty<FProperty>(UMaterialExpressionVectorParameter::StaticClass(), GET_MEMBER_NAME_STRING_CHECKED(UMaterialExpressionVectorParameter, DefaultValue));
+			FPropertyChangedEvent PropertyChangedEvent(ParamProperty);
+			PostEditChangeProperty(PropertyChangedEvent);
+		}
 		return true;
 	}
 	return false;
@@ -7951,16 +7954,6 @@ bool UMaterialExpressionVectorParameter::HasClassAndNameCollision(UMaterialExpre
 	}
 	return Super::HasClassAndNameCollision(OtherExpression);
 }
-
-void UMaterialExpressionVectorParameter::SetValueToMatchingExpression(UMaterialExpression* OtherExpression)
-{
-	FLinearColor ExistingValue = FLinearColor::Transparent;
-	Material->GetVectorParameterValue(FMaterialParameterInfo(OtherExpression->GetParameterName()), ExistingValue);
-	DefaultValue = ExistingValue;
-	FProperty* ParamProperty = FindFProperty<FProperty>(UMaterialExpressionVectorParameter::StaticClass(), GET_MEMBER_NAME_STRING_CHECKED(UMaterialExpressionVectorParameter, DefaultValue));
-	FPropertyChangedEvent PropertyChangedEvent(ParamProperty);
-	PostEditChangeProperty(PropertyChangedEvent);
-}
 #endif
 
 //
@@ -7992,7 +7985,7 @@ UMaterialExpressionChannelMaskParameter::UMaterialExpressionChannelMaskParameter
 }
 
 #if WITH_EDITOR
-bool UMaterialExpressionChannelMaskParameter::SetParameterValue(FName InParameterName, FLinearColor InValue)
+bool UMaterialExpressionChannelMaskParameter::SetParameterValue(FName InParameterName, FLinearColor InValue, EMaterialExpressionSetParameterValueFlags Flags)
 {
 	if (InParameterName == ParameterName)
 	{
@@ -8017,6 +8010,20 @@ bool UMaterialExpressionChannelMaskParameter::SetParameterValue(FName InParamete
 			MaskChannel = EChannelMaskParameterColor::Alpha;
 		}
 
+		if (EnumHasAnyFlags(Flags, EMaterialExpressionSetParameterValueFlags::SendPostEditChangeProperty))
+		{
+			{
+				FProperty* ParamProperty = FindFProperty<FProperty>(UMaterialExpressionVectorParameter::StaticClass(), GET_MEMBER_NAME_STRING_CHECKED(UMaterialExpressionVectorParameter, DefaultValue));
+				FPropertyChangedEvent PropertyChangedEvent(ParamProperty);
+				PostEditChangeProperty(PropertyChangedEvent);
+			}
+			{
+				FProperty* ParamProperty = FindFProperty<FProperty>(UMaterialExpressionChannelMaskParameter::StaticClass(), GET_MEMBER_NAME_STRING_CHECKED(UMaterialExpressionChannelMaskParameter, MaskChannel));
+				FPropertyChangedEvent PropertyChangedEvent(ParamProperty);
+				PostEditChangeProperty(PropertyChangedEvent);
+			}
+		}
+
 		// Update caption and preview
 		if (GraphNode)
 		{
@@ -8182,11 +8189,17 @@ void UMaterialExpressionScalarParameter::GetAllParameterInfo(TArray<FMaterialPar
 }
 
 #if WITH_EDITOR
-bool UMaterialExpressionScalarParameter::SetParameterValue(FName InParameterName, float InValue)
+bool UMaterialExpressionScalarParameter::SetParameterValue(FName InParameterName, float InValue, EMaterialExpressionSetParameterValueFlags Flags)
 {
 	if (InParameterName == ParameterName)
 	{
 		DefaultValue = InValue;
+		if (EnumHasAnyFlags(Flags, EMaterialExpressionSetParameterValueFlags::SendPostEditChangeProperty))
+		{
+			FProperty* ParamProperty = FindFProperty<FProperty>(UMaterialExpressionScalarParameter::StaticClass(), GET_MEMBER_NAME_STRING_CHECKED(UMaterialExpressionScalarParameter, DefaultValue));
+			FPropertyChangedEvent PropertyChangedEvent(ParamProperty);
+			PostEditChangeProperty(PropertyChangedEvent);
+		}
 		return true;
 	}
 	return false;
@@ -8252,17 +8265,6 @@ bool UMaterialExpressionScalarParameter::HasClassAndNameCollision(UMaterialExpre
 	}
 	return Super::HasClassAndNameCollision(OtherExpression);
 }
-
-void UMaterialExpressionScalarParameter::SetValueToMatchingExpression(UMaterialExpression* OtherExpression)
-{
-	float ExistingValue = 0.0f;
-	Material->GetScalarParameterDefaultValue(FMaterialParameterInfo(OtherExpression->GetParameterName()), ExistingValue);
-	DefaultValue = ExistingValue;
-	FProperty* ParamProperty = FindFProperty<FProperty>(UMaterialExpressionScalarParameter::StaticClass(), GET_MEMBER_NAME_STRING_CHECKED(UMaterialExpressionScalarParameter, DefaultValue));
-	FPropertyChangedEvent PropertyChangedEvent(ParamProperty);
-	PostEditChangeProperty(PropertyChangedEvent);
-}
-
 #endif
 
 //
@@ -8370,17 +8372,6 @@ void UMaterialExpressionStaticBoolParameter::GetCaption(TArray<FString>& OutCapt
 	OutCaptions.Add(FString::Printf(TEXT("Static Bool Param (%s)"), (DefaultValue ? TEXT("True") : TEXT("False")))); 
 	OutCaptions.Add(FString::Printf(TEXT("'%s'"), *ParameterName.ToString())); 
 }
-
-void UMaterialExpressionStaticBoolParameter::SetValueToMatchingExpression(UMaterialExpression* OtherExpression)
-{
-	bool ExistingValue = false;
-	FGuid Guid;
-	Material->GetStaticSwitchParameterValue(OtherExpression->GetParameterName(), ExistingValue, Guid);
-	DefaultValue = ExistingValue;
-	FProperty* ParamProperty = FindFProperty<FProperty>(UMaterialExpressionStaticBoolParameter::StaticClass(), GET_MEMBER_NAME_STRING_CHECKED(UMaterialExpressionStaticBoolParameter, DefaultValue));
-	FPropertyChangedEvent PropertyChangedEvent(ParamProperty);
-	PostEditChangeProperty(PropertyChangedEvent);
-}
 #endif // WITH_EDITOR
 
 bool UMaterialExpressionStaticBoolParameter::IsNamedParameter(const FHashedMaterialParameterInfo& ParameterInfo, bool& OutValue, FGuid& OutExpressionGuid) const
@@ -8396,12 +8387,21 @@ bool UMaterialExpressionStaticBoolParameter::IsNamedParameter(const FHashedMater
 }
 
 #if WITH_EDITOR
-bool UMaterialExpressionStaticBoolParameter::SetParameterValue(FName InParameterName, bool InValue, FGuid InExpressionGuid)
+bool UMaterialExpressionStaticBoolParameter::SetParameterValue(FName InParameterName, bool InValue, FGuid InExpressionGuid, EMaterialExpressionSetParameterValueFlags Flags)
 {
 	if (InParameterName == ParameterName)
 	{
 		DefaultValue = InValue;
-		ExpressionGUID = InExpressionGuid;
+		if (!EnumHasAnyFlags(Flags, EMaterialExpressionSetParameterValueFlags::NoUpdateExpressionGuid))
+		{
+			ExpressionGUID = InExpressionGuid;
+		}
+		if (EnumHasAnyFlags(Flags, EMaterialExpressionSetParameterValueFlags::SendPostEditChangeProperty))
+		{
+			FProperty* ParamProperty = FindFProperty<FProperty>(UMaterialExpressionStaticBoolParameter::StaticClass(), GET_MEMBER_NAME_STRING_CHECKED(UMaterialExpressionStaticBoolParameter, DefaultValue));
+			FPropertyChangedEvent PropertyChangedEvent(ParamProperty);
+			PostEditChangeProperty(PropertyChangedEvent);
+		}
 		return true;
 	}
 
@@ -11375,23 +11375,6 @@ void UMaterialExpressionFontSampleParameter::ValidateParameterName(const bool bA
 {
 	ValidateParameterNameInternal(this, Material, bAllowDuplicateName);
 }
-
-void UMaterialExpressionFontSampleParameter::SetValueToMatchingExpression(UMaterialExpression* OtherExpression)
-{
-	UFont* FontValue;
-	int32 FontPage;
-	if (Material->GetFontParameterValue(FMaterialParameterInfo(OtherExpression->GetParameterName()), FontValue, FontPage))
-	{
-		Font = FontValue;
-		FontTexturePage = FontPage;
-		FProperty* ParamProperty = FindFProperty<FProperty>(UMaterialExpressionFontSampleParameter::StaticClass(), GET_MEMBER_NAME_STRING_CHECKED(UMaterialExpressionFontSampleParameter, Font));
-		FPropertyChangedEvent PropertyChangedEvent(ParamProperty);
-		PostEditChangeProperty(PropertyChangedEvent);
-		ParamProperty = FindFProperty<FProperty>(UMaterialExpressionFontSampleParameter::StaticClass(), GET_MEMBER_NAME_STRING_CHECKED(UMaterialExpressionFontSampleParameter, FontTexturePage));
-		FPropertyChangedEvent PageChangedEvent(ParamProperty);
-		PostEditChangeProperty(PageChangedEvent);
-	}
-}
 #endif // WITH_EDITOR
 
 bool UMaterialExpressionFontSampleParameter::IsNamedParameter(const FHashedMaterialParameterInfo& ParameterInfo, UFont*& OutFontValue, int32& OutFontPage) const
@@ -11407,12 +11390,21 @@ bool UMaterialExpressionFontSampleParameter::IsNamedParameter(const FHashedMater
 }
 
 #if WITH_EDITOR
-bool UMaterialExpressionFontSampleParameter::SetParameterValue(FName InParameterName, UFont* InFontValue, int32 InFontPage)
+bool UMaterialExpressionFontSampleParameter::SetParameterValue(FName InParameterName, UFont* InFontValue, int32 InFontPage, EMaterialExpressionSetParameterValueFlags Flags)
 {
 	if (InParameterName == ParameterName)
 	{
 		Font = InFontValue;
 		FontTexturePage = InFontPage;
+		if (EnumHasAnyFlags(Flags, EMaterialExpressionSetParameterValueFlags::SendPostEditChangeProperty))
+		{
+			FProperty* ParamProperty = FindFProperty<FProperty>(UMaterialExpressionFontSampleParameter::StaticClass(), GET_MEMBER_NAME_STRING_CHECKED(UMaterialExpressionFontSampleParameter, Font));
+			FPropertyChangedEvent PropertyChangedEvent(ParamProperty);
+			PostEditChangeProperty(PropertyChangedEvent);
+			ParamProperty = FindFProperty<FProperty>(UMaterialExpressionFontSampleParameter::StaticClass(), GET_MEMBER_NAME_STRING_CHECKED(UMaterialExpressionFontSampleParameter, FontTexturePage));
+			FPropertyChangedEvent PageChangedEvent(ParamProperty);
+			PostEditChangeProperty(PageChangedEvent);
+		}
 		return true;
 	}
 
@@ -19105,7 +19097,6 @@ void UMaterialExpressionCurveAtlasRowParameter::GetTexturesForceMaterialRecompil
 	}
 }
 
-
 int32 UMaterialExpressionCurveAtlasRowParameter::Compile(class FMaterialCompiler* Compiler, int32 OutputIndex)
 {
 	// Some error checking. Note that when bUseCustomPrimitiveData is true we don't rely on the Curve at all
@@ -19177,6 +19168,40 @@ int32 UMaterialExpressionCurveAtlasRowParameter::Compile(class FMaterialCompiler
 	}
 }
 
+bool UMaterialExpressionCurveAtlasRowParameter::GetParameterValue(FMaterialParameterMetadata& OutMeta) const
+{
+	float Value = 0.0f;
+	if (Atlas && Curve)
+	{
+		const int32 SlotIndex = Atlas->GradientCurves.Find(Curve);
+		if (SlotIndex != INDEX_NONE)
+		{
+			Value = (float)SlotIndex;
+		}
+		OutMeta.bUsedAsAtlasPosition = true;
+		OutMeta.ScalarCurve = Curve;
+		OutMeta.ScalarAtlas = Atlas;
+	}
+
+	OutMeta.Value = Value;
+	OutMeta.ExpressionGuid = ExpressionGUID;
+	return true;
+}
+
+bool UMaterialExpressionCurveAtlasRowParameter::SetParameterValue(const FName& Name, const FMaterialParameterMetadata& Meta, EMaterialExpressionSetParameterValueFlags Flags)
+{
+	if (Name == ParameterName &&
+		Meta.Value.Type == EMaterialParameterType::Scalar &&
+		Meta.bUsedAsAtlasPosition)
+	{
+		Curve = Meta.ScalarCurve.Get();
+		Atlas = Meta.ScalarAtlas.Get();
+		DefaultValue = Meta.Value.AsScalar();
+		return true;
+	}
+	return false;
+}
+
 void UMaterialExpressionCurveAtlasRowParameter::PostEditChangeProperty(FPropertyChangedEvent& PropertyChangedEvent)
 {
 	if (PropertyChangedEvent.Property 
@@ -19189,12 +19214,11 @@ void UMaterialExpressionCurveAtlasRowParameter::PostEditChangeProperty(FProperty
 		}
 		if (SlotIndex != INDEX_NONE)
 		{
-			float NewValue = (float)SlotIndex;
-			SetParameterValue(ParameterName, NewValue);
+			DefaultValue = (float)SlotIndex;
 		}
 		else
 		{
-			SetParameterValue(ParameterName, 0.0f);
+			DefaultValue = 0.0f;
 		}
 	}
 
