Move some functions to editor-only, fix non-editor compile

diff --git a/Engine/Source/Runtime/Engine/Private/Materials/Material.cpp b/Engine/Source/Runtime/Engine/Private/Materials/Material.cpp
--- a/Engine/Source/Runtime/Engine/Private/Materials/Material.cpp
+++ b/Engine/Source/Runtime/Engine/Private/Materials/Material.cpp
@@ -4439,7 +4439,6 @@ void UMaterial::RebuildShadingModelField()
 	UsedShadingModels = GetShadingModelFieldString(ShadingModels, FShadingModelToStringDelegate::CreateLambda(ShadingModelToStringLambda), " | ");
 #endif
 }
-#endif // WITH_EDITOR
 
 bool UMaterial::GetExpressionParameterName(const UMaterialExpression* Expression, FName& OutName)
 {
@@ -4452,7 +4451,6 @@ bool UMaterial::GetExpressionParameterName(const UMaterialExpression* Expression
 	return bRet;
 }
 
-
 bool UMaterial::CopyExpressionParameters(UMaterialExpression* Source, UMaterialExpression* Destination)
 {
 	if(!Source || !Destination || Source == Destination || Source->GetClass() != Destination->GetClass())
@@ -4472,52 +4470,43 @@ bool UMaterial::CopyExpressionParameters(UMaterialExpression* Source, UMaterialE
 		{
 			UMaterialExpressionParameter* SourceParam = (UMaterialExpressionParameter*)Source;
 			UMaterialExpressionParameter* DestParam = (UMaterialExpressionParameter*)Destination;
-#if WITH_EDITOR
 			DestParam->Group = SourceParam->Group;
 			DestParam->SortPriority = SourceParam->SortPriority;
-#endif
 		}
 		else if (Source->IsA(UMaterialExpressionTextureSampleParameter::StaticClass()))
 		{
 			UMaterialExpressionTextureSampleParameter* SourceTex = (UMaterialExpressionTextureSampleParameter*)Source;
 			UMaterialExpressionTextureSampleParameter* DestTex = (UMaterialExpressionTextureSampleParameter*)Destination;
-			DestTex->Group = SourceTex->Group; // Group isn't editor only for textures
-#if WITH_EDITOR
+			DestTex->Group = SourceTex->Group;
 			DestTex->SortPriority = SourceTex->SortPriority;
-#endif
 		}
 		else if (Source->IsA(UMaterialExpressionRuntimeVirtualTextureSampleParameter::StaticClass()))
 		{
 			UMaterialExpressionRuntimeVirtualTextureSampleParameter* SourceRVT = (UMaterialExpressionRuntimeVirtualTextureSampleParameter*)Source;
 			UMaterialExpressionRuntimeVirtualTextureSampleParameter* DestRVT = (UMaterialExpressionRuntimeVirtualTextureSampleParameter*)Destination;
 			DestRVT->Group = SourceRVT->Group;
-#if WITH_EDITORONLY_DATA
 			DestRVT->SortPriority = SourceRVT->SortPriority;
-#endif
 		}
 		else if (Source->IsA(UMaterialExpressionFontSampleParameter::StaticClass()))
 		{
 			UMaterialExpressionFontSampleParameter* SourceFont = (UMaterialExpressionFontSampleParameter*)Source;
 			UMaterialExpressionFontSampleParameter* DestFont = (UMaterialExpressionFontSampleParameter*)Destination;
 			DestFont->Group = SourceFont->Group; // Group isn't editor only for fonts
-#if WITH_EDITOR
 			DestFont->SortPriority = SourceFont->SortPriority;
-#endif
 		}
 		else
 		{
 			ensureMsgf(false, TEXT("Unexpected expression parameter type %s"), *Source->GetClass()->GetName());
 		}
 
-#if WITH_EDITOR
 		Destination->Desc = Source->Desc;
 		Destination->GraphNode->OnUpdateCommentText(Destination->Desc);
-#endif
 		bResult = true;
 	}
 
 	return bResult;
 }
+#endif // WITH_EDITOR
 
 void UMaterial::BeginDestroy()
 {
