Fixing issue where jump peak substepping does not respect root motion

diff --git a/Engine/Source/Runtime/Engine/Private/Components/CharacterMovementComponent.cpp b/Engine/Source/Runtime/Engine/Private/Components/CharacterMovementComponent.cpp
--- a/Engine/Source/Runtime/Engine/Private/Components/CharacterMovementComponent.cpp
+++ b/Engine/Source/Runtime/Engine/Private/Components/CharacterMovementComponent.cpp
@@ -4359,34 +4359,43 @@ void UCharacterMovementComponent::PhysFalling(float deltaTime, int32 Iterations)
 		// Apply gravity
 		Velocity = NewFallVelocity(Velocity, Gravity, GravityTime);
 
+		//UE_LOG(LogCharacterMovement, Log, TEXT("dt=(%.6f) OldLocation=(%s) OldVelocity=(%s) OldVelocityWithRootMotion=(%s) NewVelocity=(%s)"), timeTick, *(UpdatedComponent->GetComponentLocation()).ToString(), *OldVelocity.ToString(), *OldVelocityWithRootMotion.ToString(), *Velocity.ToString());
+		ApplyRootMotionToVelocity(timeTick);
+
 		// See if we need to sub-step to exactly reach the apex. This is important for avoiding "cutting off the top" of the trajectory as framerate varies.
-		if (CharacterMovementCVars::ForceJumpPeakSubstep && OldVelocity.Z > 0.f && Velocity.Z <= 0.f && NumJumpApexAttempts < MaxJumpApexAttemptsPerSimulation)
+		if (CharacterMovementCVars::ForceJumpPeakSubstep && OldVelocityWithRootMotion.Z > 0.f && Velocity.Z <= 0.f && NumJumpApexAttempts < MaxJumpApexAttemptsPerSimulation)
 		{
-			const FVector DerivedAccel = (Velocity - OldVelocity) / timeTick;
+			const FVector DerivedAccel = (Velocity - OldVelocityWithRootMotion) / timeTick;
 			if (!FMath::IsNearlyZero(DerivedAccel.Z))
 			{
-				const float TimeToApex = -OldVelocity.Z / DerivedAccel.Z;
+				const float TimeToApex = -OldVelocityWithRootMotion.Z / DerivedAccel.Z;
 				
 				// The time-to-apex calculation should be precise, and we want to avoid adding a substep when we are basically already at the apex from the previous iteration's work.
 				const float ApexTimeMinimum = 0.0001f;
 				if (TimeToApex >= ApexTimeMinimum && TimeToApex < timeTick)
 				{
-					const FVector ApexVelocity = OldVelocity + DerivedAccel * TimeToApex;
+					const FVector ApexVelocity = OldVelocityWithRootMotion + (DerivedAccel * TimeToApex);
 					Velocity = ApexVelocity;
 					Velocity.Z = 0.f; // Should be nearly zero anyway, but this makes apex notifications consistent.
 
 					// We only want to move the amount of time it takes to reach the apex, and refund the unused time for next iteration.
-					remainingTime += (timeTick - TimeToApex);
+					const float TimeToRefund = (timeTick - TimeToApex);
+
+					remainingTime += TimeToRefund;
 					timeTick = TimeToApex;
 					Iterations--;
 					NumJumpApexAttempts++;
+
+					// Refund time to any active Root Motion Sources as well
+					for (TSharedPtr<FRootMotionSource> RootMotionSource : CurrentRootMotion.RootMotionSources)
+					{
+						const float RewoundRMSTime = FMath::Max(0.0f, RootMotionSource->GetTime() - TimeToRefund);
+						RootMotionSource->SetTime(RewoundRMSTime);
+					}
 				}
 			}
 		}
 
-		//UE_LOG(LogCharacterMovement, Log, TEXT("dt=(%.6f) OldLocation=(%s) OldVelocity=(%s) OldVelocityWithRootMotion=(%s) NewVelocity=(%s)"), timeTick, *(UpdatedComponent->GetComponentLocation()).ToString(), *OldVelocity.ToString(), *OldVelocityWithRootMotion.ToString(), *Velocity.ToString());
-		ApplyRootMotionToVelocity(timeTick);
-
 		if (bNotifyApex && (Velocity.Z < 0.f))
 		{
 			// Just passed jump apex since now going down
