Remove Ensure on actor within MoveExternalActorsToLevel and handle invalid actors
Actors can fail to load in FWorldPartitionLevelHelper::LoadActors

diff --git a/Engine/Source/Runtime/Engine/Private/WorldPartition/WorldPartitionLevelHelper.cpp b/Engine/Source/Runtime/Engine/Private/WorldPartition/WorldPartitionLevelHelper.cpp
--- a/Engine/Source/Runtime/Engine/Private/WorldPartition/WorldPartitionLevelHelper.cpp
+++ b/Engine/Source/Runtime/Engine/Private/WorldPartition/WorldPartitionLevelHelper.cpp
@@ -63,7 +63,7 @@ void FWorldPartitionLevelHelper::MoveExternalActorsToLevel(const TArray<FWorldPa
 		}
 
 		AActor* Actor = FindObject<AActor>(nullptr, *PackageObjectMapping.LoadedPath.ToString());
-		if (ensure(Actor))
+		if (Actor)
 		{
 			UPackage* ActorExternalPackage = Actor->IsPackageExternal() ? Actor->GetPackage() : nullptr;
 
