Extend FMaterialParameterMetadata with Group, SortPriority, and Description.  Cache this data inside FMaterialCachedData.
UMaterialInterface methods to query this data become non-virtual, instead it's just returned along with other Metadata when accessing the parameter value
Remove parameter group caching mechanism inside PreviewMaterial.cpp; no longer needed since group lookup now comes "for free" along with other relevant parameter data

diff --git a/Engine/Source/Editor/UnrealEd/Classes/MaterialEditor/DEditorParameterValue.h b/Engine/Source/Editor/UnrealEd/Classes/MaterialEditor/DEditorParameterValue.h
--- a/Engine/Source/Editor/UnrealEd/Classes/MaterialEditor/DEditorParameterValue.h
+++ b/Engine/Source/Editor/UnrealEd/Classes/MaterialEditor/DEditorParameterValue.h
@@ -42,6 +42,7 @@ class UNREALED_API UDEditorParameterValue : public UObject
 	virtual bool GetValue(FMaterialParameterMetadata& OutResult) const
 	{
 		OutResult.ExpressionGuid = ExpressionId;
+		OutResult.SortPriority = SortPriority;
 		OutResult.bOverride = bOverride;
 		return false;
 	}
diff --git a/Engine/Source/Editor/UnrealEd/Classes/MaterialEditor/DEditorTextureParameterValue.h b/Engine/Source/Editor/UnrealEd/Classes/MaterialEditor/DEditorTextureParameterValue.h
--- a/Engine/Source/Editor/UnrealEd/Classes/MaterialEditor/DEditorTextureParameterValue.h
+++ b/Engine/Source/Editor/UnrealEd/Classes/MaterialEditor/DEditorTextureParameterValue.h
@@ -25,10 +25,7 @@ class UNREALED_API UDEditorTextureParameterValue : public UDEditorParameterValue
 	{
 		UDEditorParameterValue::GetValue(OutResult);
 		OutResult.Value = ParameterValue;
-		OutResult.ChannelName[0] = ChannelNames.R;
-		OutResult.ChannelName[1] = ChannelNames.G;
-		OutResult.ChannelName[2] = ChannelNames.B;
-		OutResult.ChannelName[3] = ChannelNames.A;
+		OutResult.ChannelNames = ChannelNames;
 		return true;
 	}
 
diff --git a/Engine/Source/Editor/UnrealEd/Classes/MaterialEditor/DEditorVectorParameterValue.h b/Engine/Source/Editor/UnrealEd/Classes/MaterialEditor/DEditorVectorParameterValue.h
--- a/Engine/Source/Editor/UnrealEd/Classes/MaterialEditor/DEditorVectorParameterValue.h
+++ b/Engine/Source/Editor/UnrealEd/Classes/MaterialEditor/DEditorVectorParameterValue.h
@@ -29,10 +29,7 @@ class UNREALED_API UDEditorVectorParameterValue : public UDEditorParameterValue
 		UDEditorParameterValue::GetValue(OutResult);
 		OutResult.Value = ParameterValue;
 		OutResult.bUsedAsChannelMask = bIsUsedAsChannelMask;
-		OutResult.ChannelName[0] = ChannelNames.R;
-		OutResult.ChannelName[1] = ChannelNames.G;
-		OutResult.ChannelName[2] = ChannelNames.B;
-		OutResult.ChannelName[3] = ChannelNames.A;
+		OutResult.ChannelNames = ChannelNames;
 		return true;
 	}
 
diff --git a/Engine/Source/Editor/UnrealEd/Classes/MaterialEditor/MaterialEditorInstanceConstant.h b/Engine/Source/Editor/UnrealEd/Classes/MaterialEditor/MaterialEditorInstanceConstant.h
--- a/Engine/Source/Editor/UnrealEd/Classes/MaterialEditor/MaterialEditorInstanceConstant.h
+++ b/Engine/Source/Editor/UnrealEd/Classes/MaterialEditor/MaterialEditorInstanceConstant.h
@@ -319,11 +319,10 @@ class UNREALED_API UMaterialEditorInstanceConstant : public UObject
 	/** 
 	 *  Creates/adds value to group retrieved from parent material . 
 	 *
-	 * @param ParentMaterial		Name of material to search for groups.
 	 * @param ParameterValue		Current data to be grouped
-	 * @param OptionalGroupName		Optional Group Name that be used directly instead of resolving it from the material
+	 * @param GroupName				Name of the group
 	 */
-	void AssignParameterToGroup(UMaterial* ParentMaterial, UDEditorParameterValue* ParameterValue, const FName* OptionalGroupName = nullptr);
+	void AssignParameterToGroup(UDEditorParameterValue* ParameterValue, const FName& GroupName);
 
 	static FName GlobalGroupPrefix;
 
diff --git a/Engine/Source/Editor/UnrealEd/Classes/MaterialEditor/MaterialEditorPreviewParameters.h b/Engine/Source/Editor/UnrealEd/Classes/MaterialEditor/MaterialEditorPreviewParameters.h
--- a/Engine/Source/Editor/UnrealEd/Classes/MaterialEditor/MaterialEditorPreviewParameters.h
+++ b/Engine/Source/Editor/UnrealEd/Classes/MaterialEditor/MaterialEditorPreviewParameters.h
@@ -66,12 +66,10 @@ protected:
 	/**
 	*  Creates/adds value to group retrieved from parent material .
 	*
-	* @param ParentMaterial		Name of material to search for groups.
 	* @param ParameterValue		Current data to be grouped
-	* @param OptionalGroupName	Optional Group Name that be used directly instead of resolving it from the material
-
+	* @param GroupName			Name of the group
 	*/
-	void AssignParameterToGroup(UMaterial* ParentMaterial, UDEditorParameterValue* ParameterValue, FName* OptionalGroupName = nullptr);
+	void AssignParameterToGroup(UDEditorParameterValue* ParameterValue, const FName& GroupName);
 
 	static FName GlobalGroupPrefix;
 };
diff --git a/Engine/Source/Editor/UnrealEd/Private/PreviewMaterial.cpp b/Engine/Source/Editor/UnrealEd/Private/PreviewMaterial.cpp
--- a/Engine/Source/Editor/UnrealEd/Private/PreviewMaterial.cpp
+++ b/Engine/Source/Editor/UnrealEd/Private/PreviewMaterial.cpp
@@ -239,201 +239,6 @@ FMaterialResource* UPreviewMaterial::AllocateResource()
 	return new FPreviewMaterial();
 }
 
-// Helper struct to cache data for UMaterialEditorInstanceConstant/UMaterialEditorPreviewParameters::RegenerateArrays()
-struct FMaterialParamExpressionData
-{
-	FName Name = NAME_None;
-	FName Group = NAME_None;
-	UClass* ParamType = nullptr;
-	int32 SortPriority = 32;
-};
-
-// Helper struct to cache data for UMaterialEditorInstanceConstant/UMaterialEditorPreviewParameters::RegenerateArrays()
-struct FMaterialExpressionParameterDataCache
-{
-	TMap<FName, FMaterialParamExpressionData> GlobalParameters;
-	FName LayerParameterName = NAME_None;
-	TArray<TMap<FName, FMaterialParamExpressionData>> LayerParameters;
-	TArray<TMap<FName, FMaterialParamExpressionData>> BlendParameters;
-};
-
-// Helper function for UMaterialEditorInstanceConstant/UMaterialEditorPreviewParameters::RegenerateArrays()
-// Cache material expression parameter for group and sort priority for quick lookup while creating UDEditorParameterValue
-FMaterialExpressionParameterDataCache CacheMaterialExpressionParameterData(const UMaterial* InBaseMaterial, const FStaticParameterSet& InStaticParameters)
-{
-	FMaterialExpressionParameterDataCache ParamCache;
-	ParamCache.GlobalParameters.Reserve(InBaseMaterial->Expressions.Num());
-
-	// Function replicating UMaterialFunctionInterface::GetParameterGroupName & UMaterialFunctionInterface::GetParameterSortPriority behavior 
-	// but caching all the data in one pass
-	auto CacheMaterialFunctionParameterData = [](UMaterialFunctionInterface* ParameterFunction, TMap<FName, FMaterialParamExpressionData>& ParamDatas)
-	{
-		if (ParameterFunction)
-		{
-			TArray<UMaterialFunctionInterface*> Functions;
-			ParameterFunction->GetDependentFunctions(Functions);
-			Functions.AddUnique(ParameterFunction);
-
-			for (UMaterialFunctionInterface* Function : Functions)
-			{
-				for (UMaterialExpression* FunctionExpression : *Function->GetFunctionExpressions())
-				{
-					if (const UMaterialExpressionParameter* Parameter = Cast<const UMaterialExpressionParameter>(FunctionExpression))
-					{
-						FMaterialParamExpressionData ParamData;
-						ParamData.ParamType = UMaterialExpressionParameter::StaticClass();
-
-						ParamData.Name = Parameter->ParameterName;
-						ParamData.SortPriority = Parameter->SortPriority;
-						ParamData.Group = Parameter->Group;
-
-						//ensure(!ParamDatas.Contains(ParamData.Name));
-						ParamDatas.Add(ParamData.Name, ParamData);
-					}
-					else if (const UMaterialExpressionTextureSampleParameter* TexParameter = Cast<const UMaterialExpressionTextureSampleParameter>(FunctionExpression))
-					{
-						FMaterialParamExpressionData ParamData;
-						ParamData.ParamType = UMaterialExpressionTextureSampleParameter::StaticClass();
-
-						ParamData.Name = TexParameter->ParameterName;
-						ParamData.SortPriority = TexParameter->SortPriority;
-						ParamData.Group = TexParameter->Group;
-
-						//ensure(!ParamDatas.Contains(ParamData.Name));
-						ParamDatas.Add(ParamData.Name, ParamData);
-					}
-					else if (const UMaterialExpressionFontSampleParameter* FontParameter = Cast<const UMaterialExpressionFontSampleParameter>(FunctionExpression))
-					{
-						FMaterialParamExpressionData ParamData;
-						ParamData.ParamType = UMaterialExpressionFontSampleParameter::StaticClass();
-
-						ParamData.Name = FontParameter->ParameterName;
-						ParamData.SortPriority = FontParameter->SortPriority;
-						ParamData.Group = FontParameter->Group;
-
-						//ensure(!ParamDatas.Contains(ParamData.Name));
-						ParamDatas.Add(ParamData.Name, ParamData);
-					}
-				}
-			}
-		}
-	};
-
-	for (int32 Index = 0; Index < InBaseMaterial->Expressions.Num(); ++Index)
-	{
-		UMaterialExpression* Expression = InBaseMaterial->Expressions[Index];
-
-		if (UMaterialExpressionParameter* Parameter = Cast<UMaterialExpressionParameter>(Expression))
-		{
-			FMaterialParamExpressionData ParamData;
-			ParamData.ParamType = UMaterialExpressionParameter::StaticClass();
-
-			ParamData.Name = Parameter->GetParameterName();
-			ParamData.SortPriority = Parameter->SortPriority;
-			ParamData.Group = Parameter->Group;
-
-			//ensure(ParamCache.GlobalParameters.Contains(ParamData.Name));
-			ParamCache.GlobalParameters.Add(ParamData.Name, ParamData);
-		}
-		else if (UMaterialExpressionTextureSampleParameter* TexParameter = Cast<UMaterialExpressionTextureSampleParameter>(Expression))
-		{
-			FMaterialParamExpressionData ParamData;
-			ParamData.ParamType = UMaterialExpressionTextureSampleParameter::StaticClass();
-
-			ParamData.Name = TexParameter->GetParameterName();
-			ParamData.SortPriority = TexParameter->SortPriority;
-			ParamData.Group = TexParameter->Group;
-
-			//ensure(!ParamCache.GlobalParameters.Contains(ParamData.Name));
-			ParamCache.GlobalParameters.Add(ParamData.Name, ParamData);
-		}
-		else if (UMaterialExpressionRuntimeVirtualTextureSampleParameter* VTTexParameter = Cast<UMaterialExpressionRuntimeVirtualTextureSampleParameter>(Expression))
-		{
-			FMaterialParamExpressionData ParamData;
-			ParamData.ParamType = UMaterialExpressionRuntimeVirtualTextureSampleParameter::StaticClass();
-
-			ParamData.Name = VTTexParameter->GetParameterName();
-			ParamData.SortPriority = VTTexParameter->SortPriority;
-			ParamData.Group = VTTexParameter->Group;
-
-			//ensure(!ParamCache.GlobalParameters.Contains(ParamData.Name));
-			ParamCache.GlobalParameters.Add(ParamData.Name, ParamData);
-		}
-		else if (UMaterialExpressionFontSampleParameter* FontParameter = Cast<UMaterialExpressionFontSampleParameter>(Expression))
-		{
-			FMaterialParamExpressionData ParamData;
-			ParamData.ParamType = UMaterialExpressionFontSampleParameter::StaticClass();
-
-			ParamData.Name = FontParameter->GetParameterName();
-			ParamData.SortPriority = FontParameter->SortPriority;
-			ParamData.Group = FontParameter->Group;
-
-			//ensure(!ParamCache.GlobalParameters.Contains(ParamData.Name));
-			ParamCache.GlobalParameters.Add(ParamData.Name, ParamData);
-		}
-		else if (UMaterialExpressionMaterialFunctionCall* FuncParameter = Cast<UMaterialExpressionMaterialFunctionCall>(Expression))
-		{
-			if (FuncParameter->MaterialFunction)
-			{
-				if (UMaterialFunctionInterface* ParameterFunction = FuncParameter->MaterialFunction->GetBaseFunction())
-				{
-					CacheMaterialFunctionParameterData(ParameterFunction, ParamCache.GlobalParameters);
-				}
-			}
-		}
-		else if (UMaterialExpressionMaterialAttributeLayers* LayerParameter = Cast<UMaterialExpressionMaterialAttributeLayers>(Expression))
-		{
-			// there should only be one Material attribute layer expression per material
-			check(ParamCache.LayerParameterName == NAME_None);
-			ParamCache.LayerParameterName = LayerParameter->ParameterName;
-
-			// look into the instance static parameters first for overrides
-			UMaterialFunctionInterface* Function = nullptr;
-			const FStaticMaterialLayersParameter* StaticLayers = InStaticParameters.MaterialLayersParameters.FindByPredicate([LayerParameterName = LayerParameter->ParameterName](const FStaticMaterialLayersParameter& Layers)
-			{
-				return LayerParameterName == Layers.ParameterInfo.Name;
-			});
-
-			// If we found one cache those instead of what is on the material itself since they take precedence
-			if (StaticLayers)
-			{
-				// Replicate FStaticMaterialLayersParameter::GetParameterAssociatedFunction behavior while caching all function needed info
-
-				// Cache layer parameters
-				for (UMaterialFunctionInterface* Layer : StaticLayers->Value.Layers)
-				{
-					TMap<FName, FMaterialParamExpressionData>& LayerCache = ParamCache.LayerParameters.AddDefaulted_GetRef();
-					CacheMaterialFunctionParameterData(Layer, LayerCache);
-				}
-
-				// Cache blend parameters
-				for (UMaterialFunctionInterface* Blend : StaticLayers->Value.Blends)
-				{
-					TMap<FName, FMaterialParamExpressionData>& BlendCache = ParamCache.LayerParameters.AddDefaulted_GetRef();
-					CacheMaterialFunctionParameterData(Blend, BlendCache);
-				}
-			}
-			else
-			{
-				// Cache layer parameters
-				for (UMaterialFunctionInterface* Layer : LayerParameter->GetLayers())
-				{
-					TMap<FName, FMaterialParamExpressionData>& LayerCache = ParamCache.LayerParameters.AddDefaulted_GetRef();
-					CacheMaterialFunctionParameterData(Layer, LayerCache);
-				}
-
-				// Cache blend parameters
-				for (UMaterialFunctionInterface* Blend : LayerParameter->GetBlends())
-				{
-					TMap<FName, FMaterialParamExpressionData>& BlendCache = ParamCache.LayerParameters.AddDefaulted_GetRef();
-					CacheMaterialFunctionParameterData(Blend, BlendCache);
-				}
-			}
-		}
-	}
-	return ParamCache;
-}
-
 void UMaterialEditorPreviewParameters::PostEditChangeProperty(FPropertyChangedEvent& PropertyChangedEvent)
 {
 	if (PreviewMaterial && PropertyChangedEvent.ChangeType != EPropertyChangeType::Interactive)
@@ -455,21 +260,11 @@ void UMaterialEditorPreviewParameters::PostEditChangeProperty(FPropertyChangedEv
 	}
 }
 
-void UMaterialEditorPreviewParameters::AssignParameterToGroup(UMaterial* ParentMaterial, UDEditorParameterValue* ParameterValue, FName* OptionalGroupName)
+void UMaterialEditorPreviewParameters::AssignParameterToGroup(UDEditorParameterValue* ParameterValue, const FName& InParameterGroupName)
 {
-	check(ParentMaterial);
 	check(ParameterValue);
 
-	FName ParameterGroupName;
-	if (OptionalGroupName)
-	{
-		ParameterGroupName = *OptionalGroupName;
-	}
-	else
-	{
-		ParentMaterial->GetGroupName(ParameterValue->ParameterInfo, ParameterGroupName);
-	}
-	
+	FName ParameterGroupName = InParameterGroupName;
 	if (ParameterGroupName == TEXT("") || ParameterGroupName == TEXT("None"))
 	{
 		ParameterGroupName = TEXT("None");
@@ -527,10 +322,7 @@ static UDEditorParameterValue* CreateParameter_Vector(UObject* Owner, const FMat
 	{
 		Parameter->ParameterValue = Meta.Value.AsLinearColor();
 		Parameter->bIsUsedAsChannelMask = Meta.bUsedAsChannelMask;
-		Parameter->ChannelNames.R = Meta.ChannelName[0];
-		Parameter->ChannelNames.G = Meta.ChannelName[1];
-		Parameter->ChannelNames.B = Meta.ChannelName[2];
-		Parameter->ChannelNames.A = Meta.ChannelName[3];
+		Parameter->ChannelNames = Meta.ChannelNames;
 	}
 	return Parameter;
 }
@@ -541,10 +333,7 @@ static UDEditorParameterValue* CreateParameter_Texture(UObject* Owner, const FMa
 	if (Meta.Value.Type == EMaterialParameterType::Texture)
 	{
 		Parameter->ParameterValue = Meta.Value.Texture;
-		Parameter->ChannelNames.R = Meta.ChannelName[0];
-		Parameter->ChannelNames.G = Meta.ChannelName[1];
-		Parameter->ChannelNames.B = Meta.ChannelName[2];
-		Parameter->ChannelNames.A = Meta.ChannelName[3];
+		Parameter->ChannelNames = Meta.ChannelNames;
 	}
 	return Parameter;
 }
@@ -613,6 +402,7 @@ UDEditorParameterValue* UDEditorParameterValue::Create(UObject* Owner,
 
 	check(Parameter);
 	Parameter->ParameterInfo = ParameterInfo;
+	Parameter->SortPriority = Meta.SortPriority;
 	Parameter->bOverride = Meta.bOverride;
 	return Parameter;
 }
@@ -625,44 +415,9 @@ void UMaterialEditorPreviewParameters::RegenerateArrays()
 		// Only operate on base materials
 		UMaterial* ParentMaterial = PreviewMaterial;
 
-		// Use param cache to lookup group and sort priority
-		auto AssignGroupAndSortPriority = [this, ParentMaterial](UDEditorParameterValue* InEditorParamValue, const FMaterialExpressionParameterDataCache& InCachedExpressionData)
-		{
-			const FMaterialParamExpressionData* ParamData = nullptr;
-			if (InEditorParamValue->ParameterInfo.Association == EMaterialParameterAssociation::GlobalParameter)
-			{
-				ParamData = InCachedExpressionData.GlobalParameters.Find(InEditorParamValue->ParameterInfo.Name);
-			}
-			// if the association is not 'global parameter', look into attribute layers if we have a potentially valid index
-			else if (InEditorParamValue->ParameterInfo.Index >= 0)
-			{
-				if (InEditorParamValue->ParameterInfo.Association == EMaterialParameterAssociation::LayerParameter
-					&& InCachedExpressionData.LayerParameters.IsValidIndex(InEditorParamValue->ParameterInfo.Index))
-				{
-					ParamData = InCachedExpressionData.LayerParameters[InEditorParamValue->ParameterInfo.Index].Find(InEditorParamValue->ParameterInfo.Name);
-				}
-				else if (InEditorParamValue->ParameterInfo.Association == EMaterialParameterAssociation::BlendParameter
-					&& InCachedExpressionData.BlendParameters.IsValidIndex(InEditorParamValue->ParameterInfo.Index))
-				{
-					ParamData = InCachedExpressionData.BlendParameters[InEditorParamValue->ParameterInfo.Index].Find(InEditorParamValue->ParameterInfo.Name);
-				}
-			}
-			FName GroupName = NAME_None;
-			if (ParamData)
-			{
-				InEditorParamValue->SortPriority = ParamData->SortPriority;
-				GroupName = ParamData->Group;
-			}
-			AssignParameterToGroup(ParentMaterial, InEditorParamValue, &GroupName);
-		};
-
 		// This can run before UMaterial::PostEditChangeProperty has a chance to run, so explicitly call UpdateCachedExpressionData here
 		PreviewMaterial->UpdateCachedExpressionData();
 
-		// Cache relevant material expression data used to resolve editor param value info in RegenerateArrays
-		//@todo FH: can this be/should be part of `UpdateCachedExpressionData`?
-		FMaterialExpressionParameterDataCache ExpressionParameterDataCache = CacheMaterialExpressionParameterData(PreviewMaterial, FStaticParameterSet());
-
 		// Loop through all types of parameters for this material and add them to the parameter arrays.
 		TMap<FMaterialParameterInfo, FMaterialParameterMetadata> ParameterValues;
 		for (int32 TypeIndex = 0; TypeIndex < NumMaterialParameterTypes; ++TypeIndex)
@@ -674,7 +429,7 @@ void UMaterialEditorPreviewParameters::RegenerateArrays()
 			{
 				UDEditorParameterValue* Parameter = UDEditorParameterValue::Create(this, ParameterType, It.Key, It.Value);
 				Parameter->bOverride = true;
-				AssignGroupAndSortPriority(Parameter, ExpressionParameterDataCache);
+				AssignParameterToGroup(Parameter, It.Value.Group);
 			}
 		}
 
@@ -695,8 +450,7 @@ void UMaterialEditorPreviewParameters::RegenerateArrays()
 			{
 				ParameterValue->ParameterValue = Value;
 			}
-
-			AssignGroupAndSortPriority(ParameterValue, ExpressionParameterDataCache);
+			AssignParameterToGroup(ParameterValue, FName());
 		}
 	}
 	// sort contents of groups
@@ -889,20 +643,11 @@ void UMaterialEditorInstanceConstant::PostEditChangeProperty(FPropertyChangedEve
 	}
 }
 
-void  UMaterialEditorInstanceConstant::AssignParameterToGroup(UMaterial*, UDEditorParameterValue* ParameterValue, const FName* OptionalGroupName)
+void  UMaterialEditorInstanceConstant::AssignParameterToGroup(UDEditorParameterValue* ParameterValue, const FName& InParameterGroupName)
 {
 	check(ParameterValue);
 
-	FName ParameterGroupName;
-	if (OptionalGroupName)
-	{
-		ParameterGroupName = *OptionalGroupName;
-	}
-	else
-	{
-		SourceInstance->GetGroupName(ParameterValue->ParameterInfo, ParameterGroupName);
-	}
-
+	FName ParameterGroupName = InParameterGroupName;
 	if (ParameterGroupName == TEXT("") || ParameterGroupName == TEXT("None"))
 	{
 		if (bUseOldStyleMICEditorGroups == true)
@@ -945,116 +690,59 @@ void UMaterialEditorInstanceConstant::RegenerateArrays()
 
 	if (Parent)
 	{	
-		// Use param cache to lookup group and sort priority
-		auto AssignGroupAndSortPriority = [this](UDEditorParameterValue* InEditorParamValue, const FMaterialExpressionParameterDataCache& InCachedExpressionData)
-		{
-			const FMaterialParamExpressionData* ParamData = nullptr;
-			if (InEditorParamValue->ParameterInfo.Association == EMaterialParameterAssociation::GlobalParameter)
-			{
-				ParamData = InCachedExpressionData.GlobalParameters.Find(InEditorParamValue->ParameterInfo.Name);
-			}
-			// if the association is not 'global parameter', look into attribute layers if we have a potentially valid index
-			else if (InEditorParamValue->ParameterInfo.Index >= 0)
-			{
-				if (InEditorParamValue->ParameterInfo.Association == EMaterialParameterAssociation::LayerParameter
-					&& InCachedExpressionData.LayerParameters.IsValidIndex(InEditorParamValue->ParameterInfo.Index))
-				{
-					ParamData = InCachedExpressionData.LayerParameters[InEditorParamValue->ParameterInfo.Index].Find(InEditorParamValue->ParameterInfo.Name);
-				}
-				else if (InEditorParamValue->ParameterInfo.Association == EMaterialParameterAssociation::BlendParameter
-					&& InCachedExpressionData.BlendParameters.IsValidIndex(InEditorParamValue->ParameterInfo.Index))
-				{
-					ParamData = InCachedExpressionData.BlendParameters[InEditorParamValue->ParameterInfo.Index].Find(InEditorParamValue->ParameterInfo.Name);
-				}
-			}
-			if (ParamData)
-			{
-				InEditorParamValue->SortPriority = ParamData->SortPriority;
-			}
-			AssignParameterToGroup(nullptr/*useless param: Parent->GetMaterial()*/, InEditorParamValue, ParamData ? &ParamData->Group : nullptr);
-		};
-
 		// Only operate on base materials
 		UMaterial* ParentMaterial = Parent->GetMaterial();
 		SourceInstance->UpdateParameterNames();	// Update any parameter names that may have changed.
 		SourceInstance->UpdateCachedLayerParameters();
 
-		// Get all static parameters from the source instance.  This will handle inheriting parent values.	
-		FStaticParameterSet SourceStaticParameters;
- 		SourceInstance->GetStaticParameterValues(SourceStaticParameters);
-
-		// Loop through all types of parameters for this material and add them to the parameter arrays.
-		
 		// Need to get layer info first as other params are collected from layers
 		{
 			TArray<FMaterialParameterInfo> OutParameterInfo;
 			TArray<FGuid> Guids;
 			SourceInstance->GetAllMaterialLayersParameterInfo(OutParameterInfo, Guids);
 			// Copy Static Material Layers Parameters
-			for (int32 ParameterIdx = 0; ParameterIdx < SourceStaticParameters.MaterialLayersParameters.Num(); ParameterIdx++)
+			for (int32 ParameterIdx = 0; ParameterIdx < OutParameterInfo.Num(); ParameterIdx++)
 			{
-				FStaticMaterialLayersParameter MaterialLayersParameterParameterValue = FStaticMaterialLayersParameter(SourceStaticParameters.MaterialLayersParameters[ParameterIdx]);
+				FMaterialParameterInfo& ParameterInfo = OutParameterInfo[ParameterIdx];
+				FGuid ExpressionId = Guids[ParameterIdx];
+
 				UDEditorMaterialLayersParameterValue& ParameterValue = *(NewObject<UDEditorMaterialLayersParameterValue>(this));
+				ParameterValue.bOverride = true;
+				ParameterValue.ParameterInfo = ParameterInfo;
+				ParameterValue.ExpressionId = ExpressionId;
 
-				ParameterValue.ParameterValue = MaterialLayersParameterParameterValue.Value;
-				ParameterValue.bOverride = MaterialLayersParameterParameterValue.bOverride;
-				ParameterValue.ParameterInfo = MaterialLayersParameterParameterValue.ParameterInfo;
-				ParameterValue.ExpressionId = MaterialLayersParameterParameterValue.ExpressionGUID;
+				Parent->GetMaterialLayersParameterValue(ParameterInfo, ParameterValue.ParameterValue, ExpressionId);
+				ParameterValue.ParameterValue.LinkAllLayersToParent(); // Set parent guids for layers from parent material
 
-				AssignParameterToGroup(ParentMaterial, &ParameterValue);
+				// If the SourceInstance is overriding this parameter, use its settings
+				for (const FStaticMaterialLayersParameter& LayersParam : SourceInstance->GetStaticParameters().MaterialLayersParameters)
+				{
+					if (ParameterInfo == LayersParam.ParameterInfo)
+					{
+						ParameterValue.bOverride = LayersParam.bOverride;
+						if (LayersParam.bOverride)
+						{
+							ParameterValue.ParameterValue = LayersParam.Value;
+						}
+					}
+				}
+
+				AssignParameterToGroup(&ParameterValue, FName());
 			}
 		}
 
-		// Cache relevant material expression data to resolve editor param value info
-		FMaterialExpressionParameterDataCache ExpressionParameterDataCache = CacheMaterialExpressionParameterData(SourceInstance->GetMaterial(), SourceInstance->GetStaticParameters());
-
 		TMap<FMaterialParameterInfo, FMaterialParameterMetadata> ParameterValues;
 		for (int32 TypeIndex = 0; TypeIndex < NumMaterialParameterTypes; ++TypeIndex)
 		{
 			const EMaterialParameterType ParameterType = (EMaterialParameterType)TypeIndex;
-			if (IsStaticMaterialParameter(ParameterType))
-			{
-				// Logic for static parameters is slightly different, handled both by GetStaticParameterValues above, and code below
-				continue;
-			}
-
 			SourceInstance->GetAllParametersOfType(ParameterType, ParameterValues);
 			for (const auto& It : ParameterValues)
 			{
 				UDEditorParameterValue* Parameter = UDEditorParameterValue::Create(this, ParameterType, It.Key, It.Value);
-				AssignGroupAndSortPriority(Parameter, ExpressionParameterDataCache);
+				AssignParameterToGroup(Parameter, It.Value.Group);
 			}
 		}
 
-		// Copy Static Switch Parameters
-		for(int32 ParameterIdx = 0; ParameterIdx < SourceStaticParameters.StaticSwitchParameters.Num(); ParameterIdx++)
-		{	
-			const FStaticSwitchParameter& StaticSwitchParameterValue = SourceStaticParameters.StaticSwitchParameters[ParameterIdx];
-			UDEditorStaticSwitchParameterValue* ParamValue = NewObject<UDEditorStaticSwitchParameterValue>(this);
-
-			ParamValue->ParameterValue = StaticSwitchParameterValue.Value;
-			ParamValue->bOverride = StaticSwitchParameterValue.bOverride;
-			ParamValue->ParameterInfo = StaticSwitchParameterValue.ParameterInfo;
-			ParamValue->ExpressionId = StaticSwitchParameterValue.ExpressionGUID;
-			AssignGroupAndSortPriority(ParamValue, ExpressionParameterDataCache);
-		}
-
-		// Copy Static Component Mask Parameters
-		for(int32 ParameterIdx=0; ParameterIdx<SourceStaticParameters.StaticComponentMaskParameters.Num(); ParameterIdx++)
-		{
-			const FStaticComponentMaskParameter& StaticComponentMaskParameterValue = SourceStaticParameters.StaticComponentMaskParameters[ParameterIdx];
-			UDEditorStaticComponentMaskParameterValue* ParamValue = NewObject<UDEditorStaticComponentMaskParameterValue>(this);
-
-			ParamValue->ParameterValue.R = StaticComponentMaskParameterValue.R;
-			ParamValue->ParameterValue.G = StaticComponentMaskParameterValue.G;
-			ParamValue->ParameterValue.B = StaticComponentMaskParameterValue.B;
-			ParamValue->ParameterValue.A = StaticComponentMaskParameterValue.A;
-			ParamValue->bOverride = StaticComponentMaskParameterValue.bOverride;
-			ParamValue->ParameterInfo = StaticComponentMaskParameterValue.ParameterInfo;
-			ParamValue->ExpressionId = StaticComponentMaskParameterValue.ExpressionGUID;
-			AssignGroupAndSortPriority(ParamValue, ExpressionParameterDataCache);
-		}
-
 		IMaterialEditorModule* MaterialEditorModule = &FModuleManager::LoadModuleChecked<IMaterialEditorModule>( "MaterialEditor" );
 		MaterialEditorModule->GetVisibleMaterialParameters(ParentMaterial, SourceInstance, VisibleExpressions);
 	}
diff --git a/Engine/Source/Runtime/Engine/Classes/Materials/Material.h b/Engine/Source/Runtime/Engine/Classes/Materials/Material.h
--- a/Engine/Source/Runtime/Engine/Classes/Materials/Material.h
+++ b/Engine/Source/Runtime/Engine/Classes/Materials/Material.h
@@ -1060,9 +1060,6 @@ public:
 	ENGINE_API virtual bool UpdateLightmassTextureTracking() override;
 	ENGINE_API virtual int32 GetLayerParameterIndex(EMaterialParameterAssociation Association, UMaterialFunctionInterface* LayerFunction) const override;
 #if WITH_EDITOR
-	ENGINE_API virtual bool GetGroupName(const FHashedMaterialParameterInfo& ParameterInfo, FName& OutGroup) const override;
-	ENGINE_API virtual bool GetParameterDesc(const FHashedMaterialParameterInfo& ParameterInfo, FString& OutDesc, const TArray<struct FStaticMaterialLayersParameter>* MaterialLayersParameters = nullptr) const override;
-	ENGINE_API virtual bool GetParameterSortPriority(const FHashedMaterialParameterInfo& ParameterInfo, int32& OutSortPriority, const TArray<struct FStaticMaterialLayersParameter>* MaterialLayersParameters = nullptr) const override;
 	ENGINE_API virtual bool GetGroupSortPriority(const FString& InGroupName, int32& OutSortPriority) const override;
 	ENGINE_API virtual bool GetTexturesInPropertyChain(EMaterialProperty InProperty, TArray<UTexture*>& OutTextures, 
 		TArray<FName>* OutTextureParamNames, struct FStaticParameterSet* InStaticParameterSet,
diff --git a/Engine/Source/Runtime/Engine/Classes/Materials/MaterialExpressionFontSampleParameter.h b/Engine/Source/Runtime/Engine/Classes/Materials/MaterialExpressionFontSampleParameter.h
--- a/Engine/Source/Runtime/Engine/Classes/Materials/MaterialExpressionFontSampleParameter.h
+++ b/Engine/Source/Runtime/Engine/Classes/Materials/MaterialExpressionFontSampleParameter.h
@@ -51,16 +51,27 @@ class UMaterialExpressionFontSampleParameter : public UMaterialExpressionFontSam
 	virtual bool GetParameterValue(FMaterialParameterMetadata& OutMeta) const override
 	{
 		OutMeta.Value = FMaterialParameterValue(Font, FontTexturePage);
+		OutMeta.Description = Desc;
 		OutMeta.ExpressionGuid = ExpressionGUID;
+		OutMeta.Group = Group;
+		OutMeta.SortPriority = SortPriority;
 		return true;
 	}
 	virtual bool SetParameterValue(const FName& Name, const FMaterialParameterMetadata& Meta, EMaterialExpressionSetParameterValueFlags Flags) override
 	{
 		if (Meta.Value.Type == EMaterialParameterType::Font)
 		{
-			return SetParameterValue(Name, Meta.Value.Font.Value, Meta.Value.Font.Page, Flags);
+			if (SetParameterValue(Name, Meta.Value.Font.Value, Meta.Value.Font.Page, Flags))
+			{
+				if (EnumHasAnyFlags(Flags, EMaterialExpressionSetParameterValueFlags::AssignGroupAndSortPriority))
+				{
+					Group = Meta.Group;
+					SortPriority = Meta.SortPriority;
+				}
+				return true;
+			}
 		}
-		return Super::SetParameterValue(Name, Meta, Flags);
+		return false;
 	}
 #endif
 	//~ End UMaterialExpression Interface
diff --git a/Engine/Source/Runtime/Engine/Classes/Materials/MaterialExpressionParameter.h b/Engine/Source/Runtime/Engine/Classes/Materials/MaterialExpressionParameter.h
--- a/Engine/Source/Runtime/Engine/Classes/Materials/MaterialExpressionParameter.h
+++ b/Engine/Source/Runtime/Engine/Classes/Materials/MaterialExpressionParameter.h
@@ -48,6 +48,15 @@ class UMaterialExpressionParameter : public UMaterialExpression
 	virtual FName GetParameterName() const override { return ParameterName; }
 	virtual void SetParameterName(const FName& Name) override { ParameterName = Name; }
 	virtual void ValidateParameterName(const bool bAllowDuplicateName) override;
+
+	virtual bool GetParameterValue(FMaterialParameterMetadata& OutMeta) const override
+	{
+		OutMeta.Description = Desc;
+		OutMeta.ExpressionGuid = ExpressionGUID;
+		OutMeta.Group = Group;
+		OutMeta.SortPriority = SortPriority;
+		return true;
+	}
 #endif
 	//~ End UMaterialExpression Interface
 
diff --git a/Engine/Source/Runtime/Engine/Classes/Materials/MaterialExpressionRuntimeVirtualTextureSampleParameter.h b/Engine/Source/Runtime/Engine/Classes/Materials/MaterialExpressionRuntimeVirtualTextureSampleParameter.h
--- a/Engine/Source/Runtime/Engine/Classes/Materials/MaterialExpressionRuntimeVirtualTextureSampleParameter.h
+++ b/Engine/Source/Runtime/Engine/Classes/Materials/MaterialExpressionRuntimeVirtualTextureSampleParameter.h
@@ -53,16 +53,27 @@ class UMaterialExpressionRuntimeVirtualTextureSampleParameter : public UMaterial
 	virtual bool GetParameterValue(FMaterialParameterMetadata& OutMeta) const override
 	{
 		OutMeta.Value = VirtualTexture;
+		OutMeta.Description = Desc;
 		OutMeta.ExpressionGuid = ExpressionGUID;
+		OutMeta.Group = Group;
+		OutMeta.SortPriority = SortPriority;
 		return true;
 	}
 	virtual bool SetParameterValue(const FName& Name, const FMaterialParameterMetadata& Meta, EMaterialExpressionSetParameterValueFlags Flags) override
 	{
 		if (Meta.Value.Type == EMaterialParameterType::RuntimeVirtualTexture)
 		{
-			return SetParameterValue(Name, Meta.Value.RuntimeVirtualTexture, Flags);
+			if (SetParameterValue(Name, Meta.Value.RuntimeVirtualTexture, Flags))
+			{
+				if (EnumHasAnyFlags(Flags, EMaterialExpressionSetParameterValueFlags::AssignGroupAndSortPriority))
+				{
+					Group = Meta.Group;
+					SortPriority = Meta.SortPriority;
+				}
+				return true;
+			}
 		}
-		return Super::SetParameterValue(Name, Meta, Flags);
+		return false;
 	}
 #endif
 	virtual FGuid& GetParameterExpressionId() override { return ExpressionGUID; }
diff --git a/Engine/Source/Runtime/Engine/Classes/Materials/MaterialExpressionScalarParameter.h b/Engine/Source/Runtime/Engine/Classes/Materials/MaterialExpressionScalarParameter.h
--- a/Engine/Source/Runtime/Engine/Classes/Materials/MaterialExpressionScalarParameter.h
+++ b/Engine/Source/Runtime/Engine/Classes/Materials/MaterialExpressionScalarParameter.h
@@ -47,18 +47,25 @@ class UMaterialExpressionScalarParameter : public UMaterialExpressionParameter
 	virtual bool GetParameterValue(FMaterialParameterMetadata& OutMeta) const override
 	{
 		OutMeta.Value = DefaultValue;
-		OutMeta.ExpressionGuid = ExpressionGUID;
 		OutMeta.ScalarMin = SliderMin;
 		OutMeta.ScalarMax = SliderMax;
-		return true;
+		return Super::GetParameterValue(OutMeta);
 	}
 	virtual bool SetParameterValue(const FName& Name, const FMaterialParameterMetadata& Meta, EMaterialExpressionSetParameterValueFlags Flags) override
 	{
 		if (Meta.Value.Type == EMaterialParameterType::Scalar)
 		{
-			return SetParameterValue(Name, Meta.Value.AsScalar(), Flags);
+			if (SetParameterValue(Name, Meta.Value.AsScalar(), Flags))
+			{
+				if (EnumHasAnyFlags(Flags, EMaterialExpressionSetParameterValueFlags::AssignGroupAndSortPriority))
+				{
+					Group = Meta.Group;
+					SortPriority = Meta.SortPriority;
+				}
+				return true;
+			}
 		}
-		return Super::SetParameterValue(Name, Meta, Flags);
+		return false;
 	}
 #endif
 	//~ End UMaterialExpression Interface
diff --git a/Engine/Source/Runtime/Engine/Classes/Materials/MaterialExpressionStaticBoolParameter.h b/Engine/Source/Runtime/Engine/Classes/Materials/MaterialExpressionStaticBoolParameter.h
--- a/Engine/Source/Runtime/Engine/Classes/Materials/MaterialExpressionStaticBoolParameter.h
+++ b/Engine/Source/Runtime/Engine/Classes/Materials/MaterialExpressionStaticBoolParameter.h
@@ -27,16 +27,23 @@ public:
 	virtual bool GetParameterValue(FMaterialParameterMetadata& OutMeta) const override
 	{
 		OutMeta.Value = (bool)DefaultValue;
-		OutMeta.ExpressionGuid = ExpressionGUID;
-		return true;
+		return Super::GetParameterValue(OutMeta);
 	}
 	virtual bool SetParameterValue(const FName& Name, const FMaterialParameterMetadata& Meta, EMaterialExpressionSetParameterValueFlags Flags) override
 	{
 		if (Meta.Value.Type == EMaterialParameterType::StaticSwitch)
 		{
-			return SetParameterValue(Name, Meta.Value.AsStaticSwitch(), Meta.ExpressionGuid, Flags);
+			if (SetParameterValue(Name, Meta.Value.AsStaticSwitch(), Meta.ExpressionGuid, Flags))
+			{
+				if (EnumHasAnyFlags(Flags, EMaterialExpressionSetParameterValueFlags::AssignGroupAndSortPriority))
+				{
+					Group = Meta.Group;
+					SortPriority = Meta.SortPriority;
+				}
+				return true;
+			}
 		}
-		return Super::SetParameterValue(Name, Meta, Flags);
+		return false;
 	}
 #endif
 	//~ End UMaterialExpression Interface
diff --git a/Engine/Source/Runtime/Engine/Classes/Materials/MaterialExpressionStaticComponentMaskParameter.h b/Engine/Source/Runtime/Engine/Classes/Materials/MaterialExpressionStaticComponentMaskParameter.h
--- a/Engine/Source/Runtime/Engine/Classes/Materials/MaterialExpressionStaticComponentMaskParameter.h
+++ b/Engine/Source/Runtime/Engine/Classes/Materials/MaterialExpressionStaticComponentMaskParameter.h
@@ -41,22 +41,29 @@ public:
 	virtual bool GetParameterValue(FMaterialParameterMetadata& OutMeta) const override
 	{
 		OutMeta.Value = FMaterialParameterValue(DefaultR, DefaultG, DefaultB, DefaultA);
-		OutMeta.ExpressionGuid = ExpressionGUID;
-		return true;
+		return Super::GetParameterValue(OutMeta);
 	}
 	virtual bool SetParameterValue(const FName& Name, const FMaterialParameterMetadata& Meta, EMaterialExpressionSetParameterValueFlags Flags) override
 	{
 		if (Meta.Value.Type == EMaterialParameterType::StaticComponentMask)
 		{
-			return SetParameterValue(Name,
+			if (SetParameterValue(Name,
 				Meta.Value.Bool[0],
 				Meta.Value.Bool[1],
 				Meta.Value.Bool[2],
 				Meta.Value.Bool[3],
 				Meta.ExpressionGuid,
-				Flags);
+				Flags))
+			{
+				if (EnumHasAnyFlags(Flags, EMaterialExpressionSetParameterValueFlags::AssignGroupAndSortPriority))
+				{
+					Group = Meta.Group;
+					SortPriority = Meta.SortPriority;
+				}
+				return true;
+			}
 		}
-		return Super::SetParameterValue(Name, Meta, Flags);
+		return false;
 	}
 #endif
 	//~ End UMaterialExpression Interface
diff --git a/Engine/Source/Runtime/Engine/Classes/Materials/MaterialExpressionTextureSampleParameter.h b/Engine/Source/Runtime/Engine/Classes/Materials/MaterialExpressionTextureSampleParameter.h
--- a/Engine/Source/Runtime/Engine/Classes/Materials/MaterialExpressionTextureSampleParameter.h
+++ b/Engine/Source/Runtime/Engine/Classes/Materials/MaterialExpressionTextureSampleParameter.h
@@ -54,20 +54,28 @@ class ENGINE_API UMaterialExpressionTextureSampleParameter : public UMaterialExp
 	virtual bool GetParameterValue(FMaterialParameterMetadata& OutMeta) const override
 	{
 		OutMeta.Value = Texture;
+		OutMeta.Description = Desc;
 		OutMeta.ExpressionGuid = ExpressionGUID;
-		OutMeta.ChannelName[0] = ChannelNames.R;
-		OutMeta.ChannelName[1] = ChannelNames.G;
-		OutMeta.ChannelName[2] = ChannelNames.B;
-		OutMeta.ChannelName[3] = ChannelNames.A;
+		OutMeta.Group = Group;
+		OutMeta.SortPriority = SortPriority;
+		OutMeta.ChannelNames = ChannelNames;
 		return true;
 	}
 	virtual bool SetParameterValue(const FName& Name, const FMaterialParameterMetadata& Meta, EMaterialExpressionSetParameterValueFlags Flags) override
 	{
 		if (Meta.Value.Type == EMaterialParameterType::Texture)
 		{
-			return SetParameterValue(Name, Meta.Value.Texture, Flags);
+			if (SetParameterValue(Name, Meta.Value.Texture, Flags))
+			{
+				if (EnumHasAnyFlags(Flags, EMaterialExpressionSetParameterValueFlags::AssignGroupAndSortPriority))
+				{
+					Group = Meta.Group;
+					SortPriority = Meta.SortPriority;
+				}
+				return true;
+			}
 		}
-		return Super::SetParameterValue(Name, Meta, Flags);
+		return false;
 	}
 #endif
 	//~ End UMaterialExpression Interface
diff --git a/Engine/Source/Runtime/Engine/Classes/Materials/MaterialExpressionVectorParameter.h b/Engine/Source/Runtime/Engine/Classes/Materials/MaterialExpressionVectorParameter.h
--- a/Engine/Source/Runtime/Engine/Classes/Materials/MaterialExpressionVectorParameter.h
+++ b/Engine/Source/Runtime/Engine/Classes/Materials/MaterialExpressionVectorParameter.h
@@ -36,20 +36,24 @@ class UMaterialExpressionVectorParameter : public UMaterialExpressionParameter
 	virtual bool GetParameterValue(FMaterialParameterMetadata& OutMeta) const override
 	{
 		OutMeta.Value = DefaultValue;
-		OutMeta.ExpressionGuid = ExpressionGUID;
-		OutMeta.ChannelName[0] = ChannelNames.R;
-		OutMeta.ChannelName[1] = ChannelNames.G;
-		OutMeta.ChannelName[2] = ChannelNames.B;
-		OutMeta.ChannelName[3] = ChannelNames.A;
-		return true;
+		OutMeta.ChannelNames = ChannelNames;
+		return Super::GetParameterValue(OutMeta);
 	}
 	virtual bool SetParameterValue(const FName& Name, const FMaterialParameterMetadata& Meta, EMaterialExpressionSetParameterValueFlags Flags) override
 	{
 		if (Meta.Value.Type == EMaterialParameterType::Vector)
 		{
-			return SetParameterValue(Name, Meta.Value.AsLinearColor(), Flags);
+			if (SetParameterValue(Name, Meta.Value.AsLinearColor(), Flags))
+			{
+				if (EnumHasAnyFlags(Flags, EMaterialExpressionSetParameterValueFlags::AssignGroupAndSortPriority))
+				{
+					Group = Meta.Group;
+					SortPriority = Meta.SortPriority;
+				}
+				return true;
+			}
 		}
-		return Super::SetParameterValue(Name, Meta, Flags);
+		return false;
 	}
 #endif
 	//~ End UMaterialExpression Interface
diff --git a/Engine/Source/Runtime/Engine/Classes/Materials/MaterialFunctionInterface.h b/Engine/Source/Runtime/Engine/Classes/Materials/MaterialFunctionInterface.h
--- a/Engine/Source/Runtime/Engine/Classes/Materials/MaterialFunctionInterface.h
+++ b/Engine/Source/Runtime/Engine/Classes/Materials/MaterialFunctionInterface.h
@@ -232,138 +232,6 @@ public:
 		return false;
 	}
 
-	/** Finds the first matching parameter's group name */
-	bool GetParameterGroupName(const FHashedMaterialParameterInfo& ParameterInfo, FName& OutGroup)
-	{
-		if (UMaterialFunctionInterface* ParameterFunction = GetBaseFunction())
-		{
-			TArray<UMaterialFunctionInterface*> Functions;
-			ParameterFunction->GetDependentFunctions(Functions);
-			Functions.AddUnique(ParameterFunction);
-
-			for (UMaterialFunctionInterface* Function : Functions)
-			{
-				for (UMaterialExpression* FunctionExpression : *Function->GetFunctionExpressions())
-				{
-					if (const UMaterialExpressionParameter* Parameter = Cast<const UMaterialExpressionParameter>(FunctionExpression))
-					{
-						if (Parameter->ParameterName == ParameterInfo.Name)
-						{
-							OutGroup = Parameter->Group;
-							return true;
-						}
-					}
-					else if (const UMaterialExpressionTextureSampleParameter* TexParameter = Cast<const UMaterialExpressionTextureSampleParameter>(FunctionExpression))
-					{
-						if (TexParameter->ParameterName == ParameterInfo.Name)
-						{
-							OutGroup = TexParameter->Group;
-							return true;
-						}
-					}
-					else if (const UMaterialExpressionFontSampleParameter* FontParameter = Cast<const UMaterialExpressionFontSampleParameter>(FunctionExpression))
-					{
-						if (FontParameter->ParameterName == ParameterInfo.Name)
-						{
-							OutGroup = FontParameter->Group;
-							return true;
-						}
-					}
-				}
-			}
-		}
-
-		return false;
-	}
-
-	/** Finds the first matching parameter's group name */
-	bool GetParameterSortPriority(const FHashedMaterialParameterInfo& ParameterInfo, int32& OutSortPriority)
-	{
-		if (UMaterialFunctionInterface* ParameterFunction = GetBaseFunction())
-		{
-			TArray<UMaterialFunctionInterface*> Functions;
-			ParameterFunction->GetDependentFunctions(Functions);
-			Functions.AddUnique(ParameterFunction);
-
-			for (UMaterialFunctionInterface* Function : Functions)
-			{
-				for (UMaterialExpression* FunctionExpression : *Function->GetFunctionExpressions())
-				{
-					if (const UMaterialExpressionParameter* Parameter = Cast<const UMaterialExpressionParameter>(FunctionExpression))
-					{
-						if (Parameter->ParameterName == ParameterInfo.Name)
-						{
-							OutSortPriority = Parameter->SortPriority;
-							return true;
-						}
-					}
-					else if (const UMaterialExpressionTextureSampleParameter* TexParameter = Cast<const UMaterialExpressionTextureSampleParameter>(FunctionExpression))
-					{
-						if (TexParameter->ParameterName == ParameterInfo.Name)
-						{
-							OutSortPriority = TexParameter->SortPriority;
-							return true;
-						}
-					}
-					else if (const UMaterialExpressionFontSampleParameter* FontParameter = Cast<const UMaterialExpressionFontSampleParameter>(FunctionExpression))
-					{
-						if (FontParameter->ParameterName == ParameterInfo.Name)
-						{
-							OutSortPriority = FontParameter->SortPriority;
-							return true;
-						}
-					}
-				}
-			}
-		}
-
-		return false;
-	}
-
-	/** Finds the first matching parameter's description */
-	bool GetParameterDesc(const FHashedMaterialParameterInfo& ParameterInfo, FString& OutDesc)
-	{
-		if (UMaterialFunctionInterface* ParameterFunction = GetBaseFunction())
-		{
-			TArray<UMaterialFunctionInterface*> Functions;
-			ParameterFunction->GetDependentFunctions(Functions);
-			Functions.AddUnique(ParameterFunction);
-
-			for (UMaterialFunctionInterface* Function : Functions)
-			{
-				for (UMaterialExpression* FunctionExpression : *Function->GetFunctionExpressions())
-				{
-					if (const UMaterialExpressionParameter* Parameter = Cast<const UMaterialExpressionParameter>(FunctionExpression))
-					{
-						if (Parameter->ParameterName == ParameterInfo.Name)
-						{
-							OutDesc = Parameter->Desc;
-							return true;
-						}
-					}
-					else if (const UMaterialExpressionTextureSampleParameter* TexParameter = Cast<const UMaterialExpressionTextureSampleParameter>(FunctionExpression))
-					{
-						if (TexParameter->ParameterName == ParameterInfo.Name)
-						{
-							OutDesc = TexParameter->Desc;
-							return true;
-						}
-					}
-					else if (const UMaterialExpressionFontSampleParameter* FontParameter = Cast<const UMaterialExpressionFontSampleParameter>(FunctionExpression))
-					{
-						if (FontParameter->ParameterName == ParameterInfo.Name)
-						{
-							OutDesc = FontParameter->Desc;
-							return true;
-						}
-					}
-				}
-			}
-		}
-
-		return false;
-	}
-
 	/** Returns if any of the matching parameters have changed */
 	template <typename ParameterType, typename ExpressionType>
 	bool UpdateParameterSet(ParameterType& Parameter)
diff --git a/Engine/Source/Runtime/Engine/Classes/Materials/MaterialInstance.h b/Engine/Source/Runtime/Engine/Classes/Materials/MaterialInstance.h
--- a/Engine/Source/Runtime/Engine/Classes/Materials/MaterialInstance.h
+++ b/Engine/Source/Runtime/Engine/Classes/Materials/MaterialInstance.h
@@ -590,8 +590,6 @@ public:
 	virtual ENGINE_API float GetExportResolutionScale() const override;
 	virtual ENGINE_API int32 GetLayerParameterIndex(EMaterialParameterAssociation Association, UMaterialFunctionInterface* LayerFunction) const override;
 #if WITH_EDITOR
-	virtual ENGINE_API bool GetParameterDesc(const FHashedMaterialParameterInfo& ParameterInfo, FString& OutDesc, const TArray<struct FStaticMaterialLayersParameter>* MaterialLayersParameters = nullptr) const override;
-	virtual ENGINE_API bool GetParameterSortPriority(const FHashedMaterialParameterInfo& ParameterInfo, int32& OutSortPriority, const TArray<struct FStaticMaterialLayersParameter>* MaterialLayersParameters = nullptr) const override;
 	virtual ENGINE_API bool GetGroupSortPriority(const FString& InGroupName, int32& OutSortPriority) const override;
 	virtual ENGINE_API bool GetTexturesInPropertyChain(EMaterialProperty InProperty, TArray<UTexture*>& OutTextures,
 		TArray<FName>* OutTextureParamNames, struct FStaticParameterSet* InStaticParameterSet,
@@ -715,10 +713,6 @@ public:
 	ENGINE_API virtual void GetDependentFunctions(TArray<class UMaterialFunctionInterface*>& DependentFunctions) const override;
 #endif // WITH_EDITORONLY_DATA
 
-#if WITH_EDITOR
-	ENGINE_API virtual bool GetGroupName(const FHashedMaterialParameterInfo& ParameterInfo, FName& OutGroup) const override;
-#endif // WITH_EDITOR
-
 	/** Appends textures referenced by expressions, including nested functions. */
 	ENGINE_API virtual TArrayView<const TObjectPtr<UObject>> GetReferencedTextures() const override final { return CachedData ? CachedData->ReferencedTextures : TArrayView<const TObjectPtr<UObject>>(); }
 
diff --git a/Engine/Source/Runtime/Engine/Classes/Materials/MaterialInterface.h b/Engine/Source/Runtime/Engine/Classes/Materials/MaterialInterface.h
--- a/Engine/Source/Runtime/Engine/Classes/Materials/MaterialInterface.h
+++ b/Engine/Source/Runtime/Engine/Classes/Materials/MaterialInterface.h
@@ -484,8 +484,7 @@ public:
 	* @param	OutSortPriority	Will contain the sort priority of the parameter if successful
 	* @return					True if successful
 	*/
-	virtual bool GetParameterSortPriority(const FHashedMaterialParameterInfo& ParameterInfo, int32& OutSortPriority, const TArray<struct FStaticMaterialLayersParameter>* MaterialLayersParameters = nullptr) const
-		PURE_VIRTUAL(UMaterialInterface::GetParameterSortPriority, return false;);
+	ENGINE_API bool GetParameterSortPriority(const FHashedMaterialParameterInfo& ParameterInfo, int32& OutSortPriority) const;
 
 	/**
 	* Get the sort priority index of the given parameter group
@@ -494,7 +493,7 @@ public:
 	* @param	OutSortPriority	Will contain the sort priority of the parameter group if successful
 	* @return					True if successful
 	*/
-	virtual bool GetGroupSortPriority(const FString& InGroupName, int32& OutSortPriority) const
+	ENGINE_API virtual bool GetGroupSortPriority(const FString& InGroupName, int32& OutSortPriority) const
 		PURE_VIRTUAL(UMaterialInterface::GetGroupSortPriority, return false;);
 #endif // WITH_EDITOR
 
@@ -721,8 +720,8 @@ public:
 		ERHIFeatureLevel::Type InFeatureLevel = ERHIFeatureLevel::Num, EMaterialQualityLevel::Type InQuality = EMaterialQualityLevel::Num)
 		PURE_VIRTUAL(UMaterialInterface::GetTexturesInPropertyChain,return false;);
 
-	ENGINE_API virtual bool GetGroupName(const FHashedMaterialParameterInfo& ParameterInfo, FName& GroupName) const;
-	ENGINE_API virtual bool GetParameterDesc(const FHashedMaterialParameterInfo& ParameterInfo, FString& OutDesc, const TArray<struct FStaticMaterialLayersParameter>* MaterialLayersParameters = nullptr) const;
+	ENGINE_API bool GetGroupName(const FHashedMaterialParameterInfo& ParameterInfo, FName& GroupName) const;
+	ENGINE_API bool GetParameterDesc(const FHashedMaterialParameterInfo& ParameterInfo, FString& OutDesc) const;
 	ENGINE_API bool GetScalarParameterSliderMinMax(const FHashedMaterialParameterInfo& ParameterInfo, float& OutSliderMin, float& OutSliderMax) const;
 #endif // WITH_EDITOR
 
diff --git a/Engine/Source/Runtime/Engine/Private/Materials/Material.cpp b/Engine/Source/Runtime/Engine/Private/Materials/Material.cpp
--- a/Engine/Source/Runtime/Engine/Private/Materials/Material.cpp
+++ b/Engine/Source/Runtime/Engine/Private/Materials/Material.cpp
@@ -1868,178 +1868,6 @@ const UMaterial* UMaterial::GetMaterial_Concurrent(TMicRecursionGuard) const
 }
 
 #if WITH_EDITOR
-bool UMaterial::GetGroupName(const FHashedMaterialParameterInfo& ParameterInfo, FName& OutGroup) const
-{
-	for (const UMaterialExpression* Expression : Expressions)
-	{
-		// Only need to check parameters that match in associated scope
-		if (ParameterInfo.Association == EMaterialParameterAssociation::GlobalParameter)
-		{
-			if (const UMaterialExpressionParameter* Parameter = Cast<const UMaterialExpressionParameter>(Expression))
-			{
-				if (Parameter->ParameterName == ParameterInfo.Name)
-				{
-					OutGroup = Parameter->Group;
-					return true;
-				}
-			}
-			else if (const UMaterialExpressionTextureSampleParameter* TexParameter = Cast<const UMaterialExpressionTextureSampleParameter>(Expression))
-			{
-				if (TexParameter->ParameterName == ParameterInfo.Name)
-				{
-					OutGroup = TexParameter->Group;
-					return true;
-				}
-			}
-			else if (const UMaterialExpressionRuntimeVirtualTextureSampleParameter* RVTParameter = Cast<const UMaterialExpressionRuntimeVirtualTextureSampleParameter>(Expression))
-			{
-				if (RVTParameter->ParameterName == ParameterInfo.Name)
-				{
-					OutGroup = RVTParameter->Group;
-					return true;
-				}
-			}
-			else if (const UMaterialExpressionFontSampleParameter* FontParameter = Cast<const UMaterialExpressionFontSampleParameter>(Expression))
-			{
-				if (FontParameter->ParameterName == ParameterInfo.Name)
-				{
-					OutGroup = FontParameter->Group;
-					return true;
-				}
-			}
-			else if (const UMaterialExpressionMaterialFunctionCall* FunctionCall = Cast<const UMaterialExpressionMaterialFunctionCall>(Expression))
-			{
-				if (FunctionCall->MaterialFunction && FunctionCall->MaterialFunction->GetParameterGroupName(ParameterInfo, OutGroup))
-				{
-					return true;
-				}
-			}
-			else if (const UMaterialExpressionMaterialAttributeLayers* LayersParameter = Cast<const UMaterialExpressionMaterialAttributeLayers>(Expression))
-			{
-				if (LayersParameter->ParameterName == ParameterInfo.Name)
-				{
-					OutGroup = FName(); // No group as custom interface so we can end the search
-					return true;
-				}
-			}
-		}
-		else
-		{
-			if (const UMaterialExpressionMaterialAttributeLayers* LayersExpression = Cast<const UMaterialExpressionMaterialAttributeLayers>(Expression))
-			{
-				if (ParameterInfo.Association == EMaterialParameterAssociation::LayerParameter)
-				{
-					const TArray<UMaterialFunctionInterface*>& Layers = LayersExpression->GetLayers();
-					if (Layers.IsValidIndex(ParameterInfo.Index))
-					{
-						UMaterialFunctionInterface* Layer = Layers[ParameterInfo.Index];
-						if (Layer && Layer->GetParameterGroupName(ParameterInfo, OutGroup))
-						{
-							return true;
-						}
-					}
-				}
-				else if (ParameterInfo.Association == EMaterialParameterAssociation::BlendParameter)
-				{
-					const TArray<UMaterialFunctionInterface*>& Blends = LayersExpression->GetBlends();
-					if (Blends.IsValidIndex(ParameterInfo.Index))
-					{
-						UMaterialFunctionInterface* Blend = Blends[ParameterInfo.Index];
-						if (Blend && Blend->GetParameterGroupName(ParameterInfo, OutGroup))
-						{
-							return true;
-						}
-					}
-				}
-			}
-		}
-	}
-
-	return false;
-}
-
-bool UMaterial::GetParameterDesc(const FHashedMaterialParameterInfo& ParameterInfo, FString& OutDesc, const TArray<FStaticMaterialLayersParameter>* MaterialLayersParameters) const
-{
-	for (const UMaterialExpression* Expression : Expressions)
-	{
-		// Only need to check parameters that match in associated scope
-		if (ParameterInfo.Association == EMaterialParameterAssociation::GlobalParameter)
-		{
-			if (const UMaterialExpressionParameter* Parameter = Cast<const UMaterialExpressionParameter>(Expression))
-			{
-				if (Parameter->ParameterName == ParameterInfo.Name)
-				{
-					OutDesc = Parameter->Desc;
-					return true;
-				}
-			}
-			else if (const UMaterialExpressionTextureSampleParameter* TexParameter = Cast<const UMaterialExpressionTextureSampleParameter>(Expression))
-			{
-				if (TexParameter->ParameterName == ParameterInfo.Name)
-				{
-					OutDesc = TexParameter->Desc;
-					return true;
-				}
-			}
-			else if (const UMaterialExpressionFontSampleParameter* FontParameter = Cast<const UMaterialExpressionFontSampleParameter>(Expression))
-			{
-				if (FontParameter->ParameterName == ParameterInfo.Name)
-				{
-					OutDesc = FontParameter->Desc;
-					return true;
-				}
-			}
-			else if (const UMaterialExpressionMaterialFunctionCall* FunctionCall = Cast<const UMaterialExpressionMaterialFunctionCall>(Expression))
-			{
-				if (FunctionCall->MaterialFunction && FunctionCall->MaterialFunction->GetParameterDesc(ParameterInfo, OutDesc))
-				{
-					return true;
-				}
-			}
-			else if (const UMaterialExpressionMaterialAttributeLayers* LayersParameter = Cast<const UMaterialExpressionMaterialAttributeLayers>(Expression))
-			{
-				if (LayersParameter->ParameterName == ParameterInfo.Name)
-				{
-					OutDesc = LayersParameter->Desc;
-					return true;
-				}
-			}
-		}		
-		else
-		{
-			if (const UMaterialExpressionMaterialAttributeLayers* LayersExpression = Cast<const UMaterialExpressionMaterialAttributeLayers>(Expression))
-			{
-				UMaterialFunctionInterface* Function = nullptr;
-				
-				// Handle function overrides when searching for parameters
-				if (MaterialLayersParameters)
-				{
-					const FName& ParameterName = LayersExpression->ParameterName;
-					for (const FStaticMaterialLayersParameter& LayersParameter : *MaterialLayersParameters)
-					{
-						if (LayersParameter.ParameterInfo.Name == ParameterName)
-						{
-							Function = LayersParameter.GetParameterAssociatedFunction(ParameterInfo);
-							break;
-						}
-					}
-				}
-				
-				if (!Function)
-				{
-					Function = LayersExpression->GetParameterAssociatedFunction(ParameterInfo);
-				}
-
-				if (Function && Function->GetParameterDesc(ParameterInfo, OutDesc))
-				{
-					return true;
-				}
-			}
-		}
-	}
-
-	return false;
-}
 void UMaterial::UpdateCachedExpressionData()
 {
 	COOK_STAT(FScopedDurationTimer BlockingTimer(MaterialCookStats::UpdateCachedExpressionDataSec));
@@ -4464,40 +4292,9 @@ bool UMaterial::CopyExpressionParameters(UMaterialExpression* Source, UMaterialE
 	{
 		const FName ParameterName = Destination->GetParameterName();
 		Destination->Modify();
-		verify(Destination->SetParameterValue(ParameterName, Meta, EMaterialExpressionSetParameterValueFlags::NoUpdateExpressionGuid));
 
-		if (Source->IsA(UMaterialExpressionParameter::StaticClass()))
-		{
-			UMaterialExpressionParameter* SourceParam = (UMaterialExpressionParameter*)Source;
-			UMaterialExpressionParameter* DestParam = (UMaterialExpressionParameter*)Destination;
-			DestParam->Group = SourceParam->Group;
-			DestParam->SortPriority = SourceParam->SortPriority;
-		}
-		else if (Source->IsA(UMaterialExpressionTextureSampleParameter::StaticClass()))
-		{
-			UMaterialExpressionTextureSampleParameter* SourceTex = (UMaterialExpressionTextureSampleParameter*)Source;
-			UMaterialExpressionTextureSampleParameter* DestTex = (UMaterialExpressionTextureSampleParameter*)Destination;
-			DestTex->Group = SourceTex->Group;
-			DestTex->SortPriority = SourceTex->SortPriority;
-		}
-		else if (Source->IsA(UMaterialExpressionRuntimeVirtualTextureSampleParameter::StaticClass()))
-		{
-			UMaterialExpressionRuntimeVirtualTextureSampleParameter* SourceRVT = (UMaterialExpressionRuntimeVirtualTextureSampleParameter*)Source;
-			UMaterialExpressionRuntimeVirtualTextureSampleParameter* DestRVT = (UMaterialExpressionRuntimeVirtualTextureSampleParameter*)Destination;
-			DestRVT->Group = SourceRVT->Group;
-			DestRVT->SortPriority = SourceRVT->SortPriority;
-		}
-		else if (Source->IsA(UMaterialExpressionFontSampleParameter::StaticClass()))
-		{
-			UMaterialExpressionFontSampleParameter* SourceFont = (UMaterialExpressionFontSampleParameter*)Source;
-			UMaterialExpressionFontSampleParameter* DestFont = (UMaterialExpressionFontSampleParameter*)Destination;
-			DestFont->Group = SourceFont->Group; // Group isn't editor only for fonts
-			DestFont->SortPriority = SourceFont->SortPriority;
-		}
-		else
-		{
-			ensureMsgf(false, TEXT("Unexpected expression parameter type %s"), *Source->GetClass()->GetName());
-		}
+		const EMaterialExpressionSetParameterValueFlags Flags = EMaterialExpressionSetParameterValueFlags::NoUpdateExpressionGuid | EMaterialExpressionSetParameterValueFlags::AssignGroupAndSortPriority;
+		verify(Destination->SetParameterValue(ParameterName, Meta, Flags));
 
 		Destination->Desc = Source->Desc;
 		Destination->GraphNode->OnUpdateCommentText(Destination->Desc);
@@ -5226,88 +5023,6 @@ bool UMaterial::GetExpressionsInPropertyChain(EMaterialProperty InProperty,
 	return true;
 }
 
-bool UMaterial::GetParameterSortPriority(const FHashedMaterialParameterInfo& ParameterInfo, int32& OutSortPriority, const TArray<FStaticMaterialLayersParameter>* MaterialLayersParameters) const
-{
-	for (UMaterialExpression* Expression : Expressions)
-	{
-		if (ParameterInfo.Association == EMaterialParameterAssociation::GlobalParameter)
-		{
-			UMaterialExpressionParameter* Parameter = Cast<UMaterialExpressionParameter>(Expression);
-			UMaterialExpressionTextureSampleParameter* TextureParameter = Cast<UMaterialExpressionTextureSampleParameter>(Expression);
-			UMaterialExpressionRuntimeVirtualTextureSampleParameter* RVTParameter = Cast<UMaterialExpressionRuntimeVirtualTextureSampleParameter>(Expression);
-			UMaterialExpressionFontSampleParameter* FontParameter = Cast<UMaterialExpressionFontSampleParameter>(Expression);
-			UMaterialExpressionMaterialAttributeLayers* LayersParameter = Cast<UMaterialExpressionMaterialAttributeLayers>(Expression);
-			UMaterialExpressionMaterialFunctionCall* FunctionCall = Cast<UMaterialExpressionMaterialFunctionCall>(Expression);
-
-			if (Parameter && Parameter->GetParameterName() == ParameterInfo.Name)
-			{
-				OutSortPriority = Parameter->SortPriority;
-				return true;
-			}
-			else if (TextureParameter && TextureParameter->GetParameterName() == ParameterInfo.Name)
-			{
-				OutSortPriority = TextureParameter->SortPriority;
-				return true;
-			}
-			else if (RVTParameter && RVTParameter->GetParameterName() == ParameterInfo.Name)
-			{
-				OutSortPriority = RVTParameter->SortPriority;
-				return true;
-			}
-			else if (FontParameter && FontParameter->GetParameterName() == ParameterInfo.Name)
-			{
-				OutSortPriority = FontParameter->SortPriority;
-				return true;
-			}
-			else if (LayersParameter && LayersParameter->GetParameterName() == ParameterInfo.Name)
-			{
-				OutSortPriority = 0; // No sorting as custom interface so we can end the search
-				return true;
-			}
-			else if (FunctionCall && FunctionCall->MaterialFunction)
-			{
-				if (FunctionCall->MaterialFunction->GetParameterSortPriority(ParameterInfo, OutSortPriority))
-				{
-					return true;
-				}
-			}
-		}
-		else
-		{
-			if (UMaterialExpressionMaterialAttributeLayers* LayersExpression = Cast<UMaterialExpressionMaterialAttributeLayers>(Expression))
-			{
-				UMaterialFunctionInterface* Function = nullptr;
-				
-				// Handle function overrides when searching for parameters
-				if (MaterialLayersParameters)
-				{
-					const FName& ParameterName = LayersExpression->ParameterName;
-					for (const FStaticMaterialLayersParameter& LayersParameter : *MaterialLayersParameters)
-					{
-						if (LayersParameter.ParameterInfo.Name == ParameterName)
-						{
-							Function = LayersParameter.GetParameterAssociatedFunction(ParameterInfo);
-							break;
-						}
-					}
-				}
-				
-				if (!Function)
-				{
-					Function = LayersExpression->GetParameterAssociatedFunction(ParameterInfo);
-				}
-
-				if (Function && Function->GetParameterSortPriority(ParameterInfo, OutSortPriority))
-				{
-					return true;
-				}
-			}
-		}
-	}
-
-	return false;
-}
-
 bool UMaterial::GetGroupSortPriority(const FString& InGroupName, int32& OutSortPriority) const
 {
 	const FParameterGroupData* ParameterGroupDataElement = ParameterGroupData.FindByPredicate([&InGroupName](const FParameterGroupData& DataElement)
diff --git a/Engine/Source/Runtime/Engine/Private/Materials/MaterialCachedData.cpp b/Engine/Source/Runtime/Engine/Private/Materials/MaterialCachedData.cpp
--- a/Engine/Source/Runtime/Engine/Private/Materials/MaterialCachedData.cpp
+++ b/Engine/Source/Runtime/Engine/Private/Materials/MaterialCachedData.cpp
@@ -17,15 +17,8 @@
 #include "Materials/MaterialExpressionParameter.h"
 #include "Materials/MaterialExpressionPerInstanceCustomData.h"
 #include "Materials/MaterialExpressionPerInstanceRandom.h"
-#include "Materials/MaterialExpressionRuntimeVirtualTextureSampleParameter.h"
-#include "Materials/MaterialExpressionScalarParameter.h"
-#include "Materials/MaterialExpressionStaticBoolParameter.h"
-#include "Materials/MaterialExpressionStaticComponentMaskParameter.h"
-#include "Materials/MaterialExpressionStaticSwitchParameter.h"
 #include "Materials/MaterialExpressionTextureBase.h"
 #include "Materials/MaterialExpressionTextureSample.h"
-#include "Materials/MaterialExpressionTextureSampleParameter.h"
-#include "Materials/MaterialExpressionVectorParameter.h"
 #include "Materials/MaterialExpressionVertexInterpolator.h"
 #include "Materials/MaterialExpressionSceneColor.h"
 #include "Materials/MaterialExpressionRuntimeVirtualTextureOutput.h"
@@ -38,6 +31,7 @@
 #include "Materials/MaterialFunctionInstance.h"
 #include "Materials/MaterialExpressionMaterialFunctionCall.h"
 #include "Materials/MaterialParameterCollection.h"
+#include "VT/RuntimeVirtualTexture.h"
 #include "Engine/Font.h"
 #include "LandscapeGrassType.h"
 #include "ProfilingDebugging/LoadTimeTracker.h"
@@ -102,7 +96,10 @@ void FMaterialCachedParameters::AddReferencedObjects(FReferenceCollector& Collec
 }
 
 #if WITH_EDITOR
-static int32 TryAddParameter(FMaterialCachedParameters& CachedParameters, EMaterialParameterType Type, const FMaterialParameterInfo& ParameterInfo, const FGuid& ExpressionGuid)
+static int32 TryAddParameter(FMaterialCachedParameters& CachedParameters,
+	EMaterialParameterType Type,
+	const FMaterialParameterInfo& ParameterInfo,
+	const FMaterialCachedParameterEditorInfo& InEditorInfo)
 {
 	FMaterialCachedParameterEntry& Entry = CachedParameters.GetParameterTypeEntry(Type);
 	FSetElementId ElementId = Entry.ParameterInfoSet.FindId(ParameterInfo);
@@ -111,19 +108,30 @@ static int32 TryAddParameter(FMaterialCachedParameters& CachedParameters, EMater
 	{
 		ElementId = Entry.ParameterInfoSet.Add(ParameterInfo);
 		Index = ElementId.AsInteger();
-		Entry.ExpressionGuids.Insert(ExpressionGuid, Index);
+		Entry.EditorInfo.Insert(InEditorInfo, Index);
 		// should be valid as long as we don't ever remove elements from ParameterInfoSet
-		check(Entry.ParameterInfoSet.Num() == Entry.ExpressionGuids.Num());
+		check(Entry.ParameterInfoSet.Num() == Entry.EditorInfo.Num());
 		return Index;
 	}
 
+	// Update any editor values that haven't been set yet (parameter overrides don't set all of these values)
 	Index = ElementId.AsInteger();
-	if (!Entry.ExpressionGuids[Index].IsValid())
+	FMaterialCachedParameterEditorInfo& EditorInfo = Entry.EditorInfo[Index];
+	if (!EditorInfo.ExpressionGuid.IsValid())
 	{
-		// If Parameter was set by a function override, update to a valid expression guid
-		Entry.ExpressionGuids[Index] = ExpressionGuid;
+		EditorInfo.ExpressionGuid = InEditorInfo.ExpressionGuid;
 	}
-	
+	if (EditorInfo.Description.IsEmpty())
+	{
+		EditorInfo.Description = InEditorInfo.Description;
+	}
+	if (EditorInfo.Group.IsNone())
+	{
+		EditorInfo.Group = InEditorInfo.Group;
+		EditorInfo.SortPriority = InEditorInfo.SortPriority;
+	}
+
+	// Still return INDEX_NONE, to signify this parameter was already added (don't want to add it again)
 	return INDEX_NONE;
 }
 
@@ -142,7 +150,7 @@ bool FMaterialCachedExpressionData::UpdateForFunction(const FMaterialCachedExpre
 		for (const FScalarParameterValue& Param : FunctionInstance->ScalarParameterValues)
 		{
 			const FMaterialParameterInfo ParameterInfo(Param.ParameterInfo.Name, Association, ParameterIndex);
-			const int32 Index = TryAddParameter(Parameters, EMaterialParameterType::Scalar, ParameterInfo, FGuid());
+			const int32 Index = TryAddParameter(Parameters, EMaterialParameterType::Scalar, ParameterInfo, FMaterialCachedParameterEditorInfo());
 			if (Index != INDEX_NONE)
 			{
 				Parameters.ScalarValues.Insert(Param.ParameterValue, Index);
@@ -163,7 +171,7 @@ bool FMaterialCachedExpressionData::UpdateForFunction(const FMaterialCachedExpre
 		for (const FVectorParameterValue& Param : FunctionInstance->VectorParameterValues)
 		{
 			const FMaterialParameterInfo ParameterInfo(Param.ParameterInfo.Name, Association, ParameterIndex);
-			const int32 Index = TryAddParameter(Parameters, EMaterialParameterType::Vector, ParameterInfo, FGuid());
+			const int32 Index = TryAddParameter(Parameters, EMaterialParameterType::Vector, ParameterInfo, FMaterialCachedParameterEditorInfo());
 			if (Index != INDEX_NONE)
 			{
 				Parameters.VectorValues.Insert(Param.ParameterValue, Index);
@@ -175,7 +183,7 @@ bool FMaterialCachedExpressionData::UpdateForFunction(const FMaterialCachedExpre
 		for (const FTextureParameterValue& Param : FunctionInstance->TextureParameterValues)
 		{
 			const FMaterialParameterInfo ParameterInfo(Param.ParameterInfo.Name, Association, ParameterIndex);
-			const int32 Index = TryAddParameter(Parameters, EMaterialParameterType::Texture, ParameterInfo, FGuid());
+			const int32 Index = TryAddParameter(Parameters, EMaterialParameterType::Texture, ParameterInfo, FMaterialCachedParameterEditorInfo());
 			if (Index != INDEX_NONE)
 			{
 				Parameters.TextureValues.Insert(Param.ParameterValue, Index);
@@ -186,7 +194,7 @@ bool FMaterialCachedExpressionData::UpdateForFunction(const FMaterialCachedExpre
 		for (const FRuntimeVirtualTextureParameterValue& Param : FunctionInstance->RuntimeVirtualTextureParameterValues)
 		{
 			const FMaterialParameterInfo ParameterInfo(Param.ParameterInfo.Name, Association, ParameterIndex);
-			const int32 Index = TryAddParameter(Parameters, EMaterialParameterType::RuntimeVirtualTexture, ParameterInfo, FGuid());
+			const int32 Index = TryAddParameter(Parameters, EMaterialParameterType::RuntimeVirtualTexture, ParameterInfo, FMaterialCachedParameterEditorInfo());
 			if (Index != INDEX_NONE)
 			{
 				Parameters.RuntimeVirtualTextureValues.Insert(Param.ParameterValue, Index);
@@ -196,7 +204,7 @@ bool FMaterialCachedExpressionData::UpdateForFunction(const FMaterialCachedExpre
 		for (const FFontParameterValue& Param : FunctionInstance->FontParameterValues)
 		{
 			const FMaterialParameterInfo ParameterInfo(Param.ParameterInfo.Name, Association, ParameterIndex);
-			const int32 Index = TryAddParameter(Parameters, EMaterialParameterType::Font, ParameterInfo, FGuid());
+			const int32 Index = TryAddParameter(Parameters, EMaterialParameterType::Font, ParameterInfo, FMaterialCachedParameterEditorInfo());
 			if (Index != INDEX_NONE)
 			{
 				Parameters.FontValues.Insert(Param.FontValue, Index);
@@ -207,7 +215,7 @@ bool FMaterialCachedExpressionData::UpdateForFunction(const FMaterialCachedExpre
 		for (const FStaticSwitchParameter& Param : FunctionInstance->StaticSwitchParameterValues)
 		{
 			const FMaterialParameterInfo ParameterInfo(Param.ParameterInfo.Name, Association, ParameterIndex);
-			const int32 Index = TryAddParameter(Parameters, EMaterialParameterType::StaticSwitch, ParameterInfo, FGuid());
+			const int32 Index = TryAddParameter(Parameters, EMaterialParameterType::StaticSwitch, ParameterInfo, FMaterialCachedParameterEditorInfo());
 			if (Index != INDEX_NONE)
 			{
 				Parameters.StaticSwitchValues.Insert(Param.Value, Index);
@@ -217,7 +225,7 @@ bool FMaterialCachedExpressionData::UpdateForFunction(const FMaterialCachedExpre
 		for (const FStaticComponentMaskParameter& Param : FunctionInstance->StaticComponentMaskParameterValues)
 		{
 			const FMaterialParameterInfo ParameterInfo(Param.ParameterInfo.Name, Association, ParameterIndex);
-			const int32 Index = TryAddParameter(Parameters, EMaterialParameterType::StaticComponentMask, ParameterInfo, FGuid());
+			const int32 Index = TryAddParameter(Parameters, EMaterialParameterType::StaticComponentMask, ParameterInfo, FMaterialCachedParameterEditorInfo());
 			if (Index != INDEX_NONE)
 			{
 				Parameters.StaticComponentMaskValues.Insert({Param.R,Param.G,Param.B,Param.A}, Index);
@@ -336,7 +344,7 @@ void FMaterialCachedParameters_UpdateForLayerParameters(FMaterialCachedParameter
 			FMaterialParameterInfo ParameterInfo;
 			if (GetLocalLayerParameterInfo(*ParentMaterialLayers, Param.ParameterInfo, LayerParameters.Value, ParameterInfo))
 			{
-				const int32 Index = TryAddParameter(Parameters, EMaterialParameterType::Scalar, ParameterInfo, FGuid());
+				const int32 Index = TryAddParameter(Parameters, EMaterialParameterType::Scalar, ParameterInfo, FMaterialCachedParameterEditorInfo());
 				if (Index != INDEX_NONE)
 				{
 					Parameters.ScalarValues.Insert(Param.ParameterValue, Index);
@@ -360,7 +368,7 @@ void FMaterialCachedParameters_UpdateForLayerParameters(FMaterialCachedParameter
 			FMaterialParameterInfo ParameterInfo;
 			if (GetLocalLayerParameterInfo(*ParentMaterialLayers, Param.ParameterInfo, LayerParameters.Value, ParameterInfo))
 			{
-				const int32 Index = TryAddParameter(Parameters, EMaterialParameterType::Vector, ParameterInfo, FGuid());
+				const int32 Index = TryAddParameter(Parameters, EMaterialParameterType::Vector, ParameterInfo, FMaterialCachedParameterEditorInfo());
 				if (Index != INDEX_NONE)
 				{
 					Parameters.VectorValues.Insert(Param.ParameterValue, Index);
@@ -375,7 +383,7 @@ void FMaterialCachedParameters_UpdateForLayerParameters(FMaterialCachedParameter
 			FMaterialParameterInfo ParameterInfo;
 			if (GetLocalLayerParameterInfo(*ParentMaterialLayers, Param.ParameterInfo, LayerParameters.Value, ParameterInfo))
 			{
-				const int32 Index = TryAddParameter(Parameters, EMaterialParameterType::Texture, ParameterInfo, FGuid());
+				const int32 Index = TryAddParameter(Parameters, EMaterialParameterType::Texture, ParameterInfo, FMaterialCachedParameterEditorInfo());
 				if (Index != INDEX_NONE)
 				{
 					Parameters.TextureValues.Insert(Param.ParameterValue, Index);
@@ -389,7 +397,7 @@ void FMaterialCachedParameters_UpdateForLayerParameters(FMaterialCachedParameter
 			FMaterialParameterInfo ParameterInfo;
 			if (GetLocalLayerParameterInfo(*ParentMaterialLayers, Param.ParameterInfo, LayerParameters.Value, ParameterInfo))
 			{
-				const int32 Index = TryAddParameter(Parameters, EMaterialParameterType::RuntimeVirtualTexture, ParameterInfo, FGuid());
+				const int32 Index = TryAddParameter(Parameters, EMaterialParameterType::RuntimeVirtualTexture, ParameterInfo, FMaterialCachedParameterEditorInfo());
 				if (Index != INDEX_NONE)
 				{
 					Parameters.RuntimeVirtualTextureValues.Insert(Param.ParameterValue, Index);
@@ -402,7 +410,7 @@ void FMaterialCachedParameters_UpdateForLayerParameters(FMaterialCachedParameter
 			FMaterialParameterInfo ParameterInfo;
 			if (GetLocalLayerParameterInfo(*ParentMaterialLayers, Param.ParameterInfo, LayerParameters.Value, ParameterInfo))
 			{
-				const int32 Index = TryAddParameter(Parameters, EMaterialParameterType::Font, ParameterInfo, FGuid());
+				const int32 Index = TryAddParameter(Parameters, EMaterialParameterType::Font, ParameterInfo, FMaterialCachedParameterEditorInfo());
 				if (Index != INDEX_NONE)
 				{
 					Parameters.FontValues.Insert(Param.FontValue, Index);
@@ -416,7 +424,7 @@ void FMaterialCachedParameters_UpdateForLayerParameters(FMaterialCachedParameter
 			FMaterialParameterInfo ParameterInfo;
 			if (GetLocalLayerParameterInfo(*ParentMaterialLayers, Param.ParameterInfo, LayerParameters.Value, ParameterInfo))
 			{
-				const int32 Index = TryAddParameter(Parameters, EMaterialParameterType::StaticSwitch, ParameterInfo, FGuid());
+				const int32 Index = TryAddParameter(Parameters, EMaterialParameterType::StaticSwitch, ParameterInfo, FMaterialCachedParameterEditorInfo());
 				if (Index != INDEX_NONE)
 				{
 					Parameters.StaticSwitchValues.Insert(Param.Value, Index);
@@ -429,7 +437,7 @@ void FMaterialCachedParameters_UpdateForLayerParameters(FMaterialCachedParameter
 			FMaterialParameterInfo ParameterInfo;
 			if (GetLocalLayerParameterInfo(*ParentMaterialLayers, Param.ParameterInfo, LayerParameters.Value, ParameterInfo))
 			{
-				const int32 Index = TryAddParameter(Parameters, EMaterialParameterType::StaticComponentMask, ParameterInfo, FGuid());
+				const int32 Index = TryAddParameter(Parameters, EMaterialParameterType::StaticComponentMask, ParameterInfo, FMaterialCachedParameterEditorInfo());
 				if (Index != INDEX_NONE)
 				{
 					const FStaticComponentMaskValue Value(Param.R, Param.G, Param.B, Param.A);
@@ -458,91 +466,60 @@ bool FMaterialCachedExpressionData::UpdateForExpressions(const FMaterialCachedEx
 			ReferencedTextures.AddUnique(ReferencedTexture);
 		}
 
-		if (UMaterialExpressionScalarParameter* ExpressionScalarParameter = Cast<UMaterialExpressionScalarParameter>(Expression))
+		FMaterialParameterMetadata ParameterMeta;
+		if (Expression->GetParameterValue(ParameterMeta))
 		{
-			const FMaterialParameterInfo ParameterInfo(ExpressionScalarParameter->GetParameterName(), Association, ParameterIndex);
-			const int32 Index = TryAddParameter(Parameters, EMaterialParameterType::Scalar, ParameterInfo, ExpressionScalarParameter->ExpressionGUID);
+			const FMaterialParameterInfo ParameterInfo(Expression->GetParameterName(), Association, ParameterIndex);
+			const FMaterialCachedParameterEditorInfo EditorInfo(ParameterMeta.ExpressionGuid, ParameterMeta.Description, ParameterMeta.Group, ParameterMeta.SortPriority);
+			const int32 Index = TryAddParameter(Parameters, ParameterMeta.Value.Type, ParameterInfo, EditorInfo);
 			if (Index != INDEX_NONE)
 			{
-				Parameters.ScalarValues.Insert(ExpressionScalarParameter->DefaultValue, Index);
-				Parameters.ScalarMinMaxValues.Insert(FVector2D(ExpressionScalarParameter->SliderMin, ExpressionScalarParameter->SliderMax), Index);
-				if (ExpressionScalarParameter->IsUsedAsAtlasPosition())
-				{
-					UMaterialExpressionCurveAtlasRowParameter* ExpressionAtlasParameter = Cast<UMaterialExpressionCurveAtlasRowParameter>(ExpressionScalarParameter);
-					Parameters.ScalarCurveValues.Insert(ExpressionAtlasParameter->Curve, Index);
-					Parameters.ScalarCurveAtlasValues.Insert(ExpressionAtlasParameter->Atlas, Index);
-				}
-				else
+				switch (ParameterMeta.Value.Type)
 				{
-					Parameters.ScalarCurveValues.Insert(nullptr, Index);
-					Parameters.ScalarCurveAtlasValues.Insert(nullptr, Index);
+				case EMaterialParameterType::Scalar:
+					Parameters.ScalarValues.Insert(ParameterMeta.Value.AsScalar(), Index);
+					Parameters.ScalarMinMaxValues.Insert(FVector2D(ParameterMeta.ScalarMin, ParameterMeta.ScalarMax), Index);
+					if (ParameterMeta.bUsedAsAtlasPosition)
+					{
+						Parameters.ScalarCurveValues.Insert(ParameterMeta.ScalarCurve.Get(), Index);
+						Parameters.ScalarCurveAtlasValues.Insert(ParameterMeta.ScalarAtlas.Get(), Index);
+					}
+					else
+					{
+						Parameters.ScalarCurveValues.Insert(nullptr, Index);
+						Parameters.ScalarCurveAtlasValues.Insert(nullptr, Index);
+					}
+					break;
+				case EMaterialParameterType::Vector:
+					Parameters.VectorValues.Insert(ParameterMeta.Value.AsLinearColor(), Index);
+					Parameters.VectorChannelNameValues.Insert(ParameterMeta.ChannelNames, Index);
+					Parameters.VectorUsedAsChannelMaskValues.Insert(ParameterMeta.bUsedAsChannelMask, Index);
+					break;
+				case EMaterialParameterType::Texture:
+					Parameters.TextureValues.Insert(ParameterMeta.Value.Texture, Index);
+					Parameters.TextureChannelNameValues.Insert(ParameterMeta.ChannelNames, Index);
+					break;
+				case EMaterialParameterType::Font:
+					Parameters.FontValues.Insert(ParameterMeta.Value.Font.Value, Index);
+					Parameters.FontPageValues.Insert(ParameterMeta.Value.Font.Page, Index);
+					break;
+				case EMaterialParameterType::RuntimeVirtualTexture:
+					Parameters.RuntimeVirtualTextureValues.Insert(ParameterMeta.Value.RuntimeVirtualTexture, Index);
+					break;
+				case EMaterialParameterType::StaticSwitch:
+					Parameters.StaticSwitchValues.Insert(ParameterMeta.Value.AsStaticSwitch(), Index);
+					break;
+				case EMaterialParameterType::StaticComponentMask:
+					Parameters.StaticComponentMaskValues.Insert(ParameterMeta.Value.AsStaticComponentMask(), Index);
+					break;
+				default:
+					checkNoEntry();
+					break;
 				}
 			}
 		}
-		else if (UMaterialExpressionVectorParameter* ExpressionVectorParameter = Cast<UMaterialExpressionVectorParameter>(Expression))
-		{
-			const FMaterialParameterInfo ParameterInfo(ExpressionVectorParameter->GetParameterName(), Association, ParameterIndex);
-			const int32 Index = TryAddParameter(Parameters, EMaterialParameterType::Vector, ParameterInfo, ExpressionVectorParameter->ExpressionGUID);
-			if (Index != INDEX_NONE)
-			{
-				Parameters.VectorValues.Insert(ExpressionVectorParameter->DefaultValue, Index);
-				Parameters.VectorChannelNameValues.Insert(ExpressionVectorParameter->ChannelNames, Index);
-				Parameters.VectorUsedAsChannelMaskValues.Insert(ExpressionVectorParameter->IsUsedAsChannelMask(), Index);
-			}
-		}
-		else if (UMaterialExpressionTextureSampleParameter* ExpressionTextureParameter = Cast<UMaterialExpressionTextureSampleParameter>(Expression))
-		{
-			const FMaterialParameterInfo ParameterInfo(ExpressionTextureParameter->GetParameterName(), Association, ParameterIndex);
-			const int32 Index = TryAddParameter(Parameters, EMaterialParameterType::Texture, ParameterInfo, ExpressionTextureParameter->ExpressionGUID);
-			if (Index != INDEX_NONE)
-			{
-				Parameters.TextureValues.Insert(ExpressionTextureParameter->Texture, Index);
-				Parameters.TextureChannelNameValues.Insert(ExpressionTextureParameter->ChannelNames, Index);
-			}
-		}
-		else if (UMaterialExpressionFontSampleParameter* ExpressionFontParameter = Cast<UMaterialExpressionFontSampleParameter>(Expression))
-		{
-			const FMaterialParameterInfo ParameterInfo(ExpressionFontParameter->GetParameterName(), Association, ParameterIndex);
-			const int32 Index = TryAddParameter(Parameters, EMaterialParameterType::Font, ParameterInfo, ExpressionFontParameter->ExpressionGUID);
-			if (Index != INDEX_NONE)
-			{
-				Parameters.FontValues.Insert(ExpressionFontParameter->Font, Index);
-				Parameters.FontPageValues.Insert(ExpressionFontParameter->FontTexturePage, Index);
-			}
-		}
-		else if (UMaterialExpressionRuntimeVirtualTextureSampleParameter* ExpressionRTVParameter = Cast<UMaterialExpressionRuntimeVirtualTextureSampleParameter>(Expression))
-		{
-			const FMaterialParameterInfo ParameterInfo(ExpressionRTVParameter->GetParameterName(), Association, ParameterIndex);
-			const int32 Index = TryAddParameter(Parameters, EMaterialParameterType::RuntimeVirtualTexture, ParameterInfo, ExpressionRTVParameter->ExpressionGUID);
-			if (Index != INDEX_NONE)
-			{
-				Parameters.RuntimeVirtualTextureValues.Insert(ExpressionRTVParameter->VirtualTexture, Index);
-			}
-		}
-		else if (UMaterialExpressionStaticBoolParameter* ExpressionStaticBoolParameter = Cast<UMaterialExpressionStaticBoolParameter>(Expression))
-		{
-			const FMaterialParameterInfo ParameterInfo(ExpressionStaticBoolParameter->GetParameterName(), Association, ParameterIndex);
-			const int32 Index = TryAddParameter(Parameters, EMaterialParameterType::StaticSwitch, ParameterInfo, ExpressionStaticBoolParameter->ExpressionGUID);
-			if (Index != INDEX_NONE)
-			{
-				Parameters.StaticSwitchValues.Insert(ExpressionStaticBoolParameter->DefaultValue, Index);
-			}
-		}
-		else if (UMaterialExpressionStaticComponentMaskParameter* ExpressionStaticComponentMaskParameter = Cast<UMaterialExpressionStaticComponentMaskParameter>(Expression))
-		{
-			const FMaterialParameterInfo ParameterInfo(ExpressionStaticComponentMaskParameter->GetParameterName(), Association, ParameterIndex);
-			const int32 Index = TryAddParameter(Parameters, EMaterialParameterType::StaticComponentMask, ParameterInfo, ExpressionStaticComponentMaskParameter->ExpressionGUID);
-			if (Index != INDEX_NONE)
-			{
-				const FStaticComponentMaskValue Value(
-					static_cast<bool>(ExpressionStaticComponentMaskParameter->DefaultR),
-					static_cast<bool>(ExpressionStaticComponentMaskParameter->DefaultG),
-					static_cast<bool>(ExpressionStaticComponentMaskParameter->DefaultB),
-					static_cast<bool>(ExpressionStaticComponentMaskParameter->DefaultA));
-				Parameters.StaticComponentMaskValues.Insert(Value, Index);
-			}
-		}
-		else if (UMaterialExpressionCollectionParameter* ExpressionCollectionParameter = Cast<UMaterialExpressionCollectionParameter>(Expression))
+
+		if (UMaterialExpressionCollectionParameter* ExpressionCollectionParameter = Cast<UMaterialExpressionCollectionParameter>(Expression))
 		{
 			UMaterialParameterCollection* Collection = ExpressionCollectionParameter->Collection;
 			if (Collection)
@@ -650,9 +627,9 @@ bool FMaterialCachedExpressionData::UpdateForExpressions(const FMaterialCachedEx
 					SetMaterialAttributePropertyConnected(MaterialProperty, AttributeInput.Expression ? true : false);
 				}
 			}
-			}
-			else if (UMaterialExpressionMakeMaterialAttributes* MakeMatAttributes = Cast<UMaterialExpressionMakeMaterialAttributes>(Expression))
-			{
+		}
+		else if (UMaterialExpressionMakeMaterialAttributes* MakeMatAttributes = Cast<UMaterialExpressionMakeMaterialAttributes>(Expression))
+		{
 			// Only set the material property if it hasn't been set yet.  We want to specifically avoid a Set Material Attributes node which doesn't have a 
 			// attribute set from disabling the attribute from a different Set Material Attributes node which does have it enabled.
 			auto SetMatAttributeConditionally = [&](EMaterialProperty InMaterialProperty, bool InIsConnected)
@@ -700,7 +677,7 @@ void FMaterialCachedParameterEntry::Reset()
 {
 	ParameterInfoSet.Reset();
 #if WITH_EDITORONLY_DATA
-	ExpressionGuids.Reset();
+	EditorInfo.Reset();
 #endif
 }
 
@@ -748,7 +725,11 @@ void FMaterialCachedParameters::GetParameterValueByIndex(EMaterialParameterType
 	const FMaterialCachedParameterEntry& Entry = GetParameterTypeEntry(Type);
 
 #if WITH_EDITORONLY_DATA
-	OutResult.ExpressionGuid = Entry.ExpressionGuids[ParameterIndex];
+	const FMaterialCachedParameterEditorInfo& EditorInfo = Entry.EditorInfo[ParameterIndex];
+	OutResult.ExpressionGuid = EditorInfo.ExpressionGuid;
+	OutResult.Description = EditorInfo.Description;
+	OutResult.Group = EditorInfo.Group;
+	OutResult.SortPriority = EditorInfo.SortPriority;
 #endif
 
 	switch (Type)
@@ -773,20 +754,14 @@ void FMaterialCachedParameters::GetParameterValueByIndex(EMaterialParameterType
 	case EMaterialParameterType::Vector:
 		OutResult.Value = VectorValues[ParameterIndex];
 #if  WITH_EDITORONLY_DATA
-		OutResult.ChannelName[0] = VectorChannelNameValues[ParameterIndex].R;
-		OutResult.ChannelName[1] = VectorChannelNameValues[ParameterIndex].G;
-		OutResult.ChannelName[2] = VectorChannelNameValues[ParameterIndex].B;
-		OutResult.ChannelName[3] = VectorChannelNameValues[ParameterIndex].A;
+		OutResult.ChannelNames = VectorChannelNameValues[ParameterIndex];
 		OutResult.bUsedAsChannelMask = VectorUsedAsChannelMaskValues[ParameterIndex];
 #endif // WITH_EDITORONLY_DATA
 		break;
 	case EMaterialParameterType::Texture:
 		OutResult.Value = TextureValues[ParameterIndex];
 #if WITH_EDITORONLY_DATA
-		OutResult.ChannelName[0] = TextureChannelNameValues[ParameterIndex].R;
-		OutResult.ChannelName[1] = TextureChannelNameValues[ParameterIndex].G;
-		OutResult.ChannelName[2] = TextureChannelNameValues[ParameterIndex].B;
-		OutResult.ChannelName[3] = TextureChannelNameValues[ParameterIndex].A;
+		OutResult.ChannelNames = TextureChannelNameValues[ParameterIndex];
 #endif // WITH_EDITORONLY_DATA
 		break;
 	case EMaterialParameterType::RuntimeVirtualTexture:
@@ -825,7 +800,7 @@ bool FMaterialCachedParameters::GetParameterValue(EMaterialParameterType Type, c
 const FGuid& FMaterialCachedParameters::GetExpressionGuid(EMaterialParameterType Type, int32 Index) const
 {
 	const FMaterialCachedParameterEntry& Entry = GetParameterTypeEntry(Type);
-	return Entry.ExpressionGuids[Index];
+	return Entry.EditorInfo[Index].ExpressionGuid;
 }
 #endif // WITH_EDITORONLY_DATA
 
@@ -854,7 +829,7 @@ void FMaterialCachedParameters::GetAllParameterInfoOfType(EMaterialParameterType
 	{
 		OutParameterInfo.Add(*It);
 #if WITH_EDITORONLY_DATA
-		OutParameterIds.Add(Entry.ExpressionGuids[It.GetId().AsInteger()]);
+		OutParameterIds.Add(Entry.EditorInfo[It.GetId().AsInteger()].ExpressionGuid);
 #else
 		OutParameterIds.Add(FGuid());
 #endif
@@ -875,7 +850,7 @@ void FMaterialCachedParameters::GetAllGlobalParameterInfoOfType(EMaterialParamet
 		{
 			OutParameterInfo.Add(ParameterInfo);
 #if WITH_EDITORONLY_DATA
-			OutParameterIds.Add(Entry.ExpressionGuids[It.GetId().AsInteger()]);
+			OutParameterIds.Add(Entry.EditorInfo[It.GetId().AsInteger()].ExpressionGuid);
 #else
 			OutParameterIds.Add(FGuid());
 #endif
diff --git a/Engine/Source/Runtime/Engine/Private/Materials/MaterialInstance.cpp b/Engine/Source/Runtime/Engine/Private/Materials/MaterialInstance.cpp
--- a/Engine/Source/Runtime/Engine/Private/Materials/MaterialInstance.cpp
+++ b/Engine/Source/Runtime/Engine/Private/Materials/MaterialInstance.cpp
@@ -1863,56 +1863,6 @@ void UMaterialInstance::GetDependentFunctions(TArray<UMaterialFunctionInterface*
 }
 #endif // WITH_EDITORONLY_DATA
 
-#if WITH_EDITOR
-bool UMaterialInstance::GetGroupName(const FHashedMaterialParameterInfo& ParameterInfo, FName& OutGroup) const
-{
-	if (GetReentrantFlag())
-	{
-		return false;
-	}
-
-	// @TODO: Alter to match sort priority behavior?
-	for (const FStaticMaterialLayersParameter& Param : StaticParameters.MaterialLayersParameters)
-	{
-		if (Param.bOverride)
-		{
-			if (ParameterInfo.Association == EMaterialParameterAssociation::LayerParameter)
-			{
-				if(Param.Value.Layers.IsValidIndex(ParameterInfo.Index))
-				{
-					UMaterialFunctionInterface* Layer = Param.Value.Layers[ParameterInfo.Index];
-
-					if (Layer && Layer->GetParameterGroupName(ParameterInfo, OutGroup))
-					{
-						return true;
-					}
-				}
-			}
-			else if (ParameterInfo.Association == EMaterialParameterAssociation::BlendParameter)
-			{
-				if(Param.Value.Blends.IsValidIndex(ParameterInfo.Index))
-				{
-					UMaterialFunctionInterface* Blend = Param.Value.Blends[ParameterInfo.Index];
-
-					if (Blend && Blend->GetParameterGroupName(ParameterInfo, OutGroup))
-					{
-						return true;
-					}
-				}
-			}
-		}
-	}
-
-	if (Parent)
-	{
-		FMICReentranceGuard	Guard(this);
-		Parent->GetGroupName(ParameterInfo, OutGroup);
-	}
-
-	return false;
-}
-#endif // WITH_EDITOR
-
 #if WITH_EDITOR
 void UMaterialInstance::ForceRecompileForRendering()
 {
@@ -3536,28 +3486,6 @@ float UMaterialInstance::GetExportResolutionScale() const
 }
 
 #if WITH_EDITOR
-bool UMaterialInstance::GetParameterDesc(const FHashedMaterialParameterInfo& ParameterInfo, FString& OutDesc, const TArray<struct FStaticMaterialLayersParameter>*) const
-{
-	const UMaterial* BaseMaterial = GetMaterial();
-	if (BaseMaterial && BaseMaterial->GetParameterDesc(ParameterInfo, OutDesc, &StaticParameters.MaterialLayersParameters))
-	{
-		return true;
-	}
-
-	return false;
-}
-
-bool UMaterialInstance::GetParameterSortPriority(const FHashedMaterialParameterInfo& ParameterInfo, int32& OutSortPriority, const TArray<struct FStaticMaterialLayersParameter>*) const
-{
-	const UMaterial* BaseMaterial = GetMaterial();
-	if (BaseMaterial && BaseMaterial->GetParameterSortPriority(ParameterInfo, OutSortPriority, &StaticParameters.MaterialLayersParameters))
-	{
-		return true;
-	}
-
-	return false;
-}
-
 bool UMaterialInstance::GetGroupSortPriority(const FString& InGroupName, int32& OutSortPriority) const
 {
 	// @TODO: This needs to handle overridden functions, layers and blends
@@ -3569,7 +3497,6 @@ bool UMaterialInstance::GetGroupSortPriority(const FString& InGroupName, int32&
 
 	return false;
 }
-
 bool UMaterialInstance::GetTexturesInPropertyChain(EMaterialProperty InProperty, TArray<UTexture*>& OutTextures,  
 	TArray<FName>* OutTextureParamNames, struct FStaticParameterSet* InStaticParameterSet,
 	ERHIFeatureLevel::Type InFeatureLevel, EMaterialQualityLevel::Type InQuality)
diff --git a/Engine/Source/Runtime/Engine/Private/Materials/MaterialInterface.cpp b/Engine/Source/Runtime/Engine/Private/Materials/MaterialInterface.cpp
--- a/Engine/Source/Runtime/Engine/Private/Materials/MaterialInterface.cpp
+++ b/Engine/Source/Runtime/Engine/Private/Materials/MaterialInterface.cpp
@@ -441,10 +441,7 @@ bool UMaterialInterface::GetVectorParameterChannelNames(const FHashedMaterialPar
 	FMaterialParameterMetadata Result;
 	if (GetParameterValue(EMaterialParameterType::Vector, ParameterInfo, Result, EMaterialGetParameterValueFlags::CheckNonOverrides | EMaterialGetParameterValueFlags::CheckParent))
 	{
-		OutValue.R = Result.ChannelName[0];
-		OutValue.G = Result.ChannelName[1];
-		OutValue.B = Result.ChannelName[2];
-		OutValue.A = Result.ChannelName[3];
+		OutValue = Result.ChannelNames;
 		return true;
 	}
 	return false;
@@ -522,10 +519,7 @@ bool UMaterialInterface::GetTextureParameterChannelNames(const FHashedMaterialPa
 	FMaterialParameterMetadata Result;
 	if (GetParameterValue(EMaterialParameterType::Texture, ParameterInfo, Result, EMaterialGetParameterValueFlags::CheckNonOverrides | EMaterialGetParameterValueFlags::CheckParent))
 	{
-		OutValue.R = Result.ChannelName[0];
-		OutValue.G = Result.ChannelName[1];
-		OutValue.B = Result.ChannelName[2];
-		OutValue.A = Result.ChannelName[3];
+		OutValue = Result.ChannelNames;
 		return true;
 	}
 	return false;
@@ -680,13 +674,45 @@ bool UMaterialInterface::GetRefractionSettings(float& OutBiasValue) const
 }
 
 #if WITH_EDITOR
-bool UMaterialInterface::GetParameterDesc(const FHashedMaterialParameterInfo& ParameterInfo, FString& OutDesc, const TArray<struct FStaticMaterialLayersParameter>* MaterialLayersParameters) const
+bool UMaterialInterface::GetParameterDesc(const FHashedMaterialParameterInfo& ParameterInfo, FString& OutDesc) const
 {
+	for (int32 TypeIndex = 0; TypeIndex < NumMaterialParameterTypes; ++TypeIndex)
+	{
+		FMaterialParameterMetadata Meta;
+		if (GetParameterValue((EMaterialParameterType)TypeIndex, ParameterInfo, Meta, EMaterialGetParameterValueFlags::CheckAll))
+		{
+			OutDesc = Meta.Description;
+			return true;
+		}
+	}
 	return false;
 }
 
 bool UMaterialInterface::GetGroupName(const FHashedMaterialParameterInfo& ParameterInfo, FName& OutDesc) const
 {
+	for (int32 TypeIndex = 0; TypeIndex < NumMaterialParameterTypes; ++TypeIndex)
+	{
+		FMaterialParameterMetadata Meta;
+		if (GetParameterValue((EMaterialParameterType)TypeIndex, ParameterInfo, Meta, EMaterialGetParameterValueFlags::CheckAll))
+		{
+			OutDesc = Meta.Group;
+			return true;
+		}
+	}
+	return false;
+}
+
+bool UMaterialInterface::GetParameterSortPriority(const FHashedMaterialParameterInfo& ParameterInfo, int32& OutSortPriority) const
+{
+	for (int32 TypeIndex = 0; TypeIndex < NumMaterialParameterTypes; ++TypeIndex)
+	{
+		FMaterialParameterMetadata Meta;
+		if (GetParameterValue((EMaterialParameterType)TypeIndex, ParameterInfo, Meta, EMaterialGetParameterValueFlags::CheckAll))
+		{
+			OutSortPriority = Meta.SortPriority;
+			return true;
+		}
+	}
 	return false;
 }
 #endif // WITH_EDITOR
diff --git a/Engine/Source/Runtime/Engine/Public/MaterialCachedData.h b/Engine/Source/Runtime/Engine/Public/MaterialCachedData.h
--- a/Engine/Source/Runtime/Engine/Public/MaterialCachedData.h
+++ b/Engine/Source/Runtime/Engine/Public/MaterialCachedData.h
@@ -56,21 +56,30 @@ struct FMaterialParameterCollectionInfo
 };
 
 USTRUCT()
-struct FParameterChannelNames
+struct FMaterialCachedParameterEditorInfo
 {
 	GENERATED_USTRUCT_BODY()
 
-	UPROPERTY(EditAnywhere, Category = MaterialExpressionVectorParameter)
-	FText R;
+	FMaterialCachedParameterEditorInfo() = default;
+	FMaterialCachedParameterEditorInfo(const FGuid& InGuid) : ExpressionGuid(InGuid) {}
+	FMaterialCachedParameterEditorInfo(const FGuid& InGuid, const FString& InDescription, const FName& InGroup, int32 InSortPriority)
+		: Description(InDescription)
+		, Group(InGroup)
+		, SortPriority(InSortPriority)
+		, ExpressionGuid(InGuid)
+	{}
+
+	UPROPERTY()
+	FString Description;
 
-	UPROPERTY(EditAnywhere, Category = MaterialExpressionVectorParameter)
-	FText G;
+	UPROPERTY()
+	FName Group;
 
-	UPROPERTY(EditAnywhere, Category = MaterialExpressionVectorParameter)
-	FText B;
+	UPROPERTY()
+	int32 SortPriority = 0;
 
-	UPROPERTY(EditAnywhere, Category = MaterialExpressionVectorParameter)
-	FText A;
+	UPROPERTY()
+	FGuid ExpressionGuid;
 };
 
 USTRUCT()
@@ -88,8 +97,9 @@ struct FMaterialCachedParameterEntry
 	TSet<FMaterialParameterInfo> ParameterInfoSet;
 
 #if WITH_EDITORONLY_DATA
+	// Editor-only information for each parameter
 	UPROPERTY()
-	TArray<FGuid> ExpressionGuids;
+	TArray<FMaterialCachedParameterEditorInfo> EditorInfo;
 #endif // WITH_EDITORONLY_DATA
 
 };
diff --git a/Engine/Source/Runtime/Engine/Public/MaterialTypes.h b/Engine/Source/Runtime/Engine/Public/MaterialTypes.h
--- a/Engine/Source/Runtime/Engine/Public/MaterialTypes.h
+++ b/Engine/Source/Runtime/Engine/Public/MaterialTypes.h
@@ -72,6 +72,27 @@ enum class EMaterialSetParameterValueFlags : uint32
 };
 ENUM_CLASS_FLAGS(EMaterialSetParameterValueFlags);
 
+USTRUCT()
+struct FParameterChannelNames
+{
+	GENERATED_USTRUCT_BODY()
+
+	FParameterChannelNames() = default;
+	FParameterChannelNames(const FText& InR, const FText& InG, const FText& InB, const FText& InA) : R(InR), G(InG), B(InB), A(InA) {}
+
+	UPROPERTY(EditAnywhere, Category = MaterialExpressionVectorParameter)
+	FText R;
+
+	UPROPERTY(EditAnywhere, Category = MaterialExpressionVectorParameter)
+	FText G;
+
+	UPROPERTY(EditAnywhere, Category = MaterialExpressionVectorParameter)
+	FText B;
+
+	UPROPERTY(EditAnywhere, Category = MaterialExpressionVectorParameter)
+	FText A;
+};
+
 USTRUCT()
 struct FStaticComponentMaskValue
 {
@@ -184,16 +205,25 @@ struct FMaterialParameterMetadata
 
 #if WITH_EDITORONLY_DATA
 	/** Name of channels, for Vectors/Textures */
-	FText ChannelName[4];
+	FParameterChannelNames ChannelNames;
 
 	/** Curve/Atlas used to generate scalar value */
 	TSoftObjectPtr<class UCurveLinearColor> ScalarCurve;
 	TSoftObjectPtr<class UCurveLinearColorAtlas> ScalarAtlas;
 
+	/** Description of the parameter, typically taken from the 'Desc' field of the parameter's UMaterialExpression */
+	FString Description;
+
+	/** Name of the parameter's group */
+	FName Group;
+
 	/** UI range for scalar values */
 	float ScalarMin = 0.0f;
 	float ScalarMax = 0.0f;
 
+	/** Used for sorting parameter within the group, in the UI */
+	int32 SortPriority = 0;
+
 	/** GUID of the UMaterialExpression this parameter came from */
 	FGuid ExpressionGuid;
 
