Missed files from previous submit

diff --git a/Engine/Source/Runtime/Engine/Classes/Materials/MaterialExpression.h b/Engine/Source/Runtime/Engine/Classes/Materials/MaterialExpression.h
--- a/Engine/Source/Runtime/Engine/Classes/Materials/MaterialExpression.h
+++ b/Engine/Source/Runtime/Engine/Classes/Materials/MaterialExpression.h
@@ -154,6 +154,7 @@ enum class EMaterialExpressionSetParameterValueFlags : uint32
 	None = 0u,
 	SendPostEditChangeProperty = (1u << 0), // Send PostEditChangeProperty events for all properties that are modified
 	NoUpdateExpressionGuid = (1u << 1), // By default ExpressionGUI will be updated for static parameters
+	AssignGroupAndSortPriority  = (1u << 2), //  Update the Group and SortPriority along with parameter value
 };
 ENUM_CLASS_FLAGS(EMaterialExpressionSetParameterValueFlags);
 
diff --git a/Engine/Source/Runtime/Engine/Private/Materials/MaterialExpressions.cpp b/Engine/Source/Runtime/Engine/Private/Materials/MaterialExpressions.cpp
--- a/Engine/Source/Runtime/Engine/Private/Materials/MaterialExpressions.cpp
+++ b/Engine/Source/Runtime/Engine/Private/Materials/MaterialExpressions.cpp
@@ -2712,6 +2712,14 @@ void UMaterialExpressionRuntimeVirtualTextureSampleParameter::GetAllParameterInf
 
 #if WITH_EDITOR
 
+template<typename T>
+static void SendPostEditChangeProperty(T* Object, const FName& Name)
+{
+	FProperty* Property = FindFProperty<FProperty>(T::StaticClass(), Name);
+	FPropertyChangedEvent Event(Property);
+	Object->PostEditChangeProperty(Event);
+}
+
 bool UMaterialExpressionRuntimeVirtualTextureSampleParameter::SetParameterValue(FName InParameterName, URuntimeVirtualTexture* InValue, EMaterialExpressionSetParameterValueFlags Flags)
 {
 	if (InParameterName == ParameterName)
@@ -2719,9 +2727,7 @@ bool UMaterialExpressionRuntimeVirtualTextureSampleParameter::SetParameterValue(
 		VirtualTexture = InValue;
 		if (EnumHasAnyFlags(Flags, EMaterialExpressionSetParameterValueFlags::SendPostEditChangeProperty))
 		{
-			FProperty* ParamProperty = FindFProperty<FProperty>(UMaterialExpressionRuntimeVirtualTextureSampleParameter::StaticClass(), GET_MEMBER_NAME_STRING_CHECKED(UMaterialExpressionRuntimeVirtualTextureSampleParameter, VirtualTexture));
-			FPropertyChangedEvent PropertyChangedEvent(ParamProperty);
-			PostEditChangeProperty(PropertyChangedEvent);
+			SendPostEditChangeProperty(this, TEXT("VirtualTexture"));
 		}
 		return true;
 	}
@@ -2990,9 +2996,7 @@ bool UMaterialExpressionTextureSampleParameter::SetParameterValue(FName InParame
 		Texture = InValue;
 		if (EnumHasAnyFlags(Flags, EMaterialExpressionSetParameterValueFlags::SendPostEditChangeProperty))
 		{
-			FProperty* ParamProperty = FindFProperty<FProperty>(UMaterialExpressionTextureSampleParameter::StaticClass(), GET_MEMBER_NAME_STRING_CHECKED(UMaterialExpressionTextureSampleParameter, Texture));
-			FPropertyChangedEvent PropertyChangedEvent(ParamProperty);
-			PostEditChangeProperty(PropertyChangedEvent);
+			SendPostEditChangeProperty(this, TEXT("Texture"));
 		}
 		return true;
 	}
@@ -4941,18 +4945,10 @@ bool  UMaterialExpressionStaticComponentMaskParameter::SetParameterValue(FName I
 		}
 		if (EnumHasAnyFlags(Flags, EMaterialExpressionSetParameterValueFlags::SendPostEditChangeProperty))
 		{
-			FProperty* ParamProperty = FindFProperty<FProperty>(UMaterialExpressionStaticComponentMaskParameter::StaticClass(), GET_MEMBER_NAME_STRING_CHECKED(UMaterialExpressionStaticComponentMaskParameter, DefaultR));
-			FPropertyChangedEvent RChangedEvent(ParamProperty);
-			PostEditChangeProperty(RChangedEvent);
-			ParamProperty = FindFProperty<FProperty>(UMaterialExpressionStaticComponentMaskParameter::StaticClass(), GET_MEMBER_NAME_STRING_CHECKED(UMaterialExpressionStaticComponentMaskParameter, DefaultG));
-			FPropertyChangedEvent GChangedEvent(ParamProperty);
-			PostEditChangeProperty(GChangedEvent);
-			ParamProperty = FindFProperty<FProperty>(UMaterialExpressionStaticComponentMaskParameter::StaticClass(), GET_MEMBER_NAME_STRING_CHECKED(UMaterialExpressionStaticComponentMaskParameter, DefaultB));
-			FPropertyChangedEvent BChangedEvent(ParamProperty);
-			PostEditChangeProperty(BChangedEvent);
-			ParamProperty = FindFProperty<FProperty>(UMaterialExpressionStaticComponentMaskParameter::StaticClass(), GET_MEMBER_NAME_STRING_CHECKED(UMaterialExpressionStaticComponentMaskParameter, DefaultA));
-			FPropertyChangedEvent AChangedEvent(ParamProperty);
-			PostEditChangeProperty(AChangedEvent);
+			SendPostEditChangeProperty(this, TEXT("DefaultR"));
+			SendPostEditChangeProperty(this, TEXT("DefaultG"));
+			SendPostEditChangeProperty(this, TEXT("DefaultB"));
+			SendPostEditChangeProperty(this, TEXT("DefaultA"));
 		}
 		return true;
 	}
@@ -7864,9 +7860,7 @@ bool UMaterialExpressionVectorParameter::SetParameterValue(FName InParameterName
 		DefaultValue = InValue;
 		if (EnumHasAnyFlags(Flags, EMaterialExpressionSetParameterValueFlags::SendPostEditChangeProperty))
 		{
-			FProperty* ParamProperty = FindFProperty<FProperty>(UMaterialExpressionVectorParameter::StaticClass(), GET_MEMBER_NAME_STRING_CHECKED(UMaterialExpressionVectorParameter, DefaultValue));
-			FPropertyChangedEvent PropertyChangedEvent(ParamProperty);
-			PostEditChangeProperty(PropertyChangedEvent);
+			SendPostEditChangeProperty(this, TEXT("DefaultValue"));
 		}
 		return true;
 	}
@@ -8012,16 +8006,8 @@ bool UMaterialExpressionChannelMaskParameter::SetParameterValue(FName InParamete
 
 		if (EnumHasAnyFlags(Flags, EMaterialExpressionSetParameterValueFlags::SendPostEditChangeProperty))
 		{
-			{
-				FProperty* ParamProperty = FindFProperty<FProperty>(UMaterialExpressionVectorParameter::StaticClass(), GET_MEMBER_NAME_STRING_CHECKED(UMaterialExpressionVectorParameter, DefaultValue));
-				FPropertyChangedEvent PropertyChangedEvent(ParamProperty);
-				PostEditChangeProperty(PropertyChangedEvent);
-			}
-			{
-				FProperty* ParamProperty = FindFProperty<FProperty>(UMaterialExpressionChannelMaskParameter::StaticClass(), GET_MEMBER_NAME_STRING_CHECKED(UMaterialExpressionChannelMaskParameter, MaskChannel));
-				FPropertyChangedEvent PropertyChangedEvent(ParamProperty);
-				PostEditChangeProperty(PropertyChangedEvent);
-			}
+			SendPostEditChangeProperty(this, TEXT("DefaultValue"));
+			SendPostEditChangeProperty(this, TEXT("MaskChannel"));
 		}
 
 		// Update caption and preview
@@ -8196,9 +8182,7 @@ bool UMaterialExpressionScalarParameter::SetParameterValue(FName InParameterName
 		DefaultValue = InValue;
 		if (EnumHasAnyFlags(Flags, EMaterialExpressionSetParameterValueFlags::SendPostEditChangeProperty))
 		{
-			FProperty* ParamProperty = FindFProperty<FProperty>(UMaterialExpressionScalarParameter::StaticClass(), GET_MEMBER_NAME_STRING_CHECKED(UMaterialExpressionScalarParameter, DefaultValue));
-			FPropertyChangedEvent PropertyChangedEvent(ParamProperty);
-			PostEditChangeProperty(PropertyChangedEvent);
+			SendPostEditChangeProperty(this, TEXT("DefaultValue"));
 		}
 		return true;
 	}
@@ -8398,9 +8382,7 @@ bool UMaterialExpressionStaticBoolParameter::SetParameterValue(FName InParameter
 		}
 		if (EnumHasAnyFlags(Flags, EMaterialExpressionSetParameterValueFlags::SendPostEditChangeProperty))
 		{
-			FProperty* ParamProperty = FindFProperty<FProperty>(UMaterialExpressionStaticBoolParameter::StaticClass(), GET_MEMBER_NAME_STRING_CHECKED(UMaterialExpressionStaticBoolParameter, DefaultValue));
-			FPropertyChangedEvent PropertyChangedEvent(ParamProperty);
-			PostEditChangeProperty(PropertyChangedEvent);
+			SendPostEditChangeProperty(this, TEXT("DefaultValue"));
 		}
 		return true;
 	}
@@ -11398,12 +11380,8 @@ bool UMaterialExpressionFontSampleParameter::SetParameterValue(FName InParameter
 		FontTexturePage = InFontPage;
 		if (EnumHasAnyFlags(Flags, EMaterialExpressionSetParameterValueFlags::SendPostEditChangeProperty))
 		{
-			FProperty* ParamProperty = FindFProperty<FProperty>(UMaterialExpressionFontSampleParameter::StaticClass(), GET_MEMBER_NAME_STRING_CHECKED(UMaterialExpressionFontSampleParameter, Font));
-			FPropertyChangedEvent PropertyChangedEvent(ParamProperty);
-			PostEditChangeProperty(PropertyChangedEvent);
-			ParamProperty = FindFProperty<FProperty>(UMaterialExpressionFontSampleParameter::StaticClass(), GET_MEMBER_NAME_STRING_CHECKED(UMaterialExpressionFontSampleParameter, FontTexturePage));
-			FPropertyChangedEvent PageChangedEvent(ParamProperty);
-			PostEditChangeProperty(PageChangedEvent);
+			SendPostEditChangeProperty(this, TEXT("Font"));
+			SendPostEditChangeProperty(this, TEXT("FontTexturePage"));
 		}
 		return true;
 	}
@@ -19197,6 +19175,17 @@ bool UMaterialExpressionCurveAtlasRowParameter::SetParameterValue(const FName& N
 		Curve = Meta.ScalarCurve.Get();
 		Atlas = Meta.ScalarAtlas.Get();
 		DefaultValue = Meta.Value.AsScalar();
+		if (EnumHasAnyFlags(Flags, EMaterialExpressionSetParameterValueFlags::SendPostEditChangeProperty))
+		{
+			SendPostEditChangeProperty(this, TEXT("Curve"));
+			SendPostEditChangeProperty(this, TEXT("Atlas"));
+			SendPostEditChangeProperty(this, TEXT("DefaultValue"));
+		}
+		if (EnumHasAnyFlags(Flags, EMaterialExpressionSetParameterValueFlags::AssignGroupAndSortPriority))
+		{
+			Group = Meta.Group;
+			SortPriority = Meta.SortPriority;
+		}
 		return true;
 	}
 	return false;
