Make -noshaderddc apply to global shaders as well.

diff --git a/Engine/Source/Runtime/Engine/Private/ShaderCompiler/ShaderCompiler.cpp b/Engine/Source/Runtime/Engine/Private/ShaderCompiler/ShaderCompiler.cpp
--- a/Engine/Source/Runtime/Engine/Private/ShaderCompiler/ShaderCompiler.cpp
+++ b/Engine/Source/Runtime/Engine/Private/ShaderCompiler/ShaderCompiler.cpp
@@ -6733,45 +6733,55 @@ void CompileGlobalShaderMap(EShaderPlatform Platform, const ITargetPlatform* Tar
 
 			int32 HandleIndex = 0;
 
-			// Submit DDC requests.
-			for (const auto& ShaderFilenameDependencies : ShaderMapId.GetShaderFilenameToDependeciesMap())
+			// If NoShaderDDC then don't check for a material the first time we encounter it to simulate
+			// a cold DDC
+			static bool bNoShaderDDC = FParse::Param(FCommandLine::Get(), TEXT("noshaderddc"));
+			if (UNLIKELY(bNoShaderDDC))
 			{
-				SlowTask.EnterProgressFrame(ProgressStep, LOCTEXT("SubmitDDCRequests", "Submitting global shader DDC Requests..."));
-
-				const FString DataKey = GetGlobalShaderMapKeyString(ShaderMapId, Platform, TargetPlatform, ShaderFilenameDependencies.Value);
+				bShaderMapIsBeingCompiled = true;
+			}
+			else
+			{
+				// Submit DDC requests.
+				for (const auto& ShaderFilenameDependencies : ShaderMapId.GetShaderFilenameToDependeciesMap())
+				{
+					SlowTask.EnterProgressFrame(ProgressStep, LOCTEXT("SubmitDDCRequests", "Submitting global shader DDC Requests..."));
 
-				AsyncDDCRequestHandles[HandleIndex] = GetDerivedDataCacheRef().GetAsynchronous(*DataKey, TEXT("GlobalShaderMap"_SV));
+					const FString DataKey = GetGlobalShaderMapKeyString(ShaderMapId, Platform, TargetPlatform, ShaderFilenameDependencies.Value);
 
-				++HandleIndex;
-			}
+					AsyncDDCRequestHandles[HandleIndex] = GetDerivedDataCacheRef().GetAsynchronous(*DataKey, TEXT("GlobalShaderMap"_SV));
 
+					++HandleIndex;
+				}
 
-			TArray<uint8> CachedData;
 
-			HandleIndex = 0;
+				TArray<uint8> CachedData;
 
-			// Process finished DDC requests.
-			for (const auto& ShaderFilenameDependencies : ShaderMapId.GetShaderFilenameToDependeciesMap())
-			{
-				SlowTask.EnterProgressFrame(ProgressStep, LOCTEXT("ProcessDDCRequests", "Processing global shader DDC requests..."));
-				CachedData.Reset();
-				COOK_STAT(auto Timer = GlobalShaderCookStats::UsageStats.TimeSyncWork());
+				HandleIndex = 0;
 
-				GetDerivedDataCacheRef().WaitAsynchronousCompletion(AsyncDDCRequestHandles[HandleIndex]);
-				if (GetDerivedDataCacheRef().GetAsynchronousResults(AsyncDDCRequestHandles[HandleIndex], CachedData))
+				// Process finished DDC requests.
+				for (const auto& ShaderFilenameDependencies : ShaderMapId.GetShaderFilenameToDependeciesMap())
 				{
-					COOK_STAT(Timer.AddHit(CachedData.Num()));
-					FMemoryReader MemoryReader(CachedData);
-					GGlobalShaderMap[Platform]->AddSection(FGlobalShaderMapSection::CreateFromArchive(MemoryReader));
-				}
-				else
-				{
-					// it's a miss, but we haven't built anything yet. Save the counting until we actually have it built.
-					COOK_STAT(Timer.TrackCyclesOnly());
-					bShaderMapIsBeingCompiled = true;
-				}
+					SlowTask.EnterProgressFrame(ProgressStep, LOCTEXT("ProcessDDCRequests", "Processing global shader DDC requests..."));
+					CachedData.Reset();
+					COOK_STAT(auto Timer = GlobalShaderCookStats::UsageStats.TimeSyncWork());
 
-				++HandleIndex;
+					GetDerivedDataCacheRef().WaitAsynchronousCompletion(AsyncDDCRequestHandles[HandleIndex]);
+					if (GetDerivedDataCacheRef().GetAsynchronousResults(AsyncDDCRequestHandles[HandleIndex], CachedData))
+					{
+						COOK_STAT(Timer.AddHit(CachedData.Num()));
+						FMemoryReader MemoryReader(CachedData);
+						GGlobalShaderMap[Platform]->AddSection(FGlobalShaderMapSection::CreateFromArchive(MemoryReader));
+					}
+					else
+					{
+						// it's a miss, but we haven't built anything yet. Save the counting until we actually have it built.
+						COOK_STAT(Timer.TrackCyclesOnly());
+						bShaderMapIsBeingCompiled = true;
+					}
+
+					++HandleIndex;
+				}
 			}
 		}
 
