When modifying parameter value on a UMaterial, update all expressions that match the name, rather than simply the first

diff --git a/Engine/Source/Editor/UnrealEd/Private/PreviewMaterial.cpp b/Engine/Source/Editor/UnrealEd/Private/PreviewMaterial.cpp
--- a/Engine/Source/Editor/UnrealEd/Private/PreviewMaterial.cpp
+++ b/Engine/Source/Editor/UnrealEd/Private/PreviewMaterial.cpp
@@ -402,6 +402,7 @@ UDEditorParameterValue* UDEditorParameterValue::Create(UObject* Owner,
 
 	check(Parameter);
 	Parameter->ParameterInfo = ParameterInfo;
+	Parameter->ExpressionId = Meta.ExpressionGuid;
 	Parameter->SortPriority = Meta.SortPriority;
 	Parameter->bOverride = Meta.bOverride;
 	return Parameter;
diff --git a/Engine/Source/Runtime/Engine/Classes/Materials/Material.h b/Engine/Source/Runtime/Engine/Classes/Materials/Material.h
--- a/Engine/Source/Runtime/Engine/Classes/Materials/Material.h
+++ b/Engine/Source/Runtime/Engine/Classes/Materials/Material.h
@@ -1581,7 +1581,15 @@ public:
 	 * the expression nodes will not accurately display the final compiled material.
 	 * @param	Parameter	Pointer to the expression node whose state needs to be propagated.
 	 */
-	ENGINE_API virtual void PropagateExpressionParameterChanges(UMaterialExpression* Parameter);
+	ENGINE_API virtual void PropagateExpressionParameterChanges(const UMaterialExpression* Parameter);
+
+	/**
+	 * A parameter with duplicates has to update its peers so that they all have the same value. If this step isn't performed then
+	 * the expression nodes will not accurately display the final compiled material.
+	 * @param	ParameterName	Name of the parameter to update
+	 * @param	Meta			New value of the parameter
+	 */
+	ENGINE_API virtual void PropagateExpressionParameterChanges(const FName& ParameterName, const FMaterialParameterMetadata& Meta);
 
 	/**
 	 * Remove the expression from the editor parameters list (if it exists) and then re-adds it.
@@ -1624,14 +1632,6 @@ public:
 	 * @return	true if the expression is a parameter with a name.
 	 */
 	static bool GetExpressionParameterName(const UMaterialExpression* Expression, FName& OutName);
-
-	/**
-	 * Copy the values of an expression parameter to another expression parameter of the same class.
-	 *
-	 * @param	Source			The source parameter.
-	 * @param	Destination		The destination parameter that will receive Source's values.
-	 */
-	static bool CopyExpressionParameters(UMaterialExpression* Source, UMaterialExpression* Destination);
 #endif // WITH_EDITOR
 
 	/**
diff --git a/Engine/Source/Runtime/Engine/Private/Materials/Material.cpp b/Engine/Source/Runtime/Engine/Private/Materials/Material.cpp
--- a/Engine/Source/Runtime/Engine/Private/Materials/Material.cpp
+++ b/Engine/Source/Runtime/Engine/Private/Materials/Material.cpp
@@ -1335,12 +1335,12 @@ bool UMaterial::IsCompilingOrHadCompileError(ERHIFeatureLevel::Type InFeatureLev
 #if WITH_EDITOR
 bool UMaterial::SetParameterValueEditorOnly(const FName& ParameterName, const FMaterialParameterMetadata& Meta)
 {
-	// Warning: in the case of duplicate parameters with different default values, this will find the first in the expression array, not necessarily the one that's used for rendering
+	bool bResult = false;
 	for (TObjectPtr<UMaterialExpression>& Expression : Expressions)
 	{
 		if (Expression && Expression->SetParameterValue(ParameterName, Meta, EMaterialExpressionSetParameterValueFlags::SendPostEditChangeProperty))
 		{
-			return true;
+			bResult = true;
 		}
 		else if (UMaterialExpressionMaterialFunctionCall* FunctionCall = Cast<UMaterialExpressionMaterialFunctionCall>(Expression))
 		{
@@ -1359,7 +1359,7 @@ bool UMaterial::SetParameterValueEditorOnly(const FName& ParameterName, const FM
 						{
 							if (FunctionExpression && FunctionExpression->SetParameterValue(ParameterName, Meta, EMaterialExpressionSetParameterValueFlags::SendPostEditChangeProperty))
 							{
-								return true;
+								bResult = true;
 							}
 						}
 					}
@@ -1368,7 +1368,7 @@ bool UMaterial::SetParameterValueEditorOnly(const FName& ParameterName, const FM
 		}
 	}
 
-	return false;
+	return bResult;
 }
 
 bool UMaterial::SetVectorParameterValueEditorOnly(FName ParameterName, FLinearColor InValue)
@@ -4073,26 +4073,30 @@ void UMaterial::UpdateExpressionDynamicParameters(const UMaterialExpression* Exp
 	}
 }
 
-void UMaterial::PropagateExpressionParameterChanges(UMaterialExpression* Parameter)
+void UMaterial::PropagateExpressionParameterChanges(const UMaterialExpression* Parameter)
 {
-	FName ParmName;
-	bool bRet = GetExpressionParameterName(Parameter, ParmName);
-
-	if(bRet)
+	FMaterialParameterMetadata Meta;
+	if (Parameter->GetParameterValue(Meta))
 	{
-		TArray<UMaterialExpression*>* ExpressionList = EditorParameters.Find(ParmName);
+		PropagateExpressionParameterChanges(Parameter->GetParameterName(), Meta);
+	}
+}
 
-		if(ExpressionList && ExpressionList->Num() > 1)
+void UMaterial::PropagateExpressionParameterChanges(const FName& ParameterName, const FMaterialParameterMetadata& Meta)
+{
+	TArray<UMaterialExpression*>* ExpressionList = EditorParameters.Find(ParameterName);
+	if (ExpressionList && ExpressionList->Num() > 1)
+	{
+		for (UMaterialExpression* Expression : *ExpressionList)
 		{
-			for(int32 Index = 0; Index < ExpressionList->Num(); ++Index)
+			const EMaterialExpressionSetParameterValueFlags Flags = EMaterialExpressionSetParameterValueFlags::NoUpdateExpressionGuid | EMaterialExpressionSetParameterValueFlags::AssignGroupAndSortPriority;
+			if (Expression->SetParameterValue(ParameterName, Meta, Flags))
 			{
-				CopyExpressionParameters(Parameter, (*ExpressionList)[Index]);
+				Expression->Modify();
+				Expression->Desc = Meta.Description;
+				Expression->GraphNode->OnUpdateCommentText(Meta.Description);
 			}
 		}
-		else if(!ExpressionList)
-		{
-			bRet = false;
-		}
 	}
 }
 
@@ -4278,31 +4282,6 @@ bool UMaterial::GetExpressionParameterName(const UMaterialExpression* Expression
 	}
 	return bRet;
 }
-
-bool UMaterial::CopyExpressionParameters(UMaterialExpression* Source, UMaterialExpression* Destination)
-{
-	if(!Source || !Destination || Source == Destination || Source->GetClass() != Destination->GetClass())
-	{
-		return false;
-	}
-
-	bool bResult = false;
-	FMaterialParameterMetadata Meta;
-	if (Source->GetParameterValue(Meta))
-	{
-		const FName ParameterName = Destination->GetParameterName();
-		Destination->Modify();
-
-		const EMaterialExpressionSetParameterValueFlags Flags = EMaterialExpressionSetParameterValueFlags::NoUpdateExpressionGuid | EMaterialExpressionSetParameterValueFlags::AssignGroupAndSortPriority;
-		verify(Destination->SetParameterValue(ParameterName, Meta, Flags));
-
-		Destination->Desc = Source->Desc;
-		Destination->GraphNode->OnUpdateCommentText(Destination->Desc);
-		bResult = true;
-	}
-
-	return bResult;
-}
 #endif // WITH_EDITOR
 
 void UMaterial::BeginDestroy()
diff --git a/Engine/Source/Runtime/Engine/Private/Materials/MaterialExpressions.cpp b/Engine/Source/Runtime/Engine/Private/Materials/MaterialExpressions.cpp
--- a/Engine/Source/Runtime/Engine/Private/Materials/MaterialExpressions.cpp
+++ b/Engine/Source/Runtime/Engine/Private/Materials/MaterialExpressions.cpp
@@ -13019,12 +13019,12 @@ bool UMaterialFunction::HasFlippedCoordinates() const
 
 bool UMaterialFunction::SetParameterValueEditorOnly(const FName& ParameterName, const FMaterialParameterMetadata& Meta)
 {
+	bool bResult = false;
 	for (UMaterialExpression* Expression : FunctionExpressions)
 	{
 		if (Expression->SetParameterValue(ParameterName, Meta))
 		{
-			return true;
-			// Warning: in the case of duplicate parameters with different default values, this will find the first in the expression array, not necessarily the one that's used for rendering
+			bResult = true;
 		}
 		else if (UMaterialExpressionMaterialFunctionCall* FunctionCall = Cast<UMaterialExpressionMaterialFunctionCall>(Expression))
 		{
@@ -13043,7 +13043,7 @@ bool UMaterialFunction::SetParameterValueEditorOnly(const FName& ParameterName,
 						{
 							if (FunctionExpression->SetParameterValue(ParameterName, Meta))
 							{
-								return true;
+								bResult = true;
 							}
 						}
 					}
@@ -13051,7 +13051,7 @@ bool UMaterialFunction::SetParameterValueEditorOnly(const FName& ParameterName,
 			}
 		}
 	}
-	return false;
+	return bResult;
 }
 
 bool UMaterialFunction::SetVectorParameterValueEditorOnly(FName ParameterName, FLinearColor InValue)
