Fix Oodle in Editor was accidentally turned off by merge CL 17463546

EnableInEditor should not be checked anymore, Oodle is always on

diff --git a/Engine/Source/Runtime/Engine/Private/Texture.cpp b/Engine/Source/Runtime/Engine/Private/Texture.cpp
--- a/Engine/Source/Runtime/Engine/Private/Texture.cpp
+++ b/Engine/Source/Runtime/Engine/Private/Texture.cpp
@@ -1922,31 +1922,36 @@ int64 UTexture::GetBuildRequiredMemory() const
 static FName ConditionalGetPrefixedFormat(FName TextureFormatName, const ITargetPlatform* TargetPlatform)
 {
 #if WITH_EDITOR
+
 	// Prepend a texture format to allow a module to override the compression (Ex: this allows you to replace TextureFormatDXT with a different compressor)
 	FString FormatPrefix;
 	bool bHasPrefix = TargetPlatform->GetConfigSystem()->GetString(TEXT("AlternateTextureCompression"), TEXT("TextureFormatPrefix"), FormatPrefix, GEngineIni);
+	bHasPrefix = bHasPrefix && ! FormatPrefix.IsEmpty();
 
-	FString TextureCompressionFormat;
-	bool bHasFormat = TargetPlatform->GetConfigSystem()->GetString(TEXT("AlternateTextureCompression"), TEXT("TextureCompressionFormat"), TextureCompressionFormat, GEngineIni) && TextureCompressionFormat != TEXT("");
-
-	bool bEnableInEditor = false;
-	TargetPlatform->GetConfigSystem()->GetBool(TEXT("AlternateTextureCompression"), TEXT("bEnableInEditor"), bEnableInEditor, GEngineIni);
-
-	// Disable in the Editor by default but never in cooked builds
-	bEnableInEditor = !TargetPlatform->HasEditorOnlyData() || bEnableInEditor;
-
-	if (bHasPrefix && bHasFormat && bEnableInEditor)
+	if ( bHasPrefix )
 	{
-		ITextureFormat* TextureFormat = FModuleManager::LoadModuleChecked<ITextureFormatModule>(*TextureCompressionFormat).GetTextureFormat();
+		FString TextureCompressionFormat;
+		bool bHasFormat = TargetPlatform->GetConfigSystem()->GetString(TEXT("AlternateTextureCompression"), TEXT("TextureCompressionFormat"), TextureCompressionFormat, GEngineIni);
+		bHasFormat = bHasFormat && ! TextureCompressionFormat.IsEmpty();
+	
+		if ( bHasFormat )
+		{
+			ITextureFormatModule * TextureFormatModule = FModuleManager::LoadModulePtr<ITextureFormatModule>(*TextureCompressionFormat);
 
-		TArray<FName> SupportedFormats;
-		TextureFormat->GetSupportedFormats(SupportedFormats);
+			if ( TextureFormatModule )
+			{
+				ITextureFormat* TextureFormat = TextureFormatModule->GetTextureFormat();
 
-		FName NewFormatName(FormatPrefix + TextureFormatName.ToString());
+				TArray<FName> SupportedFormats;
+				TextureFormat->GetSupportedFormats(SupportedFormats);
 
-		if (SupportedFormats.Contains(NewFormatName))
-		{
-			return NewFormatName;
+				FName NewFormatName(FormatPrefix + TextureFormatName.ToString());
+
+				if (SupportedFormats.Contains(NewFormatName))
+				{
+					return NewFormatName;
+				}
+			}
 		}
 	}
 #endif
