Remove some lingering casts to various material expression parameter types, prefer using generic interface

diff --git a/Engine/Source/Runtime/Engine/Private/Materials/Material.cpp b/Engine/Source/Runtime/Engine/Private/Materials/Material.cpp
--- a/Engine/Source/Runtime/Engine/Private/Materials/Material.cpp
+++ b/Engine/Source/Runtime/Engine/Private/Materials/Material.cpp
@@ -24,24 +24,16 @@
 #include "Materials/MaterialExpressionCustomOutput.h"
 #include "Materials/MaterialExpressionFunctionOutput.h"
 #include "Materials/MaterialExpressionDynamicParameter.h"
-#include "Materials/MaterialExpressionFontSampleParameter.h"
 #include "Materials/MaterialExpressionQualitySwitch.h"
 #include "Materials/MaterialExpressionFeatureLevelSwitch.h"
 #include "Materials/MaterialExpressionShadingPathSwitch.h"
 #include "Materials/MaterialExpressionShaderStageSwitch.h"
 #include "Materials/MaterialExpressionMakeMaterialAttributes.h"
 #include "Materials/MaterialExpressionSetMaterialAttributes.h"
-#include "Materials/MaterialExpressionParameter.h"
 #include "Materials/MaterialExpressionRuntimeVirtualTextureOutput.h"
-#include "Materials/MaterialExpressionRuntimeVirtualTextureSampleParameter.h"
-#include "Materials/MaterialExpressionScalarParameter.h"
-#include "Materials/MaterialExpressionStaticBoolParameter.h"
-#include "Materials/MaterialExpressionStaticComponentMaskParameter.h"
 #include "Materials/MaterialExpressionStaticSwitchParameter.h"
 #include "Materials/MaterialExpressionTextureBase.h"
 #include "Materials/MaterialExpressionTextureSample.h"
-#include "Materials/MaterialExpressionTextureSampleParameter.h"
-#include "Materials/MaterialExpressionVectorParameter.h"
 #include "Materials/MaterialExpressionVertexInterpolator.h"
 #include "Materials/MaterialExpressionSceneColor.h"
 #include "Materials/MaterialExpressionShadingModel.h"
@@ -104,7 +96,6 @@
 #include "ThumbnailRendering/SceneThumbnailInfoWithPrimitive.h"
 #endif
 #include "ShaderCodeLibrary.h"
-#include "Materials/MaterialExpressionCurveAtlasRowParameter.h"
 #include "Curves/CurveLinearColor.h"
 #include "Curves/CurveLinearColorAtlas.h"
 #include "HAL/ThreadHeartBeat.h"
@@ -1904,8 +1895,6 @@ bool UMaterial::GetParameterValue(EMaterialParameterType Type, const FMemoryImag
 #if WITH_EDITORONLY_DATA
 bool UMaterial::GetMaterialLayersParameterValue(const FHashedMaterialParameterInfo& ParameterInfo, FMaterialLayersFunctions& OutLayers, FGuid& OutExpressionGuid, bool bCheckParent /*= true*/) const
 {
-	UMaterialExpressionStaticComponentMaskParameter* Parameter = nullptr;
-
 	for (UMaterialExpression* Expression : Expressions)
 	{
 		// Note: Check for layers in top-level only, no recursion required or supported here
@@ -3903,98 +3892,39 @@ void UMaterial::PostEditChangePropertyInternal(FPropertyChangedEvent& PropertyCh
 
 bool UMaterial::AddExpressionParameter(UMaterialExpression* Expression, TMap<FName, TArray<UMaterialExpression*> >& ParameterTypeMap)
 {
-	if(!Expression)
-	{
-		return false;
-	}
-
-	bool bRet = false;
-
-	if(Expression->IsA(UMaterialExpressionParameter::StaticClass()))
+	if(Expression && Expression->HasAParameterName())
 	{
-		UMaterialExpressionParameter *Param = (UMaterialExpressionParameter*)Expression;
-
-		TArray<UMaterialExpression*> *ExpressionList = ParameterTypeMap.Find(Param->ParameterName);
-
-		if(!ExpressionList)
+		const FName ParameterName = Expression->GetParameterName();
+		TArray<UMaterialExpression*>* ExpressionList = ParameterTypeMap.Find(ParameterName);
+		if (!ExpressionList)
 		{
-			ExpressionList = &ParameterTypeMap.Add(Param->ParameterName, TArray<UMaterialExpression*>());
+			ExpressionList = &ParameterTypeMap.Add(ParameterName, TArray<UMaterialExpression*>());
 		}
-
-		ExpressionList->Add(Param);
-		bRet = true;
-	}
-	else if(Expression->IsA(UMaterialExpressionTextureSampleParameter::StaticClass()))
-	{
-		UMaterialExpressionTextureSampleParameter *Param = (UMaterialExpressionTextureSampleParameter*)Expression;
-
-		TArray<UMaterialExpression*> *ExpressionList = ParameterTypeMap.Find(Param->ParameterName);
-
-		if(!ExpressionList)
-		{
-			ExpressionList = &ParameterTypeMap.Add(Param->ParameterName, TArray<UMaterialExpression*>());
-		}
-
-		ExpressionList->Add(Param);
-		bRet = true;
-	}
-	else if(Expression->IsA(UMaterialExpressionFontSampleParameter::StaticClass()))
-	{
-		UMaterialExpressionFontSampleParameter *Param = (UMaterialExpressionFontSampleParameter*)Expression;
-
-		TArray<UMaterialExpression*> *ExpressionList = ParameterTypeMap.Find(Param->ParameterName);
-
-		if(!ExpressionList)
-		{
-			ExpressionList = &ParameterTypeMap.Add(Param->ParameterName, TArray<UMaterialExpression*>());
-		}
-
-		ExpressionList->Add(Param);
-		bRet = true;
+		ExpressionList->Add(Expression);
+		return true;
 	}
-
-	return bRet;
+	return false;
 }
 
 bool UMaterial::RemoveExpressionParameter(UMaterialExpression* Expression)
 {
-	FName ParmName;
-
-	if(GetExpressionParameterName(Expression, ParmName))
+	if (Expression && Expression->HasAParameterName())
 	{
-		TArray<UMaterialExpression*>* ExpressionList = EditorParameters.Find(ParmName);
-
-		if(ExpressionList)
+		const FName ParameterName = Expression->GetParameterName();
+		TArray<UMaterialExpression*>* ExpressionList = EditorParameters.Find(ParameterName);
+		if (ExpressionList)
 		{
 			return ExpressionList->Remove(Expression) > 0;
 		}
 	}
-
 	return false;
 }
-#endif // WITH_EDITOR
 
 bool UMaterial::IsParameter(const UMaterialExpression* Expression)
 {
-	bool bRet = false;
-
-	if(Expression->IsA(UMaterialExpressionParameter::StaticClass()))
-	{
-		bRet = true;
-	}
-	else if(Expression->IsA(UMaterialExpressionTextureSampleParameter::StaticClass()))
-	{
-		bRet = true;
-	}
-	else if(Expression->IsA(UMaterialExpressionFontSampleParameter::StaticClass()))
-	{
-		bRet = true;
-	}
-
-	return bRet;
+	return Expression->HasAParameterName();
 }
 
-
 bool UMaterial::IsDynamicParameter(const UMaterialExpression* Expression)
 {
 	if (Expression->IsA(UMaterialExpressionDynamicParameter::StaticClass()))
@@ -4005,7 +3935,6 @@ bool UMaterial::IsDynamicParameter(const UMaterialExpression* Expression)
 	return false;
 }
 
-#if WITH_EDITOR
 void UMaterial::BuildEditorParameterList()
 {
 	EditorParameters.Empty();
@@ -4019,17 +3948,15 @@ void UMaterial::BuildEditorParameterList()
 bool UMaterial::HasDuplicateParameters(const UMaterialExpression* Expression)
 {
 	FName ExpressionName;
-
 	if(GetExpressionParameterName(Expression, ExpressionName))
 	{
-		TArray<UMaterialExpression*>* ExpressionList = EditorParameters.Find(ExpressionName);
-
+		const TArray<UMaterialExpression*>* ExpressionList = EditorParameters.Find(ExpressionName);
 		if(ExpressionList)
 		{
-			for(int32 ParmIndex = 0; ParmIndex < ExpressionList->Num(); ++ParmIndex)
+			const EMaterialParameterType ParameterType = Expression->GetParameterType();
+			for (UMaterialExpression* CurNode : *ExpressionList)
 			{
-				UMaterialExpression* CurNode = (*ExpressionList)[ParmIndex];
-				if(CurNode != Expression && CurNode->GetClass() == Expression->GetClass())
+				if(CurNode != Expression && CurNode->GetParameterType() == ParameterType)
 				{
 					return true;
 				}
