Fix up UMaterialInstance::GetParameterValue().
Now will always get metadata from base meterial if applicable, then apply override value on top of that

diff --git a/Engine/Plugins/Editor/MaterialAnalyzer/Source/Private/SMaterialAnalyzer.cpp b/Engine/Plugins/Editor/MaterialAnalyzer/Source/Private/SMaterialAnalyzer.cpp
--- a/Engine/Plugins/Editor/MaterialAnalyzer/Source/Private/SMaterialAnalyzer.cpp
+++ b/Engine/Plugins/Editor/MaterialAnalyzer/Source/Private/SMaterialAnalyzer.cpp
@@ -690,6 +690,7 @@ void FAnalyzeMaterialTreeAsyncTask::DoWork()
 	TArray<FGuid> Guids;
 
 	UMaterial* CurrentMaterial = Cast<UMaterial>(CurrentMaterialInterface);
+	UMaterialInstance* CurrentMaterialInstance = Cast<UMaterialInstance>(CurrentMaterialInterface);
 
 	bool bCanBeOverriden = CurrentMaterial != nullptr;
 
@@ -709,8 +710,6 @@ void FAnalyzeMaterialTreeAsyncTask::DoWork()
 		float TempValue = 0.0f;
 		bool bIsOverridden = false;
 
-		UMaterialInstance* CurrentMaterialInstance = Cast<UMaterialInstance>(CurrentMaterialInterface);
-
 		if (BasePropertyOverrideName.Key.IsEqual(TEXT("bOverride_OpacityMaskClipValue")))
 		{
 			TempValue = CurrentMaterialInterface->GetOpacityMaskClipValue();
@@ -832,8 +831,8 @@ void FAnalyzeMaterialTreeAsyncTask::DoWork()
 	
 	for (int ParameterIndex = 0; ParameterIndex < StaticSwitchParameterInfo.Num(); ++ParameterIndex)
 	{
-		bool bStaticSwitchValue;
-		bool bIsOverridden = CurrentMaterialInterface->GetStaticSwitchParameterValue(StaticSwitchParameterInfo[ParameterIndex], bStaticSwitchValue, StaticSwitchGuids[ParameterIndex], false, false);
+		FMaterialParameterMetadata Meta;
+		bool bIsOverridden = CurrentMaterialInstance->GetParameterOverrideValue(EMaterialParameterType::StaticSwitch, StaticSwitchParameterInfo[ParameterIndex], Meta);
 
 		if (!bIsOverridden)
 		{
@@ -852,7 +851,7 @@ void FAnalyzeMaterialTreeAsyncTask::DoWork()
 		else
 		{
 			CurrentMaterialNode->StaticSwitchParameters.Add(FStaticSwitchParameterNodeRef(
-				new FStaticSwitchParameterNode(StaticSwitchParameterInfo[ParameterIndex].Name, bStaticSwitchValue, true)));
+				new FStaticSwitchParameterNode(StaticSwitchParameterInfo[ParameterIndex].Name, Meta.Value.AsStaticSwitch(), true)));
 		}
 	}
 	
@@ -860,9 +859,8 @@ void FAnalyzeMaterialTreeAsyncTask::DoWork()
 	
 	for (int ParameterIndex = 0; ParameterIndex < StaticMaskParameterInfo.Num(); ++ParameterIndex)
 	{
-		bool R, G, B, A;
-	
-		bool bIsOverridden = CurrentMaterialInterface->GetStaticComponentMaskParameterValue(StaticMaskParameterInfo[ParameterIndex], R, G, B, A, StaticMaskGuids[ParameterIndex], false, false);
+		FMaterialParameterMetadata Meta;
+		bool bIsOverridden = CurrentMaterialInstance->GetParameterOverrideValue(EMaterialParameterType::StaticComponentMask, StaticSwitchParameterInfo[ParameterIndex], Meta);
 
 		if(!bIsOverridden)
 		{
@@ -884,7 +882,12 @@ void FAnalyzeMaterialTreeAsyncTask::DoWork()
 		else
 		{
 			CurrentMaterialNode->StaticComponentMaskParameters.Add(FStaticComponentMaskParameterNodeRef(
-				new FStaticComponentMaskParameterNode(StaticMaskParameterInfo[ParameterIndex].Name, R, G, B, A, true)));
+				new FStaticComponentMaskParameterNode(StaticMaskParameterInfo[ParameterIndex].Name,
+					Meta.Value.Bool[0],
+					Meta.Value.Bool[1],
+					Meta.Value.Bool[2],
+					Meta.Value.Bool[3],
+					true)));
 		}
 	}
 
diff --git a/Engine/Source/Runtime/Engine/Classes/Materials/Material.h b/Engine/Source/Runtime/Engine/Classes/Materials/Material.h
--- a/Engine/Source/Runtime/Engine/Classes/Materials/Material.h
+++ b/Engine/Source/Runtime/Engine/Classes/Materials/Material.h
@@ -1035,6 +1035,7 @@ public:
 	ENGINE_API virtual UMaterial* GetMaterial() override;
 	ENGINE_API virtual const UMaterial* GetMaterial() const override;
 	ENGINE_API virtual const UMaterial* GetMaterial_Concurrent(TMicRecursionGuard RecursionGuard = TMicRecursionGuard()) const override;
+	ENGINE_API virtual const UMaterial* GetMaterialInheritanceChain(FMaterialParentInstanceArray& OutInstances) const override;
 
 	ENGINE_API virtual bool GetParameterValue(EMaterialParameterType Type, const FMemoryImageMaterialParameterInfo& ParameterInfo, FMaterialParameterMetadata& OutValue, EMaterialGetParameterValueFlags Flags) const override;
 
diff --git a/Engine/Source/Runtime/Engine/Classes/Materials/MaterialInstance.h b/Engine/Source/Runtime/Engine/Classes/Materials/MaterialInstance.h
--- a/Engine/Source/Runtime/Engine/Classes/Materials/MaterialInstance.h
+++ b/Engine/Source/Runtime/Engine/Classes/Materials/MaterialInstance.h
@@ -558,11 +558,12 @@ public:
 	virtual ENGINE_API UMaterial* GetMaterial() override;
 	virtual ENGINE_API const UMaterial* GetMaterial() const override;
 	virtual ENGINE_API const UMaterial* GetMaterial_Concurrent(TMicRecursionGuard RecursionGuard = TMicRecursionGuard()) const override;
+	virtual ENGINE_API const UMaterial* GetMaterialInheritanceChain(FMaterialParentInstanceArray& OutInstances) const override;
 	virtual ENGINE_API FMaterialResource* AllocatePermutationResource();
 	virtual ENGINE_API FMaterialResource* GetMaterialResource(ERHIFeatureLevel::Type InFeatureLevel, EMaterialQualityLevel::Type QualityLevel = EMaterialQualityLevel::Num) override;
 	virtual ENGINE_API const FMaterialResource* GetMaterialResource(ERHIFeatureLevel::Type InFeatureLevel, EMaterialQualityLevel::Type QualityLevel = EMaterialQualityLevel::Num) const override;
 
-	ENGINE_API bool GetParameterOverrideValue(EMaterialParameterType Type, const FMemoryImageMaterialParameterInfo& ParameterInfo, FMaterialParameterMetadata& OutValue, EMaterialGetParameterValueFlags Flags = EMaterialGetParameterValueFlags::None) const;
+	ENGINE_API bool GetParameterOverrideValue(EMaterialParameterType Type, const FMemoryImageMaterialParameterInfo& ParameterInfo, FMaterialParameterMetadata& OutValue) const;
 	virtual ENGINE_API bool GetParameterValue(EMaterialParameterType Type, const FMemoryImageMaterialParameterInfo& ParameterInfo, FMaterialParameterMetadata& OutValue, EMaterialGetParameterValueFlags Flags) const override;
 
 	virtual ENGINE_API void GetUsedTextures(TArray<UTexture*>& OutTextures, EMaterialQualityLevel::Type QualityLevel, bool bAllQualityLevels, ERHIFeatureLevel::Type FeatureLevel, bool bAllFeatureLevels) const override;
diff --git a/Engine/Source/Runtime/Engine/Classes/Materials/MaterialInterface.h b/Engine/Source/Runtime/Engine/Classes/Materials/MaterialInterface.h
--- a/Engine/Source/Runtime/Engine/Classes/Materials/MaterialInterface.h
+++ b/Engine/Source/Runtime/Engine/Classes/Materials/MaterialInterface.h
@@ -204,6 +204,8 @@ private:
 	TMicRecursionGuard const* PreviousLink = nullptr;
 };
 
+using FMaterialParentInstanceArray = TArray<const class UMaterialInstance*, TInlineAllocator<16>>;
+
 UCLASS(abstract, BlueprintType, MinimalAPI, HideCategories = (Thumbnail))
 class UMaterialInterface : public UObject, public IBlendableInterface, public IInterface_AssetUserData
 {
@@ -329,6 +331,8 @@ public:
 	 */
 	virtual const class UMaterial* GetMaterial_Concurrent(TMicRecursionGuard RecursionGuard = TMicRecursionGuard()) const PURE_VIRTUAL(UMaterialInterface::GetMaterial_Concurrent,return NULL;);
 
+	virtual const class UMaterial* GetMaterialInheritanceChain(FMaterialParentInstanceArray & OutInstances) const { return nullptr; }
+
 	/**
 	* Test this material for dependency on a given material.
 	* @param	TestDependency - The material to test for dependency upon.
@@ -444,7 +448,7 @@ public:
 	* @param	OutValue		Will contain the value of the parameter if successful
 	* @return					True if successful
 	*/
-	ENGINE_API bool GetStaticSwitchParameterValue(const FHashedMaterialParameterInfo& ParameterInfo,bool &OutValue,FGuid &OutExpressionGuid, bool bOveriddenOnly = false, bool bCheckParent = true) const;
+	ENGINE_API bool GetStaticSwitchParameterValue(const FHashedMaterialParameterInfo& ParameterInfo,bool &OutValue,FGuid &OutExpressionGuid, bool bOveriddenOnly = false) const;
 
 	/**
 	* Get the value of the given static component mask parameter
@@ -453,7 +457,7 @@ public:
 	* @param	R, G, B, A		Will contain the values of the parameter if successful
 	* @return					True if successful
 	*/
-	ENGINE_API bool GetStaticComponentMaskParameterValue(const FHashedMaterialParameterInfo& ParameterInfo, bool& R, bool& G, bool& B, bool& A, FGuid& OutExpressionGuid, bool bOveriddenOnly = false, bool bCheckParent = true) const;
+	ENGINE_API bool GetStaticComponentMaskParameterValue(const FHashedMaterialParameterInfo& ParameterInfo, bool& R, bool& G, bool& B, bool& A, FGuid& OutExpressionGuid, bool bOveriddenOnly = false) const;
 
 	/**
 	* Get the value of the given static material layers parameter
diff --git a/Engine/Source/Runtime/Engine/Private/Materials/Material.cpp b/Engine/Source/Runtime/Engine/Private/Materials/Material.cpp
--- a/Engine/Source/Runtime/Engine/Private/Materials/Material.cpp
+++ b/Engine/Source/Runtime/Engine/Private/Materials/Material.cpp
@@ -1858,6 +1858,11 @@ const UMaterial* UMaterial::GetMaterial_Concurrent(TMicRecursionGuard) const
 	return this;
 }
 
+const UMaterial* UMaterial::GetMaterialInheritanceChain(FMaterialParentInstanceArray& OutInstances) const
+{
+	return this;
+}
+
 #if WITH_EDITOR
 void UMaterial::UpdateCachedExpressionData()
 {
diff --git a/Engine/Source/Runtime/Engine/Private/Materials/MaterialInstance.cpp b/Engine/Source/Runtime/Engine/Private/Materials/MaterialInstance.cpp
--- a/Engine/Source/Runtime/Engine/Private/Materials/MaterialInstance.cpp
+++ b/Engine/Source/Runtime/Engine/Private/Materials/MaterialInstance.cpp
@@ -768,21 +768,33 @@ UMaterial* UMaterialInstance::GetMaterial()
 	}
 }
 
-bool UMaterialInstance::GetParameterOverrideValue(EMaterialParameterType Type, const FMemoryImageMaterialParameterInfo& ParameterInfo, FMaterialParameterMetadata& OutResult, EMaterialGetParameterValueFlags Flags) const
+const UMaterial* UMaterialInstance::GetMaterialInheritanceChain(FMaterialParentInstanceArray& OutInstances) const
+{
+	if (!OutInstances.Contains(this))
+	{
+		OutInstances.Add(this);
+		if (Parent)
+		{
+			return Parent->GetMaterialInheritanceChain(OutInstances);
+		}
+	}
+
+	return UMaterial::GetDefaultMaterial(MD_Surface);
+}
+
+bool UMaterialInstance::GetParameterOverrideValue(EMaterialParameterType Type, const FMemoryImageMaterialParameterInfo& ParameterInfo, FMaterialParameterMetadata& OutResult) const
 {
-	// If we're not a parent, set the flag to signify this parameter was a MI override
-	const bool bSetOverride = !EnumHasAnyFlags(Flags, EMaterialGetParameterValueFlags::IsParent);
 	bool bResult = false;
 	switch (Type)
 	{
-	case EMaterialParameterType::Scalar: bResult = GameThread_GetParameterValue(ScalarParameterValues, ParameterInfo, bSetOverride, OutResult); break;
-	case EMaterialParameterType::Vector: bResult = GameThread_GetParameterValue(VectorParameterValues, ParameterInfo, bSetOverride, OutResult); break;
-	case EMaterialParameterType::Texture: bResult = GameThread_GetParameterValue(TextureParameterValues, ParameterInfo, bSetOverride, OutResult); break;
-	case EMaterialParameterType::RuntimeVirtualTexture: bResult = GameThread_GetParameterValue(RuntimeVirtualTextureParameterValues, ParameterInfo, bSetOverride, OutResult); break;
-	case EMaterialParameterType::Font: bResult = GameThread_GetParameterValue(FontParameterValues, ParameterInfo, bSetOverride, OutResult); break;
+	case EMaterialParameterType::Scalar: bResult = GameThread_GetParameterValue(ScalarParameterValues, ParameterInfo, OutResult); break;
+	case EMaterialParameterType::Vector: bResult = GameThread_GetParameterValue(VectorParameterValues, ParameterInfo, OutResult); break;
+	case EMaterialParameterType::Texture: bResult = GameThread_GetParameterValue(TextureParameterValues, ParameterInfo, OutResult); break;
+	case EMaterialParameterType::RuntimeVirtualTexture: bResult = GameThread_GetParameterValue(RuntimeVirtualTextureParameterValues, ParameterInfo, OutResult); break;
+	case EMaterialParameterType::Font: bResult = GameThread_GetParameterValue(FontParameterValues, ParameterInfo, OutResult); break;
 #if WITH_EDITORONLY_DATA
-	case EMaterialParameterType::StaticSwitch: bResult = GameThread_GetParameterValue(StaticParameters.StaticSwitchParameters, ParameterInfo, bSetOverride, OutResult); break;
-	case EMaterialParameterType::StaticComponentMask: bResult = GameThread_GetParameterValue(StaticParameters.StaticComponentMaskParameters, ParameterInfo, bSetOverride, OutResult); break;
+	case EMaterialParameterType::StaticSwitch: bResult = GameThread_GetParameterValue(StaticParameters.StaticSwitchParameters, ParameterInfo, OutResult); break;
+	case EMaterialParameterType::StaticComponentMask: bResult = GameThread_GetParameterValue(StaticParameters.StaticComponentMaskParameters, ParameterInfo, OutResult); break;
 #endif // WITH_EDITORONLY_DATA
 	default: checkNoEntry(); break;
 	}
@@ -791,48 +803,46 @@ bool UMaterialInstance::GetParameterOverrideValue(EMaterialParameterType Type, c
 
 bool UMaterialInstance::GetParameterValue(EMaterialParameterType Type, const FMemoryImageMaterialParameterInfo& ParameterInfo, FMaterialParameterMetadata& OutResult, EMaterialGetParameterValueFlags Flags) const
 {
-	bool bFoundAValue = false;
-
-	if (GetReentrantFlag())
-	{
-		return false;
-	}
+	FMaterialParentInstanceArray InstanceChain;
+	const UMaterial* Material = GetMaterialInheritanceChain(InstanceChain);
 
-	// Instance override
-	if (EnumHasAnyFlags(Flags, EMaterialGetParameterValueFlags::CheckInstanceOverrides))
+	bool bResult = false;
+	if (EnumHasAnyFlags(Flags, EMaterialGetParameterValueFlags::CheckNonOverrides))
 	{
-		if (GetParameterOverrideValue(Type, ParameterInfo, OutResult, Flags))
+		// First get value from base material, properly fill in any relevant metadata
+		if (Material && Material->GetParameterValue(Type, ParameterInfo, OutResult, Flags))
 		{
-			return true;
+			bResult = true;
 		}
-	}
-
-	// Instance-included default
-	if (EnumHasAnyFlags(Flags, EMaterialGetParameterValueFlags::CheckNonOverrides) &&
-		ParameterInfo.Association != EMaterialParameterAssociation::GlobalParameter &&
-		CachedData)
-	{
-		if (CachedData->Parameters.GetParameterValue(Type, ParameterInfo, OutResult))
+		// Apply layer parameter overrides from this instance
+		if (ParameterInfo.Association != EMaterialParameterAssociation::GlobalParameter &&
+			CachedData &&
+			CachedData->Parameters.GetParameterValue(Type, ParameterInfo, OutResult))
 		{
-			return true;
+			bResult = true;
 		}
 	}
 
-	// Next material in hierarchy
-	if (EnumHasAnyFlags(Flags, EMaterialGetParameterValueFlags::CheckParent) && Parent)
+	// This instance is at index 0.  So we start from 0 if checking overrides, otherwise we start at our first parent (at index 1)
+	int32 ParentIndex = EnumHasAnyFlags(Flags, EMaterialGetParameterValueFlags::CheckInstanceOverrides) ? 0 : 1;
+	while (ParentIndex < InstanceChain.Num())
 	{
-		EMaterialGetParameterValueFlags ParentFlags = Flags | EMaterialGetParameterValueFlags::IsParent;
-		if (EnumHasAnyFlags(Flags, EMaterialGetParameterValueFlags::CheckNonOverrides))
+		if (InstanceChain[ParentIndex]->GetParameterOverrideValue(Type, ParameterInfo, OutResult))
 		{
-			// Parent's instance overrides are non-overrides relative to this material
-			ParentFlags |= EMaterialGetParameterValueFlags::CheckInstanceOverrides;
+#if WITH_EDITORONLY_DATA
+			if (ParentIndex == 0)
+			{
+				// If value was set on this instance, set the override flag
+				OutResult.bOverride = true;
+			}
+#endif
+			bResult = true;
+			break;
 		}
-
-		FMICReentranceGuard	Guard(this);
-		return Parent->GetParameterValue(Type, ParameterInfo, OutResult, ParentFlags);
+		ParentIndex++;
 	}
 
-	return false;
+	return bResult;
 }
 
 bool UMaterialInstance::GetRefractionSettings(float& OutBiasValue) const
@@ -1713,7 +1723,9 @@ void UMaterialInstance::GetStaticParameterValues(FStaticParameterSet& OutStaticP
 
 void UMaterialInstance::GetAllParametersOfType(EMaterialParameterType Type, TMap<FMaterialParameterInfo, FMaterialParameterMetadata>& OutParameters) const
 {
-	const UMaterial* Material = GetMaterial_Concurrent();
+	FMaterialParentInstanceArray InstanceChain;
+	const UMaterial* Material = GetMaterialInheritanceChain(InstanceChain);
+
 	int32 NumParameters = 0;
 	if (CachedData)
 	{
@@ -1734,26 +1746,6 @@ void UMaterialInstance::GetAllParametersOfType(EMaterialParameterType Type, TMap
 		Material->GetCachedExpressionData().Parameters.GetAllParametersOfType(Type, OutParameters);
 	}
 
-	TArray<const UMaterialInstance*> InstanceChain;
-	InstanceChain.Add(this);
-	{
-		const UMaterialInterface* CurrentParent = Parent;
-		while (CurrentParent)
-		{
-			const UMaterialInstance* ParentInstance = Cast<UMaterialInstance>(CurrentParent);
-			if (ParentInstance)
-			{
-				InstanceChain.Add(ParentInstance);
-				CurrentParent = ParentInstance->Parent;
-			}
-			else
-			{
-				check(CurrentParent == Material);
-				CurrentParent = nullptr;
-			}
-		}
-	}
-
 	for (int32 i = InstanceChain.Num() - 1; i >= 0; --i)
 	{
 		const UMaterialInstance* Instance = InstanceChain[i];
diff --git a/Engine/Source/Runtime/Engine/Private/Materials/MaterialInstanceSupport.h b/Engine/Source/Runtime/Engine/Private/Materials/MaterialInstanceSupport.h
--- a/Engine/Source/Runtime/Engine/Private/Materials/MaterialInstanceSupport.h
+++ b/Engine/Source/Runtime/Engine/Private/Materials/MaterialInstanceSupport.h
@@ -275,7 +275,7 @@ const ParameterType* GameThread_FindParameterByIndex(const TArray<ParameterType>
 }
 
 template <typename ParameterType>
-FORCEINLINE bool GameThread_GetParameterValue(const TArray<ParameterType>& Parameters, const FHashedMaterialParameterInfo& ParameterInfo, bool bSetOverride, FMaterialParameterMetadata& OutResult)
+FORCEINLINE bool GameThread_GetParameterValue(const TArray<ParameterType>& Parameters, const FHashedMaterialParameterInfo& ParameterInfo, FMaterialParameterMetadata& OutResult)
 {
 	for (int32 ParameterIndex = 0; ParameterIndex < Parameters.Num(); ParameterIndex++)
 	{
@@ -283,12 +283,6 @@ FORCEINLINE bool GameThread_GetParameterValue(const TArray<ParameterType>& Param
 		if (Parameter->IsOverride() && Parameter->ParameterInfo == ParameterInfo)
 		{
 			Parameter->GetValue(OutResult);
-#if WITH_EDITORONLY_DATA
-			if (bSetOverride)
-			{
-				OutResult.bOverride = true;
-			}
-#endif // WITH_EDITORONLY_DATA
 			return true;
 		}
 	}
diff --git a/Engine/Source/Runtime/Engine/Private/Materials/MaterialInterface.cpp b/Engine/Source/Runtime/Engine/Private/Materials/MaterialInterface.cpp
--- a/Engine/Source/Runtime/Engine/Private/Materials/MaterialInterface.cpp
+++ b/Engine/Source/Runtime/Engine/Private/Materials/MaterialInterface.cpp
@@ -63,13 +63,9 @@ bool IsCompatibleWithHairStrands(EShaderPlatform Platform, const FMaterialShader
 		(Parameters.BlendMode == BLEND_Opaque || Parameters.BlendMode == BLEND_Masked);
 }
 
-static EMaterialGetParameterValueFlags MakeParameterValueFlags(bool bOveriddenOnly, bool bCheckParent = true)
+static EMaterialGetParameterValueFlags MakeParameterValueFlags(bool bOveriddenOnly)
 {
 	EMaterialGetParameterValueFlags Result = EMaterialGetParameterValueFlags::CheckInstanceOverrides;
-	if (bCheckParent)
-	{
-		Result |= EMaterialGetParameterValueFlags::CheckParent;
-	}
 	if (!bOveriddenOnly)
 	{
 		Result |= EMaterialGetParameterValueFlags::CheckNonOverrides;
@@ -136,10 +132,10 @@ void UMaterialInterface::GetUsedTexturesAndIndices(TArray<UTexture*>& OutTexture
 }
 
 #if WITH_EDITORONLY_DATA
-bool UMaterialInterface::GetStaticSwitchParameterValue(const FHashedMaterialParameterInfo& ParameterInfo, bool& OutValue, FGuid& OutExpressionGuid, bool bOveriddenOnly /*= false*/, bool bCheckParent /*= true*/) const
+bool UMaterialInterface::GetStaticSwitchParameterValue(const FHashedMaterialParameterInfo& ParameterInfo, bool& OutValue, FGuid& OutExpressionGuid, bool bOveriddenOnly /*= false*/) const
 {
 	FMaterialParameterMetadata Result;
-	if (GetParameterValue(EMaterialParameterType::StaticSwitch, ParameterInfo, Result, MakeParameterValueFlags(bOveriddenOnly, bCheckParent)))
+	if (GetParameterValue(EMaterialParameterType::StaticSwitch, ParameterInfo, Result, MakeParameterValueFlags(bOveriddenOnly)))
 	{
 		OutExpressionGuid = Result.ExpressionGuid;
 		OutValue = Result.Value.AsStaticSwitch();
@@ -148,10 +144,10 @@ bool UMaterialInterface::GetStaticSwitchParameterValue(const FHashedMaterialPara
 	return false;
 }
 
-bool UMaterialInterface::GetStaticComponentMaskParameterValue(const FHashedMaterialParameterInfo& ParameterInfo, bool& R, bool& G, bool& B, bool& A, FGuid& OutExpressionGuid, bool bOveriddenOnly /*= false*/, bool bCheckParent /*= true*/) const
+bool UMaterialInterface::GetStaticComponentMaskParameterValue(const FHashedMaterialParameterInfo& ParameterInfo, bool& R, bool& G, bool& B, bool& A, FGuid& OutExpressionGuid, bool bOveriddenOnly /*= false*/) const
 {
 	FMaterialParameterMetadata Result;
-	if (GetParameterValue(EMaterialParameterType::StaticSwitch, ParameterInfo, Result, MakeParameterValueFlags(bOveriddenOnly, bCheckParent)))
+	if (GetParameterValue(EMaterialParameterType::StaticSwitch, ParameterInfo, Result, MakeParameterValueFlags(bOveriddenOnly)))
 	{
 		OutExpressionGuid = Result.ExpressionGuid;
 		R = Result.Value.Bool[0];
@@ -428,7 +424,7 @@ bool UMaterialInterface::GetVectorParameterValue(const FHashedMaterialParameterI
 bool UMaterialInterface::IsVectorParameterUsedAsChannelMask(const FHashedMaterialParameterInfo& ParameterInfo, bool& OutValue) const
 {
 	FMaterialParameterMetadata Result;
-	if (GetParameterValue(EMaterialParameterType::Vector, ParameterInfo, Result, EMaterialGetParameterValueFlags::CheckNonOverrides | EMaterialGetParameterValueFlags::CheckParent))
+	if (GetParameterValue(EMaterialParameterType::Vector, ParameterInfo, Result, EMaterialGetParameterValueFlags::CheckNonOverrides))
 	{
 		OutValue = Result.bUsedAsChannelMask;
 		return true;
@@ -439,7 +435,7 @@ bool UMaterialInterface::IsVectorParameterUsedAsChannelMask(const FHashedMateria
 bool UMaterialInterface::GetVectorParameterChannelNames(const FHashedMaterialParameterInfo& ParameterInfo, FParameterChannelNames& OutValue) const
 {
 	FMaterialParameterMetadata Result;
-	if (GetParameterValue(EMaterialParameterType::Vector, ParameterInfo, Result, EMaterialGetParameterValueFlags::CheckNonOverrides | EMaterialGetParameterValueFlags::CheckParent))
+	if (GetParameterValue(EMaterialParameterType::Vector, ParameterInfo, Result, EMaterialGetParameterValueFlags::CheckNonOverrides))
 	{
 		OutValue = Result.ChannelNames;
 		return true;
@@ -450,7 +446,7 @@ bool UMaterialInterface::GetVectorParameterChannelNames(const FHashedMaterialPar
 bool UMaterialInterface::GetScalarParameterSliderMinMax(const FHashedMaterialParameterInfo& ParameterInfo, float& OutSliderMin, float& OutSliderMax) const
 {
 	FMaterialParameterMetadata Result;
-	if (GetParameterValue(EMaterialParameterType::Scalar, ParameterInfo, Result, EMaterialGetParameterValueFlags::CheckNonOverrides | EMaterialGetParameterValueFlags::CheckParent))
+	if (GetParameterValue(EMaterialParameterType::Scalar, ParameterInfo, Result, EMaterialGetParameterValueFlags::CheckNonOverrides))
 	{
 		OutSliderMin = Result.ScalarMin;
 		OutSliderMax = Result.ScalarMax;
@@ -480,7 +476,7 @@ bool UMaterialInterface::GetParameterValue(EMaterialParameterType Type, const FM
 bool UMaterialInterface::IsScalarParameterUsedAsAtlasPosition(const FHashedMaterialParameterInfo& ParameterInfo, bool& OutValue, TSoftObjectPtr<class UCurveLinearColor>& Curve, TSoftObjectPtr<class UCurveLinearColorAtlas>& Atlas) const
 {
 	FMaterialParameterMetadata Result;
-	if (GetParameterValue(EMaterialParameterType::Scalar, ParameterInfo, Result, EMaterialGetParameterValueFlags::CheckNonOverrides | EMaterialGetParameterValueFlags::CheckParent))
+	if (GetParameterValue(EMaterialParameterType::Scalar, ParameterInfo, Result, EMaterialGetParameterValueFlags::CheckNonOverrides))
 	{
 		OutValue = Result.bUsedAsAtlasPosition;
 		Curve = Result.ScalarCurve;
@@ -517,7 +513,7 @@ bool UMaterialInterface::GetRuntimeVirtualTextureParameterValue(const FHashedMat
 bool UMaterialInterface::GetTextureParameterChannelNames(const FHashedMaterialParameterInfo& ParameterInfo, FParameterChannelNames& OutValue) const
 {
 	FMaterialParameterMetadata Result;
-	if (GetParameterValue(EMaterialParameterType::Texture, ParameterInfo, Result, EMaterialGetParameterValueFlags::CheckNonOverrides | EMaterialGetParameterValueFlags::CheckParent))
+	if (GetParameterValue(EMaterialParameterType::Texture, ParameterInfo, Result, EMaterialGetParameterValueFlags::CheckNonOverrides))
 	{
 		OutValue = Result.ChannelNames;
 		return true;
@@ -540,7 +536,7 @@ bool UMaterialInterface::GetFontParameterValue(const FHashedMaterialParameterInf
 
 bool UMaterialInterface::GetParameterDefaultValue(EMaterialParameterType Type, const FMemoryImageMaterialParameterInfo& ParameterInfo, FMaterialParameterMetadata& OutValue) const
 {
-	return GetParameterValue(Type, ParameterInfo, OutValue, EMaterialGetParameterValueFlags::CheckNonOverrides | EMaterialGetParameterValueFlags::CheckParent);
+	return GetParameterValue(Type, ParameterInfo, OutValue, EMaterialGetParameterValueFlags::CheckNonOverrides);
 }
 
 
diff --git a/Engine/Source/Runtime/Engine/Public/MaterialTypes.h b/Engine/Source/Runtime/Engine/Public/MaterialTypes.h
--- a/Engine/Source/Runtime/Engine/Public/MaterialTypes.h
+++ b/Engine/Source/Runtime/Engine/Public/MaterialTypes.h
@@ -58,10 +58,8 @@ enum class EMaterialGetParameterValueFlags : uint32
 	None = 0u,
 	CheckNonOverrides = (1u << 0),
 	CheckInstanceOverrides = (1u << 1),
-	CheckParent = (1u << 2),
-	IsParent = (1u << 3), // automatically set when recursing into a parent
 
-	CheckAll = CheckNonOverrides | CheckInstanceOverrides | CheckParent,
+	CheckAll = CheckNonOverrides | CheckInstanceOverrides,
 };
 ENUM_CLASS_FLAGS(EMaterialGetParameterValueFlags);
 
