FIXED: Crashes due to missed merge

diff --git a/Engine/Source/Editor/DetailCustomizations/Private/CustomPrimitiveDataCustomization.cpp b/Engine/Source/Editor/DetailCustomizations/Private/CustomPrimitiveDataCustomization.cpp
--- a/Engine/Source/Editor/DetailCustomizations/Private/CustomPrimitiveDataCustomization.cpp
+++ b/Engine/Source/Editor/DetailCustomizations/Private/CustomPrimitiveDataCustomization.cpp
@@ -275,12 +275,16 @@ void FCustomPrimitiveDataCustomization::CreateParameterRow(IDetailChildrenBuilde
 			{
 			case 0:
 				ChannelName = ParameterMetadata.ChannelNames.R.IsEmpty() ? LOCTEXT("DefaultVectorChannelRed", "R") : ParameterMetadata.ChannelNames.R;
+				break;
 			case 1:
 				ChannelName = ParameterMetadata.ChannelNames.G.IsEmpty() ? LOCTEXT("DefaultVectorChannelGreen", "G") : ParameterMetadata.ChannelNames.G;
+				break;
 			case 2:
 				ChannelName = ParameterMetadata.ChannelNames.B.IsEmpty() ? LOCTEXT("DefaultVectorChannelBlue", "B") : ParameterMetadata.ChannelNames.B;
+				break;
 			case 3:
 				ChannelName = ParameterMetadata.ChannelNames.A.IsEmpty() ? LOCTEXT("DefaultVectorChannelAlpha", "A") : ParameterMetadata.ChannelNames.A;
+				break;
 			default:
 				checkNoEntry();
 				break;
diff --git a/Engine/Source/Runtime/Engine/Private/Materials/MaterialCachedData.cpp b/Engine/Source/Runtime/Engine/Private/Materials/MaterialCachedData.cpp
--- a/Engine/Source/Runtime/Engine/Private/Materials/MaterialCachedData.cpp
+++ b/Engine/Source/Runtime/Engine/Private/Materials/MaterialCachedData.cpp
@@ -479,6 +479,7 @@ bool FMaterialCachedExpressionData::UpdateForExpressions(const FMaterialCachedEx
 				case EMaterialParameterType::Scalar:
 					Parameters.ScalarValues.Insert(ParameterMeta.Value.AsScalar(), Index);
 					Parameters.ScalarMinMaxValues.Insert(FVector2D(ParameterMeta.ScalarMin, ParameterMeta.ScalarMax), Index);
+					Parameters.ScalarPrimitiveDataIndexValues.Insert(ParameterMeta.PrimitiveDataIndex, Index);
 					if (ParameterMeta.bUsedAsAtlasPosition)
 					{
 						Parameters.ScalarCurveValues.Insert(ParameterMeta.ScalarCurve.Get(), Index);
@@ -494,6 +495,7 @@ bool FMaterialCachedExpressionData::UpdateForExpressions(const FMaterialCachedEx
 					Parameters.VectorValues.Insert(ParameterMeta.Value.AsLinearColor(), Index);
 					Parameters.VectorChannelNameValues.Insert(ParameterMeta.ChannelNames, Index);
 					Parameters.VectorUsedAsChannelMaskValues.Insert(ParameterMeta.bUsedAsChannelMask, Index);
+					Parameters.VectorPrimitiveDataIndexValues.Insert(ParameterMeta.PrimitiveDataIndex, Index);
 					break;
 				case EMaterialParameterType::Texture:
 					Parameters.TextureValues.Insert(ParameterMeta.Value.Texture, Index);
