Clean up clearing material parameters when copying back from PreviewMaterial.  Fixes static parameters not properly copying back (since the runtime impl of those also includes an override flag)

diff --git a/Engine/Source/Editor/UnrealEd/Private/PreviewMaterial.cpp b/Engine/Source/Editor/UnrealEd/Private/PreviewMaterial.cpp
--- a/Engine/Source/Editor/UnrealEd/Private/PreviewMaterial.cpp
+++ b/Engine/Source/Editor/UnrealEd/Private/PreviewMaterial.cpp
@@ -891,10 +891,8 @@ void UMaterialEditorInstanceConstant::CopyToSourceInstance(const bool bForceStat
 			SourceInstance->MarkPackageDirty();
 		}
 
-		SourceInstance->ClearParameterValuesEditorOnly();
-
 		{
-			FMaterialInstanceParameterUpdateContext UpdateContext(SourceInstance);
+			FMaterialInstanceParameterUpdateContext UpdateContext(SourceInstance, EMaterialInstanceClearParameterFlag::All);
 			UpdateContext.SetBasePropertyOverrides(BasePropertyOverrides);
 			UpdateContext.SetForceStaticPermutationUpdate(bForceStaticPermutationUpdate);
 
@@ -999,7 +997,7 @@ void UMaterialEditorInstanceConstant::SetSourceInstance(UMaterialInstanceConstan
 	RegenerateArrays();
 
 	//propagate changes to the base material so the instance will be updated if it has a static permutation resource
-	FMaterialInstanceParameterUpdateContext UpdateContext(SourceInstance);
+	FMaterialInstanceParameterUpdateContext UpdateContext(SourceInstance, EMaterialInstanceClearParameterFlag::Static);
 	for (int32 GroupIdx = 0; GroupIdx < ParameterGroups.Num(); GroupIdx++)
 	{
 		FEditorParameterGroup& Group = ParameterGroups[GroupIdx];
diff --git a/Engine/Source/Runtime/Engine/Classes/Materials/MaterialInstance.h b/Engine/Source/Runtime/Engine/Classes/Materials/MaterialInstance.h
--- a/Engine/Source/Runtime/Engine/Classes/Materials/MaterialInstance.h
+++ b/Engine/Source/Runtime/Engine/Classes/Materials/MaterialInstance.h
@@ -346,11 +346,23 @@ bool CompareValueArraysByExpressionGUID(const TArray<T>& InA, const TArray<T>& I
 	return AA == BB;
 }
 
+enum class EMaterialInstanceClearParameterFlag
+{
+	None = 0u,
+	Numeric = (1u << 0),
+	Texture = (1u << 1),
+	Static = (1u << 2),
+
+	AllNonStatic = Numeric | Texture,
+	All = AllNonStatic | Static,
+};
+ENUM_CLASS_FLAGS(EMaterialInstanceClearParameterFlag);
+
 #if WITH_EDITORONLY_DATA
 class FMaterialInstanceParameterUpdateContext
 {
 public:
-	ENGINE_API explicit FMaterialInstanceParameterUpdateContext(UMaterialInstance* InInstance);
+	ENGINE_API explicit FMaterialInstanceParameterUpdateContext(UMaterialInstance* InInstance, EMaterialInstanceClearParameterFlag ClearFlags = EMaterialInstanceClearParameterFlag::None);
 	ENGINE_API ~FMaterialInstanceParameterUpdateContext();
 
 	inline FStaticParameterSet& GetStaticParameters() { return StaticParameters; }
@@ -826,7 +838,7 @@ protected:
 	void SetTextureParameterValueInternal(const FMaterialParameterInfo& ParameterInfo, class UTexture* Value);
 	void SetRuntimeVirtualTextureParameterValueInternal(const FMaterialParameterInfo& ParameterInfo, class URuntimeVirtualTexture* Value);
 	void SetFontParameterValueInternal(const FMaterialParameterInfo& ParameterInfo, class UFont* FontValue, int32 FontPage);
-	void ClearParameterValuesInternal(const bool bAllParameters = true);
+	void ClearParameterValuesInternal(EMaterialInstanceClearParameterFlag Flags = EMaterialInstanceClearParameterFlag::AllNonStatic);
 
 	/** Initialize the material instance's resources. */
 	ENGINE_API void InitResources();
diff --git a/Engine/Source/Runtime/Engine/Private/Materials/MaterialInstance.cpp b/Engine/Source/Runtime/Engine/Private/Materials/MaterialInstance.cpp
--- a/Engine/Source/Runtime/Engine/Private/Materials/MaterialInstance.cpp
+++ b/Engine/Source/Runtime/Engine/Private/Materials/MaterialInstance.cpp
@@ -1440,7 +1440,7 @@ void UMaterialInstanceDynamic::CopyScalarAndVectorParameters(const UMaterialInte
 	if (MaterialResource)
 	{
 		// first, clear out all the parameter values
-		ClearParameterValuesInternal(false);
+		ClearParameterValuesInternal(EMaterialInstanceClearParameterFlag::Numeric);
 
 		TArrayView<const FMaterialNumericParameterInfo> Array = MaterialResource->GetUniformNumericParameterExpressions();
 		for (int32 i = 0, Count = Array.Num(); i < Count; ++i)
@@ -2839,13 +2839,25 @@ bool UMaterialInstance::SetVectorParameterByIndexInternal(int32 ParameterIndex,
 }
 
 #if WITH_EDITORONLY_DATA
-FMaterialInstanceParameterUpdateContext::FMaterialInstanceParameterUpdateContext(UMaterialInstance* InInstance)
+FMaterialInstanceParameterUpdateContext::FMaterialInstanceParameterUpdateContext(UMaterialInstance* InInstance, EMaterialInstanceClearParameterFlag InFlags)
 	: Instance(InInstance)
 	, bForceStaticPermutationUpdate(false)
 {
 	check(InInstance);
-	StaticParameters = InInstance->GetStaticParameters();
+	EMaterialInstanceClearParameterFlag Flags = InFlags;
+	if (EnumHasAnyFlags(Flags, EMaterialInstanceClearParameterFlag::Static))
+	{
+		// If we ask to clear static parameters, simply avoid copying them
+		Flags &= ~EMaterialInstanceClearParameterFlag::Static;
+	}
+	else
+	{
+		StaticParameters = InInstance->GetStaticParameters();
+	}
+
 	BasePropertyOverrides = InInstance->BasePropertyOverrides;
+
+	InInstance->ClearParameterValuesInternal(Flags);
 }
 
 FMaterialInstanceParameterUpdateContext::~FMaterialInstanceParameterUpdateContext()
@@ -3155,19 +3167,30 @@ void UMaterialInstance::SetFontParameterValueInternal(const FMaterialParameterIn
 	}
 }
 
-void UMaterialInstance::ClearParameterValuesInternal(const bool bAllParameters)
+void UMaterialInstance::ClearParameterValuesInternal(EMaterialInstanceClearParameterFlag Flags)
 {
-	ScalarParameterValues.Empty();
-	VectorParameterValues.Empty();
+	bool bUpdateResource = false;
+	if (EnumHasAnyFlags(Flags, EMaterialInstanceClearParameterFlag::Numeric))
+	{
+		ScalarParameterValues.Empty();
+		VectorParameterValues.Empty();
+		bUpdateResource = true;
+	}
 
-	if(bAllParameters)
+	if (EnumHasAnyFlags(Flags, EMaterialInstanceClearParameterFlag::Texture))
 	{
 		TextureParameterValues.Empty();
 		RuntimeVirtualTextureParameterValues.Empty();
 		FontParameterValues.Empty();
+		bUpdateResource = true;
 	}
 
-	if (Resource)
+	if (EnumHasAnyFlags(Flags, EMaterialInstanceClearParameterFlag::Static))
+	{
+		StaticParameters.Empty();
+	}
+
+	if (Resource && bUpdateResource)
 	{
 		FMaterialInstanceResource* InResource = Resource;
 		ENQUEUE_RENDER_COMMAND(FClearMIParametersCommand)(
diff --git a/Engine/Source/Runtime/Engine/Public/StaticParameterSet.h b/Engine/Source/Runtime/Engine/Public/StaticParameterSet.h
--- a/Engine/Source/Runtime/Engine/Public/StaticParameterSet.h
+++ b/Engine/Source/Runtime/Engine/Public/StaticParameterSet.h
@@ -401,6 +401,14 @@ struct FStaticParameterSet
 		return StaticSwitchParameters.Num() == 0 && StaticComponentMaskParameters.Num() == 0 && TerrainLayerWeightParameters.Num() == 0 && MaterialLayersParameters.Num() == 0;
 	}
 
+	void Empty()
+	{
+		StaticSwitchParameters.Empty();
+		StaticComponentMaskParameters.Empty();
+		TerrainLayerWeightParameters.Empty();
+		MaterialLayersParameters.Empty();
+	}
+
 	void Serialize(FArchive& Ar)
 	{
 		Ar.UsingCustomVersion(FRenderingObjectVersion::GUID);
