Make it possible to load levels managed by WorldComposition with the new loader.
Enable the existing serialization path for PIE also for cooked data.

#rnx

diff --git a/Engine/Source/Runtime/Engine/Private/WorldComposition.cpp b/Engine/Source/Runtime/Engine/Private/WorldComposition.cpp
--- a/Engine/Source/Runtime/Engine/Private/WorldComposition.cpp
+++ b/Engine/Source/Runtime/Engine/Private/WorldComposition.cpp
@@ -45,9 +45,9 @@ void UWorldComposition::PostInitProperties()
 {
 	Super::PostInitProperties();
 
-	if (!IsTemplate() && !GetOutermost()->HasAnyPackageFlags(PKG_PlayInEditor))
+	if (!IsTemplate() && !GetOutermost()->HasAnyPackageFlags(PKG_PlayInEditor) && !FPlatformProperties::RequiresCookedData())
 	{
-		// Tiles information is not serialized to disk, and should be regenerated on world composition object construction
+		// Tiles information is not serialized to disk in the editor, and should be regenerated on world composition object construction
 		Rescan();
 	}
 }
@@ -56,9 +56,11 @@ void UWorldComposition::Serialize( FArchive& Ar )
 {
 	Super::Serialize(Ar);
 	
-	// We serialize this data only for PIE
-	// In normal game this data is regenerated on object construction
-	if (Ar.GetPortFlags() & PPF_DuplicateForPIE)
+	// We serialize this data only for PIE and cooked data
+	// In normal editor game this data is regenerated on object construction
+	const bool bIsPIE = Ar.GetPortFlags() & PPF_DuplicateForPIE;
+	const bool bIsCooked = Ar.IsCooking() || (Ar.IsLoading() && FPlatformProperties::RequiresCookedData());
+	if (bIsPIE || bIsCooked)
 	{
 		Ar << WorldRoot;
 		Ar << Tiles;
@@ -80,6 +82,14 @@ void UWorldComposition::PostLoad()
 {
 	Super::PostLoad();
 
+	if (FPlatformProperties::RequiresCookedData())
+	{
+		// Create streaming levels for each Tile
+		PopulateStreamingLevels();
+		// Calculate absolute positions since they are not serialized to disk
+		CaclulateTilesAbsolutePositions();
+	}
+
 	UWorld* World = GetWorld();
 
 	if (World->IsGameWorld())
@@ -231,6 +241,8 @@ struct FWorldTilesGatherer
 
 void UWorldComposition::Rescan()
 {
+	checkfSlow(!FPlatformProperties::RequiresCookedData(), TEXT("Rescan should not be called in a cooked game with serialized tile info."));
+
 	// Save tiles state, so we can restore it for dirty tiles after rescan is done
 	FTilesList SavedTileList = Tiles;
 		
