[Navigation WP] Some fixes to revive the WorldPartitionNavigationDataBuilder commandlet

diff --git a/Engine/Source/Editor/UnrealEd/Private/WorldPartition/WorldPartitionBuilder.cpp b/Engine/Source/Editor/UnrealEd/Private/WorldPartition/WorldPartitionBuilder.cpp
--- a/Engine/Source/Editor/UnrealEd/Private/WorldPartition/WorldPartitionBuilder.cpp
+++ b/Engine/Source/Editor/UnrealEd/Private/WorldPartition/WorldPartitionBuilder.cpp
@@ -76,7 +76,7 @@ bool UWorldPartitionBuilder::RunBuilder(UWorldPartitionBuilder* Builder, UWorld*
 	}
 
 	// Setup the world if needed
-	bool bWorldWasInitialized = World->bIsWorldInitialized;
+	const bool bWorldWasInitialized = World->bIsWorldInitialized;
 	if (!bWorldWasInitialized)
 	{
 		World->WorldType = EWorldType::Editor;
diff --git a/Engine/Source/Editor/UnrealEd/Private/WorldPartition/WorldPartitionNavigationDataBuilder.cpp b/Engine/Source/Editor/UnrealEd/Private/WorldPartition/WorldPartitionNavigationDataBuilder.cpp
--- a/Engine/Source/Editor/UnrealEd/Private/WorldPartition/WorldPartitionNavigationDataBuilder.cpp
+++ b/Engine/Source/Editor/UnrealEd/Private/WorldPartition/WorldPartitionNavigationDataBuilder.cpp
@@ -5,6 +5,7 @@
 #include "CoreMinimal.h"
 #include "EngineUtils.h"
 #include "Editor.h"
+#include "StaticMeshCompiler.h"
 #include "Logging/LogMacros.h"
 
 #include "WorldPartition/WorldPartition.h"
@@ -34,6 +35,8 @@ bool UWorldPartitionNavigationDataBuilder::RunInternal(UWorld* World, const FBox
 		NavigationDataChunkActorPackages.Add(ItActor->GetPackage());
 	}
 
+	FStaticMeshCompilingManager::Get().FinishAllCompilation();
+	
 	// Rebuild ANavigationDataChunkActor for the whole world
 	WorldPartition->GenerateNavigationData();
 
@@ -77,12 +80,23 @@ bool UWorldPartitionNavigationDataBuilder::RunInternal(UWorld* World, const FBox
 	{
 		// Checkout packages to save
 		UE_LOG(LogWorldPartitionNavigationDataBuilder, Log, TEXT("Checking out %d actor packages."), PackagesToSave.Num());
-		for (UPackage* Package : PackagesToSave)
+		if (PackageHelper.UseSourceControl())
+		{
+			FEditorFileUtils::CheckoutPackages(PackagesToSave, /*OutPackagesCheckedOut*/nullptr, /*bErrorIfAlreadyCheckedOut*/false);
+		}
+		else
 		{
-			if (!PackageHelper.Checkout(Package))
+			for (const UPackage* Package : PackagesToSave)
 			{
-				UE_LOG(LogWorldPartitionNavigationDataBuilder, Error, TEXT("Error checking out package %s."), *Package->GetName());
-				return 1;
+				const FString PackageFilename = SourceControlHelpers::PackageFilename(Package);
+				if (IPlatformFile::GetPlatformPhysical().FileExists(*PackageFilename))
+				{
+					if (!IPlatformFile::GetPlatformPhysical().SetReadOnly(*PackageFilename, /*bNewReadOnlyValue*/false))
+					{
+						UE_LOG(LogWorldPartitionNavigationDataBuilder, Error, TEXT("Error setting %s writable"), *PackageFilename);
+						return 1;
+					}
+				}
 			}
 		}
 
@@ -114,4 +128,4 @@ bool UWorldPartitionNavigationDataBuilder::RunInternal(UWorld* World, const FBox
 	}
 
 	return true;
-}
\ No newline at end of file
+}
diff --git a/Engine/Source/Runtime/Engine/Private/WorldPartition/RuntimeSpatialHash/RuntimeSpatialHashNavigation.cpp b/Engine/Source/Runtime/Engine/Private/WorldPartition/RuntimeSpatialHash/RuntimeSpatialHashNavigation.cpp
--- a/Engine/Source/Runtime/Engine/Private/WorldPartition/RuntimeSpatialHash/RuntimeSpatialHashNavigation.cpp
+++ b/Engine/Source/Runtime/Engine/Private/WorldPartition/RuntimeSpatialHash/RuntimeSpatialHashNavigation.cpp
@@ -72,39 +72,34 @@ bool UWorldPartitionRuntimeSpatialHash::GenerateNavigationData()
 				return;
 			}
 
-			const FSquare2DGridHelper::FGridLevel::FGridCell& GridCell = GridLevelHelper.GetCell(CellCoord);
-
-			//@todo_ow: Properly handle data layers
-			if (const FSquare2DGridHelper::FGridLevel::FGridCellDataChunk* GridCellDataChunk = GridCell.GetNoDataLayersDataChunk())
-			{
-				FActorSpawnParameters SpawnParams;
-				SpawnParams.bDeferConstruction = true;
-				SpawnParams.bCreateActorPackage = true;
-				ANavigationDataChunkActor* DataChunkActor = World->SpawnActor<ANavigationDataChunkActor>(SpawnParams);
-				ActorCount++;
-
-				const FVector2D CellCenter = CellBounds.GetCenter();
-				DataChunkActor->SetActorLocation(FVector(CellCenter.X, CellCenter.Y, 0.f));
-
-				FBox TilesBounds(EForceInit::ForceInit);
-				DataChunkActor->CollectNavData(QueryBounds, TilesBounds);
-
-				FBox ChunkActorBounds(FVector(QueryBounds.Min.X, QueryBounds.Min.Y, TilesBounds.Min.Z), FVector(QueryBounds.Max.X, QueryBounds.Max.Y, TilesBounds.Max.Z));
-				ChunkActorBounds = ChunkActorBounds.ExpandBy(FVector(-1.f, -1.f, 1.f)); //reduce XY by 1cm to avoid precision issues causing potential overflow on neighboring cell, add 1cm in Z to have a minimum of volume.
-				UE_LOG(LogWorldPartitionRuntimeSpatialHashNav, VeryVerbose, TEXT("Setting ChunkActorBounds to %s"), *ChunkActorBounds.ToString());
-				DataChunkActor->SetDataChunkActorBounds(ChunkActorBounds);
-
-				FIntVector CellGlobalCoord;
-				verify(GridHelper.GetCellGlobalCoords(FIntVector(CellCoord.X, CellCoord.Y, GridLevel), CellGlobalCoord));
-				const FName CellName = GetCellName(RuntimeGrid.GridName, CellGlobalCoord, GridCellDataChunk->GetDataLayersID());
-				DataChunkActor->SetActorLabel(FString::Printf(TEXT("NavDataChunkActor_%s_%s"), *GetName(), *CellName.ToString()));
-
-				// Set target grid
-				DataChunkActor->SetRuntimeGrid(RuntimeGrid.GridName);
-				ValidNavigationDataChunkActors.Add(DataChunkActor);
-
-				UE_LOG(LogWorldPartitionRuntimeSpatialHashNav, Verbose, TEXT("%i) %s added."), ActorCount, *DataChunkActor->GetName());
-			}
+			 //@todo_ow: Properly handle data layers
+			FActorSpawnParameters SpawnParams;
+			SpawnParams.bDeferConstruction = true;
+			SpawnParams.bCreateActorPackage = true;
+			ANavigationDataChunkActor* DataChunkActor = World->SpawnActor<ANavigationDataChunkActor>(SpawnParams);
+			ActorCount++;
+
+			const FVector2D CellCenter = CellBounds.GetCenter();
+			DataChunkActor->SetActorLocation(FVector(CellCenter.X, CellCenter.Y, 0.f));
+
+			FBox TilesBounds(EForceInit::ForceInit);
+			DataChunkActor->CollectNavData(QueryBounds, TilesBounds);
+
+			FBox ChunkActorBounds(FVector(QueryBounds.Min.X, QueryBounds.Min.Y, TilesBounds.Min.Z), FVector(QueryBounds.Max.X, QueryBounds.Max.Y, TilesBounds.Max.Z));
+			ChunkActorBounds = ChunkActorBounds.ExpandBy(FVector(-1.f, -1.f, 1.f)); //reduce XY by 1cm to avoid precision issues causing potential overflow on neighboring cell, add 1cm in Z to have a minimum of volume.
+			UE_LOG(LogWorldPartitionRuntimeSpatialHashNav, VeryVerbose, TEXT("Setting ChunkActorBounds to %s"), *ChunkActorBounds.ToString());
+			DataChunkActor->SetDataChunkActorBounds(ChunkActorBounds);
+
+			FIntVector CellGlobalCoord;
+			verify(GridHelper.GetCellGlobalCoords(FIntVector(CellCoord.X, CellCoord.Y, GridLevel), CellGlobalCoord));
+			const FName CellName = GetCellName(RuntimeGrid.GridName, CellGlobalCoord, FDataLayersID());
+			DataChunkActor->SetActorLabel(FString::Printf(TEXT("NavDataChunkActor_%s_%s"), *GetName(), *CellName.ToString()));
+
+			// Set target grid
+			DataChunkActor->SetRuntimeGrid(RuntimeGrid.GridName);
+			ValidNavigationDataChunkActors.Add(DataChunkActor);
+
+			UE_LOG(LogWorldPartitionRuntimeSpatialHashNav, Verbose, TEXT("%i) %s added."), ActorCount, *DataChunkActor->GetName());
 		});
 
 	// Destroy all invalid navigation data chunk actors
