Prevent material render proxies for destroyed materials from being updated (UE-127390)

diff --git a/Engine/Source/Runtime/Engine/Private/Materials/MaterialShared.cpp b/Engine/Source/Runtime/Engine/Private/Materials/MaterialShared.cpp
--- a/Engine/Source/Runtime/Engine/Private/Materials/MaterialShared.cpp
+++ b/Engine/Source/Runtime/Engine/Private/Materials/MaterialShared.cpp
@@ -52,6 +52,7 @@
 #include "MaterialHLSLTree.h"
 #include "MaterialHLSLEmitter.h"
 #include "HLSLTree/HLSLTreeCommon.h"
+#include "Misc/ScopeLock.h"
 #endif
 #if WITH_ODSC
 #include "ODSC/ODSCManager.h"
@@ -3303,19 +3304,26 @@ FMaterialRenderProxy::FMaterialRenderProxy()
 	, DeletedFlag(0)
 	, HasVirtualTextureCallbacks(0)
 {
-	// MaterialRenderProxyMap is only used by shader compiling
+#if WITH_EDITOR
+	// We need to register ourselves in MaterialRenderProxyMap early as sometimes (during the shader compilation) we need to be updated even before InitDynamicRHI (see e.g. UE-117842)
 	if (!FPlatformProperties::RequiresCookedData())
 	{
+		FScopeLock Locker(&MaterialRenderProxyMapLock);
 		FMaterialRenderProxy::MaterialRenderProxyMap.Add(this);
 	}
+#endif // WITH_EDITOR
 }
 
 FMaterialRenderProxy::~FMaterialRenderProxy()
 {
+#if WITH_EDITOR
+	// this may also be removed in the ReleaseDynamicRHI, but leaving it here in case we never get the state created
 	if (!FPlatformProperties::RequiresCookedData())
 	{
+		FScopeLock Locker(&MaterialRenderProxyMapLock);
 		FMaterialRenderProxy::MaterialRenderProxyMap.Remove(this);
 	}
+#endif // WITH_EDITOR
 
 	if(IsInitialized())
 	{
@@ -3335,6 +3343,15 @@ FMaterialRenderProxy::~FMaterialRenderProxy()
 
 void FMaterialRenderProxy::ReleaseDynamicRHI()
 {
+#if WITH_EDITOR
+	// unregister ourselves early to prevent getting updates from the shader compilation when the parent has already got BeginDestroy()ed
+	if (!FPlatformProperties::RequiresCookedData())
+	{
+		FScopeLock Locker(&MaterialRenderProxyMapLock);
+		FMaterialRenderProxy::MaterialRenderProxyMap.Remove(this);
+	}
+#endif // WITH_EDITOR
+
 	DeferredUniformExpressionCacheRequests.Remove(this);
 
 	InvalidateUniformExpressionCache(true);
@@ -3444,7 +3461,10 @@ void FMaterialRenderProxy::UpdateDeferredCachedUniformExpressions()
 	DeferredUniformExpressionCacheRequests.Reset();
 }
 
+#if WITH_EDITOR
 TSet<FMaterialRenderProxy*> FMaterialRenderProxy::MaterialRenderProxyMap;
+FCriticalSection FMaterialRenderProxy::MaterialRenderProxyMapLock;
+#endif // WITH_EDITOR
 TSet<FMaterialRenderProxy*> FMaterialRenderProxy::DeferredUniformExpressionCacheRequests;
 
 /*-----------------------------------------------------------------------------
@@ -4914,6 +4934,7 @@ void SetShaderMapsOnMaterialResources_RenderThread(FRHICommandListImmediate& RHI
 {
 	SCOPE_CYCLE_COUNTER(STAT_Scene_SetShaderMapsOnMaterialResources_RT);
 
+#if WITH_EDITOR
 	bool bUpdateFeatureLevel[ERHIFeatureLevel::Num] = { false };
 
 	for (auto& It : MaterialsToUpdate)
@@ -4934,6 +4955,7 @@ void SetShaderMapsOnMaterialResources_RenderThread(FRHICommandListImmediate& RHI
 		{
 			const ERHIFeatureLevel::Type MaterialFeatureLevel = (ERHIFeatureLevel::Type) FeatureLevelIndex;
 
+			FScopeLock Locker(&FMaterialRenderProxy::GetMaterialRenderProxyMapLock());
 			for (TSet<FMaterialRenderProxy*>::TConstIterator It(FMaterialRenderProxy::GetMaterialRenderProxyMap()); It; ++It)
 			{
 				FMaterialRenderProxy* MaterialProxy = *It;
@@ -4954,6 +4976,7 @@ void SetShaderMapsOnMaterialResources_RenderThread(FRHICommandListImmediate& RHI
 			}
 		}
 	}
+#endif // WITH_EDITOR
 }
 
 void FMaterial::SetShaderMapsOnMaterialResources(const TMap<TRefCountPtr<FMaterial>, TRefCountPtr<FMaterialShaderMap>>& MaterialsToUpdate)
diff --git a/Engine/Source/Runtime/Engine/Public/MaterialShared.h b/Engine/Source/Runtime/Engine/Public/MaterialShared.h
--- a/Engine/Source/Runtime/Engine/Public/MaterialShared.h
+++ b/Engine/Source/Runtime/Engine/Public/MaterialShared.h
@@ -2355,12 +2355,19 @@ public:
 	ENGINE_API virtual void ReleaseDynamicRHI() override;
 	ENGINE_API virtual void ReleaseResource() override;
 
+#if WITH_EDITOR
 	ENGINE_API static const TSet<FMaterialRenderProxy*>& GetMaterialRenderProxyMap() 
 	{
 		check(!FPlatformProperties::RequiresCookedData());
 		return MaterialRenderProxyMap;
 	}
 
+	ENGINE_API static FCriticalSection& GetMaterialRenderProxyMapLock()
+	{
+		return MaterialRenderProxyMapLock;
+	}
+#endif
+
 	void SetSubsurfaceProfileRT(const USubsurfaceProfile* Ptr) { SubsurfaceProfileRT = Ptr; }
 	const USubsurfaceProfile* GetSubsurfaceProfileRT() const { return SubsurfaceProfileRT; }
 
@@ -2389,12 +2396,19 @@ private:
 	/** If any VT producer destroyed callbacks have been registered */
 	mutable int8 HasVirtualTextureCallbacks : 1;
 
+#if WITH_EDITOR
 	/** 
-	 * Tracks all material render proxies in all scenes, can only be accessed on the rendering thread.
+	 * Tracks all material render proxies in all scenes.
 	 * This is used to propagate new shader maps to materials being used for rendering.
 	 */
 	ENGINE_API static TSet<FMaterialRenderProxy*> MaterialRenderProxyMap;
 
+	/**
+	 * Lock that guards the access to the render proxy map
+	 */
+	ENGINE_API static FCriticalSection MaterialRenderProxyMapLock;
+#endif
+
 	ENGINE_API static TSet<FMaterialRenderProxy*> DeferredUniformExpressionCacheRequests;
 };
 
