Add missing scalar/vector custom data indices to material cached data

diff --git a/Engine/Source/Runtime/Engine/Private/Materials/MaterialCachedData.cpp b/Engine/Source/Runtime/Engine/Private/Materials/MaterialCachedData.cpp
--- a/Engine/Source/Runtime/Engine/Private/Materials/MaterialCachedData.cpp
+++ b/Engine/Source/Runtime/Engine/Private/Materials/MaterialCachedData.cpp
@@ -154,6 +154,7 @@ bool FMaterialCachedExpressionData::UpdateForFunction(const FMaterialCachedExpre
 			if (Index != INDEX_NONE)
 			{
 				Parameters.ScalarValues.Insert(Param.ParameterValue, Index);
+				Parameters.ScalarPrimitiveDataIndexValues.Insert(INDEX_NONE, Index);
 				Parameters.ScalarMinMaxValues.Insert(FVector2D(), Index);
 				if (Param.AtlasData.bIsUsedAsAtlasPosition)
 				{
@@ -175,6 +176,7 @@ bool FMaterialCachedExpressionData::UpdateForFunction(const FMaterialCachedExpre
 			if (Index != INDEX_NONE)
 			{
 				Parameters.VectorValues.Insert(Param.ParameterValue, Index);
+				Parameters.VectorPrimitiveDataIndexValues.Insert(INDEX_NONE, Index);
 				Parameters.VectorChannelNameValues.Insert(FParameterChannelNames(), Index);
 				Parameters.VectorUsedAsChannelMaskValues.Insert(false, Index);
 			}
@@ -348,6 +350,7 @@ void FMaterialCachedParameters_UpdateForLayerParameters(FMaterialCachedParameter
 				if (Index != INDEX_NONE)
 				{
 					Parameters.ScalarValues.Insert(Param.ParameterValue, Index);
+					Parameters.ScalarPrimitiveDataIndexValues.Insert(INDEX_NONE, Index);
 					Parameters.ScalarMinMaxValues.Insert(FVector2D(), Index);
 					if (Param.AtlasData.bIsUsedAsAtlasPosition)
 					{
@@ -372,6 +375,7 @@ void FMaterialCachedParameters_UpdateForLayerParameters(FMaterialCachedParameter
 				if (Index != INDEX_NONE)
 				{
 					Parameters.VectorValues.Insert(Param.ParameterValue, Index);
+					Parameters.VectorPrimitiveDataIndexValues.Insert(INDEX_NONE, Index);
 					Parameters.VectorChannelNameValues.Insert(FParameterChannelNames(), Index);
 					Parameters.VectorUsedAsChannelMaskValues.Insert(false, Index);
 				}
