Fix PrimaryAssetLabel's CDO to not be editor only, this fixes a warning when cooking a runtime label reported by licensee

diff --git a/Engine/Source/Runtime/Engine/Classes/Engine/PrimaryAssetLabel.h b/Engine/Source/Runtime/Engine/Classes/Engine/PrimaryAssetLabel.h
--- a/Engine/Source/Runtime/Engine/Classes/Engine/PrimaryAssetLabel.h
+++ b/Engine/Source/Runtime/Engine/Classes/Engine/PrimaryAssetLabel.h
@@ -50,7 +50,7 @@ public:
 	/** Set to editor only if this is not available in a cooked build */
 	virtual bool IsEditorOnly() const
 	{
-		return !bIsRuntimeLabel;
+		return !bIsRuntimeLabel && !IsTemplate();
 	}
 
 #if WITH_EDITORONLY_DATA
