Fix crashes when duplicating material instances (UE-127418)

diff --git a/Engine/Source/Runtime/Engine/Private/Materials/Material.cpp b/Engine/Source/Runtime/Engine/Private/Materials/Material.cpp
--- a/Engine/Source/Runtime/Engine/Private/Materials/Material.cpp
+++ b/Engine/Source/Runtime/Engine/Private/Materials/Material.cpp
@@ -2384,7 +2384,7 @@ void UMaterial::Serialize(FArchive& Ar)
 	bool bSavedCachedExpressionData = false;
 	if (Ar.CustomVer(FUE5MainStreamObjectVersion::GUID) >= FUE5MainStreamObjectVersion::MaterialSavedCachedData)
 	{
-		if (Ar.IsCooking())
+		if (Ar.IsCooking() || (FPlatformProperties::RequiresCookedData() && Ar.IsSaving() && (Ar.GetPortFlags() & PPF_Duplicate)))
 		{
 			if (CachedExpressionData)
 			{
diff --git a/Engine/Source/Runtime/Engine/Private/Materials/MaterialInstance.cpp b/Engine/Source/Runtime/Engine/Private/Materials/MaterialInstance.cpp
--- a/Engine/Source/Runtime/Engine/Private/Materials/MaterialInstance.cpp
+++ b/Engine/Source/Runtime/Engine/Private/Materials/MaterialInstance.cpp
@@ -2365,7 +2365,8 @@ void UMaterialInstance::Serialize(FArchive& Ar)
 	bool bSavedCachedData = false;
 	if (Ar.CustomVer(FUE5MainStreamObjectVersion::GUID) >= FUE5MainStreamObjectVersion::MaterialSavedCachedData)
 	{
-		if (Ar.IsCooking())
+		// we want to save the cached data when cooking or duplicating the object in a cooked game
+		if (Ar.IsCooking() || (FPlatformProperties::RequiresCookedData() && Ar.IsSaving() && (Ar.GetPortFlags() & PPF_Duplicate)))
 		{
 			if (CachedData)
 			{
