Fix crashes in hair system code when enabling r.DeferSkeletalDynamicDataUpdateUntilGDME

This is because this variable is incompatible with the hair system, which requires early update of the skin cache data prior to the hair LOD selection and the hair simulation.

diff --git a/Engine/Source/Runtime/Engine/Private/SkeletalRenderGPUSkin.cpp b/Engine/Source/Runtime/Engine/Private/SkeletalRenderGPUSkin.cpp
--- a/Engine/Source/Runtime/Engine/Private/SkeletalRenderGPUSkin.cpp
+++ b/Engine/Source/Runtime/Engine/Private/SkeletalRenderGPUSkin.cpp
@@ -15,6 +15,7 @@
 #include "SkeletalMeshTypes.h"
 #include "HAL/LowLevelMemTracker.h"
 #include "RenderGraphUtils.h"
+#include "Interfaces/IPluginManager.h"
 
 DEFINE_LOG_CATEGORY_STATIC(LogSkeletalGPUSkinMesh, Warning, All);
 
@@ -375,6 +376,15 @@ static TAutoConsoleVariable<int32> CVarDeferSkeletalDynamicDataUpdateUntilGDME(
 	0,
 	TEXT("If > 0, then do skeletal mesh dynamic data updates will be deferred until GDME. Experimental option."));
 
+FORCEINLINE bool IsDeferredSkeletalDynamicDataUpdateEnabled()
+{
+	// When the hair plugin is enabled, deferred skel. mesh update are not allowed, as skin-cached update need to run 
+	// prior to the PreInitView calls, because hair LOD selection and hair simulation requires the skin cache 
+	// resources to be updated at that point.
+	static bool bIsHairPluginEnable = IPluginManager::Get().FindPlugin("HairStrands") != nullptr && IPluginManager::Get().FindPlugin("HairStrands")->IsEnabled();
+	return CVarDeferSkeletalDynamicDataUpdateUntilGDME.GetValueOnRenderThread() > 0 && !bIsHairPluginEnable;
+}
+
 void FSkeletalMeshObjectGPUSkin::UpdateDynamicData_RenderThread(FGPUSkinCache* GPUSkinCache, FRHICommandListImmediate& RHICmdList, FDynamicSkelMeshObjectDataGPUSkin* InDynamicData, FSceneInterface* Scene, bool bIsGameWorld, uint32 FrameNumberToPrepare, uint32 RevisionNumber)
 {
 	SCOPE_CYCLE_COUNTER(STAT_GPUSkinUpdateRTTime);
@@ -421,7 +431,7 @@ void FSkeletalMeshObjectGPUSkin::UpdateDynamicData_RenderThread(FGPUSkinCache* G
 	DynamicData = InDynamicData;
 	LastBoneTransformRevisionNumber = RevisionNumber;
 
-	if (CVarDeferSkeletalDynamicDataUpdateUntilGDME.GetValueOnRenderThread())
+	if (IsDeferredSkeletalDynamicDataUpdateEnabled())
 	{
 		bMorphNeedsUpdateDeferred = bMorphNeedsUpdate;
 		bNeedsUpdateDeferred = true;
