Another fix for the editor crashes during the shader compilation.

- Material render proxy can get created on memstack, so tracking all instances of it is impossible (and might be we should rethink the global collection of proxies).
- Instead, replace (the relatively new) check of the layout pointers with the check of their contents.

diff --git a/Engine/Source/Runtime/Engine/Private/Materials/MaterialShared.cpp b/Engine/Source/Runtime/Engine/Private/Materials/MaterialShared.cpp
--- a/Engine/Source/Runtime/Engine/Private/Materials/MaterialShared.cpp
+++ b/Engine/Source/Runtime/Engine/Private/Materials/MaterialShared.cpp
@@ -3197,7 +3197,9 @@ void FMaterialRenderProxy::EvaluateUniformExpressions(FUniformExpressionCache& O
 
 		if (IsValidRef(OutUniformExpressionCache.UniformBuffer))
 		{
-			check(OutUniformExpressionCache.UniformBuffer->GetLayoutPtr() == UniformBufferLayout);
+			// The actual pointer may not match because there are cases (in the editor, during the shader compilation) when material's shader map gets updated without proxy's cache
+			// getting invalidated, but the layout contents must match.
+			check(OutUniformExpressionCache.UniformBuffer->GetLayoutPtr() == UniformBufferLayout || *OutUniformExpressionCache.UniformBuffer->GetLayoutPtr() == *UniformBufferLayout);
 			RHIUpdateUniformBuffer(OutUniformExpressionCache.UniformBuffer, TempBuffer);
 		}
 		else
@@ -3304,27 +3306,10 @@ FMaterialRenderProxy::FMaterialRenderProxy()
 	, DeletedFlag(0)
 	, HasVirtualTextureCallbacks(0)
 {
-#if WITH_EDITOR
-	// We need to register ourselves in MaterialRenderProxyMap early as sometimes (during the shader compilation) we need to be updated even before InitDynamicRHI (see e.g. UE-117842)
-	if (!FPlatformProperties::RequiresCookedData())
-	{
-		FScopeLock Locker(&MaterialRenderProxyMapLock);
-		FMaterialRenderProxy::MaterialRenderProxyMap.Add(this);
-	}
-#endif // WITH_EDITOR
 }
 
 FMaterialRenderProxy::~FMaterialRenderProxy()
 {
-#if WITH_EDITOR
-	// this may also be removed in the ReleaseDynamicRHI, but leaving it here in case we never get the state created
-	if (!FPlatformProperties::RequiresCookedData())
-	{
-		FScopeLock Locker(&MaterialRenderProxyMapLock);
-		FMaterialRenderProxy::MaterialRenderProxyMap.Remove(this);
-	}
-#endif // WITH_EDITOR
-
 	if(IsInitialized())
 	{
 		check(IsInRenderingThread());
@@ -3341,10 +3326,21 @@ FMaterialRenderProxy::~FMaterialRenderProxy()
 	DeletedFlag = 1;
 }
 
+void FMaterialRenderProxy::InitDynamicRHI()
+{
+#if WITH_EDITOR
+	// MaterialRenderProxyMap is only used by shader compiling
+	if (!FPlatformProperties::RequiresCookedData())
+	{
+		FScopeLock Locker(&MaterialRenderProxyMapLock);
+		FMaterialRenderProxy::MaterialRenderProxyMap.Add(this);
+	}
+#endif // WITH_EDITOR
+}
+
 void FMaterialRenderProxy::ReleaseDynamicRHI()
 {
 #if WITH_EDITOR
-	// unregister ourselves early to prevent getting updates from the shader compilation when the parent has already got BeginDestroy()ed
 	if (!FPlatformProperties::RequiresCookedData())
 	{
 		FScopeLock Locker(&MaterialRenderProxyMapLock);
diff --git a/Engine/Source/Runtime/Engine/Public/MaterialShared.h b/Engine/Source/Runtime/Engine/Public/MaterialShared.h
--- a/Engine/Source/Runtime/Engine/Public/MaterialShared.h
+++ b/Engine/Source/Runtime/Engine/Public/MaterialShared.h
@@ -2352,6 +2352,7 @@ public:
 	}
 
 	// FRenderResource interface.
+	ENGINE_API virtual void InitDynamicRHI() override;
 	ENGINE_API virtual void ReleaseDynamicRHI() override;
 	ENGINE_API virtual void ReleaseResource() override;
 
diff --git a/Engine/Source/Runtime/RHI/Public/RHIResources.h b/Engine/Source/Runtime/RHI/Public/RHIResources.h
--- a/Engine/Source/Runtime/RHI/Public/RHIResources.h
+++ b/Engine/Source/Runtime/RHI/Public/RHIResources.h
@@ -1072,6 +1072,15 @@ struct FRHIUniformBufferLayout : public FRHIResource
 	const bool bHasNonGraphOutputs;
 };
 
+/** Compare two uniform buffer layouts. */
+inline bool operator==(const FRHIUniformBufferLayout& A, const FRHIUniformBufferLayout& B)
+{
+	return A.ConstantBufferSize == B.ConstantBufferSize
+		&& A.StaticSlot == B.StaticSlot
+		&& A.BindingFlags == B.BindingFlags
+		&& A.Resources == B.Resources;
+}
+
 class FRHIUniformBuffer : public FRHIResource
 #if ENABLE_RHI_VALIDATION
 	, public RHIValidation::FUniformBufferResource
