Texture: Updated serialization of cooked texture platform data

- Added a placeholder for FDerivedData to allow for derived data references to be patched into a product gradually without forcing a large patch at one time.
- Switched from absolute to relative offsets for skipping over serialized platform data.

#rnx

diff --git a/Engine/Source/Runtime/Engine/Private/TextureDerivedData.cpp b/Engine/Source/Runtime/Engine/Private/TextureDerivedData.cpp
--- a/Engine/Source/Runtime/Engine/Private/TextureDerivedData.cpp
+++ b/Engine/Source/Runtime/Engine/Private/TextureDerivedData.cpp
@@ -5,6 +5,7 @@
 =============================================================================*/
 
 #include "CoreMinimal.h"
+#include "Algo/AllOf.h"
 #include "Misc/CommandLine.h"
 #include "Stats/Stats.h"
 #include "Async/AsyncWork.h"
@@ -1557,11 +1558,19 @@ static void SerializePlatformData(
 	EPlatformDataSerializationFlags Flags
 )
 {
+	DECLARE_SCOPE_CYCLE_COUNTER(TEXT("SerializePlatformData"), STAT_Texture_SerializePlatformData, STATGROUP_LoadTime);
+
+	if (Ar.IsFilterEditorOnly())
+	{
+		constexpr int64 PlaceholderDerivedDataSize = 16;
+		uint8 PlaceholderDerivedData[PlaceholderDerivedDataSize]{};
+		Ar.Serialize(PlaceholderDerivedData, PlaceholderDerivedDataSize);
+		check(Algo::AllOf(PlaceholderDerivedData, [](uint8 Value) { return Value == 0; }));
+	}
+
 	const bool bCooked = (Flags & EPlatformDataSerializationFlags::Cooked) == EPlatformDataSerializationFlags::Cooked;
 	const bool bStreamable = (Flags & EPlatformDataSerializationFlags::Streamable) == EPlatformDataSerializationFlags::Streamable;
 
-	DECLARE_SCOPE_CYCLE_COUNTER(TEXT("SerializePlatformData"), STAT_Texture_SerializePlatformData, STATGROUP_LoadTime);
-
 	UEnum* PixelFormatEnum = UTexture::GetPixelFormatEnum();
 
 	Ar << PlatformData->SizeX;
@@ -1799,7 +1808,7 @@ void FTexturePlatformData::SerializeCooked(FArchive& Ar, UTexture* Owner, bool b
 				SetNumSlices(Mips[0].SizeZ);
 			}
 		}
-		else if ( VTData )
+		else if (VTData)
 		{
 			SizeX = VTData->Width;
 			SizeY = VTData->Height;
@@ -2328,7 +2337,6 @@ void UTexture::SerializeCookedPlatformData(FArchive& Ar)
 
 	UEnum* PixelFormatEnum = UTexture::GetPixelFormatEnum();
 
-
 #if WITH_EDITOR
 	if (Ar.IsCooking() && Ar.IsPersistent())
 	{
@@ -2357,36 +2365,35 @@ void UTexture::SerializeCookedPlatformData(FArchive& Ar)
 			}
 			else
 			{
-				TMap<FString, FTexturePlatformData*> *CookedPlatformDataPtr = GetCookedPlatformData();
-				if (CookedPlatformDataPtr == NULL)
+				TMap<FString, FTexturePlatformData*>* CookedPlatformDataPtr = GetCookedPlatformData();
+				if (CookedPlatformDataPtr == nullptr)
+				{
 					return;
+				}
 
-				TArray< TArray<FTextureBuildSettings> > BuildSettingsToCache;
+				TArray<TArray<FTextureBuildSettings>> BuildSettingsToCache;
 				GetBuildSettingsPerFormat(*this, BuildSettings, Ar.CookingTarget(), BuildSettingsToCache);
 
-				for (int32 SettingIndex = 0; SettingIndex < BuildSettingsToCache.Num(); SettingIndex++)
+				for (TArray<FTextureBuildSettings>& LayerBuildSettings : BuildSettingsToCache)
 				{
-					check(BuildSettingsToCache[SettingIndex].Num() == Source.GetNumLayers());
+					check(LayerBuildSettings.Num() == Source.GetNumLayers());
 
 					FString DerivedDataKey;
-					GetTextureDerivedDataKey(*this, BuildSettingsToCache[SettingIndex].GetData(), DerivedDataKey);
+					GetTextureDerivedDataKey(*this, LayerBuildSettings.GetData(), DerivedDataKey);
 
-					FTexturePlatformData *PlatformDataPtr = (*CookedPlatformDataPtr).FindRef(DerivedDataKey);
-					if (PlatformDataPtr == NULL)
+					FTexturePlatformData*& PlatformDataPtr = (*CookedPlatformDataPtr).FindOrAdd(DerivedDataKey);
+					if (PlatformDataPtr == nullptr)
 					{
 						PlatformDataPtr = new FTexturePlatformData();
-						PlatformDataPtr->Cache(*this, BuildSettingsToCache[SettingIndex].GetData(), uint32(ETextureCacheFlags::InlineMips | ETextureCacheFlags::Async), nullptr);
-
-						CookedPlatformDataPtr->Add(DerivedDataKey, PlatformDataPtr);
+						PlatformDataPtr->Cache(*this, LayerBuildSettings.GetData(), uint32(ETextureCacheFlags::InlineMips | ETextureCacheFlags::Async), nullptr);
 
 					}
 					PlatformDataToSerialize.Add(PlatformDataPtr);
 				}
 			}
 
-			for (int32 i = 0; i < PlatformDataToSerialize.Num(); ++i)
+			for (FTexturePlatformData* PlatformDataToSave : PlatformDataToSerialize)
 			{
-				FTexturePlatformData* PlatformDataToSave = PlatformDataToSerialize[i];
 				PlatformDataToSave->FinishCache();
 
 				// Update bCookedIsStreamable for later use in IsCandidateForTextureStreaming
@@ -2398,20 +2405,21 @@ void UTexture::SerializeCookedPlatformData(FArchive& Ar)
 
 				FName PixelFormatName = PixelFormatEnum->GetNameByValue(PlatformDataToSave->PixelFormat);
 				Ar << PixelFormatName;
-				int64 SkipOffsetLoc = Ar.Tell();
-				int64 SkipOffset = 0;
 
+				const int64 SkipOffsetLoc = Ar.Tell();
+				int64 SkipOffset = 0;
 				{
-					FArchive::FScopeSetDebugSerializationFlags S(Ar,DSF_IgnoreDiff);
+					FArchive::FScopeSetDebugSerializationFlags S(Ar, DSF_IgnoreDiff);
 					Ar << SkipOffset;
 				}
 
 				// Pass streamable flag for inlining mips
 				PlatformDataToSave->SerializeCooked(Ar, this, BuildSettings.bStreamable);
-				SkipOffset = Ar.Tell();
+
+				SkipOffset = Ar.Tell() - SkipOffsetLoc;
 				Ar.Seek(SkipOffsetLoc);
 				Ar << SkipOffset;
-				Ar.Seek(SkipOffset);
+				Ar.Seek(SkipOffsetLoc + SkipOffset);
 			}
 		}
 		FName PixelFormatName = NAME_None;
@@ -2422,25 +2430,25 @@ void UTexture::SerializeCookedPlatformData(FArchive& Ar)
 	{
 
 		FTexturePlatformData** RunningPlatformDataPtr = GetRunningPlatformData();
-		if ( RunningPlatformDataPtr == NULL )
+		if (RunningPlatformDataPtr == nullptr)
+		{
 			return;
+		}
 
+		CleanupCachedRunningPlatformData();
 		FTexturePlatformData*& RunningPlatformData = *RunningPlatformDataPtr;
+		check(RunningPlatformData == nullptr);
+		RunningPlatformData = new FTexturePlatformData();
 
 		FName PixelFormatName = NAME_None;
-
-		CleanupCachedRunningPlatformData();
-		check( RunningPlatformData == NULL );
-
-		RunningPlatformData = new FTexturePlatformData();
 		Ar << PixelFormatName;
 		while (PixelFormatName != NAME_None)
 		{
 			EPixelFormat PixelFormat = (EPixelFormat)PixelFormatEnum->GetValueByName(PixelFormatName);
+			const int64 SkipOffsetLoc = Ar.Tell();
 			int64 SkipOffset = 0;
 			Ar << SkipOffset;
-			bool bFormatSupported = GPixelFormats[PixelFormat].Supported;
-			if (RunningPlatformData->PixelFormat == PF_Unknown && bFormatSupported)
+			if (RunningPlatformData->PixelFormat == PF_Unknown && GPixelFormats[PixelFormat].Supported)
 			{
 				// Extra arg is unused here because we're loading
 				const bool bStreamable = false;
@@ -2448,13 +2456,13 @@ void UTexture::SerializeCookedPlatformData(FArchive& Ar)
 			}
 			else
 			{
-				Ar.Seek(SkipOffset);
+				Ar.Seek(SkipOffsetLoc + SkipOffset);
 			}
 			Ar << PixelFormatName;
 		}
 	}
 
-	if( Ar.IsLoading() )
+	if (Ar.IsLoading())
 	{
 		LODBias = 0;
 	}
