Extend FFinishBuildMorphTargetData class to support async serialization of externally derived UMorphTarget classes

diff --git a/Engine/Source/Runtime/Engine/Classes/Animation/MorphTarget.h b/Engine/Source/Runtime/Engine/Classes/Animation/MorphTarget.h
--- a/Engine/Source/Runtime/Engine/Classes/Animation/MorphTarget.h
+++ b/Engine/Source/Runtime/Engine/Classes/Animation/MorphTarget.h
@@ -127,6 +127,27 @@ struct FMorphTargetLODModel
 	}
 };
 
+#if WITH_EDITOR
+/**
+* Data to cache serialization results for async asset building
+*/
+struct FFinishBuildMorphTargetData
+{
+public:
+	virtual ~FFinishBuildMorphTargetData()
+	{}
+	
+	/** Load morph target data */
+	ENGINE_API virtual void LoadFromMemoryArchive(FMemoryArchive & Ar);
+	
+	/** Apply serialized data to skeletal mesh in game thread */
+	ENGINE_API virtual void ApplyEditorData(USkeletalMesh * SkeletalMesh) const;
+	
+protected:
+	bool bApplyMorphTargetsData = false;
+	TMap<FName, TArray<FMorphTargetLODModel>> MorphLODModelsPerTargetName;
+	};
+#endif
 
 UCLASS(hidecategories=Object, MinimalAPI)
 class UMorphTarget
@@ -153,7 +174,7 @@ protected:
 public:
 
 	/** Get Morphtarget Delta array for the given input Index */
-	ENGINE_API const FMorphTargetDelta* GetMorphTargetDelta(int32 LODIndex, int32& OutNumDeltas) const;
+	ENGINE_API virtual const FMorphTargetDelta* GetMorphTargetDelta(int32 LODIndex, int32& OutNumDeltas) const;
 	ENGINE_API virtual bool HasDataForLOD(int32 LODIndex) const;
 	/** return true if this morphtarget contains data for section within LOD */
 	ENGINE_API virtual bool HasDataForSection(int32 LODIndex, int32 SectionIndex) const;
@@ -162,13 +183,15 @@ public:
 	ENGINE_API virtual void EmptyMorphLODModels();
 
 	/** Discard CPU Buffers after render resources have been created. */
-	ENGINE_API void DiscardVertexData();
+	ENGINE_API virtual void DiscardVertexData();
 
 #if WITH_EDITOR
 	/** Populates the given morph target LOD model with the provided deltas */
 	ENGINE_API virtual void PopulateDeltas(const TArray<FMorphTargetDelta>& Deltas, const int32 LODIndex, const TArray<struct FSkelMeshSection>& Sections, const bool bCompareNormal = false, const bool bGeneratedByReductionSetting = false, const float PositionThreshold = THRESH_POINTS_ARE_NEAR);
 	/** Remove empty LODModels */
 	ENGINE_API virtual void RemoveEmptyMorphTargets();
+	/** Factory function to define type of FinishBuildData needed*/
+	ENGINE_API virtual TUniquePtr<FFinishBuildMorphTargetData> CreateFinishBuildMorphTargetData() const;
 #endif // WITH_EDITOR
 
 public:
diff --git a/Engine/Source/Runtime/Engine/Classes/Engine/SkeletalMesh.h b/Engine/Source/Runtime/Engine/Classes/Engine/SkeletalMesh.h
--- a/Engine/Source/Runtime/Engine/Classes/Engine/SkeletalMesh.h
+++ b/Engine/Source/Runtime/Engine/Classes/Engine/SkeletalMesh.h
@@ -121,18 +121,16 @@ enum class ESkeletalMeshAsyncProperties : uint64
 
 ENUM_CLASS_FLAGS(ESkeletalMeshAsyncProperties);
 
+#if WITH_EDITOR
 struct FFinishBuildInternalData
 {
 public:
-	/** Editor MorphTargets data to set. */
-	bool bApplyMorphTargetsData = false;
-	TMap<FName, TArray<FMorphTargetLODModel> > MorphLODModelsPerTargetName;
-
 	void ApplyEditorData(USkeletalMesh* SkeletalMesh) const;
 
-private:
-	void ApplyMorphTargetsEditorData(USkeletalMesh* SkeletalMesh) const;
+	/** Editor MorphTargets data to set. */
+	TUniquePtr<FFinishBuildMorphTargetData> MorphTargetData;
 };
+#endif
 
 class FSkeletalMeshCompilationContext
 {
@@ -144,7 +142,9 @@ public:
 	// Movable
 	FSkeletalMeshCompilationContext(FSkeletalMeshCompilationContext&&) = default;
 	FSkeletalMeshCompilationContext& operator=(FSkeletalMeshCompilationContext&&) = default;
+#if WITH_EDITOR
 	FFinishBuildInternalData FinishBuildInternalData;
+#endif
 };
 
 class FSkeletalMeshPostLoadContext : public FSkeletalMeshCompilationContext
diff --git a/Engine/Source/Runtime/Engine/Private/MorphMesh.cpp b/Engine/Source/Runtime/Engine/Private/MorphMesh.cpp
--- a/Engine/Source/Runtime/Engine/Private/MorphMesh.cpp
+++ b/Engine/Source/Runtime/Engine/Private/MorphMesh.cpp
@@ -36,26 +36,37 @@ void UMorphTarget::Serialize( FArchive& Ar )
 	}
 }
 
-void UMorphTarget::SerializeMemoryArchive(FMemoryArchive& Ar)
+namespace
 {
-	int32 MorphLODModelNumber = 0;
-	if (Ar.IsLoading())
+	void SerializeMorphLODModels(FMemoryArchive& Ar, TArray<FMorphTargetLODModel>& MorphLODModels)
 	{
-		Ar << MorphLODModelNumber;
-		MorphLODModels.AddDefaulted(MorphLODModelNumber);
-	}
-	else
-	{
-		MorphLODModelNumber = MorphLODModels.Num();
-		Ar << MorphLODModelNumber;
-	}
+		int32 MorphLODModelNumber = 0;
+		if (Ar.IsLoading())
+		{
+			Ar << MorphLODModelNumber;
+			MorphLODModels.Empty(MorphLODModelNumber);
+			MorphLODModels.AddDefaulted(MorphLODModelNumber);
+		}
+		else
+		{
+			MorphLODModelNumber = MorphLODModels.Num();
+			Ar << MorphLODModelNumber;
+		}
 
-	for (int32 MorphIndex = 0; MorphIndex < MorphLODModelNumber; ++MorphIndex)
-	{
-		Ar << MorphLODModels[MorphIndex];
+		for (int32 MorphIndex = 0; MorphIndex < MorphLODModelNumber; ++MorphIndex)
+		{
+			Ar << MorphLODModels[MorphIndex];
+		}
 	}
 }
 
+void UMorphTarget::SerializeMemoryArchive(FMemoryArchive & Ar)
+{
+	FName MorphTargetName = GetFName();
+	Ar << MorphTargetName;
+	SerializeMorphLODModels(Ar, MorphLODModels);
+}
+
 void UMorphTarget::GetResourceSizeEx(FResourceSizeEx& CumulativeResourceSize)
 {
 	Super::GetResourceSizeEx(CumulativeResourceSize);
@@ -114,3 +125,30 @@ void FMorphTargetLODModel::GetResourceSizeEx(FResourceSizeEx& CumulativeResource
 {
 	CumulativeResourceSize.AddUnknownMemoryBytes(Vertices.GetAllocatedSize() + sizeof(int32));
 }
+
+#if WITH_EDITOR
+
+void FFinishBuildMorphTargetData::LoadFromMemoryArchive(FMemoryArchive & Ar)
+{
+	check(Ar.IsLoading());
+	
+	if (!ensureMsgf(!bApplyMorphTargetsData, TEXT("Error in FFinishBuildMorphTargetData::LoadFromMemoryArchive. The compilation context morph targets data was already set.")))
+	{
+		MorphLODModelsPerTargetName.Empty();
+		bApplyMorphTargetsData = false;
+	}
+	bApplyMorphTargetsData = true;
+	
+	int32 MorphTargetNumber = 0;
+	Ar << MorphTargetNumber;
+	MorphLODModelsPerTargetName.Reserve(MorphTargetNumber);
+	for (int32 MorphTargetIndex = 0; MorphTargetIndex < MorphTargetNumber; ++MorphTargetIndex)
+	{
+		FName MorphTargetName = NAME_None;
+		Ar << MorphTargetName;
+		TArray<FMorphTargetLODModel>&MorphLODModels = MorphLODModelsPerTargetName.FindOrAdd(MorphTargetName);
+		SerializeMorphLODModels(Ar, MorphLODModels);
+	}
+}
+
+#endif // WITH_EDITOR
\ No newline at end of file
diff --git a/Engine/Source/Runtime/Engine/Private/MorphTools.cpp b/Engine/Source/Runtime/Engine/Private/MorphTools.cpp
--- a/Engine/Source/Runtime/Engine/Private/MorphTools.cpp
+++ b/Engine/Source/Runtime/Engine/Private/MorphTools.cpp
@@ -154,6 +154,71 @@ void UMorphTarget::RemoveEmptyMorphTargets()
 	}
 }
 
+TUniquePtr<FFinishBuildMorphTargetData> UMorphTarget::CreateFinishBuildMorphTargetData() const
+{
+	return MakeUnique<FFinishBuildMorphTargetData>();
+}
+
+void FFinishBuildMorphTargetData::ApplyEditorData(USkeletalMesh * SkeletalMesh) const
+{
+	check(IsInGameThread());
+	//Return if we do not need to apply data
+	if (!bApplyMorphTargetsData)
+	{
+		return;
+	}
+	
+	FSkeletalMeshModel * SkelMeshModel = SkeletalMesh->GetImportedModel();
+	check(SkelMeshModel);
+	
+	TMap<FName, UMorphTarget*> ExistingMorphTargets;
+	for (UMorphTarget* MorphTarget : SkeletalMesh->GetMorphTargets())
+	{
+		ExistingMorphTargets.Add(MorphTarget->GetFName(), MorphTarget);
+	}
+	
+	int32 MorphTargetNumber = MorphLODModelsPerTargetName.Num();
+	TArray<UMorphTarget*> ToDeleteMorphTargets;
+	ToDeleteMorphTargets.Append(SkeletalMesh->GetMorphTargets());
+	SkeletalMesh->GetMorphTargets().Empty();
+	//Rebuild the MorphTarget object
+	for (const TPair<FName, TArray<FMorphTargetLODModel>>& TargetNameAndMorphLODModels : MorphLODModelsPerTargetName)
+	{
+		FName MorphTargetName = TargetNameAndMorphLODModels.Key;
+		UMorphTarget * MorphTarget = ExistingMorphTargets.FindRef(MorphTargetName);
+		if (!MorphTarget)
+		{
+			MorphTarget = NewObject<UMorphTarget>(SkeletalMesh, MorphTargetName);
+			check(MorphTarget);
+		}
+		else
+		{
+			ToDeleteMorphTargets.Remove(MorphTarget);
+		}
+		MorphTarget->EmptyMorphLODModels();
+		SkeletalMesh->GetMorphTargets().Add(MorphTarget);
+		const TArray<FMorphTargetLODModel>&MorphTargetLODModels = TargetNameAndMorphLODModels.Value;
+		int32 MorphLODModelNumber = MorphTargetLODModels.Num();
+		MorphTarget->GetMorphLODModels().AddDefaulted(MorphLODModelNumber);
+		for (int32 MorphDataIndex = 0; MorphDataIndex < MorphLODModelNumber; ++MorphDataIndex)
+		{
+			MorphTarget->GetMorphLODModels()[MorphDataIndex] = MorphTargetLODModels[MorphDataIndex];
+		}
+	}
+	//Rebuild the mapping and rehook the curve data
+	SkeletalMesh->InitMorphTargets();
+	
+	for (UMorphTarget* ToDeleteMorphTarget : ToDeleteMorphTargets)
+	{
+		ToDeleteMorphTarget->BaseSkelMesh = nullptr;
+		ToDeleteMorphTarget->EmptyMorphLODModels();
+		
+		//Move the unused asset in the transient package and mark it pending kill
+		ToDeleteMorphTarget->Rename(nullptr, GetTransientPackage(), REN_ForceNoResetLoaders | REN_DoNotDirty | REN_DontCreateRedirectors | REN_NonTransactional);
+		ToDeleteMorphTarget->MarkPendingKill();
+	}
+}
+
 #endif // WITH_EDITOR
 
 
diff --git a/Engine/Source/Runtime/Engine/Private/SkeletalMesh.cpp b/Engine/Source/Runtime/Engine/Private/SkeletalMesh.cpp
--- a/Engine/Source/Runtime/Engine/Private/SkeletalMesh.cpp
+++ b/Engine/Source/Runtime/Engine/Private/SkeletalMesh.cpp
@@ -1975,70 +1975,13 @@ uint32 USkeletalMesh::GetVertexBufferFlags() const
 
 #if WITH_EDITOR
 
-void FFinishBuildInternalData::ApplyMorphTargetsEditorData(USkeletalMesh* SkeletalMesh) const
+void FFinishBuildInternalData::ApplyEditorData(USkeletalMesh* SkeletalMesh) const
 {
 	check(IsInGameThread());
-	//Return if we do not need to apply data
-	if (!bApplyMorphTargetsData)
+	if (MorphTargetData)
 	{
-		return;
+		MorphTargetData->ApplyEditorData(SkeletalMesh);
 	}
-
-	FSkeletalMeshModel* SkelMeshModel = SkeletalMesh->GetImportedModel();
-	check(SkelMeshModel);
-
-	TMap<FName, UMorphTarget*> ExistingMorphTargets;
-	for (UMorphTarget* MorphTarget : SkeletalMesh->GetMorphTargets())
-	{
-		ExistingMorphTargets.Add(MorphTarget->GetFName(), MorphTarget);
-	}
-
-	int32 MorphTargetNumber = MorphLODModelsPerTargetName.Num();
-	TArray<UMorphTarget*> ToDeleteMorphTargets;
-	ToDeleteMorphTargets.Append(SkeletalMesh->GetMorphTargets());
-	SkeletalMesh->GetMorphTargets().Empty();
-	//Rebuild the MorphTarget object
-	for(TPair<FName, TArray<FMorphTargetLODModel>> TargetNameAndMorphLODModels : MorphLODModelsPerTargetName)
-	{
-		FName MorphTargetName = TargetNameAndMorphLODModels.Key;
-		UMorphTarget* MorphTarget = ExistingMorphTargets.FindRef(MorphTargetName);
-		if (!MorphTarget)
-		{
-			MorphTarget = NewObject<UMorphTarget>(SkeletalMesh, MorphTargetName);
-			check(MorphTarget);
-		}
-		else
-		{
-			ToDeleteMorphTargets.Remove(MorphTarget);
-		}
-		MorphTarget->EmptyMorphLODModels();
-		SkeletalMesh->GetMorphTargets().Add(MorphTarget);
-		const TArray<FMorphTargetLODModel>& MorphTargetLODModels = TargetNameAndMorphLODModels.Value;
-		int32 MorphLODModelNumber = MorphTargetLODModels.Num();
-		MorphTarget->GetMorphLODModels().AddDefaulted(MorphLODModelNumber);
-		for (int32 MorphDataIndex = 0; MorphDataIndex < MorphLODModelNumber; ++MorphDataIndex)
-		{
-			MorphTarget->GetMorphLODModels()[MorphDataIndex] = MorphTargetLODModels[MorphDataIndex];
-		}
-	}
-	//Rebuild the mapping and rehook the curve data
-	SkeletalMesh->InitMorphTargets();
-
-	for (UMorphTarget* ToDeleteMorphTarget : ToDeleteMorphTargets)
-	{
-		ToDeleteMorphTarget->BaseSkelMesh = nullptr;
-		ToDeleteMorphTarget->EmptyMorphLODModels();
-
-		//Move the unused asset in the transient package and mark it pending kill
-		ToDeleteMorphTarget->Rename(nullptr, GetTransientPackage(), REN_ForceNoResetLoaders | REN_DoNotDirty | REN_DontCreateRedirectors | REN_NonTransactional);
-		ToDeleteMorphTarget->MarkPendingKill();
-	}
-}
-
-void FFinishBuildInternalData::ApplyEditorData(USkeletalMesh* SkeletalMesh) const
-{
-	check(IsInGameThread());
-	ApplyMorphTargetsEditorData(SkeletalMesh);
 }
 
 thread_local const USkeletalMesh* FSkeletalMeshAsyncBuildScope::SkeletalMeshBeingAsyncCompiled = nullptr;
diff --git a/Engine/Source/Runtime/Engine/Private/SkeletalMeshRenderData.cpp b/Engine/Source/Runtime/Engine/Private/SkeletalMeshRenderData.cpp
--- a/Engine/Source/Runtime/Engine/Private/SkeletalMeshRenderData.cpp
+++ b/Engine/Source/Runtime/Engine/Private/SkeletalMeshRenderData.cpp
@@ -376,31 +376,19 @@ void FSkeletalMeshRenderData::Cache(const ITargetPlatform* TargetPlatform, USkel
 
 				//Get the morph target data, we put it in the compilation context to apply them in the game thread before the InitResources
 				{
-					if (!ensureMsgf(!ContextPtr->FinishBuildInternalData.bApplyMorphTargetsData, TEXT("Error in FSkeletalMeshRenderData::Cache. The compilation context morph targets data was already set.")))
+					if (Owner->GetMorphTargets().Num() > 0)
 					{
-						ContextPtr->FinishBuildInternalData.MorphLODModelsPerTargetName.Empty();
-						ContextPtr->FinishBuildInternalData.bApplyMorphTargetsData = false;
+						ContextPtr->FinishBuildInternalData.MorphTargetData = Owner->GetMorphTargets()[0]->CreateFinishBuildMorphTargetData();
 					}
-
-					int32 MorphTargetNumber = 0;
-					Ar << MorphTargetNumber;
-					ContextPtr->FinishBuildInternalData.MorphLODModelsPerTargetName.Reserve(MorphTargetNumber);
-					ContextPtr->FinishBuildInternalData.bApplyMorphTargetsData = true;
-					//We cannot serialize directly the UMorphTarget with a FMemoryArchive. This is not supported.
-					for (int32 MorphTargetIndex = 0; MorphTargetIndex < MorphTargetNumber; ++MorphTargetIndex)
+					else
 					{
-						FName MorphTargetName = NAME_None;
-						Ar << MorphTargetName;
-						TArray<FMorphTargetLODModel>& MorphTargetLODModels = ContextPtr->FinishBuildInternalData.MorphLODModelsPerTargetName.FindOrAdd(MorphTargetName);
-						int32 MorphLODModelNumber = 0;
-						Ar << MorphLODModelNumber;
-						MorphTargetLODModels.Empty(MorphLODModelNumber);
-						MorphTargetLODModels.AddDefaulted(MorphLODModelNumber);
-						for (int32 MorphDataIndex = 0; MorphDataIndex < MorphLODModelNumber; ++MorphDataIndex)
-						{
-							Ar << MorphTargetLODModels[MorphDataIndex];
-						}
+						// Create transient MorphTarget to initialize the FinishBuildInternalData
+						UMorphTarget * DefaultMorphTarget = NewObject<UMorphTarget>();
+						ContextPtr->FinishBuildInternalData.MorphTargetData = DefaultMorphTarget->CreateFinishBuildMorphTargetData();
+						DefaultMorphTarget->Rename(nullptr, GetTransientPackage());
+						DefaultMorphTarget->MarkPendingKill();
 					}
+					ContextPtr->FinishBuildInternalData.MorphTargetData->LoadFromMemoryArchive(Ar);
 				}
 
 				//Serialize the LODModel sections since they are dependent on the reduction
@@ -490,8 +478,6 @@ void FSkeletalMeshRenderData::Cache(const ITargetPlatform* TargetPlatform, USkel
 				Ar << MorphTargetNumber;
 				for (int32 MorphTargetIndex = 0; MorphTargetIndex < MorphTargetNumber; ++MorphTargetIndex)
 				{
-					FName MorphTargetName = Owner->GetMorphTargets()[MorphTargetIndex]->GetFName();
-					Ar << MorphTargetName;
 					Owner->GetMorphTargets()[MorphTargetIndex]->SerializeMemoryArchive(Ar);
 				}
 				//No need to serialize the morph target mapping since we will rebuild the mapping when loading a ddc
