Changed the name of a streamsource to be the name of the playercontroller instead of the localplayer. It makes it easier to sync streaming sources with controllers if we eventually do streaming on server side, but help Mass code synching up controller and streaming sources in my next checkin.

#rnx

diff --git a/Engine/Source/Runtime/Engine/Private/WorldPartition/WorldPartitionStreamingPolicy.cpp b/Engine/Source/Runtime/Engine/Private/WorldPartition/WorldPartitionStreamingPolicy.cpp
--- a/Engine/Source/Runtime/Engine/Private/WorldPartition/WorldPartitionStreamingPolicy.cpp
+++ b/Engine/Source/Runtime/Engine/Private/WorldPartition/WorldPartitionStreamingPolicy.cpp
@@ -135,7 +135,7 @@ void UWorldPartitionStreamingPolicy::UpdateStreamingSources()
 						ViewLocation = WorldToLocal.TransformPosition(ViewLocation);
 						ViewRotation = WorldToLocal.TransformRotation(ViewRotation.Quaternion()).Rotator();						
 						EStreamingSourceTargetState TargetState = Player->PlayerController->StreamingSourceShouldActivate() ? EStreamingSourceTargetState::Activated : EStreamingSourceTargetState::Loaded;
-						StreamingSources.Add(FWorldPartitionStreamingSource(Player->GetFName(), ViewLocation, ViewRotation, TargetState, /*bBlockOnSlowLoading=*/true, EStreamingSourcePriority::Default));
+						StreamingSources.Add(FWorldPartitionStreamingSource(Player->PlayerController->GetFName(), ViewLocation, ViewRotation, TargetState, /*bBlockOnSlowLoading=*/true, EStreamingSourcePriority::Default));
 					}
 				}
 			}
