Fix visual logger to not reuse track upon pointer reuse.

diff --git a/Engine/Source/Runtime/Engine/Private/VisualLogger/VisualLogger.cpp b/Engine/Source/Runtime/Engine/Private/VisualLogger/VisualLogger.cpp
--- a/Engine/Source/Runtime/Engine/Private/VisualLogger/VisualLogger.cpp
+++ b/Engine/Source/Runtime/Engine/Private/VisualLogger/VisualLogger.cpp
@@ -136,7 +136,7 @@ void FVisualLogger::AddWhitelistedObject(const UObject& InObject)
 
 void FVisualLogger::ClearObjectWhitelist()
 {
-	for (const UObject* It : ObjectWhitelist)
+	for (FObjectKey It : ObjectWhitelist)
 	{
 		FVisualLogEntry* CurrentEntry = CurrentEntryPerObject.Find(It);
 		if (CurrentEntry)
@@ -187,7 +187,7 @@ FVisualLogEntry* FVisualLogger::GetEntryToWrite(const UObject* Object, float Tim
 					{
 						for (FVisualLogDevice* Device : OutputDevices)
 						{
-							Device->Serialize(CurrentPair.Key, ObjectToNameMap[CurrentPair.Key], ObjectToClassNameMap[CurrentPair.Key], *Entry);
+							Device->Serialize(CurrentPair.Key.ResolveObjectPtrEvenIfPendingKill(), ObjectToNameMap[CurrentPair.Key], ObjectToClassNameMap[CurrentPair.Key], *Entry);
 						}
 						Entry->Reset();
 					}
@@ -205,7 +205,6 @@ FVisualLogEntry* FVisualLogger::GetEntryToWrite(const UObject* Object, float Tim
 			: *FString::Printf(TEXT("%s"), *LogOwner->GetName()));
 		ObjectToNameMap.Add(LogOwner, LogName);
 		ObjectToClassNameMap.Add(LogOwner, *(LogOwner->GetClass()->GetName()));
-		ObjectToPointerMap.Add(LogOwner, LogOwner);
 		ObjectToWorldMap.Add(LogOwner, World);
 
 		// IsClassWhitelisted isn't super fast, but this gets calculated only once for every
@@ -231,7 +230,7 @@ FVisualLogEntry* FVisualLogger::GetEntryToWrite(const UObject* Object, float Tim
 		FOwnerToChildrenRedirectionMap& RedirectionMap = GetRedirectionMap(LogOwner);
 		if (RedirectionMap.Contains(LogOwner))
 		{
-			if (ObjectToPointerMap.Contains(LogOwner) && ObjectToPointerMap[LogOwner].IsValid())
+			if (CurrentEntryPerObject.Contains(LogOwner) && LogOwner->IsValidLowLevel())
 			{
 				const IVisualLoggerDebugSnapshotInterface* DebugSnapshotInterface = Cast<const IVisualLoggerDebugSnapshotInterface>(LogOwner);
 				if (DebugSnapshotInterface)
@@ -273,7 +272,7 @@ void FVisualLogger::Flush()
 		{
 			for (FVisualLogDevice* Device : OutputDevices)
 			{
-				Device->Serialize(CurrentEntry.Key, ObjectToNameMap[CurrentEntry.Key], ObjectToClassNameMap[CurrentEntry.Key], CurrentEntry.Value);
+				Device->Serialize(CurrentEntry.Key.ResolveObjectPtrEvenIfPendingKill(), ObjectToNameMap[CurrentEntry.Key], ObjectToClassNameMap[CurrentEntry.Key], CurrentEntry.Value);
 			}
 			CurrentEntry.Value.Reset();
 		}
@@ -417,40 +416,38 @@ void FVisualLogger::Cleanup(UWorld* OldWorld, bool bReleaseMemory)
 		// perform cleanup only if provided world is valid and was registered
 		if (WorldToRedirectionMap.Remove(OldWorld))
 		{
-		    if (WorldToRedirectionMap.Num() == 0)
-            {
-                WorldToRedirectionMap.Reset();
-                ObjectToWorldMap.Reset();
-                ChildToOwnerMap.Reset();
-                CurrentEntryPerObject.Reset();
-                ObjectToNameMap.Reset();
-                ObjectToClassNameMap.Reset();
-                ObjectToPointerMap.Reset();
-            }
-            else
-            {
-                for (auto It = ObjectToWorldMap.CreateIterator(); It; ++It)
-                {
-                    if (It.Value() == OldWorld)
-                    {
-                        const UObject* Obj = It.Key();
-                        ObjectToWorldMap.Remove(Obj);
-                        CurrentEntryPerObject.Remove(Obj);
-                        ObjectToNameMap.Remove(Obj);
-                        ObjectToClassNameMap.Remove(Obj);
-                        ObjectToPointerMap.Remove(Obj);
-                    }
-                }
-
-                for (FChildToOwnerRedirectionMap::TIterator It = ChildToOwnerMap.CreateIterator(); It; ++It)
-                {
-                    if (It->Key.IsValid() == false
-                        || It->Key->GetWorld() == OldWorld)
-                    {
-                        It.RemoveCurrent();
-                    }
-                }
-            }
+			if (WorldToRedirectionMap.Num() == 0)
+			{
+				WorldToRedirectionMap.Reset();
+				ObjectToWorldMap.Reset();
+				ChildToOwnerMap.Reset();
+				CurrentEntryPerObject.Reset();
+				ObjectToNameMap.Reset();
+				ObjectToClassNameMap.Reset();
+			}
+			else
+			{
+				for (auto It = ObjectToWorldMap.CreateIterator(); It; ++It)
+				{
+					if (It.Value() == OldWorld)
+					{
+						FObjectKey Obj = It.Key();
+						ObjectToWorldMap.Remove(Obj);
+						CurrentEntryPerObject.Remove(Obj);
+						ObjectToNameMap.Remove(Obj);
+						ObjectToClassNameMap.Remove(Obj);
+					}
+				}
+	
+				for (FChildToOwnerRedirectionMap::TIterator It = ChildToOwnerMap.CreateIterator(); It; ++It)
+				{
+					UObject* Object = It->Key.ResolveObjectPtrEvenIfPendingKill();
+					if (Object == nullptr || Object->GetWorld() == OldWorld)
+					{
+						It.RemoveCurrent();
+					}
+				}
+			}
 		}
 	}
 	else
@@ -461,7 +458,6 @@ void FVisualLogger::Cleanup(UWorld* OldWorld, bool bReleaseMemory)
 		CurrentEntryPerObject.Reset();
 		ObjectToNameMap.Reset();
 		ObjectToClassNameMap.Reset();
-		ObjectToPointerMap.Reset();
 	}
 
 	LastUniqueIds.Reset();
@@ -518,7 +514,7 @@ void FVisualLogger::Redirect(const UObject* FromObject, const UObject* ToObject)
 	}
 
 	FChildToOwnerRedirectionMap& ChildToOwnerMap = FVisualLogger::Get().GetChildToOwnerRedirectionMap();
-	ChildToOwnerMap.FindOrAdd(FromWeakPtr) = ToObject;
+	ChildToOwnerMap.FindOrAdd(FromWeakPtr.Get(true/*bEvenIfPendingKill*/)) = ToObject;
 
 	UE_CVLOG(FromObject != nullptr, FromObject, LogVisual, Log, TEXT("Redirected '%s' to '%s'"), *FromObject->GetName(), *NewRedirection->GetName());
 }
@@ -532,7 +528,7 @@ UObject* FVisualLogger::FindRedirection(const UObject* Object)
 
 	while (Parent)
 	{
-		Parent = Map.Find(TargetWeakPtr);
+		Parent = Map.Find(TargetWeakPtr.Get(/*bEvenIfPendingKill*/true));
 		if (Parent)
 		{
 			if (Parent->IsValid())
@@ -542,12 +538,12 @@ UObject* FVisualLogger::FindRedirection(const UObject* Object)
 			else
 			{
 				Parent = nullptr;
-				Map.Remove(TargetWeakPtr);
+				Map.Remove(TargetWeakPtr.Get(/*bEvenIfPendingKill*/true));
 			}
 		}
 	}
 
-	return const_cast<UObject*>(TargetWeakPtr.Get());
+	return const_cast<UObject*>(TargetWeakPtr.Get(/*bEvenIfPendingKill*/true));
 }
 
 void FVisualLogger::SetIsRecording(bool InIsRecording)
diff --git a/Engine/Source/Runtime/Engine/Public/VisualLogger/VisualLogger.h b/Engine/Source/Runtime/Engine/Public/VisualLogger/VisualLogger.h
--- a/Engine/Source/Runtime/Engine/Public/VisualLogger/VisualLogger.h
+++ b/Engine/Source/Runtime/Engine/Public/VisualLogger/VisualLogger.h
@@ -11,6 +11,7 @@
 #include "Templates/IsValidVariadicFunctionArg.h"
 #include "Templates/AndOrNot.h"
 #include "Templates/IsArrayOrRefOfType.h"
+#include "UObject/ObjectKey.h"
 
 #if ENABLE_VISUAL_LOG
 
@@ -538,13 +539,13 @@ public:
 	/** Sets function to call to get a timestamp instead of the default implementation (i.e. world time) */
 	void SetGetTimeStampFunc(TFunction<float(const UObject*)> Function);
 
-	typedef TMap<UObject*, TArray<TWeakObjectPtr<const UObject> > > FOwnerToChildrenRedirectionMap;
+	typedef TMap<FObjectKey, TArray<TWeakObjectPtr<const UObject> > > FOwnerToChildrenRedirectionMap;
 	static FOwnerToChildrenRedirectionMap& GetRedirectionMap(const UObject* InObject);
 
-	typedef TMap<TWeakObjectPtr<const UObject>, TWeakObjectPtr<const UObject> > FChildToOwnerRedirectionMap;
+	typedef TMap<FObjectKey, TWeakObjectPtr<const UObject> > FChildToOwnerRedirectionMap;
 	FChildToOwnerRedirectionMap& GetChildToOwnerRedirectionMap() { return ChildToOwnerMap; }
 
-	typedef TMap<const UObject*, TWeakObjectPtr<const UWorld> > FObjectToWorldMapType;
+	typedef TMap<FObjectKey, TWeakObjectPtr<const UWorld> > FObjectToWorldMapType;
 	FObjectToWorldMapType& GetObjectToWorldMap() { return ObjectToWorldMap; }
 
 	void AddWhitelistedClass(UClass& InClass);
@@ -571,7 +572,7 @@ protected:
 	// white-listed objects - takes priority over class whitelist and should be used to create exceptions in it
 	// if ObjectWhitelist is empty (default) everything will log
 	// do NOT read from those pointers, they can be invalid!
-	TSet<const UObject*> ObjectWhitelist;
+	TSet<FObjectKey> ObjectWhitelist;
 
 	// white list of categories to bypass blocking
 	TArray<FName>	CategoriesWhitelist;
@@ -580,13 +581,11 @@ protected:
 	// last generated unique id for given times tamp
 	TMap<float, int32> LastUniqueIds;
 	// Current entry with all data
-	TMap<const UObject*, FVisualLogEntry>	CurrentEntryPerObject;
+	TMap<FObjectKey, FVisualLogEntry>	CurrentEntryPerObject;
 	// Map to contain names for Objects (they can be destroyed after while)
-	TMap<const UObject*, FName> ObjectToNameMap;
+	TMap<FObjectKey, FName> ObjectToNameMap;
 	// Map to contain class names for Objects (they can be destroyed after while)
-	TMap<const UObject*, FName> ObjectToClassNameMap;
-	// Map to contain information about pointers in game
-	TMap<const UObject*, TWeakObjectPtr<const UObject> > ObjectToPointerMap;
+	TMap<FObjectKey, FName> ObjectToClassNameMap;
 	// Cached map to world information because it's just raw pointer and not real object
 	FObjectToWorldMapType ObjectToWorldMap;
 	// for any object that has requested redirection this map holds where we should
