Added missing resets for primitive data indices

diff --git a/Engine/Source/Runtime/Engine/Private/Materials/MaterialCachedData.cpp b/Engine/Source/Runtime/Engine/Private/Materials/MaterialCachedData.cpp
--- a/Engine/Source/Runtime/Engine/Private/Materials/MaterialCachedData.cpp
+++ b/Engine/Source/Runtime/Engine/Private/Materials/MaterialCachedData.cpp
@@ -700,6 +700,8 @@ void FMaterialCachedParameters::Reset()
 	}
 #endif
 
+	ScalarPrimitiveDataIndexValues.Reset();
+	VectorPrimitiveDataIndexValues.Reset();
 	ScalarValues.Reset();
 	VectorValues.Reset();
 	TextureValues.Reset();
