Sampling raw data is now slightly more accurate by switching some calculations to use doubles internally.
Also fixed some parameter names that were called NumFrames instead of NumKeys, and the other way around.

diff --git a/Engine/Source/Runtime/Engine/Private/Animation/AnimationRuntime.cpp b/Engine/Source/Runtime/Engine/Private/Animation/AnimationRuntime.cpp
--- a/Engine/Source/Runtime/Engine/Private/Animation/AnimationRuntime.cpp
+++ b/Engine/Source/Runtime/Engine/Private/Animation/AnimationRuntime.cpp
@@ -1388,48 +1388,47 @@ void FAnimationRuntime::ApplyWeightToTransform(const FBoneContainer& RequiredBon
 	}			
 }
 
-/* from % from OutKeyIndex1, meaning (CurrentKeyIndex(float)-OutKeyIndex1)/(OutKeyIndex2-OutKeyIndex1) */
-void FAnimationRuntime::GetKeyIndicesFromTime(int32& OutKeyIndex1, int32& OutKeyIndex2, float& OutAlpha, const float Time, const int32 NumFrames, const float SequenceLength)
+/* from % from OutKeyIndex1, meaning (CurrentKeyIndex(double)-OutKeyIndex1)/(OutKeyIndex2-OutKeyIndex1) */
+void FAnimationRuntime::GetKeyIndicesFromTime(int32& OutKeyIndex1, int32& OutKeyIndex2, float& OutAlpha, const double Time, const int32 NumKeys, const double SequenceLength)
 {
 	// Check for 1-frame, before-first-frame and after-last-frame cases.
-	if( Time <= 0.f || NumFrames == 1 )
+	if (Time <= 0.0 || NumKeys == 1)
 	{
 		OutKeyIndex1 = 0;
 		OutKeyIndex2 = 0;
-		OutAlpha = 0.f;
+		OutAlpha = 0.0f;
 		return;
 	}
 
-	const int32 LastIndex		= NumFrames - 1;
-	if( Time >= SequenceLength )
+	const int32 LastIndex = NumKeys - 1;
+	if (Time >= SequenceLength)
 	{
 		OutKeyIndex1 = LastIndex;
-		OutKeyIndex2 = (OutKeyIndex1 + 1) % (NumFrames);
-		OutAlpha = 0.f;
+		OutKeyIndex2 = 0;
+		OutAlpha = 0.0f;
 		return;
 	}
 
-	// This assumes that all keys are equally spaced (ie. won't work if we have dropped unimportant frames etc).
-	const int32 NumKeys = NumFrames - 1;
-	const float KeyPos = ((float)NumKeys * Time) / SequenceLength;
+	const int32 NumFrames = NumKeys - 1;
+	const double KeyPos = (Time / SequenceLength) * (double)NumFrames;
 
 	// Find the integer part (ensuring within range) and that gives us the 'starting' key index.
-	const int32 KeyIndex1 = FMath::Clamp<int32>( FMath::FloorToInt(KeyPos), 0, NumFrames-1 );  // @todo should be changed to FMath::TruncToInt
+	const int32 KeyIndex1 = FMath::Clamp<int32>( FMath::FloorToInt(KeyPos), 0, NumKeys - 1 );  // @todo should be changed to FMath::TruncToInt
 
 	// The alpha (fractional part) is then just the remainder.
-	const float Alpha = KeyPos - (float)KeyIndex1;
+	const double Alpha = KeyPos - (double)KeyIndex1;
 
 	int32 KeyIndex2 = KeyIndex1 + 1;
 
 	// If we have gone over the end, do different things in case of looping
-	if( KeyIndex2 == NumFrames )
+	if (KeyIndex2 == NumKeys)
 	{
 		KeyIndex2 = KeyIndex1;
 	}
 
 	OutKeyIndex1 = KeyIndex1;
 	OutKeyIndex2 = KeyIndex2;
-	OutAlpha = Alpha;
+	OutAlpha = (float)Alpha;
 }
 
 FTransform FAnimationRuntime::GetComponentSpaceRefPose(const FCompactPoseBoneIndex& CompactPoseBoneIndex, const FBoneContainer& BoneContainer)
diff --git a/Engine/Source/Runtime/Engine/Public/AnimationRuntime.h b/Engine/Source/Runtime/Engine/Public/AnimationRuntime.h
--- a/Engine/Source/Runtime/Engine/Public/AnimationRuntime.h
+++ b/Engine/Source/Runtime/Engine/Public/AnimationRuntime.h
@@ -518,11 +518,11 @@ public:
 	static void ApplyWeightToTransform(const FBoneContainer& RequiredBones, /*inout*/ FTransformArrayA2& Atoms, float Weight);
 
 	/** 
-	 * Get Key Indices (start/end with alpha from start) with input parameter Time, NumFrames
+	 * Get Key Indices (start/end with alpha from start) with input parameter Time, NumKeys
 	 * from % from StartKeyIndex, meaning (CurrentKeyIndex(float)-StartKeyIndex)/(EndKeyIndex-StartKeyIndex)
-	 * by this Start-End, it will be between 0-(NumFrames-1), not number of Pos/Rot key tracks 
+	 * by this Start-End, it will be between 0-(NumKeys-1), not number of Pos/Rot key tracks 
 	 **/
-	static void GetKeyIndicesFromTime(int32& OutKeyIndex1, int32& OutKeyIndex2, float& OutAlpha, const float Time, const int32 NumFrames, const float SequenceLength);
+	static void GetKeyIndicesFromTime(int32& OutKeyIndex1, int32& OutKeyIndex2, float& OutAlpha, const double Time, const int32 NumKeys, const double SequenceLength);
 
 	/** 
 	 *	Utility for taking an array of bone indices and ensuring that all parents are present 
