Fix crash when using skeletalmesh ddc with morph targets
#rnx

diff --git a/Engine/Source/Runtime/Engine/Private/SkeletalMeshRenderData.cpp b/Engine/Source/Runtime/Engine/Private/SkeletalMeshRenderData.cpp
--- a/Engine/Source/Runtime/Engine/Private/SkeletalMeshRenderData.cpp
+++ b/Engine/Source/Runtime/Engine/Private/SkeletalMeshRenderData.cpp
@@ -382,12 +382,10 @@ void FSkeletalMeshRenderData::Cache(const ITargetPlatform* TargetPlatform, USkel
 					}
 					else
 					{
-						// Create transient MorphTarget to initialize the FinishBuildInternalData
-						UMorphTarget * DefaultMorphTarget = NewObject<UMorphTarget>();
-						ContextPtr->FinishBuildInternalData.MorphTargetData = DefaultMorphTarget->CreateFinishBuildMorphTargetData();
-						DefaultMorphTarget->Rename(nullptr, GetTransientPackage());
-						DefaultMorphTarget->MarkPendingKill();
+						// Create and initialize the FinishBuildInternalData, use the class default object to call the virtual function
+						ContextPtr->FinishBuildInternalData.MorphTargetData = UMorphTarget::StaticClass()->GetDefaultObject<UMorphTarget>()->CreateFinishBuildMorphTargetData();
 					}
+					check(ContextPtr->FinishBuildInternalData.MorphTargetData);
 					ContextPtr->FinishBuildInternalData.MorphTargetData->LoadFromMemoryArchive(Ar);
 				}
 
