PR #8237: Adding force reload option for AutomationOpenMap function (Contributed by [REDACTED_PERSON])

diff --git a/Engine/Source/Editor/UnrealEd/Classes/Editor/EditorEngine.h b/Engine/Source/Editor/UnrealEd/Classes/Editor/EditorEngine.h
--- a/Engine/Source/Editor/UnrealEd/Classes/Editor/EditorEngine.h
+++ b/Engine/Source/Editor/UnrealEd/Classes/Editor/EditorEngine.h
@@ -3171,7 +3171,7 @@ public:
 
 	virtual void HandleTravelFailure(UWorld* InWorld, ETravelFailure::Type FailureType, const FString& ErrorString);
 
-	void AutomationLoadMap(const FString& MapName, FString* OutError);
+	void AutomationLoadMap(const FString& MapName, bool bForceReload, FString* OutError);
 
 	/** This function should be called to notify the editor that new materials were added to our scene or some materials were modified */
 	void OnSceneMaterialsModified();
diff --git a/Engine/Source/Editor/UnrealEd/Private/EditorEngine.cpp b/Engine/Source/Editor/UnrealEd/Private/EditorEngine.cpp
--- a/Engine/Source/Editor/UnrealEd/Private/EditorEngine.cpp
+++ b/Engine/Source/Editor/UnrealEd/Private/EditorEngine.cpp
@@ -7246,7 +7246,7 @@ void UEditorEngine::HandleTravelFailure(UWorld* InWorld, ETravelFailure::Type Fa
 	}
 }
 
-void UEditorEngine::AutomationLoadMap(const FString& MapName, FString* OutError)
+void UEditorEngine::AutomationLoadMap(const FString& MapName, bool bForceReload, FString* OutError)
 {
 #if WITH_AUTOMATION_TESTS
 	struct FFailedGameStartHandler
@@ -7301,7 +7301,7 @@ void UEditorEngine::AutomationLoadMap(const FString& MapName, FString* OutError)
 		}
 	}
 
-	if (bNeedLoadEditorMap)
+	if (bNeedLoadEditorMap || bForceReload)
 	{
 		if (bPieRunning)
 		{
diff --git a/Engine/Source/Runtime/Engine/Private/Tests/AutomationCommon.cpp b/Engine/Source/Runtime/Engine/Private/Tests/AutomationCommon.cpp
--- a/Engine/Source/Runtime/Engine/Private/Tests/AutomationCommon.cpp
+++ b/Engine/Source/Runtime/Engine/Private/Tests/AutomationCommon.cpp
@@ -301,14 +301,14 @@ namespace AutomationCommon
 	}
 }
 
-bool AutomationOpenMap(const FString& MapName)
+bool AutomationOpenMap(const FString& MapName, bool bForceReload)
 {
 	bool bCanProceed = true;
 	FString OutString = TEXT("");
 #if WITH_EDITOR
 	if (GIsEditor && AutomationCommon::OnEditorAutomationMapLoad.IsBound())
 	{
-		AutomationCommon::OnEditorAutomationMapLoad.Broadcast(MapName, &OutString);
+		AutomationCommon::OnEditorAutomationMapLoad.Broadcast(MapName, bForceReload, &OutString);
 	}
 	else
 #endif
@@ -324,7 +324,7 @@ bool AutomationOpenMap(const FString& MapName)
 			FString PIEPrefix = FString::Printf(PLAYWORLD_PACKAGE_PREFIX TEXT("_%d_"), TestWorld->GetOutermost()->PIEInstanceID);
 			ShortWorldMapName.ReplaceInline(*PIEPrefix, TEXT(""));
 		}
-		if (ShortMapName != ShortWorldMapName)
+		if (ShortMapName != ShortWorldMapName || bForceReload)
 		{
 			FString OpenCommand = FString::Printf(TEXT("Open %s"), *MapName);
 			GEngine->Exec(TestWorld, *OpenCommand);
diff --git a/Engine/Source/Runtime/Engine/Public/Tests/AutomationCommon.h b/Engine/Source/Runtime/Engine/Public/Tests/AutomationCommon.h
--- a/Engine/Source/Runtime/Engine/Public/Tests/AutomationCommon.h
+++ b/Engine/Source/Runtime/Engine/Public/Tests/AutomationCommon.h
@@ -16,7 +16,7 @@ class SWindow;
 ENGINE_API DECLARE_LOG_CATEGORY_EXTERN(LogEditorAutomationTests, Log, All);
 ENGINE_API DECLARE_LOG_CATEGORY_EXTERN(LogEngineAutomationTests, Log, All);
 
-DECLARE_MULTICAST_DELEGATE_TwoParams(FOnEditorAutomationMapLoad, const FString&, FString*);
+DECLARE_MULTICAST_DELEGATE_ThreeParams(FOnEditorAutomationMapLoad, const FString&, bool, FString*);
 
 #endif
 
@@ -69,7 +69,7 @@ struct WindowScreenshotParameters
 /**
  * If Editor, Opens map and PIES.  If Game, transitions to map and waits for load
  */
-ENGINE_API bool AutomationOpenMap(const FString& MapName);
+ENGINE_API bool AutomationOpenMap(const FString& MapName, bool bForceReload = false);
 
 /**
  * Wait for the given amount of time
