Fix for a potential BP compiler crash after deleting a default instanced subobject element from a set or map variable.

diff --git a/Engine/Source/Runtime/Engine/Private/InstancedReferenceSubobjectHelper.cpp b/Engine/Source/Runtime/Engine/Private/InstancedReferenceSubobjectHelper.cpp
--- a/Engine/Source/Runtime/Engine/Private/InstancedReferenceSubobjectHelper.cpp
+++ b/Engine/Source/Runtime/Engine/Private/InstancedReferenceSubobjectHelper.cpp
@@ -54,10 +54,10 @@ UObject* FInstancedPropertyPath::Resolve(const UObject* Container) const
 				break;
 			}
 
-			const int32 TargetIndex = PropertyLink.ArrayIndex;
-			check(TargetIndex != INDEX_NONE);
-
 			FScriptSetHelper SetHelper(SetProperty, ValuePtr);
+
+			// Convert the logical index (recorded in the path) to the actual index used internally by the set.
+			const int32 TargetIndex = SetHelper.FindInternalIndex(PropertyLink.ArrayIndex);
 			if (SetHelper.IsValidIndex(TargetIndex))
 			{
 				CurrentProp = SetProperty->ElementProp;
@@ -70,10 +70,10 @@ UObject* FInstancedPropertyPath::Resolve(const UObject* Container) const
 		}
 		else if (const FMapProperty* MapProperty = CastField<FMapProperty>(CurrentProp))
 		{
-			const int32 TargetIndex = PropertyLink.ArrayIndex;
-			check(TargetIndex != INDEX_NONE);
-				
 			FScriptMapHelper MapHelper(MapProperty, ValuePtr);
+
+			// Convert the logical index (recorded in the path) to the actual index used internally by the map.
+			const int32 TargetIndex = MapHelper.FindInternalIndex(PropertyLink.ArrayIndex);
 			if (!MapHelper.IsValidIndex(TargetIndex))
 			{
 				CurrentProp = nullptr;
@@ -145,6 +145,7 @@ void FFindInstancedReferenceSubobjectHelper::ForEachInstancedSubObject(FInstance
 			return;
 		}
 
+		int32 LogicalIndex = 0;
 		FScriptMapHelper MapHelper(MapProperty, ContainerAddress);
 		for (int32 ElementIndex = 0; ElementIndex < MapHelper.GetMaxIndex(); ++ElementIndex)
 		{
@@ -153,13 +154,16 @@ void FFindInstancedReferenceSubobjectHelper::ForEachInstancedSubObject(FInstance
 				T KeyAddress = MapHelper.GetKeyPtr(ElementIndex);
 				T ValueAddress = MapHelper.GetValuePtr(ElementIndex);
 
-				PropertyPath.Push(MapProperty->KeyProp, ElementIndex);
+				// Note: Keep these as the logical (Nth) index in case the map changes internally after we construct the path or in case we resolve using a different object.
+				PropertyPath.Push(MapProperty->KeyProp, LogicalIndex);
 				ForEachInstancedSubObject(PropertyPath, KeyAddress, ObjRefFunc);
 				PropertyPath.Pop();
 
-				PropertyPath.Push(MapProperty->ValueProp, ElementIndex, true);
+				PropertyPath.Push(MapProperty->ValueProp, LogicalIndex, true);
 				ForEachInstancedSubObject(PropertyPath, ValueAddress, ObjRefFunc);
 				PropertyPath.Pop();
+
+				++LogicalIndex;
 			}
 		}
 	}
@@ -171,6 +175,7 @@ void FFindInstancedReferenceSubobjectHelper::ForEachInstancedSubObject(FInstance
 			return;
 		}
 
+		int32 LogicalIndex = 0;
 		FScriptSetHelper SetHelper(SetProperty, ContainerAddress);
 		for (int32 ElementIndex = 0; ElementIndex < SetHelper.GetMaxIndex(); ++ElementIndex)
 		{
@@ -178,9 +183,12 @@ void FFindInstancedReferenceSubobjectHelper::ForEachInstancedSubObject(FInstance
 			{
 				T ValueAddress = SetHelper.GetElementPtr(ElementIndex);
 
-				PropertyPath.Push(SetProperty->ElementProp, ElementIndex);
+				// Note: Keep this as the logical (Nth) index in case the set changes internally after we construct the path or in case we resolve using a different object.
+				PropertyPath.Push(SetProperty->ElementProp, LogicalIndex);
 				ForEachInstancedSubObject(PropertyPath, ValueAddress, ObjRefFunc);
 				PropertyPath.Pop();
+
+				++LogicalIndex;
 			}
 		}
 	}
