Add an option to disable support for per-pixel material shading models on mobile

diff --git a/Engine/Source/Runtime/Core/Private/HAL/ConsoleManager.cpp b/Engine/Source/Runtime/Core/Private/HAL/ConsoleManager.cpp
--- a/Engine/Source/Runtime/Core/Private/HAL/ConsoleManager.cpp
+++ b/Engine/Source/Runtime/Core/Private/HAL/ConsoleManager.cpp
@@ -2230,6 +2230,14 @@ static TAutoConsoleVariable<int32> CVarMobileAllowPixelDepthOffset(
 	ECVF_ReadOnly | ECVF_RenderThreadSafe
 );
 
+static TAutoConsoleVariable<int32> CVarMobileAllowPerPixelShadingModels(
+	TEXT("r.Mobile.AllowPerPixelShadingModels"),
+	1,
+	TEXT("Whether to allow 'Per-Pixel Shader Models (From Material Expression)' in materials for ES3.1 feature level."),
+	ECVF_ReadOnly | ECVF_RenderThreadSafe
+);
+
+
 static TAutoConsoleVariable<int32> CVarMobileSupportGPUScene(
 	TEXT("r.Mobile.SupportGPUScene"),
 	0,
diff --git a/Engine/Source/Runtime/Engine/Private/Materials/Material.cpp b/Engine/Source/Runtime/Engine/Private/Materials/Material.cpp
--- a/Engine/Source/Runtime/Engine/Private/Materials/Material.cpp
+++ b/Engine/Source/Runtime/Engine/Private/Materials/Material.cpp
@@ -232,7 +232,17 @@ int32 FMaterialResource::CompilePropertyAndSetMaterialProperty(EMaterialProperty
 				Ret = FMaterialAttributeDefinitionMap::CompileDefaultExpression(Compiler, Property);
 			}
 			break;
-
+		case MP_ShadingModel:
+			if (AllowPerPixelShadingModels(Compiler->GetShaderPlatform()))
+			{
+				Ret = MaterialInterface->CompileProperty(Compiler, Property);
+			}
+			else
+			{
+				FMaterialShadingModelField ShadingModels = Compiler->GetMaterialShadingModels();
+				Ret = Compiler->ShadingModel(ShadingModels.GetFirstShadingModel());
+			}
+			break;
 		case MP_MaterialAttributes:
 			Ret = MaterialInterface->CompileProperty(Compiler, Property);
 			break;
diff --git a/Engine/Source/Runtime/RenderCore/Private/RenderUtils.cpp b/Engine/Source/Runtime/RenderCore/Private/RenderUtils.cpp
--- a/Engine/Source/Runtime/RenderCore/Private/RenderUtils.cpp
+++ b/Engine/Source/Runtime/RenderCore/Private/RenderUtils.cpp
@@ -1142,6 +1142,16 @@ RENDERCORE_API bool AllowPixelDepthOffset(const FStaticShaderPlatform Platform)
 	return true;
 }
 
+RENDERCORE_API bool AllowPerPixelShadingModels(const FStaticShaderPlatform Platform)
+{
+	if (IsMobilePlatform(Platform))
+	{
+		static const auto CVar = IConsoleManager::Get().FindTConsoleVariableDataInt(TEXT("r.Mobile.AllowPerPixelShadingModels"));
+		return CVar->GetValueOnAnyThread() != 0;
+	}
+	return true;
+}
+
 RENDERCORE_API uint64 GMobileAmbientOcclusionPlatformMask = 0;
 static_assert(SP_NumPlatforms <= sizeof(GMobileAmbientOcclusionPlatformMask) * 8, "GMobileAmbientOcclusionPlatformMask must be large enough to support all shader platforms");
 
diff --git a/Engine/Source/Runtime/RenderCore/Private/Shader.cpp b/Engine/Source/Runtime/RenderCore/Private/Shader.cpp
--- a/Engine/Source/Runtime/RenderCore/Private/Shader.cpp
+++ b/Engine/Source/Runtime/RenderCore/Private/Shader.cpp
@@ -1477,6 +1477,11 @@ void ShaderMapAppendKeyString(EShaderPlatform Platform, FString& KeyString)
 	{
 		KeyString += TEXT("_NoPDO");
 	}
+
+	if (!AllowPerPixelShadingModels(Platform))
+	{
+		KeyString += TEXT("_NoPPSM");
+	}
 	
 	if (IsD3DPlatform(Platform) && IsPCPlatform(Platform))
 	{
diff --git a/Engine/Source/Runtime/RenderCore/Public/RenderUtils.h b/Engine/Source/Runtime/RenderCore/Public/RenderUtils.h
--- a/Engine/Source/Runtime/RenderCore/Public/RenderUtils.h
+++ b/Engine/Source/Runtime/RenderCore/Public/RenderUtils.h
@@ -497,6 +497,8 @@ RENDERCORE_API bool MaskedInEarlyPass(const FStaticShaderPlatform Platform);
 
 RENDERCORE_API bool AllowPixelDepthOffset(const FStaticShaderPlatform Platform);
 
+RENDERCORE_API bool AllowPerPixelShadingModels(const FStaticShaderPlatform Platform);
+
 RENDERCORE_API bool UseMobileAmbientOcclusion(const FStaticShaderPlatform Platform);
 
 RENDERCORE_API bool IsMobileDistanceFieldEnabled(const FStaticShaderPlatform Platform);
