run the stray world settings cleanup in RepairWorldSettings even if there's a properly set world settings object

diff --git a/Engine/Source/Runtime/Engine/Private/World.cpp b/Engine/Source/Runtime/Engine/Private/World.cpp
--- a/Engine/Source/Runtime/Engine/Private/World.cpp
+++ b/Engine/Source/Runtime/Engine/Private/World.cpp
@@ -1557,18 +1557,18 @@ void UWorld::RepairWorldSettings()
 		{
 			PersistentLevel->OwningWorld = nullptr;
 		}
+	}
 
-		// Now that we have set the proper world settings, clean up any other stay that may have accumulated due to legacy behaviors
-		if (PersistentLevel->Actors.Num() > 1)
+	// Now that we have set the proper world settings, clean up any other stay that may have accumulated due to legacy behaviors
+	if (PersistentLevel->Actors.Num() > 1)
+	{
+		for (int32 Index = 1, ActorNum = PersistentLevel->Actors.Num(); Index < ActorNum; ++Index)
 		{
-			for (int32 Index = 1, ActorNum = PersistentLevel->Actors.Num(); Index < ActorNum; ++Index)
+			AActor* Actor = PersistentLevel->Actors[Index];
+			if (Actor->IsA<AWorldSettings>())
 			{
-				AActor* Actor = PersistentLevel->Actors[Index];
-				if (Actor->IsA<AWorldSettings>())
-				{
-					UE_LOG(LogWorld, Warning, TEXT("Extra World Settings '%s' actor found. Resave level or actors to clean up"), *Actor->GetName());
-					Actor->Destroy();
-				}
+				UE_LOG(LogWorld, Warning, TEXT("Extra World Settings '%s' actor found. Resave level or actors to clean up"), *Actor->GetName());
+				Actor->Destroy();
 			}
 		}
 	}
