BSP: Fix issue in OptimizeIntoConvexPoly()

diff --git a/Engine/Source/Runtime/Engine/Classes/Engine/Polys.h b/Engine/Source/Runtime/Engine/Classes/Engine/Polys.h
--- a/Engine/Source/Runtime/Engine/Classes/Engine/Polys.h
+++ b/Engine/Source/Runtime/Engine/Classes/Engine/Polys.h
@@ -218,186 +218,7 @@ public:
 	* @param	InPolygons		An array of FPolys that will be replaced with a new set of polygons that are merged together as much as possible.
 	*/
 	template<typename ArrayType>
-	static void OptimizeIntoConvexPolys( ABrush* InOwnerBrush, ArrayType& InPolygons )
-	{
-		bool bDidMergePolygons = true;
-
-		while( bDidMergePolygons )
-		{
-			bDidMergePolygons = false;
-
-			for( int32 PolyIndex = 0 ; PolyIndex < InPolygons.Num() && !bDidMergePolygons ; ++PolyIndex )
-			{
-				const FPoly* PolyMain = &InPolygons[PolyIndex];
-
-				// Find all the polygons that neighbor this one (aka share an edge)
-
-				for( int32 p2 = 0 ; p2 < InPolygons.Num() && !bDidMergePolygons ; ++p2 )
-				{
-					const FPoly* PolyNeighbor = &InPolygons[p2];
-
-					if( PolyMain != PolyNeighbor && PolyMain->Normal.Equals( PolyNeighbor->Normal ) )
-					{
-						// See if PolyNeighbor is sharing an edge with Poly
-
-						int32 Index1 = INDEX_NONE, Index2 = INDEX_NONE;
-						FVector3f EdgeVtx1(0), EdgeVtx2(0);
-
-						for( int32 v = 0 ; v < PolyMain->Vertices.Num() ; ++v )
-						{
-							FVector3f vtx = PolyMain->Vertices[v];
-
-							int32 idx = INDEX_NONE;
-
-							for( int32 v2 = 0 ; v2 < PolyNeighbor->Vertices.Num() ; ++v2 )
-							{
-								const FVector3f* vtx2 = &PolyNeighbor->Vertices[v2];
-
-								if( vtx.Equals( *vtx2 ) )
-								{
-									idx = v2;
-									break;
-								}
-							}
-
-							if( idx != INDEX_NONE )
-							{
-								if( Index1 == INDEX_NONE )
-								{
-									Index1 = v;
-									EdgeVtx1 = vtx;
-								}
-								else if( Index2 == INDEX_NONE )
-								{
-									Index2 = v;
-									EdgeVtx2 = vtx;
-									break;
-								}
-							}
-						}
-
-						if( Index1 != INDEX_NONE && Index2 != INDEX_NONE )
-						{
-							// Found a shared edge.  Let's see if we can merge these polygons together.
-
-							// Create a list of cutting planes
-
-							TArray<FPlane> CuttingPlanes;
-
-							for( int32 v = 0 ; v < PolyMain->Vertices.Num() ; ++v )
-							{
-								const FVector3f* vtx1 = &PolyMain->Vertices[v];
-								const FVector3f* vtx2 = &PolyMain->Vertices[(v+1)%PolyMain->Vertices.Num()];
-
-								if( (vtx1->Equals( EdgeVtx1 ) && vtx2->Equals( EdgeVtx2 ) )
-									|| (vtx1->Equals( EdgeVtx2 ) && vtx2->Equals( EdgeVtx1 ) ) )
-								{
-									// If these verts are on the shared edge, don't include this edge in the cutting planes array.
-									continue;
-								}
-
-								FPlane plane( *vtx1, *vtx2, (*vtx2) + (PolyMain->Normal * 16.f) );
-								CuttingPlanes.Add( plane );
-							}
-
-							for( int32 v = 0 ; v < PolyNeighbor->Vertices.Num() ; ++v )
-							{
-								const FVector3f* vtx1 = &PolyNeighbor->Vertices[v];
-								const FVector3f* vtx2 = &PolyNeighbor->Vertices[(v+1)%PolyNeighbor->Vertices.Num()];
-
-								if( (vtx1->Equals( EdgeVtx1 ) && vtx2->Equals( EdgeVtx2 ) )
-									|| (vtx1->Equals( EdgeVtx2 ) && vtx2->Equals( EdgeVtx1 ) ) )
-								{
-									// If these verts are on the shared edge, don't include this edge in the cutting planes array.
-									continue;
-								}
-
-								FPlane plane( *vtx1, *vtx2, (*vtx2) + (PolyNeighbor->Normal * 16.f) );
-								CuttingPlanes.Add( plane );
-							}
-
-							// Make sure that all verts lie behind all cutting planes.  This serves as our convexity check for the merged polygon.
-
-							bool bMergedPolyIsConvex = true;
-							float dot;
-
-							for( int32 p = 0 ; p < CuttingPlanes.Num() && bMergedPolyIsConvex ; ++p )
-							{
-								FPlane* plane = &CuttingPlanes[p];
-
-								for( int32 v = 0 ; v < PolyMain->Vertices.Num() ; ++v )
-								{
-									dot = plane->PlaneDot( PolyMain->Vertices[v] );
-									if( dot > THRESH_POINT_ON_PLANE )
-									{
-										bMergedPolyIsConvex = false;
-										break;
-									}
-								}
-
-								for( int32 v = 0 ; v < PolyNeighbor->Vertices.Num() && bMergedPolyIsConvex ; ++v )
-								{
-									dot = plane->PlaneDot( PolyNeighbor->Vertices[v] );
-									if( dot > THRESH_POINT_ON_PLANE )
-									{
-										bMergedPolyIsConvex = false;
-										break;
-									}
-								}
-							}
-
-							if( bMergedPolyIsConvex )
-							{
-								// OK, the resulting polygon will result in a convex polygon.  So create it by clipping a large polygon using
-								// the cutting planes we created earlier.  The resulting polygon will be the merged result.
-
-								FPlane NormalPlane = FPlane( PolyMain->Vertices[0], PolyMain->Vertices[1], PolyMain->Vertices[2] );
-								FPoly PolyMerged = BuildAndCutInfiniteFPoly( NormalPlane, CuttingPlanes, InOwnerBrush );
-
-								// Verts that result from these clips are assured of being on the 1 grid.  Go through and snap them
-								// there to make sure (aka avoid float drift).
-
-								for( int32 v = 0 ; v < PolyMerged.Vertices.Num() ; ++v )
-								{
-									FVector3f* vtx = &PolyMerged.Vertices[v];
-									*vtx = vtx->GridSnap( 1.0f );
-								}
-
-								PolyMerged.CalcNormal(1);
-
-								if( PolyMerged.Finalize( InOwnerBrush, 1 ) == 0 )
-								{
-									// Remove the original polygons from the list
-
-									int32 idx1 = InPolygons.Find( *PolyMain );
-									int32 idx2 = InPolygons.Find( *PolyNeighbor );
-
-									if( idx2 > idx1 )
-									{
-										InPolygons.RemoveAt( idx2 );
-										InPolygons.RemoveAt( idx1 );
-									}
-									else
-									{
-										InPolygons.RemoveAt( idx1 );
-										InPolygons.RemoveAt( idx2 );
-									}
-
-									// Add the newly merged polygon into the list
-
-									InPolygons.Add( PolyMerged );
-
-									// Tell the outside loop that we merged polygons and need to run through it all again
-
-									bDidMergePolygons = true;
-								}
-							}
-						}
-					}
-				}
-			}
-		}
-	}
+	ENGINE_API static void OptimizeIntoConvexPolys(ABrush* InOwnerBrush, ArrayType& InPolygons);
 
 	/**
 	* Takes a set of polygons and returns a vertex array representing the outside winding
diff --git a/Engine/Source/Runtime/Engine/Private/Polygon.cpp b/Engine/Source/Runtime/Engine/Private/Polygon.cpp
--- a/Engine/Source/Runtime/Engine/Private/Polygon.cpp
+++ b/Engine/Source/Runtime/Engine/Private/Polygon.cpp
@@ -695,27 +695,27 @@ bool FPoly::IsCoplanar()
 
 bool FPoly::IsConvex()
 {
-	// Create a set of planes that represent each edge of the polygon.
-
-	TArray<FPlane> Planes;
-
-	for( int32 EdgeVert = 0 ; EdgeVert < Vertices.Num() ; ++EdgeVert )
+	const int32 VerticesNum = Vertices.Num();
+	if (VerticesNum > 2)
 	{
-		const FVector3f& vtx1 = Vertices[EdgeVert];
-		const FVector3f& vtx2 = Vertices[(EdgeVert + 1) % Vertices.Num()];
-		FVector3f v2v = vtx2 - vtx1;
+		FVector3f Vertex1 = Vertices[VerticesNum - 2];
+		FVector3f Vertex2 = Vertices[VerticesNum - 1];
+		FVector3f EdgeVertex1Vertex2 = (Vertex2 - Vertex1).GetSafeNormal();
 
-		FVector3f EdgeNormal = v2v ^ Normal;
-
-		for( int32 CheckVertLoop = 2 ; CheckVertLoop < Vertices.Num(); ++CheckVertLoop )
-	{
-			int32 CheckVert = (CheckVertLoop + EdgeVert) % Vertices.Num();
-			FVector3f RelativePos = Vertices[CheckVert] - Vertices[EdgeVert];
+		for (int32 Index3 = 0; Index3 < VerticesNum; ++Index3)
+		{
+			const FVector3f EdgeVertex1Vertex2Normal = FVector3f::CrossProduct(EdgeVertex1Vertex2, Normal);
+			const FVector3f Vertex3 = Vertices[Index3];
+			const FVector3f EdgeVertex2Vertex3 = (Vertex3 - Vertex2).GetSafeNormal();
 
-			if (0.0 < (EdgeNormal | RelativePos))
+			if (KINDA_SMALL_NUMBER < FVector3f::DotProduct(EdgeVertex1Vertex2Normal, EdgeVertex2Vertex3))
 			{
 				return false;
 			}
+
+			Vertex1 = Vertex2;
+			Vertex2 = Vertex3;
+			EdgeVertex1Vertex2 = EdgeVertex2Vertex3;
 		}
 	}
 
@@ -954,6 +954,177 @@ int32 FPoly::Finalize( ABrush* InOwner, int32 NoError )
 	return 0;
 }
 
+template<typename ArrayType>
+void FPoly::OptimizeIntoConvexPolys(ABrush* InOwnerBrush, ArrayType& InPolygons)
+{
+	// In this implementation, we find all coplanar polygons that share an edge and merge them together provided they form a convex polygon.
+
+	// Flag that indicates that we might have missed some potential neighbors after we jumped over one but then found another one.
+	// If this is the case, we need to check everything again.
+	bool bMightHaveMissedSomeNeighbors;
+
+	do
+	{
+		bMightHaveMissedSomeNeighbors = false;
+
+		// We iterate over the polygons in reverse order. This is effectively the inverse of how polygons get split up during triangulation.
+		// Starting at the end, the current polygon is considered the "main" polygon, and we try to grow it as much as possible using polygons earlier in
+		// the list called "neighbor" polygons. Whenever we jump over a potential neighbor polygon, and then find a polygon that is an actual neighbor, we
+		// effectively need to check again since the one we jumped over earlier might now be connectable.
+
+		for (int32 PolyMainIndex = InPolygons.Num() - 1; PolyMainIndex > 0; --PolyMainIndex)
+		{
+			FPoly* PolyMain = &InPolygons[PolyMainIndex];
+
+			// Iterate over all potential neighbor polygons, again in reverse order. 
+			for (int32 PolyNeighborIndex = PolyMainIndex - 1; PolyNeighborIndex >= 0; --PolyNeighborIndex)
+			{
+				const FPoly* PolyNeighbor = &InPolygons[PolyNeighborIndex];
+
+				if (PolyMain->Normal.Equals(PolyNeighbor->Normal))
+				{
+					// See if PolyNeighbor is sharing an edge with PolyMain.
+
+					const int32 Poly1Num = PolyMain->Vertices.Num();
+					const int32 Poly2Num = PolyNeighbor->Vertices.Num();
+
+					// Shared edge information for both polygons; this will be uninitialized if there is no shared edge.
+					FVector3f EdgeVtxA, EdgeVtxB; // Position for edge vertices A and B
+					int32 Poly1IndexA = INDEX_NONE, Poly1IndexB = INDEX_NONE; // Indices for edge vertices A and B in first polygon, PolyMain.
+					int32 Poly2IndexA = INDEX_NONE, Poly2IndexB = INDEX_NONE; // Indices for edge vertices A and B in second polygon, PolyNeighbor.
+
+					// Determine if there is a shared edge, and if yes, fills in the shared edge information for both polygons.
+					// The edge vertices, EdgeVtxA and EdgeVtxB, are set to the average between the corresponding vertices in both polygons.
+					const bool bSharedEdge = [PolyMain, PolyNeighbor, Poly1Num, Poly2Num, &EdgeVtxA, &EdgeVtxB,
+						&Poly1IndexA, &Poly1IndexB, &Poly2IndexA, &Poly2IndexB]
+					{
+						if (Poly1Num > 1 && Poly2Num > 1)
+						{
+							for (Poly1IndexA = Poly1Num - 1, Poly1IndexB = 0; Poly1IndexB < Poly1Num; Poly1IndexA = Poly1IndexB++)
+							{
+								// For each edge in the first polygon ...
+								const FVector3f& Poly1Vtx1 = PolyMain->Vertices[Poly1IndexA];
+								const FVector3f& Poly1Vtx2 = PolyMain->Vertices[Poly1IndexB];
+
+								for (Poly2IndexB = Poly2Num - 1, Poly2IndexA = 0; Poly2IndexA < Poly2Num; Poly2IndexB = Poly2IndexA++)
+								{
+									// ... and for each edge in the second polygon ...
+									const FVector3f& Poly2Vtx1 = PolyNeighbor->Vertices[Poly2IndexB];
+									const FVector3f& Poly2Vtx2 = PolyNeighbor->Vertices[Poly2IndexA];
+
+									// ... check if they have the same vertices for opposite orientation, i.e. a shared edge.
+									if (Poly1Vtx1.Equals(Poly2Vtx2) && Poly1Vtx2.Equals(Poly2Vtx1))
+									{
+										EdgeVtxA = (Poly1Vtx1 + Poly2Vtx2) * 0.5f;
+										EdgeVtxB = (Poly1Vtx2 + Poly2Vtx1) * 0.5f;
+										return true;
+									}
+								}
+							}
+						}
+						return false;
+					}();
+
+					if (bSharedEdge)
+					{
+						// Found a shared edge. Merge both polygons around the shared edge and check if the result is convex.
+
+						// Note that adding the vertices and polygons in this very order keeps the repeated execution of "Triangulate" and "Optimize"
+						// more stable, i.e. when you have a polygon and click "Triangulate" followed by "Optimize" it will more likely give you the
+						// same polygon with the same vertex order and the same UVs.
+
+						FPoly PolyMerged;
+						int32 PolyMergedNum = Poly1Num + Poly2Num - 2;
+						PolyMerged.Vertices.SetNumUninitialized(PolyMergedNum);
+
+						int32 PolyMergedIndex = 0;
+						for (int32 V = (Poly2IndexA + 1) % Poly2Num; V != Poly2IndexB; V = (V + 1) % Poly2Num)
+						{
+							PolyMerged.Vertices[PolyMergedIndex++] = PolyNeighbor->Vertices[V];
+						}
+						PolyMerged.Vertices[PolyMergedIndex++] = EdgeVtxB;
+						for (int32 V = (Poly1IndexB + 1) % Poly1Num; V != Poly1IndexA; V = (V + 1) % Poly1Num)
+						{
+							PolyMerged.Vertices[PolyMergedIndex++] = PolyMain->Vertices[V];
+						}
+						PolyMerged.Vertices[PolyMergedIndex] = EdgeVtxA;
+						checkSlow(++PolyMergedIndex == PolyMerged.Vertices.Num());
+
+						// Computing the normal is not cheap, but we need it for the convexity test. Thus, we just use the normals we already have since it
+						// is good enough for that. The proper normal will be computed during the later call to `FPoly::Finalize()` anyway.
+						PolyMerged.Normal = (PolyMain->Normal + PolyNeighbor->Normal) * 0.5f;
+
+						// Check convexity equivalent to FPoly::IsConvex(), but also remove collinear edges in the process.
+						const bool bIsConvex = [&PolyMerged, &PolyMergedNum]
+						{
+							if (PolyMergedNum > 2)
+							{
+								FVector3f Vertex1 = PolyMerged.Vertices[PolyMergedNum - 2];
+								int32 Index2 = PolyMergedNum - 1;
+								FVector3f Vertex2 = PolyMerged.Vertices[Index2];
+								FVector3f EdgeVertex1Vertex2 = (Vertex2 - Vertex1).GetSafeNormal();
+
+								for (int32 Index3 = 0; Index3 < PolyMergedNum; ++Index3)
+								{
+									const FVector3f Vertex3 = PolyMerged.Vertices[Index3];
+									const FVector3f EdgeVertex2Vertex3 = (Vertex3 - Vertex2).GetSafeNormal();
+
+									// Neighboring edges are collinear if they point into the same direction.
+									const bool bIsCollinear = EdgeVertex1Vertex2.Equals(EdgeVertex2Vertex3, KINDA_SMALL_NUMBER);
+
+									if (bIsCollinear)
+									{
+										// Collapse collinear edges by removing vertex 2.
+										PolyMerged.Vertices.RemoveAt(Index2);
+										--PolyMergedNum;
+
+										// Vertex 1 stays the same, but recompute the edge vector.
+										EdgeVertex1Vertex2 = (Vertex3 - Vertex1).GetSafeNormal();
+									}
+									else
+									{
+										// Check convexity for neighboring non-collinear edges.
+										const FVector3f EdgeVertex1Vertex2Normal = FVector3f::CrossProduct(EdgeVertex1Vertex2, PolyMerged.Normal);
+										if (KINDA_SMALL_NUMBER < FVector3f::DotProduct(EdgeVertex1Vertex2Normal, EdgeVertex2Vertex3))
+										{
+											return false;
+										}
+
+										// Vertex 1 moves to vertex 2, and the edge vector gets updated accordingly.
+										Vertex1 = Vertex2;
+										EdgeVertex1Vertex2 = EdgeVertex2Vertex3;
+									}
+
+									// Vertex 2 moves to vertex 3.
+									Index2 = Index3;
+									Vertex2 = Vertex3;
+								}
+							}
+
+							return true;
+						}();
+
+						// If the merged polygon is convex, finalize it, and enter it into the list.
+						if (bIsConvex && PolyMerged.Finalize(InOwnerBrush, 1) == 0)
+						{
+							// If the main polygon is not last one or the neighbor polygon is not the second to last one, we might have missed neighbors.
+							bMightHaveMissedSomeNeighbors |= PolyMainIndex != InPolygons.Num() - 1 || PolyNeighborIndex != PolyMainIndex - 1;
+
+							// Remove the neighbor polygon, fix up the main polygon index, and make the merged polygon the new main polygon.
+							InPolygons.RemoveAt(PolyNeighborIndex);
+							InPolygons[--PolyMainIndex] = PolyMerged;
+							PolyMain = &InPolygons[PolyMainIndex];
+						}
+					}
+				}
+			}
+		}
+	}
+	while (bMightHaveMissedSomeNeighbors);
+}
+
+template ENGINE_API void FPoly::OptimizeIntoConvexPolys<TArray<FPoly>>(ABrush*, TArray<FPoly>&);
+
 #endif // WITH_EDITOR
 
 
