Fix crash on nullptr AActor in RepairWorldSettings

diff --git a/Engine/Source/Runtime/Engine/Private/World.cpp b/Engine/Source/Runtime/Engine/Private/World.cpp
--- a/Engine/Source/Runtime/Engine/Private/World.cpp
+++ b/Engine/Source/Runtime/Engine/Private/World.cpp
@@ -1565,7 +1565,7 @@ void UWorld::RepairWorldSettings()
 		for (int32 Index = 1, ActorNum = PersistentLevel->Actors.Num(); Index < ActorNum; ++Index)
 		{
 			AActor* Actor = PersistentLevel->Actors[Index];
-			if (Actor->IsA<AWorldSettings>())
+			if (Actor && Actor->IsA<AWorldSettings>())
 			{
 				UE_LOG(LogWorld, Warning, TEXT("Extra World Settings '%s' actor found. Resave level or actors to clean up"), *Actor->GetName());
 				Actor->Destroy();
