Remove Legacy tonemapper related parameters

diff --git a/Engine/Plugins/Enterprise/DatasmithContent/Source/DatasmithContent/Private/ObjectTemplates/DatasmithCineCameraComponentTemplate.cpp b/Engine/Plugins/Enterprise/DatasmithContent/Source/DatasmithContent/Private/ObjectTemplates/DatasmithCineCameraComponentTemplate.cpp
--- a/Engine/Plugins/Enterprise/DatasmithContent/Source/DatasmithContent/Private/ObjectTemplates/DatasmithCineCameraComponentTemplate.cpp
+++ b/Engine/Plugins/Enterprise/DatasmithContent/Source/DatasmithContent/Private/ObjectTemplates/DatasmithCineCameraComponentTemplate.cpp
@@ -35,9 +35,6 @@ void FDatasmithPostProcessSettingsTemplate::Apply( FPostProcessSettings* Destina
 	DATASMITHOBJECTTEMPLATE_CONDITIONALSET( bOverride_VignetteIntensity, Destination, PreviousTemplate );
 	DATASMITHOBJECTTEMPLATE_CONDITIONALSET( VignetteIntensity, Destination, PreviousTemplate );
 
-	DATASMITHOBJECTTEMPLATE_CONDITIONALSET( bOverride_FilmWhitePoint, Destination, PreviousTemplate );
-	DATASMITHOBJECTTEMPLATE_CONDITIONALSET( FilmWhitePoint, Destination, PreviousTemplate );
-
 	DATASMITHOBJECTTEMPLATE_CONDITIONALSET( bOverride_ColorSaturation, Destination, PreviousTemplate );
 	DATASMITHOBJECTTEMPLATE_CONDITIONALSET( ColorSaturation, Destination, PreviousTemplate );
 
@@ -62,9 +59,6 @@ void FDatasmithPostProcessSettingsTemplate::Load( const FPostProcessSettings& So
 	bOverride_VignetteIntensity = Source.bOverride_VignetteIntensity;
 	VignetteIntensity = Source.VignetteIntensity;
 
-	bOverride_FilmWhitePoint = Source.bOverride_FilmWhitePoint;
-	FilmWhitePoint = Source.FilmWhitePoint;
-
 	bOverride_ColorSaturation = Source.bOverride_ColorSaturation;
 	ColorSaturation = Source.ColorSaturation;
 
@@ -89,9 +83,6 @@ bool FDatasmithPostProcessSettingsTemplate::Equals( const FDatasmithPostProcessS
 	bEquals = bEquals && bOverride_VignetteIntensity == Other.bOverride_VignetteIntensity;
 	bEquals = bEquals && FMath::IsNearlyEqual( VignetteIntensity, Other.VignetteIntensity );
 
-	bEquals = bEquals && bOverride_FilmWhitePoint == Other.bOverride_FilmWhitePoint;
-	bEquals = bEquals && FilmWhitePoint.Equals( Other.FilmWhitePoint );
-
 	bEquals = bEquals && bOverride_ColorSaturation == Other.bOverride_ColorSaturation;
 	bEquals = bEquals && ColorSaturation.Equals( Other.ColorSaturation );
 
diff --git a/Engine/Plugins/Enterprise/DatasmithContent/Source/DatasmithContent/Public/ObjectTemplates/DatasmithCineCameraComponentTemplate.h b/Engine/Plugins/Enterprise/DatasmithContent/Source/DatasmithContent/Public/ObjectTemplates/DatasmithCineCameraComponentTemplate.h
--- a/Engine/Plugins/Enterprise/DatasmithContent/Source/DatasmithContent/Public/ObjectTemplates/DatasmithCineCameraComponentTemplate.h
+++ b/Engine/Plugins/Enterprise/DatasmithContent/Source/DatasmithContent/Public/ObjectTemplates/DatasmithCineCameraComponentTemplate.h
@@ -72,9 +72,6 @@ public:
 	UPROPERTY()
 	uint32 bOverride_VignetteIntensity:1;
 
-	UPROPERTY()
-	uint32 bOverride_FilmWhitePoint:1;
-
 	UPROPERTY()
 	uint32 bOverride_AutoExposureMethod:1;
 
@@ -93,9 +90,6 @@ public:
 	UPROPERTY()
 	float VignetteIntensity = 0.0f;
 
-	UPROPERTY()
-	FLinearColor FilmWhitePoint;
-
 	UPROPERTY()
 	FVector4 ColorSaturation;
 
diff --git a/Engine/Plugins/Enterprise/DatasmithImporter/Source/DatasmithImporter/Private/DatasmithPostProcessImporter.cpp b/Engine/Plugins/Enterprise/DatasmithImporter/Source/DatasmithImporter/Private/DatasmithPostProcessImporter.cpp
--- a/Engine/Plugins/Enterprise/DatasmithImporter/Source/DatasmithImporter/Private/DatasmithPostProcessImporter.cpp
+++ b/Engine/Plugins/Enterprise/DatasmithImporter/Source/DatasmithImporter/Private/DatasmithPostProcessImporter.cpp
@@ -27,12 +27,6 @@ FDatasmithPostProcessSettingsTemplate FDatasmithPostProcessImporter::CopyDatasmi
 		PostProcessSettingsTemplate.VignetteIntensity = Src->GetVignette();
 	}
 
-	if (Src->GetColorFilter() != FLinearColor::Black && Src->GetColorFilter() != FLinearColor::White )
-	{
-		PostProcessSettingsTemplate.bOverride_FilmWhitePoint = true;
-		PostProcessSettingsTemplate.FilmWhitePoint = Src->GetColorFilter();
-	}
-
 	if ( !FMath::IsNearlyEqual( Src->GetSaturation(), 1.f ) )
 	{
 		PostProcessSettingsTemplate.bOverride_ColorSaturation = true;
diff --git a/Engine/Plugins/Experimental/Enterprise/DatasmithRuntime/Source/Private/ActorImporter.cpp b/Engine/Plugins/Experimental/Enterprise/DatasmithRuntime/Source/Private/ActorImporter.cpp
--- a/Engine/Plugins/Experimental/Enterprise/DatasmithRuntime/Source/Private/ActorImporter.cpp
+++ b/Engine/Plugins/Experimental/Enterprise/DatasmithRuntime/Source/Private/ActorImporter.cpp
@@ -263,12 +263,6 @@ namespace DatasmithRuntime
 					PostProcessSettings.VignetteIntensity = PostProcess->GetVignette();
 				}
 
-				if (PostProcess->GetColorFilter() != FLinearColor::Black && PostProcess->GetColorFilter() != FLinearColor::White )
-				{
-					PostProcessSettings.bOverride_FilmWhitePoint = true;
-					PostProcessSettings.FilmWhitePoint = PostProcess->GetColorFilter();
-				}
-
 				if ( !FMath::IsNearlyEqual( PostProcess->GetSaturation(), 1.f ) )
 				{
 					PostProcessSettings.bOverride_ColorSaturation = true;
diff --git a/Engine/Source/Runtime/Engine/Classes/Engine/BlendableInterface.h b/Engine/Source/Runtime/Engine/Classes/Engine/BlendableInterface.h
--- a/Engine/Source/Runtime/Engine/Classes/Engine/BlendableInterface.h
+++ b/Engine/Source/Runtime/Engine/Classes/Engine/BlendableInterface.h
@@ -22,8 +22,6 @@ enum EBlendableLocation
 	BL_BeforeTranslucency UMETA(DisplayName="Before Translucency"),
 	/**
 	* Input0:former pass color, Input1:SeparateTranslucency, Input2: BloomOutput
-	* vector parameters: Engine.FilmWhitePoint
-	* scalar parameters: Engine.FilmSaturation, Engine.FilmContrast
 	*/
 	BL_ReplacingTonemapper UMETA(DisplayName="Replacing the Tonemapper"),
 //	BL_AfterOpaque,
diff --git a/Engine/Source/Runtime/Engine/Classes/Engine/Scene.h b/Engine/Source/Runtime/Engine/Classes/Engine/Scene.h
--- a/Engine/Source/Runtime/Engine/Classes/Engine/Scene.h
+++ b/Engine/Source/Runtime/Engine/Classes/Engine/Scene.h
@@ -733,42 +733,6 @@ struct FPostProcessSettings
 	UPROPERTY(EditAnywhere, BlueprintReadWrite, Category = Overrides, meta = (PinHiddenByDefault, InlineEditConditionToggle))
 	uint8 bOverride_ToneCurveAmount : 1;
 
-	UPROPERTY(EditAnywhere, BlueprintReadWrite, Category=Overrides, meta=(PinHiddenByDefault, InlineEditConditionToggle))
-	uint8 bOverride_FilmWhitePoint:1;
-
-	UPROPERTY(EditAnywhere, BlueprintReadWrite, Category=Overrides, meta=(PinHiddenByDefault, InlineEditConditionToggle))
-	uint8 bOverride_FilmSaturation:1;
-
-	UPROPERTY(EditAnywhere, BlueprintReadWrite, Category=Overrides, meta=(PinHiddenByDefault, InlineEditConditionToggle))
-	uint8 bOverride_FilmChannelMixerRed:1;
-
-	UPROPERTY(EditAnywhere, BlueprintReadWrite, Category=Overrides, meta=(PinHiddenByDefault, InlineEditConditionToggle))
-	uint8 bOverride_FilmChannelMixerGreen:1;
-
-	UPROPERTY(EditAnywhere, BlueprintReadWrite, Category=Overrides, meta=(PinHiddenByDefault, InlineEditConditionToggle))
-	uint8 bOverride_FilmChannelMixerBlue:1;
-
-	UPROPERTY(EditAnywhere, BlueprintReadWrite, Category=Overrides, meta=(PinHiddenByDefault, InlineEditConditionToggle))
-	uint8 bOverride_FilmContrast:1;
-
-	UPROPERTY(EditAnywhere, BlueprintReadWrite, Category=Overrides, meta=(PinHiddenByDefault, InlineEditConditionToggle))
-	uint8 bOverride_FilmDynamicRange:1;
-
-	UPROPERTY(EditAnywhere, BlueprintReadWrite, Category=Overrides, meta=(PinHiddenByDefault, InlineEditConditionToggle))
-	uint8 bOverride_FilmHealAmount:1;
-
-	UPROPERTY(EditAnywhere, BlueprintReadWrite, Category=Overrides, meta=(PinHiddenByDefault, InlineEditConditionToggle))
-	uint8 bOverride_FilmToeAmount:1;
-
-	UPROPERTY(EditAnywhere, BlueprintReadWrite, Category=Overrides, meta=(PinHiddenByDefault, InlineEditConditionToggle))
-	uint8 bOverride_FilmShadowTint:1;
-
-	UPROPERTY(EditAnywhere, BlueprintReadWrite, Category=Overrides, meta=(PinHiddenByDefault, InlineEditConditionToggle))
-	uint8 bOverride_FilmShadowTintBlend:1;
-
-	UPROPERTY(EditAnywhere, BlueprintReadWrite, Category=Overrides, meta=(PinHiddenByDefault, InlineEditConditionToggle))
-	uint8 bOverride_FilmShadowTintAmount:1;
-
 	UPROPERTY(EditAnywhere, BlueprintReadWrite, Category=Overrides, meta=(PinHiddenByDefault, InlineEditConditionToggle))
 	uint8 bOverride_FilmSlope:1;
 	UPROPERTY(EditAnywhere, BlueprintReadWrite, Category=Overrides, meta=(PinHiddenByDefault, InlineEditConditionToggle))
@@ -1362,33 +1326,6 @@ struct FPostProcessSettings
 	UPROPERTY(interp, BlueprintReadWrite, Category="Film", meta=(UIMin = "0.0", UIMax = "1.0", editcondition = "bOverride_FilmWhiteClip", DisplayName = "White clip"))
 	float FilmWhiteClip;
 
-	UPROPERTY(interp, BlueprintReadWrite, Category="Film", meta=(editcondition = "bOverride_FilmWhitePoint", DisplayName = "Tint", HideAlphaChannel, LegacyTonemapper))
-	FLinearColor FilmWhitePoint;
-	UPROPERTY(interp, BlueprintReadWrite, Category="Film", AdvancedDisplay, meta=(editcondition = "bOverride_FilmShadowTint", DisplayName = "Tint Shadow", HideAlphaChannel, LegacyTonemapper))
-	FLinearColor FilmShadowTint;
-	UPROPERTY(interp, BlueprintReadWrite, Category="Film", AdvancedDisplay, meta=(UIMin = "0.0", UIMax = "1.0", editcondition = "bOverride_FilmShadowTintBlend", DisplayName = "Tint Shadow Blend", LegacyTonemapper))
-	float FilmShadowTintBlend;
-	UPROPERTY(interp, BlueprintReadWrite, Category="Film", AdvancedDisplay, meta=(UIMin = "0.0", UIMax = "1.0", editcondition = "bOverride_FilmShadowTintAmount", DisplayName = "Tint Shadow Amount", LegacyTonemapper))
-	float FilmShadowTintAmount;
-
-	UPROPERTY(interp, BlueprintReadWrite, Category="Film", meta=(UIMin = "0.0", UIMax = "2.0", editcondition = "bOverride_FilmSaturation", DisplayName = "Saturation", LegacyTonemapper))
-	float FilmSaturation;
-	UPROPERTY(interp, BlueprintReadWrite, Category="Film", AdvancedDisplay, meta=(editcondition = "bOverride_FilmChannelMixerRed", DisplayName = "Channel Mixer Red", HideAlphaChannel, LegacyTonemapper))
-	FLinearColor FilmChannelMixerRed;
-	UPROPERTY(interp, BlueprintReadWrite, Category="Film", AdvancedDisplay, meta=(editcondition = "bOverride_FilmChannelMixerGreen", DisplayName = "Channel Mixer Green", HideAlphaChannel, LegacyTonemapper))
-	FLinearColor FilmChannelMixerGreen;
-	UPROPERTY(interp, BlueprintReadWrite, Category="Film", AdvancedDisplay, meta=(editcondition = "bOverride_FilmChannelMixerBlue", DisplayName = " Channel Mixer Blue", HideAlphaChannel, LegacyTonemapper))
-	FLinearColor FilmChannelMixerBlue;
-
-	UPROPERTY(interp, BlueprintReadWrite, Category="Film", meta=(UIMin = "0.0", UIMax = "1.0", editcondition = "bOverride_FilmContrast", DisplayName = "Contrast", LegacyTonemapper))
-	float FilmContrast;
-	UPROPERTY(interp, BlueprintReadWrite, Category="Film", AdvancedDisplay, meta=(UIMin = "0.0", UIMax = "1.0", editcondition = "bOverride_FilmToeAmount", DisplayName = "Crush Shadows", LegacyTonemapper))
-	float FilmToeAmount;
-	UPROPERTY(interp, BlueprintReadWrite, Category="Film", AdvancedDisplay, meta=(UIMin = "0.0", UIMax = "1.0", editcondition = "bOverride_FilmHealAmount", DisplayName = "Crush Highlights", LegacyTonemapper))
-	float FilmHealAmount;
-	UPROPERTY(interp, BlueprintReadWrite, Category="Film", AdvancedDisplay, meta=(UIMin = "1.0", UIMax = "4.0", editcondition = "bOverride_FilmDynamicRange", DisplayName = "Dynamic Range", LegacyTonemapper))
-	float FilmDynamicRange;
-
 	/** Scene tint color */
 	UPROPERTY(interp, BlueprintReadWrite, Category="Color Grading|Misc", meta=(editcondition = "bOverride_SceneColorTint", HideAlphaChannel))
 	FLinearColor SceneColorTint;
diff --git a/Engine/Source/Runtime/Engine/Private/Scene.cpp b/Engine/Source/Runtime/Engine/Private/Scene.cpp
--- a/Engine/Source/Runtime/Engine/Private/Scene.cpp
+++ b/Engine/Source/Runtime/Engine/Private/Scene.cpp
@@ -339,19 +339,6 @@ static void DoPostProcessSettingsSanityCheck()
 
 	{
 		TMap<FString, FString> RenameMap;
-		// Old tonemapper parameters are ignored.
-		RenameMap.Add(TEXT("FilmWhitePoint"), TEXT(""));
-		RenameMap.Add(TEXT("FilmSaturation"), TEXT(""));
-		RenameMap.Add(TEXT("FilmChannelMixerRed"), TEXT(""));
-		RenameMap.Add(TEXT("FilmChannelMixerGreen"), TEXT(""));
-		RenameMap.Add(TEXT("FilmChannelMixerBlue"), TEXT(""));
-		RenameMap.Add(TEXT("FilmContrast"), TEXT(""));
-		RenameMap.Add(TEXT("FilmDynamicRange"), TEXT(""));
-		RenameMap.Add(TEXT("FilmHealAmount"), TEXT(""));
-		RenameMap.Add(TEXT("FilmToeAmount"), TEXT(""));
-		RenameMap.Add(TEXT("FilmShadowTint"), TEXT(""));
-		RenameMap.Add(TEXT("FilmShadowTintBlend"), TEXT(""));
-		RenameMap.Add(TEXT("FilmShadowTintAmount"), TEXT(""));
 		// Film Grain are ignored
 		RenameMap.Add(TEXT("FilmGrainIntensity"), TEXT(""));
 		RenameMap.Add(TEXT("FilmGrainIntensityShadows"), TEXT(""));
@@ -411,18 +398,6 @@ FPostProcessSettings::FPostProcessSettings()
 	ToneCurveAmount = 1.0;
 
 	// default values:
-	FilmWhitePoint = FLinearColor(1.0f, 1.0f, 1.0f);
-	FilmSaturation = 1.0f;
-	FilmChannelMixerRed = FLinearColor(1.0f, 0.0f, 0.0f);
-	FilmChannelMixerGreen = FLinearColor(0.0f, 1.0f, 0.0f);
-	FilmChannelMixerBlue = FLinearColor(0.0f, 0.0f, 1.0f);
-	FilmContrast = 0.03f;
-	FilmDynamicRange = 4.0f;
-	FilmHealAmount = 1.0f;
-	FilmToeAmount = 1.0f;
-	FilmShadowTint = FLinearColor(1.0f, 1.0f, 1.0f);
-	FilmShadowTintBlend = 0.5;
-	FilmShadowTintAmount = 0.0;
 
 	// ACES settings
 	FilmSlope = 0.88f;
@@ -653,18 +628,6 @@ FPostProcessSettings::FPostProcessSettings(const FPostProcessSettings& Settings)
 	, bOverride_BlueCorrection(Settings.bOverride_BlueCorrection)
 	, bOverride_ExpandGamut(Settings.bOverride_ExpandGamut)
 	, bOverride_ToneCurveAmount(Settings.bOverride_ToneCurveAmount)
-	, bOverride_FilmWhitePoint(Settings.bOverride_FilmWhitePoint)
-	, bOverride_FilmSaturation(Settings.bOverride_FilmSaturation)
-	, bOverride_FilmChannelMixerRed(Settings.bOverride_FilmChannelMixerRed)
-	, bOverride_FilmChannelMixerGreen(Settings.bOverride_FilmChannelMixerGreen)
-	, bOverride_FilmChannelMixerBlue(Settings.bOverride_FilmChannelMixerBlue)
-	, bOverride_FilmContrast(Settings.bOverride_FilmContrast)
-	, bOverride_FilmDynamicRange(Settings.bOverride_FilmDynamicRange)
-	, bOverride_FilmHealAmount(Settings.bOverride_FilmHealAmount)
-	, bOverride_FilmToeAmount(Settings.bOverride_FilmToeAmount)
-	, bOverride_FilmShadowTint(Settings.bOverride_FilmShadowTint)
-	, bOverride_FilmShadowTintBlend(Settings.bOverride_FilmShadowTintBlend)
-	, bOverride_FilmShadowTintAmount(Settings.bOverride_FilmShadowTintAmount)
 	, bOverride_FilmSlope(Settings.bOverride_FilmSlope)
 	, bOverride_FilmToe(Settings.bOverride_FilmToe)
 	, bOverride_FilmShoulder(Settings.bOverride_FilmShoulder)
@@ -849,18 +812,6 @@ FPostProcessSettings::FPostProcessSettings(const FPostProcessSettings& Settings)
 	, FilmShoulder(Settings.FilmShoulder)
 	, FilmBlackClip(Settings.FilmBlackClip)
 	, FilmWhiteClip(Settings.FilmWhiteClip)
-	, FilmWhitePoint(Settings.FilmWhitePoint)
-	, FilmShadowTint(Settings.FilmShadowTint)
-	, FilmShadowTintBlend(Settings.FilmShadowTintBlend)
-	, FilmShadowTintAmount(Settings.FilmShadowTintAmount)
-	, FilmSaturation(Settings.FilmSaturation)
-	, FilmChannelMixerRed(Settings.FilmChannelMixerRed)
-	, FilmChannelMixerGreen(Settings.FilmChannelMixerGreen)
-	, FilmChannelMixerBlue(Settings.FilmChannelMixerBlue)
-	, FilmContrast(Settings.FilmContrast)
-	, FilmToeAmount(Settings.FilmToeAmount)
-	, FilmHealAmount(Settings.FilmHealAmount)
-	, FilmDynamicRange(Settings.FilmDynamicRange)
 	, SceneColorTint(Settings.SceneColorTint)
 	, SceneFringeIntensity(Settings.SceneFringeIntensity)
 	, ChromaticAberrationStartOffset(Settings.ChromaticAberrationStartOffset)
diff --git a/Engine/Source/Runtime/Engine/Private/SceneView.cpp b/Engine/Source/Runtime/Engine/Private/SceneView.cpp
--- a/Engine/Source/Runtime/Engine/Private/SceneView.cpp
+++ b/Engine/Source/Runtime/Engine/Private/SceneView.cpp
@@ -268,9 +268,7 @@ static TAutoConsoleVariable<int32> CVarTonemapperQuality(
 	TEXT("Defines the Tonemapper Quality in the range 0..5\n")
 	TEXT("Depending on the used settings we might pick a faster shader permutation\n")
 	TEXT(" 0: basic tonemapper only, lowest quality\n")
-	TEXT(" 1: + FilmContrast\n")
 	TEXT(" 2: + Vignette\n")
-	TEXT(" 3: + FilmShadowTintAmount\n")
 	TEXT(" 4: + Grain\n")
 	TEXT(" 5: + GrainJitter = full quality (default)"),
 	ECVF_Scalability | ECVF_RenderThreadSafe);
@@ -1378,19 +1376,6 @@ void FSceneView::OverridePostProcessSettings(const FPostProcessSettings& Src, fl
 		LERP_PP(ExpandGamut);
 		LERP_PP(ToneCurveAmount);
 
-		LERP_PP(FilmWhitePoint);
-		LERP_PP(FilmSaturation);
-		LERP_PP(FilmChannelMixerRed);
-		LERP_PP(FilmChannelMixerGreen);
-		LERP_PP(FilmChannelMixerBlue);
-		LERP_PP(FilmContrast);
-		LERP_PP(FilmDynamicRange);
-		LERP_PP(FilmHealAmount);
-		LERP_PP(FilmToeAmount);
-		LERP_PP(FilmShadowTint);
-		LERP_PP(FilmShadowTintBlend);
-		LERP_PP(FilmShadowTintAmount);
-
 		LERP_PP(FilmSlope);
 		LERP_PP(FilmToe);
 		LERP_PP(FilmShoulder);
@@ -1908,21 +1893,11 @@ void FSceneView::EndFinalPostprocessSettings(const FSceneViewInitOptions& ViewIn
 	{
 		int32 Quality = CVarTonemapperQuality.GetValueOnGameThread();
 
-		if(Quality < 1)
-		{
-			FinalPostProcessSettings.FilmContrast = 0;
-		}
-
 		if(Quality < 2)
 		{
 			FinalPostProcessSettings.VignetteIntensity = 0;
 		}
 
-		if(Quality < 3)
-		{
-			FinalPostProcessSettings.FilmShadowTintAmount = 0;
-		}
-
 		if(Quality < 4)
 		{
 			FinalPostProcessSettings.FilmGrainIntensity = 0;
