Alexis.Matte, please merge this change by hand.
More info at https://robomerge.epicgames.net#STARSHIP

{"clientFile":"/src/ROBOMERGE_STARSHIP_UE5_Release_Engine_Test/Engine/Source/Runtime/Engine/Private/SkeletalMeshRenderData.cpp","fromFile":"//UE5/Release-Engine-Staging/Engine/Source/Runtime/Engine/Private/SkeletalMeshRenderData.cpp","startFromRev":54,"endFromRev":55,"resolveType":"content","resolveFlag":"c","contentResolveType":"3waytext"}
Fix an error when loading old asset with uninitialized cloth data
#rnx

diff --git a/Engine/Source/Runtime/Engine/Private/SkeletalMeshRenderData.cpp b/Engine/Source/Runtime/Engine/Private/SkeletalMeshRenderData.cpp
--- a/Engine/Source/Runtime/Engine/Private/SkeletalMeshRenderData.cpp
+++ b/Engine/Source/Runtime/Engine/Private/SkeletalMeshRenderData.cpp
@@ -186,7 +186,7 @@ static void SerializeLODInfoForDDC(USkeletalMesh* SkeletalMesh, FString& KeySuff
 // differences, etc.) replace the version GUID below with a new one.
 // In case of merge conflicts with DDC versions, you MUST generate a new GUID
 // and set this new GUID as the version.
-#define SKELETALMESH_DERIVEDDATA_VER TEXT("0F778E3964824C509A348A7D390C13D8")
+#define SKELETALMESH_DERIVEDDATA_VER TEXT("6E444E7FF7E948B9B8B1F33F898CD121")
 
 const FString& GetSkeletalMeshDerivedDataVersion()
 {
