Bad Merge cleanup, a bad merge integrate empty change, this restor the change properly.
cherry pick the cl: 17680434 from Release5.0
Fix an error when loading old asset with uninitialized cloth data
#rnx

diff --git a/Engine/Source/Runtime/Engine/Private/Rendering/SkeletalMeshLODImporterData.cpp b/Engine/Source/Runtime/Engine/Private/Rendering/SkeletalMeshLODImporterData.cpp
--- a/Engine/Source/Runtime/Engine/Private/Rendering/SkeletalMeshLODImporterData.cpp
+++ b/Engine/Source/Runtime/Engine/Private/Rendering/SkeletalMeshLODImporterData.cpp
@@ -553,7 +553,16 @@ void FReductionBaseSkeletalMeshBulkData::LoadReductionData(FSkeletalMeshLODModel
 			//This call will filled missing chunked data for old asset that cannot build normal and chunking (not re import since the new skeletal mesh chunk build refactor)
 			if (!BaseLODModel.bIsBuildDataAvailable)
 			{
-				BaseLODModel.UpdateChunkedSectionInfo(Owner ? Owner->GetName() : FString(TEXT("")));
+				int32 NumberOfClothAssets = 0;
+				if (USkeletalMesh* SkeletalMesh = Cast<USkeletalMesh>(Owner))
+				{
+					NumberOfClothAssets = SkeletalMesh->GetMeshClothingAssets().Num();
+				}
+				else
+				{
+					UE_LOG(LogSkeletalMeshLODImporterData, Warning, TEXT("In function FReductionBaseSkeletalMeshBulkData::LoadReductionData, the Owner parameter should be a valid skeletalmesh."));
+				}
+				BaseLODModel.UpdateChunkedSectionInfo(Owner ? Owner->GetName() : FString(TEXT("")), NumberOfClothAssets);
 			}
 		}
 		// Unlock the bulk data
diff --git a/Engine/Source/Runtime/Engine/Private/SkeletalMesh.cpp b/Engine/Source/Runtime/Engine/Private/SkeletalMesh.cpp
--- a/Engine/Source/Runtime/Engine/Private/SkeletalMesh.cpp
+++ b/Engine/Source/Runtime/Engine/Private/SkeletalMesh.cpp
@@ -2700,7 +2700,9 @@ void USkeletalMesh::CreateUserSectionsDataForLegacyAssets()
 			}
 			bMustUseReductionSourceData = false;
 		}
-		ThisLODModel.UpdateChunkedSectionInfo(GetName());
+
+		const int32 NumberOfClothAssets = GetMeshClothingAssets().Num();
+		ThisLODModel.UpdateChunkedSectionInfo(GetName(), NumberOfClothAssets);
 
 		if (bMustUseReductionSourceData)
 		{
diff --git a/Engine/Source/Runtime/Engine/Private/SkeletalMeshLODModel.cpp b/Engine/Source/Runtime/Engine/Private/SkeletalMeshLODModel.cpp
--- a/Engine/Source/Runtime/Engine/Private/SkeletalMeshLODModel.cpp
+++ b/Engine/Source/Runtime/Engine/Private/SkeletalMeshLODModel.cpp
@@ -1171,7 +1171,7 @@ FString FSkeletalMeshLODModel::GetLODModelDeriveDataKey() const
 	return KeySuffix;
 }
 
-void FSkeletalMeshLODModel::UpdateChunkedSectionInfo(const FString& SkeletalMeshName)
+void FSkeletalMeshLODModel::UpdateChunkedSectionInfo(const FString& SkeletalMeshName, const int32 NumberOfClothAssets)
 {
 	int32 LODModelSectionNum = Sections.Num();
 	//Fill the ChunkedParentSectionIndex data, we assume that every section using the same material are chunked
@@ -1217,6 +1217,14 @@ void FSkeletalMeshLODModel::UpdateChunkedSectionInfo(const FString& SkeletalMesh
 			SectionUserData.bRecomputeTangent = Section.bRecomputeTangent;
 			SectionUserData.RecomputeTangentsVertexMaskChannel = Section.RecomputeTangentsVertexMaskChannel;
 			SectionUserData.GenerateUpToLodIndex = Section.GenerateUpToLodIndex;
+			//Make sure the CorrespondClothAssetIndex is valid it goes from -1(INDEX_NONE) to the number of cloth asset -1
+			if (Section.CorrespondClothAssetIndex < -1 ||  Section.CorrespondClothAssetIndex >= NumberOfClothAssets)
+			{
+				Section.CorrespondClothAssetIndex = INDEX_NONE;
+				Section.ClothingData.AssetGuid = FGuid();
+				Section.ClothingData.AssetLodIndex = INDEX_NONE;
+			}
+
 			SectionUserData.CorrespondClothAssetIndex = Section.CorrespondClothAssetIndex;
 			SectionUserData.ClothingData.AssetGuid = Section.ClothingData.AssetGuid;
 			SectionUserData.ClothingData.AssetLodIndex = Section.ClothingData.AssetLodIndex;
diff --git a/Engine/Source/Runtime/Engine/Public/Rendering/SkeletalMeshLODModel.h b/Engine/Source/Runtime/Engine/Public/Rendering/SkeletalMeshLODModel.h
--- a/Engine/Source/Runtime/Engine/Public/Rendering/SkeletalMeshLODModel.h
+++ b/Engine/Source/Runtime/Engine/Public/Rendering/SkeletalMeshLODModel.h
@@ -486,7 +486,7 @@ public:
 	* This function will update the chunked information for each section. Only old data before the 
 	* skeletal mesh build refactor should need to call this function.
 	*/
-	ENGINE_API void UpdateChunkedSectionInfo(const FString& SkeletalMeshName);
+	ENGINE_API void UpdateChunkedSectionInfo(const FString& SkeletalMeshName, const int32 NumberOfClothAssets);
 
 	/**
 	* Copy one structure to the other, make sure all bulk data is unlock and the data can be read before copying.
