[VT] Restrict staging in GPU memory to GRHISupportsUpdateFromBufferTexture (in addition to GRHISupportsDirectGPUMemoryLock) for those platforms that don't support both.

diff --git a/Engine/Source/Runtime/Engine/Private/VT/VirtualTextureUploadCache.cpp b/Engine/Source/Runtime/Engine/Private/VT/VirtualTextureUploadCache.cpp
--- a/Engine/Source/Runtime/Engine/Private/VT/VirtualTextureUploadCache.cpp
+++ b/Engine/Source/Runtime/Engine/Private/VT/VirtualTextureUploadCache.cpp
@@ -181,7 +181,7 @@ void FVTUploadTileAllocator::FStagingBuffer::Init(uint32 InBufferStrideBytes, ui
 		TileFreeList[Index] = NumTiles - Index - 1;
 	}
 
-	if (GRHISupportsDirectGPUMemoryLock)
+	if (GRHISupportsDirectGPUMemoryLock && GRHISupportsUpdateFromBufferTexture)
 	{
 		// Allocate staging buffer directly in GPU memory.
 		checkSlow(IsInRenderingThread());
