Fix occasional failure in the System.Engine.VisualLogger unit tests.

#rnx

- In CL 17578544, NewTimestamp was changed from a float to a double to fix truncation warnings. However FVisualLogge::GetEntryToWrite would take this as a float and it would be stored as a float in FVisualLogEntry.
- Depending on the value of the timestamp this could cause small enough differences between NewTimestamp and the timestamp stored in the FVisualLogEntry being tests, causing the comparision check of those timestamps to fail.
- Changing NewTimestamp back to float in all cases and instead, changing the value being added to the timestamp to be a float constant rather than a double constant (the cause of the original warning) so that the data stays in the same format.
- The original cast warning fixed in 17578544 should remain fixed.

diff --git a/Engine/Source/Runtime/Engine/Private/VisualLogger/VisualLoggerAutomationTests.cpp b/Engine/Source/Runtime/Engine/Private/VisualLogger/VisualLoggerAutomationTests.cpp
--- a/Engine/Source/Runtime/Engine/Private/VisualLogger/VisualLoggerAutomationTests.cpp
+++ b/Engine/Source/Runtime/Engine/Private/VisualLogger/VisualLoggerAutomationTests.cpp
@@ -143,7 +143,7 @@ bool FVisualLogTest::RunTest(const FString& Parameters)
 			CHECK_SUCCESS(CurrentEntry->LogLines[0].Category == LogVisual.GetCategoryName());
 			CHECK_SUCCESS(CurrentEntry->LogLines[0].Line == TextToLog);
 
-			const double NewTimestamp = CurrentTimestamp + 0.1;
+			const float NewTimestamp = CurrentTimestamp + 0.1f;
 			FVisualLogEntry* NewEntry = FVisualLogger::Get().GetEntryToWrite(World, NewTimestamp); //generate new entry and serialize old one
 			CurrentEntry = &Context.Device.LastEntry;
 			CHECK_NOT_NULL(CurrentEntry);
@@ -199,7 +199,7 @@ bool FVisualLogSegmentsTest::RunTest(const FString& Parameters)
 			CHECK_SUCCESS(CurrentEntry->ElementsToDraw[0].Points[0] == StartPoint);
 			CHECK_SUCCESS(CurrentEntry->ElementsToDraw[0].Points[1] == EndPoint);
 
-			const double NewTimestamp = CurrentTimestamp + 0.1;
+			const float NewTimestamp = CurrentTimestamp + 0.1f;
 			FVisualLogEntry* NewEntry = FVisualLogger::Get().GetEntryToWrite(World, NewTimestamp); //generate new entry and serialize old one
 			CurrentEntry = &Context.Device.LastEntry;
 			CHECK_NOT_NULL(CurrentEntry);
@@ -277,7 +277,7 @@ bool FVisualLogEventsTest::RunTest(const FString& Parameters)
 		CHECK_SUCCESS(CurrentEntry->Events[1].UserFriendlyDesc == TEXT("Second simple event for vlog tests"));
 		CHECK_SUCCESS(CurrentEntry->Events[2].UserFriendlyDesc == TEXT("Third simple event for vlog tests"));
 
-		const double NewTimestamp = CurrentTimestamp + 0.1;
+		const float NewTimestamp = CurrentTimestamp + 0.1f;
 		FVisualLogEntry* NewEntry = FVisualLogger::Get().GetEntryToWrite(World, NewTimestamp); //generate new entry and serialize old one
 		CurrentEntry = &Context.Device.LastEntry;
 		CHECK_NOT_NULL(CurrentEntry);
@@ -303,7 +303,7 @@ bool FVisualLogEventsTest::RunTest(const FString& Parameters)
 	const FName EventTag2 = TEXT("ATLAS_C_1");
 	const FName EventTag3 = TEXT("ATLAS_C_2");
 
-	CurrentTimestamp = World->TimeSeconds + 0.2;
+	CurrentTimestamp = World->TimeSeconds + 0.2f;
 	CurrentEntry = FVisualLogger::Get().GetEntryToWrite(World, CurrentTimestamp); //generate new entry and serialize old one
 	UE_VLOG_EVENT_WITH_DATA(World, EventTest, EventTag1);
 	CHECK_NOT_NULL(CurrentEntry);
@@ -313,7 +313,7 @@ bool FVisualLogEventsTest::RunTest(const FString& Parameters)
 	CHECK_SUCCESS(CurrentEntry->Events[0].EventTags.Num() == 1);
 	CHECK_SUCCESS(CurrentEntry->Events[0].EventTags[EventTag1] == 1);
 
-	CurrentTimestamp = World->TimeSeconds + 0.3;
+	CurrentTimestamp = World->TimeSeconds + 0.3f;
 	CurrentEntry = FVisualLogger::Get().GetEntryToWrite(World, CurrentTimestamp); //generate new entry and serialize old one
 	UE_VLOG_EVENT_WITH_DATA(World, EventTest, EventTag1, EventTag2, EventTag3);
 	UE_VLOG_EVENT_WITH_DATA(World, EventTest, EventTag3);
