Adding functions on UExponentialHeightFogComponent to allow users to set Second Fog variables and the whole struct through Blueprint

diff --git a/Engine/Source/Runtime/Engine/Classes/Components/ExponentialHeightFogComponent.h b/Engine/Source/Runtime/Engine/Classes/Components/ExponentialHeightFogComponent.h
--- a/Engine/Source/Runtime/Engine/Classes/Components/ExponentialHeightFogComponent.h
+++ b/Engine/Source/Runtime/Engine/Classes/Components/ExponentialHeightFogComponent.h
@@ -188,6 +188,9 @@ public:
 	UFUNCTION(BlueprintCallable, Category="Rendering|Components|ExponentialHeightFog")
 	void SetFogDensity(float Value);
 
+	UFUNCTION(BlueprintCallable, Category="Rendering|Components|ExponentialHeightFog")
+	void SetSecondFogDensity(float Value);
+
 	UFUNCTION(BlueprintCallable, Category="Rendering|Components|ExponentialHeightFog")
 	void SetFogInscatteringColor(FLinearColor Value);
 
@@ -214,10 +217,16 @@ public:
 
 	UFUNCTION(BlueprintCallable, Category="Rendering|Components|ExponentialHeightFog")
 	void SetDirectionalInscatteringColor(FLinearColor Value);
+
+	UFUNCTION(BlueprintCallable, Category="Rendering|Components|ExponentialHeightFog")
+	void SetSecondFogHeightOffset(float Value);
 	
 	UFUNCTION(BlueprintCallable, Category="Rendering|Components|ExponentialHeightFog")
 	void SetFogHeightFalloff(float Value);
 
+	UFUNCTION(BlueprintCallable, Category="Rendering|Components|ExponentialHeightFog")
+	void SetSecondFogHeightFalloff(float Value);
+
 	UFUNCTION(BlueprintCallable, Category="Rendering|Components|ExponentialHeightFog")
 	void SetFogMaxOpacity(float Value);
 	
@@ -226,7 +235,7 @@ public:
 
 	UFUNCTION(BlueprintCallable, Category="Rendering|Components|ExponentialHeightFog")
 	void SetFogCutoffDistance(float Value);
-
+	
 	UFUNCTION(BlueprintCallable, Category="Rendering|VolumetricFog")
 	void SetVolumetricFog(bool bNewValue);
 
@@ -245,6 +254,9 @@ public:
 	UFUNCTION(BlueprintCallable, Category="Rendering|VolumetricFog")
 	void SetVolumetricFogDistance(float NewValue);
 
+	UFUNCTION(BlueprintCallable, Category="Rendering|VolumetricFog")
+	void SetSecondFogData(FExponentialHeightFogData NewValue);
+
 protected:
 	//~ Begin UActorComponent Interface.
 	virtual void CreateRenderState_Concurrent(FRegisterComponentContext* Context) override;
diff --git a/Engine/Source/Runtime/Engine/Private/Components/HeightFogComponent.cpp b/Engine/Source/Runtime/Engine/Private/Components/HeightFogComponent.cpp
--- a/Engine/Source/Runtime/Engine/Private/Components/HeightFogComponent.cpp
+++ b/Engine/Source/Runtime/Engine/Private/Components/HeightFogComponent.cpp
@@ -52,6 +52,9 @@ UExponentialHeightFogComponent::UExponentialHeightFogComponent(const FObjectInit
 
 void UExponentialHeightFogComponent::AddFogIfNeeded()
 {
+	// For safety, clamp the values for SecondFogData here.
+	SecondFogData.ClampToValidRanges();
+	
 	if (ShouldComponentAddToScene() && ShouldRender() && IsRegistered() && ((FogDensity + SecondFogData.FogDensity) * 1000) > DELTA && FogMaxOpacity > DELTA
 		&& (GetOuter() == NULL || !GetOuter()->HasAnyFlags(RF_ClassDefaultObject)))
 	{
@@ -138,6 +141,15 @@ void UExponentialHeightFogComponent::SetFogDensity(float Value)
 	}
 }
 
+void UExponentialHeightFogComponent::SetSecondFogDensity(float Value)
+{
+	if(SecondFogData.FogDensity != Value)
+	{
+		SecondFogData.FogDensity = Value;
+		MarkRenderStateDirty();
+	}
+}
+
 void UExponentialHeightFogComponent::SetFogInscatteringColor(FLinearColor Value)
 {
 	if(FogInscatteringLuminance != Value)
@@ -219,6 +231,15 @@ void UExponentialHeightFogComponent::SetDirectionalInscatteringColor(FLinearColo
 	}
 }
 
+void UExponentialHeightFogComponent::SetSecondFogHeightOffset(float Value)
+{
+	if(SecondFogData.FogHeightOffset != Value)
+	{
+		SecondFogData.FogHeightOffset = Value;
+		MarkRenderStateDirty();
+	}
+}
+
 void UExponentialHeightFogComponent::SetFogHeightFalloff(float Value)
 {
 	if(FogHeightFalloff != Value)
@@ -228,6 +249,15 @@ void UExponentialHeightFogComponent::SetFogHeightFalloff(float Value)
 	}
 }
 
+void UExponentialHeightFogComponent::SetSecondFogHeightFalloff(float Value)
+{
+	if(SecondFogData.FogHeightFalloff != Value)
+	{
+		SecondFogData.FogHeightFalloff = Value;
+		MarkRenderStateDirty();
+	}
+}
+
 void UExponentialHeightFogComponent::SetFogMaxOpacity(float Value)
 {
 	if(FogMaxOpacity != Value)
@@ -309,6 +339,17 @@ void UExponentialHeightFogComponent::SetVolumetricFogDistance(float NewValue)
 	}
 }
 
+void UExponentialHeightFogComponent::SetSecondFogData(FExponentialHeightFogData NewValue)
+{
+	if(SecondFogData.FogDensity != NewValue.FogDensity ||
+	   SecondFogData.FogHeightOffset != NewValue.FogHeightOffset ||
+	   SecondFogData.FogHeightFalloff != NewValue.FogHeightFalloff)
+	{
+		SecondFogData = NewValue;
+		MarkRenderStateDirty();
+	}
+}
+
 void UExponentialHeightFogComponent::Serialize(FArchive& Ar)
 {
 	Super::Serialize(Ar);
