Fix a crash with cloth mapping data
#rnx

diff --git a/Engine/Source/Runtime/Engine/Private/Rendering/SkeletalMeshLODImporterData.cpp b/Engine/Source/Runtime/Engine/Private/Rendering/SkeletalMeshLODImporterData.cpp
--- a/Engine/Source/Runtime/Engine/Private/Rendering/SkeletalMeshLODImporterData.cpp
+++ b/Engine/Source/Runtime/Engine/Private/Rendering/SkeletalMeshLODImporterData.cpp
@@ -553,16 +553,7 @@ void FReductionBaseSkeletalMeshBulkData::LoadReductionData(FSkeletalMeshLODModel
 			//This call will filled missing chunked data for old asset that cannot build normal and chunking (not re import since the new skeletal mesh chunk build refactor)
 			if (!BaseLODModel.bIsBuildDataAvailable)
 			{
-				int32 NumberOfClothAssets = 0;
-				if (USkeletalMesh* SkeletalMesh = Cast<USkeletalMesh>(Owner))
-				{
-					NumberOfClothAssets = SkeletalMesh->GetMeshClothingAssets().Num();
-				}
-				else
-				{
-					UE_LOG(LogSkeletalMeshLODImporterData, Warning, TEXT("In function FReductionBaseSkeletalMeshBulkData::LoadReductionData, the Owner parameter should be a valid skeletalmesh."));
-				}
-				BaseLODModel.UpdateChunkedSectionInfo(Owner ? Owner->GetName() : FString(TEXT("")), NumberOfClothAssets);
+				BaseLODModel.UpdateChunkedSectionInfo(Owner ? Owner->GetName() : FString(TEXT("")));
 			}
 		}
 		// Unlock the bulk data
diff --git a/Engine/Source/Runtime/Engine/Private/SkeletalMesh.cpp b/Engine/Source/Runtime/Engine/Private/SkeletalMesh.cpp
--- a/Engine/Source/Runtime/Engine/Private/SkeletalMesh.cpp
+++ b/Engine/Source/Runtime/Engine/Private/SkeletalMesh.cpp
@@ -2489,6 +2489,13 @@ void USkeletalMesh::RemoveLegacyClothingSections()
 						// Remove the reference index
 						Section.CorrespondClothSectionIndex_DEPRECATED = INDEX_NONE;
 
+						//Make sure the UserSectionsData is up to date
+						if (FSkelMeshSourceSectionUserData* SectionUserData = LodModel.UserSectionsData.Find(Section.OriginalDataSectionIndex))
+						{
+							SectionUserData->CorrespondClothAssetIndex = Section.CorrespondClothAssetIndex;
+							SectionUserData->ClothingData = Section.ClothingData;
+						}
+
 						ClothingSectionCount++;
 					}
 					else
@@ -2701,8 +2708,7 @@ void USkeletalMesh::CreateUserSectionsDataForLegacyAssets()
 			bMustUseReductionSourceData = false;
 		}
 
-		const int32 NumberOfClothAssets = GetMeshClothingAssets().Num();
-		ThisLODModel.UpdateChunkedSectionInfo(GetName(), NumberOfClothAssets);
+		ThisLODModel.UpdateChunkedSectionInfo(GetName());
 
 		if (bMustUseReductionSourceData)
 		{
diff --git a/Engine/Source/Runtime/Engine/Private/SkeletalMeshLODModel.cpp b/Engine/Source/Runtime/Engine/Private/SkeletalMeshLODModel.cpp
--- a/Engine/Source/Runtime/Engine/Private/SkeletalMeshLODModel.cpp
+++ b/Engine/Source/Runtime/Engine/Private/SkeletalMeshLODModel.cpp
@@ -1171,7 +1171,7 @@ FString FSkeletalMeshLODModel::GetLODModelDeriveDataKey() const
 	return KeySuffix;
 }
 
-void FSkeletalMeshLODModel::UpdateChunkedSectionInfo(const FString& SkeletalMeshName, const int32 NumberOfClothAssets)
+void FSkeletalMeshLODModel::UpdateChunkedSectionInfo(const FString& SkeletalMeshName)
 {
 	int32 LODModelSectionNum = Sections.Num();
 	//Fill the ChunkedParentSectionIndex data, we assume that every section using the same material are chunked
@@ -1217,12 +1217,10 @@ void FSkeletalMeshLODModel::UpdateChunkedSectionInfo(const FString& SkeletalMesh
 			SectionUserData.bRecomputeTangent = Section.bRecomputeTangent;
 			SectionUserData.RecomputeTangentsVertexMaskChannel = Section.RecomputeTangentsVertexMaskChannel;
 			SectionUserData.GenerateUpToLodIndex = Section.GenerateUpToLodIndex;
-			//Make sure the CorrespondClothAssetIndex is valid it goes from -1(INDEX_NONE) to the number of cloth asset -1
-			if (Section.CorrespondClothAssetIndex < -1 ||  Section.CorrespondClothAssetIndex >= NumberOfClothAssets)
+			//Make sure the CorrespondClothAssetIndex is valid
+			if (Section.CorrespondClothAssetIndex < -1)
 			{
 				Section.CorrespondClothAssetIndex = INDEX_NONE;
-				Section.ClothingData.AssetGuid = FGuid();
-				Section.ClothingData.AssetLodIndex = INDEX_NONE;
 			}
 
 			SectionUserData.CorrespondClothAssetIndex = Section.CorrespondClothAssetIndex;
diff --git a/Engine/Source/Runtime/Engine/Private/SkeletalMeshRenderData.cpp b/Engine/Source/Runtime/Engine/Private/SkeletalMeshRenderData.cpp
--- a/Engine/Source/Runtime/Engine/Private/SkeletalMeshRenderData.cpp
+++ b/Engine/Source/Runtime/Engine/Private/SkeletalMeshRenderData.cpp
@@ -186,7 +186,7 @@ static void SerializeLODInfoForDDC(USkeletalMesh* SkeletalMesh, FString& KeySuff
 // differences, etc.) replace the version GUID below with a new one.
 // In case of merge conflicts with DDC versions, you MUST generate a new GUID
 // and set this new GUID as the version.
-#define SKELETALMESH_DERIVEDDATA_VER TEXT("6E444E7FF7E948B9B8B1F33F898CD121")
+#define SKELETALMESH_DERIVEDDATA_VER TEXT("6BAEFEB10ADB4FC2961282E8D9BC7B77")
 
 const FString& GetSkeletalMeshDerivedDataVersion()
 {
diff --git a/Engine/Source/Runtime/Engine/Public/Rendering/SkeletalMeshLODModel.h b/Engine/Source/Runtime/Engine/Public/Rendering/SkeletalMeshLODModel.h
--- a/Engine/Source/Runtime/Engine/Public/Rendering/SkeletalMeshLODModel.h
+++ b/Engine/Source/Runtime/Engine/Public/Rendering/SkeletalMeshLODModel.h
@@ -486,7 +486,7 @@ public:
 	* This function will update the chunked information for each section. Only old data before the 
 	* skeletal mesh build refactor should need to call this function.
 	*/
-	ENGINE_API void UpdateChunkedSectionInfo(const FString& SkeletalMeshName, const int32 NumberOfClothAssets);
+	ENGINE_API void UpdateChunkedSectionInfo(const FString& SkeletalMeshName);
 
 	/**
 	* Copy one structure to the other, make sure all bulk data is unlock and the data can be read before copying.
