PR #8243: [New][SubtitleManager] Expose ENGINE_API for QueueSubtitles (Contributed by [REDACTED_PERSON])

diff --git a/Engine/Source/Runtime/Engine/Public/SubtitleManager.h b/Engine/Source/Runtime/Engine/Public/SubtitleManager.h
--- a/Engine/Source/Runtime/Engine/Public/SubtitleManager.h
+++ b/Engine/Source/Runtime/Engine/Public/SubtitleManager.h
@@ -112,7 +112,7 @@ public:
 	 * @param SoundDuration Time  after which the subtitles do not display (in seconds).
 	 * @param Subtitles Collection of lines of subtitle and time offset to play them.
 	 */
-	void QueueSubtitles(PTRINT SubtitleID, float Priority, bool bManualWordWrap, bool bSingleLine, float SoundDuration, const TArray<FSubtitleCue>& Subtitles, float InStartTime, float InCurrentTime);
+	ENGINE_API void QueueSubtitles(PTRINT SubtitleID, float Priority, bool bManualWordWrap, bool bSingleLine, float SoundDuration, const TArray<FSubtitleCue>& Subtitles, float InStartTime, float InCurrentTime);
 
 	static void QueueSubtitles(const FQueueSubtitleParams& QueueSubtitlesParams);
 
