Narrowing missing LSA repair to occur only in pure editor modes (i.e. not in PIE). This fixes the case in multiplayer mode where a duplicate is created on a client if the project sets the LSA to be replicated.

diff --git a/Engine/Source/Runtime/Engine/Private/Level.cpp b/Engine/Source/Runtime/Engine/Private/Level.cpp
--- a/Engine/Source/Runtime/Engine/Private/Level.cpp
+++ b/Engine/Source/Runtime/Engine/Private/Level.cpp
@@ -3089,7 +3089,7 @@ void ULevel::RepairLevelScript()
 	// Catch the edge case where we have a level blueprint but have never created the LevelScriptActor based on it.
 	//    This could happen if a new level is saved before the level blueprint is compiled.
 	if (!LevelScriptActor && LevelScriptBlueprint && !LevelScriptBlueprint->bIsRegeneratingOnLoad &&
-		OwningWorld && OwningWorld->IsEditorWorld())
+		OwningWorld && OwningWorld->IsEditorWorld() && !OwningWorld->IsGameWorld())
 	{
 		RegenerateLevelScriptActor();
 	}
