[Actor Iterators] added support for `const UWorld*`
Affected classes: FActorIteratorState, TActorIteratorBase, FActorIterator, TActorIterator, TActorRange, FSelectedActorIterator, FSelectedActorRange

diff --git a/Engine/Source/Runtime/Engine/Classes/Engine/World.h b/Engine/Source/Runtime/Engine/Classes/Engine/World.h
--- a/Engine/Source/Runtime/Engine/Classes/Engine/World.h
+++ b/Engine/Source/Runtime/Engine/Classes/Engine/World.h
@@ -1452,10 +1452,10 @@ private:
 #endif
 
 	/** a delegate that broadcasts a notification whenever an actor is spawned */
-	FOnActorSpawned OnActorSpawned;
+	mutable FOnActorSpawned OnActorSpawned;
 
 	/** a delegate that broadcasts a notification before a newly spawned actor is initialized */
-	FOnActorSpawned OnActorPreSpawnInitialization;
+	mutable FOnActorSpawned OnActorPreSpawnInitialization;
 
 	/** Reset Async Trace Buffer **/
 	void ResetAsyncTrace();
@@ -2601,16 +2601,16 @@ public:
 	void RemoveNetworkActor( AActor* Actor ) const;
 
 	/** Add a listener for OnActorSpawned events */
-	FDelegateHandle AddOnActorSpawnedHandler( const FOnActorSpawned::FDelegate& InHandler );
+	FDelegateHandle AddOnActorSpawnedHandler( const FOnActorSpawned::FDelegate& InHandler ) const;
 
 	/** Remove a listener for OnActorSpawned events */
-	void RemoveOnActorSpawnedHandler( FDelegateHandle InHandle );
+	void RemoveOnActorSpawnedHandler( FDelegateHandle InHandle ) const;
 
 	/** Add a listener for OnActorPreSpawnInitialization events */
-	FDelegateHandle AddOnActorPreSpawnInitialization(const FOnActorSpawned::FDelegate& InHandler);
+	FDelegateHandle AddOnActorPreSpawnInitialization(const FOnActorSpawned::FDelegate& InHandler) const;
 
 	/** Remove a listener for OnActorPreSpawnInitialization events */
-	void RemoveOnActorPreSpawnInitialization(FDelegateHandle InHandle);	
+	void RemoveOnActorPreSpawnInitialization(FDelegateHandle InHandle) const;
 
 	/**
 	 * Returns whether the passed in actor is part of any of the loaded levels actors array.
diff --git a/Engine/Source/Runtime/Engine/Private/World.cpp b/Engine/Source/Runtime/Engine/Private/World.cpp
--- a/Engine/Source/Runtime/Engine/Private/World.cpp
+++ b/Engine/Source/Runtime/Engine/Private/World.cpp
@@ -5087,22 +5087,22 @@ void UWorld::RemoveNetworkActor( AActor* Actor ) const
 	}
 }
 
-FDelegateHandle UWorld::AddOnActorSpawnedHandler( const FOnActorSpawned::FDelegate& InHandler )
+FDelegateHandle UWorld::AddOnActorSpawnedHandler(const FOnActorSpawned::FDelegate& InHandler) const
 {
 	return OnActorSpawned.Add(InHandler);
 }
 
-void UWorld::RemoveOnActorSpawnedHandler( FDelegateHandle InHandle )
+void UWorld::RemoveOnActorSpawnedHandler(FDelegateHandle InHandle) const
 {
 	OnActorSpawned.Remove(InHandle);
 }
 
-FDelegateHandle UWorld::AddOnActorPreSpawnInitialization(const FOnActorSpawned::FDelegate& InHandler)
+FDelegateHandle UWorld::AddOnActorPreSpawnInitialization(const FOnActorSpawned::FDelegate& InHandler) const
 {
 	return OnActorPreSpawnInitialization.Add(InHandler);
 }
 
-void UWorld::RemoveOnActorPreSpawnInitialization(FDelegateHandle InHandle)
+void UWorld::RemoveOnActorPreSpawnInitialization(FDelegateHandle InHandle) const
 {
 	OnActorPreSpawnInitialization.Remove(InHandle);
 }
diff --git a/Engine/Source/Runtime/Engine/Public/EngineUtils.h b/Engine/Source/Runtime/Engine/Public/EngineUtils.h
--- a/Engine/Source/Runtime/Engine/Public/EngineUtils.h
+++ b/Engine/Source/Runtime/Engine/Public/EngineUtils.h
@@ -161,7 +161,7 @@ class FActorIteratorState
 {
 public:
 	/** Current world we are iterating upon						*/
-	UWorld* CurrentWorld;
+	const UWorld* CurrentWorld;
 	/** Results from the GetObjectsOfClass query				*/
 	TArray<UObject*> ObjectArray;
 	/** index of the current element in the object array		*/
@@ -182,12 +182,12 @@ public:
 	/**
 	 * Default ctor, inits everything
 	 */
-	FActorIteratorState( UWorld* InWorld, TSubclassOf<AActor> InClass ) :
+	FActorIteratorState(const UWorld* InWorld, const TSubclassOf<AActor> InClass) :
 		CurrentWorld( InWorld ),
 		Index( -1 ),
 		ReachedEnd( false ),
 		ConsideredCount( 0 ),
-		CurrentActor( NULL ),
+		CurrentActor(nullptr),
 		DesiredClass(InClass)
 	{
 		check(IsInGameThread());
@@ -226,11 +226,11 @@ public:
 		else
 #endif // WITH_EDITOR
 		{
-			EObjectFlags ExcludeFlags = RF_ClassDefaultObject;
+			constexpr EObjectFlags ExcludeFlags = RF_ClassDefaultObject;
 			GetObjectsOfClass(InClass, ObjectArray, true, ExcludeFlags, EInternalObjectFlags::PendingKill);
 		}
 
-		auto ActorSpawnedDelegate = FOnActorSpawned::FDelegate::CreateRaw(this, &FActorIteratorState::OnActorSpawned);
+		const auto ActorSpawnedDelegate = FOnActorSpawned::FDelegate::CreateRaw(this, &FActorIteratorState::OnActorSpawned);
 		ActorSpawnedDelegateHandle = CurrentWorld->AddOnActorSpawnedHandler(ActorSpawnedDelegate);
 	}
 
@@ -294,7 +294,7 @@ public:
 		int32             LocalIndex             = State->Index;
 		TArray<UObject*>& LocalObjectArray       = State->ObjectArray;
 		TArray<AActor*>&  LocalSpawnedActorArray = State->SpawnedActorArray;
-		UWorld*           LocalCurrentWorld      = State->CurrentWorld;
+		const UWorld*     LocalCurrentWorld      = State->CurrentWorld;
 		while(++LocalIndex < (LocalObjectArray.Num() + LocalSpawnedActorArray.Num()))
 		{
 			if (LocalIndex < LocalObjectArray.Num())
@@ -322,7 +322,7 @@ public:
 				}
 			}
 		}
-		State->CurrentActor = NULL;
+		State->CurrentActor = nullptr;
 		State->ReachedEnd = true;
 	}
 
@@ -385,7 +385,7 @@ protected:
 	{
 	}
 
-	TActorIteratorBase(UWorld* InWorld, TSubclassOf<AActor> InClass, EActorIteratorFlags InFlags)
+	TActorIteratorBase(const UWorld* InWorld, TSubclassOf<AActor> InClass, const EActorIteratorFlags InFlags)
 		: Flags(InFlags)
 	{
 		State.Emplace(InWorld, InClass);
@@ -398,7 +398,7 @@ protected:
 	 * @param	Actor	Actor to check
 	 * @return	true
 	 */
-	FORCEINLINE bool IsActorSuitable(AActor* Actor) const
+	FORCEINLINE bool IsActorSuitable(const AActor* Actor) const
 	{
 		if (EnumHasAnyFlags(Flags, EActorIteratorFlags::SkipPendingKill) && !IsValid(Actor))
 		{
@@ -420,7 +420,7 @@ protected:
 	 * @param Level the level to check for iteration
 	 * @return true if the level can be iterated, false otherwise
 	 */
-	FORCEINLINE bool CanIterateLevel(ULevel* Level) const
+	FORCEINLINE bool CanIterateLevel(const ULevel* Level) const
 	{
 		if (EnumHasAnyFlags(Flags, EActorIteratorFlags::OnlyActiveLevels))
 		{
@@ -464,9 +464,10 @@ public:
 	/**
 	 * Constructor
 	 *
-	 * @param  InWorld  The world whose actors are to be iterated over.
+	 * @param InWorld	The world whose actors are to be iterated over.
+	 * @param InFlags	Iteration flags indicating which types of levels and actors should be iterated
 	 */
-	explicit FActorIterator(UWorld* InWorld, EActorIteratorFlags InFlags = EActorIteratorFlags::OnlyActiveLevels | EActorIteratorFlags::SkipPendingKill)
+	explicit FActorIterator(const UWorld* InWorld, const EActorIteratorFlags InFlags = EActorIteratorFlags::OnlyActiveLevels | EActorIteratorFlags::SkipPendingKill)
 		: Super(InWorld, AActor::StaticClass(), InFlags)
 	{
 		++(*this);
@@ -475,10 +476,11 @@ public:
 	/**
 	 * Constructor
 	 *
-	 * @param  InWorld  The world whose actors are to be iterated over.
-	 * @param  InClass  The type of actors to be iterated over.
+	 * @param InWorld	The world whose actors are to be iterated over.
+	 * @param InClass	The type of actors to be iterated over.
+	 * @param InFlags	Iteration flags indicating which types of levels and actors should be iterated
 	 */
-	explicit FActorIterator(UWorld* InWorld, TSubclassOf<AActor> InClass, EActorIteratorFlags InFlags = EActorIteratorFlags::OnlyActiveLevels | EActorIteratorFlags::SkipPendingKill)
+	explicit FActorIterator(const UWorld* InWorld, const TSubclassOf<AActor> InClass, const EActorIteratorFlags InFlags = EActorIteratorFlags::OnlyActiveLevels | EActorIteratorFlags::SkipPendingKill)
 		: Super(InWorld, InClass, InFlags)
 	{
 		++(*this);
@@ -503,9 +505,10 @@ public:
 	/**
 	 * Constructor
 	 *
-	 * @param  InWorld  The world whose actors are to be iterated over.
+	 * @param InWorld	The world whose actors are to be iterated over.
+	 * @param InFlags	Iteration flags indicating which types of levels and actors should be iterated
 	 */
-	explicit FActorRange(UWorld* InWorld, EActorIteratorFlags InFlags = EActorIteratorFlags::OnlyActiveLevels | EActorIteratorFlags::SkipPendingKill)
+	explicit FActorRange(const UWorld* InWorld, const EActorIteratorFlags InFlags = EActorIteratorFlags::OnlyActiveLevels | EActorIteratorFlags::SkipPendingKill)
 		: Flags(InFlags)
 		, World(InWorld)
 	{
@@ -513,7 +516,7 @@ public:
 
 private:
 	EActorIteratorFlags	Flags;
-	UWorld*				World;
+	const UWorld* World;
 
 	friend FActorIterator begin(const FActorRange& Range) { return FActorIterator(Range.World, Range.Flags); }
 	friend FActorIterator end  (const FActorRange& Range) { return FActorIterator(EActorIteratorType::End); }
@@ -532,10 +535,11 @@ public:
 	/**
 	 * Constructor
 	 *
-	 * @param  InWorld  The world whose actors are to be iterated over.
-	 * @param  InClass  The subclass of actors to be iterated over.
+	 * @param InWorld	The world whose actors are to be iterated over.
+	 * @param InClass	The subclass of actors to be iterated over.
+	 * @param InFlags	Iteration flags indicating which types of levels and actors should be iterated
 	 */
-	explicit TActorIterator(UWorld* InWorld, TSubclassOf<ActorType> InClass = ActorType::StaticClass(), EActorIteratorFlags InFlags = EActorIteratorFlags::OnlyActiveLevels | EActorIteratorFlags::SkipPendingKill)
+	explicit TActorIterator(const UWorld* InWorld, TSubclassOf<ActorType> InClass = ActorType::StaticClass(), EActorIteratorFlags InFlags = EActorIteratorFlags::OnlyActiveLevels | EActorIteratorFlags::SkipPendingKill)
 		: Super(InWorld, InClass, InFlags)
 	{
 		++(*this);
@@ -580,10 +584,11 @@ public:
 	/**
 	 * Constructor
 	 *
-	 * @param  InWorld  The world whose actors are to be iterated over.
-	 * @param  InClass  The subclass of actors to be iterated over.
+	 * @param InWorld	The world whose actors are to be iterated over.
+	 * @param InClass	The subclass of actors to be iterated over.
+	 * @param InFlags	Iteration flags indicating which types of levels and actors should be iterated
 	 */
-	explicit TActorRange(UWorld* InWorld, TSubclassOf<ActorType> InClass = ActorType::StaticClass(), EActorIteratorFlags InFlags = EActorIteratorFlags::OnlyActiveLevels | EActorIteratorFlags::SkipPendingKill)
+	explicit TActorRange(const UWorld* InWorld, TSubclassOf<ActorType> InClass = ActorType::StaticClass(), const EActorIteratorFlags InFlags = EActorIteratorFlags::OnlyActiveLevels | EActorIteratorFlags::SkipPendingKill)
 		: Flags(InFlags)
 		, World(InWorld)
 		, Class(InClass)
@@ -592,7 +597,7 @@ public:
 
 private:
 	EActorIteratorFlags		Flags;
-	UWorld*					World;
+	const UWorld*			World;
 	TSubclassOf<ActorType>	Class;
 
 	friend TActorIterator<ActorType> begin(const TActorRange& Range) { return TActorIterator<ActorType>(Range.World, Range.Class, Range.Flags); }
@@ -611,9 +616,9 @@ public:
 	/**
 	 * Constructor
 	 *
-	 * @param  InWorld  The world whose actors are to be iterated over.
+	 * @param InWorld	The world whose actors are to be iterated over.
 	 */
-	explicit FSelectedActorIterator(UWorld* InWorld)
+	explicit FSelectedActorIterator(const UWorld* InWorld)
 		: Super(InWorld, AActor::StaticClass(), EActorIteratorFlags::SkipPendingKill | EActorIteratorFlags::OnlySelectedActors)
 	{
 		++(*this);
@@ -625,7 +630,7 @@ public:
 	 * @param  InWorld  The world whose actors are to be iterated over.
 	 * @param  InClass  The type of actors to be iterated over.
 	 */
-	explicit FSelectedActorIterator(UWorld* InWorld, TSubclassOf<AActor> InClass)
+	explicit FSelectedActorIterator(const UWorld* InWorld, const TSubclassOf<AActor> InClass)
 		: Super(InWorld, InClass, EActorIteratorFlags::SkipPendingKill | EActorIteratorFlags::OnlySelectedActors)
 	{
 		++(*this);
@@ -651,13 +656,13 @@ public:
 	 *
 	 * @param  InWorld  The world whose actors are to be iterated over.
 	 */
-	explicit FSelectedActorRange(UWorld* InWorld)
+	explicit FSelectedActorRange(const UWorld* InWorld)
 		: World(InWorld)
 	{
 	}
 
 private:
-	UWorld* World;
+	const UWorld* World;
 
 	friend FSelectedActorIterator begin(const FSelectedActorRange& Range) { return FSelectedActorIterator(Range.World); }
 	friend FSelectedActorIterator end  (const FSelectedActorRange& Range) { return FSelectedActorIterator(EActorIteratorType::End); }
@@ -692,12 +697,13 @@ private:
 /**
  *	Renders stats
  *
- *	@param Viewport	The viewport to render to
- *	@param Canvas	Canvas object to use for rendering
+ *	@param InWorld			The World to render stats
+ *	@param Viewport			The viewport to render to
+ *	@param Canvas			Canvas object to use for rendering
  *	@param CanvasObject		Optional canvas object for visualizing properties
  *	@param DebugProperties	List of properties to visualize (in/out)
- *	@param ViewLocation	Location of camera
- *	@param ViewRotation	Rotation of camera
+ *	@param ViewLocation		Location of camera
+ *	@param ViewRotation		Rotation of camera
  */
 ENGINE_API void DrawStatsHUD( UWorld* InWorld, FViewport* Viewport, FCanvas* Canvas, UCanvas* CanvasObject, TArray<struct FDebugDisplayProperty>& DebugProperties, const FVector& ViewLocation, const FRotator& ViewRotation );
 
@@ -735,7 +741,7 @@ struct FContentComparisonAssetInfo
 	}
 
 	/** operator == */
-	bool operator==(const FContentComparisonAssetInfo& Other)
+	bool operator==(const FContentComparisonAssetInfo& Other) const
 	{
 		return (
 			(AssetName == Other.AssetName) &&
@@ -783,10 +789,10 @@ public:
 	/**
 	 *	Recursive function for collecting objects referenced by the given object.
 	 *
-	 *	@param	InStartObject			The object to collect the referencees for
+	 *	@param	InStartObject			The object to collect the references for
 	 *	@param	InCurrDepth				The current depth being processed
 	 *	@param	InMaxDepth				The maximum depth to traverse the reference chain
-	 *	@param	OutCollectReferences	The resulting referenced object list
+	 *	@param	OutCollectedReferences	The resulting referenced object list
 	 */
 	void RecursiveObjectCollection(UObject* InStartObject, int32 InCurrDepth, int32 InMaxDepth, TMap<UObject*,bool>& OutCollectedReferences);
 
@@ -842,7 +848,7 @@ public:
 	 * @param InClassFlags - User defined per class flags .
 	 * @param InVersion - Minimal strip version required to serialize strip flags
 	 */
-	FStripDataFlags(FArchive& Ar, uint8 InClassFlags = 0, const FPackageFileVersion& InVersion = GOldestLoadablePackageFileUEVersion);
+	explicit FStripDataFlags(FArchive& Ar, uint8 InClassFlags = 0, const FPackageFileVersion& InVersion = GOldestLoadablePackageFileUEVersion);
 	UE_DEPRECATED(5.0, "Use the other overload that takes InVersion as a FPackageFileVersion. See the @FPackageFileVersion documentation for further details")
 	inline FStripDataFlags(FArchive& Ar, uint8 InClassFlags, int32 InVersion)
 		: FStripDataFlags(Ar, InClassFlags, FPackageFileVersion::CreateUE4Version(InVersion))
@@ -856,6 +862,7 @@ public:
 	 * when saving. Class flags also need to be defined by the user.
 	 *
 	 * @param Ar - Archive to serialize with.
+	 * @param InGlobalFlags Engine flags
 	 * @param InClassFlags - User defined per class flags.
 	 * @param InVersion - Minimal version required to serialize strip flags
 	 */
@@ -872,11 +879,11 @@ public:
 	* Serializes strip data flags. Global (engine) flags are automatically generated from target platform
 	* when saving. Class flags need to be defined by the user.
 	*
-	* @param Ar - Archive to serialize with.
+	* @param Slot - Slot holding the archive to serialize with.
 	* @param InClassFlags - User defined per class flags .
 	* @param InVersion - Minimal strip version required to serialize strip flags
 	*/
-	FStripDataFlags(FStructuredArchive::FSlot Slot, uint8 InClassFlags = 0, const FPackageFileVersion& InVersion = GOldestLoadablePackageFileUEVersion);
+	explicit FStripDataFlags(FStructuredArchive::FSlot Slot, uint8 InClassFlags = 0, const FPackageFileVersion& InVersion = GOldestLoadablePackageFileUEVersion);
 	UE_DEPRECATED(5.0, "Use the other overload that takes InVersion as a FPackageFileVersion. See the @FPackageFileVersion documentation for further details")
 	inline FStripDataFlags(FStructuredArchive::FSlot Slot, uint8 InClassFlags, int32 InVersion)
 		: FStripDataFlags(Slot, InClassFlags, FPackageFileVersion::CreateUE4Version(InVersion))
@@ -889,7 +896,8 @@ public:
 	* Serializes strip data flags. Global (engine) flags are user defined and will not be automatically generated
 	* when saving. Class flags also need to be defined by the user.
 	*
-	* @param Ar - Archive to serialize with.
+	* @param Slot - Slot holding the archive to serialize with.
+	* @param InGlobalFlags - Engine flags.
 	* @param InClassFlags - User defined per class flags.
 	* @param InVersion - Minimal version required to serialize strip flags
 	*/
@@ -936,6 +944,6 @@ public:
 class UTexture;
 namespace VirtualTextureUtils
 {
-	/** Function that will test if the passed in texture is VT. If so, print to messaglog that the property does not support VT textures */
+	/** Function that will test if the passed in texture is VT. If so, print to Messagelog that the property does not support VT textures */
 	ENGINE_API void CheckAndReportInvalidUsage(const UObject* Owner, const FName& PropertyName, const UTexture* Texture);
 }
\ No newline at end of file
