Explicitly SetWireframe in EngineShowFlagOverride so default previews of orthographic cameras display correctly.

diff --git a/Engine/Source/Runtime/Engine/Private/ShowFlags.cpp b/Engine/Source/Runtime/Engine/Private/ShowFlags.cpp
--- a/Engine/Source/Runtime/Engine/Private/ShowFlags.cpp
+++ b/Engine/Source/Runtime/Engine/Private/ShowFlags.cpp
@@ -443,6 +443,16 @@ void EngineShowFlagOverride(EShowFlagInitMode ShowFlagInitMode, EViewModeIndex V
 
 	// Some view modes want some features off or on (no state)
 	{
+		if (ViewModeIndex == VMI_BrushWireframe ||
+			ViewModeIndex == VMI_Wireframe)
+		{
+			EngineShowFlags.SetWireframe(true);
+		}
+		else
+		{
+			EngineShowFlags.SetWireframe(false);
+		}
+
 		if( ViewModeIndex == VMI_BrushWireframe ||
 			ViewModeIndex == VMI_Wireframe ||
 			ViewModeIndex == VMI_Unlit ||
