Change bNetStartUp to not be a UPROP.  This makes it so when creating an Actor using a template of another Actor that it will not copy over the value.  This change will ensure that new Actors created this way will not have this flag set.

diff --git a/Engine/Source/Runtime/Engine/Classes/GameFramework/Actor.h b/Engine/Source/Runtime/Engine/Classes/GameFramework/Actor.h
--- a/Engine/Source/Runtime/Engine/Classes/GameFramework/Actor.h
+++ b/Engine/Source/Runtime/Engine/Classes/GameFramework/Actor.h
@@ -191,7 +191,6 @@ public:
 	uint8 bNetTemporary:1;
 
 	/** If true, this actor was loaded directly from the map, and for networking purposes can be addressed by its full path name */
-	UPROPERTY()
 	uint8 bNetStartup:1;
 
 	/** If true, this actor is only relevant to its owner. If this flag is changed during play, all non-owner channels would need to be explicitly closed. */
