[Navmesh] PR #8081: Client side dynamic navmesh generation fix (Contributed by [REDACTED_PERSON])

diff --git a/Engine/Source/Runtime/Engine/Private/Components/SceneComponent.cpp b/Engine/Source/Runtime/Engine/Private/Components/SceneComponent.cpp
--- a/Engine/Source/Runtime/Engine/Private/Components/SceneComponent.cpp
+++ b/Engine/Source/Runtime/Engine/Private/Components/SceneComponent.cpp
@@ -3819,10 +3819,9 @@ void USceneComponent::UpdateNavigationData()
 
 	if (IsRegistered())
 	{
-		UWorld* MyWorld = GetWorld();
-		if ((MyWorld != nullptr) && (!MyWorld->IsGameWorld() || !MyWorld->IsNetMode(ENetMode::NM_Client)))
+		if (GetWorld() != nullptr)
 		{
-			// use propagated component's transform update in editor OR server game with additional navsys check
+			// use propagated component's transform update in editor OR game with additional navsys check
 			FNavigationSystem::UpdateComponentData(*this);
 		}
 	}
