Android: Rendering artifacts on some skeletal meshes with GPUSkin cache enabled

diff --git a/Engine/Source/Runtime/Engine/Private/SkeletalRenderGPUSkin.cpp b/Engine/Source/Runtime/Engine/Private/SkeletalRenderGPUSkin.cpp
--- a/Engine/Source/Runtime/Engine/Private/SkeletalRenderGPUSkin.cpp
+++ b/Engine/Source/Runtime/Engine/Private/SkeletalRenderGPUSkin.cpp
@@ -2106,6 +2106,15 @@ void FDynamicSkelMeshObjectDataGPUSkin::InitDynamicSkelMeshObjectDataGPUSkin(
 
 	bIsSkinCacheAllowed = InMeshComponent->IsSkinCacheAllowed(InLODIndex);
 
+	if (bIsSkinCacheAllowed && InMeshObject->FeatureLevel == ERHIFeatureLevel::ES3_1)
+	{
+		// Some mobile GPUs (MALI) has a 64K elements limitation on texel buffers
+		// SkinCache fetches mesh position through R32F texel buffer, thus any mesh that has more than 64K/3 vertices will not work correctly on such GPUs
+		// We force this limitation for all mobile, to have an uniform behaviour accross all mobile platforms
+		const FSkeletalMeshLODRenderData& LOD = InSkeletalMeshRenderData->LODRenderData[LODIndex];
+		bIsSkinCacheAllowed = (LOD.GetNumVertices()*3 < 0xffff);
+	}
+
 #if RHI_RAYTRACING
 	if (SkeletalMeshProxy != nullptr)
 	{
