Get non-unity fix to UE5/Main

diff --git a/Engine/Source/Runtime/Engine/Private/BlueprintGeneratedClass.cpp b/Engine/Source/Runtime/Engine/Private/BlueprintGeneratedClass.cpp
--- a/Engine/Source/Runtime/Engine/Private/BlueprintGeneratedClass.cpp
+++ b/Engine/Source/Runtime/Engine/Private/BlueprintGeneratedClass.cpp
@@ -23,6 +23,7 @@
 #include "UObject/CoreObjectVersion.h"
 #include "Net/Core/PushModel/PushModel.h"
 #include "UObject/CoreObjectVersion.h"
+#include "UObject/FortniteMainBranchObjectVersion.h"
 #include "UObject/UE5MainStreamObjectVersion.h"
 #include "GenericPlatform/GenericPlatformCrashContext.h"
 
