Adding an exec command to list subobjects of open actor channels "net dumpsubobjects"
Adding a stat scope for deleted subobject iteration

diff --git a/Engine/Source/Runtime/Engine/Classes/Engine/NetDriver.h b/Engine/Source/Runtime/Engine/Classes/Engine/NetDriver.h
--- a/Engine/Source/Runtime/Engine/Classes/Engine/NetDriver.h
+++ b/Engine/Source/Runtime/Engine/Classes/Engine/NetDriver.h
@@ -1394,13 +1394,14 @@ public:
 	/**
 	 * Exec command handlers
 	 */
-	bool HandleSocketsCommand( const TCHAR* Cmd, FOutputDevice& Ar );
-	bool HandlePackageMapCommand( const TCHAR* Cmd, FOutputDevice& Ar );
-	bool HandleNetFloodCommand( const TCHAR* Cmd, FOutputDevice& Ar );
-	bool HandleNetDebugTextCommand( const TCHAR* Cmd, FOutputDevice& Ar );
-	bool HandleNetDisconnectCommand( const TCHAR* Cmd, FOutputDevice& Ar );
-	bool HandleNetDumpServerRPCCommand( const TCHAR* Cmd, FOutputDevice& Ar );
-	bool HandleNetDumpDormancy( const TCHAR* Cmd, FOutputDevice& Ar );
+	bool HandleSocketsCommand(const TCHAR* Cmd, FOutputDevice& Ar);
+	bool HandlePackageMapCommand(const TCHAR* Cmd, FOutputDevice& Ar);
+	bool HandleNetFloodCommand(const TCHAR* Cmd, FOutputDevice& Ar);
+	bool HandleNetDebugTextCommand(const TCHAR* Cmd, FOutputDevice& Ar);
+	bool HandleNetDisconnectCommand(const TCHAR* Cmd, FOutputDevice& Ar);
+	bool HandleNetDumpServerRPCCommand(const TCHAR* Cmd, FOutputDevice& Ar);
+	bool HandleNetDumpDormancy(const TCHAR* Cmd, FOutputDevice& Ar);
+	bool HandleDumpSubObjectsCommand(const TCHAR* Cmd, FOutputDevice& Ar);
 #endif
 
 	void HandlePacketLossBurstCommand( int32 DurationInMilliseconds );
diff --git a/Engine/Source/Runtime/Engine/Private/DataChannel.cpp b/Engine/Source/Runtime/Engine/Private/DataChannel.cpp
--- a/Engine/Source/Runtime/Engine/Private/DataChannel.cpp
+++ b/Engine/Source/Runtime/Engine/Private/DataChannel.cpp
@@ -3057,7 +3057,7 @@ int64 UActorChannel::ReplicateActor()
 		UE_LOG(LogNet, Warning, TEXT("ReplicateActor: PackageMap->GetMustBeMappedGuidsInLastBunch().Num() != 0: %i"), PackageMapClient->GetMustBeMappedGuidsInLastBunch().Num());
 	}
 
-	bool WroteSomethingImportant = bIsNewlyReplicationUnpaused || bIsNewlyReplicationPaused;
+	bool bWroteSomethingImportant = bIsNewlyReplicationUnpaused || bIsNewlyReplicationPaused;
 
 	// Create an outgoing bunch, and skip this actor if the channel is saturated.
 	FOutBunch Bunch( this, 0 );
@@ -3134,7 +3134,7 @@ int64 UActorChannel::ReplicateActor()
 		UE_NET_TRACE_SCOPE(NewActor, Bunch, GetTraceCollector(Bunch), ENetTraceVerbosity::Trace);
 
 		Connection->PackageMap->SerializeNewActor(Bunch, this, static_cast<AActor*&>(Actor));
-		WroteSomethingImportant = true;
+		bWroteSomethingImportant = true;
 
 		Actor->OnSerializeNewActor(Bunch);
 	}
@@ -3164,55 +3164,59 @@ int64 UActorChannel::ReplicateActor()
 		// The Actor
 		{
 			UE_NET_TRACE_OBJECT_SCOPE(ActorReplicator->ObjectNetGUID, Bunch, GetTraceCollector(Bunch), ENetTraceVerbosity::Trace);
-			WroteSomethingImportant |= ActorReplicator->ReplicateProperties(Bunch, RepFlags);
+			bWroteSomethingImportant |= ActorReplicator->ReplicateProperties(Bunch, RepFlags);
 		}
 
 		// The SubObjects
-		WroteSomethingImportant |= Actor->ReplicateSubobjects(this, &Bunch, &RepFlags);
+		bWroteSomethingImportant |= Actor->ReplicateSubobjects(this, &Bunch, &RepFlags);
 
 		if (Connection->ResendAllDataState != EResendAllDataState::None)
 		{
-			if (WroteSomethingImportant)
+			if (bWroteSomethingImportant)
 			{
 				SendBunch(&Bunch, 1);
 			}
 
 			MemMark.Pop();
 
-			return WroteSomethingImportant;
+			return bWroteSomethingImportant;
 		}
 
-		// Look for deleted subobjects
-		FObjectReplicator* LocalActorReplicator = ActorReplicator.Get();
-		for (auto RepComp = ReplicationMap.CreateIterator(); RepComp; ++RepComp)
 		{
-			TSharedRef<FObjectReplicator>& LocalReplicator = RepComp.Value();
+			QUICK_SCOPE_CYCLE_COUNTER(STAT_NetDeletedSubObjects);
 
-			if (!LocalReplicator->GetWeakObjectPtr().IsValid())
+			// Look for deleted subobjects
+			FObjectReplicator* LocalActorReplicator = ActorReplicator.Get();
+			for (auto RepComp = ReplicationMap.CreateIterator(); RepComp; ++RepComp)
 			{
-				if (LocalReplicator->ObjectNetGUID.IsValid())
-				{
-					// Write a deletion content header:
-					WriteContentBlockForSubObjectDelete(Bunch, LocalReplicator->ObjectNetGUID);
+				TSharedRef<FObjectReplicator>& LocalReplicator = RepComp.Value();
 
-					WroteSomethingImportant = true;
-					Bunch.bReliable = true;
-				}
-				else
+				if (!LocalReplicator->GetWeakObjectPtr().IsValid())
 				{
-					UE_LOG(LogNetTraffic, Error, TEXT("Unable to write subobject delete for (%s), object replicator has invalid NetGUID"), *GetPathNameSafe(Actor));
-				}
+					if (LocalReplicator->ObjectNetGUID.IsValid())
+					{
+						// Write a deletion content header:
+						WriteContentBlockForSubObjectDelete(Bunch, LocalReplicator->ObjectNetGUID);
 
-				// The only way this case would be possible is if someone tried destroying the Actor as a part of
-				// a Subobject's Pre / Post replication, during Replicate Subobjects, or OnSerializeNewActor.
-				// All of those are bad.
-				if (!ensureMsgf(LocalActorReplicator != &LocalReplicator.Get(), TEXT("UActorChannel::ReplicateActor: Actor was deleting during replication: %s"), *Describe()))
-				{
-					ActorReplicator.Reset();
-				}
+						bWroteSomethingImportant = true;
+						Bunch.bReliable = true;
+					}
+					else
+					{
+						UE_LOG(LogNetTraffic, Error, TEXT("Unable to write subobject delete for (%s), object replicator has invalid NetGUID"), *GetPathNameSafe(Actor));
+					}
 
-				LocalReplicator->CleanUp();
-				RepComp.RemoveCurrent();
+					// The only way this case would be possible is if someone tried destroying the Actor as a part of
+					// a Subobject's Pre / Post replication, during Replicate Subobjects, or OnSerializeNewActor.
+					// All of those are bad.
+					if (!ensureMsgf(LocalActorReplicator != &LocalReplicator.Get(), TEXT("UActorChannel::ReplicateActor: Actor was deleting during replication: %s"), *Describe()))
+					{
+						ActorReplicator.Reset();
+					}
+
+					LocalReplicator->CleanUp();
+					RepComp.RemoveCurrent();
+				}
 			}
 		}
 	}
@@ -3224,7 +3228,7 @@ int64 UActorChannel::ReplicateActor()
 	// -----------------------------
 
 	int64 NumBitsWrote = 0;
-	if (WroteSomethingImportant)
+	if (bWroteSomethingImportant)
 	{
 		FPacketIdRange PacketRange = SendBunch( &Bunch, 1 );
 
@@ -4216,16 +4220,16 @@ bool UActorChannel::ReplicateSubobject(UObject *Obj, FOutBunch &Bunch, const FRe
 		NewSubobject = true;
 	}
 	UE_NET_TRACE_OBJECT_SCOPE(ObjectReplicator->ObjectNetGUID, Bunch, GetTraceCollector(Bunch), ENetTraceVerbosity::Trace);
-	bool WroteSomething = ObjectReplicator.Get().ReplicateProperties(Bunch, RepFlags);
-	if (NewSubobject && !WroteSomething)
+	bool bWroteSomething = ObjectReplicator.Get().ReplicateProperties(Bunch, RepFlags);
+	if (NewSubobject && !bWroteSomething)
 	{
 		// Write empty payload to force object creation
 		FNetBitWriter EmptyPayload;
 		WriteContentBlockPayload( Obj, Bunch, false, EmptyPayload );
-		WroteSomething= true;
+		bWroteSomething = true;
 	}
 
-	return WroteSomething;
+	return bWroteSomething;
 }
 
 //------------------------------------------------------
diff --git a/Engine/Source/Runtime/Engine/Private/NetDriver.cpp b/Engine/Source/Runtime/Engine/Private/NetDriver.cpp
--- a/Engine/Source/Runtime/Engine/Private/NetDriver.cpp
+++ b/Engine/Source/Runtime/Engine/Private/NetDriver.cpp
@@ -2982,6 +2982,40 @@ bool UNetDriver::HandleNetDumpDormancy(const TCHAR* Cmd, FOutputDevice& Ar)
 	Ar.Logf(TEXT(""));
 	return true;
 }
+
+bool UNetDriver::HandleDumpSubObjectsCommand(const TCHAR* Cmd, FOutputDevice& Ar)
+{
+#if WITH_SERVER_CODE
+	for(UNetConnection* ClientConnection : ClientConnections)
+	{
+		Ar.Logf(TEXT("Logging subobjects for %d open actor channels on %s"), ClientConnection->OpenChannels.Num(), *GetNameSafe(ClientConnection));
+		Ar.Logf(TEXT(""));
+
+		for (const UChannel* Channel : ClientConnection->OpenChannels)
+		{
+			if (const UActorChannel* ActorChannel = Cast<UActorChannel>(Channel))
+			{
+				if (ActorChannel->ReplicationMap.Num() > 1)
+				{
+					Ar.Logf(TEXT("   Actor: %s"), *GetFullNameSafe(ActorChannel->Actor));
+
+					for (auto RepComp = ActorChannel->ReplicationMap.CreateConstIterator(); RepComp; ++RepComp)
+					{
+						const TSharedRef<FObjectReplicator>& LocalReplicator = RepComp.Value();
+						const UObject* Obj = LocalReplicator->GetWeakObjectPtr().Get();
+
+						if (Obj != ActorChannel->Actor)
+						{
+							Ar.Logf(TEXT("       Object: %s Class: %s"), *GetNameSafe(Obj), *GetNameSafe(LocalReplicator->ObjectClass));
+						}
+					}
+				}
+			}
+		}
+	}
+#endif
+	return true;
+}
 #endif // !UE_BUILD_SHIPPING
 
 void UNetDriver::HandlePacketLossBurstCommand( int32 DurationInMilliseconds )
@@ -3097,6 +3131,11 @@ bool UNetDriver::Exec( UWorld* InWorld, const TCHAR* Cmd, FOutputDevice& Ar )
 		HandleNetFlushAllDormancy(this, InWorld);
 		return true;
 	}
+	else if (FParse::Command(&Cmd, TEXT("DUMPSUBOBJECTS")))
+	{
+		HandleDumpSubObjectsCommand(Cmd, Ar);
+		return true;
+	}
 	else
 #endif // !UE_BUILD_SHIPPING
 	{
