Removing duplicate console CSV output. It seems that PIE console and in-game console use slightly different paths, but the ConsoleManager path that I recently added the CSV log to is shared by both, so this change removes it from the in-game specific one.

diff --git a/Engine/Source/Runtime/Engine/Private/UserInterface/Console.cpp b/Engine/Source/Runtime/Engine/Private/UserInterface/Console.cpp
--- a/Engine/Source/Runtime/Engine/Private/UserInterface/Console.cpp
+++ b/Engine/Source/Runtime/Engine/Private/UserInterface/Console.cpp
@@ -601,7 +601,6 @@ void UConsole::SetCursorPos(int32 Position)
 
 void UConsole::ConsoleCommand(const FString& Command)
 {
-	CSV_EVENT_GLOBAL(TEXT("Cmd: %s"), *Command);
 	// insert into history buffer
 	{
 		HistoryBuffer.Remove(Command);
