Fix crash when creating zero vertice sections.
Set a minimum value of 12 (MAX_TOTAL_INFLUENCES) for the renderer settings "MaxSkinBones". Setting this value under the influence skinning count make no sense.

diff --git a/Engine/Source/Runtime/Engine/Classes/Engine/RendererSettings.h b/Engine/Source/Runtime/Engine/Classes/Engine/RendererSettings.h
--- a/Engine/Source/Runtime/Engine/Classes/Engine/RendererSettings.h
+++ b/Engine/Source/Runtime/Engine/Classes/Engine/RendererSettings.h
@@ -889,10 +889,14 @@ class ENGINE_API URendererSettings : public UDeveloperSettings
 		ToolTip = "When Unlimited Bone Influence is enabled, it still uses a fixed bone inflence buffer until the max bone influence of a mesh exceeds this value"))
 		int32 UnlimitedBonInfluencesThreshold;
 
+	/*
+	 * The maximum bones count section vertices's skinning can use before being chunked into more sections. The minimum value is the maximum total influences define (MAX_TOTAL_INFLUENCES).
+	 */
 	UPROPERTY(config, EditAnywhere, Category = Skinning, meta = (
 		DisplayName = "Maximum bones per Sections",
 		ToolTip = "Max number of bones that can be skinned on the GPU in a single draw call. The default value is set by the Compat.MAX_GPUSKIN_BONES consolevariable. Changing this setting requires restarting the editor.",
-		ConfigRestartRequired = true))
+		ConfigRestartRequired = true,
+		ClampMin = "12", UIMin = "12"))
 		FPerPlatformInt MaxSkinBones;
 	
 	UPROPERTY(config, EditAnywhere, Category = Mobile, meta = (
diff --git a/Engine/Source/Runtime/Engine/Private/GPUSkinVertexFactory.cpp b/Engine/Source/Runtime/Engine/Private/GPUSkinVertexFactory.cpp
--- a/Engine/Source/Runtime/Engine/Private/GPUSkinVertexFactory.cpp
+++ b/Engine/Source/Runtime/Engine/Private/GPUSkinVertexFactory.cpp
@@ -412,6 +412,10 @@ int32 FGPUBaseSkinVertexFactory::GetMaxGPUSkinBones(const ITargetPlatform* Targe
 		MaxGPUSkinBones = FMath::Min(MaxGPUSkinBones, GCVarMaxGPUSkinBones);
 		
 	}
+
+	//We cannot go under MAX_TOTAL_INFLUENCES
+	MaxGPUSkinBones = FMath::Max(MaxGPUSkinBones, MAX_TOTAL_INFLUENCES);
+
 	if (GCVarSupport16BitBoneIndex > 0)
 	{
 		// 16-bit bone index is supported
diff --git a/Engine/Source/Runtime/Engine/Public/Rendering/SkeletalMeshDuplicatedVerticesBuffer.h b/Engine/Source/Runtime/Engine/Public/Rendering/SkeletalMeshDuplicatedVerticesBuffer.h
--- a/Engine/Source/Runtime/Engine/Public/Rendering/SkeletalMeshDuplicatedVerticesBuffer.h
+++ b/Engine/Source/Runtime/Engine/Public/Rendering/SkeletalMeshDuplicatedVerticesBuffer.h
@@ -53,7 +53,7 @@ public:
         }
 
         DupVertData.ResizeBuffer(NumDups ? NumDups : 1);
-        DupVertIndexData.ResizeBuffer(NumVertices);
+        DupVertIndexData.ResizeBuffer(NumVertices ? NumVertices : 1);
 
         uint8* VertData = DupVertData.GetDataPointer();
         uint32 VertStride = DupVertData.GetStride();
