Fixed typos related to the usage of NAME_None.
Was reported on UDN and validated by the authors.

diff --git a/Engine/Source/Editor/UnrealEd/Private/UnrealEdEngine.cpp b/Engine/Source/Editor/UnrealEd/Private/UnrealEdEngine.cpp
--- a/Engine/Source/Editor/UnrealEd/Private/UnrealEdEngine.cpp
+++ b/Engine/Source/Editor/UnrealEd/Private/UnrealEdEngine.cpp
@@ -1339,7 +1339,7 @@ void UUnrealEdEngine::FixAnyInvertedBrushes(UWorld* World)
 
 void UUnrealEdEngine::RegisterComponentVisualizer(FName ComponentClassName, TSharedPtr<FComponentVisualizer> Visualizer)
 {
-	if( ComponentClassName != NAME_Name )
+	if( !ComponentClassName.IsNone() )
 	{
 		ComponentVisualizerMap.Add(ComponentClassName, Visualizer);		
 	}
diff --git a/Engine/Source/Runtime/AnimGraphRuntime/Private/AnimNodes/AnimNode_PoseSnapshot.cpp b/Engine/Source/Runtime/AnimGraphRuntime/Private/AnimNodes/AnimNode_PoseSnapshot.cpp
--- a/Engine/Source/Runtime/AnimGraphRuntime/Private/AnimNodes/AnimNode_PoseSnapshot.cpp
+++ b/Engine/Source/Runtime/AnimGraphRuntime/Private/AnimNodes/AnimNode_PoseSnapshot.cpp
@@ -43,7 +43,7 @@ void FAnimNode_PoseSnapshot::PreUpdate(const UAnimInstance* InAnimInstance)
 	}
 	else
 	{
-		TargetBoneNameMesh = NAME_Name;
+		TargetBoneNameMesh = NAME_None;
 	}
 }
 
diff --git a/Engine/Source/Runtime/Core/Private/Serialization/CustomVersion.cpp b/Engine/Source/Runtime/Core/Private/Serialization/CustomVersion.cpp
--- a/Engine/Source/Runtime/Core/Private/Serialization/CustomVersion.cpp
+++ b/Engine/Source/Runtime/Core/Private/Serialization/CustomVersion.cpp
@@ -363,7 +363,7 @@ const FCustomVersion* FCustomVersionContainer::GetVersion(FGuid Key) const
 
 const FName FCustomVersionContainer::GetFriendlyName(FGuid Key) const
 {
-	FName FriendlyName = NAME_Name;
+	FName FriendlyName = NAME_None;
 	const FCustomVersion* CustomVersion = GetVersion(Key);
 	if (CustomVersion)
 	{
diff --git a/Engine/Source/Runtime/Engine/Private/MemberReference.cpp b/Engine/Source/Runtime/Engine/Private/MemberReference.cpp
--- a/Engine/Source/Runtime/Engine/Private/MemberReference.cpp
+++ b/Engine/Source/Runtime/Engine/Private/MemberReference.cpp
@@ -308,7 +308,7 @@ TFieldType* FindRemappedFieldImpl(FName FieldClassOutermostName, FName FieldClas
 				// Try remapping class, this only works if class is in memory
 				FString ClassName = NewRedirectName.OuterName.ToString();
 
-				if (NewRedirectName.PackageName != NAME_Name)
+				if ( !NewRedirectName.PackageName.IsNone() )
 				{
 					// Use package if it's there
 					ClassName = FString::Printf(TEXT("%s.%s"), *NewRedirectName.PackageName.ToString(), *NewRedirectName.OuterName.ToString());
