PR #7977: Remove BlueprintAuthority from MinimumAreaRectangle (Contributed by [REDACTED_PERSON])

diff --git a/Engine/Source/Runtime/Engine/Classes/Kismet/KismetMathLibrary.h b/Engine/Source/Runtime/Engine/Classes/Kismet/KismetMathLibrary.h
--- a/Engine/Source/Runtime/Engine/Classes/Kismet/KismetMathLibrary.h
+++ b/Engine/Source/Runtime/Engine/Classes/Kismet/KismetMathLibrary.h
@@ -3987,7 +3987,7 @@ class ENGINE_API UKismetMathLibrary : public UBlueprintFunctionLibrary
 	 * @outparam	OutRectSideA - Vector oriented and sized to represent one edge of the enclosing rectangle, orthogonal to OutRectSideB
 	 * @outparam	OutRectSideB - Vector oriented and sized to represent one edge of the enclosing rectangle, orthogonal to OutRectSideA
 	*/
-	UFUNCTION(BlueprintCallable, BlueprintAuthorityOnly, Category="Math|Geometry", meta=(WorldContext="WorldContextObject", CallableWithoutWorldContext))
+	UFUNCTION(BlueprintCallable, Category="Math|Geometry", meta=(WorldContext="WorldContextObject", CallableWithoutWorldContext))
 	static void MinimumAreaRectangle(UObject* WorldContextObject, const TArray<FVector>& InVerts, const FVector& SampleSurfaceNormal, FVector& OutRectCenter, FRotator& OutRectRotation, float& OutSideLengthX, float& OutSideLengthY, bool bDebugDraw = false);
 
 	/**
