Fix LongPress triggering release events.

diff --git a/Engine/Source/Runtime/Engine/Private/Slate/SceneViewport.cpp b/Engine/Source/Runtime/Engine/Private/Slate/SceneViewport.cpp
--- a/Engine/Source/Runtime/Engine/Private/Slate/SceneViewport.cpp
+++ b/Engine/Source/Runtime/Engine/Private/Slate/SceneViewport.cpp
@@ -901,7 +901,10 @@ FReply FSceneViewport::OnTouchGesture( const FGeometry& MyGeometry, const FPoint
 		// Switch to the viewport clients world before processing input
 		FScopedConditionalWorldSwitcher WorldSwitcher( ViewportClient );
 
-		FSlateApplication::Get().SetKeyboardFocus(ViewportWidget.Pin());
+		if (GestureEvent.GetGestureType() != EGestureEvent::LongPress)
+		{
+			FSlateApplication::Get().SetKeyboardFocus(ViewportWidget.Pin());
+		}
 
 		if( !ViewportClient->InputGesture( this, GestureEvent.GetGestureType(), GestureEvent.GetGestureDelta(), GestureEvent.IsDirectionInvertedFromDevice() ) )
 		{
