PR #7971: Expose APlayerState::IsOnlyASpectator to blueprints (Contributed by [REDACTED_PERSON])

diff --git a/Engine/Source/Runtime/Engine/Classes/GameFramework/PlayerState.h b/Engine/Source/Runtime/Engine/Classes/GameFramework/PlayerState.h
--- a/Engine/Source/Runtime/Engine/Classes/GameFramework/PlayerState.h
+++ b/Engine/Source/Runtime/Engine/Classes/GameFramework/PlayerState.h
@@ -371,6 +371,7 @@ public:
 	void SetIsSpectator(const bool bNewSpectator);
 
 	/** Gets the literal value of bOnlySpectator. */
+	UFUNCTION(BlueprintCallable, Category = "PlayerState")
 	bool IsOnlyASpectator() const
 	{
 		return bOnlySpectator;
@@ -438,4 +439,4 @@ private:
 	{
 		PlayerState->PawnPrivate = Pawn;
 	}
-};
\ No newline at end of file
+};
