PR #8016: UE-39873: Add CallableWithoutWorldContext to BP library functions (Contributed by [REDACTED_PERSON])

diff --git a/Engine/Source/Runtime/Engine/Classes/Kismet/GameplayStatics.h b/Engine/Source/Runtime/Engine/Classes/Kismet/GameplayStatics.h
--- a/Engine/Source/Runtime/Engine/Classes/Kismet/GameplayStatics.h
+++ b/Engine/Source/Runtime/Engine/Classes/Kismet/GameplayStatics.h
@@ -1103,8 +1103,8 @@ public:
 	static float GetAudioTimeSeconds(const UObject* WorldContextObject);
 
 	/** Returns time in seconds since the application was started. Unlike the other time functions this is accurate to the exact time this function is called instead of set once per frame. */
-	UFUNCTION(BlueprintPure, Category="Utilities|Time", meta=(WorldContext="WorldContextObject"))
-	static void GetAccurateRealTime(const UObject* WorldContextObject, int32& Seconds, float& PartialSeconds);
+	UFUNCTION(BlueprintPure, Category="Utilities|Time")
+	static void GetAccurateRealTime(int32& Seconds, float& PartialSeconds);
 
 	/*~ DVRStreaming API */
 	
diff --git a/Engine/Source/Runtime/Engine/Classes/Kismet/KismetSystemLibrary.h b/Engine/Source/Runtime/Engine/Classes/Kismet/KismetSystemLibrary.h
--- a/Engine/Source/Runtime/Engine/Classes/Kismet/KismetSystemLibrary.h
+++ b/Engine/Source/Runtime/Engine/Classes/Kismet/KismetSystemLibrary.h
@@ -427,7 +427,7 @@ class ENGINE_API UKismetSystemLibrary : public UBlueprintFunctionLibrary
 	 * @param	Command			Command to send to the console
 	 * @param	SpecificPlayer	If specified, the console command will be routed through the specified player
 	 */
-	UFUNCTION(BlueprintCallable, Category="Development",meta=(WorldContext="WorldContextObject"))
+	UFUNCTION(BlueprintCallable, Category="Development",meta=(WorldContext="WorldContextObject", CallableWithoutWorldContext))
 	static void ExecuteConsoleCommand(const UObject* WorldContextObject, const FString& Command, class APlayerController* SpecificPlayer = NULL );
 
 	/**
@@ -463,7 +463,7 @@ class ENGINE_API UKismetSystemLibrary : public UBlueprintFunctionLibrary
 	 * @param	QuitPreference	Form of quitting.
 	 * @param	bIgnorePlatformRestrictions	Ignores and best-practices based on platform (e.g PS4 games should never quit). Non-shipping only
 	 */
-	UFUNCTION(BlueprintCallable, Category="Game",meta=(WorldContext="WorldContextObject"))
+	UFUNCTION(BlueprintCallable, Category="Game",meta=(WorldContext="WorldContextObject", CallableWithoutWorldContext))
 	static void QuitGame(const UObject* WorldContextObject, class APlayerController* SpecificPlayer, TEnumAsByte<EQuitPreference::Type> QuitPreference, bool bIgnorePlatformRestrictions);
 	
 #if WITH_EDITOR
diff --git a/Engine/Source/Runtime/Engine/Private/GameplayStatics.cpp b/Engine/Source/Runtime/Engine/Private/GameplayStatics.cpp
--- a/Engine/Source/Runtime/Engine/Private/GameplayStatics.cpp
+++ b/Engine/Source/Runtime/Engine/Private/GameplayStatics.cpp
@@ -2286,7 +2286,7 @@ float UGameplayStatics::GetAudioTimeSeconds(const UObject* WorldContextObject)
 	return World ? World->GetAudioTimeSeconds() : 0.f;
 }
 
-void UGameplayStatics::GetAccurateRealTime(const UObject* WorldContextObject, int32& Seconds, float& PartialSeconds)
+void UGameplayStatics::GetAccurateRealTime(int32& Seconds, float& PartialSeconds)
 {
 	double TimeSeconds = FPlatformTime::Seconds() - GStartTime;
 	Seconds = floor(TimeSeconds);
