Cloth - Add a ClearColllisionSources method so that all collision sources can be removed without knowledge of any of the original sources.

diff --git a/Engine/Source/Runtime/Engine/Classes/Components/SkeletalMeshComponent.h b/Engine/Source/Runtime/Engine/Classes/Components/SkeletalMeshComponent.h
--- a/Engine/Source/Runtime/Engine/Classes/Components/SkeletalMeshComponent.h
+++ b/Engine/Source/Runtime/Engine/Classes/Components/SkeletalMeshComponent.h
@@ -2175,6 +2175,9 @@ public:
 	/** Remove a cloth collision source defined by both a component and a physics asset */
 	void RemoveClothCollisionSource(USkeletalMeshComponent* InSourceComponent, UPhysicsAsset* InSourcePhysicsAsset);
 
+	/** Remove all cloth collision sources */
+	void ResetClothCollisionSources();
+
 protected:
 	/** copy cloth collision sources to this, where parent means components above it in the hierarchy */
 	void CopyClothCollisionSources();
@@ -2197,6 +2200,7 @@ public:
 	void AddClothCollisionSource(USkeletalMeshComponent* InSourceComponent, UPhysicsAsset* InSourcePhysicsAsset) {}
 	void RemoveClothCollisionSource(USkeletalMeshComponent* InSourceComponent) {}
 	void RemoveClothCollisionSource(USkeletalMeshComponent* InSourceComponent, UPhysicsAsset* InSourcePhysicsAsset) {}
+	void ResetClothCollisionSources() {}
 
 #endif
 
diff --git a/Engine/Source/Runtime/Engine/Private/SkeletalMeshComponentPhysics.cpp b/Engine/Source/Runtime/Engine/Private/SkeletalMeshComponentPhysics.cpp
--- a/Engine/Source/Runtime/Engine/Private/SkeletalMeshComponentPhysics.cpp
+++ b/Engine/Source/Runtime/Engine/Private/SkeletalMeshComponentPhysics.cpp
@@ -3221,6 +3221,11 @@ void USkeletalMeshComponent::RemoveClothCollisionSource(USkeletalMeshComponent*
 	}
 }
 
+void USkeletalMeshComponent::ResetClothCollisionSources()
+{
+	ClothCollisionSources.Reset();
+}
+
 #endif// #if WITH_CLOTH_COLLISION_DETECTION
 
 void USkeletalMeshComponent::EndPhysicsTickComponent(FSkeletalMeshComponentEndPhysicsTickFunction& ThisTickFunction)
