[CIS] Fixing ensure hit during S3 DDC update

diff --git a/Engine/Source/Runtime/Engine/Private/AudioDevice.cpp b/Engine/Source/Runtime/Engine/Private/AudioDevice.cpp
--- a/Engine/Source/Runtime/Engine/Private/AudioDevice.cpp
+++ b/Engine/Source/Runtime/Engine/Private/AudioDevice.cpp
@@ -5871,9 +5871,13 @@ void FAudioDevice::Flush(UWorld* WorldToFlush, bool bClearActivatedReverb)
 		}
 	}
 
-	for (AudioDeviceUtils::FVirtualLoopPair& Pair : VirtualLoops)
+	// We use a copy as some operations may modify VirtualLoops
 	{
-		AddSoundToStop(Pair.Key);
+		TMap<FActiveSound*, FAudioVirtualLoop> VirtualLoopsCopy = VirtualLoops;
+		for (AudioDeviceUtils::FVirtualLoopPair& Pair : VirtualLoopsCopy)
+		{
+			AddSoundToStop(Pair.Key);
+		}
 	}
 
 	// Immediately stop all pending active sounds
