Fix for unitialized structures in Material.

diff --git a/Engine/Source/Runtime/Engine/Public/MaterialExpressionIO.h b/Engine/Source/Runtime/Engine/Public/MaterialExpressionIO.h
--- a/Engine/Source/Runtime/Engine/Public/MaterialExpressionIO.h
+++ b/Engine/Source/Runtime/Engine/Public/MaterialExpressionIO.h
@@ -188,6 +188,12 @@ template<class InputType> struct FMaterialInput : FExpressionInput
 #if WITH_EDITORONLY_DATA
 	uint32	UseConstant : 1;
 	InputType	Constant;
+
+	FMaterialInput():
+		UseConstant(0),
+		Constant()
+	{
+	}
 #endif
 };
 
