Fix async character movement state management so we have GT state separated from simulation state.
Fix issues with marshalling jump input state and movement mode.
Move Pawn outputs to their own struct, stub out derived movement component input/output types, register derived callback.
Fixes some animation glitches in async mode.

diff --git a/Engine/Source/Runtime/Engine/Classes/GameFramework/CharacterMovementComponent.h b/Engine/Source/Runtime/Engine/Classes/GameFramework/CharacterMovementComponent.h
--- a/Engine/Source/Runtime/Engine/Classes/GameFramework/CharacterMovementComponent.h
+++ b/Engine/Source/Runtime/Engine/Classes/GameFramework/CharacterMovementComponent.h
@@ -19,6 +19,7 @@
 #include "GameFramework/PawnMovementComponent.h"
 #include "GameFramework/CharacterMovementReplication.h"
 #include "Interfaces/NetworkPredictionInterface.h"
+#include "CharacterMovementComponentAsync.h"
 #include "CharacterMovementComponent.generated.h"
 
 class ACharacter;
@@ -28,18 +29,6 @@ class FSavedMove_Character;
 class UPrimitiveComponent;
 class INavigationData;
 class UCharacterMovementComponent;
-class FCharacterMovementAsyncCallback;
-struct FCharacterMovementAsyncInput;
-struct FCharacterMovementAsyncOutput;
-
-// Enum used to control GetPawnCapsuleExtent behavior
-enum EShrinkCapsuleExtent
-{
-	SHRINK_None,			// Don't change the size of the capsule
-	SHRINK_RadiusCustom,	// Change only the radius, based on a supplied param
-	SHRINK_HeightCustom,	// Change only the height, based on a supplied param
-	SHRINK_AllCustom,		// Change both radius and height, based on a supplied param
-};
 
 DECLARE_DELEGATE_RetVal_ThreeParams(FTransform, FOnProcessRootMotion, const FTransform&, UCharacterMovementComponent*, float)
 
@@ -56,92 +45,6 @@ namespace CharacterMovementCVars
 	extern int32 ForceJumpPeakSubstep;
 }
 
-/** Data about the floor for walking movement, used by CharacterMovementComponent. */
-USTRUCT(BlueprintType)
-struct ENGINE_API FFindFloorResult
-{
-	GENERATED_USTRUCT_BODY()
-
-	/**
-	* True if there was a blocking hit in the floor test that was NOT in initial penetration.
-	* The HitResult can give more info about other circumstances.
-	*/
-	UPROPERTY(VisibleInstanceOnly, BlueprintReadOnly, Category=CharacterFloor)
-	uint32 bBlockingHit:1;
-
-	/** True if the hit found a valid walkable floor. */
-	UPROPERTY(VisibleInstanceOnly, BlueprintReadOnly, Category=CharacterFloor)
-	uint32 bWalkableFloor:1;
-
-	/** True if the hit found a valid walkable floor using a line trace (rather than a sweep test, which happens when the sweep test fails to yield a walkable surface). */
-	UPROPERTY(VisibleInstanceOnly, BlueprintReadOnly, Category=CharacterFloor)
-	uint32 bLineTrace:1;
-
-	/** The distance to the floor, computed from the swept capsule trace. */
-	UPROPERTY(VisibleInstanceOnly, BlueprintReadOnly, Category=CharacterFloor)
-	float FloorDist;
-	
-	/** The distance to the floor, computed from the trace. Only valid if bLineTrace is true. */
-	UPROPERTY(VisibleInstanceOnly, BlueprintReadOnly, Category=CharacterFloor)
-	float LineDist;
-
-	/** Hit result of the test that found a floor. Includes more specific data about the point of impact and surface normal at that point. */
-	UPROPERTY(VisibleInstanceOnly, BlueprintReadOnly, Category=CharacterFloor)
-	FHitResult HitResult;
-
-public:
-
-	FFindFloorResult()
-		: bBlockingHit(false)
-		, bWalkableFloor(false)
-		, bLineTrace(false)
-		, FloorDist(0.f)
-		, LineDist(0.f)
-		, HitResult(1.f)
-	{
-	}
-
-	/** Returns true if the floor result hit a walkable surface. */
-	bool IsWalkableFloor() const
-	{
-		return bBlockingHit && bWalkableFloor;
-	}
-
-	void Clear()
-	{
-		bBlockingHit = false;
-		bWalkableFloor = false;
-		bLineTrace = false;
-		FloorDist = 0.f;
-		LineDist = 0.f;
-		HitResult.Reset(1.f, false);
-	}
-
-	/** Gets the distance to floor, either LineDist or FloorDist. */
-	float GetDistanceToFloor() const
-	{
-		// When the floor distance is set using SetFromSweep, the LineDist value will be reset.
-		// However, when SetLineFromTrace is used, there's no guarantee that FloorDist is set.
-		return bLineTrace ? LineDist : FloorDist;
-	}
-
-	void SetFromSweep(const FHitResult& InHit, const float InSweepFloorDist, const bool bIsWalkableFloor);
-	void SetFromLineTrace(const FHitResult& InHit, const float InSweepFloorDist, const float InLineDist, const bool bIsWalkableFloor);
-};
-
-
-/** Struct updated by StepUp() to return result of final step down, if applicable. */
-struct FStepDownResult
-{
-	uint32 bComputedFloor : 1;		// True if the floor was computed as a result of the step down.
-	FFindFloorResult FloorResult;	// The result of the floor test if the floor was updated.
-
-	FStepDownResult()
-		: bComputedFloor(false)
-	{
-	}
-};
-
 /** 
  * Tick function that calls UCharacterMovementComponent::PostPhysicsTickComponent
  **/
@@ -263,8 +166,14 @@ public:
 	UPROPERTY(Category="Character Movement: Jumping / Falling", EditAnywhere, BlueprintReadWrite, AdvancedDisplay, meta=(ClampMin="0", UIMin="0"))
 	float JumpOffJumpZFactor;
 
+
+private:
 	FCharacterMovementAsyncCallback* AsyncCallback;
 
+protected:
+	// This is the most recent async state from simulated. Only safe for access on physics thread.
+	TSharedPtr<FCharacterMovementAsyncOutput, ESPMode::ThreadSafe> AsyncSimState;
+	bool bMovementModeDirty = false; // Gamethread changed movement mode, need to update sim.
 private:
 
 	/**
@@ -2114,11 +2023,17 @@ protected:
 	virtual void OnUnableToFollowBaseMove(const FVector& DeltaPosition, const FVector& OldLocation, const FHitResult& MoveOnBaseHit);
 
 
-private:
+protected:
 	/* Prepare inputs for asynchronous simulation on physics thread */ 
 	virtual void FillAsyncInput(const FVector& InputVector, FCharacterMovementAsyncInput& AsyncInput) const;
+	virtual void BuildAsyncInput();
+	virtual void PostBuildAsyncInput();
 	/* Apply outputs from async sim. */
 	virtual void ApplyAsyncOutput(FCharacterMovementAsyncOutput& Output);
+	virtual void ProcessAsyncOutput();
+	
+	/* Register async callback with physics system. */
+	virtual void RegisterAsyncCallback();
 	
 public:
 
diff --git a/Engine/Source/Runtime/Engine/Private/CharacterMovementComponentAsync.cpp b/Engine/Source/Runtime/Engine/Private/CharacterMovementComponentAsync.cpp
--- a/Engine/Source/Runtime/Engine/Private/CharacterMovementComponentAsync.cpp
+++ b/Engine/Source/Runtime/Engine/Private/CharacterMovementComponentAsync.cpp
@@ -1,6 +1,7 @@
 // Copyright Epic Games, Inc. All Rights Reserved.
 
 #include "CharacterMovementComponentAsync.h"
+#include "GameFramework/CharacterMovementComponent.h"
 #include "PhysicsProxy/SingleParticlePhysicsProxy.h"
 #include "Chaos/SpatialAccelerationCollection.h"
 #include "PBDRigidsSolver.h"
@@ -624,7 +625,7 @@ void FCharacterMovementAsyncInput::UpdateBasedMovement(float DeltaSeconds, FChar
 				// Do we need this value after all?
 
 				CharacterInput->FaceRotation(TargetRotator, 0.0f, *this, Output);
-				FinalQuat =    Output.CharacterOwnerRotation.Quaternion();//UpdatedComponent->GetComponentQuat(); supposed to be modified by MockFaceRotation, si this ok?
+				FinalQuat =  Output.Pawn->Rotation.Quaternion();//UpdatedComponent->GetComponentQuat(); supposed to be modified by MockFaceRotation, si this ok?
 
 
 				if (PawnOldQuat.Equals(FinalQuat, 1e-6f))
@@ -1057,15 +1058,15 @@ void FCharacterMovementAsyncInput::PhysFalling(float deltaTime, int32 Iterations
 
 		// If jump is providing force, gravity may be affected.
 		bool bEndingJumpForce = false;
-		if (Output.JumpForceTimeRemaining > 0.0f)
+		if (Output.Pawn->JumpForceTimeRemaining > 0.0f)
 		{
 			// Consume some of the force time. Only the remaining time (if any) is affected by gravity when bApplyGravityWhileJumping=false.
-			const float JumpForceTime = FMath::Min(Output.JumpForceTimeRemaining, timeTick);
+			const float JumpForceTime = FMath::Min(Output.Pawn->JumpForceTimeRemaining, timeTick);
 			GravityTime = bApplyGravityWhileJumping ? timeTick : FMath::Max(0.0f, timeTick - JumpForceTime);
 
 			// Update Character state
-			Output.JumpForceTimeRemaining -= JumpForceTime;
-			if (Output.JumpForceTimeRemaining <= 0.0f)
+			Output.Pawn->JumpForceTimeRemaining -= JumpForceTime;
+			if (Output.Pawn->JumpForceTimeRemaining <= 0.0f)
 			{
 				CharacterInput->ResetJumpState(*this, Output);
 				bEndingJumpForce = true;
@@ -1780,8 +1781,7 @@ void FCharacterMovementAsyncInput::OnMovementModeChanged(EMovementMode PreviousM
 		}*/
 	}
 
-	// TODO Character->OnMovementModeChanged
-	//CharacterOwner->OnMovementModeChanged(PreviousMovementMode, PreviousCustomMode);
+	CharacterInput->OnMovementModeChanged(PreviousMovementMode, *this, Output, PreviousCustomMode);
 
 	//ensureMsgf(GroundMovementMode == MOVE_Walking || GroundMovementMode == MOVE_NavWalking, TEXT("Invalid GroundMovementMode %d. MovementMode: %d, PreviousMovementMode: %d"), GroundMovementMode.GetValue(), MovementMode.GetValue(), PreviousMovementMode);
 }
@@ -4376,7 +4376,7 @@ void FCharacterAsyncInput::FaceRotation(FRotator NewControlRotation, float Delta
 	// Only if we actually are going to use any component of rotation.
 	if (bUseControllerRotationPitch || bUseControllerRotationYaw || bUseControllerRotationRoll)
 	{
-		FRotator& CurrentRotation = Output.CharacterOwnerRotation;
+		FRotator& CurrentRotation = Output.Pawn->Rotation;
 		if (!bUseControllerRotationPitch)
 		{
 			NewControlRotation.Pitch = CurrentRotation.Pitch;
@@ -4406,16 +4406,16 @@ void FCharacterAsyncInput::FaceRotation(FRotator NewControlRotation, float Delta
 
 void FCharacterAsyncInput::CheckJumpInput(float DeltaSeconds, const FCharacterMovementAsyncInput& Input, FCharacterMovementAsyncOutput& Output) const
 {
-	Output.JumpCurrentCountPreJump = Output.JumpCurrentCount;
+	Output.Pawn->JumpCurrentCountPreJump = Output.Pawn->JumpCurrentCount;
 
-	if (Output.bPressedJump)
+	if (Output.Pawn->bPressedJump)
 	{
 		// If this is the first jump and we're already falling,
 		// then increment the JumpCount to compensate.
-		const bool bFirstJump = Output.JumpCurrentCount == 0;
+		const bool bFirstJump = Output.Pawn->JumpCurrentCount == 0;
 		if (bFirstJump && Input.IsFalling(Output))
 		{
-			Output.JumpCurrentCount++;
+			Output.Pawn->JumpCurrentCount++;
 		}
 
 		// TODO bClientUpdating
@@ -4423,35 +4423,36 @@ void FCharacterAsyncInput::CheckJumpInput(float DeltaSeconds, const FCharacterMo
 		if (bDidJump)
 		{
 			// Transition from not (actively) jumping to jumping.
-			if (!Output.bWasJumping)
+			if (!Output.Pawn->bWasJumping)
 			{
-				Output.JumpCurrentCount++;
-				Output.JumpForceTimeRemaining = JumpMaxHoldTime;
+				Output.Pawn->JumpCurrentCount++;
+				Output.Pawn->JumpForceTimeRemaining = JumpMaxHoldTime;
 				//OnJumped(); TODO Jumping
 			}
 		}
 
-		Output.bWasJumping = bDidJump;
+		Output.Pawn->bWasJumping = bDidJump;
 	}
 }
 
 void FCharacterAsyncInput::ClearJumpInput(float DeltaSeconds, const FCharacterMovementAsyncInput& Input, FCharacterMovementAsyncOutput& Output) const
 {
-	if (Output.bPressedJump)
+	if (Output.Pawn->bPressedJump)
 	{
-		Output.JumpKeyHoldTime += DeltaSeconds;
+		Output.Pawn->JumpKeyHoldTime += DeltaSeconds;
 
 		// Don't disable bPressedJump right away if it's still held.
 		// Don't modify JumpForceTimeRemaining because a frame of update may be remaining.
-		if (Output.JumpKeyHoldTime >= JumpMaxHoldTime)
+		if (Output.Pawn->JumpKeyHoldTime >= JumpMaxHoldTime)
 		{
-			Output.bPressedJump = false;
+			Output.Pawn->bClearJumpInput = true;
+			Output.Pawn->bPressedJump = false;
 		}
 	}
 	else
 	{
-		Output.JumpForceTimeRemaining = 0.0f;
-		Output.bWasJumping = false;
+		Output.Pawn->JumpForceTimeRemaining = 0.0f;
+		Output.Pawn->bWasJumping = false;
 	}
 }
 
@@ -4466,15 +4467,15 @@ bool FCharacterAsyncInput::CanJump(const FCharacterMovementAsyncInput& Input, FC
 	if (bCanJump)
 	{
 		// Ensure JumpHoldTime and JumpCount are valid.
-		if (!Output.bWasJumping || Input.CharacterInput->JumpMaxHoldTime <= 0.0f)
+		if (!Output.Pawn->bWasJumping || Input.CharacterInput->JumpMaxHoldTime <= 0.0f)
 		{
-			if (Output.JumpCurrentCount == 0 && Input.IsFalling(Output))
+			if (Output.Pawn->JumpCurrentCount == 0 && Input.IsFalling(Output))
 			{
-				bCanJump = Output.JumpCurrentCount + 1 < Input.CharacterInput->JumpMaxCount;
+				bCanJump = Output.Pawn->JumpCurrentCount + 1 < Input.CharacterInput->JumpMaxCount;
 			}
 			else
 			{
-				bCanJump = Output.JumpCurrentCount < Input.CharacterInput->JumpMaxCount;
+				bCanJump = Output.Pawn->JumpCurrentCount < Input.CharacterInput->JumpMaxCount;
 			}
 		}
 		else
@@ -4482,9 +4483,9 @@ bool FCharacterAsyncInput::CanJump(const FCharacterMovementAsyncInput& Input, FC
 			// Only consider JumpKeyHoldTime as long as:
 			// A) The jump limit hasn't been met OR
 			// B) The jump limit has been met AND we were already jumping
-			const bool bJumpKeyHeld = (Output.bPressedJump && Output.JumpKeyHoldTime < Input.CharacterInput->JumpMaxHoldTime);
+			const bool bJumpKeyHeld = (Output.Pawn->bPressedJump && Output.Pawn->JumpKeyHoldTime < Input.CharacterInput->JumpMaxHoldTime);
 			bCanJump = bJumpKeyHeld &&
-				((Output.JumpCurrentCount < Input.CharacterInput->JumpMaxCount) || (Output.bWasJumping && Output.JumpCurrentCount == Input.CharacterInput->JumpMaxCount));
+				((Output.Pawn->JumpCurrentCount < Input.CharacterInput->JumpMaxCount) || (Output.Pawn->bWasJumping && Output.Pawn->JumpCurrentCount == Input.CharacterInput->JumpMaxCount));
 		}
 	}
 
@@ -4493,14 +4494,99 @@ bool FCharacterAsyncInput::CanJump(const FCharacterMovementAsyncInput& Input, FC
 
 void FCharacterAsyncInput::ResetJumpState(const FCharacterMovementAsyncInput& Input, FCharacterMovementAsyncOutput& Output) const
 {
-	Output.bPressedJump = false;
-	Output.bWasJumping = false;
-	Output.JumpKeyHoldTime = 0.0f;
-	Output.JumpForceTimeRemaining = 0.0f;
+	if (Output.Pawn->bPressedJump == true)
+	{
+		Output.Pawn->bClearJumpInput = true;
+	}
+
+	Output.Pawn->bPressedJump = false;
+	Output.Pawn->bWasJumping = false;
+	Output.Pawn->JumpKeyHoldTime = 0.0f;
+	Output.Pawn->JumpForceTimeRemaining = 0.0f;
+
+	if (!Input.IsFalling(Output))
+	{
+		Output.Pawn->JumpCurrentCount = 0;
+		Output.Pawn->JumpCurrentCountPreJump = 0;
+	}
+}
 
-	if (Input.IsFalling(Output))
+void FCharacterAsyncInput::OnMovementModeChanged(EMovementMode PrevMovementMode, const FCharacterMovementAsyncInput& Input, FCharacterMovementAsyncOutput& Output, uint8 PreviousCustomMode)
+{
+	if (!Output.Pawn->bPressedJump || !Input.IsFalling(Output))
 	{
-		Output.JumpCurrentCount = 0;
-		Output.JumpCurrentCountPreJump = 0;
+		ResetJumpState(Input, Output);
 	}
-}
\ No newline at end of file
+	
+	// TODO Proxy
+	// Record jump force start time for proxies. Allows us to expire the jump even if not continually ticking down a timer.
+	/*if (bProxyIsJumpForceApplied && Input->IsFalling(Output))
+	{
+		ProxyJumpForceStartedTime = GetWorld()->GetTimeSeconds();
+	}*/
+
+	// TODO Blueprint calls
+	//K2_OnMovementModeChanged(PrevMovementMode, CharacterMovement->MovementMode, PrevCustomMode, CharacterMovement->CustomMovementMode);
+	//MovementModeChangedDelegate.Broadcast(this, PrevMovementMode, PrevCustomMode);
+}
+
+void FCharacterMovementAsyncCallback::OnPreSimulate_Internal()
+{
+	PreSimulateImpl<FCharacterMovementAsyncInput, FCharacterMovementAsyncOutput>(*this);
+}
+
+void FCharacterMovementAsyncOutput::Copy(const FCharacterMovementAsyncOutput& Value)
+{
+	bIsValid = Value.bIsValid;
+
+	bWasSimulatingRootMotion = Value.bWasSimulatingRootMotion;
+	MovementMode = Value.MovementMode;
+	GroundMovementMode = Value.GroundMovementMode;
+	CustomMovementMode = Value.CustomMovementMode;
+	Acceleration = Value.Acceleration;
+	AnalogInputModifier = Value.AnalogInputModifier;
+	LastUpdateLocation = Value.LastUpdateLocation;
+	LastUpdateRotation = Value.LastUpdateRotation;
+	LastUpdateVelocity = Value.LastUpdateVelocity;
+	bForceNextFloorCheck = Value.bForceNextFloorCheck;
+	CurrentRootMotion = Value.CurrentRootMotion;
+	Velocity = Value.Velocity;
+	bDeferUpdateBasedMovement = Value.bDeferUpdateBasedMovement;
+	MoveComponentFlags = Value.MoveComponentFlags;
+	PendingForceToApply = Value.PendingForceToApply;
+	PendingImpulseToApply = Value.PendingImpulseToApply;
+	PendingLaunchVelocity = Value.PendingLaunchVelocity;
+	bCrouchMaintainsBaseLocation = Value.bCrouchMaintainsBaseLocation;
+	bJustTeleported = Value.bJustTeleported;
+	ScaledCapsuleRadius = Value.ScaledCapsuleRadius;
+	ScaledCapsuleHalfHeight = Value.ScaledCapsuleHalfHeight;
+	bIsCrouched = Value.bIsCrouched;
+	bWantsToCrouch = Value.bWantsToCrouch;
+	bMovementInProgress = Value.bMovementInProgress;
+	CurrentFloor = Value.CurrentFloor;
+	bHasRequestedVelocity = Value.bHasRequestedVelocity;
+	bRequestedMoveWithMaxSpeed = Value.bRequestedMoveWithMaxSpeed;
+	RequestedVelocity = Value.RequestedVelocity;
+	NumJumpApexAttempts = Value.NumJumpApexAttempts;
+	RootMotionParams = Value.RootMotionParams;
+	bShouldApplyDeltaToMeshPhysicsTransforms = Value.bShouldApplyDeltaToMeshPhysicsTransforms;
+	DeltaPosition = Value.DeltaPosition;
+	DeltaQuat = Value.DeltaQuat;
+	DeltaTime = Value.DeltaTime;
+	OldVelocity = Value.OldVelocity;
+	OldLocation = Value.OldLocation;
+
+	// See MaybeUpdateBasedMovement
+	// TODO MovementBase, handle tick group changes properly
+	bShouldDisablePostPhysicsTick = Value.bShouldDisablePostPhysicsTick;
+	bShouldEnablePostPhysicsTick = Value.bShouldEnablePostPhysicsTick;
+	bShouldAddMovementBaseTickDependency = Value.bShouldAddMovementBaseTickDependency;
+	bShouldRemoveMovementBaseTickDependency = Value.bShouldRemoveMovementBaseTickDependency;
+
+	NewMovementBase = Value.NewMovementBase;
+	NewMovementBaseOwner = Value.NewMovementBaseOwner;
+
+	UpdatedComponent = Value.UpdatedComponent;
+	*Pawn = *Value.Pawn;
+}
+
diff --git a/Engine/Source/Runtime/Engine/Private/Components/CharacterMovementComponent.cpp b/Engine/Source/Runtime/Engine/Private/Components/CharacterMovementComponent.cpp
--- a/Engine/Source/Runtime/Engine/Private/Components/CharacterMovementComponent.cpp
+++ b/Engine/Source/Runtime/Engine/Private/Components/CharacterMovementComponent.cpp
@@ -620,17 +620,7 @@ void UCharacterMovementComponent::BeginPlay()
 {
 	Super::BeginPlay();
 	
-	if(CharacterMovementCVars::AsyncCharacterMovement == 1)
-	{
-		if (UWorld* World = GetWorld())
-		{
-			if (FPhysScene* PhysScene = World->GetPhysicsScene())
-			{
-				AsyncCallback = PhysScene->GetSolver()->CreateAndRegisterSimCallbackObject_External<FCharacterMovementAsyncCallback>();
-
-			}
-		}
-	}
+	RegisterAsyncCallback();
 }
 
 void UCharacterMovementComponent::PostLoad()
@@ -1159,6 +1149,8 @@ void UCharacterMovementComponent::SetMovementMode(EMovementMode NewMovementMode,
 	OnMovementModeChanged(PrevMovementMode, PrevCustomMode);
 
 	// @todo do we need to disable ragdoll physics here? Should this function do nothing if in ragdoll?
+
+	bMovementModeDirty = true; // lets async callback know movement mode was dirtied on game thread
 }
 
 
@@ -1474,25 +1466,12 @@ void UCharacterMovementComponent::TickComponent(float DeltaTime, enum ELevelTick
 
 void UCharacterMovementComponent::PrePhysicsTickComponent(float DeltaTime, FCharacterMovementComponentPrePhysicsTickFunction& ThisTickFunction)
 {
-	if (CharacterMovementCVars::AsyncCharacterMovement == 1)
-	{
-		FCharacterMovementAsyncInput* Input = AsyncCallback->GetProducerInputData_External();
-
-		const FVector InputVector = ConsumeInputVector();
-		FillAsyncInput(InputVector, *Input);
-	}
+	BuildAsyncInput();
 }
 
 void UCharacterMovementComponent::PostPhysicsTickComponent(float DeltaTime, FCharacterMovementComponentPostPhysicsTickFunction& ThisTickFunction)
 {
-	if (CharacterMovementCVars::AsyncCharacterMovement == 1)
-	{
-		while (auto Output = AsyncCallback->PopOutputData_External())
-		{
-			ApplyAsyncOutput(*Output);
-		}
-	}
-
+	ProcessAsyncOutput();
 
 	if (bDeferUpdateBasedMovement)
 	{
@@ -10735,24 +10714,6 @@ void UCharacterMovementComponent::RegisterComponentTickFunctions(bool bRegister)
 	}
 }
 
-void FCharacterMovementAsyncCallback::OnPreSimulate_Internal()
-{
-	const FCharacterMovementAsyncInput* Input = GetConsumerInput_Internal();
-	if (Input && Input->bInitialized)
-	{
-		// Initialize callback output
-		FCharacterMovementAsyncOutput& Output = GetProducerOutputData_Internal();
-		Output = *Input->InitialOutput;
-
-		Input->Simulate(GetDeltaTime_Internal(), Output);
-
-		// Save callback output to initial output for subsequent sim steps using same input.
-		*Input->InitialOutput = Output;
-
-		//UE_LOG(LogChaos, Warning, TEXT("Output Pos: (%f, %f, %f)"), Output.UpdatedComponent.Location.X, Output.UpdatedComponent.Location.Y, Output.UpdatedComponent.Location.Z);
-	}
-}
-
 void UCharacterMovementComponent::ApplyWorldOffset(const FVector& InOffset, bool bWorldShift)
 {
 	OldBaseLocation += InOffset;
@@ -12202,7 +12163,7 @@ void UCharacterMovementComponent::FillAsyncInput(const FVector& InputVector, FCh
 
 	// Character owner inputs
 	TUniquePtr<FCharacterAsyncInput>& CharacterInput = AsyncInput.CharacterInput;
-	CharacterInput = MakeUnique<FCharacterAsyncInput>();
+	ensure(CharacterInput.IsValid());
 	CharacterInput->JumpMaxHoldTime = CharacterOwner->GetJumpMaxHoldTime();
 	CharacterInput->JumpMaxCount = CharacterOwner->JumpMaxCount;
 	CharacterInput->LocalRole = ENetRole::ROLE_Authority;//CharacterOwner->GetLocalRole(); Override as we aren't currently replicating to server. TODO NetRole
@@ -12219,7 +12180,7 @@ void UCharacterMovementComponent::FillAsyncInput(const FVector& InputVector, FCh
 	if (ensure(UpdatedComponent))
 	{
 		TUniquePtr<FUpdatedComponentAsyncInput>& UpdatedComponentInput = AsyncInput.UpdatedComponentInput;
-		UpdatedComponentInput = MakeUnique<FUpdatedComponentAsyncInput>();
+		ensure(UpdatedComponentInput.IsValid());
 
 		UpdatedComponentInput->bIsQueryCollisionEnabled = UpdatedComponent->IsQueryCollisionEnabled();
 		UpdatedComponentInput->bIsSimulatingPhysics = UpdatedComponent->IsSimulatingPhysics();
@@ -12259,72 +12220,133 @@ void UCharacterMovementComponent::FillAsyncInput(const FVector& InputVector, FCh
 	}
 	AsyncInput.RandomStream.Initialize(FApp::bUseFixedSeed ? GetFName() : NAME_None);
 
-	// FILL OUT INITIAL OUTPUT:
-	TUniquePtr<FCharacterMovementAsyncOutput>& Output = AsyncInput.InitialOutput;
-	Output = MakeUnique<FCharacterMovementAsyncOutput>();
-	Output->bWasSimulatingRootMotion = bWasSimulatingRootMotion;
-	Output->MovementMode = MovementMode;
-	Output->GroundMovementMode = GroundMovementMode;
-	Output->CustomMovementMode = CustomMovementMode;
-	Output->Acceleration = Acceleration;
-	Output->AnalogInputModifier = AnalogInputModifier;
-	Output->LastUpdateLocation = LastUpdateLocation;
-	Output->LastUpdateRotation = LastUpdateRotation;
-	Output->LastUpdateVelocity = LastUpdateVelocity;
-	Output->bForceNextFloorCheck = bForceNextFloorCheck;
-	Output->Velocity = Velocity;
-	Output->bDeferUpdateBasedMovement = bDeferUpdateBasedMovement;
-	Output->MoveComponentFlags = MoveComponentFlags;
-	Output->PendingForceToApply = PendingForceToApply;
-	Output->PendingImpulseToApply = PendingImpulseToApply;
-	Output->PendingLaunchVelocity = PendingLaunchVelocity;
-	Output->bCrouchMaintainsBaseLocation = bCrouchMaintainsBaseLocation;
-	Output->bJustTeleported = bJustTeleported;
-	CharacterOwner->GetCapsuleComponent()->GetScaledCapsuleSize(Output->ScaledCapsuleRadius, Output->ScaledCapsuleHalfHeight);
-	Output->bIsCrouched = CharacterOwner->bIsCrouched;
-	Output->bWantsToCrouch = bWantsToCrouch;
-	Output->bMovementInProgress = bMovementInProgress;
-	Output->CurrentFloor = CurrentFloor;
-	Output->bHasRequestedVelocity = bHasRequestedVelocity;
-	Output->bRequestedMoveWithMaxSpeed = bRequestedMoveWithMaxSpeed;
-	Output->RequestedVelocity = RequestedVelocity;
-	Output->NumJumpApexAttempts = NumJumpApexAttempts;
-	Output->RootMotionParams = RootMotionParams;
-	Output->bShouldApplyDeltaToMeshPhysicsTransforms = false;
-	Output->DeltaPosition = FVector::ZeroVector;
-	Output->DeltaQuat = FQuat::Identity;
-	Output->DeltaTime = 0.0f; // Fill out in async callback.
-	Output->OldLocation = UpdatedComponent ? UpdatedComponent->GetComponentLocation() : FVector::ZeroVector;
-	Output->OldVelocity = Velocity;
-	Output->bShouldDisablePostPhysicsTick = false;
-	Output->bShouldEnablePostPhysicsTick = false;
-	Output->bShouldAddMovementBaseTickDependency = false;
-	Output->bShouldRemoveMovementBaseTickDependency = false;
-	Output->NewMovementBase = AsyncInput.MovementBaseAsyncData.CachedMovementBase;
-	Output->NewMovementBaseOwner = Output->NewMovementBase ? Output->NewMovementBase->GetOwner() : nullptr;
-	Output->CurrentRootMotion = CurrentRootMotion;
-
-	// Character owner data
-	Output->CharacterOwnerRotation = CharacterOwner->GetActorRotation();
-	Output->JumpCurrentCountPreJump = CharacterOwner->JumpCurrentCountPreJump;
-	Output->JumpCurrentCount = CharacterOwner->JumpCurrentCount;
-	Output->JumpForceTimeRemaining = CharacterOwner->JumpForceTimeRemaining;
-	Output->bWasJumping = CharacterOwner->bWasJumping;
-	Output->bPressedJump = CharacterOwner->bPressedJump;
-	Output->JumpKeyHoldTime = CharacterOwner->JumpKeyHoldTime;
-
-	/*FVector Pos = UpdatedComponent->GetComponentLocation();
-	if (CharacterOwner->GetLocalRole() != ROLE_Authority) // only print on client
-	{
-		float timestamp = FPlatformTime::ToMilliseconds64(FPlatformTime::Cycles64());
-		UE_LOG(LogChaos, Warning, TEXT("Time: %f Movement Input Pos: (%f, %f, %f) this: %llx"), timestamp, Pos.X, Pos.Y, Pos.Z, this);
-	}*/
+
+	// Only game thread inputs need to be updated here.
+	AsyncInput.GTInputs.bWantsToCrouch = bWantsToCrouch;
+
+	if (MovementMode == MOVE_None)
+	{
+		// TODO fix this, sometimes we don't fall on start and just float in air because movement mode is none. Not sure why.
+		// Pawn typically sets movement mode after possess, but something goes wrong.
+		AsyncInput.GTInputs.MovementMode = GroundMovementMode;
+		AsyncInput.GTInputs.bValidMovementMode = true;
+	}
+	else
+	{
+		AsyncInput.GTInputs.MovementMode = MovementMode;
+		AsyncInput.GTInputs.bValidMovementMode = bMovementModeDirty;
+	}
+	AsyncInput.GTInputs.bPressedJump = CharacterOwner->bPressedJump;
+
+
+	if(AsyncSimState->IsValid() == false)
+	{
+		// Need to fully initialize output as we do not have any stored async state.
+		AsyncSimState->bWasSimulatingRootMotion = bWasSimulatingRootMotion;
+		AsyncSimState->MovementMode = MovementMode;
+		AsyncSimState->GroundMovementMode = GroundMovementMode;
+		AsyncSimState->CustomMovementMode = CustomMovementMode;
+		AsyncSimState->Acceleration = Acceleration;
+		AsyncSimState->AnalogInputModifier = AnalogInputModifier;
+		AsyncSimState->LastUpdateLocation = LastUpdateLocation;
+		AsyncSimState->LastUpdateRotation = LastUpdateRotation;
+		AsyncSimState->LastUpdateVelocity = LastUpdateVelocity;
+		AsyncSimState->bForceNextFloorCheck = bForceNextFloorCheck;
+		AsyncSimState->Velocity = Velocity;
+		AsyncSimState->bDeferUpdateBasedMovement = bDeferUpdateBasedMovement;
+		AsyncSimState->MoveComponentFlags = MoveComponentFlags;
+		AsyncSimState->PendingForceToApply = PendingForceToApply;
+		AsyncSimState->PendingImpulseToApply = PendingImpulseToApply;
+		AsyncSimState->PendingLaunchVelocity = PendingLaunchVelocity;
+		AsyncSimState->bCrouchMaintainsBaseLocation = bCrouchMaintainsBaseLocation;
+		AsyncSimState->bJustTeleported = bJustTeleported;
+		CharacterOwner->GetCapsuleComponent()->GetScaledCapsuleSize(AsyncSimState->ScaledCapsuleRadius, AsyncSimState->ScaledCapsuleHalfHeight);
+		AsyncSimState->bIsCrouched = CharacterOwner->bIsCrouched;
+		AsyncSimState->bWantsToCrouch = bWantsToCrouch;
+		AsyncSimState->bMovementInProgress = bMovementInProgress;
+		AsyncSimState->CurrentFloor = CurrentFloor;
+		AsyncSimState->bHasRequestedVelocity = bHasRequestedVelocity;
+		AsyncSimState->bRequestedMoveWithMaxSpeed = bRequestedMoveWithMaxSpeed;
+		AsyncSimState->RequestedVelocity = RequestedVelocity;
+		AsyncSimState->NumJumpApexAttempts = NumJumpApexAttempts;
+		AsyncSimState->RootMotionParams = RootMotionParams;
+		AsyncSimState->bShouldApplyDeltaToMeshPhysicsTransforms = false;
+		AsyncSimState->DeltaPosition = FVector::ZeroVector;
+		AsyncSimState->DeltaQuat = FQuat::Identity;
+		AsyncSimState->DeltaTime = 0.0f; // Fill out in async callback.
+		AsyncSimState->OldLocation = UpdatedComponent ? UpdatedComponent->GetComponentLocation() : FVector::ZeroVector;
+		AsyncSimState->OldVelocity = Velocity;
+		AsyncSimState->bShouldDisablePostPhysicsTick = false;
+		AsyncSimState->bShouldEnablePostPhysicsTick = false;
+		AsyncSimState->bShouldAddMovementBaseTickDependency = false;
+		AsyncSimState->bShouldRemoveMovementBaseTickDependency = false;
+		AsyncSimState->NewMovementBase = AsyncInput.MovementBaseAsyncData.CachedMovementBase;
+		AsyncSimState->NewMovementBaseOwner = AsyncSimState->NewMovementBase ? AsyncSimState->NewMovementBase->GetOwner() : nullptr;
+		AsyncSimState->CurrentRootMotion = CurrentRootMotion;
+
+		// Character owner data
+		TUniquePtr<FCharacterAsyncOutput>& PawnOutput = AsyncSimState->Pawn;
+		PawnOutput->Rotation = CharacterOwner->GetActorRotation();
+		PawnOutput->JumpCurrentCountPreJump = CharacterOwner->JumpCurrentCountPreJump;
+		PawnOutput->JumpCurrentCount = CharacterOwner->JumpCurrentCount;
+		PawnOutput->JumpForceTimeRemaining = CharacterOwner->JumpForceTimeRemaining;
+		PawnOutput->bWasJumping = CharacterOwner->bWasJumping;
+		PawnOutput->bPressedJump = CharacterOwner->bPressedJump;
+		PawnOutput->JumpKeyHoldTime = CharacterOwner->JumpKeyHoldTime;
+		PawnOutput->bClearJumpInput = false;
+
+		AsyncSimState->bIsValid = true;
+	}
+}
+
+void UCharacterMovementComponent::BuildAsyncInput()
+{
+	if (CharacterMovementCVars::AsyncCharacterMovement == 1)
+	{
+		FCharacterMovementAsyncInput* Input = AsyncCallback->GetProducerInputData_External();
+		if (Input->bInitialized == false)
+		{
+			Input->Initialize<FCharacterMovementAsyncInput, FCharacterMovementAsyncOutput>();
+		}
+
+		if (AsyncSimState.IsValid() == false)
+		{
+			AsyncSimState = MakeShared<FCharacterMovementAsyncOutput, ESPMode::ThreadSafe>();
+		}
+		Input->AsyncSimState = AsyncSimState;
+
+		const FVector InputVector = ConsumeInputVector();
+		FillAsyncInput(InputVector, *Input);
+
+		PostBuildAsyncInput();
+	}
+}
+
+void UCharacterMovementComponent::PostBuildAsyncInput()
+{
+	// Reset so we can tell if movement mode change comes from game thread.
+	bMovementModeDirty = false;
 }
 
 void UCharacterMovementComponent::ApplyAsyncOutput(FCharacterMovementAsyncOutput& Output)
 {
 	ensure(Output.DeltaTime > 0.0f);
 
+	if (Output.IsValid() == false)
+	{
+		return;
+	}
+
+
+	// TODO does anyhting here need to be interpolated?
+
+	// TODO:
+	// Not all of this stuff should actually be copied to game thread, if it is not read by other GT
+	// systems it doesn't need to be copied back. Some values representing inputs will cause issues if copied back
+	// (like bPressedJump for example)
+	// Need to sort through and see what should/shouldn't be copied back from outputs.
+
+
 	// TODO verify order
 	bWasSimulatingRootMotion = Output.bWasSimulatingRootMotion;
 	MovementMode = Output.MovementMode;
@@ -12363,12 +12385,17 @@ void UCharacterMovementComponent::ApplyAsyncOutput(FCharacterMovementAsyncOutput
 
 	if (CharacterOwner)
 	{
-		CharacterOwner->JumpCurrentCountPreJump = Output.JumpCurrentCountPreJump;
-		CharacterOwner->JumpCurrentCount = Output.JumpCurrentCount;
-		CharacterOwner->bPressedJump = Output.bPressedJump;
-		CharacterOwner->JumpForceTimeRemaining = Output.JumpForceTimeRemaining;
-		CharacterOwner->bWasJumping = Output.bWasJumping;
-		CharacterOwner->JumpKeyHoldTime = Output.JumpKeyHoldTime;
+		TUniquePtr<FCharacterAsyncOutput>& CharacterOutput = Output.Pawn;
+		CharacterOwner->JumpCurrentCountPreJump = CharacterOutput->JumpCurrentCountPreJump;
+		CharacterOwner->JumpCurrentCount = CharacterOutput->JumpCurrentCount;
+		CharacterOwner->JumpForceTimeRemaining = CharacterOutput->JumpForceTimeRemaining;
+		CharacterOwner->bWasJumping = CharacterOutput->bWasJumping;
+		CharacterOwner->JumpKeyHoldTime = CharacterOutput->JumpKeyHoldTime;
+
+		if (CharacterOutput->bClearJumpInput)
+		{
+			CharacterOwner->bPressedJump = false;
+		}
 	}
 
 	NumJumpApexAttempts = Output.NumJumpApexAttempts;
@@ -12394,14 +12421,35 @@ void UCharacterMovementComponent::ApplyAsyncOutput(FCharacterMovementAsyncOutput
 	// TODO Should this happen before or after movement update above?
 	if (CharacterOwner)
 	{
-		CharacterOwner->FaceRotation(Output.CharacterOwnerRotation);
+		CharacterOwner->FaceRotation(Output.Pawn->Rotation);
 	}
 
 	// TODO MovementBase
 	// Need to call SavedBaseLocation?
 	// Call SetBase with NewMovementBase
+}
 
+void UCharacterMovementComponent::ProcessAsyncOutput()
+{
+	if (CharacterMovementCVars::AsyncCharacterMovement == 1)
+	{
+		while (auto Output = AsyncCallback->PopOutputData_External())
+		{
+			ApplyAsyncOutput(*Output);
+		}
+	}
+}
 
-
-
+void UCharacterMovementComponent::RegisterAsyncCallback()
+{
+	if (CharacterMovementCVars::AsyncCharacterMovement == 1)
+	{
+		if (UWorld* World = GetWorld())
+		{
+			if (FPhysScene* PhysScene = World->GetPhysicsScene())
+			{
+				AsyncCallback = PhysScene->GetSolver()->CreateAndRegisterSimCallbackObject_External<FCharacterMovementAsyncCallback>();
+			}
+		}
+	}
 }
diff --git a/Engine/Source/Runtime/Engine/Public/CharacterMovementComponentAsync.h b/Engine/Source/Runtime/Engine/Public/CharacterMovementComponentAsync.h
--- a/Engine/Source/Runtime/Engine/Public/CharacterMovementComponentAsync.h
+++ b/Engine/Source/Runtime/Engine/Public/CharacterMovementComponentAsync.h
@@ -2,11 +2,124 @@
 
 #pragma once
 #include "GameFramework/RootMotionSource.h"
-#include "GameFramework/CharacterMovementComponent.h"
 #include "Chaos/SimCallbackObject.h"
+#include "CharacterMovementComponentAsync.generated.h"
 
-struct FFindFloorResult;
-struct FStepDownResult;
+
+struct FCharacterMovementAsyncOutput;
+struct FCharacterMovementAsyncInput;
+
+// TODO move these common structures to separate header?
+
+// Enum used to control GetPawnCapsuleExtent behavior
+enum EShrinkCapsuleExtent
+{
+	SHRINK_None,			// Don't change the size of the capsule
+	SHRINK_RadiusCustom,	// Change only the radius, based on a supplied param
+	SHRINK_HeightCustom,	// Change only the height, based on a supplied param
+	SHRINK_AllCustom,		// Change both radius and height, based on a supplied param
+};
+
+/** Data about the floor for walking movement, used by CharacterMovementComponent. */
+USTRUCT(BlueprintType)
+struct ENGINE_API FFindFloorResult
+{
+	GENERATED_USTRUCT_BODY()
+
+		/**
+		* True if there was a blocking hit in the floor test that was NOT in initial penetration.
+		* The HitResult can give more info about other circumstances.
+		*/
+		UPROPERTY(VisibleInstanceOnly, BlueprintReadOnly, Category = CharacterFloor)
+		uint32 bBlockingHit : 1;
+
+	/** True if the hit found a valid walkable floor. */
+	UPROPERTY(VisibleInstanceOnly, BlueprintReadOnly, Category = CharacterFloor)
+		uint32 bWalkableFloor : 1;
+
+	/** True if the hit found a valid walkable floor using a line trace (rather than a sweep test, which happens when the sweep test fails to yield a walkable surface). */
+	UPROPERTY(VisibleInstanceOnly, BlueprintReadOnly, Category = CharacterFloor)
+		uint32 bLineTrace : 1;
+
+	/** The distance to the floor, computed from the swept capsule trace. */
+	UPROPERTY(VisibleInstanceOnly, BlueprintReadOnly, Category = CharacterFloor)
+		float FloorDist;
+
+	/** The distance to the floor, computed from the trace. Only valid if bLineTrace is true. */
+	UPROPERTY(VisibleInstanceOnly, BlueprintReadOnly, Category = CharacterFloor)
+		float LineDist;
+
+	/** Hit result of the test that found a floor. Includes more specific data about the point of impact and surface normal at that point. */
+	UPROPERTY(VisibleInstanceOnly, BlueprintReadOnly, Category = CharacterFloor)
+		FHitResult HitResult;
+
+public:
+
+	FFindFloorResult()
+		: bBlockingHit(false)
+		, bWalkableFloor(false)
+		, bLineTrace(false)
+		, FloorDist(0.f)
+		, LineDist(0.f)
+		, HitResult(1.f)
+	{
+	}
+
+	/** Returns true if the floor result hit a walkable surface. */
+	bool IsWalkableFloor() const
+	{
+		return bBlockingHit && bWalkableFloor;
+	}
+
+	void Clear()
+	{
+		bBlockingHit = false;
+		bWalkableFloor = false;
+		bLineTrace = false;
+		FloorDist = 0.f;
+		LineDist = 0.f;
+		HitResult.Reset(1.f, false);
+	}
+
+	/** Gets the distance to floor, either LineDist or FloorDist. */
+	float GetDistanceToFloor() const
+	{
+		// When the floor distance is set using SetFromSweep, the LineDist value will be reset.
+		// However, when SetLineFromTrace is used, there's no guarantee that FloorDist is set.
+		return bLineTrace ? LineDist : FloorDist;
+	}
+
+	void SetFromSweep(const FHitResult& InHit, const float InSweepFloorDist, const bool bIsWalkableFloor);
+	void SetFromLineTrace(const FHitResult& InHit, const float InSweepFloorDist, const float InLineDist, const bool bIsWalkableFloor);
+};
+
+
+/** Struct updated by StepUp() to return result of final step down, if applicable. */
+struct FStepDownResult
+{
+	uint32 bComputedFloor : 1;		// True if the floor was computed as a result of the step down.
+	FFindFloorResult FloorResult;	// The result of the floor test if the floor was updated.
+
+	FStepDownResult()
+		: bComputedFloor(false)
+	{
+	}
+};
+
+struct FCharacterAsyncOutput
+{
+	virtual ~FCharacterAsyncOutput() {}
+
+	// Character Owner Data
+	FRotator Rotation;
+	int32 JumpCurrentCountPreJump;
+	int32 JumpCurrentCount;
+	float JumpForceTimeRemaining;
+	bool bWasJumping;
+	bool bPressedJump;
+	float JumpKeyHoldTime;
+	bool bClearJumpInput; // If true when applying output we will clear bPressedJump on game thread.
+};
 
 struct FUpdatedComponentAsyncOutput
 {
@@ -37,7 +150,30 @@ struct FUpdatedComponentAsyncOutput
 
 struct FCharacterMovementAsyncOutput : public Chaos::FSimCallbackOutput
 {
-	void Reset() { }
+	using FCharacterOutput = FCharacterAsyncOutput;
+
+	FCharacterMovementAsyncOutput()
+		: FSimCallbackOutput()
+		, bIsValid(false)
+	{
+		Pawn = MakeUnique<FCharacterOutput>();
+	}
+
+	FCharacterMovementAsyncOutput(TUniquePtr<FCharacterOutput>&& PawnOutput)
+		: FSimCallbackOutput()
+		, Pawn(MoveTemp(PawnOutput))
+		, bIsValid(false)
+	{
+	}
+
+	virtual ~FCharacterMovementAsyncOutput() {}
+
+public:
+	void Reset() { bIsValid = false; }
+
+	ENGINE_API void Copy(const FCharacterMovementAsyncOutput& Value);
+	
+	bool IsValid() const { return bIsValid; }
 
 	bool bWasSimulatingRootMotion;
 	EMovementMode MovementMode;
@@ -85,17 +221,10 @@ struct FCharacterMovementAsyncOutput : public Chaos::FSimCallbackOutput
 	UPrimitiveComponent* NewMovementBase; // call SetBase
 	AActor* NewMovementBaseOwner; // make sure this is set whenever base component is. TODO
 
-	// UpdatedComponent data
 	FUpdatedComponentAsyncOutput UpdatedComponent;
+	TUniquePtr<FCharacterAsyncOutput> Pawn;
 
-	// Character Ownner Data
-	FRotator CharacterOwnerRotation; 
-	int32 JumpCurrentCountPreJump;
-	int32 JumpCurrentCount;
-	float JumpForceTimeRemaining;
-	bool bWasJumping;
-	bool bPressedJump;
-	float JumpKeyHoldTime;
+	bool bIsValid;
 };
 
 // Don't read into this part too much it needs to be changed.
@@ -126,7 +255,7 @@ struct FCachedMovementBaseAsyncData
 };
 
 // Data and implementation that lives on movement component's character owner
-struct FCharacterAsyncInput
+struct ENGINE_API FCharacterAsyncInput
 {
 	virtual ~FCharacterAsyncInput() {}
 
@@ -147,10 +276,11 @@ struct FCharacterAsyncInput
 	virtual void ClearJumpInput(float DeltaSeconds, const FCharacterMovementAsyncInput& Input, FCharacterMovementAsyncOutput& Output) const;
 	virtual bool CanJump(const FCharacterMovementAsyncInput& Input, FCharacterMovementAsyncOutput& Output) const;
 	virtual void ResetJumpState(const FCharacterMovementAsyncInput& Input, FCharacterMovementAsyncOutput& Output) const;
+	virtual void OnMovementModeChanged(EMovementMode PrevMovementMode, const FCharacterMovementAsyncInput& Input, FCharacterMovementAsyncOutput& Output, uint8 PreviousCustomMode = 0);
 };
 
 // Represents the UpdatedComponent's state and implementation
-struct FUpdatedComponentAsyncInput
+struct ENGINE_API FUpdatedComponentAsyncInput
 {
 	virtual ~FUpdatedComponentAsyncInput() {}
 
@@ -191,14 +321,40 @@ struct FUpdatedComponentAsyncInput
 	FVector Scale;
 };
 
+
+// This contains inputs from GT that are applied to async sim state before simulation.
+struct FCharacterMovementGTInputs
+{
+	bool bWantsToCrouch;
+	
+	bool bValidMovementMode; // Should we actually copy movement mode?
+	EMovementMode MovementMode;
+
+	// Pawn inputs
+	bool bPressedJump;
+
+	void UpdateOutput(FCharacterMovementAsyncOutput& Output) const
+	{
+		Output.bWantsToCrouch = bWantsToCrouch;
+		if (bValidMovementMode)
+		{
+			Output.MovementMode = MovementMode;
+		}
+		Output.Pawn->bPressedJump = bPressedJump;
+	}
+};
+
 /*
 * Contains all input and implementation required to run async character movement.
 * Base implementation is from CharacterMovementComponent.
 * Contains 'CharacterInput' and 'UpdatedComponentInput' represent data/impl of Character and our UpdatedComponent.
-* All input is const, non-const data goes in output. 'InitialOutput' member is copied to output before sim to initialize non-const data.
+* All input is const, non-const data goes in output. 'AsyncSimState' points to non-const sim state.
 */
-struct FCharacterMovementAsyncInput : public Chaos::FSimCallbackInput
+struct ENGINE_API FCharacterMovementAsyncInput : public Chaos::FSimCallbackInput
 {
+	using FCharacterInput = FCharacterAsyncInput;
+	using FUpdatedComponentInput = FUpdatedComponentAsyncInput;
+
 	// Has this been filled out?
 	bool bInitialized = false;
 
@@ -281,16 +437,37 @@ struct FCharacterMovementAsyncInput : public Chaos::FSimCallbackInput
 	FCollisionQueryParams CapsuleParams;
 	FRandomStream RandomStream;
 
-	// When filling inputs we put initial value of outputs in here.
-	// This initializes callback output at beginning of sim,
-	// and callback output is copied back onto this at end of sim,
-	// so it may be used to initialize subsequent steps using this same input
-	TUniquePtr<FCharacterMovementAsyncOutput> InitialOutput;
+	
+	// This is the latest simulated state of this movement component.
+	TSharedPtr<FCharacterMovementAsyncOutput, ESPMode::ThreadSafe> AsyncSimState;
+
+	FCharacterMovementGTInputs GTInputs;
+
+	virtual ~FCharacterMovementAsyncInput() {}
+	
+	template <typename FMovementInput, typename FMovementOutput> 
+	void Initialize()
+	{
+		CharacterInput = MakeUnique<FMovementInput::FCharacterInput>();
+		UpdatedComponentInput = MakeUnique<FMovementInput::FUpdatedComponentInput>();
+	}
+
+	void Reset()
+	{
+		/* TODO Should actually implement this.*/
+		bInitialized = false;
+		CharacterInput.Reset();
+		UpdatedComponentInput.Reset();
+		AsyncSimState.Reset();
+	}
+
+	void UpdateAsyncStateFromGTInputs_Internal() const
+	{
+		GTInputs.UpdateOutput(*AsyncSimState);
+	}
 
 	// Entry point of async tick
 	void Simulate(const float DeltaSeconds, FCharacterMovementAsyncOutput& Output) const;
-	void Reset() { bInitialized = false; /* TODO, currently writing everything each frame. Should actually implement this.*/ }
-	virtual ~FCharacterMovementAsyncInput() {}
 
 	// TODO organize these
 	virtual void ControlledCharacterMove(const float DeltaSeconds, FCharacterMovementAsyncOutput& Output) const;
@@ -397,4 +574,37 @@ class FCharacterMovementAsyncCallback : public Chaos::TSimCallbackObject<FCharac
 {
 private:
 	virtual void OnPreSimulate_Internal() override;
-};
\ No newline at end of file
+};
+
+
+template <typename FAsyncCallbackInput, typename FAsyncCallbackOutput, typename FAsyncCallback>
+void PreSimulateImpl(FAsyncCallback& Callback)
+{
+	const FAsyncCallbackInput* Input = Callback.GetConsumerInput_Internal();
+	if (Input && Input->bInitialized)
+	{
+		// Update sim state from game thread inputs
+		Input->UpdateAsyncStateFromGTInputs_Internal();
+
+		// Ensure that if we reset jump recently we do not process stale inputs enqueued from game thread.
+		if (Input->AsyncSimState->Pawn->bClearJumpInput)
+		{
+			if (Input->GTInputs.bPressedJump == false)
+			{
+				// Game thread has consumed output clearing jump input, reset our flag so we accept next jump input.
+				Input->AsyncSimState->Pawn->bClearJumpInput = false;
+			}
+			else
+			{
+				// Game thread has not consumed output clearing jump input yet, this is stale input that we do not want
+				Input->AsyncSimState->Pawn->bPressedJump = false;
+			}
+		}
+
+		Input->Simulate(Callback.GetDeltaTime_Internal(), static_cast<FAsyncCallbackOutput&>(*Input->AsyncSimState));
+
+		//  Copy sim state to callback output that will be pushed to game thread
+		FAsyncCallbackOutput& Output = Callback.GetProducerOutputData_Internal();
+		Output.Copy(static_cast<FAsyncCallbackOutput&>(*Input->AsyncSimState));
+	}
+}
