Fix template compilation error in CharacterMovementComponentAsync.h

diff --git a/Engine/Source/Runtime/Engine/Private/Components/CharacterMovementComponent.cpp b/Engine/Source/Runtime/Engine/Private/Components/CharacterMovementComponent.cpp
--- a/Engine/Source/Runtime/Engine/Private/Components/CharacterMovementComponent.cpp
+++ b/Engine/Source/Runtime/Engine/Private/Components/CharacterMovementComponent.cpp
@@ -12306,7 +12306,7 @@ void UCharacterMovementComponent::BuildAsyncInput()
 		FCharacterMovementAsyncInput* Input = AsyncCallback->GetProducerInputData_External();
 		if (Input->bInitialized == false)
 		{
-			Input->Initialize<FCharacterMovementAsyncInput, FCharacterMovementAsyncOutput>();
+			Input->Initialize<FCharacterMovementAsyncInput::FCharacterInput, FCharacterMovementAsyncInput::FUpdatedComponentInput>();
 		}
 
 		if (AsyncSimState.IsValid() == false)
diff --git a/Engine/Source/Runtime/Engine/Public/CharacterMovementComponentAsync.h b/Engine/Source/Runtime/Engine/Public/CharacterMovementComponentAsync.h
--- a/Engine/Source/Runtime/Engine/Public/CharacterMovementComponentAsync.h
+++ b/Engine/Source/Runtime/Engine/Public/CharacterMovementComponentAsync.h
@@ -445,11 +445,11 @@ struct ENGINE_API FCharacterMovementAsyncInput : public Chaos::FSimCallbackInput
 
 	virtual ~FCharacterMovementAsyncInput() {}
 	
-	template <typename FMovementInput, typename FMovementOutput> 
+	template <typename FCharacterInput, typename FUpdatedComponentInput> 
 	void Initialize()
 	{
-		CharacterInput = MakeUnique<FMovementInput::FCharacterInput>();
-		UpdatedComponentInput = MakeUnique<FMovementInput::FUpdatedComponentInput>();
+		CharacterInput = MakeUnique<FCharacterInput>();
+		UpdatedComponentInput = MakeUnique<FUpdatedComponentInput>();
 	}
 
 	void Reset()
