Force CurrentFirstLODIdx to be 0 on static mesh render data when using cooked data in the editor - otherwised, the inlined LODs would be used which will often be low LODs.

diff --git a/Engine/Source/Runtime/Engine/Private/StaticMesh.cpp b/Engine/Source/Runtime/Engine/Private/StaticMesh.cpp
--- a/Engine/Source/Runtime/Engine/Private/StaticMesh.cpp
+++ b/Engine/Source/Runtime/Engine/Private/StaticMesh.cpp
@@ -1578,7 +1578,17 @@ void FStaticMeshRenderData::Serialize(FArchive& Ar, UStaticMesh* Owner, bool bCo
 	}
 #endif
 	Ar << NumInlinedLODs;
-	CurrentFirstLODIdx = LODResources.Num() - NumInlinedLODs;
+
+#if WITH_EDITOR
+	if (bCooked && Ar.IsLoading())
+	{
+		CurrentFirstLODIdx = 0;
+	}
+	else
+#endif
+	{
+		CurrentFirstLODIdx = LODResources.Num() - NumInlinedLODs;
+	}
 
 	if (Ar.IsLoading())
 	{
