Adding some error text to ActorWorld check

diff --git a/Engine/Source/Runtime/Engine/Private/DataChannel.cpp b/Engine/Source/Runtime/Engine/Private/DataChannel.cpp
--- a/Engine/Source/Runtime/Engine/Private/DataChannel.cpp
+++ b/Engine/Source/Runtime/Engine/Private/DataChannel.cpp
@@ -2956,7 +2956,7 @@ int64 UActorChannel::ReplicateActor()
 	check(nullptr != Cast<UPackageMapClient>(Connection->PackageMap));
 
 	const UWorld* const ActorWorld = Actor->GetWorld();
-	check(ActorWorld);
+	checkf(ActorWorld, TEXT("ActorWorld for Actor [%s] is Null"), *GetPathNameSafe(Actor));
 
 #if STATS || ENABLE_STATNAMEDEVENTS
 	UClass* ParentNativeClass = GetParentNativeClass(Actor->GetClass());
