Removing AudioGameplay Proxy subsystem (not needed)
Removing param from AudioEngineSubsystem::GetDeviceHandle (also not needed)

diff --git a/Engine/Source/Runtime/Engine/Private/Subsystems/AudioEngineSubsystem.cpp b/Engine/Source/Runtime/Engine/Private/Subsystems/AudioEngineSubsystem.cpp
--- a/Engine/Source/Runtime/Engine/Private/Subsystems/AudioEngineSubsystem.cpp
+++ b/Engine/Source/Runtime/Engine/Private/Subsystems/AudioEngineSubsystem.cpp
@@ -7,11 +7,14 @@ UAudioEngineSubsystem::UAudioEngineSubsystem()
 {
 }
 
-FAudioDeviceHandle UAudioEngineSubsystem::GetAudioDeviceHandle(Audio::FDeviceId InDeviceID)
+FAudioDeviceHandle UAudioEngineSubsystem::GetAudioDeviceHandle() const
 {
+	const UAudioSubsystemCollectionRoot* SubsystemRoot = Cast<UAudioSubsystemCollectionRoot>(GetOuter());
+	check(SubsystemRoot);
+
 	if (FAudioDeviceManager* AudioDeviceManager = FAudioDeviceManager::Get())
 	{
-		return AudioDeviceManager->GetAudioDevice(InDeviceID);
+		return AudioDeviceManager->GetAudioDevice(SubsystemRoot->GetAudioDeviceID());
 	}
 
 	return FAudioDeviceHandle();
diff --git a/Engine/Source/Runtime/Engine/Public/Subsystems/AudioEngineSubsystem.h b/Engine/Source/Runtime/Engine/Public/Subsystems/AudioEngineSubsystem.h
--- a/Engine/Source/Runtime/Engine/Public/Subsystems/AudioEngineSubsystem.h
+++ b/Engine/Source/Runtime/Engine/Public/Subsystems/AudioEngineSubsystem.h
@@ -23,7 +23,10 @@ public:
 	/** Set the ID of the owning audio device */
 	void FORCEINLINE SetAudioDeviceID(Audio::FDeviceId DeviceID) { OwningDeviceID = DeviceID; }
 
-protected: 
+	/** Get the ID of the owning audio device */
+	Audio::FDeviceId GetAudioDeviceID() const { return OwningDeviceID; }
+
+protected:
 
 	Audio::FDeviceId OwningDeviceID = INDEX_NONE;
 };
@@ -41,5 +44,5 @@ public:
 
 	UAudioEngineSubsystem();
 
-	virtual FAudioDeviceHandle GetAudioDeviceHandle(Audio::FDeviceId InDeviceID);
+	virtual FAudioDeviceHandle GetAudioDeviceHandle() const final;
 };
