#UE Fix a determinism problems when exporting data tables to json due to writing float values as doubles, which can have precision problems between runs.

diff --git a/Engine/Source/Runtime/Engine/Private/DataTableJSON.cpp b/Engine/Source/Runtime/Engine/Private/DataTableJSON.cpp
--- a/Engine/Source/Runtime/Engine/Private/DataTableJSON.cpp
+++ b/Engine/Source/Runtime/Engine/Private/DataTableJSON.cpp
@@ -235,6 +235,11 @@ bool TDataTableExporterJSON<CharType>::WriteStructEntry(const void* InRowData, c
 			const int64 PropertyValue = NumProp->GetSignedIntPropertyValue(InPropertyData);
 			JsonWriter->WriteValue(Identifier, PropertyValue);
 		}
+		else if (NumProp->IsA(FFloatProperty::StaticClass()))
+		{
+			const float PropertyValue = (float)NumProp->GetFloatingPointPropertyValue(InPropertyData);
+			JsonWriter->WriteValue(Identifier, PropertyValue);
+		}
 		else
 		{
 			const double PropertyValue = NumProp->GetFloatingPointPropertyValue(InPropertyData);
