[Debug] Adding distance option to 'minimal' sound cue debug output

diff --git a/Engine/Source/Runtime/Engine/Private/AudioDebug.cpp b/Engine/Source/Runtime/Engine/Private/AudioDebug.cpp
--- a/Engine/Source/Runtime/Engine/Private/AudioDebug.cpp
+++ b/Engine/Source/Runtime/Engine/Private/AudioDebug.cpp
@@ -127,6 +127,14 @@ FAutoConsoleVariableRef CVarAudioSoundCueDebugShowPath(
 	TEXT("0: Not Enabled, 1: Enabled"),
 	ECVF_Default);
 
+static int32 SoundCueDebugShowDistanceCVar = 0;
+FAutoConsoleVariableRef CVarAudioSoundCueDebugShowDistance(
+	TEXT("au.Debug.Soundcues.ShowDistance"),
+	SoundCueDebugShowDistanceCVar,
+	TEXT("Display distance of sound cue when enabled.\n")
+	TEXT("0: Not Enabled, 1: Enabled"),
+	ECVF_Default);
+
 static int32 SoundCueDebugMinimalCVar = 0;
 FAutoConsoleVariableRef CVarAudioSoundCueDebugMinimal(
 	TEXT("au.Debug.SoundCues.Minimal"),
@@ -1591,6 +1599,14 @@ namespace Audio
 						CurrentX += NumberSpacing;
 						Canvas->DrawShadowedString(CurrentX, Y, *Volume, StatsFont, bMutedOrSoloed ? Color : FColor::Orange);
 						CurrentX += NumberSpacing;
+
+						if (SoundCueDebugShowDistanceCVar)
+						{
+							const FString DistanceText = FString::Printf(TEXT("%6.2f "), StatSoundInfo.Distance);
+							Canvas->DrawShadowedString(CurrentX, Y, *DistanceText, StatsFont, bMutedOrSoloed ? Color : FColor::White);
+							CurrentX += (NumberSpacing * 2);
+						}
+
 						Canvas->DrawShadowedString(CurrentX, Y, *SplitNames[SoundNameIndex].Key, StatsFont, bMutedOrSoloed ? Color : FColor(0, 255, 255));
 						CurrentX += (MaxNameLength * CharSpacing);
 						Canvas->DrawShadowedString(CurrentX, Y, *StatSoundInfo.SoundClassName.ToString(), StatsFont, bMutedOrSoloed ? Color : FColor::Yellow);
@@ -2439,7 +2455,11 @@ namespace Audio
 						FAudioStats::FStatSoundInfo& StatSoundInfo = StatSoundInfos.Last();
 						StatSoundInfo.SoundName = SoundBase->GetName();
 						StatSoundInfo.SoundPath = SoundBase->GetPathName();
-						StatSoundInfo.Distance = AudioDevice.GetDistanceToNearestListener(ActiveSound->Transform.GetTranslation());
+						StatSoundInfo.Distance = 0.f;
+						if (ActiveSound->bAllowSpatialization)
+						{
+							StatSoundInfo.Distance = AudioDevice.GetDistanceToNearestListener(ActiveSound->Transform.GetTranslation());
+						}
 						StatSoundInfo.PlaybackTime = ActiveSound->PlaybackTime;
 						StatSoundInfo.Priority = ActiveSound->GetHighestPriority();
 						StatSoundInfo.PlaybackTimeNonVirtualized = ActiveSound->PlaybackTimeNonVirtualized;
