nDisplay: Linux support (Engine)

diff --git a/Engine/Source/Runtime/Core/Public/HAL/Platform.h b/Engine/Source/Runtime/Core/Public/HAL/Platform.h
--- a/Engine/Source/Runtime/Core/Public/HAL/Platform.h
+++ b/Engine/Source/Runtime/Core/Public/HAL/Platform.h
@@ -360,6 +360,10 @@
 	#define PLATFORM_SUPPORTS_TBB 0
 #endif
 
+#ifndef PLATFORM_SUPPORTS_BORDERLESS_WINDOW
+	#define PLATFORM_SUPPORTS_BORDERLESS_WINDOW		0
+#endif
+
 #ifndef PLATFORM_MAX_CACHED_SYNC_FILE_HANDLES_PER_GENERIC_ASYNC_FILE_HANDLE
 	#define PLATFORM_MAX_CACHED_SYNC_FILE_HANDLES_PER_GENERIC_ASYNC_FILE_HANDLE 4
 #endif
diff --git a/Engine/Source/Runtime/Core/Public/Linux/LinuxPlatform.h b/Engine/Source/Runtime/Core/Public/Linux/LinuxPlatform.h
--- a/Engine/Source/Runtime/Core/Public/Linux/LinuxPlatform.h
+++ b/Engine/Source/Runtime/Core/Public/Linux/LinuxPlatform.h
@@ -9,3 +9,5 @@
 #include "Unix/UnixPlatform.h"
 
 #define PLATFORM_GLOBAL_LOG_CATEGORY			LogLinux
+
+#define PLATFORM_SUPPORTS_BORDERLESS_WINDOW		1
diff --git a/Engine/Source/Runtime/Core/Public/Windows/WindowsPlatform.h b/Engine/Source/Runtime/Core/Public/Windows/WindowsPlatform.h
--- a/Engine/Source/Runtime/Core/Public/Windows/WindowsPlatform.h
+++ b/Engine/Source/Runtime/Core/Public/Windows/WindowsPlatform.h
@@ -90,6 +90,8 @@ typedef FWindowsPlatformTypes FPlatformTypes;
 
 #define PLATFORM_GLOBAL_LOG_CATEGORY						LogWindows
 
+#define PLATFORM_SUPPORTS_BORDERLESS_WINDOW					1
+
 #define WINDOWS_USE_FEATURE_APPLICATIONMISC_CLASS			1
 #define WINDOWS_USE_FEATURE_PLATFORMPROCESS_CLASS			1
 #define WINDOWS_USE_FEATURE_PLATFORMMISC_CLASS				1
diff --git a/Engine/Source/Runtime/Engine/Private/GameEngine.cpp b/Engine/Source/Runtime/Engine/Private/GameEngine.cpp
--- a/Engine/Source/Runtime/Engine/Private/GameEngine.cpp
+++ b/Engine/Source/Runtime/Engine/Private/GameEngine.cpp
@@ -525,7 +525,7 @@ TSharedRef<SWindow> UGameEngine::CreateGameWindow()
 	};
 
 	const bool bShouldPreserveAspectRatio = GetProjectSettingBool(TEXT("bShouldWindowPreserveAspectRatio"), true);
-	const bool bUseBorderlessWindow = GetProjectSettingBool(TEXT("bUseBorderlessWindow"), false) && PLATFORM_WINDOWS;
+	const bool bUseBorderlessWindow = GetProjectSettingBool(TEXT("bUseBorderlessWindow"), false) && PLATFORM_SUPPORTS_BORDERLESS_WINDOW;
 	const bool bAllowWindowResize = GetProjectSettingBool(TEXT("bAllowWindowResize"), true);
 	const bool bAllowClose = GetProjectSettingBool(TEXT("bAllowClose"), true);
 	const bool bAllowMaximize = GetProjectSettingBool(TEXT("bAllowMaximize"), true);
