Allow composite texture size to be different from source size, as long as aspect ratio matches

diff --git a/Engine/Source/Runtime/Engine/Private/TextureDerivedDataTask.cpp b/Engine/Source/Runtime/Engine/Private/TextureDerivedDataTask.cpp
--- a/Engine/Source/Runtime/Engine/Private/TextureDerivedDataTask.cpp
+++ b/Engine/Source/Runtime/Engine/Private/TextureDerivedDataTask.cpp
@@ -468,9 +468,10 @@ FTextureCacheDerivedDataWorker::FTextureCacheDerivedDataWorker(
 	TextureData.Init(Texture, BuildSettingsPerLayer.GetData(), bAllowAsyncLoading);
 	if (Texture.CompositeTexture && Texture.CompositeTextureMode != CTM_Disabled)
 	{
-		bool bCompositeTextureMatching = Texture.CompositeTexture->Source.GetNumBlocks() == Texture.Source.GetNumBlocks();
+		bool bMatchingBlocks = Texture.CompositeTexture->Source.GetNumBlocks() == Texture.Source.GetNumBlocks();
+		bool bMatchingAspectRatio = true;
 		bool bOnlyPowerOfTwoSize = true;
-		if (bCompositeTextureMatching)
+		if (bMatchingBlocks)
 		{
 			for (int32 BlockIdx = 0; BlockIdx < Texture.Source.GetNumBlocks(); ++BlockIdx)
 			{
@@ -479,25 +480,26 @@ FTextureCacheDerivedDataWorker::FTextureCacheDerivedDataWorker(
 				FTextureSourceBlock CompositeTextureBlock;
 				Texture.CompositeTexture->Source.GetBlock(BlockIdx, CompositeTextureBlock);
 
-				bCompositeTextureMatching = bCompositeTextureMatching
-					&& TextureBlock.BlockX == CompositeTextureBlock.BlockX
-					&& TextureBlock.BlockY == CompositeTextureBlock.BlockY
-					&& TextureBlock.SizeX == CompositeTextureBlock.SizeX
-					&& TextureBlock.SizeY == CompositeTextureBlock.SizeY;
-
+				bMatchingBlocks = bMatchingBlocks && TextureBlock.BlockX == CompositeTextureBlock.BlockX && TextureBlock.BlockY == CompositeTextureBlock.BlockY;
+				bMatchingAspectRatio = bMatchingAspectRatio && TextureBlock.SizeX * CompositeTextureBlock.SizeY == TextureBlock.SizeY * CompositeTextureBlock.SizeX;
 				bOnlyPowerOfTwoSize = bOnlyPowerOfTwoSize && FMath::IsPowerOfTwo(TextureBlock.SizeX) && FMath::IsPowerOfTwo(TextureBlock.SizeY);
 			}
 		}
-		if (!bOnlyPowerOfTwoSize)
+
+		if (!bMatchingBlocks)
+		{
+			UE_LOG(LogTexture, Warning, TEXT("Issue while building %s : Composite texture resolution/UDIMs do not match. Composite texture will be ignored"), *Texture.GetPathName());
+		}
+		else if (!bOnlyPowerOfTwoSize)
 		{
 			UE_LOG(LogTexture, Warning, TEXT("Issue while building %s : Some blocks (UDIMs) have a non power of two size. Composite texture will be ignored"), *Texture.GetPathName());
 		}
-		else if (!bCompositeTextureMatching)
+		else if (!bMatchingAspectRatio)
 		{
-			UE_LOG(LogTexture, Warning, TEXT("Issue while building %s : Composite texture resolution/UDIMs do not match. Composite texture will be ignored"), *Texture.GetPathName());
+			UE_LOG(LogTexture, Warning, TEXT("Issue while building %s : Some blocks (UDIMs) have mismatched aspect ratio. Composite texture will be ignored"), *Texture.GetPathName());
 		}
 
-		if (bCompositeTextureMatching && bOnlyPowerOfTwoSize)
+		if (bMatchingBlocks && bMatchingAspectRatio && bOnlyPowerOfTwoSize)
 		{
 			CompositeTextureData.Init(*Texture.CompositeTexture, BuildSettingsPerLayer.GetData(), bAllowAsyncLoading);
 		}
