Fixed issue with bAllowCPUAccess not being honored when set to true

diff --git a/Engine/Source/Runtime/Engine/Classes/Engine/StaticMesh.h b/Engine/Source/Runtime/Engine/Classes/Engine/StaticMesh.h
--- a/Engine/Source/Runtime/Engine/Classes/Engine/StaticMesh.h
+++ b/Engine/Source/Runtime/Engine/Classes/Engine/StaticMesh.h
@@ -1325,6 +1325,7 @@ public:
 			, bBuildSimpleCollision(false)
 			, bCommitMeshDescription(true)
 			, bFastBuild(false)
+			, bAllowCpuAccess(false)
 		{}
 
 		/**
@@ -1354,6 +1355,11 @@ public:
 		 * Set to false by default.
 		 */
 		bool bFastBuild;
+
+		/**
+		 * Ored with the value of bAllowCpuAccess on the static mesh. Set to false by default.
+		 */
+		bool bAllowCpuAccess;
 	};
 
 	/**
diff --git a/Engine/Source/Runtime/Engine/Private/StaticMesh.cpp b/Engine/Source/Runtime/Engine/Private/StaticMesh.cpp
--- a/Engine/Source/Runtime/Engine/Private/StaticMesh.cpp
+++ b/Engine/Source/Runtime/Engine/Private/StaticMesh.cpp
@@ -5732,6 +5732,12 @@ bool UStaticMesh::BuildFromMeshDescriptions(const TArray<const FMeshDescription*
 	SetRenderData(MakeUnique<FStaticMeshRenderData>());
 	GetRenderData()->AllocateLODResources(MeshDescriptions.Num());
 
+	FStaticMeshLODResourcesArray& LODResourcesArray = GetRenderData()->LODResources;
+	for (int32 LODIndex = 0; LODIndex < LODResourcesArray.Num(); ++LODIndex)
+	{
+		LODResourcesArray[LODIndex].IndexBuffer.TrySetAllowCPUAccess(bAllowCPUAccess | Params.bAllowCpuAccess);
+	}
+
 	// Build render data from each mesh description
 
 #if WITH_EDITOR
diff --git a/Engine/Source/Runtime/Engine/Public/RawIndexBuffer.h b/Engine/Source/Runtime/Engine/Public/RawIndexBuffer.h
--- a/Engine/Source/Runtime/Engine/Public/RawIndexBuffer.h
+++ b/Engine/Source/Runtime/Engine/Public/RawIndexBuffer.h
@@ -134,6 +134,19 @@ public:
 	 */
 	ENGINE_API FRawStaticIndexBuffer(bool InNeedsCPUAccess=false);
 
+	/**
+	 * Copy everything, keeping reference to the same RHI resources.
+	 * Returns true if input value for bAllowCPUAccess will be honored:
+	 *		- if bAllowCPUAccess is true then no indices have already been cached or IndexStorage is not empty
+	 *		- or if bAllowCPUAccess is false
+	 */
+	bool TrySetAllowCPUAccess(bool bAllowCPUAccess)
+	{
+		IndexStorage.SetAllowCPUAccess(bAllowCPUAccess);
+
+		return !bAllowCPUAccess || CachedNumIndices == 0 || IndexStorage.Num() > 0;
+	}
+
 	/**
 	 * Sets a single index value.  Consider using SetIndices() instead if you're setting a lot of indices.
 	 * @param	At	The index of the index to set
