UE4.27 - Deprecate unused function UStaticMeshComponent::UpdatePreCulledData

diff --git a/Engine/Source/Runtime/Engine/Classes/Components/StaticMeshComponent.h b/Engine/Source/Runtime/Engine/Classes/Components/StaticMeshComponent.h
--- a/Engine/Source/Runtime/Engine/Classes/Components/StaticMeshComponent.h
+++ b/Engine/Source/Runtime/Engine/Classes/Components/StaticMeshComponent.h
@@ -646,6 +646,7 @@ public:
 	void RemoveInstanceVertexColors();
 #endif
 
+	UE_DEPRECATED(4.27, "This function is no longer used")
 	void UpdatePreCulledData(int32 LODIndex, const TArray<uint32>& PreCulledData, const TArray<int32>& NumTrianglesPerSection);
 
 #if WITH_EDITORONLY_DATA
