Fixed compilation warnings

diff --git a/Engine/Source/Runtime/Engine/Private/StaticMesh.cpp b/Engine/Source/Runtime/Engine/Private/StaticMesh.cpp
--- a/Engine/Source/Runtime/Engine/Private/StaticMesh.cpp
+++ b/Engine/Source/Runtime/Engine/Private/StaticMesh.cpp
@@ -5735,7 +5735,7 @@ bool UStaticMesh::BuildFromMeshDescriptions(const TArray<const FMeshDescription*
 	FStaticMeshLODResourcesArray& LODResourcesArray = GetRenderData()->LODResources;
 	for (int32 LODIndex = 0; LODIndex < LODResourcesArray.Num(); ++LODIndex)
 	{
-		LODResourcesArray[LODIndex].IndexBuffer.TrySetAllowCPUAccess(bAllowCPUAccess | Params.bAllowCpuAccess);
+		LODResourcesArray[LODIndex].IndexBuffer.TrySetAllowCPUAccess(bAllowCPUAccess || Params.bAllowCpuAccess);
 	}
 
 	// Build render data from each mesh description
