Add a command to dump flags for everything in the replayout map
Add a LexToString for ERepLayoutFlags
Add a replayout flag to indicate when there are no replicated properties (including disabled)

diff --git a/Engine/Source/Runtime/Engine/Classes/Engine/NetDriver.h b/Engine/Source/Runtime/Engine/Classes/Engine/NetDriver.h
--- a/Engine/Source/Runtime/Engine/Classes/Engine/NetDriver.h
+++ b/Engine/Source/Runtime/Engine/Classes/Engine/NetDriver.h
@@ -1402,6 +1402,7 @@ public:
 	bool HandleNetDumpServerRPCCommand(const TCHAR* Cmd, FOutputDevice& Ar);
 	bool HandleNetDumpDormancy(const TCHAR* Cmd, FOutputDevice& Ar);
 	bool HandleDumpSubObjectsCommand(const TCHAR* Cmd, FOutputDevice& Ar);
+	bool HandleDumpRepLayoutFlagsCommand(const TCHAR* Cmd, FOutputDevice& Ar);
 #endif
 
 	void HandlePacketLossBurstCommand( int32 DurationInMilliseconds );
diff --git a/Engine/Source/Runtime/Engine/Private/NetDriver.cpp b/Engine/Source/Runtime/Engine/Private/NetDriver.cpp
--- a/Engine/Source/Runtime/Engine/Private/NetDriver.cpp
+++ b/Engine/Source/Runtime/Engine/Private/NetDriver.cpp
@@ -3016,6 +3016,71 @@ bool UNetDriver::HandleDumpSubObjectsCommand(const TCHAR* Cmd, FOutputDevice& Ar
 #endif
 	return true;
 }
+
+bool UNetDriver::HandleDumpRepLayoutFlagsCommand(const TCHAR* Cmd, FOutputDevice& Ar)
+{
+#if WITH_SERVER_CODE
+
+#if WITH_PUSH_MODEL
+	const bool bPartialPushDump = FParse::Command(&Cmd, TEXT("PARTIALPUSHMODEL"));
+#endif
+
+	Ar.Logf(TEXT("Logging replayout flags for: %s"), *GetNameSafe(this));
+	Ar.Logf(TEXT(""));
+
+	for (auto LayoutIt = RepLayoutMap.CreateConstIterator(); LayoutIt; ++LayoutIt)
+	{
+		const TSharedPtr<FRepLayout>& RepLayout = LayoutIt.Value();
+		if (RepLayout.IsValid())
+		{
+			FString FlagStr;
+
+			const ERepLayoutFlags Flags = RepLayout->GetFlags();
+
+			for (uint32 i = 0; i < sizeof(ERepLayoutFlags) * 8; ++i)
+			{
+				ERepLayoutFlags Flag = (ERepLayoutFlags)(1 << i);
+
+				if (EnumHasAnyFlags(Flags, Flag))
+				{
+					FlagStr += (FlagStr.IsEmpty() ? TEXT("") : TEXT("|"));
+					FlagStr += LexToString(Flag);
+				}
+			}
+
+			if (FlagStr.IsEmpty())
+			{
+				FlagStr = TEXT("None");
+			}
+
+			Ar.Logf(TEXT("  Owner: %s"), *GetFullNameSafe(RepLayout->GetOwner()));
+			Ar.Logf(TEXT("    Flags: %s"), *FlagStr);
+
+#if WITH_PUSH_MODEL
+			if (bPartialPushDump && EnumHasAnyFlags(Flags, ERepLayoutFlags::PartialPushSupport))
+			{
+				int32 ParentCount = RepLayout->GetNumParents();
+
+				Ar.Logf(TEXT("    Parent Properties: %d Non-push:"), ParentCount);
+
+				for (int32 i = 0; i < ParentCount; ++i)
+				{
+					if ((RepLayout->GetParentCondition(i) != COND_Never) && !RepLayout->IsPushModelProperty(i) && (RepLayout->GetParentArrayIndex(i) == 0))
+					{
+						if (const FProperty* Property = RepLayout->GetParentProperty(i))
+						{
+							Ar.Logf(TEXT("      Property: %s Class: %s"), *Property->GetName(), *GetNameSafe(Property->GetOwnerClass()));
+						}
+					}
+				}
+			}
+#endif
+		}
+	}
+#endif
+	return true;
+}
+
 #endif // !UE_BUILD_SHIPPING
 
 void UNetDriver::HandlePacketLossBurstCommand( int32 DurationInMilliseconds )
@@ -3136,6 +3201,11 @@ bool UNetDriver::Exec( UWorld* InWorld, const TCHAR* Cmd, FOutputDevice& Ar )
 		HandleDumpSubObjectsCommand(Cmd, Ar);
 		return true;
 	}
+	else if (FParse::Command(&Cmd, TEXT("DUMPREPLAYOUTFLAGS")))
+	{
+		HandleDumpRepLayoutFlagsCommand(Cmd, Ar);
+		return true;
+	}
 	else
 #endif // !UE_BUILD_SHIPPING
 	{
diff --git a/Engine/Source/Runtime/Engine/Private/RepLayout.cpp b/Engine/Source/Runtime/Engine/Private/RepLayout.cpp
--- a/Engine/Source/Runtime/Engine/Private/RepLayout.cpp
+++ b/Engine/Source/Runtime/Engine/Private/RepLayout.cpp
@@ -6096,16 +6096,16 @@ void FRepLayout::InitFromClass(
 	}
 
 	if (bIsObjectActor)
-		{
-			// We handle remote role specially, since it can change between connections when downgraded
-			// So we force it on the conditional list
+	{
+		// We handle remote role specially, since it can change between connections when downgraded
+		// So we force it on the conditional list
 		FRepParentCmd& RemoteRoleParent = Parents[(int32)AActor::ENetFields_Private::RemoteRole];
 		if (RemoteRoleParent.Condition != COND_Never)
 		{
 			if (COND_None != RemoteRoleParent.Condition)
 			{
 				UE_LOG(LogRep, Warning, TEXT("FRepLayout::InitFromClass: Forcing replication of RemoteRole. Owner=%s"), *InObjectClass->GetPathName());
-		}
+			}
 
 			Parents[(int32)AActor::ENetFields_Private::RemoteRole].Flags |= ERepParentFlags::IsConditional;
 			Parents[(int32)AActor::ENetFields_Private::RemoteRole].Condition = COND_None;
@@ -6121,6 +6121,11 @@ void FRepLayout::InitFromClass(
 	}
 #endif
 
+	if (NumberOfLifetimeProperties == 0)
+	{
+		Flags |= ERepLayoutFlags::NoReplicatedProperties;
+	}
+
 	if (!ServerConnection || EnumHasAnyFlags(CreateFlags, ECreateRepLayoutFlags::MaySendProperties))
 	{
 		BuildHandleToCmdIndexTable_r(0, Cmds.Num() - 1, BaseHandleToCmdIndex);
@@ -8097,6 +8102,26 @@ FRepStateStaticBuffer::~FRepStateStaticBuffer()
 	}
 }
 
+const TCHAR* LexToString(ERepLayoutFlags Flag)
+{
+	switch (Flag)
+	{
+	case ERepLayoutFlags::IsActor:
+		return TEXT("IsActor");
+	case ERepLayoutFlags::PartialPushSupport:
+		return TEXT("PartialPushSupport");
+	case ERepLayoutFlags::FullPushSupport:
+		return TEXT("FullPushSupport");
+	case ERepLayoutFlags::HasObjectOrNetSerializeProperties:
+		return TEXT("HasObjectOrNetSerializeProperties");
+	case ERepLayoutFlags::NoReplicatedProperties:
+		return TEXT("NoReplicatedProperties");
+	default:
+		check(false);
+		return TEXT("Unknown");
+	}
+}
+
 #define REPDATATYPE_SPECIALIZATION(DstType, SrcType) \
 template bool FRepLayout::DiffStableProperties(TArray<FProperty*>*, TArray<UObject*>*, TRepDataBuffer<DstType>, TConstRepDataBuffer<SrcType>) const; \
 template bool FRepLayout::DiffProperties(TArray<FProperty*>*, TRepDataBuffer<DstType>, TConstRepDataBuffer<SrcType>, const EDiffPropertiesFlags) const;
diff --git a/Engine/Source/Runtime/Engine/Public/Net/RepLayout.h b/Engine/Source/Runtime/Engine/Public/Net/RepLayout.h
--- a/Engine/Source/Runtime/Engine/Public/Net/RepLayout.h
+++ b/Engine/Source/Runtime/Engine/Public/Net/RepLayout.h
@@ -1085,9 +1085,12 @@ enum class ERepLayoutFlags : uint8
 	PartialPushSupport					= (1 << 1),	//! This RepLayout has some properties that use Push Model and some that don't.
 	FullPushSupport						= (1 << 2),	//! All properties in this RepLayout use Push Model.
 	HasObjectOrNetSerializeProperties	= (1 << 3),	//! Will be set for any RepLayout that contains Object or Net Serialize property commands.
+	NoReplicatedProperties				= (1 << 4), //! Will be set if the RepLayout has no lifetime properties, or they are all disabled.
 };
 ENUM_CLASS_FLAGS(ERepLayoutFlags);
 
+const TCHAR* LexToString(ERepLayoutFlags Flag);
+
 enum class ERepLayoutResult
 {
 	Success,	// Operation succeeded
@@ -1521,6 +1524,28 @@ public:
 		return Parents.Num();
 	}
 
+	const FProperty* GetParentProperty(int32 Index) const
+	{ 
+		return Parents.IsValidIndex(Index) ? Parents[Index].Property : nullptr;
+	}
+
+	const int32 GetParentArrayIndex(int32 Index) const
+	{
+		return Parents.IsValidIndex(Index) ? Parents[Index].ArrayIndex : 0;
+	}
+
+	const int32 GetParentCondition(int32 Index) const
+	{
+		return Parents.IsValidIndex(Index) ? Parents[Index].Condition : COND_None;
+	}
+
+#if WITH_PUSH_MODEL
+	const bool IsPushModelProperty(int32 Index) const
+	{
+		return PushModelProperties.IsValidIndex(Index) ? PushModelProperties[Index] : false;
+	}
+#endif
+
 	void CountBytes(FArchive& Ar) const;
 
 private:
