Fix collision physic data to use the proper material index if there is a valid LodMaterialMap for the section

diff --git a/Engine/Source/Runtime/Engine/Classes/Engine/SkeletalMesh.h b/Engine/Source/Runtime/Engine/Classes/Engine/SkeletalMesh.h
--- a/Engine/Source/Runtime/Engine/Classes/Engine/SkeletalMesh.h
+++ b/Engine/Source/Runtime/Engine/Classes/Engine/SkeletalMesh.h
@@ -2743,6 +2743,7 @@ public:
 	{
 		return true;
 	}
+	virtual void GetMeshId(FString& OutMeshId) override { OutMeshId = TEXT("3FC28DC87B814E08BA852C92D18D41D4"); }
 	//~ End Interface_CollisionDataProvider Interface
 
 	//~ Begin IInterface_AssetUserData Interface
diff --git a/Engine/Source/Runtime/Engine/Private/SkeletalMesh.cpp b/Engine/Source/Runtime/Engine/Private/SkeletalMesh.cpp
--- a/Engine/Source/Runtime/Engine/Private/SkeletalMesh.cpp
+++ b/Engine/Source/Runtime/Engine/Private/SkeletalMesh.cpp
@@ -4682,8 +4682,13 @@ bool USkeletalMesh::GetPhysicsTriMeshData(FTriMeshCollisionData* CollisionData,
 	FSkeletalMeshRenderData* SkelMeshRenderData = GetResourceForRendering();
 	const FSkeletalMeshLODRenderData& LODData = SkelMeshRenderData->LODRenderData[0];
 
-		// Copy all verts into collision vertex buffer.
-		CollisionData->Vertices.Empty();
+	const TArray<int32>* MaterialMapPtr = nullptr;
+	if (const FSkeletalMeshLODInfo* LODZeroInfo = GetLODInfo(0))
+	{
+		MaterialMapPtr = &LODZeroInfo->LODMaterialMap;
+	}
+	// Copy all verts into collision vertex buffer.
+	CollisionData->Vertices.Empty();
 	CollisionData->Vertices.AddUninitialized(LODData.GetNumVertices());
 
 	for (uint32 VertIdx = 0; VertIdx < LODData.GetNumVertices(); ++VertIdx)
@@ -4706,8 +4711,19 @@ bool USkeletalMesh::GetPhysicsTriMeshData(FTriMeshCollisionData* CollisionData,
 		for (int32 SectionIndex = 0; SectionIndex < LODData.RenderSections.Num(); ++SectionIndex)
 		{
 			const FSkelMeshRenderSection& Section = LODData.RenderSections[SectionIndex];
-
 			const uint32 OnePastLastIndex = Section.BaseIndex + Section.NumTriangles * 3;
+			uint16 MaterialIndex = Section.MaterialIndex;
+			if (MaterialMapPtr)
+			{
+				if (MaterialMapPtr->IsValidIndex(SectionIndex))
+				{
+					const uint16 RemapMaterialIndex = static_cast<uint16>((*MaterialMapPtr)[SectionIndex]);
+					if (GetMaterials().IsValidIndex(RemapMaterialIndex))
+					{
+						MaterialIndex = RemapMaterialIndex;
+					}
+				}
+			}
 
 			for (uint32 i = Section.BaseIndex; i < OnePastLastIndex; i += 3)
 			{
@@ -4716,7 +4732,7 @@ bool USkeletalMesh::GetPhysicsTriMeshData(FTriMeshCollisionData* CollisionData,
 				TriIndex.v2 = Indices[i + 2];
 
 				CollisionData->Indices.Add(TriIndex);
-				CollisionData->MaterialIndices.Add(Section.MaterialIndex);
+				CollisionData->MaterialIndices.Add(MaterialIndex);
 			}
 		}
 	}
