Removing deprecated property

diff --git a/Engine/Source/Runtime/Engine/Classes/Sound/SoundClass.h b/Engine/Source/Runtime/Engine/Classes/Sound/SoundClass.h
--- a/Engine/Source/Runtime/Engine/Classes/Sound/SoundClass.h
+++ b/Engine/Source/Runtime/Engine/Classes/Sound/SoundClass.h
@@ -72,11 +72,6 @@ struct FSoundClassProperties
 	UPROPERTY(EditAnywhere, BlueprintReadOnly, Category = General)
 	float AttenuationDistanceScale;
 
-	/** The amount of stereo sounds to bleed to the rear speakers */
-	UE_DEPRECATED(4.25, "Stereo Bleed is no longer supported.")
-	UPROPERTY(EditAnywhere, BlueprintReadOnly, Category = Routing, meta = (DeprecatedProperty, DeprecationMessage = "Stereo Bleed no longer supported."))
-	float StereoBleed = 0.f;
-
 	/** The amount of a sound to bleed to the LFE channel */
 	UPROPERTY(EditAnywhere, BlueprintReadOnly, Category = Routing, meta = (DisplayName = "LFE Bleed"))
 	float LFEBleed;
