Cache material attributes that are connected and use the bitmask to more intelligently cache shaders.

- Update all Has*Connected() functions to check the cached bitmask to see if that specific material attribute is connected.
- Update FMaterialResource::Has*Connected() functions to not blindly check if _Any_  material attribute is connected.  This was caching too many shaders this way.

diff --git a/Engine/Source/Runtime/Engine/Classes/Materials/Material.h b/Engine/Source/Runtime/Engine/Classes/Materials/Material.h
--- a/Engine/Source/Runtime/Engine/Classes/Materials/Material.h
+++ b/Engine/Source/Runtime/Engine/Classes/Materials/Material.h
@@ -1760,15 +1760,17 @@ public:
 	void SaveShaderStableKeys(const class ITargetPlatform* TP);
 	ENGINE_API virtual void SaveShaderStableKeysInner(const class ITargetPlatform* TP, const struct FStableShaderKeyAndValue& SaveKeyVal) override;
 
-	bool HasBaseColorConnected() const { return BaseColor.IsConnected(); }
-	bool HasRoughnessConnected() const { return Roughness.IsConnected(); }
-	bool HasAmbientOcclusionConnected() const { return AmbientOcclusion.IsConnected(); }
-	bool HasNormalConnected() const { return Normal.IsConnected(); }
-	bool HasSpecularConnected() const { return Specular.IsConnected(); }
-	bool HasMetallicConnected() const { return Metallic.IsConnected(); }
-	bool HasEmissiveColorConnected() const { return EmissiveColor.IsConnected(); }
-	bool HasAnisotropyConnected() const { return Anisotropy.IsConnected(); }
+	bool HasBaseColorConnected() const { return BaseColor.IsConnected() || GetCachedExpressionData().IsMaterialAttributePropertyConnected(MP_BaseColor); }
+	bool HasRoughnessConnected() const { return Roughness.IsConnected() || GetCachedExpressionData().IsMaterialAttributePropertyConnected(MP_Roughness); }
+	bool HasAmbientOcclusionConnected() const { return AmbientOcclusion.IsConnected() || GetCachedExpressionData().IsMaterialAttributePropertyConnected(MP_AmbientOcclusion); }
+	bool HasNormalConnected() const { return Normal.IsConnected() || GetCachedExpressionData().IsMaterialAttributePropertyConnected(MP_Normal); }
+	bool HasSpecularConnected() const { return Specular.IsConnected() || GetCachedExpressionData().IsMaterialAttributePropertyConnected(MP_Specular); }
+	bool HasMetallicConnected() const { return Metallic.IsConnected() || GetCachedExpressionData().IsMaterialAttributePropertyConnected(MP_Metallic); }
+	bool HasEmissiveColorConnected() const { return EmissiveColor.IsConnected() || GetCachedExpressionData().IsMaterialAttributePropertyConnected(MP_EmissiveColor); }
+	bool HasAnisotropyConnected() const { return Anisotropy.IsConnected() || GetCachedExpressionData().IsMaterialAttributePropertyConnected(MP_Anisotropy); }
 	bool HasStrataFrontMaterialConnected() const { return FrontMaterial.IsConnected(); }
+	bool HasVertexPositionOffsetConnected() const { return WorldPositionOffset.IsConnected() || GetCachedExpressionData().IsMaterialAttributePropertyConnected(MP_WorldPositionOffset); }
+	bool HasPixelDepthOffsetConnected() const { return PixelDepthOffset.IsConnected() || GetCachedExpressionData().IsMaterialAttributePropertyConnected(MP_PixelDepthOffset); }
 
 #if WITH_EDITOR
 	static void NotifyCompilationFinished(UMaterialInterface* Material);
diff --git a/Engine/Source/Runtime/Engine/Private/Materials/MaterialCachedData.cpp b/Engine/Source/Runtime/Engine/Private/Materials/MaterialCachedData.cpp
--- a/Engine/Source/Runtime/Engine/Private/Materials/MaterialCachedData.cpp
+++ b/Engine/Source/Runtime/Engine/Private/Materials/MaterialCachedData.cpp
@@ -29,6 +29,8 @@
 #include "Materials/MaterialExpressionRuntimeVirtualTextureOutput.h"
 #include "Materials/MaterialExpressionLandscapeGrassOutput.h"
 #include "Materials/MaterialExpressionCurveAtlasRowParameter.h"
+#include "Materials/MaterialExpressionSetMaterialAttributes.h"
+#include "Materials/MaterialExpressionMakeMaterialAttributes.h"
 #include "Materials/MaterialInstance.h"
 #include "Materials/MaterialFunction.h"
 #include "Materials/MaterialFunctionInstance.h"
@@ -54,6 +56,9 @@ void FMaterialCachedExpressionData::Reset()
 	QualityLevelsUsed.AddDefaulted(EMaterialQualityLevel::Num);
 	bHasRuntimeVirtualTextureOutput = false;
 	bHasSceneColor = false;
+	MaterialAttributesPropertyConnectedBitmask = 0;
+
+	static_assert((uint32)(EMaterialProperty::MP_MAX)-1 <= (8 * sizeof(MaterialAttributesPropertyConnectedBitmask)), "MaterialAttributesPropertyConnectedBitmask cannot contain entire EMaterialProperty enumeration.");
 }
 
 void FMaterialCachedExpressionData::AddReferencedObjects(FReferenceCollector& Collector)
@@ -618,9 +623,9 @@ bool FMaterialCachedExpressionData::UpdateForExpressions(const FMaterialCachedEx
 
 			LayersExpression->RebuildLayerGraph(false);
 		}
-		else if (Context.bUpdateFunctionExpressions)
+		else if (UMaterialExpressionMaterialFunctionCall* FunctionCall = Cast<UMaterialExpressionMaterialFunctionCall>(Expression))
 		{
-			if (UMaterialExpressionMaterialFunctionCall* FunctionCall = Cast<UMaterialExpressionMaterialFunctionCall>(Expression))
+			if (Context.bUpdateFunctionExpressions)
 			{
 				if (!UpdateForFunction(Context, FunctionCall->MaterialFunction, GlobalParameter, -1))
 				{
@@ -632,6 +637,62 @@ bool FMaterialCachedExpressionData::UpdateForExpressions(const FMaterialCachedEx
 				FunctionCall->UpdateFromFunctionResource();
 			}
 		}
+		else if (UMaterialExpressionSetMaterialAttributes* SetMatAttributes = Cast<UMaterialExpressionSetMaterialAttributes>(Expression))
+		{
+			for (int PinIndex = 0; PinIndex < SetMatAttributes->AttributeSetTypes.Num(); ++PinIndex)
+			{
+				// For this material attribute pin do we have something connected?
+				const FGuid& Guid = SetMatAttributes->AttributeSetTypes[PinIndex];
+				const FExpressionInput& AttributeInput = SetMatAttributes->Inputs[PinIndex + 1];
+				const EMaterialProperty MaterialProperty = FMaterialAttributeDefinitionMap::GetProperty(Guid);
+
+				// Only set the material property if it hasn't been set yet.  We want to specifically avoid a Set Material Attributes node which doesn't have a 
+				// attribute set from disabling the attribute from a different Set Material Attributes node which does have it enabled.
+				if (!IsMaterialAttributePropertyConnected(MaterialProperty))
+				{
+					SetMaterialAttributePropertyConnected(MaterialProperty, AttributeInput.Expression ? true : false);
+				}
+			}
+			}
+			else if (UMaterialExpressionMakeMaterialAttributes* MakeMatAttributes = Cast<UMaterialExpressionMakeMaterialAttributes>(Expression))
+			{
+			// Only set the material property if it hasn't been set yet.  We want to specifically avoid a Set Material Attributes node which doesn't have a 
+			// attribute set from disabling the attribute from a different Set Material Attributes node which does have it enabled.
+			auto SetMatAttributeConditionally = [&](EMaterialProperty InMaterialProperty, bool InIsConnected)
+			{
+				if (!IsMaterialAttributePropertyConnected(InMaterialProperty))
+				{
+					SetMaterialAttributePropertyConnected(InMaterialProperty, InIsConnected);
+				}
+			};
+
+			SetMatAttributeConditionally(EMaterialProperty::MP_BaseColor, MakeMatAttributes->BaseColor.IsConnected());
+			SetMatAttributeConditionally(EMaterialProperty::MP_Metallic, MakeMatAttributes->Metallic.IsConnected());
+			SetMatAttributeConditionally(EMaterialProperty::MP_Specular, MakeMatAttributes->Specular.IsConnected());
+			SetMatAttributeConditionally(EMaterialProperty::MP_Roughness, MakeMatAttributes->Roughness.IsConnected());
+			SetMatAttributeConditionally(EMaterialProperty::MP_Anisotropy, MakeMatAttributes->Anisotropy.IsConnected());
+			SetMatAttributeConditionally(EMaterialProperty::MP_EmissiveColor, MakeMatAttributes->EmissiveColor.IsConnected());
+			SetMatAttributeConditionally(EMaterialProperty::MP_Opacity, MakeMatAttributes->Opacity.IsConnected());
+			SetMatAttributeConditionally(EMaterialProperty::MP_OpacityMask, MakeMatAttributes->OpacityMask.IsConnected());
+			SetMatAttributeConditionally(EMaterialProperty::MP_Normal, MakeMatAttributes->Normal.IsConnected());
+			SetMatAttributeConditionally(EMaterialProperty::MP_Tangent, MakeMatAttributes->Tangent.IsConnected());
+			SetMatAttributeConditionally(EMaterialProperty::MP_WorldPositionOffset, MakeMatAttributes->WorldPositionOffset.IsConnected());
+			SetMatAttributeConditionally(EMaterialProperty::MP_SubsurfaceColor, MakeMatAttributes->SubsurfaceColor.IsConnected());
+			SetMatAttributeConditionally(EMaterialProperty::MP_CustomData0, MakeMatAttributes->ClearCoat.IsConnected());
+			SetMatAttributeConditionally(EMaterialProperty::MP_CustomData1, MakeMatAttributes->ClearCoatRoughness.IsConnected());
+			SetMatAttributeConditionally(EMaterialProperty::MP_AmbientOcclusion, MakeMatAttributes->AmbientOcclusion.IsConnected());
+			SetMatAttributeConditionally(EMaterialProperty::MP_Refraction, MakeMatAttributes->Refraction.IsConnected());
+			SetMatAttributeConditionally(EMaterialProperty::MP_CustomizedUVs0, MakeMatAttributes->CustomizedUVs[0].IsConnected());
+			SetMatAttributeConditionally(EMaterialProperty::MP_CustomizedUVs1, MakeMatAttributes->CustomizedUVs[1].IsConnected());
+			SetMatAttributeConditionally(EMaterialProperty::MP_CustomizedUVs2, MakeMatAttributes->CustomizedUVs[2].IsConnected());
+			SetMatAttributeConditionally(EMaterialProperty::MP_CustomizedUVs3, MakeMatAttributes->CustomizedUVs[3].IsConnected());
+			SetMatAttributeConditionally(EMaterialProperty::MP_CustomizedUVs4, MakeMatAttributes->CustomizedUVs[4].IsConnected());
+			SetMatAttributeConditionally(EMaterialProperty::MP_CustomizedUVs5, MakeMatAttributes->CustomizedUVs[5].IsConnected());
+			SetMatAttributeConditionally(EMaterialProperty::MP_CustomizedUVs6, MakeMatAttributes->CustomizedUVs[6].IsConnected());
+			SetMatAttributeConditionally(EMaterialProperty::MP_CustomizedUVs7, MakeMatAttributes->CustomizedUVs[7].IsConnected());
+			SetMatAttributeConditionally(EMaterialProperty::MP_PixelDepthOffset, MakeMatAttributes->PixelDepthOffset.IsConnected());
+			SetMatAttributeConditionally(EMaterialProperty::MP_ShadingModel, MakeMatAttributes->ShadingModel.IsConnected());
+		}
 	}
 
 	return bResult;
diff --git a/Engine/Source/Runtime/Engine/Private/Materials/MaterialShared.cpp b/Engine/Source/Runtime/Engine/Private/Materials/MaterialShared.cpp
--- a/Engine/Source/Runtime/Engine/Private/Materials/MaterialShared.cpp
+++ b/Engine/Source/Runtime/Engine/Private/Materials/MaterialShared.cpp
@@ -880,7 +880,7 @@ bool FMaterial::MaterialMayModifyMeshPosition() const
 {
 	// Conservative estimate when called before material translation has occurred. 
 	// This function is only intended for use in deciding whether or not shader permutations are required.
-	return HasVertexPositionOffsetConnected() || HasPixelDepthOffsetConnected() || HasMaterialAttributesConnected();
+	return HasVertexPositionOffsetConnected() || HasPixelDepthOffsetConnected();
 }
 
 bool FMaterial::MaterialUsesPixelDepthOffset() const
@@ -1296,8 +1296,8 @@ bool FMaterialResource::IsUsedWithEditorCompositing() const { return Material->b
 bool FMaterialResource::IsDeferredDecal() const { return Material->MaterialDomain == MD_DeferredDecal; }
 bool FMaterialResource::IsVolumetricPrimitive() const { return Material->MaterialDomain == MD_Volume; }
 bool FMaterialResource::IsSpecialEngineMaterial() const { return Material->bUsedAsSpecialEngineMaterial; }
-bool FMaterialResource::HasVertexPositionOffsetConnected() const { return HasMaterialAttributesConnected() || (!Material->bUseMaterialAttributes && Material->WorldPositionOffset.IsConnected()); }
-bool FMaterialResource::HasPixelDepthOffsetConnected() const { return HasMaterialAttributesConnected() || (!Material->bUseMaterialAttributes && Material->PixelDepthOffset.IsConnected()); }
+bool FMaterialResource::HasVertexPositionOffsetConnected() const { return Material->HasVertexPositionOffsetConnected(); }
+bool FMaterialResource::HasPixelDepthOffsetConnected() const { return Material->HasPixelDepthOffsetConnected(); }
 bool FMaterialResource::HasMaterialAttributesConnected() const { return (Material->bUseMaterialAttributes && Material->MaterialAttributes.IsConnected()) || Material->bEnableExecWire; }
 EMaterialShadingRate FMaterialResource::GetShadingRate() const { return Material->ShadingRate; }
 FString FMaterialResource::GetBaseMaterialPathName() const { return Material->GetPathName(); }
@@ -1585,22 +1585,22 @@ uint32 FMaterialResource::GetMaterialDecalResponse() const
 
 bool FMaterialResource::HasBaseColorConnected() const
 {
-	return HasMaterialAttributesConnected() || Material->HasBaseColorConnected();
+	return Material->HasBaseColorConnected();
 }
 
 bool FMaterialResource::HasNormalConnected() const
 {
-	return HasMaterialAttributesConnected() || Material->HasNormalConnected();
+	return Material->HasNormalConnected();
 }
 
 bool FMaterialResource::HasRoughnessConnected() const
 {
-	return HasMaterialAttributesConnected() || Material->HasRoughnessConnected();
+	return Material->HasRoughnessConnected();
 }
 
 bool FMaterialResource::HasSpecularConnected() const
 {
-	return HasMaterialAttributesConnected() || Material->HasSpecularConnected();
+	return Material->HasSpecularConnected();
 }
 
 bool FMaterialResource::HasMetallicConnected() const
@@ -1610,17 +1610,17 @@ bool FMaterialResource::HasMetallicConnected() const
 
 bool FMaterialResource::HasEmissiveColorConnected() const
 {
-	return HasMaterialAttributesConnected() || Material->HasEmissiveColorConnected();
+	return Material->HasEmissiveColorConnected();
 }
 
 bool FMaterialResource::HasAnisotropyConnected() const
 {
-	return HasMaterialAttributesConnected() || Material->HasAnisotropyConnected();
+	return Material->HasAnisotropyConnected();
 }
 
 bool FMaterialResource::HasAmbientOcclusionConnected() const
 {
-	return HasMaterialAttributesConnected() || Material->HasAmbientOcclusionConnected();
+	return Material->HasAmbientOcclusionConnected();
 }
 
 bool FMaterialResource::IsStrataMaterial() const
diff --git a/Engine/Source/Runtime/Engine/Public/MaterialCachedData.h b/Engine/Source/Runtime/Engine/Public/MaterialCachedData.h
--- a/Engine/Source/Runtime/Engine/Public/MaterialCachedData.h
+++ b/Engine/Source/Runtime/Engine/Public/MaterialCachedData.h
@@ -249,6 +249,16 @@ struct FMaterialCachedExpressionData
 
 	void AddReferencedObjects(FReferenceCollector& Collector);
 
+	bool IsMaterialAttributePropertyConnected(EMaterialProperty Property) const
+	{
+		return ((MaterialAttributesPropertyConnectedBitmask >> (uint32)Property) & 0x1) != 0;
+	}
+
+	void SetMaterialAttributePropertyConnected(EMaterialProperty Property, bool bIsConnected)
+	{
+		MaterialAttributesPropertyConnectedBitmask = bIsConnected ? MaterialAttributesPropertyConnectedBitmask | (1 << (uint32)Property) : MaterialAttributesPropertyConnectedBitmask & ~(1 << (uint32)Property);
+	}
+
 	UPROPERTY()
 	FMaterialCachedParameters Parameters;
 
@@ -284,6 +294,10 @@ struct FMaterialCachedExpressionData
 
 	UPROPERTY()
 	uint32 bHasSceneColor : 1;
+
+	/** Each bit corresponds to EMaterialProperty connection status. */
+	UPROPERTY()
+	uint32 MaterialAttributesPropertyConnectedBitmask = 0;
 };
 
 USTRUCT()
