Optimized FTickTaskManager::RunTickGroup by ~45% when the new tick set is empty. Constructing TSet::TIterator doesn't seem to be a trivial operation, as the object is 48 bytes and has nested bit-iterators. The code-gen for constructing and checking validity is pretty hefty and does a lot of copying using SSE instructions, so we avoid it if we know there won't be any work done.

diff --git a/Engine/Source/Runtime/Engine/Private/TickTaskManager.cpp b/Engine/Source/Runtime/Engine/Private/TickTaskManager.cpp
--- a/Engine/Source/Runtime/Engine/Private/TickTaskManager.cpp
+++ b/Engine/Source/Runtime/Engine/Private/TickTaskManager.cpp
@@ -1006,20 +1006,25 @@ public:
 	{
 		Context.TickGroup = CurrentTickGroup;
 		int32 Num = 0;
-		FTickTaskSequencer& TTS = FTickTaskSequencer::Get();
-		for (TSet<FTickFunction*>::TIterator It(NewlySpawnedTickFunctions); It; ++It)
-		{
-			FTickFunction* TickFunction = *It;
-			TickFunction->QueueTickFunction(TTS, Context);
-			Num++;
 
-			if (TickFunction->TickInterval > 0.f)
+		// Calling Empty() on an empty set is not free, as it resets the previously reserved hash entries
+		if (NewlySpawnedTickFunctions.Num() != 0)
+		{
+			FTickTaskSequencer& TTS = FTickTaskSequencer::Get();
+			for (TSet<FTickFunction*>::TIterator It(NewlySpawnedTickFunctions); It; ++It)
 			{
-				AllEnabledTickFunctions.Remove(TickFunction);
-				RescheduleForInterval(TickFunction, TickFunction->TickInterval);
+				FTickFunction* TickFunction = *It;
+				TickFunction->QueueTickFunction(TTS, Context);
+				Num++;
+
+				if (TickFunction->TickInterval > 0.f)
+				{
+					AllEnabledTickFunctions.Remove(TickFunction);
+					RescheduleForInterval(TickFunction, TickFunction->TickInterval);
+				}
 			}
+			NewlySpawnedTickFunctions.Empty();
 		}
-		NewlySpawnedTickFunctions.Empty();
 		return Num;
 	}
 	/**
