PR #8042: Fixed typo in MaterialTextureScaleAccuracy viewmode console command (Contributed by [REDACTED_PERSON])

#rnx

diff --git a/Engine/Source/Runtime/Engine/Private/ShowFlags.cpp b/Engine/Source/Runtime/Engine/Private/ShowFlags.cpp
--- a/Engine/Source/Runtime/Engine/Private/ShowFlags.cpp
+++ b/Engine/Source/Runtime/Engine/Private/ShowFlags.cpp
@@ -773,7 +773,7 @@ const TCHAR* GetViewModeName(EViewModeIndex ViewModeIndex)
 		case VMI_ShaderComplexityWithQuadOverdraw: return TEXT("ShaderComplexityWithQuadOverdraw");
 		case VMI_PrimitiveDistanceAccuracy:	return TEXT("PrimitiveDistanceAccuracy");
 		case VMI_MeshUVDensityAccuracy:		return TEXT("MeshUVDensityAccuracy");
-		case VMI_MaterialTextureScaleAccuracy: return TEXT("MaterialTexturecaleAccuracy");
+		case VMI_MaterialTextureScaleAccuracy: return TEXT("MaterialTextureScaleAccuracy");
 		case VMI_RequiredTextureResolution: return TEXT("RequiredTextureResolution");
 		case VMI_StationaryLightOverlap:	return TEXT("StationaryLightOverlap");
 		case VMI_LightmapDensity:			return TEXT("LightmapDensity");
