StaticMesh: Avoid a crash when the cache data is larger than 2GB.
- Change the cache archive from an FArchiveWriter (which crash on atempt to write more than 2GB) to an FLargeMemoryWritter.
- In that scenario, we skip the DDC (which API also limits to 2GB buffers)
#rnx

diff --git a/Engine/Source/Runtime/Engine/Private/StaticMesh.cpp b/Engine/Source/Runtime/Engine/Private/StaticMesh.cpp
--- a/Engine/Source/Runtime/Engine/Private/StaticMesh.cpp
+++ b/Engine/Source/Runtime/Engine/Private/StaticMesh.cpp
@@ -6,6 +6,7 @@
 
 #include "Engine/StaticMesh.h"
 #include "Serialization/MemoryWriter.h"
+#include "Serialization/LargeMemoryWriter.h"
 #include "Misc/ConfigCacheIni.h"
 #include "Misc/PackageSegment.h"
 #include "Misc/ScopedSlowTask.h"
@@ -2686,7 +2687,7 @@ void FStaticMeshRenderData::Cache(const ITargetPlatform* TargetPlatform, UStatic
 			}
 			
 			bLODsShareStaticLighting = Owner->CanLODsShareStaticLighting();
-			FMemoryWriter Ar(DerivedData, /*bIsPersistent=*/ true);
+			FLargeMemoryWriter Ar(0, /*bIsPersistent=*/ true);
 			Serialize(Ar, Owner, /*bCooked=*/ false);
 			check(NaniteResources.RootClusterPage.Num() == 0 || NaniteResources.bLZCompressed);
 
@@ -2714,9 +2715,12 @@ void FStaticMeshRenderData::Cache(const ITargetPlatform* TargetPlatform, UStatic
 				bSaveDDC = false;
 			}
 #endif
-			if (bSaveDDC)
+			int64 DerivedDataNum = Ar.TotalSize();
+			bool bCanStoreInDDC = DerivedDataNum <= TNumericLimits<TArrayView<uint8>::SizeType>::Max();
+			if (bSaveDDC && bCanStoreInDDC)
 			{
-				GetDerivedDataCacheRef().Put(*DerivedDataKey, DerivedData, Owner->GetPathName());
+				TArrayView<uint8> DataView(Ar.GetData(), DerivedDataNum);
+				GetDerivedDataCacheRef().Put(*DerivedDataKey, DataView, Owner->GetPathName());
 			}
 
 			double T1 = FPlatformTime::Seconds();
@@ -2725,7 +2729,7 @@ void FStaticMeshRenderData::Cache(const ITargetPlatform* TargetPlatform, UStatic
 				*Owner->GetPathName()
 				);
 			FPlatformAtomics::InterlockedAdd(&StaticMeshDerivedDataTimings::BuildCycles, T1 - T0);
-			COOK_STAT(Timer.AddMiss(DerivedData.Num()));
+			COOK_STAT(Timer.AddMiss(DerivedDataNum));
 		}
 	}
 
