Fixed tick task manager check for a licensee.

Happens in extremely rare case if an actor, that has ticking while paused, disables its ticking on the first paused tick.

diff --git a/Engine/Source/Runtime/Engine/Private/TickTaskManager.cpp b/Engine/Source/Runtime/Engine/Private/TickTaskManager.cpp
--- a/Engine/Source/Runtime/Engine/Private/TickTaskManager.cpp
+++ b/Engine/Source/Runtime/Engine/Private/TickTaskManager.cpp
@@ -1065,13 +1065,12 @@ public:
 			CumulativeCooldown += TickFunction->InternalData->RelativeTickCooldown;
 			if (TickFunction->bTickEvenWhenPaused)
 			{
+				bool bExecuteTick = false;
 				TickFunction->InternalData->TaskPointer = nullptr; // this is stale, clear it out now
 				if (CumulativeCooldown < InContext.DeltaSeconds)
 				{
-					TickFunction->InternalData->TickVisitedGFrameCounter = GFrameCounter;
-					TickFunction->InternalData->TickQueuedGFrameCounter = GFrameCounter;
-					TickFunction->ExecuteTick(TickFunction->CalculateDeltaTime(InContext), InContext.TickType, ENamedThreads::GameThread, FGraphEventRef());
-
+					// Do any rescheduling before the ticking as users can call methods inside the tick that can affect the current tick function state and we end up into a inconsistent state.
+					bExecuteTick = true;
 					RescheduleForInterval(TickFunction, TickFunction->TickInterval - (InContext.DeltaSeconds - CumulativeCooldown)); // Give credit for any overrun
 				}
 				else
@@ -1092,6 +1091,13 @@ public:
 					TickFunction->InternalData->Next->InternalData->RelativeTickCooldown += TickFunction->InternalData->RelativeTickCooldown;
 					CumulativeCooldown -= TickFunction->InternalData->RelativeTickCooldown; // Since the next object in the list will have this cooldown included take it back out of the cumulative
 				}
+
+				if (bExecuteTick)
+				{
+					TickFunction->InternalData->TickVisitedGFrameCounter = GFrameCounter;
+					TickFunction->InternalData->TickQueuedGFrameCounter = GFrameCounter;
+					TickFunction->ExecuteTick(TickFunction->CalculateDeltaTime(InContext), InContext.TickType, ENamedThreads::GameThread, FGraphEventRef());
+				}
 			}
 			else
 			{
