Tweaking comment to more accurately reflect why the if-check was put into QueueNewlySpawned. Set iterators are 48 bytes and do quite a bit of copying, so we avoid the work when we can. Empty() generally avoids work if the set is empty, but Core is looking into making further optimizations to eliminate work when Num < HashSize/4 and just deallocate the hash.

diff --git a/Engine/Source/Runtime/Engine/Private/TickTaskManager.cpp b/Engine/Source/Runtime/Engine/Private/TickTaskManager.cpp
--- a/Engine/Source/Runtime/Engine/Private/TickTaskManager.cpp
+++ b/Engine/Source/Runtime/Engine/Private/TickTaskManager.cpp
@@ -1007,7 +1007,7 @@ public:
 		Context.TickGroup = CurrentTickGroup;
 		int32 Num = 0;
 
-		// Calling Empty() on an empty set is not free, as it resets the previously reserved hash entries
+		// Constructing set iterators is not trivial, so avoid the following block if it will have no effect
 		if (NewlySpawnedTickFunctions.Num() != 0)
 		{
 			FTickTaskSequencer& TTS = FTickTaskSequencer::Get();
