Removing unused property

diff --git a/Engine/Source/Runtime/Engine/Classes/Sound/AudioBus.h b/Engine/Source/Runtime/Engine/Classes/Sound/AudioBus.h
--- a/Engine/Source/Runtime/Engine/Classes/Sound/AudioBus.h
+++ b/Engine/Source/Runtime/Engine/Classes/Sound/AudioBus.h
@@ -35,10 +35,6 @@ public:
 	UPROPERTY(EditAnywhere, Category = BusProperties)
 	EAudioBusChannels AudioBusChannels;
 
-	/** If the audio bus can be instantiated and destroyed automatically when sources send audio to it. If this audio bus is manually started, it will override this value to be false, meaning you will need to stop the audio bus manually.*/
-	UPROPERTY(EditAnywhere, Category = BusProperties)
-	bool bIsAutomatic;
-
 	//~ Begin UObject
 	virtual void BeginDestroy() override;
 	//~ End UObject
