Fixing tooltip

diff --git a/Engine/Source/Runtime/Engine/Classes/Sound/AudioSettings.h b/Engine/Source/Runtime/Engine/Classes/Sound/AudioSettings.h
--- a/Engine/Source/Runtime/Engine/Classes/Sound/AudioSettings.h
+++ b/Engine/Source/Runtime/Engine/Classes/Sound/AudioSettings.h
@@ -83,7 +83,7 @@ struct ENGINE_API FDefaultAudioBusSettings
 {
 	GENERATED_BODY()
 
-	/** The SoundConcurrency assigned to newly created sounds */
+	/** The audio bus to start up by default on init. */
 	UPROPERTY(EditAnywhere, Category = "Mix", meta = (AllowedClasses = "AudioBus"))
 	FSoftObjectPath AudioBus;
 };
