Fix CIS (bring back InDataKey, now that we need it).

diff --git a/Engine/Source/Runtime/Engine/Private/Materials/MaterialShader.cpp b/Engine/Source/Runtime/Engine/Private/Materials/MaterialShader.cpp
--- a/Engine/Source/Runtime/Engine/Private/Materials/MaterialShader.cpp
+++ b/Engine/Source/Runtime/Engine/Private/Materials/MaterialShader.cpp
@@ -1257,6 +1257,7 @@ void FMaterialShaderMap::LoadFromDerivedDataCache(const FMaterial* Material, con
 				InOutShaderMap->Serialize(Ar);
 				//InOutShaderMap->RegisterSerializedShaders(false);
 		
+				const FString InDataKey = GetMaterialShaderMapKeyString(InOutShaderMap->GetShaderMapId(), InPlatform, TargetPlatform);
 				checkf(InDataKey == DataKey, TEXT("Data deserialized from the DDC would need a different DDC key!"));
 				checkf(InOutShaderMap->GetShaderMapId() == ShaderMapId, TEXT("Shadermap data deserialized from the DDC does not match the ID we used to build the key!"));
 
