Fixed UI and Never texture memory stats that would never decrease.

diff --git a/Engine/Source/Runtime/Engine/Private/Rendering/StreamableTextureResource.cpp b/Engine/Source/Runtime/Engine/Private/Rendering/StreamableTextureResource.cpp
--- a/Engine/Source/Runtime/Engine/Private/Rendering/StreamableTextureResource.cpp
+++ b/Engine/Source/Runtime/Engine/Private/Rendering/StreamableTextureResource.cpp
@@ -151,11 +151,11 @@ void FStreamableTextureResource::DecrementTextureStats() const
 
 	if (LODGroup == TEXTUREGROUP_UI)
 	{
-		GUITextureMemory += TextureSize;
+		GUITextureMemory -= TextureSize;
 	}
 	else if (bIsNeverStream)
 	{
-		GNeverStreamTextureMemory += TextureSize;
+		GNeverStreamTextureMemory -= TextureSize;
 	}
 }
 #endif
