Merging CL#16407567

#UE Newly created inactive worlds should not have a physics scene, just like loaded inactive worlds. This was causing a problem where the physics scene was not properly getting destroyed and keeping a reference to the world, resulting in a world leak.

diff --git a/Engine/Source/Runtime/Engine/Private/World.cpp b/Engine/Source/Runtime/Engine/Private/World.cpp
--- a/Engine/Source/Runtime/Engine/Private/World.cpp
+++ b/Engine/Source/Runtime/Engine/Private/World.cpp
@@ -1972,7 +1972,7 @@ UWorld* UWorld::CreateWorld(const EWorldType::Type InWorldType, bool bInformEngi
 	NewWorld->SetFlags(RF_Transactional);
 	NewWorld->WorldType = InWorldType;
 	NewWorld->FeatureLevel = InFeatureLevel;
-	NewWorld->InitializeNewWorld(InIVS ? *InIVS : UWorld::InitializationValues().ShouldSimulatePhysics(false).EnableTraceCollision(true).CreateNavigation(InWorldType == EWorldType::Editor).CreateAISystem(InWorldType == EWorldType::Editor), bInSkipInitWorld);
+	NewWorld->InitializeNewWorld(InIVS ? *InIVS : UWorld::InitializationValues().CreatePhysicsScene(InWorldType != EWorldType::Inactive).ShouldSimulatePhysics(false).EnableTraceCollision(true).CreateNavigation(InWorldType == EWorldType::Editor).CreateAISystem(InWorldType == EWorldType::Editor), bInSkipInitWorld);
 
 	// Clear the dirty flag set during SpawnActor and UpdateLevelComponents
 	WorldPackage->SetDirtyFlag(false);
