An other fix for tick paused frame crash from a licensee

diff --git a/Engine/Source/Runtime/Engine/Private/TickTaskManager.cpp b/Engine/Source/Runtime/Engine/Private/TickTaskManager.cpp
--- a/Engine/Source/Runtime/Engine/Private/TickTaskManager.cpp
+++ b/Engine/Source/Runtime/Engine/Private/TickTaskManager.cpp
@@ -1057,6 +1057,8 @@ public:
 	{
 		check(!NewlySpawnedTickFunctions.Num()); // There shouldn't be any in here at this point in the frame
 
+		TArray<FTickFunction*> ExecuteTickFunctions;
+
 		float CumulativeCooldown = 0.f;
 		FTickFunction* PrevTickFunction = nullptr;
 		FTickFunction* TickFunction = AllCoolingDownTickFunctions.Head;
@@ -1069,8 +1071,8 @@ public:
 				TickFunction->InternalData->TaskPointer = nullptr; // this is stale, clear it out now
 				if (CumulativeCooldown < InContext.DeltaSeconds)
 				{
-					// Do any rescheduling before the ticking as users can call methods inside the tick that can affect the current tick function state and we end up into a inconsistent state.
-					bExecuteTick = true;
+					// Queue up the tick function for later and do the reschedule before as it is in the normal ticking logic
+					ExecuteTickFunctions.Add(TickFunction);
 					RescheduleForInterval(TickFunction, TickFunction->TickInterval - (InContext.DeltaSeconds - CumulativeCooldown)); // Give credit for any overrun
 				}
 				else
@@ -1091,13 +1093,6 @@ public:
 					TickFunction->InternalData->Next->InternalData->RelativeTickCooldown += TickFunction->InternalData->RelativeTickCooldown;
 					CumulativeCooldown -= TickFunction->InternalData->RelativeTickCooldown; // Since the next object in the list will have this cooldown included take it back out of the cumulative
 				}
-
-				if (bExecuteTick)
-				{
-					TickFunction->InternalData->TickVisitedGFrameCounter = GFrameCounter;
-					TickFunction->InternalData->TickQueuedGFrameCounter = GFrameCounter;
-					TickFunction->ExecuteTick(TickFunction->CalculateDeltaTime(InContext), InContext.TickType, ENamedThreads::GameThread, FGraphEventRef());
-				}
 			}
 			else
 			{
@@ -1112,9 +1107,8 @@ public:
 			TickFunction->InternalData->TaskPointer = nullptr; // this is stale, clear it out now
 			if (TickFunction->bTickEvenWhenPaused && TickFunction->TickState == FTickFunction::ETickState::Enabled)
 			{
-				TickFunction->InternalData->TickVisitedGFrameCounter = GFrameCounter;
-				TickFunction->InternalData->TickQueuedGFrameCounter = GFrameCounter;
-				TickFunction->ExecuteTick(TickFunction->CalculateDeltaTime(InContext), InContext.TickType, ENamedThreads::GameThread, FGraphEventRef());
+				// Queue up the tick function for later and do the reschedule before as it is in the normal ticking logic
+				ExecuteTickFunctions.Add(TickFunction);
 
 				if (TickFunction->TickInterval > 0.f)
 				{
@@ -1125,6 +1119,13 @@ public:
 		}
 
 		check(!NewlySpawnedTickFunctions.Num()); // We don't support new spawns during pause ticks
+
+		for (FTickFunction* TickFunctionToExecute : ExecuteTickFunctions)
+		{
+			TickFunctionToExecute->InternalData->TickVisitedGFrameCounter = GFrameCounter;
+			TickFunctionToExecute->InternalData->TickQueuedGFrameCounter = GFrameCounter;
+			TickFunctionToExecute->ExecuteTick(TickFunctionToExecute->CalculateDeltaTime(InContext), InContext.TickType, ENamedThreads::GameThread, FGraphEventRef());
+		}
 	}
 
 	/** End a tick frame **/
