CIS fix

diff --git a/Engine/Source/Editor/UnrealEd/Private/Commandlets/ContentCommandlets.cpp b/Engine/Source/Editor/UnrealEd/Private/Commandlets/ContentCommandlets.cpp
--- a/Engine/Source/Editor/UnrealEd/Private/Commandlets/ContentCommandlets.cpp
+++ b/Engine/Source/Editor/UnrealEd/Private/Commandlets/ContentCommandlets.cpp
@@ -3389,9 +3389,9 @@ int32 UStaticMeshMinLodCommandlet::Main(const FString& Params)
 				}
 			}
 
-			if (const PlatformInfo::FPlatformInfo* PlatformInfo = PlatformInfo::FindPlatformInfo(PlatformName))
+			if (const PlatformInfo::FTargetPlatformInfo* PlatformInfo = PlatformInfo::FindPlatformInfo(PlatformName))
 			{
-				GroupToPlatform.AddUnique(PlatformInfo->PlatformGroupName, PlatformName);
+				GroupToPlatform.AddUnique(PlatformInfo->DataDrivenPlatformInfo->PlatformGroupName, PlatformName);
 			}
 			return true;
 		});
diff --git a/Engine/Source/Runtime/Engine/Private/StaticMesh.cpp b/Engine/Source/Runtime/Engine/Private/StaticMesh.cpp
--- a/Engine/Source/Runtime/Engine/Private/StaticMesh.cpp
+++ b/Engine/Source/Runtime/Engine/Private/StaticMesh.cpp
@@ -314,7 +314,7 @@ int32 FStaticMeshLODResources::GetPlatformMinLODIdx(const ITargetPlatform* Targe
 	if (StaticMesh->IsMinLodQualityLevelEnable())
 	{
 		// get all supported quality level from scalability + engine ini files
-		FSupportedQualityLevelArray SupportedQualityLevels = StaticMesh->GetQualityLevelMinLOD().GetSupportedQualityLevels(*TargetPlatform->GetPlatformInfo().TargetPlatformName.ToString());
+		FSupportedQualityLevelArray SupportedQualityLevels = StaticMesh->GetQualityLevelMinLOD().GetSupportedQualityLevels(*TargetPlatform->PlatformName());
 		
 		// loop through all the supported quality level to find the min lod index
 		int32 MinLodIdx = MAX_int32;
