Fix FURL relative URL's overriding Base URL Host:Port (selective/non-exhaustive fix, restricted to the part of code that caused a regression).

diff --git a/Engine/Source/Runtime/Engine/Private/URL.cpp b/Engine/Source/Runtime/Engine/Private/URL.cpp
--- a/Engine/Source/Runtime/Engine/Private/URL.cpp
+++ b/Engine/Source/Runtime/Engine/Private/URL.cpp
@@ -408,30 +408,34 @@ FURL::FURL( FURL* Base, const TCHAR* TextURL, ETravelType Type )
 				URLStr = (MapIdx != INDEX_NONE) ? URLStr.Mid(MapIdx) : TEXT("");
 
 
-				// Skip past all the ':' characters in the IPv6 address to get to the port.
-				int32 ClosingBracketIdx = FindCharIdx(HostAndPort, ']');
-				int32 PortSearchIdx = HostAndPort.Find(":", ESearchCase::CaseSensitive, ESearchDir::FromStart, FMath::Max(ClosingBracketIdx, 0));
-				int32 NextColonIdx = HostAndPort.Find(":", ESearchCase::CaseSensitive, ESearchDir::FromStart, PortSearchIdx + 1);
-
-				// If we don't have a square bracket, check for another instance of a colon. This will tell us if we're dealing with an IPv6 address
-				// if there isn't one, then we're dealing with a port.
-				if (PortSearchIdx != INDEX_NONE && (ClosingBracketIdx != INDEX_NONE || NextColonIdx == INDEX_NONE))
+				// Don't override the Base Host:Port, for relative travel (not rigorously enforced in this function, selective regression fix)
+				if (Type != TRAVEL_Relative || Host.IsEmpty())
 				{
-					Port = FCString::Atoi(*HostAndPort.Mid(PortSearchIdx+1));
-					HostAndPort = HostAndPort.Left(PortSearchIdx);
-				}
+					// Skip past all the ':' characters in the IPv6 address to get to the port.
+					int32 ClosingBracketIdx = FindCharIdx(HostAndPort, ']');
+					int32 PortSearchIdx = HostAndPort.Find(":", ESearchCase::CaseSensitive, ESearchDir::FromStart, FMath::Max(ClosingBracketIdx, 0));
+					int32 NextColonIdx = HostAndPort.Find(":", ESearchCase::CaseSensitive, ESearchDir::FromStart, PortSearchIdx + 1);
+
+					// If we don't have a square bracket, check for another colon. This will tell us if we're dealing with an IPv6 address
+					// if there isn't one, then we're dealing with a port.
+					if (PortSearchIdx != INDEX_NONE && (ClosingBracketIdx != INDEX_NONE || NextColonIdx == INDEX_NONE))
+					{
+						Port = FCString::Atoi(*HostAndPort.Mid(PortSearchIdx+1));
+						HostAndPort = HostAndPort.Left(PortSearchIdx);
+					}
 
-				int32 OpeningBracketIdx = FindCharIdx(HostAndPort, '[');
+					int32 OpeningBracketIdx = FindCharIdx(HostAndPort, '[');
 
-				// If the input was an IPv6 address with a port, we need to remove the brackets then.
-				if (OpeningBracketIdx != INDEX_NONE && ClosingBracketIdx != INDEX_NONE)
-				{
-					Host = HostAndPort.Mid(OpeningBracketIdx + 1, (ClosingBracketIdx - OpeningBracketIdx) - 1);
-				}
-				else
-				{
-					// Otherwise, leave the address as is.
-					Host = HostAndPort;
+					// If the input was an IPv6 address with a port, we need to remove the brackets then.
+					if (OpeningBracketIdx != INDEX_NONE && ClosingBracketIdx != INDEX_NONE)
+					{
+						Host = HostAndPort.Mid(OpeningBracketIdx + 1, (ClosingBracketIdx - OpeningBracketIdx) - 1);
+					}
+					else
+					{
+						// Otherwise, leave the address as is.
+						Host = HostAndPort;
+					}
 				}
 
 				if( FCString::Stricmp(*Protocol,*UrlConfig.DefaultProtocol)==0 )
