Added some extra defensive safety checks to ensure we dont access bone transform out of bounds

This should prevent some rare crashes when dealing with meshes that have master-pose setups

diff --git a/Engine/Source/Runtime/Engine/Private/Components/SkinnedMeshComponent.cpp b/Engine/Source/Runtime/Engine/Private/Components/SkinnedMeshComponent.cpp
--- a/Engine/Source/Runtime/Engine/Private/Components/SkinnedMeshComponent.cpp
+++ b/Engine/Source/Runtime/Engine/Private/Components/SkinnedMeshComponent.cpp
@@ -1286,11 +1286,11 @@ FTransform USkinnedMeshComponent::GetBoneTransform(int32 BoneIdx, const FTransfo
 		}
 		if(BoneIdx < MasterBoneMap.Num())
 		{
-			int32 MasterBoneIndex = MasterBoneMap[BoneIdx];
+			const int32 MasterBoneIndex = MasterBoneMap[BoneIdx];
+			const int32 NumMasterTransforms = MasterPoseComponentInst->GetNumComponentSpaceTransforms();
 
-			// If ParentBoneIndex is valid, grab matrix from MasterPoseComponent.
-			if(	MasterBoneIndex != INDEX_NONE && 
-				MasterBoneIndex < MasterPoseComponentInst->GetNumComponentSpaceTransforms())
+			// If MasterBoneIndex is valid, grab matrix from MasterPoseComponent.
+			if(MasterBoneIndex >= 0 && MasterBoneIndex < NumMasterTransforms)
 			{
 				return MasterPoseComponentInst->GetComponentSpaceTransforms()[MasterBoneIndex] * LocalToWorld;
 			}
@@ -1301,12 +1301,20 @@ FTransform USkinnedMeshComponent::GetBoneTransform(int32 BoneIdx, const FTransfo
 				const FMissingMasterBoneCacheEntry* MissingBoneInfoPtr = MissingMasterBoneMap.Find(BoneIdx);
 				if(MissingBoneInfoPtr != nullptr)
 				{
-					return MissingBoneInfoPtr->RelativeTransform * MasterPoseComponentInst->GetComponentSpaceTransforms()[MissingBoneInfoPtr->CommonAncestorBoneIndex] * LocalToWorld;
+					const int32 MissingMasterBoneIndex = MissingBoneInfoPtr->CommonAncestorBoneIndex;
+					if(MissingMasterBoneIndex >= 0 && MissingMasterBoneIndex < NumMasterTransforms)
+					{
+						return MissingBoneInfoPtr->RelativeTransform * MasterPoseComponentInst->GetComponentSpaceTransforms()[MissingBoneInfoPtr->CommonAncestorBoneIndex] * LocalToWorld;
+					}
 				}
 				// Otherwise we might be able to generate the missing transform on the fly (although this is expensive)
 				else if(GetMissingMasterBoneRelativeTransform(BoneIdx, MissingBoneInfo))
 				{
-					return MissingBoneInfo.RelativeTransform * MasterPoseComponentInst->GetComponentSpaceTransforms()[MissingBoneInfo.CommonAncestorBoneIndex] * LocalToWorld;
+					const int32 MissingMasterBoneIndex = MissingBoneInfo.CommonAncestorBoneIndex;
+					if (MissingMasterBoneIndex >= 0 && MissingMasterBoneIndex < NumMasterTransforms)
+					{
+						return MissingBoneInfo.RelativeTransform * MasterPoseComponentInst->GetComponentSpaceTransforms()[MissingBoneInfo.CommonAncestorBoneIndex] * LocalToWorld;
+					}
 				}
 
 				UE_LOG(LogSkinnedMeshComp, Verbose, TEXT("GetBoneTransform : ParentBoneIndex(%d) out of range of MasterPoseComponent->SpaceBases for %s"), BoneIdx, *this->GetFName().ToString() );
