Set MinLod stripping quality level (disable by default)

Low=Mobile
Medium=Switch
High=Desktop,Xbox,PS4,XboxOneGDK
Epic=PS5,XSX

diff --git a/Engine/Source/Runtime/Engine/Private/StaticMesh.cpp b/Engine/Source/Runtime/Engine/Private/StaticMesh.cpp
--- a/Engine/Source/Runtime/Engine/Private/StaticMesh.cpp
+++ b/Engine/Source/Runtime/Engine/Private/StaticMesh.cpp
@@ -151,7 +151,7 @@ static FAutoConsoleVariableRef CVarMinLodQualityLevel(
 	FConsoleVariableDelegate::CreateStatic(&UStaticMesh::OnLodStrippingQualityLevelChanged),
 	ECVF_Scalability);
 
-int32 GDisableMinLODQualityLevel = 0;
+int32 GDisableMinLODQualityLevel = 1;
 static FAutoConsoleVariableRef CVarCDisableMinLODQualityLevel(
 	TEXT("r.StaticMesh.DisableMinLODQualityLevel"),
 	GDisableMinLODQualityLevel,
