Always take into account the pending first LOD when calling MeshObject->Update() which creates new dynamic data.
Without this we end up trying to render a LOD that isn't streamed.
USkinnedMeshComponent::CreateRenderState_Concurrent() already does the right thing but this was being overwritten by the next call to USkinnedMeshComponent::SendRenderDynamicData_Concurrent()

diff --git a/Engine/Source/Runtime/Engine/Private/Components/SkinnedMeshComponent.cpp b/Engine/Source/Runtime/Engine/Private/Components/SkinnedMeshComponent.cpp
--- a/Engine/Source/Runtime/Engine/Private/Components/SkinnedMeshComponent.cpp
+++ b/Engine/Source/Runtime/Engine/Private/Components/SkinnedMeshComponent.cpp
@@ -689,7 +689,12 @@ void USkinnedMeshComponent::SendRenderDynamicData_Concurrent()
 	{
 		SCOPE_CYCLE_COUNTER(STAT_MeshObjectUpdate);
 
-		const int32 UseLOD = GetPredictedLODLevel();
+		int32 UseLOD = GetPredictedLODLevel();
+		// Clamp to loaded streaming data if available
+		if (SkeletalMesh->IsStreamable() && MeshObject)
+		{
+			UseLOD = FMath::Max<int32>(UseLOD, MeshObject->GetSkeletalMeshRenderData().PendingFirstLODIdx);
+		}
 
 		// Only update the state if PredictedLODLevel is valid
 		FSkeletalMeshRenderData* SkelMeshRenderData = GetSkeletalMeshRenderData();
@@ -715,10 +720,15 @@ void USkinnedMeshComponent::SendRenderDynamicData_Concurrent()
 
 void USkinnedMeshComponent::ClearMotionVector()
 {
-	const int32 UseLOD = GetPredictedLODLevel();
-
 	if (MeshObject)
 	{
+		int32 UseLOD = GetPredictedLODLevel();
+		// Clamp to loaded streaming data if available
+		if (SkeletalMesh->IsStreamable() && MeshObject)
+		{
+			UseLOD = FMath::Max<int32>(UseLOD, MeshObject->GetSkeletalMeshRenderData().PendingFirstLODIdx);
+		}
+
 		// rendering bone velocity is updated by revision number
 		// if you have situation where you want to clear the bone velocity (that causes temporal AA or motion blur)
 		// use this function to clear it
@@ -735,8 +745,15 @@ void USkinnedMeshComponent::ForceMotionVector()
 {
 	if (MeshObject)
 	{
+		int32 UseLOD = GetPredictedLODLevel();
+		// Clamp to loaded streaming data if available
+		if (SkeletalMesh->IsStreamable() && MeshObject)
+		{
+			UseLOD = FMath::Max<int32>(UseLOD, MeshObject->GetSkeletalMeshRenderData().PendingFirstLODIdx);
+		}
+
 		++CurrentBoneTransformRevisionNumber;
-		MeshObject->Update(GetPredictedLODLevel(), this, ActiveMorphTargets, MorphTargetWeights, EPreviousBoneTransformUpdateMode::None);
+		MeshObject->Update(UseLOD, this, ActiveMorphTargets, MorphTargetWeights, EPreviousBoneTransformUpdateMode::None);
 	}
 }
 
