Simplify DataDrivenShaderPlatform logic in IsHairStrandsGeometrySupported by changing bSupportsHairStrandGeometry to an opt-in setting.
Enable the setting on the corresponding console and PC platforms (console platforms were set in previous CLs).

#rnx

diff --git a/Engine/Config/Windows/DataDrivenPlatformInfo.ini b/Engine/Config/Windows/DataDrivenPlatformInfo.ini
--- a/Engine/Config/Windows/DataDrivenPlatformInfo.ini
+++ b/Engine/Config/Windows/DataDrivenPlatformInfo.ini
@@ -80,6 +80,7 @@ bSupportsGPUSkinCache=true
 bSupportsGen5TemporalAA=true
 bSupportsMeshShaders=true
 FriendlyName=LOCTEXT("FriendlyShaderPlatformName_PCD3D_SM6", "Shader Model 6")
+bSupportsHairStrandGeometry=true
 
 [PreviewPlatform PCD3D_SM6]
 PlatformName=PCD3D_SM6
@@ -195,6 +196,8 @@ bSupportsVariableRateShading = false
 
 NumberOfComputeThreads = 64
 
+bSupportsHairStrandGeometry=true
+
 [ShaderPlatform PCD3D_ES3_1]
 Language=D3D
 MaxFeatureLevel=ES3_1
diff --git a/Engine/Source/Runtime/Engine/Private/Materials/MaterialInterface.cpp b/Engine/Source/Runtime/Engine/Private/Materials/MaterialInterface.cpp
--- a/Engine/Source/Runtime/Engine/Private/Materials/MaterialInterface.cpp
+++ b/Engine/Source/Runtime/Engine/Private/Materials/MaterialInterface.cpp
@@ -43,17 +43,8 @@ bool IsHairStrandsGeometrySupported(const EShaderPlatform Platform)
 {
 	check(Platform != SP_NumPlatforms);
 
-	bool bPCPlatformSupportsHairStrands = (IsD3DPlatform(Platform) || IsVulkanSM5Platform(Platform));
-	bPCPlatformSupportsHairStrands = bPCPlatformSupportsHairStrands && IsPCPlatform(Platform);
-	bPCPlatformSupportsHairStrands = bPCPlatformSupportsHairStrands && !IsMobilePlatform(Platform);
-
-	bool bConsolePlatformSupportsHairStrands = IsConsolePlatform(Platform) && !IsSwitchPlatform(Platform);
-	bConsolePlatformSupportsHairStrands = bConsolePlatformSupportsHairStrands && FDataDrivenShaderPlatformInfo::GetSupportsHairStrandGeometry(Platform);
-
-	bool bSupportsHairStrands = (bConsolePlatformSupportsHairStrands || bPCPlatformSupportsHairStrands);
-	bSupportsHairStrands = bSupportsHairStrands && IsFeatureLevelSupported(Platform, ERHIFeatureLevel::SM5);
-
-	return bSupportsHairStrands;
+	return FDataDrivenShaderPlatformInfo::GetSupportsHairStrandGeometry(Platform)
+		&& IsFeatureLevelSupported(Platform, ERHIFeatureLevel::SM5);
 }
 
 bool IsCompatibleWithHairStrands(const FMaterial* Material, const ERHIFeatureLevel::Type FeatureLevel)
diff --git a/Engine/Source/Runtime/RHI/Private/RHI.cpp b/Engine/Source/Runtime/RHI/Private/RHI.cpp
--- a/Engine/Source/Runtime/RHI/Private/RHI.cpp
+++ b/Engine/Source/Runtime/RHI/Private/RHI.cpp
@@ -2030,7 +2030,6 @@ void FGenericDataDrivenShaderPlatformInfo::SetDefaultValues()
 	bSupportsMSAA = true;
 
 	bNeedsToSwitchVerticalAxisOnMobileOpenGL = true;
-	bSupportsHairStrandGeometry = true;
 	bSupportsDOFHybridScattering = true;
 	bSupportsHZBOcclusion = true;
 	bSupportsWaterIndirectDraw = true;
