Cloth - Fix external collision sources that wrongly applied the cloth scale onto the source collision spheres.

diff --git a/Engine/Source/Runtime/Engine/Private/SkeletalMeshComponentPhysics.cpp b/Engine/Source/Runtime/Engine/Private/SkeletalMeshComponentPhysics.cpp
--- a/Engine/Source/Runtime/Engine/Private/SkeletalMeshComponentPhysics.cpp
+++ b/Engine/Source/Runtime/Engine/Private/SkeletalMeshComponentPhysics.cpp
@@ -2771,10 +2771,15 @@ void USkeletalMeshComponent::ExtractCollisionsForCloth(
 	// (but we want it to affect other meshes with cloth simulation)
 	if(SourceComponent->SkeletalMesh && PhysicsAsset)
 	{
-		const FTransform& ComponentToComponentTransform = 
-			SourceComponent != DestClothComponent ? 
-				SourceComponent->GetComponentTransform() * DestClothComponent->GetComponentTransform().Inverse() : 
-				FTransform::Identity;
+		FTransform ComponentToComponentTransform;
+		if(SourceComponent != DestClothComponent)
+		{
+			FTransform SourceClothComponentTransform = SourceComponent->GetComponentTransform();
+			SourceClothComponentTransform.RemoveScaling();  // The source component scale will be present in the world space bone transform, and is not needed here
+			FTransform DestClothComponentTransform = DestClothComponent->GetComponentTransform();
+			DestClothComponentTransform.RemoveScaling();  // The collision source doesn't need the scale of the cloth skeletal mesh applied to it
+			ComponentToComponentTransform = SourceClothComponentTransform * DestClothComponentTransform.Inverse();
+		}
 
 		// Init cache on first copy
 		if(!ClothCollisionSource.bCached || ClothCollisionSource.CachedSkeletalMesh.Get() != SourceComponent->SkeletalMesh)
@@ -2859,6 +2864,7 @@ void USkeletalMeshComponent::ExtractCollisionsForCloth(
 
 			const FTransform BoneTransform = SourceComponent->GetBoneTransform(OutSphere.BoneIndex, FTransform::Identity) * ComponentToComponentTransform;
 			OutSphere.LocalPosition = BoneTransform.TransformPosition(OutSphere.LocalPosition);
+			OutSphere.Radius *= BoneTransform.GetScale3D().X;  // Cloth collisions only uniformly scale
 			OutSphere.BoneIndex = INDEX_NONE;
 		}
 
