Fixed NavModifierVolumes not exporting concave shapes properly.
The core of the issue was that ANavModifierVolume was using a malfunctioning FAreaNavModifier constructor so I've marked that one up as deprecated.
Note that this change affects the behavior of NavModifierVolume and can affect projects depending on previous, erroneous behavior.

diff --git a/Engine/Source/Runtime/Engine/Public/AI/NavigationModifier.h b/Engine/Source/Runtime/Engine/Public/AI/NavigationModifier.h
--- a/Engine/Source/Runtime/Engine/Public/AI/NavigationModifier.h
+++ b/Engine/Source/Runtime/Engine/Public/AI/NavigationModifier.h
@@ -103,6 +103,7 @@ struct ENGINE_API FAreaNavModifier : public FNavigationModifier
 	FAreaNavModifier(const TArray<FVector3f>& Points, ENavigationCoordSystem::Type CoordType, const FTransform& LocalToWorld, const TSubclassOf<UNavAreaBase> AreaClass);
 	FAreaNavModifier(const TArray<FVector3f>& Points, const int32 FirstIndex, const int32 LastIndex, ENavigationCoordSystem::Type CoordType, const FTransform& LocalToWorld, const TSubclassOf<UNavAreaBase> AreaClass);
 	FAreaNavModifier(const TNavStatArray<FVector3f>& Points, const int32 FirstIndex, const int32 LastIndex, ENavigationCoordSystem::Type CoordType, const FTransform& LocalToWorld, const TSubclassOf<UNavAreaBase> AreaClass);
+	UE_DEPRECATED(5.0, "FAreaNavModifier constructor with a UBrushComponent* parameter has been deprecated since it wasn't able to handle concave shapes. Use FCompositeNavModifier::CreateAreaModifiers instead")
 	FAreaNavModifier(const UBrushComponent* BrushComponent, const TSubclassOf<UNavAreaBase> AreaClass);
 
 	void InitializePerInstanceConvex(const TNavStatArray<FVector>& Points, const int32 FirstIndex, const int32 LastIndex, const TSubclassOf<UNavAreaBase> AreaClass);
diff --git a/Engine/Source/Runtime/NavigationSystem/Private/NavModifierVolume.cpp b/Engine/Source/Runtime/NavigationSystem/Private/NavModifierVolume.cpp
--- a/Engine/Source/Runtime/NavigationSystem/Private/NavModifierVolume.cpp
+++ b/Engine/Source/Runtime/NavigationSystem/Private/NavModifierVolume.cpp
@@ -29,8 +29,7 @@ void ANavModifierVolume::GetNavigationData(FNavigationRelevantData& Data) const
 {
 	if (Brush && AreaClass && AreaClass != FNavigationSystem::GetDefaultWalkableArea())
 	{
-		FAreaNavModifier AreaMod(GetBrushComponent(), AreaClass);
-		Data.Modifiers.Add(AreaMod);
+		Data.Modifiers.CreateAreaModifiers(GetBrushComponent(), AreaClass);
 	}
 
 	if (bMaskFillCollisionUnderneathForNavmesh)
