Fix for an editor crash that could occur when reparenting a Blueprint with at least one Blueprint-added component to its duplicate.

Addresses an editor workflow issue that made it difficult to effectively branch an existing Blueprint class hierarchy using the reparent feature.

diff --git a/Engine/Source/Runtime/Engine/Private/Blueprint.cpp b/Engine/Source/Runtime/Engine/Private/Blueprint.cpp
--- a/Engine/Source/Runtime/Engine/Private/Blueprint.cpp
+++ b/Engine/Source/Runtime/Engine/Private/Blueprint.cpp
@@ -875,11 +875,13 @@ bool UBlueprint::CanRecompileWhilePlayingInEditor() const
 void UBlueprint::GetAssetRegistryTags(TArray<FAssetRegistryTag>& OutTags) const
 {
 	// We use Generated instead of Skeleton because the CDO data is more accurate on Generated
+	UObject* BlueprintCDO = nullptr;
 	if (GeneratedClass)
 	{
-		if (UObject* CDO = GeneratedClass->GetDefaultObject())
+		BlueprintCDO = GeneratedClass->GetDefaultObject();
+		if (BlueprintCDO)
 		{
-			CDO->GetAssetRegistryTags(OutTags);
+			BlueprintCDO->GetAssetRegistryTags(OutTags);
 		}
 	}
 
@@ -966,10 +968,10 @@ void UBlueprint::GetAssetRegistryTags(TArray<FAssetRegistryTag>& OutTags) const
 	{
 		// Determine how many inherited native components exist
 		int32 NumNativeComponents = 0;
-		if (BlueprintClass != nullptr)
+		if (BlueprintCDO != nullptr)
 		{
 			TArray<UObject*> PotentialComponents;
-			BlueprintClass->GetDefaultObjectSubobjects(/*out*/ PotentialComponents);
+			BlueprintCDO->GetDefaultSubobjects(/*out*/ PotentialComponents);
 
 			for (UObject* TestSubObject : PotentialComponents)
 			{
