Add AllowAudioVisualData() to the target platform interface so that TPs can tell the outside world whether they need shaders/audio etc. Allows for a generic abstraction away from using something specific like "IsServerOnly" to determine when to strip that kind of data, such that custom target platforms which also don't need AV data can get in on the action

diff --git a/Engine/Source/Developer/TargetPlatform/Public/Common/TargetPlatformBase.h b/Engine/Source/Developer/TargetPlatform/Public/Common/TargetPlatformBase.h
--- a/Engine/Source/Developer/TargetPlatform/Public/Common/TargetPlatformBase.h
+++ b/Engine/Source/Developer/TargetPlatform/Public/Common/TargetPlatformBase.h
@@ -188,6 +188,11 @@ public:
 	}
 
 #if WITH_ENGINE
+	virtual bool AllowAudioVisualData() const override
+	{
+		return !IsServerOnly();
+	}
+
 	virtual FName GetMeshBuilderModuleName() const override
 	{
 		// MeshBuilder is the default module. Platforms may override this to provide platform specific mesh data.
diff --git a/Engine/Source/Developer/TargetPlatform/Public/Interfaces/ITargetPlatform.h b/Engine/Source/Developer/TargetPlatform/Public/Interfaces/ITargetPlatform.h
--- a/Engine/Source/Developer/TargetPlatform/Public/Interfaces/ITargetPlatform.h
+++ b/Engine/Source/Developer/TargetPlatform/Public/Interfaces/ITargetPlatform.h
@@ -556,6 +556,14 @@ public:
 	*/
 	virtual void GetAllWaveFormats( TArray<FName>& OutFormats ) const = 0;
 
+	/**
+	 * Checks whether if this platform wants AV data (defaults to !IsServerOnly(), which is the standard reason why we don't want AV data)
+	 * Used so that custom target platforms can remove AV data, but is not a server-only platform
+	 *
+	 * @return true if this platform allows AV data to be cooked, false otherwise.
+	 */
+	virtual bool AllowAudioVisualData() const = 0;
+
 	/**
 	 * Gets the texture LOD settings used by this platform.
 	 *
diff --git a/Engine/Source/Runtime/Engine/Private/SoundWave.cpp b/Engine/Source/Runtime/Engine/Private/SoundWave.cpp
--- a/Engine/Source/Runtime/Engine/Private/SoundWave.cpp
+++ b/Engine/Source/Runtime/Engine/Private/SoundWave.cpp
@@ -466,7 +466,7 @@ void USoundWave::Serialize( FArchive& Ar )
 #if WITH_ENGINE
 				TArray<FName> ActualFormatsToSave;
 				const ITargetPlatform* CookingTarget = Ar.CookingTarget();
-				if (!CookingTarget->IsServerOnly())
+				if (CookingTarget->AllowAudioVisualData())
 				{
 					// for now we only support one format per wave
 					FName Format = CookingTarget->GetWaveFormat(this);
@@ -1118,7 +1118,7 @@ void USoundWave::PostLoad()
 
 		for (const ITargetPlatform* Platform : Platforms)
 		{
-			if (!Platform->IsServerOnly())
+			if (Platform->AllowAudioVisualData())
 			{
 				BeginGetCompressedData(Platform->GetWaveFormat(this), FPlatformCompressionUtilities::GetCookOverrides(*Platform->IniPlatformName()));
 			}
