Cloth - Partially undo the external collision sources fix that still needs the source component scale after all. The code should now be able to handle all scale situations on either the source or the destination compoments.

diff --git a/Engine/Source/Runtime/Engine/Private/SkeletalMeshComponentPhysics.cpp b/Engine/Source/Runtime/Engine/Private/SkeletalMeshComponentPhysics.cpp
--- a/Engine/Source/Runtime/Engine/Private/SkeletalMeshComponentPhysics.cpp
+++ b/Engine/Source/Runtime/Engine/Private/SkeletalMeshComponentPhysics.cpp
@@ -2774,11 +2774,9 @@ void USkeletalMeshComponent::ExtractCollisionsForCloth(
 		FTransform ComponentToComponentTransform;
 		if(SourceComponent != DestClothComponent)
 		{
-			FTransform SourceClothComponentTransform = SourceComponent->GetComponentTransform();
-			SourceClothComponentTransform.RemoveScaling();  // The source component scale will be present in the world space bone transform, and is not needed here
 			FTransform DestClothComponentTransform = DestClothComponent->GetComponentTransform();
-			DestClothComponentTransform.RemoveScaling();  // The collision source doesn't need the scale of the cloth skeletal mesh applied to it
-			ComponentToComponentTransform = SourceClothComponentTransform * DestClothComponentTransform.Inverse();
+			DestClothComponentTransform.RemoveScaling();  // The collision source doesn't need the scale of the cloth skeletal mesh applied to it (but it does need the source scale from the component transform)
+			ComponentToComponentTransform = SourceComponent->GetComponentTransform() * DestClothComponentTransform.Inverse();
 		}
 
 		// Init cache on first copy
