Use anisotropic sampler with runtime virtual textures
#8006

diff --git a/Engine/Source/Runtime/Engine/Private/Materials/HLSLMaterialTranslator.cpp b/Engine/Source/Runtime/Engine/Private/Materials/HLSLMaterialTranslator.cpp
--- a/Engine/Source/Runtime/Engine/Private/Materials/HLSLMaterialTranslator.cpp
+++ b/Engine/Source/Runtime/Engine/Private/Materials/HLSLMaterialTranslator.cpp
@@ -5533,8 +5533,9 @@ int32 FHLSLMaterialTranslator::TextureSample(
 		if (SamplerSource != SSM_FromTextureAsset)
 		{
 			// VT doesn't care if the shared sampler is wrap or clamp this is handled in the shader explicitly by our code so we still inherit this from the texture
-			TextureName += FString::Printf(TEXT("Material.VirtualTexturePhysical_%d, GetMaterialSharedSampler(Material.VirtualTexturePhysical_%dSampler, View.SharedBilinearClampedSampler)")
-				, VirtualTextureIndex, VirtualTextureIndex);
+			const TCHAR* SharedSamplerName = (MipValueMode == TMVM_MipLevel) ? TEXT("View.SharedBilinearClampedSampler") : TEXT("View.SharedBilinearAnisoClampedSampler");
+			TextureName += FString::Printf(TEXT("Material.VirtualTexturePhysical_%d, GetMaterialSharedSampler(Material.VirtualTexturePhysical_%dSampler, %s)")
+				, VirtualTextureIndex, VirtualTextureIndex, SharedSamplerName);
 		}
 		else
 		{
diff --git a/Engine/Source/Runtime/Engine/Private/SceneManagement.cpp b/Engine/Source/Runtime/Engine/Private/SceneManagement.cpp
--- a/Engine/Source/Runtime/Engine/Private/SceneManagement.cpp
+++ b/Engine/Source/Runtime/Engine/Private/SceneManagement.cpp
@@ -766,6 +766,7 @@ FViewUniformShaderParameters::FViewUniformShaderParameters()
 	SharedPointClampedSampler = TStaticSamplerState<SF_Point, AM_Clamp, AM_Clamp, AM_Clamp>::GetRHI();
 	SharedBilinearWrappedSampler = TStaticSamplerState<SF_Bilinear, AM_Wrap, AM_Wrap, AM_Wrap>::GetRHI();
 	SharedBilinearClampedSampler = TStaticSamplerState<SF_Bilinear, AM_Clamp, AM_Clamp, AM_Clamp>::GetRHI();
+	SharedBilinearAnisoClampedSampler = TStaticSamplerState<SF_AnisotropicLinear, AM_Clamp, AM_Clamp, AM_Clamp, 0, 0>::GetRHI();
 	SharedTrilinearWrappedSampler = TStaticSamplerState<SF_Trilinear, AM_Wrap, AM_Wrap, AM_Wrap>::GetRHI();
 	SharedTrilinearClampedSampler = TStaticSamplerState<SF_Trilinear, AM_Clamp, AM_Clamp, AM_Clamp>::GetRHI();
 
diff --git a/Engine/Source/Runtime/Engine/Public/SceneView.h b/Engine/Source/Runtime/Engine/Public/SceneView.h
--- a/Engine/Source/Runtime/Engine/Public/SceneView.h
+++ b/Engine/Source/Runtime/Engine/Public/SceneView.h
@@ -838,6 +838,7 @@ BEGIN_GLOBAL_SHADER_PARAMETER_STRUCT_WITH_CONSTRUCTOR(FViewUniformShaderParamete
 	SHADER_PARAMETER_SAMPLER(SamplerState, SharedPointClampedSampler)
 	SHADER_PARAMETER_SAMPLER(SamplerState, SharedBilinearWrappedSampler)
 	SHADER_PARAMETER_SAMPLER(SamplerState, SharedBilinearClampedSampler)
+	SHADER_PARAMETER_SAMPLER(SamplerState, SharedBilinearAnisoClampedSampler)
 	SHADER_PARAMETER_SAMPLER(SamplerState, SharedTrilinearWrappedSampler)
 	SHADER_PARAMETER_SAMPLER(SamplerState, SharedTrilinearClampedSampler)
 	SHADER_PARAMETER_TEXTURE(Texture2D, PreIntegratedBRDF)
