Adding config to KillZ

diff --git a/Engine/Source/Runtime/Engine/Classes/GameFramework/WorldSettings.h b/Engine/Source/Runtime/Engine/Classes/GameFramework/WorldSettings.h
--- a/Engine/Source/Runtime/Engine/Classes/GameFramework/WorldSettings.h
+++ b/Engine/Source/Runtime/Engine/Classes/GameFramework/WorldSettings.h
@@ -558,7 +558,7 @@ public:
 	float WorldToMeters;
 
 	// any actor falling below this level gets destroyed
-	UPROPERTY(EditAnywhere, BlueprintReadOnly, Category=World, meta=(editcondition = "bEnableWorldBoundsChecks"))
+	UPROPERTY(config, EditAnywhere, BlueprintReadOnly, Category=World, meta=(editcondition = "bEnableWorldBoundsChecks"))
 	float KillZ;
 
 	// The type of damage inflicted when a actor falls below KillZ
