Adding RESET_REPLIFETIME macros that support push model

diff --git a/Engine/Source/Runtime/Engine/Private/UnrealNetwork.cpp b/Engine/Source/Runtime/Engine/Private/UnrealNetwork.cpp
--- a/Engine/Source/Runtime/Engine/Private/UnrealNetwork.cpp
+++ b/Engine/Source/Runtime/Engine/Private/UnrealNetwork.cpp
@@ -149,6 +149,38 @@ void ResetReplicatedLifetimeProperty(const UClass* ThisClass, const UClass* Prop
 	ResetReplicatedLifetimeProperty(NetworkingPrivate::FRepPropertyDescriptor(ReplicatedProperty), LifetimeCondition, OutLifetimeProps);
 }
 
+void ResetReplicatedLifetimeProperty(const UClass* ThisClass, const UClass* PropertyClass, FName PropertyName, const FDoRepLifetimeParams& Params, TArray< FLifetimeProperty >& OutLifetimeProps)
+{
+	const FProperty* ReplicatedProperty = GetReplicatedProperty(ThisClass, PropertyClass, PropertyName);
+	if (!ReplicatedProperty)
+	{
+		return;
+	}
+
+	ResetReplicatedLifetimeProperty(NetworkingPrivate::FRepPropertyDescriptor(ReplicatedProperty), Params, OutLifetimeProps);
+}
+
+void ResetReplicatedLifetimeProperty(const NetworkingPrivate::FRepPropertyDescriptor& PropertyDescriptor, const FDoRepLifetimeParams& Params, TArray< FLifetimeProperty >& OutLifetimeProps)
+{
+	for (int32 i = 0; i < PropertyDescriptor.ArrayDim; i++)
+	{
+		uint16 RepIndex = PropertyDescriptor.RepIndex + i;
+		FLifetimeProperty* RegisteredPropertyPtr = OutLifetimeProps.FindByPredicate([&RepIndex](const FLifetimeProperty& Var) { return Var.RepIndex == RepIndex; });
+
+		// Set the new condition
+		if (RegisteredPropertyPtr)
+		{
+			RegisteredPropertyPtr->bIsPushBased = Params.bIsPushBased;
+			RegisteredPropertyPtr->Condition = Params.Condition;
+			RegisteredPropertyPtr->RepNotifyCondition = Params.RepNotifyCondition;
+		}
+		else
+		{
+			OutLifetimeProps.Add(FLifetimeProperty(RepIndex, Params.Condition, Params.RepNotifyCondition, Params.bIsPushBased));
+		}
+	}
+}
+
 void DisableAllReplicatedPropertiesOfClass(const NetworkingPrivate::FRepClassDescriptor& ClassDescriptor, EFieldIteratorFlags::SuperClassFlags SuperClassBehavior, TArray<FLifetimeProperty>& OutLifetimeProps)
 {
 	const int32 StartIndex = (EFieldIteratorFlags::IncludeSuper == SuperClassBehavior) ? 0 : ClassDescriptor.StartRepIndex;
diff --git a/Engine/Source/Runtime/Engine/Public/Net/UnrealNetwork.h b/Engine/Source/Runtime/Engine/Public/Net/UnrealNetwork.h
--- a/Engine/Source/Runtime/Engine/Public/Net/UnrealNetwork.h
+++ b/Engine/Source/Runtime/Engine/Public/Net/UnrealNetwork.h
@@ -384,6 +384,16 @@ ENGINE_API void DisableAllReplicatedPropertiesOfClass(const NetworkingPrivate::F
 #define RESET_REPLIFETIME_FAST(c,v) RESET_REPLIFETIME_CONDITION_FAST(c, v, COND_None)
 #define RESET_REPLIFETIME_FAST_STATIC_ARRAY(c,v) RESET_REPLIFETIME_FAST_STATIC_ARRAY(c, v, COND_None)
 
+#define RESET_REPLIFETIME_WITH_PARAMS(c,v,params)  ResetReplicatedLifetimeProperty(StaticClass(), c::StaticClass(), GET_MEMBER_NAME_CHECKED(c,v), params, OutLifetimeProps);
+
+#define RESET_REPLIFETIME_FAST_WITH_PARAMS(c,v,params) \
+{ \
+	static const bool bIsValid_##c_##v = ValidateReplicatedClassInheritance(StaticClass(), c::StaticClass(), TEXT(#v)); \
+	const TCHAR* DoRepPropertyName_##c_##v(TEXT(#v)); \
+	const NetworkingPrivate::FRepPropertyDescriptor PropertyDescriptor_##c_##v(DoRepPropertyName_##c_##v, (int32)c::ENetFields_Private::v, 1); \
+	ResetReplicatedLifetimeProperty(StaticClass(), c::StaticClass(), GET_MEMBER_NAME_CHECKED(c,v), params, OutLifetimeProps); \
+}
+
 ENGINE_API void ResetReplicatedLifetimeProperty(
 	const UClass* ThisClass,
 	const UClass* PropertyClass,
@@ -396,6 +406,19 @@ ENGINE_API void ResetReplicatedLifetimeProperty(
 	ELifetimeCondition LifetimeCondition,
 	TArray<FLifetimeProperty>& OutLifetimeProps);
 
+ENGINE_API void ResetReplicatedLifetimeProperty(
+	const UClass* ThisClass,
+	const UClass* PropertyClass,
+	FName PropertyName,
+	const FDoRepLifetimeParams& Params,
+	TArray< FLifetimeProperty >& OutLifetimeProps);
+
+ENGINE_API void ResetReplicatedLifetimeProperty(
+	const NetworkingPrivate::FRepPropertyDescriptor& PropertyDescriptor,
+	const FDoRepLifetimeParams& Params,
+	TArray< FLifetimeProperty >& OutLifetimeProps);
+
+
 /*-----------------------------------------------------------------------------
 	RPC Parameter Validation Helpers
 -----------------------------------------------------------------------------*/
