Fix up new 'instance' material transform to work with LWC

diff --git a/Engine/Source/Runtime/Engine/Private/Materials/HLSLMaterialTranslator.cpp b/Engine/Source/Runtime/Engine/Private/Materials/HLSLMaterialTranslator.cpp
--- a/Engine/Source/Runtime/Engine/Private/Materials/HLSLMaterialTranslator.cpp
+++ b/Engine/Source/Runtime/Engine/Private/Materials/HLSLMaterialTranslator.cpp
@@ -8276,7 +8276,7 @@ int32 FHLSLMaterialTranslator::TransformBase(EMaterialCommonBasis SourceCoordBas
 			}
 			else if (DestCoordBasis == MCB_Instance)
 			{
-				CodeStr = TEXT("mul(<A>, <MATRIX>GetWorldToInstance(Parameters))");
+				CodeStr = LWCMultiplyMatrix(TEXT("<A>"), TEXT("GetWorldToInstance(Parameters)"), AWComponent);
 				bUsesInstanceWorldToLocalPS = ShaderFrequency == SF_Pixel;
 			}
 
@@ -8318,7 +8318,7 @@ int32 FHLSLMaterialTranslator::TransformBase(EMaterialCommonBasis SourceCoordBas
 		{
 			if (DestCoordBasis == MCB_World)
 			{
-				CodeStr = TEXT("mul(<A>, <MATRIX>GetInstanceToWorld(Parameters))");
+				CodeStr = LWCMultiplyMatrix(TEXT("<A>"), TEXT("GetInstanceToWorld(Parameters)"), AWComponent);
 				bUsesInstanceLocalToWorldPS = ShaderFrequency == SF_Pixel;
 			}
 			// use World as an intermediary base
