Added property handle param to IBoneReferenceSkeletonProvider::GetSkeleton so implementers can have multiple FBoneReference properties using different skeletons

diff --git a/Engine/Plugins/Experimental/Animation/PoseSearch/Source/Runtime/Public/PoseSearch/PoseSearch.h b/Engine/Plugins/Experimental/Animation/PoseSearch/Source/Runtime/Public/PoseSearch/PoseSearch.h
--- a/Engine/Plugins/Experimental/Animation/PoseSearch/Source/Runtime/Public/PoseSearch/PoseSearch.h
+++ b/Engine/Plugins/Experimental/Animation/PoseSearch/Source/Runtime/Public/PoseSearch/PoseSearch.h
@@ -215,7 +215,7 @@ public: // UObject
 	virtual void PostLoad() override;
 
 public: // IBoneReferenceSkeletonProvider
-	class USkeleton* GetSkeleton(bool& bInvalidSkeletonIsError) override { bInvalidSkeletonIsError = false; return Skeleton; }
+	class USkeleton* GetSkeleton(bool& bInvalidSkeletonIsError, const IPropertyHandle* PropertyHandle) override { bInvalidSkeletonIsError = false; return Skeleton; }
 
 private:
 	void GenerateLayout();
diff --git a/Engine/Plugins/Experimental/MeshModelingToolsetExp/Source/MeshModelingToolsExp/Private/SkinWeightsBindingTool.cpp b/Engine/Plugins/Experimental/MeshModelingToolsetExp/Source/MeshModelingToolsExp/Private/SkinWeightsBindingTool.cpp
--- a/Engine/Plugins/Experimental/MeshModelingToolsetExp/Source/MeshModelingToolsExp/Private/SkinWeightsBindingTool.cpp
+++ b/Engine/Plugins/Experimental/MeshModelingToolsetExp/Source/MeshModelingToolsExp/Private/SkinWeightsBindingTool.cpp
@@ -736,7 +736,7 @@ UInteractiveTool* USkinWeightsBindingToolBuilder::BuildTool(const FToolBuilderSt
 }
 
 
-USkeleton* USkinWeightsBindingToolProperties::GetSkeleton(bool& bInvalidSkeletonIsError)
+USkeleton* USkinWeightsBindingToolProperties::GetSkeleton(bool& bInvalidSkeletonIsError, const IPropertyHandle* PropertyHandle)
 {
 	bInvalidSkeletonIsError = false;
 	return SkeletalMesh ? SkeletalMesh->GetSkeleton() : nullptr;
diff --git a/Engine/Plugins/Experimental/MeshModelingToolsetExp/Source/MeshModelingToolsExp/Private/SkinWeightsPaintTool.cpp b/Engine/Plugins/Experimental/MeshModelingToolsetExp/Source/MeshModelingToolsExp/Private/SkinWeightsPaintTool.cpp
--- a/Engine/Plugins/Experimental/MeshModelingToolsetExp/Source/MeshModelingToolsExp/Private/SkinWeightsPaintTool.cpp
+++ b/Engine/Plugins/Experimental/MeshModelingToolsetExp/Source/MeshModelingToolsExp/Private/SkinWeightsPaintTool.cpp
@@ -722,7 +722,7 @@ void USkinWeightsPaintTool::CommitSkinWeights(const IMeshDescriptionCommitter::F
 }
 
 
-USkeleton* USkinWeightsPaintToolProperties::GetSkeleton(bool& bInvalidSkeletonIsError)
+USkeleton* USkinWeightsPaintToolProperties::GetSkeleton(bool& bInvalidSkeletonIsError, const IPropertyHandle* PropertyHandle)
 {
 	bInvalidSkeletonIsError = false;
 	return SkeletalMesh ? SkeletalMesh->GetSkeleton() : nullptr;
diff --git a/Engine/Plugins/Experimental/MeshModelingToolsetExp/Source/MeshModelingToolsExp/Public/SkinWeightsBindingTool.h b/Engine/Plugins/Experimental/MeshModelingToolsetExp/Source/MeshModelingToolsExp/Public/SkinWeightsBindingTool.h
--- a/Engine/Plugins/Experimental/MeshModelingToolsetExp/Source/MeshModelingToolsExp/Public/SkinWeightsBindingTool.h
+++ b/Engine/Plugins/Experimental/MeshModelingToolsetExp/Source/MeshModelingToolsExp/Public/SkinWeightsBindingTool.h
@@ -69,7 +69,7 @@ public:
 	bool bDebugDraw = false;
 	
 	// IBoneReferenceSkeletonProvider
-	USkeleton* GetSkeleton(bool& bInvalidSkeletonIsError) override;
+	USkeleton* GetSkeleton(bool& bInvalidSkeletonIsError, const IPropertyHandle* PropertyHandle) override;
 
 	TObjectPtr<USkeletalMesh> SkeletalMesh;
 };
diff --git a/Engine/Plugins/Experimental/MeshModelingToolsetExp/Source/MeshModelingToolsExp/Public/SkinWeightsPaintTool.h b/Engine/Plugins/Experimental/MeshModelingToolsetExp/Source/MeshModelingToolsExp/Public/SkinWeightsPaintTool.h
--- a/Engine/Plugins/Experimental/MeshModelingToolsetExp/Source/MeshModelingToolsExp/Public/SkinWeightsPaintTool.h
+++ b/Engine/Plugins/Experimental/MeshModelingToolsetExp/Source/MeshModelingToolsExp/Public/SkinWeightsPaintTool.h
@@ -71,7 +71,7 @@ public:
 	FBoneReference CurrentBone;
 
 	// IBoneReferenceSkeletonProvider
-	USkeleton* GetSkeleton(bool& bInvalidSkeletonIsError) override;
+	USkeleton* GetSkeleton(bool& bInvalidSkeletonIsError, const IPropertyHandle* PropertyHandle) override;
 
 	TObjectPtr<USkeletalMesh> SkeletalMesh;
 };
diff --git a/Engine/Source/Editor/Persona/Private/AnimGraphNodeDetails.cpp b/Engine/Source/Editor/Persona/Private/AnimGraphNodeDetails.cpp
--- a/Engine/Source/Editor/Persona/Private/AnimGraphNodeDetails.cpp
+++ b/Engine/Source/Editor/Persona/Private/AnimGraphNodeDetails.cpp
@@ -510,7 +510,7 @@ void FBoneReferenceCustomization::SetEditableSkeleton(TSharedRef<IPropertyHandle
 
 	bEnsureOnInvalidSkeleton = true;
 
-	auto FindSkeletonForObject = [this, &TargetSkeleton, &EditableSkeleton](UObject* InObject)
+	auto FindSkeletonForObject = [this, &TargetSkeleton, &EditableSkeleton, &StructPropertyHandle](UObject* InObject)
 	{
 		for( ; InObject; InObject = InObject->GetOuter())
 		{
@@ -567,7 +567,7 @@ void FBoneReferenceCustomization::SetEditableSkeleton(TSharedRef<IPropertyHandle
 
 			if (IBoneReferenceSkeletonProvider* SkeletonProvider = Cast<IBoneReferenceSkeletonProvider>(InObject))
 			{
-				TargetSkeleton = SkeletonProvider->GetSkeleton(bEnsureOnInvalidSkeleton);
+				TargetSkeleton = SkeletonProvider->GetSkeleton(bEnsureOnInvalidSkeleton, &StructPropertyHandle.Get());
 				break;
 			}
 		}
diff --git a/Engine/Source/Runtime/Engine/Classes/Interfaces/Interface_BoneReferenceSkeletonProvider.h b/Engine/Source/Runtime/Engine/Classes/Interfaces/Interface_BoneReferenceSkeletonProvider.h
--- a/Engine/Source/Runtime/Engine/Classes/Interfaces/Interface_BoneReferenceSkeletonProvider.h
+++ b/Engine/Source/Runtime/Engine/Classes/Interfaces/Interface_BoneReferenceSkeletonProvider.h
@@ -21,6 +21,9 @@ class ENGINE_API IBoneReferenceSkeletonProvider
 
 public:
 
+	UE_DEPRECATED(5.0, "Please use GetSkeleton(bool& bInvalidSkeletonIsError, const class IPropertyHandle* PropertyHandle)")
+	virtual class USkeleton* GetSkeleton(bool& bInvalidSkeletonIsError) { return GetSkeleton(bInvalidSkeletonIsError, nullptr); }
+
 	/**
 	 * Called to get the skeleton that FBoneReference's details customization will use to populate
 	 * bone names.
@@ -29,5 +32,5 @@ public:
 	 *
 	 * @return The skeleton we should use.
 	 */
-	virtual class USkeleton* GetSkeleton(bool& bInvalidSkeletonIsError) PURE_VIRTUAL(IBoneReferenceSkeletonProvider::GetSkeleton, return nullptr; );
+	virtual class USkeleton* GetSkeleton(bool& bInvalidSkeletonIsError, const class IPropertyHandle* PropertyHandle) PURE_VIRTUAL(IBoneReferenceSkeletonProvider::GetSkeleton, return nullptr; );
 };
