Added "stat help" command to print all STAT available

diff --git a/Engine/Source/Runtime/Engine/Private/UnrealEngine.cpp b/Engine/Source/Runtime/Engine/Private/UnrealEngine.cpp
--- a/Engine/Source/Runtime/Engine/Private/UnrealEngine.cpp
+++ b/Engine/Source/Runtime/Engine/Private/UnrealEngine.cpp
@@ -4729,6 +4729,16 @@ bool UEngine::HandleGameVerCommand( const TCHAR* Cmd, FOutputDevice& Ar )
 
 bool UEngine::HandleStatCommand( UWorld* World, FCommonViewportClient* ViewportClient, const TCHAR* Cmd, FOutputDevice& Ar )
 {
+	if (FParse::Command(&Cmd, TEXT("help")))
+	{
+		for (int32 StatIdx = 0; StatIdx < EngineStats.Num(); StatIdx++)
+		{
+			const FEngineStatFuncs& EngineStat = EngineStats[StatIdx];
+			Ar.Logf(TEXT("%s"), *EngineStat.CommandNameString);
+		}
+		return true;
+	}
+
 	const TCHAR* Temp = Cmd;
 	for (int32 StatIdx = 0; StatIdx < EngineStats.Num(); StatIdx++)
 	{
