RBAN memory usage optimization
Make UniqueIndices object dependency injectable into the SOAs object
Add a simple implementation for RBAN ( imamediate physics ) that does not need pre-allocation and lockfree list

diff --git a/Engine/Source/Runtime/Engine/Private/PhysicsEngine/ImmediatePhysics/ImmediatePhysicsChaos/ImmediatePhysicsSimulation_Chaos.cpp b/Engine/Source/Runtime/Engine/Private/PhysicsEngine/ImmediatePhysics/ImmediatePhysicsChaos/ImmediatePhysicsSimulation_Chaos.cpp
--- a/Engine/Source/Runtime/Engine/Private/PhysicsEngine/ImmediatePhysics/ImmediatePhysicsChaos/ImmediatePhysicsSimulation_Chaos.cpp
+++ b/Engine/Source/Runtime/Engine/Private/PhysicsEngine/ImmediatePhysics/ImmediatePhysicsChaos/ImmediatePhysicsSimulation_Chaos.cpp
@@ -215,6 +215,37 @@ FAutoConsoleVariableRef CVarChaosImmPhysScale(TEXT("p.Chaos.ImmPhys.DebugDraw.Sc
 
 namespace ImmediatePhysics_Chaos
 {
+
+	class FSimpleParticleUniqueIndices : public Chaos::IParticleUniqueIndices
+	{
+	public:
+		Chaos::FUniqueIdx GenerateUniqueIdx() override
+		{
+			if (FreeIndices.Num())
+			{
+				const int32 FreeIndex = FreeIndices.Pop();
+				return Chaos::FUniqueIdx(FreeIndex);
+			}
+
+			Chaos::FUniqueIdx NewUniqueIndex(NextUniqueIndex);
+			NextUniqueIndex++;
+			return NewUniqueIndex;
+		}
+
+		void ReleaseIdx(Chaos::FUniqueIdx Unique) override
+		{
+			ensure(Unique.IsValid());
+			FreeIndices.Add(Unique.Idx);
+		}
+
+		~FSimpleParticleUniqueIndices() = default;
+
+	private:
+		int32 NextUniqueIndex; // this includes all valid and freed indices
+		TArray<int32> FreeIndices;
+	};
+
+
 	struct FSimulation::FImplementation
 	{
 	public:
@@ -225,7 +256,7 @@ namespace ImmediatePhysics_Chaos
 		using FRigidParticleSOAs = Chaos::FPBDRigidsSOAs;
 
 		FImplementation()
-			: Particles()
+			: Particles(UniqueIndices)
 			, Joints()
 			, Collisions(Particles, CollidedParticles, ParticleMaterials, PerParticleMaterials, 0, 0, ChaosImmediate_Collision_CullDistance)
 			, BroadPhase(&ActivePotentiallyCollidingPairs, nullptr, nullptr, ChaosImmediate_Collision_CullDistance)
@@ -265,6 +296,7 @@ namespace ImmediatePhysics_Chaos
 		Chaos::TArrayCollectionArray<Chaos::FRotation3> ParticlePrevRs;
 
 		FRigidParticleSOAs Particles;
+		FSimpleParticleUniqueIndices UniqueIndices;
 		Chaos::FPBDJointConstraints Joints;
 		FCollisionConstraints Collisions;
 		Chaos::FParticlePairBroadPhase BroadPhase;
diff --git a/Engine/Source/Runtime/Experimental/Chaos/Private/PBDRigidsSolver.cpp b/Engine/Source/Runtime/Experimental/Chaos/Private/PBDRigidsSolver.cpp
--- a/Engine/Source/Runtime/Experimental/Chaos/Private/PBDRigidsSolver.cpp
+++ b/Engine/Source/Runtime/Experimental/Chaos/Private/PBDRigidsSolver.cpp
@@ -397,6 +397,7 @@ namespace Chaos
 		, bHasFloor(true)
 		, bIsFloorAnalytic(false)
 		, FloorHeight(0.f)
+		, Particles(UniqueIndices)
 		, MEvolution(new FPBDRigidsEvolution(Particles, SimMaterials, &ContactModifiers, BufferingModeIn == Chaos::EMultiBufferMode::Single))
 		, MEventManager(new FEventManager(BufferingModeIn))
 		, MSolverEventFilters(new FSolverEventFilters())
diff --git a/Engine/Source/Runtime/Experimental/Chaos/Public/Chaos/PBDRigidsSOAs.h b/Engine/Source/Runtime/Experimental/Chaos/Public/Chaos/PBDRigidsSOAs.h
--- a/Engine/Source/Runtime/Experimental/Chaos/Public/Chaos/PBDRigidsSOAs.h
+++ b/Engine/Source/Runtime/Experimental/Chaos/Public/Chaos/PBDRigidsSOAs.h
@@ -6,17 +6,25 @@
 
 namespace Chaos
 {
-class FParticleUniqueIndices
+class IParticleUniqueIndices
 {
 public:
-	FParticleUniqueIndices()
+	virtual ~IParticleUniqueIndices() = default;
+	virtual FUniqueIdx GenerateUniqueIdx() = 0;
+	virtual void ReleaseIdx(FUniqueIdx Unique) = 0;
+};
+
+class FParticleUniqueIndicesMultithreaded: public IParticleUniqueIndices
+{
+public:
+	FParticleUniqueIndicesMultithreaded()
 		: Block(0)
 	{
 		//Note: tune this so that all allocation is done at initialization
 		AddPageAndAcquireNextId(/*bAcquireNextId = */ false);
 	}
 
-	FUniqueIdx GenerateUniqueIdx()
+	FUniqueIdx GenerateUniqueIdx() override
 	{
 		while(true)
 		{
@@ -42,7 +50,7 @@ public:
 		}
 	}
 
-	void ReleaseIdx(FUniqueIdx Unique)
+	void ReleaseIdx(FUniqueIdx Unique) override
 	{
 		ensure(Unique.IsValid());
 		int32 PageIdx = Unique.Idx / IndicesPerPage;
@@ -51,7 +59,7 @@ public:
 		FreeIndices.Push(&Page.Indices[Entry]);
 	}
 
-	~FParticleUniqueIndices()
+	~FParticleUniqueIndicesMultithreaded()
 	{
 		//Make sure queue is empty, memory management of actual pages is handled automatically by TUniquePtr
 		while(FreeIndices.Pop())
@@ -90,6 +98,8 @@ private:
 	volatile int8 Block;
 };
 
+using FParticleUniqueIndices = FParticleUniqueIndicesMultithreaded;
+
 template <typename TParticleType>
 class TParticleMapArray
 {
@@ -180,7 +190,8 @@ private:
 class FPBDRigidsSOAs
 {
 public:
-	FPBDRigidsSOAs()
+	FPBDRigidsSOAs(IParticleUniqueIndices& InUniqueIndices)
+		: UniqueIndices(InUniqueIndices)
 	{
 #if CHAOS_DETERMINISTIC
 		BiggestParticleID = 0;
@@ -1159,7 +1170,7 @@ private:
 	//Auxiliary data synced with particle handles
 	TGeometryParticleHandles<FReal, 3> ParticleHandles;
 
-	FParticleUniqueIndices UniqueIndices;
+	IParticleUniqueIndices& UniqueIndices;
 
 #if CHAOS_DETERMINISTIC
 	int32 BiggestParticleID;
diff --git a/Engine/Source/Runtime/Experimental/Chaos/Public/PBDRigidsSolver.h b/Engine/Source/Runtime/Experimental/Chaos/Public/PBDRigidsSolver.h
--- a/Engine/Source/Runtime/Experimental/Chaos/Public/PBDRigidsSolver.h
+++ b/Engine/Source/Runtime/Experimental/Chaos/Public/PBDRigidsSolver.h
@@ -320,6 +320,7 @@ namespace Chaos
 		FReal FloorHeight;
 
 		FParticlesType Particles;
+		FParticleUniqueIndicesMultithreaded UniqueIndices;
 		TUniquePtr<FPBDRigidsEvolution> MEvolution;
 		TUniquePtr<FEventManager> MEventManager;
 		TUniquePtr<FSolverEventFilters> MSolverEventFilters;
diff --git a/Engine/Source/Runtime/PhysicsCore/Private/PhysTestSerializer.cpp b/Engine/Source/Runtime/PhysicsCore/Private/PhysTestSerializer.cpp
--- a/Engine/Source/Runtime/PhysicsCore/Private/PhysTestSerializer.cpp
+++ b/Engine/Source/Runtime/PhysicsCore/Private/PhysTestSerializer.cpp
@@ -29,6 +29,7 @@ using namespace Chaos;
 FPhysTestSerializer::FPhysTestSerializer()
 	: bDiskDataIsChaos(false)
 	, bChaosDataReady(false)
+	, Particles(UniqueIndices)
 {
 }
 
diff --git a/Engine/Source/Runtime/PhysicsCore/Public/PhysTestSerializer.h b/Engine/Source/Runtime/PhysicsCore/Public/PhysTestSerializer.h
--- a/Engine/Source/Runtime/PhysicsCore/Public/PhysTestSerializer.h
+++ b/Engine/Source/Runtime/PhysicsCore/Public/PhysTestSerializer.h
@@ -141,6 +141,7 @@ private:
 
 	TUniquePtr<Chaos::FPBDRigidsEvolution> ChaosEvolution;
 	Chaos::FPBDRigidsSOAs Particles;
+	Chaos::FParticleUniqueIndicesMultithreaded UniqueIndices;
 	Chaos::THandleArray<Chaos::FChaosPhysicsMaterial> PhysicalMaterials;
 	TArray <TUniquePtr<Chaos::FGeometryParticle>> GTParticles;
 
