Statgroups are added to the list of stats printed when running "stat help"

diff --git a/Engine/Source/Runtime/Engine/Private/UnrealEngine.cpp b/Engine/Source/Runtime/Engine/Private/UnrealEngine.cpp
--- a/Engine/Source/Runtime/Engine/Private/UnrealEngine.cpp
+++ b/Engine/Source/Runtime/Engine/Private/UnrealEngine.cpp
@@ -4731,11 +4731,27 @@ bool UEngine::HandleStatCommand( UWorld* World, FCommonViewportClient* ViewportC
 {
 	if (FParse::Command(&Cmd, TEXT("help")))
 	{
+		const TSet<FName>& StatGroupNames = FStatGroupGameThreadNotifier::Get().StatGroupNames;
+		TArray<FString> StatNames;
+		for (const FName& StatName : StatGroupNames)
+		{
+			FString StatNameStr = StatName.ToString();
+			StatNameStr.RemoveFromStart(TEXT("STATGROUP_"));
+			StatNames.Add_GetRef(StatNameStr);
+		}
+
 		for (int32 StatIdx = 0; StatIdx < EngineStats.Num(); StatIdx++)
 		{
 			const FEngineStatFuncs& EngineStat = EngineStats[StatIdx];
-			Ar.Logf(TEXT("%s"), *EngineStat.CommandNameString);
+			StatNames.Add_GetRef( EngineStat.CommandNameString);
 		}
+
+		StatNames.Sort();
+		for (const FString& StatName : StatNames)
+		{
+			Ar.Logf(TEXT("%s"), *StatName);
+		}
+
 		return true;
 	}
 
