Fixed compile error

#horde 109527

diff --git a/Engine/Source/Runtime/Engine/Private/UnrealEngine.cpp b/Engine/Source/Runtime/Engine/Private/UnrealEngine.cpp
--- a/Engine/Source/Runtime/Engine/Private/UnrealEngine.cpp
+++ b/Engine/Source/Runtime/Engine/Private/UnrealEngine.cpp
@@ -4729,6 +4729,7 @@ bool UEngine::HandleGameVerCommand( const TCHAR* Cmd, FOutputDevice& Ar )
 
 bool UEngine::HandleStatCommand( UWorld* World, FCommonViewportClient* ViewportClient, const TCHAR* Cmd, FOutputDevice& Ar )
 {
+#if STATS
 	if (FParse::Command(&Cmd, TEXT("help")))
 	{
 		const TSet<FName>& StatGroupNames = FStatGroupGameThreadNotifier::Get().StatGroupNames;
@@ -4754,6 +4755,7 @@ bool UEngine::HandleStatCommand( UWorld* World, FCommonViewportClient* ViewportC
 
 		return true;
 	}
+#endif //STATS
 
 	const TCHAR* Temp = Cmd;
 	for (int32 StatIdx = 0; StatIdx < EngineStats.Num(); StatIdx++)
