Added CPUData to the EngineAnalyticsSessionSummary to get full information about the CPU.
Fixed CPUBrand that was never set in the CrashDescription (emitted as CRC analytics)

diff --git a/Engine/Source/Runtime/CrashReportCore/Private/CrashDescription.cpp b/Engine/Source/Runtime/CrashReportCore/Private/CrashDescription.cpp
--- a/Engine/Source/Runtime/CrashReportCore/Private/CrashDescription.cpp
+++ b/Engine/Source/Runtime/CrashReportCore/Private/CrashDescription.cpp
@@ -572,6 +572,10 @@ void FCrashContext::SetupPrimaryCrashProperties()
 		GetCrashProperty( bInBackground, FGenericCrashContext::GameDataTag, TEXT("bInBackground"));
 		GetCrashProperty( bIsRequestingExit, FGenericCrashContext::GameDataTag, TEXT("IsRequestingExit"));
 
+		FString CPUBrandString;
+		GetCrashProperty( CPUBrandString, FGenericCrashContext::RuntimePropertiesTag, TEXT("Misc.CPUBrand"));
+		CPUBrand = CPUBrandString;
+
 		if (CrashDumpMode == ECrashDumpMode::FullDump)
 		{
 			// Set the full dump crash location when we have a full dump.
diff --git a/Engine/Source/Runtime/Engine/Private/Analytics/EditorAnalyticsSessionSummary.cpp b/Engine/Source/Runtime/Engine/Private/Analytics/EditorAnalyticsSessionSummary.cpp
--- a/Engine/Source/Runtime/Engine/Private/Analytics/EditorAnalyticsSessionSummary.cpp
+++ b/Engine/Source/Runtime/Engine/Private/Analytics/EditorAnalyticsSessionSummary.cpp
@@ -59,7 +59,8 @@ FEditorAnalyticsSessionSummary::FEditorAnalyticsSessionSummary(TSharedPtr<IAnaly
 	// also be incremented when we change some behaviors to be able to compare between versions. Mostly useful for 'dev' branches where users bypasses UGS to update/build.
 	//    - V3 -> Windows optimization for stall/ensure -> The engine only captures the responsible thread so CRC walks 1 thread rather than all threads.
 	//    - V4 -> Added DirtyPackageCount property.
-	EditorAnalyticsProperties::SummaryEventVersion.Set(GetStore(), 4);
+	//    - V5 -> Added CPUInfo to the base class (EngineAnalyticsSessionSummary)
+	EditorAnalyticsProperties::SummaryEventVersion.Set(GetStore(), 5);
 
 	// Persist the session to disk.
 	GetStore()->Flush();
diff --git a/Engine/Source/Runtime/Engine/Private/Analytics/EngineAnalyticsSessionSummary.cpp b/Engine/Source/Runtime/Engine/Private/Analytics/EngineAnalyticsSessionSummary.cpp
--- a/Engine/Source/Runtime/Engine/Private/Analytics/EngineAnalyticsSessionSummary.cpp
+++ b/Engine/Source/Runtime/Engine/Private/Analytics/EngineAnalyticsSessionSummary.cpp
@@ -160,6 +160,7 @@ FEngineAnalyticsSessionSummary::FEngineAnalyticsSessionSummary(TSharedPtr<IAnaly
 	Store->Set(TEXT("CPULogicalCores"), FPlatformMisc::NumberOfCoresIncludingHyperthreads());
 	Store->Set(TEXT("CPUVendor"), FPlatformMisc::GetCPUVendor());
 	Store->Set(TEXT("CPUBrand"), FPlatformMisc::GetCPUBrand());
+	Store->Set(TEXT("CPUInfo"), FPlatformMisc::GetCPUInfo()); // This is a bitfield. See the function documentation in GenericPlatformMisc.h to interpret the bits.
 	Store->Set(TEXT("DesktopGPUAdapter"), FPlatformMisc::GetPrimaryGPUBrand());
 	Store->Set(TEXT("RenderingGPUAdapter"), GRHIAdapterName);
 	Store->Set(TEXT("GPUVendorID"), GRHIVendorId);
