Simulate engine tick in UWorldPartitionBuilder::Run() when processing with loading mode == ELoadingMode::IterativeCells
* Fixes issue with VT allocations not being freed
* Reused code from RuntimeSpatialHashHLOD.cpp, now exposed in the FWorldPartitionHelpers utility class

diff --git a/Engine/Source/Editor/UnrealEd/Private/WorldPartition/WorldPartitionBuilder.cpp b/Engine/Source/Editor/UnrealEd/Private/WorldPartition/WorldPartitionBuilder.cpp
--- a/Engine/Source/Editor/UnrealEd/Private/WorldPartition/WorldPartitionBuilder.cpp
+++ b/Engine/Source/Editor/UnrealEd/Private/WorldPartition/WorldPartitionBuilder.cpp
@@ -212,6 +212,12 @@ bool UWorldPartitionBuilder::Run(UWorld* World, FPackageSourceControlHelper& Pac
 
 						FWorldPartitionHelpers::DoCollectGarbage();
 					}
+
+					// When running with -AllowCommandletRendering we want to simulate an engine tick
+					if (IsAllowCommandletRendering())
+					{
+						FWorldPartitionHelpers::FakeEngineTick(World);
+					}
 				}
 			}
 		}
diff --git a/Engine/Source/Runtime/Engine/Private/WorldPartition/RuntimeSpatialHash/RuntimeSpatialHashHLOD.cpp b/Engine/Source/Runtime/Engine/Private/WorldPartition/RuntimeSpatialHash/RuntimeSpatialHashHLOD.cpp
--- a/Engine/Source/Runtime/Engine/Private/WorldPartition/RuntimeSpatialHash/RuntimeSpatialHashHLOD.cpp
+++ b/Engine/Source/Runtime/Engine/Private/WorldPartition/RuntimeSpatialHash/RuntimeSpatialHashHLOD.cpp
@@ -106,26 +106,7 @@ static void GameTick(UWorld* InWorld)
 	// When running with -AllowCommandletRendering we want to flush
 	if (((++TickRendering % FlushRenderingFrequency) == 0) && IsAllowCommandletRendering())
 	{
-		if (FSceneInterface* Scene = InWorld->Scene)
-		{
-			// BeingFrame/EndFrame (taken from FEngineLoop)
-			uint64 CurrentFrameCounter = GFrameCounter;
-
-			ENQUEUE_RENDER_COMMAND(BeginFrame)([CurrentFrameCounter](FRHICommandListImmediate& RHICmdList)
-			{ 
-				GFrameNumberRenderThread++;
-				RHICmdList.BeginFrame();
-				FCoreDelegates::OnBeginFrameRT.Broadcast(); 
-			});
-
-			ENQUEUE_RENDER_COMMAND(EndFrame)([CurrentFrameCounter](FRHICommandListImmediate& RHICmdList)
-			{
-				FCoreDelegates::OnEndFrameRT.Broadcast();
-				RHICmdList.EndFrame(); 
-			});
-
-			FlushRenderingCommands();
-		}
+		FWorldPartitionHelpers::FakeEngineTick(InWorld);
 	}
 }
 
diff --git a/Engine/Source/Runtime/Engine/Private/WorldPartition/WorldPartitionHelpers.cpp b/Engine/Source/Runtime/Engine/Private/WorldPartition/WorldPartitionHelpers.cpp
--- a/Engine/Source/Runtime/Engine/Private/WorldPartition/WorldPartitionHelpers.cpp
+++ b/Engine/Source/Runtime/Engine/Private/WorldPartition/WorldPartitionHelpers.cpp
@@ -2,6 +2,8 @@
 
 #include "WorldPartition/WorldPartitionHelpers.h"
 
+#include "Engine/Engine.h"
+#include "RenderingThread.h"
 #include "WorldPartition/WorldPartition.h"
 #include "WorldPartition/WorldPartitionEditorHash.h"
 
@@ -130,4 +132,33 @@ void FWorldPartitionHelpers::DoCollectGarbage()
 	);
 };
 
+void FWorldPartitionHelpers::FakeEngineTick(UWorld* InWorld)
+{
+	check(InWorld);
+
+	// Simulate an engine frame tick
+	// Will make sure systems can perform their internal bookkeeping properly. For example, the VT system needs to 
+	// process deleted VTs.
+
+	if (FSceneInterface* Scene = InWorld->Scene)
+	{
+		// BeingFrame/EndFrame (taken from FEngineLoop)
+
+		ENQUEUE_RENDER_COMMAND(BeginFrame)([](FRHICommandListImmediate& RHICmdList)
+			{
+				GFrameNumberRenderThread++;
+				RHICmdList.BeginFrame();
+				FCoreDelegates::OnBeginFrameRT.Broadcast();
+			});
+
+		ENQUEUE_RENDER_COMMAND(EndFrame)([](FRHICommandListImmediate& RHICmdList)
+			{
+				FCoreDelegates::OnEndFrameRT.Broadcast();
+				RHICmdList.EndFrame();
+			});
+
+		FlushRenderingCommands();
+	}
+}
+
 #endif // #if WITH_EDITOR
diff --git a/Engine/Source/Runtime/Engine/Public/WorldPartition/WorldPartitionHelpers.h b/Engine/Source/Runtime/Engine/Public/WorldPartition/WorldPartitionHelpers.h
--- a/Engine/Source/Runtime/Engine/Public/WorldPartition/WorldPartitionHelpers.h
+++ b/Engine/Source/Runtime/Engine/Public/WorldPartition/WorldPartitionHelpers.h
@@ -6,6 +6,7 @@
 
 #if WITH_EDITOR
 
+class UWorld;
 class UWorldPartition;
 class FWorldPartitionActorDesc;
 
@@ -37,6 +38,9 @@ public:
 
 	static bool HasExceededMaxMemory();
 	static void DoCollectGarbage();
+
+	// Simulate an engine frame tick
+	static void FakeEngineTick(UWorld* World);
 };
 
 #endif // WITH_EDITOR
