Fixed dead code in the world partition helpers

diff --git a/Engine/Source/Runtime/Engine/Public/WorldPartition/WorldPartitionHelpers.h b/Engine/Source/Runtime/Engine/Public/WorldPartition/WorldPartitionHelpers.h
--- a/Engine/Source/Runtime/Engine/Public/WorldPartition/WorldPartitionHelpers.h
+++ b/Engine/Source/Runtime/Engine/Public/WorldPartition/WorldPartitionHelpers.h
@@ -14,9 +14,9 @@ class ENGINE_API FWorldPartitionHelpers
 {
 public:
 	template <class ActorClass = AActor>
-	static void ForEachIntersectingActorDesc(UWorldPartition* WorldPartition, const FBox& Box, bool bIncludeFromChildActors, TFunctionRef<bool(const FWorldPartitionActorDesc*)> Func)
+	static void ForEachIntersectingActorDesc(UWorldPartition* WorldPartition, const FBox& Box, TFunctionRef<bool(const FWorldPartitionActorDesc*)> Func)
 	{
-		ForEachIntersectingActorDesc(WorldPartition, Box, bIncludeFromChildActors, ActorClass::StaticClass(), Func);
+		ForEachIntersectingActorDesc(WorldPartition, Box, ActorClass::StaticClass(), Func);
 	}
 
 	template<class ActorClass = AActor>
