#6086 : PR : Expose NetworkProfiler API for linking

diff --git a/Engine/Source/Runtime/Engine/Public/Net/NetworkProfiler.h b/Engine/Source/Runtime/Engine/Public/Net/NetworkProfiler.h
--- a/Engine/Source/Runtime/Engine/Public/Net/NetworkProfiler.h
+++ b/Engine/Source/Runtime/Engine/Public/Net/NetworkProfiler.h
@@ -80,7 +80,7 @@ public:
 /**
  * Network profiler, using serialized token emission like e.g. script and malloc profiler.
  */
-class FNetworkProfiler
+class ENGINE_API FNetworkProfiler
 {
 private:
 
@@ -207,17 +207,17 @@ public:
 	 *
 	 * @param	bShouldEnableTracking	Whether tracking should be enabled or not
 	 */
-	ENGINE_API void EnableTracking( bool bShouldEnableTracking );
+	void EnableTracking( bool bShouldEnableTracking );
 
 	/**
 	 * Marks the beginning of a frame.
 	 */
-	ENGINE_API void TrackFrameBegin();
+	void TrackFrameBegin();
 
 	/**
 	* Tracks when connection address changes
 	*/
-	ENGINE_API void SetCurrentConnection( UNetConnection* Connection );
+	void SetCurrentConnection( UNetConnection* Connection );
 
 	
 	/**
@@ -259,7 +259,7 @@ public:
 	 * @param	Data					Data sent
 	 * @param	BytesSent				Bytes actually being sent
 	 */
-	ENGINE_API void TrackSocketSend( const FString& SocketDesc, const void* Data, uint16 BytesSent );
+	void TrackSocketSend( const FString& SocketDesc, const void* Data, uint16 BytesSent );
 
 	/**
 	 * Low level FSocket::SendTo information.
@@ -273,7 +273,7 @@ public:
 	 * @param	NumPaddingBits			Number of bits appended to the end to byte-align the data
 	 * @param	Connection				Destination address
 	 */
-	ENGINE_API void TrackSocketSendTo(
+	void TrackSocketSendTo(
 		const FString& SocketDesc,
 		const void* Data,
 		uint16 BytesSent,
@@ -333,7 +333,7 @@ public:
 	 *
 	 * @param Connection the connection which is about to send any pending bunches over the network
 	 */
-	ENGINE_API void FlushOutgoingBunches( UNetConnection* Connection );
+	void FlushOutgoingBunches( UNetConnection* Connection );
 
 	/**
 	 * Track actor being replicated.
@@ -478,7 +478,7 @@ public:
 	 *
 	 * @return			True if processed, false otherwise
 	 */
-	ENGINE_API bool Exec( UWorld * InWorld, const TCHAR* Cmd, FOutputDevice & Ar );
+	bool Exec( UWorld * InWorld, const TCHAR* Cmd, FOutputDevice & Ar );
 
 	bool FORCEINLINE IsTrackingEnabled() const { return bIsTrackingEnabled || bShouldTrackingBeEnabled; }
 	bool IsComparisonTrackingEnabled() const { return (bIsTrackingEnabled || bShouldTrackingBeEnabled) && bIsComparisonTrackingEnabled; }
