[Backout] - CL17887288
Original CL Desc
-----------------------------------------------------------------
Generate cached data for MID created at runtime in non-editor builds
Make sure material instance resources are updated from cached data loaded in non-editor builds

diff --git a/Engine/Source/Editor/MaterialEditor/Private/MaterialStats.cpp b/Engine/Source/Editor/MaterialEditor/Private/MaterialStats.cpp
--- a/Engine/Source/Editor/MaterialEditor/Private/MaterialStats.cpp
+++ b/Engine/Source/Editor/MaterialEditor/Private/MaterialStats.cpp
@@ -198,7 +198,7 @@ bool FShaderPlatformSettings::CheckShaders()
 
 				if (MaterialInstance != nullptr)
 				{
-					MaterialInstance->UpdateCachedDataConstant();
+					MaterialInstance->UpdateCachedLayerParameters();
 				}
 
 				Data.LastTimeCompilationRequested = CurrentTime;
diff --git a/Engine/Source/Editor/UnrealEd/Private/PreviewMaterial.cpp b/Engine/Source/Editor/UnrealEd/Private/PreviewMaterial.cpp
--- a/Engine/Source/Editor/UnrealEd/Private/PreviewMaterial.cpp
+++ b/Engine/Source/Editor/UnrealEd/Private/PreviewMaterial.cpp
@@ -577,7 +577,7 @@ void UMaterialEditorInstanceConstant::RegenerateArrays()
 		// Only operate on base materials
 		UMaterial* ParentMaterial = Parent->GetMaterial();
 		SourceInstance->UpdateParameterNames();	// Update any parameter names that may have changed.
-		SourceInstance->UpdateCachedDataConstant();
+		SourceInstance->UpdateCachedLayerParameters();
 
 		// Need to get layer info first as other params are collected from layers
 		const FMaterialLayersFunctions* MaterialLayers = SourceInstance->GetMaterialLayers();
diff --git a/Engine/Source/Runtime/Engine/Classes/Materials/MaterialInstance.h b/Engine/Source/Runtime/Engine/Classes/Materials/MaterialInstance.h
--- a/Engine/Source/Runtime/Engine/Classes/Materials/MaterialInstance.h
+++ b/Engine/Source/Runtime/Engine/Classes/Materials/MaterialInstance.h
@@ -590,11 +590,10 @@ private:
 	bool bSavedCachedData_DEPRECATED;
 #endif
 
-protected:
 	bool bLoadedCachedData;
+
 	FMaterialInstanceCachedData* CachedData;
 
-private:
 #if WITH_EDITOR
 	mutable TOptional<FStaticParameterSet> CachedStaticParameterValues;
 	mutable uint8 AllowCachingStaticParameterValuesCounter = 0;
@@ -782,7 +781,7 @@ public:
 	/** Add to the set any texture referenced by expressions, including nested functions, as well as any overrides from parameters. */
 	ENGINE_API virtual void GetReferencedTexturesAndOverrides(TSet<const UTexture*>& InOutTextures) const;
 
-	ENGINE_API virtual void UpdateCachedDataConstant();
+	ENGINE_API void UpdateCachedLayerParameters();
 #endif
 
 	void GetBasePropertyOverridesHash(FSHAHash& OutHash)const;
@@ -856,7 +855,7 @@ protected:
 	 */
 	bool UpdateParameters();
 
-	ENGINE_API bool SetParentInternal(class UMaterialInterface* NewParent, bool RecacheShaders);
+	ENGINE_API void SetParentInternal(class UMaterialInterface* NewParent, bool RecacheShaders);
 
 	void GetTextureExpressionValues(const FMaterialResource* MaterialResource, TArray<UTexture*>& OutTextures, TArray< TArray<int32> >* OutIndices = nullptr) const;
 
diff --git a/Engine/Source/Runtime/Engine/Classes/Materials/MaterialInstanceConstant.h b/Engine/Source/Runtime/Engine/Classes/Materials/MaterialInstanceConstant.h
--- a/Engine/Source/Runtime/Engine/Classes/Materials/MaterialInstanceConstant.h
+++ b/Engine/Source/Runtime/Engine/Classes/Materials/MaterialInstanceConstant.h
@@ -35,7 +35,6 @@ class UMaterialInstanceConstant : public UMaterialInstance
 	friend class UMaterialEditorInstanceConstant;
 
 	virtual ENGINE_API void PostEditChangeProperty(FPropertyChangedEvent& PropertyChangedEvent) override;
-	virtual ENGINE_API void UpdateCachedDataConstant() override final;
 #endif
 
 	/** Physical material mask to use for this graphics material. Used for sounds, effects etc.*/
diff --git a/Engine/Source/Runtime/Engine/Classes/Materials/MaterialInstanceDynamic.h b/Engine/Source/Runtime/Engine/Classes/Materials/MaterialInstanceDynamic.h
--- a/Engine/Source/Runtime/Engine/Classes/Materials/MaterialInstanceDynamic.h
+++ b/Engine/Source/Runtime/Engine/Classes/Materials/MaterialInstanceDynamic.h
@@ -183,7 +183,5 @@ class ENGINE_API UMaterialInstanceDynamic : public UMaterialInstance
 	// This overrides does the remapping before looking at the parent data.
 	virtual float GetTextureDensity(FName TextureName, const struct FMeshUVChannelInfo& UVChannelData) const override;
 
-private:
-	void InitializeMID(class UMaterialInterface* ParentMaterial);
 };
 
diff --git a/Engine/Source/Runtime/Engine/Private/Materials/MaterialCachedData.cpp b/Engine/Source/Runtime/Engine/Private/Materials/MaterialCachedData.cpp
--- a/Engine/Source/Runtime/Engine/Private/Materials/MaterialCachedData.cpp
+++ b/Engine/Source/Runtime/Engine/Private/Materials/MaterialCachedData.cpp
@@ -738,15 +738,3 @@ void FMaterialInstanceCachedData::Initialize(FMaterialCachedExpressionData&& InC
 	NumParentLayers = ParentLayers ? ParentLayers->Layers.Num() : 0;
 }
 #endif // WITH_EDITOR
-
-void FMaterialInstanceCachedData::InitializeForDynamic(const FMaterialLayersFunctions* ParentLayers)
-{
-	const int32 NumLayers = ParentLayers ? ParentLayers->Layers.Num() : 0;
-	ParentLayerIndexRemap.Empty(NumLayers);
-	for (int32 LayerIndex = 0; LayerIndex < NumLayers; ++LayerIndex)
-	{
-		ParentLayerIndexRemap.Add(LayerIndex);
-	}
-
-	NumParentLayers = NumLayers;
-}
diff --git a/Engine/Source/Runtime/Engine/Private/Materials/MaterialInstance.cpp b/Engine/Source/Runtime/Engine/Private/Materials/MaterialInstance.cpp
--- a/Engine/Source/Runtime/Engine/Private/Materials/MaterialInstance.cpp
+++ b/Engine/Source/Runtime/Engine/Private/Materials/MaterialInstance.cpp
@@ -50,6 +50,21 @@
 #include "UObject/UE5MainStreamObjectVersion.h"
 #include "ShaderCompilerCore.h"
 
+#if ENABLE_COOK_STATS
+#include "ProfilingDebugging/ScopedTimers.h"
+namespace MaterialInstanceCookStats
+{
+	static double UpdateCachedExpressionDataSec = 0.0;
+
+	static FCookStatsManager::FAutoRegisterCallback RegisterCookStats([](FCookStatsManager::AddStatFuncRef AddStat)
+		{
+			AddStat(TEXT("MaterialInstance"), FCookStatsManager::CreateKeyValueArray(
+				TEXT("UpdateCachedExpressionDataSec"), UpdateCachedExpressionDataSec
+			));
+		});
+}
+#endif
+
 DECLARE_CYCLE_STAT(TEXT("MaterialInstance CopyMatInstParams"), STAT_MaterialInstance_CopyMatInstParams, STATGROUP_Shaders);
 DECLARE_CYCLE_STAT(TEXT("MaterialInstance Serialize"), STAT_MaterialInstance_Serialize, STATGROUP_Shaders);
 DECLARE_CYCLE_STAT(TEXT("MaterialInstance CopyUniformParamsInternal"), STAT_MaterialInstance_CopyUniformParamsInternal, STATGROUP_Shaders);
@@ -710,15 +725,10 @@ void UMaterialInstance::InitResources()
 
 	checkf(SafeParent, TEXT("Invalid parent on %s"), *GetFullName());
 
-	// TODO - should merge all of render commands sent to initialize resource into a single command
 	// Set the material instance's parent on its resources.
 	if (Resource != nullptr)
 	{
 		Resource->GameThread_SetParent(SafeParent);
-		if (CachedData)
-		{
-			Resource->GameThread_UpdateCachedData(*CachedData);
-		}
 	}
 
 	GameThread_InitMIParameters(*this);
@@ -1854,7 +1864,7 @@ void UMaterialInstance::GetDependentFunctions(TArray<UMaterialFunctionInterface*
 #if WITH_EDITOR
 void UMaterialInstance::ForceRecompileForRendering()
 {
-	UpdateCachedDataConstant();
+	UpdateCachedLayerParameters();
 	CacheResourceShadersForRendering();
 }
 #endif // WITH_EDITOR
@@ -2124,7 +2134,7 @@ void UMaterialInstance::CacheShadersForResources(EShaderPlatform ShaderPlatform,
 #endif
 
 #if WITH_EDITOR
-	UpdateCachedDataConstant();
+	UpdateCachedLayerParameters();
 #endif
 
 	for (int32 ResourceIndex = 0; ResourceIndex < ResourcesToCache.Num(); ResourceIndex++)
@@ -2590,7 +2600,7 @@ void UMaterialInstance::PostLoad()
 	}
 
 #if WITH_EDITOR
-	UpdateCachedDataConstant();
+	UpdateCachedLayerParameters();
 #endif
 
 	// called before we cache the uniform expression as a call to SubsurfaceProfileRT affects the dta in there
@@ -2735,13 +2745,13 @@ void UMaterialInstance::AddReferencedObjects(UObject* InThis, FReferenceCollecto
 	Super::AddReferencedObjects(This, Collector);
 }
 
-bool UMaterialInstance::SetParentInternal(UMaterialInterface* NewParent, bool RecacheShaders)
+void UMaterialInstance::SetParentInternal(UMaterialInterface* NewParent, bool RecacheShaders)
 {
-	bool bSetParent = false;
 	if (!Parent || Parent != NewParent)
 	{
 		// Check if the new parent is already an existing child
 		UMaterialInstance* ParentAsMaterialInstance = Cast<UMaterialInstance>(NewParent);
+		bool bSetParent = false;
 
 		if (ParentAsMaterialInstance != nullptr && ParentAsMaterialInstance->IsChildOf(this))
 		{
@@ -2788,8 +2798,10 @@ bool UMaterialInstance::SetParentInternal(UMaterialInterface* NewParent, bool Re
 		{
 			InitResources();
 		}
+#if WITH_EDITOR
+		UpdateCachedLayerParameters();
+#endif
 	}
-	return bSetParent;
 }
 
 bool UMaterialInstance::SetVectorParameterByIndexInternal(int32 ParameterIndex, FLinearColor Value)
@@ -3242,7 +3254,7 @@ void UMaterialInstance::UpdateStaticPermutation(const FStaticParameterSet& NewPa
 		bHasStaticPermutationResource = bWantsStaticPermutationResource;
 		StaticParameters = CompareParameters;
 
-		UpdateCachedDataConstant();
+		UpdateCachedLayerParameters();
 		CacheResourceShadersForRendering(EMaterialShaderPrecompileMode::None);
 		RecacheUniformExpressions(true);
 
@@ -3287,9 +3299,43 @@ void UMaterialInstance::GetReferencedTexturesAndOverrides(TSet<const UTexture*>&
 	}
 }
 
-void UMaterialInstance::UpdateCachedDataConstant()
+void UMaterialInstance::UpdateCachedLayerParameters()
 {
-	// By default do nothing, data is only cached for UMaterialInstanceConstant
+	COOK_STAT(FScopedDurationTimer BlockingTimer(MaterialInstanceCookStats::UpdateCachedExpressionDataSec));
+
+	// Don't need to rebuild cached data if it was serialized
+	if (!bLoadedCachedData)
+	{
+		UMaterialInstance* ParentInstance = nullptr;
+		FMaterialCachedExpressionData CachedExpressionData;
+		CachedExpressionData.Reset();
+		if (Parent)
+		{
+			CachedExpressionData.ReferencedTextures = Parent->GetReferencedTextures();
+			ParentInstance = Cast<UMaterialInstance>(Parent);
+		}
+
+		const FMaterialLayersFunctions* Layers = GetMaterialLayers();
+		const FMaterialLayersFunctions* ParentLayers = Parent ? Parent->GetMaterialLayers() : nullptr;
+
+		if (Layers)
+		{
+			FMaterialCachedExpressionContext Context;
+			CachedExpressionData.UpdateForLayerFunctions(Context, *Layers);
+		}
+
+		if (!CachedData)
+		{
+			CachedData = new FMaterialInstanceCachedData();
+		}
+		CachedData->Initialize(MoveTemp(CachedExpressionData), Layers, ParentLayers);
+		if (Resource)
+		{
+			Resource->GameThread_UpdateCachedData(*CachedData);
+		}
+	}
+
+	FObjectCacheEventSink::NotifyReferencedTextureChanged_Concurrent(this);
 }
 
 void UMaterialInstance::UpdateStaticPermutation(const FStaticParameterSet& NewParameters, FMaterialUpdateContext* MaterialUpdateContext)
@@ -3345,7 +3391,7 @@ void UMaterialInstance::PostEditChangeProperty(FPropertyChangedEvent& PropertyCh
 		RecacheMaterialInstanceUniformExpressions(this, false);
 	}
 
-	UpdateCachedDataConstant();
+	UpdateCachedLayerParameters();
 
 	if (GIsEditor)
 	{
diff --git a/Engine/Source/Runtime/Engine/Private/Materials/MaterialInstanceConstant.cpp b/Engine/Source/Runtime/Engine/Private/Materials/MaterialInstanceConstant.cpp
--- a/Engine/Source/Runtime/Engine/Private/Materials/MaterialInstanceConstant.cpp
+++ b/Engine/Source/Runtime/Engine/Private/Materials/MaterialInstanceConstant.cpp
@@ -5,26 +5,8 @@
 =============================================================================*/
 
 #include "Materials/MaterialInstanceConstant.h"
-#include "Materials/MaterialInstanceSupport.h"
-#include "ProfilingDebugging/CookStats.h"
 #if WITH_EDITOR
 #include "MaterialEditor/DEditorScalarParameterValue.h"
-#include "ObjectCacheEventSink.h"
-#endif
-
-#if ENABLE_COOK_STATS
-#include "ProfilingDebugging/ScopedTimers.h"
-namespace MaterialInstanceCookStats
-{
-static double UpdateCachedExpressionDataSec = 0.0;
-
-static FCookStatsManager::FAutoRegisterCallback RegisterCookStats([](FCookStatsManager::AddStatFuncRef AddStat)
-	{
-		AddStat(TEXT("MaterialInstance"), FCookStatsManager::CreateKeyValueArray(
-			TEXT("UpdateCachedExpressionDataSec"), UpdateCachedExpressionDataSec
-		));
-	});
-}
 #endif
 
 UMaterialInstanceConstant::UMaterialInstanceConstant(const FObjectInitializer& ObjectInitializer)
@@ -76,10 +58,7 @@ void UMaterialInstanceConstant::PostEditChangeProperty(FPropertyChangedEvent& Pr
 void UMaterialInstanceConstant::SetParentEditorOnly(UMaterialInterface* NewParent, bool RecacheShader)
 {
 	check(GIsEditor || IsRunningCommandlet());
-	if (SetParentInternal(NewParent, RecacheShader))
-	{
-		UpdateCachedDataConstant();
-	}
+	SetParentInternal(NewParent, RecacheShader);
 }
 
 void UMaterialInstanceConstant::CopyMaterialUniformParametersEditorOnly(UMaterialInterface* Source, bool bIncludeStaticParams)
@@ -146,44 +125,4 @@ void UMaterialInstanceConstant::ClearParameterValuesEditorOnly()
 	check(GIsEditor || IsRunningCommandlet());
 	ClearParameterValuesInternal();
 }
-
-void UMaterialInstanceConstant::UpdateCachedDataConstant()
-{
-	COOK_STAT(FScopedDurationTimer BlockingTimer(MaterialInstanceCookStats::UpdateCachedExpressionDataSec));
-
-	// Don't need to rebuild cached data if it was serialized
-	if (!bLoadedCachedData)
-	{
-		UMaterialInstance* ParentInstance = nullptr;
-		FMaterialCachedExpressionData CachedExpressionData;
-		CachedExpressionData.Reset();
-		if (Parent)
-		{
-			CachedExpressionData.ReferencedTextures = Parent->GetReferencedTextures();
-			ParentInstance = Cast<UMaterialInstance>(Parent);
-		}
-
-		const FMaterialLayersFunctions* Layers = GetMaterialLayers();
-		const FMaterialLayersFunctions* ParentLayers = Parent ? Parent->GetMaterialLayers() : nullptr;
-
-		if (Layers)
-		{
-			FMaterialCachedExpressionContext Context;
-			CachedExpressionData.UpdateForLayerFunctions(Context, *Layers);
-		}
-
-		if (!CachedData)
-		{
-			CachedData = new FMaterialInstanceCachedData();
-		}
-		CachedData->Initialize(MoveTemp(CachedExpressionData), Layers, ParentLayers);
-		if (Resource)
-		{
-			Resource->GameThread_UpdateCachedData(*CachedData);
-		}
-	}
-
-	FObjectCacheEventSink::NotifyReferencedTextureChanged_Concurrent(this);
-}
-
 #endif // #if WITH_EDITOR
diff --git a/Engine/Source/Runtime/Engine/Private/Materials/MaterialInstanceDynamic.cpp b/Engine/Source/Runtime/Engine/Private/Materials/MaterialInstanceDynamic.cpp
--- a/Engine/Source/Runtime/Engine/Private/Materials/MaterialInstanceDynamic.cpp
+++ b/Engine/Source/Runtime/Engine/Private/Materials/MaterialInstanceDynamic.cpp
@@ -20,28 +20,12 @@ UMaterialInstanceDynamic::UMaterialInstanceDynamic(const FObjectInitializer& Obj
 {
 }
 
-void UMaterialInstanceDynamic::InitializeMID(class UMaterialInterface* ParentMaterial)
-{
-	SetParentInternal(ParentMaterial, false);
-
-	const FMaterialLayersFunctions* ParentLayers = ParentMaterial ? ParentMaterial->GetMaterialLayers() : nullptr;
-
-	CachedData = new FMaterialInstanceCachedData();
-	CachedData->InitializeForDynamic(ParentLayers);
-	// TODO - should we copy ReferencedTextures from our parent?
-
-	if (Resource)
-	{
-		Resource->GameThread_UpdateCachedData(*CachedData);
-	}
-}
-
 UMaterialInstanceDynamic* UMaterialInstanceDynamic::Create(UMaterialInterface* ParentMaterial, UObject* InOuter)
 {
 	LLM_SCOPE(ELLMTag::MaterialInstance);
 	UObject* Outer = InOuter ? InOuter : GetTransientPackage();
 	UMaterialInstanceDynamic* MID = NewObject<UMaterialInstanceDynamic>(Outer);
-	MID->InitializeMID(ParentMaterial);
+	MID->SetParentInternal(ParentMaterial, false);
 	return MID;
 }
 
@@ -88,7 +72,7 @@ UMaterialInstanceDynamic* UMaterialInstanceDynamic::Create(UMaterialInterface* P
 		}
 	}
 	UMaterialInstanceDynamic* MID = NewObject<UMaterialInstanceDynamic>(Outer, Name);
-	MID->InitializeMID(ParentMaterial);
+	MID->SetParentInternal(ParentMaterial, false);
 	return MID;
 }
 
diff --git a/Engine/Source/Runtime/Engine/Private/Materials/MaterialShared.cpp b/Engine/Source/Runtime/Engine/Private/Materials/MaterialShared.cpp
--- a/Engine/Source/Runtime/Engine/Private/Materials/MaterialShared.cpp
+++ b/Engine/Source/Runtime/Engine/Private/Materials/MaterialShared.cpp
@@ -4017,7 +4017,7 @@ FMaterialUpdateContext::~FMaterialUpdateContext()
 		}
 
 #if WITH_EDITOR
-		MI->UpdateCachedDataConstant();
+		MI->UpdateCachedLayerParameters();
 #endif
 		MI->RecacheUniformExpressions(true);
 		MI->InitStaticPermutation(EMaterialShaderPrecompileMode::None);//bHasStaticPermutation can change.
diff --git a/Engine/Source/Runtime/Engine/Public/MaterialCachedData.h b/Engine/Source/Runtime/Engine/Public/MaterialCachedData.h
--- a/Engine/Source/Runtime/Engine/Public/MaterialCachedData.h
+++ b/Engine/Source/Runtime/Engine/Public/MaterialCachedData.h
@@ -292,7 +292,6 @@ struct FMaterialInstanceCachedData
 #if WITH_EDITOR
 	void Initialize(FMaterialCachedExpressionData&& InCachedExpressionData, const FMaterialLayersFunctions* Layers, const FMaterialLayersFunctions* ParentLayers);
 #endif // WITH_EDITOR
-	void InitializeForDynamic(const FMaterialLayersFunctions* ParentLayers);
 
 	void AddReferencedObjects(FReferenceCollector& Collector);
 
