[WorldSubsystem] added new virtual callback 'OnWorldComponentsUpdated' and 'GetWorldRef()'

Note that this step comes after the following existing subsystem callbacks:

Initialize:
* called from initialization of the subsystem collection
* level collection => NOT set
* actor iterator   => NOT available
* actor components => NOT registered

PostInitialize:
* called once all UWorldSubsystems have been initialized
* level collection => set
* actor iterator   => available
* actor components => NOT registered

OnWorldComponentsUpdated:
* called once all world components have been updated (e.g. level components) and actors registered
* level collection => set
* actor iterator   => available
* actor components => registered

diff --git a/Engine/Source/Runtime/Engine/Private/Subsystems/WorldSubsystem.cpp b/Engine/Source/Runtime/Engine/Private/Subsystems/WorldSubsystem.cpp
--- a/Engine/Source/Runtime/Engine/Private/Subsystems/WorldSubsystem.cpp
+++ b/Engine/Source/Runtime/Engine/Private/Subsystems/WorldSubsystem.cpp
@@ -11,6 +11,11 @@ UWorldSubsystem::UWorldSubsystem()
 
 }
 
+UWorld& UWorldSubsystem::GetWorldRef() const
+{
+	return *CastChecked<UWorld>(GetOuter(), ECastCheckedType::NullChecked);
+}
+
 UWorld* UWorldSubsystem::GetWorld() const
 {
 	return Cast<UWorld>(GetOuter());
@@ -28,7 +33,7 @@ bool UWorldSubsystem::ShouldCreateSubsystem(UObject* Outer) const
 	return DoesSupportWorldType(World->WorldType);
 }
 
-bool UWorldSubsystem::DoesSupportWorldType(EWorldType::Type WorldType) const
+bool UWorldSubsystem::DoesSupportWorldType(const EWorldType::Type WorldType) const
 {
 	return WorldType == EWorldType::Game || WorldType == EWorldType::Editor || WorldType == EWorldType::PIE;
 }
diff --git a/Engine/Source/Runtime/Engine/Private/World.cpp b/Engine/Source/Runtime/Engine/Private/World.cpp
--- a/Engine/Source/Runtime/Engine/Private/World.cpp
+++ b/Engine/Source/Runtime/Engine/Private/World.cpp
@@ -2164,6 +2164,12 @@ void UWorld::UpdateWorldComponents(bool bRerunConstructionScripts, bool bCurrent
 		}
 	}
 
+	const TArray<UWorldSubsystem*>& WorldSubsystems = SubsystemCollection.GetSubsystemArray<UWorldSubsystem>(UWorldSubsystem::StaticClass());
+	for (UWorldSubsystem* WorldSubsystem : WorldSubsystems)
+	{
+		WorldSubsystem->OnWorldComponentsUpdated(*this);
+	}
+
 	UpdateCullDistanceVolumes();
 }
 
diff --git a/Engine/Source/Runtime/Engine/Public/Subsystems/WorldSubsystem.h b/Engine/Source/Runtime/Engine/Public/Subsystems/WorldSubsystem.h
--- a/Engine/Source/Runtime/Engine/Public/Subsystems/WorldSubsystem.h
+++ b/Engine/Source/Runtime/Engine/Public/Subsystems/WorldSubsystem.h
@@ -20,19 +20,29 @@ public:
 	UWorldSubsystem();
 
 	virtual UWorld* GetWorld() const override final;
+
+	/**
+	 * Returns a reference to the UWorld this subsystem is contained within.
+	 * @note This should not be called on default object since the method assumes a valid outer world.
+	 */
+	UWorld& GetWorldRef() const;
+
 	virtual bool ShouldCreateSubsystem(UObject* Outer) const override;
 
 	/** Called once all UWorldSubsystems have been initialized */
 	virtual void PostInitialize() {}
 	
 	/** Called when world is ready to start gameplay before the game mode transitions to the correct state and call BeginPlay on all actors */
-	virtual void OnWorldBeginPlay(UWorld& InWorld) {};
+	virtual void OnWorldBeginPlay(UWorld& InWorld) {}
+
+	/** Called after world components (e.g. line batcher and all level components) have been updated */
+	virtual void OnWorldComponentsUpdated(UWorld& World) {}
 
 	/** Updates sub-system required streaming levels (called by world's UpdateStreamingState function) */
 	virtual void UpdateStreamingState() {}
 
 protected:
-	virtual bool DoesSupportWorldType(EWorldType::Type WorldType) const;
+	virtual bool DoesSupportWorldType(const EWorldType::Type WorldType) const;
 };
 
 /**
