Handle emptying the map name if CurrentMapName == DestinationMapName when no transition map is specified.

diff --git a/Engine/Source/Runtime/Engine/Private/World.cpp b/Engine/Source/Runtime/Engine/Private/World.cpp
--- a/Engine/Source/Runtime/Engine/Private/World.cpp
+++ b/Engine/Source/Runtime/Engine/Private/World.cpp
@@ -6439,18 +6439,6 @@ bool FSeamlessTravelHandler::StartTravel(UWorld* InCurrentWorld, const FURL& InU
 					StartLoadingDestination();
 				}
 			}
-			else if (TransitionMap.IsEmpty())
-			{
-				// If a default transition map doesn't exist, create a dummy World to use as the transition
-				if (CurrentWorld->WorldType == EWorldType::PIE)
-				{
-					SetHandlerLoadedData(NULL, UWorld::CreateWorld(EWorldType::PIE, false));
-				}
-				else
-				{
-					SetHandlerLoadedData(NULL, UWorld::CreateWorld(EWorldType::None, false));
-				}
-			}
 			else
 			{
 				if (CurrentMapName == DestinationMapName)
@@ -6471,18 +6459,33 @@ bool FSeamlessTravelHandler::StartTravel(UWorld* InCurrentWorld, const FURL& InU
 					}
 				}
 
-				// Set the world type in the static map, so that UWorld::PostLoad can set the world type
-				UWorld::WorldTypePreLoadMap.FindOrAdd(*TransitionMap) = CurrentWorld->WorldType;
-
-				// first, load the entry level package
-				STAT_ADD_CUSTOMMESSAGE_NAME( STAT_NamedMarker, *(FString( TEXT( "StartTravel - " ) + TransitionMap )) );
-				TRACE_BOOKMARK(TEXT("StartTravel - %s"), *TransitionMap);
-				LoadPackageAsync(TransitionMap, 
-					FLoadPackageAsyncDelegate::CreateRaw(this, &FSeamlessTravelHandler::SeamlessTravelLoadCallback),
-					0, 
-					(CurrentWorld->WorldType == EWorldType::PIE ? PKG_PlayInEditor : PKG_None),
-					Context.PIEInstance
-					);
+				if (TransitionMap.IsEmpty())
+				{ 
+					// If a default transition map doesn't exist, create a dummy World to use as the transition
+					if (CurrentWorld->WorldType == EWorldType::PIE)
+					{
+						SetHandlerLoadedData(NULL, UWorld::CreateWorld(EWorldType::PIE, false));
+					}
+					else
+					{
+						SetHandlerLoadedData(NULL, UWorld::CreateWorld(EWorldType::None, false));
+					}
+				}
+				else
+				{
+					// Set the world type in the static map, so that UWorld::PostLoad can set the world type
+					UWorld::WorldTypePreLoadMap.FindOrAdd(*TransitionMap) = CurrentWorld->WorldType;
+
+					// first, load the entry level package
+					STAT_ADD_CUSTOMMESSAGE_NAME( STAT_NamedMarker, *(FString( TEXT( "StartTravel - " ) + TransitionMap )) );
+					TRACE_BOOKMARK(TEXT("StartTravel - %s"), *TransitionMap);
+					LoadPackageAsync(TransitionMap, 
+						FLoadPackageAsyncDelegate::CreateRaw(this, &FSeamlessTravelHandler::SeamlessTravelLoadCallback),
+						0, 
+						(CurrentWorld->WorldType == EWorldType::PIE ? PKG_PlayInEditor : PKG_None),
+						Context.PIEInstance
+						);
+				}
 			}
 
 			return true;
