Remove SCOPED_SUSPEND_RENDERING_THREAD from UpdateViewportRHI

diff --git a/Engine/Source/Runtime/Engine/Private/Slate/DebugCanvas.cpp b/Engine/Source/Runtime/Engine/Private/Slate/DebugCanvas.cpp
--- a/Engine/Source/Runtime/Engine/Private/Slate/DebugCanvas.cpp
+++ b/Engine/Source/Runtime/Engine/Private/Slate/DebugCanvas.cpp
@@ -110,8 +110,6 @@ void FDebugCanvasDrawer::ReleaseResources()
 		{
 			ReleaseMe->ReleaseTexture();
 		});
-
-	FlushRenderingCommands();
 }
 
 FDebugCanvasDrawer::~FDebugCanvasDrawer()
diff --git a/Engine/Source/Runtime/Engine/Private/Slate/SceneViewport.cpp b/Engine/Source/Runtime/Engine/Private/Slate/SceneViewport.cpp
--- a/Engine/Source/Runtime/Engine/Private/Slate/SceneViewport.cpp
+++ b/Engine/Source/Runtime/Engine/Private/Slate/SceneViewport.cpp
@@ -1618,7 +1618,8 @@ void FSceneViewport::SetRenderTargetTextureRenderThread(FTexture2DRHIRef& RT)
 void FSceneViewport::UpdateViewportRHI(bool bDestroyed, uint32 NewSizeX, uint32 NewSizeY, EWindowMode::Type NewWindowMode, EPixelFormat PreferredPixelFormat)
 {
 	{
-		SCOPED_SUSPEND_RENDERING_THREAD(true);
+		check(IsInGameThread());
+		FlushRenderingCommands();
 
 		// Update the viewport attributes.
 		// This is done AFTER the command flush done by UpdateViewportRHI, to avoid disrupting rendering thread accesses to the old viewport size.
@@ -1632,12 +1633,23 @@ void FSceneViewport::UpdateViewportRHI(bool bDestroyed, uint32 NewSizeX, uint32
 		if( !bDestroyed )
 		{
 			BeginInitResource(this);
-				
-			if( !UseSeparateRenderTarget() )
+			
+			FSlateRenderer* Renderer = FSlateApplication::Get().GetRenderer();
+
+			TSharedPtr<SWidget> PinnedViewport = ViewportWidget.Pin();
+			if (PinnedViewport.IsValid())
 			{
-				// Get the viewport for this window from the renderer so we can render directly to the backbuffer
-				FSlateRenderer* Renderer = FSlateApplication::Get().GetRenderer();
+				TSharedPtr<SWindow> Window = FSlateApplication::Get().FindWidgetWindow(PinnedViewport.ToSharedRef());
+
+				WindowRenderTargetUpdate(Renderer, Window.Get());
+				if (UseSeparateRenderTarget())
+				{
+					RTTSize = FIntPoint(SizeX, SizeY);
+				}
+			}
 
+			if( !UseSeparateRenderTarget() )
+			{
 				TSharedPtr<SWindow> Window = FSlateApplication::Get().FindWidgetWindow(ViewportWidget.Pin().ToSharedRef());
 				void* ViewportResource = Renderer->GetViewportResource(*Window);
 				if( ViewportResource )
@@ -1652,21 +1664,19 @@ void FSceneViewport::UpdateViewportRHI(bool bDestroyed, uint32 NewSizeX, uint32
 		else
 		{
 			// Enqueue a render command to delete the handle.  It must be deleted on the render thread after the resource is released
-			FSlateRenderTargetRHI** RenderThreadSlateTexturePtr = &RenderThreadSlateTexture;
-			TArray<FSlateRenderTargetRHI*>* BufferedSlateHandlesPtr = &BufferedSlateHandles;
+			FSlateRenderTargetRHI* RenderThreadSlateTexturePtr = RenderThreadSlateTexture;
+			RenderThreadSlateTexture = nullptr;
+
 			ENQUEUE_RENDER_COMMAND(DeleteSlateRenderTarget)(
-				[BufferedSlateHandlesPtr, RenderThreadSlateTexturePtr](FRHICommandListImmediate& RHICmdList)
+				[BufferedSlateHandles = MoveTemp(BufferedSlateHandles), RenderThreadSlateTexturePtr](FRHICommandListImmediate& RHICmdList)
 				{
-					for (int32 i = 0; i < BufferedSlateHandlesPtr->Num(); ++i)
+					for (int32 i = 0; i < BufferedSlateHandles.Num(); ++i)
 					{
-						delete (*BufferedSlateHandlesPtr)[i];
-						(*BufferedSlateHandlesPtr)[i] = nullptr;
+						delete BufferedSlateHandles[i];
 					}
 
-					delete *RenderThreadSlateTexturePtr;
-					*RenderThreadSlateTexturePtr = nullptr;
+					delete RenderThreadSlateTexturePtr;
 				});
-
 		}
 	}
 }
@@ -1933,7 +1943,6 @@ void FSceneViewport::InitDynamicRHI()
 	}
 	RTTSize = FIntPoint(0, 0);
 
-	FSlateRenderer* Renderer = FSlateApplication::Get().GetRenderer();
 	uint32 TexSizeX = SizeX, TexSizeY = SizeY;
 	if (UseSeparateRenderTarget())
 	{
@@ -2044,20 +2053,6 @@ void FSceneViewport::InitDynamicRHI()
 		RenderTargetTextureRHI = nullptr;		
 		CurrentBufferedTargetIndex = NextBufferedTargetIndex = 0;
 	}
-
-	//how is this useful at all?  Pinning a weakptr to get a non-threadsafe shared ptr?  Pinning a weakptr is supposed to be protecting me from my weakptr dying underneath me...
-	TSharedPtr<SWidget> PinnedViewport = ViewportWidget.Pin();
-	if (PinnedViewport.IsValid())
-	{
-
-		TSharedPtr<SWindow> Window = FSlateApplication::Get().FindWidgetWindow(PinnedViewport.ToSharedRef());
-		
-		WindowRenderTargetUpdate(Renderer, Window.Get());
-		if (UseSeparateRenderTarget())
-		{
-			RTTSize = FIntPoint(TexSizeX, TexSizeY);
-		}
-	}
 }
 
 void FSceneViewport::ReleaseDynamicRHI()
