U5 - More removal of SCOPED_SUSPEND_RENDERING_THREAD

diff --git a/Engine/Source/Runtime/Engine/Private/UnrealClient.cpp b/Engine/Source/Runtime/Engine/Private/UnrealClient.cpp
--- a/Engine/Source/Runtime/Engine/Private/UnrealClient.cpp
+++ b/Engine/Source/Runtime/Engine/Private/UnrealClient.cpp
@@ -1511,15 +1511,10 @@ void FViewport::EnqueueEndRenderFrame(const bool bLockToVsync, const bool bShoul
 		});
 }
 
-// true: The CompositionInspectur Slate UI requests it's data
-bool GCaptureCompositionNextFrame = false;
-
-
 void FViewport::Draw( bool bShouldPresent /*= true */)
 {
 	SCOPED_NAMED_EVENT(FViewport_Draw, FColor::Red);
 	UWorld* World = GetClient()->GetWorld();
-	static TUniquePtr<FSuspendRenderingThread> GRenderingThreadSuspension;
 
 	// Ignore reentrant draw calls, since we can only redraw one viewport at a time.
 	static bool bReentrant = false;
@@ -1531,14 +1526,6 @@ void FViewport::Draw( bool bShouldPresent /*= true */)
 		bool bAnyScreenshotsRequired = FScreenshotRequest::IsScreenshotRequested() || GIsHighResScreenshot || GIsDumpingMovie;
 		bool bBufferVisualizationDumpingRequired = bAnyScreenshotsRequired && CVarDumpFrames && CVarDumpFrames->GetValueOnGameThread();
 
-
-		if(GCaptureCompositionNextFrame)
-		{
-			// To capture the CompositionGraph we go into single threaded for one frame
-			// so that the Slate UI gets the data on the game thread.
-			GRenderingThreadSuspension = MakeUnique<FSuspendRenderingThread>(true);
-		}
-
 		// if this is a game viewport, and game rendering is disabled, then we don't want to actually draw anything
 		if ( World && World->IsGameWorld() && !bIsGameRenderingEnabled)
 		{
@@ -1647,12 +1634,6 @@ void FViewport::Draw( bool bShouldPresent /*= true */)
 				bIsGameRenderingEnabled = true;
 			}
 		}
-
-		if(GCaptureCompositionNextFrame)
-		{
-			GRenderingThreadSuspension.Reset();
-			GCaptureCompositionNextFrame = false;
-		}
 	}
 }
 
@@ -1945,7 +1926,7 @@ void FViewport::UpdateViewportRHI(bool bDestroyed, uint32 NewSizeX, uint32 NewSi
 {
 	{
 		// Temporarily stop rendering thread.
-		SCOPED_SUSPEND_RENDERING_THREAD(true);
+		FlushRenderingCommands();
 
 		// Update the viewport attributes.
 		// This is done AFTER the command flush done by UpdateViewportRHI, to avoid disrupting rendering thread accesses to the old viewport size.
@@ -2069,14 +2050,13 @@ void FViewport::ReleaseDynamicRHI()
 
 void FViewport::ReleaseRHI()
 {
-	SCOPED_SUSPEND_RENDERING_THREAD(true);
+	FlushRenderingCommands();
 	ViewportRHI.SafeRelease();
 }
 
 void FViewport::InitRHI()
 {
-	SCOPED_SUSPEND_RENDERING_THREAD(true);
-
+	FlushRenderingCommands();
 	if(!IsValidRef(ViewportRHI))
 	{
 		ViewportRHI = RHICreateViewport(
