Added support for full precision in Material expressions to project
Added override for half precision to Materials

diff --git a/Engine/Source/Developer/Apple/MetalShaderFormat/Private/MetalShaderCompiler.cpp b/Engine/Source/Developer/Apple/MetalShaderFormat/Private/MetalShaderCompiler.cpp
--- a/Engine/Source/Developer/Apple/MetalShaderFormat/Private/MetalShaderCompiler.cpp
+++ b/Engine/Source/Developer/Apple/MetalShaderFormat/Private/MetalShaderCompiler.cpp
@@ -1073,8 +1073,7 @@ void CompileShader_Metal(const FShaderCompilerInput& _Input,FShaderCompilerOutpu
 	}
 	
 	// Force floats if the material requests it
-	const bool bUseFullPrecisionInPS = Input.Environment.CompilerFlags.Contains(CFLAG_UseFullPrecisionInPS);
-	if (bUseFullPrecisionInPS || (VersionEnum < 2)) // Too many bugs in Metal 1.0 & 1.1 with half floats the more time goes on and the compiler stack changes
+	if (Input.Environment.FullPrecisionInPS || (VersionEnum < 2)) // Too many bugs in Metal 1.0 & 1.1 with half floats the more time goes on and the compiler stack changes
 	{
 		AdditionalDefines.SetDefine(TEXT("FORCE_FLOATS"), (uint32)1);
 	}
diff --git a/Engine/Source/Developer/ShaderFormatOpenGL/Private/OpenGLShaderCompiler.cpp b/Engine/Source/Developer/ShaderFormatOpenGL/Private/OpenGLShaderCompiler.cpp
--- a/Engine/Source/Developer/ShaderFormatOpenGL/Private/OpenGLShaderCompiler.cpp
+++ b/Engine/Source/Developer/ShaderFormatOpenGL/Private/OpenGLShaderCompiler.cpp
@@ -1164,11 +1164,11 @@ uint32 FOpenGLFrontend::GetMaxSamplers(GLSLVersion Version)
 	return 16;
 }
 
-uint32 FOpenGLFrontend::CalculateCrossCompilerFlags(GLSLVersion Version, const TArray<uint32>& CompilerFlags)
+uint32 FOpenGLFrontend::CalculateCrossCompilerFlags(GLSLVersion Version, const bool bFullPrecisionInPS, const TArray<uint32>& CompilerFlags)
 {
 	uint32  CCFlags = HLSLCC_NoPreprocess | HLSLCC_PackUniforms | HLSLCC_DX11ClipSpace | HLSLCC_RetainSizes;
 
-	if (CompilerFlags.Contains(CFLAG_UseFullPrecisionInPS))
+	if (bFullPrecisionInPS)
 	{
 		CCFlags |= HLSLCC_UseFullPrecisionInPS;
 	}
@@ -2753,7 +2753,7 @@ void FOpenGLFrontend::CompileShader(const FShaderCompilerInput& Input, FShaderCo
 		AdditionalDefines.SetDefine(TEXT("COMPILER_SUPPORTS_ATTRIBUTES"), (uint32)0);
 	}
 
-	if (Input.Environment.CompilerFlags.Contains(CFLAG_UseFullPrecisionInPS))
+	if (Input.Environment.FullPrecisionInPS)
 	{
 		AdditionalDefines.SetDefine(TEXT("FORCE_FLOATS"), (uint32)1);
 	}
@@ -2817,7 +2817,7 @@ void FOpenGLFrontend::CompileShader(const FShaderCompilerInput& Input, FShaderCo
 	// This requires removing the HLSLCC_NoPreprocess flag later on!
 	RemoveUniformBuffersFromSource(Input.Environment, PreprocessedShader);
 
-	uint32 CCFlags = CalculateCrossCompilerFlags(Version, Input.Environment.CompilerFlags);
+	uint32 CCFlags = CalculateCrossCompilerFlags(Version, Input.Environment.FullPrecisionInPS, Input.Environment.CompilerFlags);
 
 	// Required as we added the RemoveUniformBuffersFromSource() function (the cross-compiler won't be able to interpret comments w/o a preprocessor)
 	CCFlags &= ~HLSLCC_NoPreprocess;
diff --git a/Engine/Source/Developer/ShaderFormatOpenGL/Public/ShaderFormatOpenGL.h b/Engine/Source/Developer/ShaderFormatOpenGL/Public/ShaderFormatOpenGL.h
--- a/Engine/Source/Developer/ShaderFormatOpenGL/Public/ShaderFormatOpenGL.h
+++ b/Engine/Source/Developer/ShaderFormatOpenGL/Public/ShaderFormatOpenGL.h
@@ -59,7 +59,7 @@ protected:
 	// what is the max number of samplers the shader platform can use?
 	virtual uint32 GetMaxSamplers(GLSLVersion Version);
 
-	virtual uint32 CalculateCrossCompilerFlags(GLSLVersion Version, const TArray<uint32>& CompilerFlags);
+	virtual uint32 CalculateCrossCompilerFlags(GLSLVersion Version, const bool bFullPrecisionInPS, const TArray<uint32>& CompilerFlags);
 
 	// set up compilation information like defines and HlslCompileTarget
 	virtual void SetupPerVersionCompilationEnvironment(GLSLVersion Version, class FShaderCompilerDefinitions& AdditionalDefines, EHlslCompileTarget& HlslCompilerTarget);
diff --git a/Engine/Source/Developer/VulkanShaderFormat/Private/VulkanShaderCompiler.cpp b/Engine/Source/Developer/VulkanShaderFormat/Private/VulkanShaderCompiler.cpp
--- a/Engine/Source/Developer/VulkanShaderFormat/Private/VulkanShaderCompiler.cpp
+++ b/Engine/Source/Developer/VulkanShaderFormat/Private/VulkanShaderCompiler.cpp
@@ -2078,8 +2078,7 @@ void DoCompileVulkanShader(const FShaderCompilerInput& Input, FShaderCompilerOut
 	AdditionalDefines.SetDefine(TEXT("COMPILER_SUPPORTS_ATTRIBUTES"), (uint32)1);
 	AdditionalDefines.SetDefine(TEXT("COMPILER_SUPPORTS_DUAL_SOURCE_BLENDING_SLOT_DECORATION"), (uint32)1);
 
-	const bool bUseFullPrecisionInPS = Input.Environment.CompilerFlags.Contains(CFLAG_UseFullPrecisionInPS);
-	if (bUseFullPrecisionInPS)
+	if (Input.Environment.FullPrecisionInPS)
 	{
 		AdditionalDefines.SetDefine(TEXT("FORCE_FLOATS"), (uint32)1);
 	}
@@ -2128,7 +2127,7 @@ void DoCompileVulkanShader(const FShaderCompilerInput& Input, FShaderCompilerOut
 		// Only flatten structures inside UBs
 		CompilerInfo.CCFlags |= HLSLCC_FlattenUniformBufferStructures;
 
-		if (bUseFullPrecisionInPS)
+		if (Input.Environment.FullPrecisionInPS)
 		{
 			CompilerInfo.CCFlags |= HLSLCC_UseFullPrecisionInPS;
 		}
diff --git a/Engine/Source/Runtime/Core/Private/HAL/ConsoleManager.cpp b/Engine/Source/Runtime/Core/Private/HAL/ConsoleManager.cpp
--- a/Engine/Source/Runtime/Core/Private/HAL/ConsoleManager.cpp
+++ b/Engine/Source/Runtime/Core/Private/HAL/ConsoleManager.cpp
@@ -2163,11 +2163,12 @@ static TAutoConsoleVariable<int32> CVarMobileSkyLightPermutation(
 	TEXT("2: Generate only skylight permutations"),
 	ECVF_RenderThreadSafe | ECVF_ReadOnly);
 
-static TAutoConsoleVariable<int32> CVarMobileForceFullPrecisionInPS(
-	TEXT("r.Mobile.ForceFullPrecisionInPS"),
+static TAutoConsoleVariable<int32> CVarMobileFloatPrecisionMode(
+	TEXT("r.Mobile.FloatPrecisionMode"),
 	0,
-	TEXT("0: Use precision specified in shader code (default)\n")
-	TEXT("1: Force use of high precision in pixel shaders.\n"),
+	TEXT("0: Use Half-precision (default)\n")
+	TEXT("1: Half precision, except Full precision for material expressions\n")
+	TEXT("2: Force use of high precision in pixel shaders.\n"),
 	ECVF_ReadOnly);
 
 static TAutoConsoleVariable<int32> CVarMobileAllowDitheredLODTransition(
diff --git a/Engine/Source/Runtime/Engine/Classes/Engine/EngineTypes.h b/Engine/Source/Runtime/Engine/Classes/Engine/EngineTypes.h
--- a/Engine/Source/Runtime/Engine/Classes/Engine/EngineTypes.h
+++ b/Engine/Source/Runtime/Engine/Classes/Engine/EngineTypes.h
@@ -254,12 +254,14 @@ enum EBlendMode
 UENUM()
 enum EMaterialFloatPrecisionMode
 {
-	/** Half precision, except explict 'float' in .ush/.usf*/
-	MFPM_Half UMETA(DisplayName = "Half-precision"),
+	/** Uses project based precision mode setting */
+	MFPM_Default UMETA(DisplayName = "Default"),
 	/** Force full-precision for MaterialFloat only, no effect on shader codes in .ush/.usf*/
-	MFPM_Full_MaterialExpressionOnly UMETA(DisplayName = "Full-precision for MaterialExpressions only"),
+	MFPM_Full_MaterialExpressionOnly UMETA(DisplayName = "Use Full-precision for MaterialExpressions only"),
 	/** All the floats are full-precision */
-	MFPM_Full UMETA(DisplayName = "Full-precision for every float"),
+	MFPM_Full UMETA(DisplayName = "Use Full-precision for every float"),
+	/** Half precision, except explict 'float' in .ush/.usf*/
+	MFPM_Half UMETA(DisplayName = "Use Half-precision"),
 	MFPM_MAX,
 };
 
diff --git a/Engine/Source/Runtime/Engine/Classes/Engine/RendererSettings.h b/Engine/Source/Runtime/Engine/Classes/Engine/RendererSettings.h
--- a/Engine/Source/Runtime/Engine/Classes/Engine/RendererSettings.h
+++ b/Engine/Source/Runtime/Engine/Classes/Engine/RendererSettings.h
@@ -169,6 +169,22 @@ namespace EMobileAntiAliasingMethod
 	};
 }
 
+/** The default float precision for material's pixel shaders on mobile devices*/
+UENUM()
+namespace EMobileFloatPrecisionMode
+{
+	enum Type
+	{
+		/** Half precision, except explict 'float' in .ush/.usf*/
+		Half = 0 UMETA(DisplayName = "Use Half-precision"),
+		/** Half precision, except Full precision for material floats and explicit floats in .ush/.usf*/
+		Full_MaterialExpressionOnly = 1 UMETA(DisplayName = "Use Full-precision for MaterialExpressions only"),
+		/** All the floats are full-precision */
+		Full = 2 UMETA(DisplayName = "Use Full-precision for every float"),
+	};
+}
+
+
 namespace EDefaultBackBufferPixelFormat
 {
 	ENGINE_API EPixelFormat Convert2PixelFormat(EDefaultBackBufferPixelFormat::Type InDefaultBackBufferPixelFormat);
@@ -239,6 +255,12 @@ class ENGINE_API URendererSettings : public UDeveloperSettings
 		ToolTip = "The mobile default anti-aliasing method."))
 	TEnumAsByte<EMobileAntiAliasingMethod::Type> MobileAntiAliasing;
 
+	UPROPERTY(config, EditAnywhere, Category = Mobile, meta = (
+		ConsoleVariable = "r.Mobile.FloatPrecisionMode", DisplayName = "Mobile Float Precision Mode",
+		ToolTip = "Project wide mobile float precision mode for shaders and materials. Changing this setting requires restarting the editor.",
+		ConfigRestartRequired = true))
+	TEnumAsByte<EMobileFloatPrecisionMode::Type> MobileFloatPrecisionMode;
+
 	UPROPERTY(config, EditAnywhere, Category = Mobile, meta=(
 		ConsoleVariable="r.Mobile.AllowDitheredLODTransition", DisplayName="Allow Dithered LOD Transition",
 		ToolTip="Whether to support 'Dithered LOD Transition' material option on mobile platforms. Enabling this may degrade performance as rendering will not benefit from Early-Z optimization.",
diff --git a/Engine/Source/Runtime/Engine/Private/Materials/Material.cpp b/Engine/Source/Runtime/Engine/Private/Materials/Material.cpp
--- a/Engine/Source/Runtime/Engine/Private/Materials/Material.cpp
+++ b/Engine/Source/Runtime/Engine/Private/Materials/Material.cpp
@@ -967,7 +967,7 @@ UMaterial::UMaterial(const FObjectInitializer& ObjectInitializer)
 
 	bLoadedCachedExpressionData = false;
 
-	FloatPrecisionMode = EMaterialFloatPrecisionMode::MFPM_Half;
+	FloatPrecisionMode = EMaterialFloatPrecisionMode::MFPM_Default;
 }
 
 void UMaterial::PreSave(const class ITargetPlatform* TargetPlatform)
diff --git a/Engine/Source/Runtime/Engine/Private/Materials/MaterialShared.cpp b/Engine/Source/Runtime/Engine/Private/Materials/MaterialShared.cpp
--- a/Engine/Source/Runtime/Engine/Private/Materials/MaterialShared.cpp
+++ b/Engine/Source/Runtime/Engine/Private/Materials/MaterialShared.cpp
@@ -2059,16 +2059,18 @@ void FMaterial::SetupMaterialEnvironment(
 	OutEnvironment.SetDefine(TEXT("MATERIAL_USE_ALPHA_TO_COVERAGE"), IsUsingAlphaToCoverage());
 	OutEnvironment.SetDefine(TEXT("MOBILE_HIGH_QUALITY_BRDF"), IsMobileHighQualityBRDFEnabled());
 
-	EMaterialFloatPrecisionMode FloatPrecisionMode = GetMaterialFloatPrecisionMode();
-	if (FloatPrecisionMode == EMaterialFloatPrecisionMode::MFPM_Full)
-	{
-		OutEnvironment.CompilerFlags.Add(CFLAG_UseFullPrecisionInPS);
-	}
-	else if (FloatPrecisionMode == EMaterialFloatPrecisionMode::MFPM_Full_MaterialExpressionOnly)
+	bool bFullPrecisionInMaterial = false;
+	bool bFullPrecisionInPS = false;
+
+	GetOutputPrecision(GetMaterialFloatPrecisionMode(), bFullPrecisionInPS, bFullPrecisionInMaterial);
+
+	if (bFullPrecisionInMaterial)
 	{
 		OutEnvironment.SetDefine(TEXT("FORCE_MATERIAL_FLOAT_FULL_PRECISION"), TEXT("1"));
 	}
 
+	OutEnvironment.FullPrecisionInPS = bFullPrecisionInPS;
+
 	switch(GetMaterialDomain())
 	{
 		case MD_Surface:				OutEnvironment.SetDefine(TEXT("MATERIAL_DOMAIN_SURFACE"),			TEXT("1")); break;
@@ -2612,6 +2614,24 @@ TShaderRef<FShader> FMaterial::GetShader(FMeshMaterialShaderType* ShaderType, FV
 	return TShaderRef<FShader>(Shader, *RenderingThreadShaderMap);
 }
 
+void FMaterial::GetOutputPrecision(EMaterialFloatPrecisionMode FloatPrecisionMode, bool& FullPrecisionInPS, bool& FullPrecisionInMaterial)
+{
+	static const IConsoleVariable* CVarFloatPrecisionMode = IConsoleManager::Get().FindConsoleVariable(TEXT("r.Mobile.FloatPrecisionMode"));
+
+	if (FloatPrecisionMode != EMaterialFloatPrecisionMode::MFPM_Default)
+	{
+		FullPrecisionInMaterial = FloatPrecisionMode == EMaterialFloatPrecisionMode::MFPM_Full_MaterialExpressionOnly;
+		FullPrecisionInPS = FloatPrecisionMode == EMaterialFloatPrecisionMode::MFPM_Full;
+	}
+	else if (CVarFloatPrecisionMode)
+	{
+		int MobilePrecisionMode = FMath::Clamp(CVarFloatPrecisionMode->GetInt(), (int32_t)EMobileFloatPrecisionMode::Half, (int32_t)EMobileFloatPrecisionMode::Full);
+
+		FullPrecisionInMaterial = MobilePrecisionMode == EMobileFloatPrecisionMode::Full_MaterialExpressionOnly;
+		FullPrecisionInPS = MobilePrecisionMode == EMobileFloatPrecisionMode::Full;
+	}
+}
+
 TRACE_DECLARE_INT_COUNTER(Shaders_OnDemandShaderRequests, TEXT("Shaders/OnDemandShaderRequests"));
 bool FMaterial::TryGetShaders(const FMaterialShaderTypes& InTypes, const FVertexFactoryType* InVertexFactoryType, FMaterialShaders& OutShaders) const
 {
diff --git a/Engine/Source/Runtime/Engine/Private/ShaderCompiler/ShaderCompiler.cpp b/Engine/Source/Runtime/Engine/Private/ShaderCompiler/ShaderCompiler.cpp
--- a/Engine/Source/Runtime/Engine/Private/ShaderCompiler/ShaderCompiler.cpp
+++ b/Engine/Source/Runtime/Engine/Private/ShaderCompiler/ShaderCompiler.cpp
@@ -62,6 +62,7 @@
 #include "RenderUtils.h"
 #include "ProfilingDebugging/CountersTrace.h"
 #include "ClearReplacementShaders.h"
+#include "Engine/RendererSettings.h"
 #include "Misc/LargeWorldRenderPosition.h"
 
 #if WITH_EDITOR
@@ -5761,11 +5762,8 @@ void GlobalBeginCompileShader(
 	}
 
 	{
-		static IConsoleVariable* CVar = IConsoleManager::Get().FindConsoleVariable(TEXT("r.Mobile.ForceFullPrecisionInPS"));
-		if (CVar && CVar->GetInt() != 0)
-		{
-			Input.Environment.CompilerFlags.Add(CFLAG_UseFullPrecisionInPS);
-		}
+		static IConsoleVariable* CVar = IConsoleManager::Get().FindConsoleVariable(TEXT("r.Mobile.FloatPrecisionMode"));
+		Input.Environment.FullPrecisionInPS = CVar ? (CVar->GetInt() == EMobileFloatPrecisionMode::Full) : false;
 	}
 
 	{
diff --git a/Engine/Source/Runtime/Engine/Private/ShaderCompiler/ShaderGenerationUtil.cpp b/Engine/Source/Runtime/Engine/Private/ShaderCompiler/ShaderGenerationUtil.cpp
--- a/Engine/Source/Runtime/Engine/Private/ShaderCompiler/ShaderGenerationUtil.cpp
+++ b/Engine/Source/Runtime/Engine/Private/ShaderCompiler/ShaderGenerationUtil.cpp
@@ -1779,6 +1779,8 @@ void FShaderCompileUtilities::ApplyDerivedDefines(FShaderCompilerEnvironment& Ou
 		ApplyFetchEnvironment(MaterialDefines, *SharedEnvironment);
 		ApplyFetchEnvironment(LightmapDefines, *SharedEnvironment);
 		ApplyFetchEnvironment(CompilerDefines, *SharedEnvironment);
+
+		OutEnvironment.FullPrecisionInPS = SharedEnvironment->FullPrecisionInPS;
 	}
 
 	// for reference, here are the existing defines
diff --git a/Engine/Source/Runtime/Engine/Public/MaterialShared.h b/Engine/Source/Runtime/Engine/Public/MaterialShared.h
--- a/Engine/Source/Runtime/Engine/Public/MaterialShared.h
+++ b/Engine/Source/Runtime/Engine/Public/MaterialShared.h
@@ -1741,7 +1741,7 @@ public:
 	virtual bool IsUsedWithVirtualHeightfieldMesh() const { return false; }
 	virtual bool IsFullyRough() const { return false; }
 	virtual bool UseNormalCurvatureToRoughness() const { return false; }
-	virtual enum EMaterialFloatPrecisionMode GetMaterialFloatPrecisionMode() const { return EMaterialFloatPrecisionMode::MFPM_Half; };
+	virtual enum EMaterialFloatPrecisionMode GetMaterialFloatPrecisionMode() const { return EMaterialFloatPrecisionMode::MFPM_Default; };
 	virtual bool IsUsingAlphaToCoverage() const { return false; }
 	virtual bool IsUsingPreintegratedGFForSimpleIBL() const { return false; }
 	virtual bool IsUsingHQForwardReflections() const { return false; }
@@ -2245,6 +2245,9 @@ private:
 		FShaderCompilerEnvironment& OutEnvironment
 		) const;
 
+	/** Get the float precision mode for the material and shader taking into account any project wide precision values */
+	static void GetOutputPrecision(EMaterialFloatPrecisionMode FloatPrecisionMode, bool& bFullPrecisionInPS, bool& bFullPrecisionInMaterial);
+
 	/**
 	 * Finds the shader matching the template type and the passed in vertex factory, asserts if not found.
 	 */
diff --git a/Engine/Source/Runtime/RenderCore/Private/Shader.cpp b/Engine/Source/Runtime/RenderCore/Private/Shader.cpp
--- a/Engine/Source/Runtime/RenderCore/Private/Shader.cpp
+++ b/Engine/Source/Runtime/RenderCore/Private/Shader.cpp
@@ -1522,8 +1522,11 @@ void ShaderMapAppendKeyString(EShaderPlatform Platform, FString& KeyString)
 		}	
 		
 		{
-			static IConsoleVariable* CVar = IConsoleManager::Get().FindConsoleVariable(TEXT("r.Mobile.ForceFullPrecisionInPS"));
-			KeyString += (CVar && CVar->GetInt() != 0) ? TEXT("_highp") : TEXT("");
+			static IConsoleVariable* CVar = IConsoleManager::Get().FindConsoleVariable(TEXT("r.Mobile.FloatPrecisionMode"));
+			if(CVar && CVar->GetInt() > 0)
+			{
+				KeyString.Appendf(TEXT("_highp%d"), CVar->GetInt());
+			}
 		}
 
 		{
diff --git a/Engine/Source/Runtime/RenderCore/Public/ShaderCompilerCore.h b/Engine/Source/Runtime/RenderCore/Public/ShaderCompilerCore.h
--- a/Engine/Source/Runtime/RenderCore/Public/ShaderCompilerCore.h
+++ b/Engine/Source/Runtime/RenderCore/Public/ShaderCompilerCore.h
@@ -81,8 +81,6 @@ enum ECompilerFlags
 	CFLAG_BoundsChecking,
 	// Force removing unused interpolators for platforms that can opt out
 	CFLAG_ForceRemoveUnusedInterpolators,
-	// Set default precision to highp in a pixel shader (default is mediump on ES platforms)
-	CFLAG_UseFullPrecisionInPS,
 	// Hint that it is a vertex to geometry shader
 	CFLAG_VertexToGeometryShader,
 	// Hint that it is a vertex to primitive shader
diff --git a/Engine/Source/Runtime/RenderCore/Public/ShaderCore.h b/Engine/Source/Runtime/RenderCore/Public/ShaderCore.h
--- a/Engine/Source/Runtime/RenderCore/Public/ShaderCore.h
+++ b/Engine/Source/Runtime/RenderCore/Public/ShaderCore.h
@@ -421,6 +421,9 @@ struct FShaderCompilerEnvironment
 
 	const ITargetPlatform* TargetPlatform = nullptr;
 
+	// Used for mobile platforms to allow per shader/material precision modes
+	bool FullPrecisionInPS = 0;
+
 	/** Default constructor. */
 	FShaderCompilerEnvironment()
 	{
@@ -475,6 +478,7 @@ struct FShaderCompilerEnvironment
 		Ar << Environment.UniformBufferMap;
 		Ar << Environment.RemoteServerData;
 		Ar << Environment.ShaderFormatCVars;
+		Ar << Environment.FullPrecisionInPS;
 
 		return Ar;
 	}
@@ -506,6 +510,7 @@ struct FShaderCompilerEnvironment
 		RenderTargetOutputFormatsMap.Append(Other.RenderTargetOutputFormatsMap);
 		RemoteServerData.Append(Other.RemoteServerData);
 		ShaderFormatCVars.Append(Other.ShaderFormatCVars);
+		FullPrecisionInPS = Other.FullPrecisionInPS;
 	}
 
 private:
