Made non-user placed actors saved in unloaded WP cells unload as expected

diff --git a/Engine/Source/Runtime/Engine/Private/WorldPartition/WorldPartitionEditorCell.cpp b/Engine/Source/Runtime/Engine/Private/WorldPartition/WorldPartitionEditorCell.cpp
--- a/Engine/Source/Runtime/Engine/Private/WorldPartition/WorldPartitionEditorCell.cpp
+++ b/Engine/Source/Runtime/Engine/Private/WorldPartition/WorldPartitionEditorCell.cpp
@@ -54,7 +54,7 @@ void UWorldPartitionEditorCell::AddActor(const FGuid& Source, const FWorldPartit
 
 	if (!bIsAlreadyInSet)
 	{	
-		if (ActorHandle.IsLoaded() && !IsRunningCommandlet())
+		if (ActorHandle.IsLoaded() && !IsRunningCommandlet() && (IsLoaded() || ActorHandle->GetActor()->IsUserManaged()))
 		{
 			LoadedActors.Add(FActorReference(Source, ActorHandle), &bIsAlreadyInSet);
 			check(!bIsAlreadyInSet);
