U5 - Blocking the RenderingThread because BufferedSlateHandles is modified on the RT but later read on the GT

diff --git a/Engine/Source/Runtime/Engine/Private/Slate/SceneViewport.cpp b/Engine/Source/Runtime/Engine/Private/Slate/SceneViewport.cpp
--- a/Engine/Source/Runtime/Engine/Private/Slate/SceneViewport.cpp
+++ b/Engine/Source/Runtime/Engine/Private/Slate/SceneViewport.cpp
@@ -1634,6 +1634,10 @@ void FSceneViewport::UpdateViewportRHI(bool bDestroyed, uint32 NewSizeX, uint32
 		{
 			BeginInitResource(this);
 			
+			FRenderCommandFence InitResourceFence;
+			InitResourceFence.BeginFence();
+			InitResourceFence.Wait();
+
 			FSlateRenderer* Renderer = FSlateApplication::Get().GetRenderer();
 
 			TSharedPtr<SWidget> PinnedViewport = ViewportWidget.Pin();
