PR #8579: Fix missing FActorInstanceHandle::InstanceUID initialization in ctor (Contributed by [REDACTED_PERSON])

diff --git a/Engine/Source/Runtime/Engine/Private/EngineTypes.cpp b/Engine/Source/Runtime/Engine/Private/EngineTypes.cpp
--- a/Engine/Source/Runtime/Engine/Private/EngineTypes.cpp
+++ b/Engine/Source/Runtime/Engine/Private/EngineTypes.cpp
@@ -248,6 +248,7 @@ FActorInstanceHandle::FActorInstanceHandle()
 FActorInstanceHandle::FActorInstanceHandle(AActor* InActor)
 	: Actor(InActor)
 	, InstanceIndex(INDEX_NONE)
+	, InstanceUID(0)
 {
 	if (InActor)
 	{
