Unshelved from pending changelist '17793494':

PR #8504: Near Clipping Plane handling in FMinimalViewInfo (Contributed by [REDACTED_PERSON])

diff --git a/Engine/Source/Editor/UnrealEd/Public/EditorViewportClient.h b/Engine/Source/Editor/UnrealEd/Public/EditorViewportClient.h
--- a/Engine/Source/Editor/UnrealEd/Public/EditorViewportClient.h
+++ b/Engine/Source/Editor/UnrealEd/Public/EditorViewportClient.h
@@ -1183,6 +1183,9 @@ public:
 	/** Get the near clipping plane for this viewport. */
 	float GetNearClipPlane() const;
 
+	/** Override the near clipping plane. Set to a negative value to disable the override. */
+	void OverrideNearClipPlane(float InNearPlane);
+
 	/** Get the far clipping plane override for this viewport. */
 	float GetFarClipPlaneOverride() const;
 	
@@ -1320,9 +1323,6 @@ protected:
 	/** Invalidates the viewport widget (if valid) to register its active timer */
 	void InvalidateViewportWidget();
 
-	/** Subclasses may override the near clipping plane. Set to a negative value to disable the override. */
-	void OverrideNearClipPlane(float InNearPlane);
-
 	/** Constant for how much the camera safe zone rectangle is inset when being displayed in the editor */
 	static float const SafePadding;
 	/**
diff --git a/Engine/Source/Runtime/Engine/Classes/Camera/CameraTypes.h b/Engine/Source/Runtime/Engine/Classes/Camera/CameraTypes.h
--- a/Engine/Source/Runtime/Engine/Classes/Camera/CameraTypes.h
+++ b/Engine/Source/Runtime/Engine/Classes/Camera/CameraTypes.h
@@ -79,6 +79,10 @@ struct FMinimalViewInfo
 	UPROPERTY(Interp, EditAnywhere, BlueprintReadWrite, Category=Camera)
 	float OrthoFarClipPlane;
 
+	/** The near plane distance of the perspective view (in world units). Set to a negative value to use the default global value of GNearClippingPlane */
+	UPROPERTY(Interp, EditAnywhere, BlueprintReadWrite, Category=Camera)
+	float PerspectiveNearClipPlane;
+
 	// Aspect Ratio (Width/Height)
 	UPROPERTY(EditAnywhere, BlueprintReadWrite, Category=Camera)
 	float AspectRatio;
@@ -118,6 +122,7 @@ struct FMinimalViewInfo
 		, OrthoWidth(512.0f)
 		, OrthoNearClipPlane(0.0f)
 		, OrthoFarClipPlane(WORLD_MAX)
+		, PerspectiveNearClipPlane(-1.0f)
 		, AspectRatio(1.33333333f)
 		, bConstrainAspectRatio(false)
 		, bUseFieldOfViewForLOD(true)
@@ -145,4 +150,10 @@ struct FMinimalViewInfo
 
 	/** Calculates the projection matrix (and potentially a constrained view rectangle) given a FMinimalViewInfo and partially configured projection data (must have the view rect already set) */
 	ENGINE_API static void CalculateProjectionMatrixGivenView(const FMinimalViewInfo& ViewInfo, TEnumAsByte<enum EAspectRatioAxisConstraint> AspectRatioAxisConstraint, class FViewport* Viewport, struct FSceneViewProjectionData& InOutProjectionData);
+
+	/** The near plane distance of the perspective view (in world units). Returns the value of PerspectiveNearClipPlane if positive, and GNearClippingPlane otherwise */
+	ENGINE_API FORCEINLINE float GetFinalPerspectiveNearClipPlane() const
+	{
+		return PerspectiveNearClipPlane > 0.0f ? PerspectiveNearClipPlane : GNearClippingPlane;
+	}
 };
diff --git a/Engine/Source/Runtime/Engine/Private/Camera/CameraStackTypes.cpp b/Engine/Source/Runtime/Engine/Private/Camera/CameraStackTypes.cpp
--- a/Engine/Source/Runtime/Engine/Private/Camera/CameraStackTypes.cpp
+++ b/Engine/Source/Runtime/Engine/Private/Camera/CameraStackTypes.cpp
@@ -23,6 +23,8 @@ bool FMinimalViewInfo::Equals(const FMinimalViewInfo& OtherInfo) const
 		(OrthoWidth == OtherInfo.OrthoWidth) &&
 		(OrthoNearClipPlane == OtherInfo.OrthoNearClipPlane) &&
 		(OrthoFarClipPlane == OtherInfo.OrthoFarClipPlane) &&
+		((PerspectiveNearClipPlane == OtherInfo.PerspectiveNearClipPlane) || //either they are the same or both don't override
+			(PerspectiveNearClipPlane <= 0.f && OtherInfo.PerspectiveNearClipPlane <= 0.f)) &&
 		(AspectRatio == OtherInfo.AspectRatio) &&
 		(bConstrainAspectRatio == OtherInfo.bConstrainAspectRatio) &&
 		(bUseFieldOfViewForLOD == OtherInfo.bUseFieldOfViewForLOD) &&
@@ -41,6 +43,7 @@ void FMinimalViewInfo::BlendViewInfo(FMinimalViewInfo& OtherInfo, float OtherWei
 	OrthoWidth = FMath::Lerp(OrthoWidth, OtherInfo.OrthoWidth, OtherWeight);
 	OrthoNearClipPlane = FMath::Lerp(OrthoNearClipPlane, OtherInfo.OrthoNearClipPlane, OtherWeight);
 	OrthoFarClipPlane = FMath::Lerp(OrthoFarClipPlane, OtherInfo.OrthoFarClipPlane, OtherWeight);
+	PerspectiveNearClipPlane = FMath::Lerp(PerspectiveNearClipPlane, OtherInfo.PerspectiveNearClipPlane, OtherWeight);
 	OffCenterProjectionOffset = FMath::Lerp(OffCenterProjectionOffset, OtherInfo.OffCenterProjectionOffset, OtherWeight);
 
 	AspectRatio = FMath::Lerp(AspectRatio, OtherInfo.AspectRatio, OtherWeight);
@@ -57,6 +60,7 @@ void FMinimalViewInfo::ApplyBlendWeight(const float& Weight)
 	OrthoWidth *= Weight;
 	OrthoNearClipPlane *= Weight;
 	OrthoFarClipPlane *= Weight;
+	PerspectiveNearClipPlane *= Weight;
 	AspectRatio *= Weight;
 	OffCenterProjectionOffset *= Weight;
 }
@@ -72,6 +76,7 @@ void FMinimalViewInfo::AddWeightedViewInfo(const FMinimalViewInfo& OtherView, co
 	OrthoWidth += OtherViewWeighted.OrthoWidth;
 	OrthoNearClipPlane += OtherViewWeighted.OrthoNearClipPlane;
 	OrthoFarClipPlane += OtherViewWeighted.OrthoFarClipPlane;
+	PerspectiveNearClipPlane += OtherViewWeighted.PerspectiveNearClipPlane;
 	AspectRatio += OtherViewWeighted.AspectRatio;
 	OffCenterProjectionOffset += OtherViewWeighted.OffCenterProjectionOffset;
 
@@ -105,12 +110,13 @@ FMatrix FMinimalViewInfo::CalculateProjectionMatrix() const
 	}
 	else
 	{
+		const float ClippingPlane = GetFinalPerspectiveNearClipPlane();
 		// Avoid divide by zero in the projection matrix calculation by clamping FOV
 		ProjectionMatrix = FReversedZPerspectiveMatrix(
 			FMath::Max(0.001f, FOV) * (float)PI / 360.0f,
 			AspectRatio,
 			1.0f,
-			GNearClippingPlane );
+			ClippingPlane);
 	}
 
 	if (!OffCenterProjectionOffset.IsZero())
@@ -201,14 +207,15 @@ void FMinimalViewInfo::CalculateProjectionMatrixGivenView(const FMinimalViewInfo
 		}
 		else
 		{
+			const float ClippingPlane = ViewInfo.GetFinalPerspectiveNearClipPlane();
 			InOutProjectionData.ProjectionMatrix = FReversedZPerspectiveMatrix(
 				MatrixHalfFOV,
 				MatrixHalfFOV,
 				XAxisMultiplier,
 				YAxisMultiplier,
-				GNearClippingPlane,
-				GNearClippingPlane
-				);
+				ClippingPlane,
+				ClippingPlane
+			);
 		}
 	}
 
