Fix crash when running with -server in an editor build.

Fonts aren't initialized in that case, so Font.Value will be null.

#rnx

diff --git a/Engine/Source/Runtime/Engine/Private/Materials/MaterialCachedData.cpp b/Engine/Source/Runtime/Engine/Private/Materials/MaterialCachedData.cpp
--- a/Engine/Source/Runtime/Engine/Private/Materials/MaterialCachedData.cpp
+++ b/Engine/Source/Runtime/Engine/Private/Materials/MaterialCachedData.cpp
@@ -272,7 +272,7 @@ bool FMaterialCachedExpressionData::UpdateForExpressions(const FMaterialCachedEx
 				case EMaterialParameterType::Font:
 					Parameters.FontValues.Insert(ParameterMeta.Value.Font.Value, Index);
 					Parameters.FontPageValues.Insert(ParameterMeta.Value.Font.Page, Index);
-					if (ParameterMeta.Value.Font.Value->Textures.IsValidIndex(ParameterMeta.Value.Font.Page))
+					if (ParameterMeta.Value.Font.Value && ParameterMeta.Value.Font.Value->Textures.IsValidIndex(ParameterMeta.Value.Font.Page))
 					{
 						ReferencedTexture = ParameterMeta.Value.Font.Value->Textures[ParameterMeta.Value.Font.Page];
 					}
