Let UEngine::Browse handle cancelling the PendingNetGame if UOnlineSessionClient::HandleDisconnectInternal doesn't handle the disconnect.

diff --git a/Engine/Source/Runtime/Engine/Private/UnrealEngine.cpp b/Engine/Source/Runtime/Engine/Private/UnrealEngine.cpp
--- a/Engine/Source/Runtime/Engine/Private/UnrealEngine.cpp
+++ b/Engine/Source/Runtime/Engine/Private/UnrealEngine.cpp
@@ -12843,9 +12843,6 @@ void UEngine::HandleDisconnect( UWorld *InWorld, UNetDriver *NetDriver )
 	// There must be some context for this disconnect
 	check(InWorld || NetDriver);
 
-	// If the NetDriver that failed was a pending netgame driver, cancel the PendingNetGame
-	CancelPending(NetDriver);
-
 	// InWorld might be null. It might also not map to any valid world context (for example, a pending net game disconnect)
 	// If there is a context for this world, setup client travel.
 	if (FWorldContext* WorldContext = GetWorldContextFromWorld(InWorld))
@@ -12855,10 +12852,14 @@ void UEngine::HandleDisconnect( UWorld *InWorld, UNetDriver *NetDriver )
 		WorldContext->LastURL.RemoveOption( TEXT("LAN") );
 
 		// Net driver destruction will occur during LoadMap (prevents GetNetMode from changing output for the remainder of the frame)
+		// The ?closed option will also make sure any PendingNetGame is cancelled in Browse
 		SetClientTravel( InWorld, TEXT("?closed"), TRAVEL_Absolute );
 	}
 	else if (NetDriver)
 	{
+		// If the NetDriver that failed was a pending netgame driver, cancel the PendingNetGame
+		CancelPending(NetDriver);
+
 		// Shut down any existing game connections
 		if (InWorld)
 		{
