FBX Skeletal Mesh Importer: Store extra per individual mesh information from Fbx scenes.

diff --git a/Engine/Source/Editor/UnrealEd/Private/Fbx/FbxSkeletalMeshImport.cpp b/Engine/Source/Editor/UnrealEd/Private/Fbx/FbxSkeletalMeshImport.cpp
--- a/Engine/Source/Editor/UnrealEd/Private/Fbx/FbxSkeletalMeshImport.cpp
+++ b/Engine/Source/Editor/UnrealEd/Private/Fbx/FbxSkeletalMeshImport.cpp
@@ -1422,9 +1422,24 @@ bool UnFbx::FFbxImporter::FillSkeletalMeshImportPoints(FSkeletalMeshImportData*
 {
 	FbxMesh* FbxMesh = Node->GetMesh();
 
+	// Extract the name.
+	FString MeshName = MakeName(FbxMesh->GetName());
+	if (MeshName.IsEmpty())
+	{
+		MeshName = MakeName(Node->GetName());
+	}
+
+	// Add the mesh info
+	OutData->MeshInfos.AddDefaulted();
+	SkeletalMeshImportData::FMeshInfo& MeshInfo = OutData->MeshInfos.Last();
+	MeshInfo.Name = *MeshName;
+
 	const int32 ControlPointsCount = FbxMesh->GetControlPointsCount();
 	const int32 ExistPointNum = OutData->Points.Num();
 	OutData->Points.AddUninitialized(ControlPointsCount);
+	
+	MeshInfo.StartImportedVertex = ExistPointNum;
+	MeshInfo.NumVertices = ControlPointsCount;
 
 	// Construct the matrices for the conversion from right handed to left handed system
 	FbxAMatrix TotalMatrix = ComputeSkeletalMeshTotalMatrix(Node, RootNode);
@@ -1859,6 +1874,18 @@ USkeletalMesh* UnFbx::FFbxImporter::ImportSkeletalMesh(FImportSkeletalMeshArgs &
 	// Pass the number of texture coordinate sets to the LODModel.  Ensure there is at least one UV coord
 	LODModel.NumTexCoords = FMath::Max<uint32>(1, SkelMeshImportDataPtr->NumTexCoords);
 
+	// Copy mesh infos into the LOD model.
+	LODModel.ImportedMeshInfos.Empty();
+	LODModel.ImportedMeshInfos.Reserve(SkelMeshImportDataPtr->MeshInfos.Num());
+	for (const SkeletalMeshImportData::FMeshInfo& MeshInfo : SkelMeshImportDataPtr->MeshInfos)
+	{
+		LODModel.ImportedMeshInfos.AddDefaulted();
+		FSkelMeshImportedMeshInfo& LODMeshInfo = LODModel.ImportedMeshInfos.Last();
+		LODMeshInfo.Name = MeshInfo.Name;
+		LODMeshInfo.NumVertices = MeshInfo.NumVertices;
+		LODMeshInfo.StartImportedVertex = MeshInfo.StartImportedVertex;
+	}
+
 	if(ImportSkeletalMeshArgs.bCreateRenderData )
 	{
 		TArray<FVector3f> LODPoints;
diff --git a/Engine/Source/Runtime/Core/Public/UObject/UE5MainStreamObjectVersion.h b/Engine/Source/Runtime/Core/Public/UObject/UE5MainStreamObjectVersion.h
--- a/Engine/Source/Runtime/Core/Public/UObject/UE5MainStreamObjectVersion.h
+++ b/Engine/Source/Runtime/Core/Public/UObject/UE5MainStreamObjectVersion.h
@@ -183,6 +183,9 @@ struct CORE_API FUE5MainStreamObjectVersion
 
 		// Remove the inline reduction bulkdata and replace it by a simple vertex and triangle count cache
 		ConvertReductionBaseSkeletalMeshBulkDataToInlineReductionCacheData,
+
+		// Added some new MeshInfo to the FSkeletalMeshLODModel class.
+		SkeletalMeshLODModelMeshInfo,
 		
 		// -----<new versions can be added above this line>-------------------------------------------------
 		VersionPlusOne,
diff --git a/Engine/Source/Runtime/Engine/Private/Rendering/SkeletalMeshLODImporterData.cpp b/Engine/Source/Runtime/Engine/Private/Rendering/SkeletalMeshLODImporterData.cpp
--- a/Engine/Source/Runtime/Engine/Private/Rendering/SkeletalMeshLODImporterData.cpp
+++ b/Engine/Source/Runtime/Engine/Private/Rendering/SkeletalMeshLODImporterData.cpp
@@ -161,6 +161,7 @@ bool FSkeletalMeshImportData::ReplaceSkeletalMeshGeometryImportData(const USkele
 	ImportData->MorphTargetNames.Reset();
 	ImportData->MorphTargets.Reset();
 	ImportData->MorphTargetModifiedPoints.Reset();
+	ImportData->MeshInfos.Reset();
 
 	//Material is a special case since we cannot serialize the UMaterialInstance when saving the RawSkeletalMeshBulkData
 	//So it has to be reconstructed.
@@ -184,6 +185,7 @@ bool FSkeletalMeshImportData::ReplaceSkeletalMeshGeometryImportData(const USkele
 	ImportData->MorphTargetNames += OriginalSkeletalMeshImportData.MorphTargetNames;
 	ImportData->MorphTargets += OriginalSkeletalMeshImportData.MorphTargets;
 	ImportData->MorphTargetModifiedPoints += OriginalSkeletalMeshImportData.MorphTargetModifiedPoints;
+	ImportData->MeshInfos += OriginalSkeletalMeshImportData.MeshInfos;
 
 	return ImportData->ApplyRigToGeo(NewGeometryAndRigData);
 }
diff --git a/Engine/Source/Runtime/Engine/Private/SkeletalMeshLODModel.cpp b/Engine/Source/Runtime/Engine/Private/SkeletalMeshLODModel.cpp
--- a/Engine/Source/Runtime/Engine/Private/SkeletalMeshLODModel.cpp
+++ b/Engine/Source/Runtime/Engine/Private/SkeletalMeshLODModel.cpp
@@ -18,12 +18,24 @@
 #include "UObject/RenderingObjectVersion.h"
 #include "Rendering/SkeletalMeshLODImporterData.h"
 #include "UObject/FortniteMainBranchObjectVersion.h"
+#include "UObject/UE5MainStreamObjectVersion.h"
 #include "GPUSkinVertexFactory.h"
 #include "UObject/AnimObjectVersion.h"
 #include "Misc/ScopeLock.h"
 #include "MeshDescription.h"
 #include "SkeletalMeshAttributes.h"
 
+/*-----------------------------------------------------------------------------
+FSkelMeshImportedMeshInfo
+-----------------------------------------------------------------------------*/
+FArchive& operator<<(FArchive& Ar, FSkelMeshImportedMeshInfo& MeshInfo)
+{
+	Ar << MeshInfo.Name;
+	Ar << MeshInfo.NumVertices;
+	Ar << MeshInfo.StartImportedVertex;
+	return Ar;
+}
+
 /*-----------------------------------------------------------------------------
 FSoftSkinVertex
 -----------------------------------------------------------------------------*/
@@ -694,6 +706,7 @@ void FSkeletalMeshLODModel::Serialize(FArchive& Ar, UObject* Owner, int32 Idx)
 	Ar.UsingCustomVersion(FSkeletalMeshCustomVersion::GUID);
 	Ar.UsingCustomVersion(FFortniteMainBranchObjectVersion::GUID);
 	Ar.UsingCustomVersion(FEditorObjectVersion::GUID);
+	Ar.UsingCustomVersion(FUE5MainStreamObjectVersion::GUID);
 
 	if (StripFlags.IsDataStrippedForServer())
 	{
@@ -720,6 +733,12 @@ void FSkeletalMeshLODModel::Serialize(FArchive& Ar, UObject* Owner, int32 Idx)
 
 		TArray<FBoneIndexType> TempActiveBoneIndices;
 		Ar << TempActiveBoneIndices;
+
+		if (Ar.CustomVer(FUE5MainStreamObjectVersion::GUID) >= FUE5MainStreamObjectVersion::SkeletalMeshLODModelMeshInfo)
+		{
+			TArray<FSkelMeshImportedMeshInfo> TempMeshInfos;
+			Ar << TempMeshInfos;
+		}
 	}
 	else
 	{
@@ -750,6 +769,12 @@ void FSkeletalMeshLODModel::Serialize(FArchive& Ar, UObject* Owner, int32 Idx)
 		}
 
 		Ar << ActiveBoneIndices;
+
+		// Editor only data.
+		if (!StripFlags.IsEditorDataStripped() && Ar.CustomVer(FUE5MainStreamObjectVersion::GUID) >= FUE5MainStreamObjectVersion::SkeletalMeshLODModelMeshInfo)
+		{
+			Ar << ImportedMeshInfos;
+		}
 	}
 
 	// Array of Sections for backwards compat
@@ -910,6 +935,17 @@ void FSkeletalMeshLODModel::Serialize(FArchive& Ar, UObject* Owner, int32 Idx)
 	}
 }
 
+int32 FSkeletalMeshLODModel::FindMeshInfoIndex(FName Name) const
+{
+	for (int32 Index = 0; Index < ImportedMeshInfos.Num(); ++Index)
+	{
+		if (ImportedMeshInfos[Index].Name == Name)
+		{
+			return Index;
+		}
+	}
+	return INDEX_NONE;
+}
 
 void FSkeletalMeshLODModel::GetSectionFromVertexIndex(int32 InVertIndex, int32& OutSectionIndex, int32& OutVertIndex) const
 {
diff --git a/Engine/Source/Runtime/Engine/Public/Rendering/SkeletalMeshLODImporterData.h b/Engine/Source/Runtime/Engine/Public/Rendering/SkeletalMeshLODImporterData.h
--- a/Engine/Source/Runtime/Engine/Public/Rendering/SkeletalMeshLODImporterData.h
+++ b/Engine/Source/Runtime/Engine/Public/Rendering/SkeletalMeshLODImporterData.h
@@ -51,6 +51,17 @@ enum class ESkeletalMeshSkinningImportVersions : uint8
 
 namespace SkeletalMeshImportData
 {
+	/**
+	* Some information per individual mesh, as appearing in the source asset.
+	* This could store the data for say meshes "UpperBody", "Legs", "Hat", etc.
+	*/
+	struct FMeshInfo
+	{
+		FName Name;	// The name of the mesh.
+		int32 NumVertices;	// The number of imported (dcc) vertices that are part of this mesh. This is a value of 8 for a cube. So NOT the number of render vertices.
+		int32 StartImportedVertex;	// The first index of imported (dcc) vertices in the mesh. So this NOT an index into the render vertex buffer. In range of 0..7 for a cube.
+	};
+
 	struct FMeshWedge
 	{
 		uint32			iVertex;			// Vertex index.
@@ -335,6 +346,7 @@ public:
 	TArray <SkeletalMeshImportData::FTriangle> Faces;
 	TArray <SkeletalMeshImportData::FBone> RefBonesBinary;
 	TArray <SkeletalMeshImportData::FRawBoneInfluence> Influences;
+	TArray <SkeletalMeshImportData::FMeshInfo> MeshInfos;
 	TArray <int32> PointToRawMap;	// Mapping from current point index to the original import point index
 	uint32 NumTexCoords; // The number of texture coordinate sets
 	uint32 MaxMaterialIndex; // The max material index found on a triangle
@@ -407,6 +419,7 @@ public:
 		RefBonesBinary.Empty();
 		Influences.Empty();
 		PointToRawMap.Empty();
+		MeshInfos.Empty();
 	}
 
 	static bool ReplaceSkeletalMeshGeometryImportData(const USkeletalMesh* SkeletalMesh, FSkeletalMeshImportData* ImportData, int32 LodIndex);
diff --git a/Engine/Source/Runtime/Engine/Public/Rendering/SkeletalMeshLODModel.h b/Engine/Source/Runtime/Engine/Public/Rendering/SkeletalMeshLODModel.h
--- a/Engine/Source/Runtime/Engine/Public/Rendering/SkeletalMeshLODModel.h
+++ b/Engine/Source/Runtime/Engine/Public/Rendering/SkeletalMeshLODModel.h
@@ -311,6 +311,23 @@ struct FSkelMeshSourceSectionUserData
 	friend FArchive& operator<<(FArchive& Ar, FSkelMeshSourceSectionUserData& S);
 };
 
+/**
+ * The information for a given mesh as it was imported and appears in the dcc.
+ * This gives us some extra information about the mesh as it appeared inside the source asset.
+ * There can be multiple of those per skeletal mesh. The skeletal mesh is some merged big mesh separated in sections.
+ * However sometimes we want to know some information about the actual individual meshes as they appeared in the dcc.
+ * This class provides some of that information.
+ */
+struct FSkelMeshImportedMeshInfo
+{
+	FName Name;	// The name of the mesh.
+	int32 NumVertices;	// The number of imported (dcc) vertices that are part of this mesh. This is a value of 8 for a cube. So NOT the number of render vertices.
+	int32 StartImportedVertex;	// The first index of imported (dcc) vertices in the mesh. So this NOT an index into the render vertex buffer. In range of 0..7 for a cube.
+
+	// Serialization.
+	friend FArchive& operator<<(FArchive& Ar, FSkelMeshImportedMeshInfo& MeshInfo);
+};
+
 /**
 * All data to define a certain LOD model for a skeletal mesh.
 */
@@ -320,6 +337,9 @@ public:
 	/** Sections. */
 	TArray<FSkelMeshSection> Sections;
 
+	/** The information about individual meshes. This can be empty if we did choose not to store this information. */
+	TArray<FSkelMeshImportedMeshInfo> ImportedMeshInfos;
+
 	/*
 	 * When user change section data in the UI, we store it here to be able to regenerate the changes
 	 * Note: the key (int32) is the original imported section data, because of BONE chunk the size of
@@ -401,6 +421,9 @@ public:
 		BulkDataReadMutex = BackupBulkDataReadMutex;
 	}
 
+	/** Find a mesh info by a mesh name. Returns INDEX_NONE when not found. */
+	int32 FindMeshInfoIndex(FName Name) const;
+
 private:
 	//Mutex use by the CopyStructure function. It's a pointer because FCriticalSection privatize the operator= function, which will prevent this class operator= to use the default.
 	//We want to avoid having a custom equal operator that will get deprecated if dev forget to add the new member in this class
