KismetMathLibrary: Deprecate broken MinimumAreaRectangle() and introduce properly working MinAreaRectangle()

diff --git a/Engine/Source/Runtime/Engine/Classes/Kismet/KismetMathLibrary.h b/Engine/Source/Runtime/Engine/Classes/Kismet/KismetMathLibrary.h
--- a/Engine/Source/Runtime/Engine/Classes/Kismet/KismetMathLibrary.h
+++ b/Engine/Source/Runtime/Engine/Classes/Kismet/KismetMathLibrary.h
@@ -4071,9 +4071,28 @@ class ENGINE_API UKismetMathLibrary : public UBlueprintFunctionLibrary
 	 * @outparam	OutRectSideA - Vector oriented and sized to represent one edge of the enclosing rectangle, orthogonal to OutRectSideB
 	 * @outparam	OutRectSideB - Vector oriented and sized to represent one edge of the enclosing rectangle, orthogonal to OutRectSideA
 	*/
-	UFUNCTION(BlueprintCallable, Category="Math|Geometry", meta=(WorldContext="WorldContextObject", CallableWithoutWorldContext))
+	UE_DEPRECATED(5.2, "Use MinAreaRectangle instead. This deprecated version incorrectly returns the average of all input points as the rectangle center.")
+	UFUNCTION(BlueprintCallable, Category="Math|Geometry", meta=(WorldContext="WorldContextObject", CallableWithoutWorldContext,
+		DeprecatedFunction, DeprecationMessage = "Use 'Min Area Rectangle' instead; this deprecated version incorrectly returns the average of all input points as the rectangle center."))
 	static void MinimumAreaRectangle(UObject* WorldContextObject, const TArray<FVector>& InVerts, const FVector& SampleSurfaceNormal, FVector& OutRectCenter, FRotator& OutRectRotation, float& OutSideLengthX, float& OutSideLengthY, bool bDebugDraw = false);
 
+	/**
+	 * Finds the minimum area rectangle that encloses a set of coplanar points.
+	 * Uses the exhaustive search algorithm in http://www.geometrictools.com/Documentation/MinimumAreaRectangle.pdf
+	 *
+	 * @param	WorldContextObject - Pointer to world context; only used when debug draw is enabled
+	 * @param	InPoints - Points to enclose in the rectangle; need to be within the same plane for correct results
+	 * @param	SampleSurfaceNormal - Normal indicating the surface direction for the points
+	 * @param	OutRectCenter - Translation for the output rectangle from the origin
+	 * @param	OutRectRotation - Rotation for the output rectangle from the XY plane
+	 * @param	OutRectLengthX - Length of the output rectangle along the X axis before rotation
+	 * @param	OutRectLengthY - Length of the output rectangle along the Y axis before rotation
+	 * @param	bDebugDraw - Draws the output rectangle for debugging purposes provided the world context is set as well
+	*/
+	UFUNCTION(BlueprintCallable, Category="Math|Geometry", meta=(WorldContext="WorldContextObject", CallableWithoutWorldContext))
+	static void MinAreaRectangle(UObject* WorldContextObject, const TArray<FVector>& InPoints, const FVector& SampleSurfaceNormal, FVector& OutRectCenter,
+	                             FRotator& OutRectRotation, float& OutRectLengthX, float& OutRectLengthY, bool bDebugDraw = false);
+
 	/**
 	 * Determines whether a given set of points are coplanar, with a tolerance. Any three points or less are always coplanar.
 	 *
diff --git a/Engine/Source/Runtime/Engine/Private/KismetMathLibrary.cpp b/Engine/Source/Runtime/Engine/Private/KismetMathLibrary.cpp
--- a/Engine/Source/Runtime/Engine/Private/KismetMathLibrary.cpp
+++ b/Engine/Source/Runtime/Engine/Private/KismetMathLibrary.cpp
@@ -1149,6 +1149,130 @@ void UKismetMathLibrary::MinimumAreaRectangle(class UObject* WorldContextObject,
 #endif
 }
 
+void UKismetMathLibrary::MinAreaRectangle(class UObject* WorldContextObject, const TArray<FVector>& InPoints, const FVector& SampleSurfaceNormal,
+                                          FVector& OutRectCenter, FRotator& OutRectRotation, float& OutRectLengthX, float& OutRectLengthY, bool bDebugDraw)
+{
+	const int32 InVertsNum = InPoints.Num();
+	
+	// Bail if we receive an empty InVerts array
+	if (InVertsNum == 0)
+	{
+		OutRectCenter = FVector::ZeroVector;
+		OutRectRotation = FRotator::ZeroRotator;
+		OutRectLengthX = 0.0f;
+		OutRectLengthY = 0.0f;
+		return;
+	}
+
+	// Use 'Newell's Method' to compute a robust 'best fit' plane from the input points
+	FVector PlaneNormal = FVector::ZeroVector;
+	for (int32 Vert0 = InVertsNum - 1, Vert1 = 0; Vert1 < InVertsNum; Vert0 = Vert1++)
+	{
+		const FVector& P0 = InPoints[Vert0];
+		const FVector& P1 = InPoints[Vert1];
+		PlaneNormal.X += (P1.Y - P0.Y) * (P0.Z + P1.Z);
+		PlaneNormal.Y += (P1.Z - P0.Z) * (P0.X + P1.X);
+		PlaneNormal.Z += (P1.X - P0.X) * (P0.Y + P1.Y);
+	}
+	PlaneNormal.Normalize();
+	if ((PlaneNormal | SampleSurfaceNormal) < 0.f)
+	{
+		PlaneNormal = -PlaneNormal;
+	}
+
+	// Transform the sample points to 2D
+	const FMatrix SurfaceNormalMatrix = FRotationMatrix::MakeFromZX(PlaneNormal, FVector(1.f, 0.f, 0.f));
+	TArray<FVector> TransformedVerts;
+	TransformedVerts.SetNumUninitialized(InVertsNum);
+	const FMatrix InverseSurfaceNormalMatrix = SurfaceNormalMatrix.Inverse();
+	for (int32 Idx = 0; Idx < InVertsNum; ++Idx)
+	{
+		const FVector TransformedVert = InverseSurfaceNormalMatrix.TransformVector(InPoints[Idx]);
+		TransformedVerts[Idx] = {TransformedVert.X, TransformedVert.Y, 0.0f};
+	}
+
+	// Compute the convex hull of the sample points
+	TArray<int32> PolyVertIndices;
+	ConvexHull2D::ComputeConvexHullLegacy(TransformedVerts, PolyVertIndices);
+
+	// Minimum area rectangle as computed by the exhaustive search algorithm in http://www.geometrictools.com/Documentation/MinimumAreaRectangle.pdf
+	FVector Side0, Side1;
+	[&PolyVertIndices, &TransformedVerts, &OutRectCenter, &Side0, &Side1]
+	{
+		float MinArea = TNumericLimits<float>::Max();
+
+		for (int32 Idx0 = PolyVertIndices.Num() - 1, Idx1 = 0, Num = PolyVertIndices.Num(); Idx1 < Num; Idx0 = Idx1++)
+		{
+			const FVector Origin = TransformedVerts[PolyVertIndices[Idx0]];
+			const FVector U0 = (TransformedVerts[PolyVertIndices[Idx1]] - Origin).GetSafeNormal2D();
+			const FVector U1{-U0.Y, U0.X, 0.0f};
+
+			float MinU0 = 0.0f;
+			float MaxU0 = 0.0f;
+			float MaxU1 = 0.0f;
+
+			for (int VertIdx = 0; VertIdx < Num; ++VertIdx)
+			{
+				const FVector D = TransformedVerts[PolyVertIndices[VertIdx]] - Origin;
+				const float Dot0 = U0 | D;
+				if (Dot0 < MinU0)
+				{
+					MinU0 = Dot0;
+				}
+				else if (Dot0 > MaxU0)
+				{
+					MaxU0 = Dot0;
+				}
+
+				const float Dot1 = U1 | D;
+				if (Dot1 > MaxU1)
+				{
+					MaxU1 = Dot1;
+				}
+			}
+
+			const float Side0Length = MaxU0 - MinU0;
+			const float Side1Length = MaxU1;
+			const float CurrentArea = Side0Length * Side1Length;
+			if (CurrentArea < MinArea)
+			{
+				MinArea = CurrentArea;
+				Side0 = U0 * Side0Length;
+				Side1 = U1 * Side1Length;
+				OutRectCenter = Origin + ((MinU0 + MaxU0) / 2.0f) * U0 + (MaxU1 / 2.0f) * U1;
+			}
+		}
+	}();
+
+	Side0 = SurfaceNormalMatrix.TransformVector(Side0);
+	Side1 = SurfaceNormalMatrix.TransformVector(Side1);
+	OutRectCenter = SurfaceNormalMatrix.TransformPosition(OutRectCenter);
+	OutRectRotation = FRotationMatrix::MakeFromZX(PlaneNormal, Side0).Rotator();
+	OutRectLengthX = Side0.Length();
+	OutRectLengthY = Side1.Length();
+
+#if ENABLE_DRAW_DEBUG
+	if (bDebugDraw)
+	{
+		if (const UWorld* World = GEngine->GetWorldFromContextObject(WorldContextObject, EGetWorldErrorMode::LogAndReturnNull))
+		{
+			DrawDebugSphere(World, OutRectCenter, 10.0f, 12, FColor::Yellow, true);
+			DrawDebugCoordinateSystem(World, OutRectCenter, OutRectRotation, 50.0f, true);
+
+			auto DrawLine = [World, &OutRectCenter, &Side0, &Side1](float A, float B, float C, float D)
+			{
+				DrawDebugLine(World, OutRectCenter + A * Side0 + B * Side1, OutRectCenter + C * Side0 + D * Side1, FColor::Blue, true, -1, 0, 1.0f);
+			};
+
+			DrawLine(-0.5f, -0.5f,  0.5f, -0.5f);
+			DrawLine( 0.5f, -0.5f,  0.5f,  0.5f);
+			DrawLine( 0.5f,  0.5f, -0.5f,  0.5f);
+			DrawLine(-0.5f,  0.5f, -0.5f, -0.5f);
+		}
+	}
+#endif
+}
+
 bool UKismetMathLibrary::IsPointInBox(FVector Point, FVector BoxOrigin, FVector BoxExtent)
 {
 	const FBox Box(BoxOrigin - BoxExtent, BoxOrigin + BoxExtent);
