Increase default value for NetPackedMovementMaxBits

diff --git a/Engine/Source/Runtime/Engine/Private/Components/CharacterMovementComponent.cpp b/Engine/Source/Runtime/Engine/Private/Components/CharacterMovementComponent.cpp
--- a/Engine/Source/Runtime/Engine/Private/Components/CharacterMovementComponent.cpp
+++ b/Engine/Source/Runtime/Engine/Private/Components/CharacterMovementComponent.cpp
@@ -112,7 +112,7 @@ namespace CharacterMovementCVars
 		TEXT("0: Disable, 1: Enable"),
 		ECVF_Default);
 
-	static int32 NetPackedMovementMaxBits = 2048;
+	static int32 NetPackedMovementMaxBits = 4096;
 	FAutoConsoleVariableRef CVarNetPackedMovementMaxBits(
 		TEXT("p.NetPackedMovementMaxBits"),
 		NetPackedMovementMaxBits,
