CustomBoneIndexArray improved to support construction, assignment, allocators, and ArrayView

diff --git a/Engine/Source/Runtime/AnimationCore/Public/BoneIndices.h b/Engine/Source/Runtime/AnimationCore/Public/BoneIndices.h
--- a/Engine/Source/Runtime/AnimationCore/Public/BoneIndices.h
+++ b/Engine/Source/Runtime/AnimationCore/Public/BoneIndices.h
@@ -17,6 +17,10 @@ struct FBoneIndexBase
 
 	FORCEINLINE bool IsValid() const { return BoneIndex != INDEX_NONE; }
 
+	FORCEINLINE explicit operator int32() const { return BoneIndex; }
+
+	FORCEINLINE explicit operator bool() const { return IsValid(); }
+
 protected:
 	int32 BoneIndex;
 };
diff --git a/Engine/Source/Runtime/Engine/Public/CustomBoneIndexArray.h b/Engine/Source/Runtime/Engine/Public/CustomBoneIndexArray.h
--- a/Engine/Source/Runtime/Engine/Public/CustomBoneIndexArray.h
+++ b/Engine/Source/Runtime/Engine/Public/CustomBoneIndexArray.h
@@ -4,24 +4,114 @@
 
 #include "CoreMinimal.h"
 
-template<typename ArrayType, typename BoneIndexType>
-class TCustomBoneIndexArray : public TArray < ArrayType >
+#include <initializer_list>
+#include <type_traits>
+
+
+namespace UE::Anim { namespace Private {
+
+// TRangeIndexType extracts RangeType::IndexType if it exists, otherwise void
+template<typename RangeType, typename = void>
+struct TRangeIndexType 
+{
+	using Type = void;
+};
+
+template<typename RangeType>
+struct TRangeIndexType<RangeType, std::void_t<typename std::decay_t<RangeType>::IndexType>>
+{
+	using Type = typename std::decay_t<RangeType>::IndexType;
+};
+
+// IndexType is compatible with RangeType::IndexType if they are the same or if RangeType doesn't have an index type
+template<typename IndexType, typename RangeType>
+struct TIsCompatibleRangeIndexType
+{
+	enum
+	{
+		Value = TOr<
+			TIsSame<IndexType, typename TRangeIndexType<RangeType>::Type>,
+			TIsVoidType<typename TRangeIndexType<RangeType>::Type>
+		>::Value
+	};
+};
+
+} // namespace UE::Anim::Private
+
+
+template<typename InIndexType, typename InRangeType>
+class TTypedIndexRange : public InRangeType
 {
 public:
-	FORCEINLINE ArrayType& operator[](int32 BoneIndex)
-	{ 
-		return TArray<ArrayType>::operator [](BoneIndex); 
+	using BaseType = InRangeType;
+	using IndexType = InIndexType;
+	using ElementType = typename InRangeType::ElementType;
+	using SizeType = typename InRangeType::SizeType;
+
+	TTypedIndexRange() = default;
+
+	// Forwarding constructor that disallows mixing index types
+	template<typename T>
+	TTypedIndexRange(T&& Other)
+		: BaseType(Forward<T>(Other))
+	{
+		static_assert(Private::TIsCompatibleRangeIndexType<IndexType, T>::Value, "TTypedIndexRange can't construct from a different index type");
 	}
-	FORCEINLINE const ArrayType& operator[](int32 BoneIndex) const
-	{ 
-		return TArray<ArrayType>::operator [](BoneIndex); 
+
+	TTypedIndexRange(std::initializer_list<ElementType> List)
+		: BaseType(List)
+	{
+	}
+
+	// Forwarding assignment that disallows mixing index types
+	template<typename T>
+	TTypedIndexRange& operator=(T&& Other)
+	{
+		static_assert(Private::TIsCompatibleRangeIndexType<IndexType, T>::Value, "TTypedIndexRange can't assign from a different index type");
+		BaseType::operator=(Forward<T>(Other));
+		return *this;
+	}
+
+	TTypedIndexRange& operator=(std::initializer_list<ElementType> List)
+	{
+		BaseType::operator=(List);
+		return *this;
+	}
+
+	FORCEINLINE ElementType& operator[](SizeType Index)
+	{
+		return BaseType::operator[](Index);
+	}
+	FORCEINLINE const ElementType& operator[](SizeType Index) const
+	{
+		return BaseType::operator[](Index);
 	}
-	FORCEINLINE ArrayType& operator[](const BoneIndexType& BoneIndex) 
-	{ 
-		return TArray<ArrayType>::operator [] (BoneIndex.GetInt()); 
+	FORCEINLINE ElementType& operator[](const IndexType& Index)
+	{
+		return BaseType::operator[]((SizeType)Index);
 	}
-	FORCEINLINE const ArrayType& operator[](const BoneIndexType& BoneIndex) const 
-	{ 
-		return TArray<ArrayType>::operator [] (BoneIndex.GetInt()); 
+	FORCEINLINE const ElementType& operator[](const IndexType& Index) const
+	{
+		return BaseType::operator[]((SizeType)Index);
 	}
 };
+
+template<typename IndexType, typename ...TArrayArgs>
+using TTypedIndexArray = TTypedIndexRange<IndexType, TArray<TArrayArgs...>>;
+
+template<typename IndexType, typename ...TArrayViewArgs>
+using TTypedIndexArrayView = TTypedIndexRange<IndexType, TArrayView<TArrayViewArgs...>>;
+
+} // namespace UE::Anim
+
+template<typename ElementType, typename BoneIndexType>
+using TCustomBoneIndexArray = UE::Anim::TTypedIndexRange<BoneIndexType, TArray<ElementType>>;
+
+template<typename ElementType, typename BoneIndexType>
+using TCustomBoneIndexArrayView = UE::Anim::TTypedIndexRange<BoneIndexType, TArrayView<ElementType>>;
+
+template <typename ...ArgTypes>
+struct TIsContiguousContainer<UE::Anim::TTypedIndexRange<ArgTypes...>>
+{
+	enum { Value = true };
+};
\ No newline at end of file
