Don't cache HLSLTree on UMaterialFunction object, instead just generate the tree inline along with the material
This removes potential optimization of caching tree generation per-function
However, I've decided that I'd rather just make the HLSLTree completely transient, and avoid complications related to allowing nodes to be cached/re-used across different materials

diff --git a/Engine/Source/Runtime/Engine/Classes/Materials/MaterialFunctionInterface.h b/Engine/Source/Runtime/Engine/Classes/Materials/MaterialFunctionInterface.h
--- a/Engine/Source/Runtime/Engine/Classes/Materials/MaterialFunctionInterface.h
+++ b/Engine/Source/Runtime/Engine/Classes/Materials/MaterialFunctionInterface.h
@@ -18,7 +18,6 @@ class UMaterial;
 class UTexture;
 struct FPropertyChangedEvent;
 class FMaterialHLSLGenerator;
-class FMaterialHLSLTree;
 
 /** Usage set on a material function determines feature compatibility and validation. */
 UENUM()
@@ -50,15 +49,6 @@ class UMaterialFunctionInterface : public UObject
 	UPROPERTY(duplicatetransient)
 	FGuid StateId;
 
-#if WITH_EDITOR
-public:
-	FMaterialHLSLTree& AcquireHLSLTree(FMaterialHLSLGenerator& Generator);
-#endif
-#if WITH_EDITORONLY_DATA
-private:
-	FMaterialHLSLTree* CachedHLSLTree = nullptr;
-#endif
-
 protected:
 	/** The intended usage of this function, required for material layers. */
 	UPROPERTY(AssetRegistrySearchable)
diff --git a/Engine/Source/Runtime/Engine/Private/Materials/MaterialExpressionHLSL.cpp b/Engine/Source/Runtime/Engine/Private/Materials/MaterialExpressionHLSL.cpp
--- a/Engine/Source/Runtime/Engine/Private/Materials/MaterialExpressionHLSL.cpp
+++ b/Engine/Source/Runtime/Engine/Private/Materials/MaterialExpressionHLSL.cpp
@@ -56,16 +56,6 @@ EMaterialGenerateHLSLStatus UMaterialExpression::GenerateHLSLTexture(FMaterialHL
 	return Generator.Error(TEXT("Node does not support textures"));
 }
 
-FMaterialHLSLTree& UMaterialFunctionInterface::AcquireHLSLTree(FMaterialHLSLGenerator& Generator)
-{
-	if (!CachedHLSLTree)
-	{
-		CachedHLSLTree = new FMaterialHLSLTree();
-		CachedHLSLTree->InitializeForFunction(Generator.GetCompileTarget(), this);
-	}
-	return *CachedHLSLTree;
-}
-
 EMaterialGenerateHLSLStatus UMaterialExpressionGenericConstant::GenerateHLSLExpression(FMaterialHLSLGenerator& Generator, UE::HLSLTree::FScope& Scope, int32 OutputIndex, UE::HLSLTree::FExpression*& OutExpression)
 {
 	OutExpression = Generator.NewConstant(GetConstantValue());
@@ -311,29 +301,12 @@ EMaterialGenerateHLSLStatus UMaterialExpressionFunctionInput::GenerateHLSLExpres
 
 EMaterialGenerateHLSLStatus UMaterialExpressionMaterialFunctionCall::GenerateHLSLExpression(FMaterialHLSLGenerator& Generator, UE::HLSLTree::FScope& Scope, int32 OutputIndex, UE::HLSLTree::FExpression*& OutExpression)
 {
-	if (!MaterialFunction)
+	OutExpression = Generator.GenerateFunctionCall(Scope, MaterialFunction, FunctionInputs, OutputIndex);
+	if (OutExpression)
 	{
-		return Generator.Error(TEXT("Missing function"));
+		return EMaterialGenerateHLSLStatus::Success;
 	}
-
-	TArray<UE::HLSLTree::FExpression*> InputExpressions;
-	InputExpressions.Empty(FunctionInputs.Num());
-
-	for (int32 InputIndex = 0; InputIndex < FunctionInputs.Num(); ++InputIndex)
-	{
-		const FFunctionExpressionInput& Input = FunctionInputs[InputIndex];
-		UE::HLSLTree::FExpression* InputExpression = Input.Input.AcquireHLSLExpression(Generator, Scope);
-		if (!InputExpression)
-		{
-			InputExpression = Generator.AcquireExpression(Scope, Input.ExpressionInput, 0);
-		}
-		check(InputExpression);
-		InputExpressions.Add(InputExpression);
-	}
-
-	UE::HLSLTree::FFunctionCall* FunctionCall = Generator.AcquireFunctionCall(Scope, MaterialFunction, InputExpressions);
-	OutExpression = Generator.GetTree().NewExpression<UE::HLSLTree::FExpressionFunctionOutput>(Scope, FunctionCall, OutputIndex);
-	return EMaterialGenerateHLSLStatus::Success;
+	return EMaterialGenerateHLSLStatus::Error;
 }
 
 EMaterialGenerateHLSLStatus UMaterialExpressionExecBegin::GenerateHLSLStatements(FMaterialHLSLGenerator& Generator, UE::HLSLTree::FScope& Scope)
diff --git a/Engine/Source/Runtime/Engine/Private/Materials/MaterialHLSLGenerator.cpp b/Engine/Source/Runtime/Engine/Private/Materials/MaterialHLSLGenerator.cpp
--- a/Engine/Source/Runtime/Engine/Private/Materials/MaterialHLSLGenerator.cpp
+++ b/Engine/Source/Runtime/Engine/Private/Materials/MaterialHLSLGenerator.cpp
@@ -19,19 +19,12 @@
 FMaterialHLSLGenerator::FMaterialHLSLGenerator(UMaterial* InTargetMaterial, const FMaterialCompileTargetParameters& InCompileTarget, UE::HLSLTree::FTree& InOutTree)
 	: CompileTarget(InCompileTarget)
 	, TargetMaterial(InTargetMaterial)
-	, TargetMaterialFunction(nullptr)
 	, HLSLTree(&InOutTree)
 	, bGeneratedResult(false)
 {
-}
+	FFunctionCallEntry* RootFunctionEntry = new(InOutTree.GetAllocator()) FFunctionCallEntry();
+	FunctionCallStack.Add(RootFunctionEntry);
 
-FMaterialHLSLGenerator::FMaterialHLSLGenerator(UMaterialFunctionInterface* InTargetMaterialFunction, const FMaterialCompileTargetParameters& InCompileTarget, UE::HLSLTree::FTree& InOutTree)
-	: CompileTarget(InCompileTarget)
-	, TargetMaterial(nullptr)
-	, TargetMaterialFunction(InTargetMaterialFunction)
-	, HLSLTree(&InOutTree)
-	, bGeneratedResult(false)
-{
 }
 
 void FMaterialHLSLGenerator::AcquireErrors(TArray<FString>& OutCompileErrors, TArray<UMaterialExpression*>& OutErrorExpressions)
@@ -42,6 +35,8 @@ void FMaterialHLSLGenerator::AcquireErrors(TArray<FString>& OutCompileErrors, TA
 
 bool FMaterialHLSLGenerator::Finalize()
 {
+	check(FunctionCallStack.Num() == 1);
+
 	if (!bGeneratedResult)
 	{
 		Error(TEXT("Missing connection to material output"));
@@ -263,34 +258,6 @@ UE::HLSLTree::FExpression* FMaterialHLSLGenerator::NewCast(UE::HLSLTree::FScope&
 	return Expression;
 }
 
-UE::HLSLTree::FExpression* FMaterialHLSLGenerator::NewFunctionInput(UE::HLSLTree::FScope& Scope, int32 InputIndex, UMaterialExpressionFunctionInput* MaterialFunctionInput)
-{
-	UE::Shader::EValueType ExpressionType = UE::Shader::EValueType::Void;
-	switch (MaterialFunctionInput->InputType)
-	{
-	case FunctionInput_Scalar: ExpressionType = UE::Shader::EValueType::Float1; break;
-	case FunctionInput_Vector2: ExpressionType = UE::Shader::EValueType::Float2; break;
-	case FunctionInput_Vector3: ExpressionType = UE::Shader::EValueType::Float3; break;
-	case FunctionInput_Vector4: ExpressionType = UE::Shader::EValueType::Float4; break;
-	case FunctionInput_Texture2D:
-	case FunctionInput_TextureCube:
-	case FunctionInput_Texture2DArray:
-	case FunctionInput_VolumeTexture:
-	case FunctionInput_StaticBool:
-	case FunctionInput_MaterialAttributes:
-	case FunctionInput_TextureExternal:
-		// TODO
-		break;
-	}
-	check(ExpressionType != UE::Shader::EValueType::Void);
-
-	UE::HLSLTree::FExpressionFunctionInput* Expression = HLSLTree->NewExpression<UE::HLSLTree::FExpressionFunctionInput>(Scope, MaterialFunctionInput->InputName, ExpressionType, InputIndex);
-	const FExpressionKey Key(MaterialFunctionInput, 0);
-	ExpressionMap.Add(Key, Expression);
-
-	return Expression;
-}
-
 UE::HLSLTree::FTextureParameterDeclaration* FMaterialHLSLGenerator::AcquireTextureDeclaration(const UE::HLSLTree::FTextureDescription& Value)
 {
 	FString SamplerTypeError;
@@ -325,35 +292,6 @@ UE::HLSLTree::FTextureParameterDeclaration* FMaterialHLSLGenerator::AcquireTextu
 	return Declaration;
 }
 
-UE::HLSLTree::FFunctionCall* FMaterialHLSLGenerator::AcquireFunctionCall(UE::HLSLTree::FScope& Scope, UMaterialFunctionInterface* Function, TArrayView<UE::HLSLTree::FExpression*> Inputs)
-{
-	FSHA1 Hasher;
-	for (UE::HLSLTree::FExpression* Input : Inputs)
-	{
-		Hasher.Update((uint8*)&Input, sizeof(UE::HLSLTree::FExpression*));
-	}
-	Hasher.Final();
-
-	FFunctionCallKey Key;
-	Key.Function = Function;
-	Hasher.GetHash(Key.InputHash.Hash);
-
-	UE::HLSLTree::FFunctionCall** PrevFunctionCall = FunctionCallMap.Find(Key);
-	UE::HLSLTree::FFunctionCall* FunctionCall = nullptr;
-	if (!PrevFunctionCall)
-	{
-		const FMaterialHLSLTree& FunctionTree = Function->AcquireHLSLTree(*this);
-		FunctionCall = FunctionTree.GenerateFunctionCall(*this, Scope, Inputs);
-		FunctionCallMap.Add(Key, FunctionCall);
-	}
-	else
-	{
-		FunctionCall = *PrevFunctionCall;
-		Scope.UseFunctionCall(FunctionCall);
-	}
-	return FunctionCall;
-}
-
 bool FMaterialHLSLGenerator::GenerateAssignLocal(UE::HLSLTree::FScope& Scope, const FName& LocalName, UE::HLSLTree::FExpression* Value)
 {
 	const FLocalKey Key(&Scope, LocalName);
@@ -400,15 +338,17 @@ UE::HLSLTree::FExpression* FMaterialHLSLGenerator::AcquireLocalValue(UE::HLSLTre
 
 UE::HLSLTree::FExpression* FMaterialHLSLGenerator::AcquireExpression(UE::HLSLTree::FScope& Scope, UMaterialExpression* MaterialExpression, int32 OutputIndex)
 {
+	FFunctionCallEntry* FunctionEntry = FunctionCallStack.Last();
+
 	const FExpressionKey Key(MaterialExpression, OutputIndex);
-	UE::HLSLTree::FExpression** PrevExpression = ExpressionMap.Find(Key);
+	UE::HLSLTree::FExpression** PrevExpression = FunctionEntry->ExpressionMap.Find(Key);
 	UE::HLSLTree::FExpression* Expression = nullptr;
 	if (!PrevExpression)
 	{
 		ExpressionStack.Add(Key);
 		const EMaterialGenerateHLSLStatus Status = MaterialExpression->GenerateHLSLExpression(*this, Scope, OutputIndex, Expression);
 		verify(ExpressionStack.Pop() == Key);
-		ExpressionMap.Add(Key, Expression);
+		FunctionEntry->ExpressionMap.Add(Key, Expression);
 	}
 	else
 	{
@@ -462,6 +402,88 @@ bool FMaterialHLSLGenerator::GenerateStatements(UE::HLSLTree::FScope& Scope, UMa
 	return bResult;
 }
 
+template<typename T>
+TArrayView<T> CopyArrayView(FMemStackBase& Allocator, TArrayView<T> InView)
+{
+	const int32 Num = InView.Num();
+	T* Result = new(Allocator) T[Num];
+	for (int32 i = 0; i < Num; ++i)
+	{
+		Result[i] = InView[i];
+	}
+	return MakeArrayView(Result, Num);
+}
+
+UE::HLSLTree::FExpression* FMaterialHLSLGenerator::GenerateFunctionCall(UE::HLSLTree::FScope& Scope, UMaterialFunctionInterface* Function, TArrayView<const FFunctionExpressionInput> Inputs, int32 OutputIndex)
+{
+	if (!Function)
+	{
+		Error(TEXT("Missing material function"));
+		return nullptr;
+	}
+
+	TArray<FFunctionExpressionInput> FunctionInputs;
+	TArray<FFunctionExpressionOutput> FunctionOutputs;
+	Function->GetInputsAndOutputs(FunctionInputs, FunctionOutputs);
+
+	if (FunctionInputs.Num() != Inputs.Num())
+	{
+		Error(TEXT("Mismatched function inputs"));
+		return nullptr;
+	}
+
+	const UMaterialExpressionFunctionOutput* ExpressionOutput = FunctionOutputs.IsValidIndex(OutputIndex) ? FunctionOutputs[OutputIndex].ExpressionOutput : nullptr;
+	if (!ExpressionOutput)
+	{
+		Error(TEXT("Invalid function output"));
+		return nullptr;
+	}
+
+	FSHA1 Hasher;
+	Hasher.Update((uint8*)&Function, sizeof(UMaterialFunctionInterface*));
+
+	TArray<UE::HLSLTree::FExpression*> InputExpressions;
+	InputExpressions.Reserve(Inputs.Num());
+	for (int32 InputIndex = 0; InputIndex < Inputs.Num(); ++InputIndex)
+	{
+		const FFunctionExpressionInput& Input = Inputs[InputIndex];
+		UE::HLSLTree::FExpression* InputExpression = Input.Input.AcquireHLSLExpression(*this, Scope);
+
+		check(InputExpression);
+		InputExpressions.Add(InputExpression);
+		Hasher.Update((uint8*)&InputExpression, sizeof(UE::HLSLTree::FExpression*));
+	}
+
+	Hasher.Final();
+
+	FSHAHash Hash;
+	Hasher.GetHash(Hash.Hash);
+
+	FFunctionCallEntry** ExistingFunctionCall = FunctionCallMap.Find(Hash);
+	FFunctionCallEntry* FunctionCall = ExistingFunctionCall ? *ExistingFunctionCall : nullptr;
+	if (!FunctionCall)
+	{
+		FunctionCall = new(HLSLTree->GetAllocator()) FFunctionCallEntry();
+		FunctionCall->Function = Function;
+		// Inject the function inputs into the function scope
+		for (int32 InputIndex = 0; InputIndex < Inputs.Num(); ++InputIndex)
+		{
+			const FFunctionExpressionInput& Input = FunctionInputs[InputIndex];
+			UE::HLSLTree::FExpression* InputExpression = InputExpressions[InputIndex];
+			const FExpressionKey ExpressionKey(Input.ExpressionInput, 0);
+			FunctionCall->ExpressionMap.Add(ExpressionKey, InputExpression);
+		}
+
+		FunctionCallMap.Add(Hash, FunctionCall);
+	}
+
+	FunctionCallStack.Add(FunctionCall);
+	UE::HLSLTree::FExpression* Result = ExpressionOutput->A.AcquireHLSLExpression(*this, Scope);
+	verify(FunctionCallStack.Pop() == FunctionCall);
+
+	return Result;
+}
+
 void FMaterialHLSLGenerator::InternalRegisterExpressionData(const FName& Type, const UMaterialExpression* MaterialExpression, void* Data)
 {
 	const FExpressionDataKey Key(Type, MaterialExpression);
diff --git a/Engine/Source/Runtime/Engine/Private/Materials/MaterialHLSLTree.cpp b/Engine/Source/Runtime/Engine/Private/Materials/MaterialHLSLTree.cpp
--- a/Engine/Source/Runtime/Engine/Private/Materials/MaterialHLSLTree.cpp
+++ b/Engine/Source/Runtime/Engine/Private/Materials/MaterialHLSLTree.cpp
@@ -52,46 +52,4 @@ bool FMaterialHLSLTree::InitializeForMaterial(const FMaterialCompileTargetParame
 	return bResult;
 }
 
-bool FMaterialHLSLTree::InitializeForFunction(const FMaterialCompileTargetParameters& InCompilerTarget, UMaterialFunctionInterface* InOutFunction)
-{
-	TArray<FFunctionExpressionInput> FunctionInputs;
-	TArray<FFunctionExpressionOutput> FunctionOutputs;
-	InOutFunction->GetInputsAndOutputs(FunctionInputs, FunctionOutputs);
-
-	HLSLTree = UE::HLSLTree::FTree::Create(Allocator);
-	UE::HLSLTree::FScope& RootScope = HLSLTree->GetRootScope();
-
-	NumFunctionInputs = FunctionInputs.Num();
-	NumFunctionOutputs = FunctionOutputs.Num();
-	FunctionOutputExpressions = New<UE::HLSLTree::FExpression*>(Allocator, NumFunctionOutputs);
-
-	FMaterialHLSLGenerator Generator(InOutFunction, InCompilerTarget, *HLSLTree);
-	for (int32 InputIndex = 0; InputIndex < FunctionInputs.Num(); ++InputIndex)
-	{
-		const FFunctionExpressionInput& Input = FunctionInputs[InputIndex];
-		Generator.NewFunctionInput(RootScope, InputIndex, Input.ExpressionInput);
-	}
-
-	for (int32 OutputIndex = 0; OutputIndex < FunctionOutputs.Num(); ++OutputIndex)
-	{
-		const FFunctionExpressionOutput& Output = FunctionOutputs[OutputIndex];
-		UMaterialExpressionFunctionOutput* ExpressionOutput = Output.ExpressionOutput;
-		FunctionOutputExpressions[OutputIndex] = ExpressionOutput->A.AcquireHLSLExpression(Generator, RootScope);
-	}
-
-	return Generator.Finalize();
-}
-
-UE::HLSLTree::FFunctionCall* FMaterialHLSLTree::GenerateFunctionCall(FMaterialHLSLGenerator& Generator, UE::HLSLTree::FScope& Scope, TArrayView<UE::HLSLTree::FExpression*> Inputs) const
-{
-	check(Inputs.Num() == NumFunctionInputs);
-
-	return HLSLTree->NewFunctionCall(Scope,
-		HLSLTree->GetRootScope(),
-		Inputs.GetData(),
-		FunctionOutputExpressions,
-		NumFunctionInputs,
-		NumFunctionOutputs);
-}
-
 #endif // WITH_EDITOR
diff --git a/Engine/Source/Runtime/Engine/Public/HLSLTree/HLSLTree.h b/Engine/Source/Runtime/Engine/Public/HLSLTree/HLSLTree.h
--- a/Engine/Source/Runtime/Engine/Public/HLSLTree/HLSLTree.h
+++ b/Engine/Source/Runtime/Engine/Public/HLSLTree/HLSLTree.h
@@ -464,6 +464,8 @@ public:
 	static FTree* Create(FMemStackBase& Allocator);
 	static void Destroy(FTree* Tree);
 
+	FMemStackBase& GetAllocator() { return *Allocator; }
+
 	bool EmitHLSL(FEmitContext& Context, FCodeWriter& Writer) const;
 
 	FScope& GetRootScope() const { return *RootScope; }
diff --git a/Engine/Source/Runtime/Engine/Public/MaterialHLSLGenerator.h b/Engine/Source/Runtime/Engine/Public/MaterialHLSLGenerator.h
--- a/Engine/Source/Runtime/Engine/Public/MaterialHLSLGenerator.h
+++ b/Engine/Source/Runtime/Engine/Public/MaterialHLSLGenerator.h
@@ -59,8 +59,6 @@ class FMaterialHLSLGenerator
 public:
 	FMaterialHLSLGenerator(UMaterial* InTargetMaterial, const FMaterialCompileTargetParameters& InCompilerTarget,
 		UE::HLSLTree::FTree& InOutTree);
-	FMaterialHLSLGenerator(UMaterialFunctionInterface* InTargetMaterialFunction, const FMaterialCompileTargetParameters& InCompilerTarget,
-		UE::HLSLTree::FTree& InOutTree);
 
 	const FMaterialCompileTargetParameters& GetCompileTarget() const { return CompileTarget; }
 
@@ -90,16 +88,12 @@ public:
 	UE::HLSLTree::FExpressionSwizzle* NewSwizzle(UE::HLSLTree::FScope& Scope, const UE::HLSLTree::FSwizzleParameters& Params, UE::HLSLTree::FExpression* Input);
 	UE::HLSLTree::FExpression* NewCast(UE::HLSLTree::FScope& Scope, UE::Shader::EValueType Type, UE::HLSLTree::FExpression* Input, UE::HLSLTree::ECastFlags Flags = UE::HLSLTree::ECastFlags::None);
 
-	UE::HLSLTree::FExpression* NewFunctionInput(UE::HLSLTree::FScope& Scope, int32 InputIndex, UMaterialExpressionFunctionInput* MaterialFunctionInput);
-
 	/** Returns a declaration to access the given texture, with no parameter */
 	UE::HLSLTree::FTextureParameterDeclaration* AcquireTextureDeclaration(const UE::HLSLTree::FTextureDescription& Value);
 
 	/** Returns a declaration to access the given texture parameter */
 	UE::HLSLTree::FTextureParameterDeclaration* AcquireTextureParameterDeclaration(const FName& Name, const UE::HLSLTree::FTextureDescription& DefaultValue);
 
-	UE::HLSLTree::FFunctionCall* AcquireFunctionCall(UE::HLSLTree::FScope& Scope, UMaterialFunctionInterface* Function, TArrayView<UE::HLSLTree::FExpression*> Inputs);
-
 	bool GenerateAssignLocal(UE::HLSLTree::FScope& Scope, const FName& LocalName, UE::HLSLTree::FExpression* Value);
 	UE::HLSLTree::FExpression* AcquireLocalValue(UE::HLSLTree::FScope& Scope, const FName& LocalName);
 
@@ -112,6 +106,8 @@ public:
 	UE::HLSLTree::FTextureParameterDeclaration* AcquireTextureDeclaration(UE::HLSLTree::FScope& Scope, UMaterialExpression* MaterialExpression, int32 OutputIndex);
 	bool GenerateStatements(UE::HLSLTree::FScope& Scope, UMaterialExpression* MaterialExpression);
 
+	UE::HLSLTree::FExpression* GenerateFunctionCall(UE::HLSLTree::FScope& Scope, UMaterialFunctionInterface* Function, TArrayView<const FFunctionExpressionInput> Inputs, int32 OutputIndex);
+
 	template<typename T, typename... ArgTypes>
 	T* NewExpressionData(const UMaterialExpression* MaterialExpression, ArgTypes... Args)
 	{
@@ -158,22 +154,6 @@ private:
 		}
 	};
 
-	struct FFunctionCallKey
-	{
-		UMaterialFunctionInterface* Function;
-		FSHAHash InputHash;
-
-		friend inline uint32 GetTypeHash(const FFunctionCallKey& Value)
-		{
-			return HashCombine(GetTypeHash(Value.Function), GetTypeHash(Value.InputHash));
-		}
-
-		friend inline bool operator==(const FFunctionCallKey& Lhs, const FFunctionCallKey& Rhs)
-		{
-			return Lhs.Function == Rhs.Function && Lhs.InputHash == Rhs.InputHash;
-		}
-	};
-
 	struct FLocalKey
 	{
 		FLocalKey(UE::HLSLTree::FScope* InScope, const FName& InName) : Scope(InScope), Name(InName) {}
@@ -210,6 +190,12 @@ private:
 		}
 	};
 
+	struct FFunctionCallEntry
+	{
+		UMaterialFunctionInterface* Function = nullptr;
+		TMap<FExpressionKey, UE::HLSLTree::FExpression*> ExpressionMap;
+	};
+
 	struct FStatementEntry
 	{
 		UE::HLSLTree::FScope* PreviousScope[MaxNumPreviousScopes];
@@ -223,19 +209,18 @@ private:
 
 	const FMaterialCompileTargetParameters& CompileTarget;
 	UMaterial* TargetMaterial;
-	UMaterialFunctionInterface* TargetMaterialFunction;
 
 	UE::HLSLTree::FTree* HLSLTree;
 	TArray<FExpressionKey> ExpressionStack;
+	TArray<FFunctionCallEntry*> FunctionCallStack;
 	TArray<UE::HLSLTree::FScope*> JoinedScopeStack;
 	TArray<UE::HLSLTree::FExpressionLocalPHI*> PHIExpressions;
 	TArray<FString> CompileErrors;
 	TArray<UMaterialExpression*> ErrorExpressions;
 	TMap<UE::HLSLTree::FTextureDescription, UE::HLSLTree::FTextureParameterDeclaration*> TextureDeclarationMap;
 	TMap<FName, UE::HLSLTree::FTextureParameterDeclaration*> TextureParameterDeclarationMap;
-	TMap<FFunctionCallKey, UE::HLSLTree::FFunctionCall*> FunctionCallMap;
+	TMap<FSHAHash, FFunctionCallEntry*> FunctionCallMap;
 	TMap<FLocalKey, UE::HLSLTree::FExpression*> LocalMap;
-	TMap<FExpressionKey, UE::HLSLTree::FExpression*> ExpressionMap;
 	TMap<UMaterialExpression*, FStatementEntry> StatementMap;
 	TMap<FExpressionDataKey, void*> ExpressionDataMap;
 	bool bGeneratedResult;
diff --git a/Engine/Source/Runtime/Engine/Public/MaterialHLSLTree.h b/Engine/Source/Runtime/Engine/Public/MaterialHLSLTree.h
--- a/Engine/Source/Runtime/Engine/Public/MaterialHLSLTree.h
+++ b/Engine/Source/Runtime/Engine/Public/MaterialHLSLTree.h
@@ -26,22 +26,11 @@ class FMaterialHLSLTree
 public:
 	FMaterialHLSLTree();
 	~FMaterialHLSLTree();
-
 	bool InitializeForMaterial(const FMaterialCompileTargetParameters& InCompilerTarget, FMaterial& InOutMaterial);
-
-	bool InitializeForFunction(const FMaterialCompileTargetParameters& InCompilerTarget, UMaterialFunctionInterface* InOutFunction);
-
 	UE::HLSLTree::FTree& GetTree() const { return *HLSLTree; }
-	int32 GetNumFunctionInputs() const { return NumFunctionInputs; }
-	int32 GetNumFunctionOutputs() const { return NumFunctionOutputs; }
-
-	UE::HLSLTree::FFunctionCall* GenerateFunctionCall(FMaterialHLSLGenerator& Generator, UE::HLSLTree::FScope& Scope, TArrayView<UE::HLSLTree::FExpression*> Inputs) const;
 private:
 	FMemStackBase Allocator;
 	UE::HLSLTree::FTree* HLSLTree = nullptr;
-	UE::HLSLTree::FExpression** FunctionOutputExpressions = nullptr;
-	int32 NumFunctionInputs = 0;
-	int32 NumFunctionOutputs = 0;
 };
 
 #endif // WITH_EDITOR
