Update FMaterial::GetMaterialExpressionSource() to use the new HLSL translator, if it's enabled

diff --git a/Engine/Source/Runtime/Engine/Private/Materials/MaterialShared.cpp b/Engine/Source/Runtime/Engine/Private/Materials/MaterialShared.cpp
--- a/Engine/Source/Runtime/Engine/Private/Materials/MaterialShared.cpp
+++ b/Engine/Source/Runtime/Engine/Private/Materials/MaterialShared.cpp
@@ -2373,6 +2373,23 @@ bool FMaterial::Translate_New(const FMaterialShaderMapId& ShaderMapId,
 #endif
 }
 
+bool FMaterial::Translate(const FMaterialShaderMapId& InShaderMapId,
+	const FStaticParameterSet& InStaticParameters,
+	EShaderPlatform InPlatform,
+	const ITargetPlatform* InTargetPlatform,
+	FMaterialCompilationOutput& OutCompilationOutput,
+	TRefCountPtr<FSharedShaderCompilerEnvironment>& OutMaterialEnvironment)
+{
+	if (InShaderMapId.bUsingNewHLSLGenerator)
+	{
+		return Translate_New(InShaderMapId, InStaticParameters, InPlatform, InTargetPlatform, OutCompilationOutput, OutMaterialEnvironment);
+	}
+	else
+	{
+		return Translate_Legacy(InShaderMapId, InStaticParameters, InPlatform, InTargetPlatform, OutCompilationOutput, OutMaterialEnvironment);
+	}
+}
+
 /**
 * Compiles this material for Platform
 *
@@ -2404,14 +2421,7 @@ bool FMaterial::BeginCompileShaderMap(
 	// Generate the material shader code.
 	FMaterialCompilationOutput NewCompilationOutput;
 	TRefCountPtr<FSharedShaderCompilerEnvironment> MaterialEnvironment;
-	if (ShaderMapId.bUsingNewHLSLGenerator)
-	{
-		bSuccess = Translate_New(ShaderMapId, StaticParameterSet, Platform, TargetPlatform, NewCompilationOutput, MaterialEnvironment);
-	}
-	else
-	{
-		bSuccess = Translate_Legacy(ShaderMapId, StaticParameterSet, Platform, TargetPlatform, NewCompilationOutput, MaterialEnvironment);
-	}
+	bSuccess = Translate(ShaderMapId, StaticParameterSet, Platform, TargetPlatform, NewCompilationOutput, MaterialEnvironment);
 
 	if(bSuccess)
 	{
@@ -3753,29 +3763,28 @@ void FMaterial::GetReferencedTexturesHash(EShaderPlatform Platform, FSHAHash& Ou
 bool FMaterial::GetMaterialExpressionSource( FString& OutSource )
 {
 #if WITH_EDITORONLY_DATA
-	class FViewSourceMaterialTranslator : public FHLSLMaterialTranslator
-	{
-	public:
-		FViewSourceMaterialTranslator(FMaterial* InMaterial,FMaterialCompilationOutput& InMaterialCompilationOutput,const FStaticParameterSet& StaticParameters,EShaderPlatform InPlatform,EMaterialQualityLevel::Type InQualityLevel,ERHIFeatureLevel::Type InFeatureLevel)
-		:	FHLSLMaterialTranslator(InMaterial,InMaterialCompilationOutput,StaticParameters,InPlatform,InQualityLevel,InFeatureLevel)
-		{}
-	};
+	const EShaderPlatform ShaderPlatform = GMaxRHIShaderPlatform;
+	const ITargetPlatform* TargetPlatform = nullptr;
 
-	FMaterialCompilationOutput TempOutput;
-	FMaterialShaderMapId ShaderMapID;
-	GetShaderMapId(GMaxRHIShaderPlatform, nullptr, ShaderMapID);
-	FStaticParameterSet StaticParamSet;
-	GetStaticParameterSet(GMaxRHIShaderPlatform, StaticParamSet);
+	FMaterialShaderMapId ShaderMapId;
+	GetShaderMapId(ShaderPlatform, TargetPlatform, ShaderMapId);
+	FStaticParameterSet StaticParameterSet;
+	GetStaticParameterSet(ShaderPlatform, StaticParameterSet);
 
-	FViewSourceMaterialTranslator MaterialTranslator(this, TempOutput, StaticParamSet, GMaxRHIShaderPlatform, GetQualityLevel(), GetFeatureLevel());
-	bool bSuccess = MaterialTranslator.Translate();
+	FMaterialCompilationOutput NewCompilationOutput;
+	TRefCountPtr<FSharedShaderCompilerEnvironment> MaterialEnvironment;
+	const bool bSuccess = Translate(ShaderMapId, StaticParameterSet, ShaderPlatform, TargetPlatform, NewCompilationOutput, MaterialEnvironment);
 
-	if( bSuccess )
+	if (bSuccess)
 	{
-		// Generate the HLSL
-		OutSource = MaterialTranslator.GetMaterialShaderCode();
+		const FString* Source = MaterialEnvironment->IncludeVirtualPathToContentsMap.Find(TEXT("/Engine/Generated/Material.ush"));
+		if (Source)
+		{
+			OutSource = *Source;
+			return true;
+		}
 	}
-	return bSuccess;
+	return false;
 #else
 	UE_LOG(LogMaterial, Fatal,TEXT("Not supported."));
 	return false;
diff --git a/Engine/Source/Runtime/Engine/Public/MaterialShared.h b/Engine/Source/Runtime/Engine/Public/MaterialShared.h
--- a/Engine/Source/Runtime/Engine/Public/MaterialShared.h
+++ b/Engine/Source/Runtime/Engine/Public/MaterialShared.h
@@ -2237,6 +2237,13 @@ private:
 		FMaterialCompilationOutput& OutCompilationOutput,
 		TRefCountPtr<FSharedShaderCompilerEnvironment>& OutMaterialEnvironment);
 
+	bool Translate(const FMaterialShaderMapId& InShaderMapId,
+		const FStaticParameterSet& InStaticParameters,
+		EShaderPlatform InPlatform,
+		const ITargetPlatform* InTargetPlatform,
+		FMaterialCompilationOutput& OutCompilationOutput,
+		TRefCountPtr<FSharedShaderCompilerEnvironment>& OutMaterialEnvironment);
+
 	/** Populates OutEnvironment with defines needed to compile shaders for this material. */
 	void SetupMaterialEnvironment(
 		EShaderPlatform Platform,
