Fix non-editor compile

diff --git a/Engine/Source/Runtime/Engine/Private/Materials/MaterialShared.cpp b/Engine/Source/Runtime/Engine/Private/Materials/MaterialShared.cpp
--- a/Engine/Source/Runtime/Engine/Private/Materials/MaterialShared.cpp
+++ b/Engine/Source/Runtime/Engine/Private/Materials/MaterialShared.cpp
@@ -2380,6 +2380,7 @@ bool FMaterial::Translate(const FMaterialShaderMapId& InShaderMapId,
 	FMaterialCompilationOutput& OutCompilationOutput,
 	TRefCountPtr<FSharedShaderCompilerEnvironment>& OutMaterialEnvironment)
 {
+#if WITH_EDITOR
 	if (InShaderMapId.bUsingNewHLSLGenerator)
 	{
 		return Translate_New(InShaderMapId, InStaticParameters, InPlatform, InTargetPlatform, OutCompilationOutput, OutMaterialEnvironment);
@@ -2388,6 +2389,10 @@ bool FMaterial::Translate(const FMaterialShaderMapId& InShaderMapId,
 	{
 		return Translate_Legacy(InShaderMapId, InStaticParameters, InPlatform, InTargetPlatform, OutCompilationOutput, OutMaterialEnvironment);
 	}
+#else
+	checkNoEntry();
+	return false;
+#endif
 }
 
 /**
@@ -3777,10 +3782,10 @@ bool FMaterial::GetMaterialExpressionSource( FString& OutSource )
 
 	if (bSuccess)
 	{
-		const FString* Source = MaterialEnvironment->IncludeVirtualPathToContentsMap.Find(TEXT("/Engine/Generated/Material.ush"));
+		FString* Source = MaterialEnvironment->IncludeVirtualPathToContentsMap.Find(TEXT("/Engine/Generated/Material.ush"));
 		if (Source)
 		{
-			OutSource = *Source;
+			OutSource = MoveTemp(*Source);
 			return true;
 		}
 	}
