Clean up some HLSLTree function stuff

diff --git a/Engine/Source/Runtime/Engine/Private/HLSLTree/HLSLTree.cpp b/Engine/Source/Runtime/Engine/Private/HLSLTree/HLSLTree.cpp
--- a/Engine/Source/Runtime/Engine/Private/HLSLTree/HLSLTree.cpp
+++ b/Engine/Source/Runtime/Engine/Private/HLSLTree/HLSLTree.cpp
@@ -151,14 +151,10 @@ void UE::HLSLTree::FCodeWriter::WriteConstant(const Shader::FValue& Value)
 
 UE::HLSLTree::FEmitContext::FEmitContext()
 {
-	FunctionStack.AddDefaulted();
 }
 
 UE::HLSLTree::FEmitContext::~FEmitContext()
 {
-	// Make sure we have only the root stack entry
-	check(FunctionStack.Num() == 1);
-
 	for (Shader::FPreshaderData* Preshader : TempPreshaders)
 	{
 		delete Preshader;
@@ -220,8 +216,7 @@ const UE::HLSLTree::FEmitValue* UE::HLSLTree::FEmitContext::AcquireValue(FExpres
 		return nullptr;
 	}
 
-	FFunctionStackEntry& FunctionStackEntry = FunctionStack.Last();
-	FDeclarationEntry** FoundEntry = FunctionStackEntry.DeclarationMap.Find(Expression);
+	FDeclarationEntry** FoundEntry = DeclarationMap.Find(Expression);
 	FDeclarationEntry* Entry = FoundEntry ? *FoundEntry : nullptr;
 	if (!Entry)
 	{
@@ -246,7 +241,7 @@ const UE::HLSLTree::FEmitValue* UE::HLSLTree::FEmitContext::AcquireValue(FExpres
 			check(EmitResult.Type != Shader::EValueType::Void);
 
 			Entry = new(*Allocator) FDeclarationEntry();
-			FunctionStackEntry.DeclarationMap.Add(Expression, Entry);
+			DeclarationMap.Add(Expression, Entry);
 			Entry->Value.ExpressionType = EmitResult.Type;
 			Entry->Value.EvaluationType = EmitResult.EvaluationType;
 			if (EmitResult.EvaluationType == EExpressionEvaluationType::Constant)
@@ -310,62 +305,6 @@ const UE::HLSLTree::FEmitValue* UE::HLSLTree::FEmitContext::AcquireValue(FExpres
 	return Entry ? &Entry->Value : nullptr;
 }
 
-const UE::HLSLTree::FEmitValue* UE::HLSLTree::FEmitContext::AcquireValue(FFunctionCall* FunctionCall, int32 OutputIndex)
-{
-	check(FunctionCall);
-
-	FFunctionCallEntry*& Entry = FunctionStack.Last().FunctionCallMap.FindOrAdd(FunctionCall);
-	if (!Entry)
-	{
-		FEmitScope* EmitScope = FindScope(*FunctionCall->ParentScope);
-		check(EmitScope);
-
-		{
-			FFunctionStackEntry& StackEntry = FunctionStack.AddDefaulted_GetRef();
-			StackEntry.FunctionCall = FunctionCall;
-		}
-
-		// Link the function scope to our parent scope
-		ScopeMap.Add(FunctionCall->FunctionScope, EmitScope);
-		FunctionCall->FunctionScope->EmitHLSL(*this, *EmitScope);
-
-		// Assign function outputs
-		FEmitValue* OutputValues = new(*Allocator) FEmitValue[FunctionCall->NumOutputs];
-		for (int32 i = 0; i < FunctionCall->NumOutputs; ++i)
-		{
-			FExpression* OutputExpression = FunctionCall->Outputs[i];
-
-			const FEmitValue* OutputValue = AcquireValue(OutputExpression);
-			FEmitValue& Output = OutputValues[i];
-			if (OutputValue)
-			{
-				Output = *OutputValue;
-			}
-			else
-			{
-				Output.EvaluationType = EExpressionEvaluationType::Constant;
-				Output.ExpressionType = Shader::EValueType::Float1;
-				Output.ConstantValue = 0.0f;
-			}
-		}
-
-		verify(ScopeMap.Remove(FunctionCall->FunctionScope) == 1);
-
-		{
-			const FFunctionStackEntry PoppedStackEntry = FunctionStack.Pop(false);
-			check(PoppedStackEntry.FunctionCall == FunctionCall);
-		}
-
-		Entry = new(*Allocator) FFunctionCallEntry();
-		Entry->OutputValues = OutputValues;
-		Entry->NumOutputs = FunctionCall->NumOutputs;
-	}
-
-	check(Entry->NumOutputs == FunctionCall->NumOutputs);
-	check(OutputIndex >= 0 && OutputIndex < Entry->NumOutputs);
-	return &Entry->OutputValues[OutputIndex];
-}
-
 const TCHAR* UE::HLSLTree::FEmitContext::AcquireLocalDeclarationCode()
 {
 	return AllocateStringf(*Allocator, TEXT("Local%d"), NumExpressionLocals++);
@@ -709,20 +648,6 @@ UE::HLSLTree::ENodeVisitResult UE::HLSLTree::FTextureParameterDeclaration::Visit
 	return Visitor.OnTextureParameterDeclaration(*this);
 }
 
-UE::HLSLTree::ENodeVisitResult UE::HLSLTree::FFunctionCall::Visit(FNodeVisitor& Visitor)
-{
-	const ENodeVisitResult Result = Visitor.OnFunctionCall(*this);
-	if (ShouldVisitDependentNodes(Result))
-	{
-		// Don't visit the function scope, as that is from a different tree
-		for (int32 i = 0; i < NumInputs; ++i)
-		{
-			Visitor.VisitNode(Inputs[i]);
-		}
-	}
-	return Result;
-}
-
 UE::HLSLTree::ENodeVisitResult UE::HLSLTree::FScope::Visit(FNodeVisitor& Visitor)
 {
 	const ENodeVisitResult Result = Visitor.OnScope(*this);
@@ -787,12 +712,6 @@ public:
 		return ENodeVisitResult::VisitDependentNodes;
 	}
 
-	virtual ENodeVisitResult OnFunctionCall(FFunctionCall& InFunctionCall) override
-	{
-		InFunctionCall.ParentScope = FScope::FindSharedParent(Scope, InFunctionCall.ParentScope);
-		return ENodeVisitResult::VisitDependentNodes;
-	}
-
 	FScope* Scope;
 };
 } // namespace HLSLTree
@@ -805,12 +724,6 @@ void UE::HLSLTree::FScope::UseExpression(FExpression* Expression)
 	Visitor.VisitNode(Expression);
 }
 
-void UE::HLSLTree::FScope::UseFunctionCall(FFunctionCall* FunctionCall)
-{
-	FNodeVisitor_MoveToScope Visitor(this);
-	Visitor.VisitNode(FunctionCall);
-}
-
 UE::HLSLTree::FTree* UE::HLSLTree::FTree::Create(FMemStackBase& Allocator)
 {
 	FTree* Tree = new(Allocator) FTree();
@@ -953,24 +866,3 @@ UE::HLSLTree::FTextureParameterDeclaration* UE::HLSLTree::FTree::NewTextureParam
 	return Declaration;
 }
 
-UE::HLSLTree::FFunctionCall* UE::HLSLTree::FTree::NewFunctionCall(FScope& Scope,
-	const FScope& InFunctionScope,
-	FExpression* const* InInputs,
-	FExpression* const* InOutputs,
-	int32 InNumInputs,
-	int32 InNumOutputs)
-{
-	FExpression** Inputs = New<FExpression*>(*Allocator, InNumInputs);
-	FExpression** Outputs = New<FExpression*>(*Allocator, InNumOutputs);
-	FMemory::Memcpy(Inputs, InInputs, InNumInputs * sizeof(FExpression*));
-	FMemory::Memcpy(Outputs, InOutputs, InNumOutputs * sizeof(FExpression*));
-
-	FFunctionCall* FunctionCall = NewNode<FFunctionCall>();
-	FunctionCall->ParentScope = &Scope;
-	FunctionCall->FunctionScope = &InFunctionScope;
-	FunctionCall->Inputs = Inputs;
-	FunctionCall->Outputs = Outputs;
-	FunctionCall->NumInputs = InNumInputs;
-	FunctionCall->NumOutputs = InNumOutputs;
-	return FunctionCall;
-}
diff --git a/Engine/Source/Runtime/Engine/Private/HLSLTree/HLSLTreeCommon.cpp b/Engine/Source/Runtime/Engine/Private/HLSLTree/HLSLTreeCommon.cpp
--- a/Engine/Source/Runtime/Engine/Private/HLSLTree/HLSLTreeCommon.cpp
+++ b/Engine/Source/Runtime/Engine/Private/HLSLTree/HLSLTreeCommon.cpp
@@ -548,29 +548,6 @@ bool UE::HLSLTree::FExpressionReflectionVector::EmitCode(FEmitContext& Context,
 	return true;
 }
 
-bool UE::HLSLTree::FExpressionFunctionInput::EmitCode(FEmitContext& Context, FExpressionEmitResult& OutResult) const
-{
-	const FEmitContext::FFunctionStackEntry& StackEntry = Context.FunctionStack.Last();
-
-	check(InputIndex >= 0 && InputIndex < StackEntry.FunctionCall->NumInputs);
-	FExpression* InputExpression = StackEntry.FunctionCall->Inputs[InputIndex];
-	const FEmitValue* Value = Context.AcquireValue(InputExpression);
-	OutResult.ForwardValue(Context, Value);
-	return true;
-}
-
-bool UE::HLSLTree::FExpressionFunctionOutput::EmitCode(FEmitContext& Context, FExpressionEmitResult& OutResult) const
-{
-	const FEmitValue* Value = Context.AcquireValue(FunctionCall, OutputIndex);
-	if (!Value)
-	{
-		return false;
-	}
-
-	OutResult.ForwardValue(Context, Value);
-	return true;
-}
-
 bool UE::HLSLTree::FStatementBreak::EmitHLSL(FEmitContext& Context) const
 {
 	Context.WriteScopeTerminatorf(TEXT("break;"));
diff --git a/Engine/Source/Runtime/Engine/Public/HLSLTree/HLSLTree.h b/Engine/Source/Runtime/Engine/Public/HLSLTree/HLSLTree.h
--- a/Engine/Source/Runtime/Engine/Public/HLSLTree/HLSLTree.h
+++ b/Engine/Source/Runtime/Engine/Public/HLSLTree/HLSLTree.h
@@ -163,9 +163,6 @@ public:
 	/** Returns a value that references the given expression */
 	const FEmitValue* AcquireValue(FExpression* Expression);
 
-	/** Returns value that references the given local function output */
-	const FEmitValue* AcquireValue(FFunctionCall* FunctionCall, int32 OutputIndex);
-
 	/** Get a unique local variable name */
 	const TCHAR* AcquireLocalDeclarationCode();
 
@@ -253,18 +250,10 @@ public:
 		int32 NumOutputs;
 	};
 
-	struct FFunctionStackEntry
-	{
-		FFunctionCall* FunctionCall = nullptr;
-
-		TMap<FNode*, FDeclarationEntry*> DeclarationMap;
-		TMap<FFunctionCall*, FFunctionCallEntry*> FunctionCallMap;
-	};
-
 	TArray<FEmitScope*> ScopeStack;
 	TMap<const FScope*, FEmitScope*> ScopeMap;
 
-	TArray<FFunctionStackEntry> FunctionStack;
+	TMap<FNode*, FDeclarationEntry*> DeclarationMap;
 	TArray<Shader::FPreshaderData*> TempPreshaders;
 	TSet<const FExpression*> PendingEmitValueExpressions;
 	FMemStackBase* Allocator = nullptr;
@@ -315,7 +304,6 @@ public:
 	virtual ENodeVisitResult OnStatement(FStatement& Statement) { return ENodeVisitResult::VisitDependentNodes; }
 	virtual ENodeVisitResult OnExpression(FExpression& Expression) { return ENodeVisitResult::VisitDependentNodes; }
 	virtual ENodeVisitResult OnTextureParameterDeclaration(FTextureParameterDeclaration& Declaration) { return ENodeVisitResult::VisitDependentNodes; }
-	virtual ENodeVisitResult OnFunctionCall(FFunctionCall& FunctionCall) { return ENodeVisitResult::VisitDependentNodes; }
 };
 
 /** Root class of the HLSL AST */
@@ -395,29 +383,6 @@ public:
 	FTextureDescription Description;
 };
 
-/**
- * Represents a call to an HLSL function from a separate tree
- */
-class FFunctionCall final : public FNode
-{
-public:
-	virtual ENodeVisitResult Visit(FNodeVisitor& Visitor) override;
-
-	FScope* ParentScope = nullptr;
-
-	/** Root scope of the function to call. Note that this scope will be from a separate (external) tree */
-	const FScope* FunctionScope;
-
-	/** Outputs are expressions from the function's scope */
-	FExpression* const* Outputs;
-
-	/** Inputs are connected to the calling scope */
-	FExpression* const* Inputs;
-	
-	int32 NumInputs;
-	int32 NumOutputs;
-};
-
 /**
  * Represents an HLSL scope.  A scope contains a single statement, along with any expressions required by that statement
  */
@@ -441,7 +406,6 @@ public:
 
 	void AddPreviousScope(FScope& Scope);
 
-	void UseFunctionCall(FFunctionCall* FunctionCall);
 	void UseExpression(FExpression* Expression);
 
 private:
@@ -489,13 +453,6 @@ public:
 	FScope* NewScope(FScope& Scope);
 	FTextureParameterDeclaration* NewTextureParameterDeclaration(const FName& Name, const FTextureDescription& DefaultValue);
 
-	FFunctionCall* NewFunctionCall(FScope& Scope,
-		const FScope& FunctionScope,
-		FExpression* const* Inputs,
-		FExpression* const* Outputs,
-		int32 NumInputs,
-		int32 NumOutputs);
-
 private:
 	template<typename T, typename... ArgTypes>
 	inline T* NewNode(ArgTypes&&... Args)
diff --git a/Engine/Source/Runtime/Engine/Public/HLSLTree/HLSLTreeCommon.h b/Engine/Source/Runtime/Engine/Public/HLSLTree/HLSLTreeCommon.h
--- a/Engine/Source/Runtime/Engine/Public/HLSLTree/HLSLTreeCommon.h
+++ b/Engine/Source/Runtime/Engine/Public/HLSLTree/HLSLTreeCommon.h
@@ -299,46 +299,6 @@ public:
 	virtual bool EmitCode(FEmitContext& Context, FExpressionEmitResult& OutResult) const override;
 };
 
-class FExpressionFunctionInput : public FExpression
-{
-public:
-	FExpressionFunctionInput(const FName& InName, Shader::EValueType InType, int32 InIndex)
-		: Name(InName), Type(InType), InputIndex(InIndex)
-	{}
-
-	virtual bool EmitCode(FEmitContext& Context, FExpressionEmitResult& OutResult) const override;
-
-	FName Name;
-	Shader::EValueType Type;
-	int32 InputIndex;
-};
-
-class FExpressionFunctionOutput : public FExpression
-{
-public:
-	FExpressionFunctionOutput(FFunctionCall* InFunctionCall, int32 InIndex)
-		: FunctionCall(InFunctionCall)
-		, OutputIndex(InIndex)
-	{
-		check(InIndex >= 0 && InIndex < InFunctionCall->NumOutputs);
-	}
-
-	FFunctionCall* FunctionCall;
-	int32 OutputIndex;
-
-	virtual ENodeVisitResult Visit(FNodeVisitor& Visitor) override
-	{
-		const ENodeVisitResult Result = FExpression::Visit(Visitor);
-		if (ShouldVisitDependentNodes(Result))
-		{
-			Visitor.VisitNode(FunctionCall);
-		}
-		return Result;
-	}
-
-	virtual bool EmitCode(FEmitContext& Context, FExpressionEmitResult& OutResult) const override;
-};
-
 class FStatementReturn : public FStatement
 {
 public:
diff --git a/Engine/Source/Runtime/Engine/Public/HLSLTree/HLSLTreeTypes.h b/Engine/Source/Runtime/Engine/Public/HLSLTree/HLSLTreeTypes.h
--- a/Engine/Source/Runtime/Engine/Public/HLSLTree/HLSLTreeTypes.h
+++ b/Engine/Source/Runtime/Engine/Public/HLSLTree/HLSLTreeTypes.h
@@ -16,7 +16,6 @@ class FExpression;
 class FTextureParameterDeclaration;
 class FStatement;
 class FScope;
-class FFunctionCall;
 class FTree;
 class FCodeWriter;
 
diff --git a/Engine/Source/Runtime/Engine/Public/MaterialHLSLTree.h b/Engine/Source/Runtime/Engine/Public/MaterialHLSLTree.h
--- a/Engine/Source/Runtime/Engine/Public/MaterialHLSLTree.h
+++ b/Engine/Source/Runtime/Engine/Public/MaterialHLSLTree.h
@@ -17,7 +17,6 @@ namespace HLSLTree
 class FTree;
 class FScope;
 class FExpression;
-class FFunctionCall;
 }
 }
 
