Remove HLSLTree FCodeWriter, replace with FStringBuilderBase

diff --git a/Engine/Source/Runtime/Engine/Private/HLSLTree/HLSLTree.cpp b/Engine/Source/Runtime/Engine/Private/HLSLTree/HLSLTree.cpp
--- a/Engine/Source/Runtime/Engine/Private/HLSLTree/HLSLTree.cpp
+++ b/Engine/Source/Runtime/Engine/Private/HLSLTree/HLSLTree.cpp
@@ -4,14 +4,6 @@
 #include "Misc/MemStack.h"
 #include "MaterialShared.h" // TODO - split preshader out into its own module
 
-class FLocalHLSLCodeWriter : public UE::HLSLTree::FCodeWriter
-{
-public:
-	FLocalHLSLCodeWriter() : UE::HLSLTree::FCodeWriter(&LocalStringBuilder) {}
-
-	TStringBuilder<2048> LocalStringBuilder;
-};
-
 UE::HLSLTree::EExpressionEvaluationType UE::HLSLTree::CombineEvaluationTypes(EExpressionEvaluationType Lhs, EExpressionEvaluationType Rhs)
 {
 	if (Lhs == EExpressionEvaluationType::Constant && Rhs == EExpressionEvaluationType::Constant)
@@ -49,106 +41,13 @@ const TCHAR* AllocateString(FMemStackBase& Allocator, const FStringBuilderBase&
 	return Result;
 }
 
-UE::HLSLTree::FCodeWriter* UE::HLSLTree::FCodeWriter::Create(FMemStackBase& Allocator)
-{
-	static const int32 InitialBufferSize = 4 * 1024;
-	TCHAR* Buffer = New<TCHAR>(Allocator, InitialBufferSize);
-	FStringBuilderBase* LocalStringBuilder = new(Allocator) TStringBuilderBase<TCHAR>(Buffer, InitialBufferSize);
-	return new(Allocator) FCodeWriter(LocalStringBuilder);
-}
-
-FSHAHash UE::HLSLTree::FCodeWriter::GetCodeHash() const
+FSHAHash HashString(const FStringBuilderBase& StringBuilder)
 {
 	FSHAHash Hash;
-	FSHA1::HashBuffer(StringBuilder->GetData(), StringBuilder->Len() * sizeof(TCHAR), Hash.Hash);
+	FSHA1::HashBuffer(StringBuilder.GetData(), StringBuilder.Len() * sizeof(TCHAR), Hash.Hash);
 	return Hash;
 }
 
-void UE::HLSLTree::FCodeWriter::IncreaseIndent()
-{
-	++IndentLevel;
-}
-
-void UE::HLSLTree::FCodeWriter::DecreaseIndent()
-{
-	check(IndentLevel > 0);
-	--IndentLevel;
-}
-
-void UE::HLSLTree::FCodeWriter::WriteIndent()
-{
-	for (int32 i = 0; i < IndentLevel; ++i)
-	{
-		StringBuilder->Append(TCHAR('\t'));
-	}
-}
-
-void UE::HLSLTree::FCodeWriter::Reset()
-{
-	StringBuilder->Reset();
-}
-
-void UE::HLSLTree::FCodeWriter::Append(const FCodeWriter& InWriter)
-{
-	const int32 Length = InWriter.StringBuilder->Len();
-	if (Length > 0)
-	{
-		StringBuilder->Append(InWriter.StringBuilder->GetData(), Length);
-	}
-}
-
-
-void UE::HLSLTree::FCodeWriter::WriteConstant(const Shader::FValue& Value)
-{
-	auto ToString = [](uint8 v)
-	{
-		return v ? TEXT("true") : TEXT("false");
-	};
-
-	switch (Value.GetType())
-	{
-	case Shader::EValueType::Float1:
-		Writef(TEXT("%0.8f"), Value.Component[0].Float);
-		break;
-	case Shader::EValueType::Float2:
-		Writef(TEXT("float2(%0.8f, %0.8f)"), Value.Component[0].Float, Value.Component[1].Float);
-		break;
-	case Shader::EValueType::Float3:
-		Writef(TEXT("float3(%0.8f, %0.8f, %0.8f)"), Value.Component[0].Float, Value.Component[1].Float, Value.Component[2].Float);
-		break;
-	case Shader::EValueType::Float4:
-		Writef(TEXT("float4(%0.8f, %0.8f, %0.8f, %0.8f)"), Value.Component[0].Float, Value.Component[1].Float, Value.Component[2].Float, Value.Component[3].Float);
-		break;
-	case Shader::EValueType::Int1:
-		Writef(TEXT("%d"), Value.Component[0].Int);
-		break;
-	case Shader::EValueType::Int2:
-		Writef(TEXT("int2(%d, %d)"), Value.Component[0].Int, Value.Component[1].Int);
-		break;
-	case Shader::EValueType::Int3:
-		Writef(TEXT("int3(%d, %d, %d)"), Value.Component[0].Int, Value.Component[1].Int, Value.Component[2].Int);
-		break;
-	case Shader::EValueType::Int4:
-		Writef(TEXT("int4(%d, %d, %d, %d)"), Value.Component[0].Int, Value.Component[1].Int, Value.Component[2].Int, Value.Component[3].Int);
-		break;
-	case Shader::EValueType::Bool1:
-		Writef(TEXT("%s"), ToString(Value.Component[0].Bool));
-		break;
-	case Shader::EValueType::Bool2:
-		Writef(TEXT("bool2(%s, %s)"), ToString(Value.Component[0].Bool), ToString(Value.Component[1].Bool));
-		break;
-	case Shader::EValueType::Bool3:
-		Writef(TEXT("bool3(%s, %s, %s)"), ToString(Value.Component[0].Bool), ToString(Value.Component[1].Bool), ToString(Value.Component[2].Bool));
-		break;
-	case Shader::EValueType::Bool4:
-		Writef(TEXT("bool4(%s, %s, %s, %s)"), ToString(Value.Component[0].Bool), ToString(Value.Component[1].Bool), ToString(Value.Component[2].Bool), ToString(Value.Component[3].Bool));
-		break;
-	default:
-		checkNoEntry();
-		break;
-	}
-}
-
 UE::HLSLTree::FEmitContext::FEmitContext()
 {
 }
@@ -201,7 +100,7 @@ void UE::HLSLTree::FExpressionEmitResult::ForwardValue(FEmitContext& Context, co
 	if (EvaluationType == EExpressionEvaluationType::Shader)
 	{
 		bInline = true;
-		Writer.Writef(TEXT("%s"), Context.GetCode(InValue));
+		Code.Appendf(TEXT("%s"), Context.GetCode(InValue));
 	}
 	else
 	{
@@ -220,9 +119,9 @@ const UE::HLSLTree::FEmitValue* UE::HLSLTree::FEmitContext::AcquireValue(FExpres
 	FDeclarationEntry* Entry = FoundEntry ? *FoundEntry : nullptr;
 	if (!Entry)
 	{
-		FLocalHLSLCodeWriter LocalWriter;
+		TStringBuilder<2048> LocalCodeString;
 		Shader::FPreshaderData LocalPreshader;
-		FExpressionEmitResult EmitResult(LocalWriter, LocalPreshader);
+		FExpressionEmitResult EmitResult(LocalCodeString, LocalPreshader);
 
 		bool bEmitResult = false;
 		{
@@ -262,7 +161,7 @@ const UE::HLSLTree::FEmitValue* UE::HLSLTree::FEmitContext::AcquireValue(FExpres
 				check(EmitResult.EvaluationType == EExpressionEvaluationType::Shader);
 				if (EmitResult.bInline)
 				{
-					Entry->Value.Code = AllocateString(*Allocator, LocalWriter.GetStringBuilder());
+					Entry->Value.Code = AllocateString(*Allocator, LocalCodeString);
 				}
 				else
 				{
@@ -270,7 +169,7 @@ const UE::HLSLTree::FEmitValue* UE::HLSLTree::FEmitContext::AcquireValue(FExpres
 					check(EmitScope);
 
 					// Check to see if we've already generated code for an equivalent expression in either this scope, or any outer visible scope
-					const FSHAHash Hash = LocalWriter.GetCodeHash();
+					const FSHAHash Hash = HashString(LocalCodeString);
 					const TCHAR* Declaration = nullptr;
 
 					FEmitScope* CheckScope = EmitScope;
@@ -293,7 +192,7 @@ const UE::HLSLTree::FEmitValue* UE::HLSLTree::FEmitContext::AcquireValue(FExpres
 						WriteStatementToScopef(*EmitScope, TEXT("const %s %s = %s;"),
 							TypeDesc.Name,
 							Declaration,
-							LocalWriter.GetStringBuilder().ToString());
+							LocalCodeString.ToString());
 						EmitScope->ExpressionMap.Add(Hash, Declaration);
 					}
 					Entry->Value.Code = Declaration;
@@ -621,7 +520,7 @@ bool UE::HLSLTree::FExpressionLocalPHI::EmitCode(FEmitContext& Context, FExpress
 	OutResult.EvaluationType = EExpressionEvaluationType::Shader;
 	OutResult.Type = CombinedValueType;
 	OutResult.bInline = true;
-	OutResult.Writer.Writef(TEXT("%s"), Declaration);
+	OutResult.Code.Appendf(TEXT("%s"), Declaration);
 	return true;
 }
 
@@ -816,7 +715,7 @@ static void WriteScope(const UE::HLSLTree::FEmitScope& EmitScope, int32 IndentLe
 	}
 }
 
-bool UE::HLSLTree::FTree::EmitHLSL(UE::HLSLTree::FEmitContext& Context, FCodeWriter& Writer) const
+bool UE::HLSLTree::FTree::EmitHLSL(UE::HLSLTree::FEmitContext& Context, FStringBuilderBase& Writer) const
 {
 	FEmitScope* EmitRootScope = Context.AcquireScope(*RootScope);
 	if (RootScope->EmitHLSL(Context, *EmitRootScope))
@@ -830,7 +729,7 @@ bool UE::HLSLTree::FTree::EmitHLSL(UE::HLSLTree::FEmitContext& Context, FCodeWri
 
 		Context.Finalize();
 
-		WriteScope(*EmitRootScope, 0, *Writer.StringBuilder);
+		WriteScope(*EmitRootScope, 0, Writer);
 
 		return true;
 	}
diff --git a/Engine/Source/Runtime/Engine/Private/HLSLTree/HLSLTreeCommon.cpp b/Engine/Source/Runtime/Engine/Private/HLSLTree/HLSLTreeCommon.cpp
--- a/Engine/Source/Runtime/Engine/Private/HLSLTree/HLSLTreeCommon.cpp
+++ b/Engine/Source/Runtime/Engine/Private/HLSLTree/HLSLTreeCommon.cpp
@@ -132,7 +132,7 @@ bool UE::HLSLTree::FExpressionExternalInput::EmitCode(FEmitContext& Context, FEx
 	if (TexCoordIndex >= 0 && TexCoordIndex < 8)
 	{
 		Context.NumTexCoords = FMath::Max(Context.NumTexCoords, TexCoordIndex + 1);
-		OutResult.Writer.Writef(TEXT("Parameters.TexCoords[%u].xy"), TexCoordIndex);
+		OutResult.Code.Appendf(TEXT("Parameters.TexCoords[%u].xy"), TexCoordIndex);
 	}
 	else
 	{
@@ -275,7 +275,7 @@ bool UE::HLSLTree::FExpressionTextureSample::EmitCode(FEmitContext& Context, FEx
 
 	OutResult.EvaluationType = EExpressionEvaluationType::Shader;
 	OutResult.Type = Shader::EValueType::Float4;
-	OutResult.Writer.Writef(TEXT("%s(%s(%s, %s, %s))"), SamplerTypeFunction, SampleFunctionName, *TextureName, *SamplerStateCode, Context.GetCode(TexCoordValue));
+	OutResult.Code.Appendf(TEXT("%s(%s(%s, %s, %s))"), SamplerTypeFunction, SampleFunctionName, *TextureName, *SamplerStateCode, Context.GetCode(TexCoordValue));
 	return true;
 }
 
@@ -284,7 +284,7 @@ bool UE::HLSLTree::FExpressionDefaultMaterialAttributes::EmitCode(FEmitContext&
 	OutResult.EvaluationType = EExpressionEvaluationType::Shader;
 	OutResult.Type = Shader::EValueType::MaterialAttributes;
 	OutResult.bInline = true;
-	OutResult.Writer.Write(TEXT("DefaultMaterialAttributes"));
+	OutResult.Code.Append(TEXT("DefaultMaterialAttributes"));
 	return true;
 }
 
@@ -300,7 +300,7 @@ bool UE::HLSLTree::FExpressionSetMaterialAttribute::EmitCode(FEmitContext& Conte
 	const FString PropertyName = FMaterialAttributeDefinitionMap::GetAttributeName(AttributeID);
 	OutResult.EvaluationType = EExpressionEvaluationType::Shader;
 	OutResult.Type = Shader::EValueType::MaterialAttributes;
-	OutResult.Writer.Writef(TEXT("FMaterialAttributes_Set%s(%s, %s)"), *PropertyName, Context.GetCode(AttributresValue), Context.GetCode(Value));
+	OutResult.Code.Appendf(TEXT("FMaterialAttributes_Set%s(%s, %s)"), *PropertyName, Context.GetCode(AttributresValue), Context.GetCode(Value));
 	return true;
 }
 
@@ -332,7 +332,7 @@ bool UE::HLSLTree::FExpressionSelect::EmitCode(FEmitContext& Context, FExpressio
 		const FEmitValue* FalseValue = Context.AcquireValue(FalseExpression);
 		OutResult.EvaluationType = EExpressionEvaluationType::Shader; // TODO - preshader
 		OutResult.Type = TrueValue->GetExpressionType(); // TODO - FalseValue?
-		OutResult.Writer.Writef(TEXT("(%s ? %s : %s)"),
+		OutResult.Code.Appendf(TEXT("(%s ? %s : %s)"),
 			Context.GetCode(ConditionValue),
 			Context.GetCode(TrueValue),
 			Context.GetCode(FalseValue));
@@ -380,9 +380,9 @@ bool UE::HLSLTree::FExpressionBinaryOp::EmitCode(FEmitContext& Context, FExpress
 		const TCHAR* RhsCode = Context.GetCode(RhsValue);
 		switch (Op)
 		{
-		case EBinaryOp::Add: OutResult.Writer.Writef(TEXT("(%s + %s)"), LhsCode, RhsCode); break;
-		case EBinaryOp::Mul: OutResult.Writer.Writef(TEXT("(%s * %s)"), LhsCode, RhsCode); break;
-		case EBinaryOp::Less: OutResult.Writer.Writef(TEXT("(%s < %s)"), LhsCode, RhsCode); break;
+		case EBinaryOp::Add: OutResult.Code.Appendf(TEXT("(%s + %s)"), LhsCode, RhsCode); break;
+		case EBinaryOp::Mul: OutResult.Code.Appendf(TEXT("(%s * %s)"), LhsCode, RhsCode); break;
+		case EBinaryOp::Less: OutResult.Code.Appendf(TEXT("(%s < %s)"), LhsCode, RhsCode); break;
 		default: checkNoEntry(); break;
 		}
 	}
@@ -418,7 +418,7 @@ bool UE::HLSLTree::FExpressionSwizzle::EmitCode(FEmitContext& Context, FExpressi
 	OutResult.Type = Shader::MakeValueType(InputValue->GetExpressionType(), Parameters.NumComponents);
 	if (OutResult.EvaluationType == EExpressionEvaluationType::Shader)
 	{
-		OutResult.Writer.Writef(TEXT("%s.%s"), Context.GetCode(InputValue), Swizzle);
+		OutResult.Code.Appendf(TEXT("%s.%s"), Context.GetCode(InputValue), Swizzle);
 	}
 	else
 	{
@@ -460,7 +460,7 @@ bool UE::HLSLTree::FExpressionAppend::EmitCode(FEmitContext& Context, FExpressio
 	if (OutResult.EvaluationType == EExpressionEvaluationType::Shader)
 	{
 		const Shader::FValueTypeDescription ResultTypeDesc = Shader::GetValueTypeDescription(OutResult.Type);
-		OutResult.Writer.Writef(TEXT("%s(%s, %s)"), ResultTypeDesc.Name, Context.GetCode(LhsValue), Context.GetCode(RhsValue));
+		OutResult.Code.Appendf(TEXT("%s(%s, %s)"), ResultTypeDesc.Name, Context.GetCode(LhsValue), Context.GetCode(RhsValue));
 	}
 	else
 	{
@@ -497,7 +497,7 @@ bool UE::HLSLTree::FExpressionCast::EmitCode(FEmitContext& Context, FExpressionE
 	if (InputTypeDesc.NumComponents == OutputTypeDesc.NumComponents)
 	{
 		// Cast between different underlying types with the same number of components, use a C-style cast
-		OutResult.Writer.Writef(TEXT("((%s)%s)"), OutputTypeDesc.Name, InputCode);
+		OutResult.Code.Appendf(TEXT("((%s)%s)"), OutputTypeDesc.Name, InputCode);
 	}
 	else if (OutputTypeDesc.NumComponents < InputTypeDesc.NumComponents)
 	{
@@ -506,12 +506,12 @@ bool UE::HLSLTree::FExpressionCast::EmitCode(FEmitContext& Context, FExpressionE
 		if (OutputTypeDesc.ComponentType == InputTypeDesc.ComponentType)
 		{
 			// underlying types are the same, no C-style cast needed
-			OutResult.Writer.Writef(TEXT("%s.%s"), InputCode, SwizzleName[OutputTypeDesc.NumComponents]);
+			OutResult.Code.Appendf(TEXT("%s.%s"), InputCode, SwizzleName[OutputTypeDesc.NumComponents]);
 		}
 		else
 		{
 			// Different underlying types, add a C-style cast as well
-			OutResult.Writer.Writef(TEXT("((%s)%s.%s)"), OutputTypeDesc.Name, InputCode, SwizzleName[OutputTypeDesc.NumComponents]);
+			OutResult.Code.Appendf(TEXT("((%s)%s.%s)"), OutputTypeDesc.Name, InputCode, SwizzleName[OutputTypeDesc.NumComponents]);
 		}
 	}
 	else
@@ -521,19 +521,18 @@ bool UE::HLSLTree::FExpressionCast::EmitCode(FEmitContext& Context, FExpressionE
 		{
 			// Replicating a scalar value, use a swizzle
 			static const TCHAR* ScalarSwizzleName[] = { TEXT("x"), TEXT("xx"), TEXT("xxx"), TEXT("xxxx") };
-			OutResult.Writer.Writef(TEXT("%s.%s"), InputCode, ScalarSwizzleName[OutputTypeDesc.NumComponents]);
+			OutResult.Code.Appendf(TEXT("%s.%s"), InputCode, ScalarSwizzleName[OutputTypeDesc.NumComponents]);
 		}
 		else
 		{
 			// Padding with 0s, use a constructor
 			const Shader::FValue ZeroValue(OutputTypeDesc.ComponentType, 1);
-			OutResult.Writer.Writef(TEXT("%s(%s"), OutputTypeDesc.Name, InputCode);
+			OutResult.Code.Appendf(TEXT("%s(%s"), OutputTypeDesc.Name, InputCode);
 			for (int32 i = InputTypeDesc.NumComponents; i < OutputTypeDesc.NumComponents; ++i)
 			{
-				OutResult.Writer.Write(TEXT(", "));
-				OutResult.Writer.WriteConstant(ZeroValue);
+				OutResult.Code.Appendf(TEXT(", %s"), *ZeroValue.ToString(UE::Shader::EValueStringFormat::HLSL));
 			}
-			OutResult.Writer.Write(TEXT(")"));
+			OutResult.Code.Append(TEXT(")"));
 		}
 	}
 	return true;
@@ -544,7 +543,7 @@ bool UE::HLSLTree::FExpressionReflectionVector::EmitCode(FEmitContext& Context,
 	OutResult.bInline = true;
 	OutResult.Type = Shader::EValueType::Float3;
 	OutResult.EvaluationType = EExpressionEvaluationType::Shader;
-	OutResult.Writer.Write(TEXT("Parameters.ReflectionVector"));
+	OutResult.Code.Append(TEXT("Parameters.ReflectionVector"));
 	return true;
 }
 
diff --git a/Engine/Source/Runtime/Engine/Private/Materials/MaterialHLSLEmitter.cpp b/Engine/Source/Runtime/Engine/Private/Materials/MaterialHLSLEmitter.cpp
--- a/Engine/Source/Runtime/Engine/Private/Materials/MaterialHLSLEmitter.cpp
+++ b/Engine/Source/Runtime/Engine/Private/Materials/MaterialHLSLEmitter.cpp
@@ -674,19 +674,19 @@ bool MaterialEmitHLSL(const FMaterialCompileTargetParameters& InCompilerTarget,
 	FMemStackBase Allocator;
 	FMemMark MemMark(Allocator);
 
-	UE::HLSLTree::FCodeWriter* CodeWriter = UE::HLSLTree::FCodeWriter::Create(Allocator);
+	TStringBuilder<32 * 1024> Code;
 
 	UE::HLSLTree::FEmitContext EmitContext;
 	EmitContext.Allocator = &Allocator;
 	EmitContext.Material = &InMaterial;
 	EmitContext.StaticParameters = &InStaticParameters;
 	EmitContext.MaterialCompilationOutput = &OutCompilationOutput;
-	InTree.EmitHLSL(EmitContext, *CodeWriter);
+	InTree.EmitHLSL(EmitContext, Code);
 
 	//InOutMaterial.CompileErrors = MoveTemp(CompileErrors);
 	//InOutMaterial.ErrorExpressions = MoveTemp(ErrorExpressions);
 
-	const FString MaterialTemplateSource = GenerateMaterialTemplateHLSL(InCompilerTarget.ShaderPlatform, InMaterial, EmitContext, CodeWriter->StringBuilder->ToString());
+	const FString MaterialTemplateSource = GenerateMaterialTemplateHLSL(InCompilerTarget.ShaderPlatform, InMaterial, EmitContext, Code.ToString());
 
 	OutMaterialEnvironment = new FSharedShaderCompilerEnvironment();
 	OutMaterialEnvironment->TargetPlatform = InCompilerTarget.TargetPlatform;
diff --git a/Engine/Source/Runtime/Engine/Public/HLSLTree/HLSLTree.h b/Engine/Source/Runtime/Engine/Public/HLSLTree/HLSLTree.h
--- a/Engine/Source/Runtime/Engine/Public/HLSLTree/HLSLTree.h
+++ b/Engine/Source/Runtime/Engine/Public/HLSLTree/HLSLTree.h
@@ -34,64 +34,6 @@ class FExpressionLocalPHI;
 
 static constexpr int32 MaxNumPreviousScopes = 2;
 
-/** Allows building a string incrementally, with some additional features to support code generation, such as indent handling */
-class FCodeWriter
-{
-public:
-	static FCodeWriter* Create(FMemStackBase& Allocator);
-
-	const FStringBuilderBase& GetStringBuilder() const { return *StringBuilder; }
-
-	FSHAHash GetCodeHash() const;
-
-	void IncreaseIndent();
-	void DecreaseIndent();
-
-	template <typename TextType>
-	inline void Write(const TextType& Text)
-	{
-		StringBuilder->Append(Text);
-	}
-
-	template <typename FormatType, typename... ArgTypes>
-	inline void Writef(const FormatType& Format, ArgTypes... Args)
-	{
-		StringBuilder->Appendf(Format, Args...);
-	}
-
-	template <typename TextType>
-	inline void WriteLine(const TextType& Text)
-	{
-		WriteIndent();
-		StringBuilder->Append(Text);
-		StringBuilder->Append('\n');
-	}
-
-	template <typename FormatType, typename... ArgTypes>
-	inline void WriteLinef(const FormatType& Format, ArgTypes... Args)
-	{
-		WriteIndent();
-		StringBuilder->Appendf(Format, Args...);
-		StringBuilder->Append('\n');
-	}
-
-	void WriteConstant(const Shader::FValue& Value);
-
-	explicit FCodeWriter(FStringBuilderBase* InStringBuilder)
-		: StringBuilder(InStringBuilder)
-		, IndentLevel(0)
-	{}
-
-	void WriteIndent();
-
-	void Reset();
-
-	void Append(const FCodeWriter& InWriter);
-
-	FStringBuilderBase* StringBuilder;
-	int32 IndentLevel;
-};
-
 class FScope;
 class FEmitScope;
 
@@ -269,8 +211,8 @@ public:
 
 struct FExpressionEmitResult
 {
-	FExpressionEmitResult(FCodeWriter& InWriter, Shader::FPreshaderData& InPreshader)
-		: Writer(InWriter)
+	FExpressionEmitResult(FStringBuilderBase& InCode, Shader::FPreshaderData& InPreshader)
+		: Code(InCode)
 		, Preshader(InPreshader)
 		, EvaluationType(EExpressionEvaluationType::None)
 		, Type(Shader::EValueType::Void)
@@ -279,7 +221,7 @@ struct FExpressionEmitResult
 
 	void ForwardValue(FEmitContext& Context, const FEmitValue* InValue);
 
-	FCodeWriter& Writer;
+	FStringBuilderBase& Code;
 	Shader::FPreshaderData& Preshader;
 	EExpressionEvaluationType EvaluationType;
 	Shader::EValueType Type;
@@ -430,7 +372,7 @@ public:
 
 	FMemStackBase& GetAllocator() { return *Allocator; }
 
-	bool EmitHLSL(FEmitContext& Context, FCodeWriter& Writer) const;
+	bool EmitHLSL(FEmitContext& Context, FStringBuilderBase& Writer) const;
 
 	FScope& GetRootScope() const { return *RootScope; }
 
