Fix static analysis warning

diff --git a/Engine/Source/Runtime/Engine/Private/Materials/MaterialHLSLGenerator.cpp b/Engine/Source/Runtime/Engine/Private/Materials/MaterialHLSLGenerator.cpp
--- a/Engine/Source/Runtime/Engine/Private/Materials/MaterialHLSLGenerator.cpp
+++ b/Engine/Source/Runtime/Engine/Private/Materials/MaterialHLSLGenerator.cpp
@@ -432,7 +432,7 @@ UE::HLSLTree::FExpression* FMaterialHLSLGenerator::GenerateFunctionCall(UE::HLSL
 		return nullptr;
 	}
 
-	const UMaterialExpressionFunctionOutput* ExpressionOutput = FunctionOutputs.IsValidIndex(OutputIndex) ? FunctionOutputs[OutputIndex].ExpressionOutput : nullptr;
+	const UMaterialExpressionFunctionOutput* ExpressionOutput = FunctionOutputs.IsValidIndex(OutputIndex) ? FunctionOutputs[OutputIndex].ExpressionOutput.Get() : nullptr;
 	if (!ExpressionOutput)
 	{
 		Error(TEXT("Invalid function output"));
