[DebugDrawService] added synchronization object for delegate registration since it can happen from multiple threads from primitives added in batch through
ParallelFor(AddPrimitiveBatches.Num(), [&](int32 Index)
{
   FOptionalTaskTagScope Scope(ETaskTag::EParallelGameThread);
   Scene->AddPrimitive(...);
}

Critical section used only for registration since it should not be required for the other accesses (i.e. Unregister &Draw). In those cases we use the AccessDetector.

diff --git a/Engine/Source/Runtime/Engine/Classes/Debug/DebugDrawService.h b/Engine/Source/Runtime/Engine/Classes/Debug/DebugDrawService.h
--- a/Engine/Source/Runtime/Engine/Classes/Debug/DebugDrawService.h
+++ b/Engine/Source/Runtime/Engine/Classes/Debug/DebugDrawService.h
@@ -1,10 +1,10 @@
 // Copyright Epic Games, Inc. All Rights Reserved.
 #pragma once
 
-#include "CoreMinimal.h"
-#include "UObject/ObjectMacros.h"
 #include "ShowFlags.h"
 #include "Kismet/BlueprintFunctionLibrary.h"
+#include "Misc/MTAccessDetector.h"
+#include "HAL/CriticalSection.h"
 #include "DebugDrawService.generated.h"
 
 class FCanvas;
@@ -20,8 +20,6 @@ class ENGINE_API UDebugDrawService : public UBlueprintFunctionLibrary
 {
 	GENERATED_UCLASS_BODY()
 
-	//void Register
-
 	static FDelegateHandle Register(const TCHAR* Name, const FDebugDrawDelegate& NewDelegate);
 	static void Unregister(FDelegateHandle HandleToRemove);
 
@@ -32,6 +30,23 @@ class ENGINE_API UDebugDrawService : public UBlueprintFunctionLibrary
 	static void Draw(const FEngineShowFlags Flags, class FViewport* Viewport, FSceneView* View, FCanvas* Canvas, class UCanvas* CanvasObject = nullptr);
 
 private:
+
+	/**
+	 * Synchronization object for delegate registration since it can happen from multiple threads from
+	 * primitives added in batch through
+	 *		ParallelFor(AddPrimitiveBatches.Num(), [&](int32 Index)
+	 *			{
+	 *				FOptionalTaskTagScope Scope(ETaskTag::EParallelGameThread);
+	 *				Scene->AddPrimitive(...);
+	 *			}
+	 *
+	 * Critical section used only for registration since it should not be required for the other accesses (i.e. Unregister &Draw).
+	 * In those cases we use the AccessDetector.
+	 */
+	static FCriticalSection DelegatesLock;
+#if ENABLE_MT_DETECTOR
+	static FRWAccessDetector DelegatesDetector;
+#endif
 	static TArray<TArray<FDebugDrawDelegate> > Delegates;
 	static FEngineShowFlags ObservedFlags;
 };
diff --git a/Engine/Source/Runtime/Engine/Private/Debug/DebugDrawService.cpp b/Engine/Source/Runtime/Engine/Private/Debug/DebugDrawService.cpp
--- a/Engine/Source/Runtime/Engine/Private/Debug/DebugDrawService.cpp
+++ b/Engine/Source/Runtime/Engine/Private/Debug/DebugDrawService.cpp
@@ -3,8 +3,11 @@
 #include "Debug/DebugDrawService.h"
 #include "UObject/Package.h"
 #include "Engine/Canvas.h"
-#include "Engine/Engine.h"
-#include "IXRTrackingSystem.h"
+
+FCriticalSection UDebugDrawService::DelegatesLock;
+#if ENABLE_MT_DETECTOR
+FRWAccessDetector UDebugDrawService::DelegatesDetector;
+#endif
 
 TArray<TArray<FDebugDrawDelegate> > UDebugDrawService::Delegates;
 FEngineShowFlags UDebugDrawService::ObservedFlags(ESFIM_Editor);
@@ -17,13 +20,13 @@ UDebugDrawService::UDebugDrawService(const FObjectInitializer& ObjectInitializer
 
 FDelegateHandle UDebugDrawService::Register(const TCHAR* Name, const FDebugDrawDelegate& NewDelegate)
 {
-	check(IsInGameThread());
-
-	int32 Index = FEngineShowFlags::FindIndexByName(Name);
+	const int32 Index = FEngineShowFlags::FindIndexByName(Name);
 
 	FDelegateHandle Result;
 	if (Index != INDEX_NONE)
 	{
+		FScopeLock PathLock(&DelegatesLock);
+		UE_MT_SCOPED_WRITE_ACCESS(DelegatesDetector);
 		if (Index >= Delegates.Num())
 		{
 			Delegates.AddZeroed(Index - Delegates.Num() + 1);
@@ -35,10 +38,9 @@ FDelegateHandle UDebugDrawService::Register(const TCHAR* Name, const FDebugDrawD
 	return Result;
 }
 
-void UDebugDrawService::Unregister(FDelegateHandle HandleToRemove)
+void UDebugDrawService::Unregister(const FDelegateHandle HandleToRemove)
 {
-	check(IsInGameThread());
-
+	UE_MT_SCOPED_WRITE_ACCESS(DelegatesDetector);
 	TArray<FDebugDrawDelegate>* DelegatesArray = Delegates.GetData();
 	for (int32 Flag = 0; Flag < Delegates.Num(); ++Flag, ++DelegatesArray)
 	{
@@ -67,7 +69,7 @@ void UDebugDrawService::Draw(const FEngineShowFlags Flags, FViewport* Viewport,
 		}
 	}
 
-	// Canvas must be initialize every draw because the FCanvas passed in is on the stack in some senarioes.
+	// Canvas must be initialized every draw because the FCanvas passed in is on the stack in some scenarios.
 	CanvasObject->Init(View->UnscaledViewRect.Width(), View->UnscaledViewRect.Height(), View, Canvas);
 
 	CanvasObject->Update();	
@@ -79,11 +81,12 @@ void UDebugDrawService::Draw(const FEngineShowFlags Flags, FViewport* Viewport,
 
 void UDebugDrawService::Draw(const FEngineShowFlags Flags, UCanvas* Canvas)
 {
-	if (Canvas == NULL)
+	if (Canvas == nullptr)
 	{
 		return;
 	}
-	
+
+	UE_MT_SCOPED_READ_ACCESS(DelegatesDetector);
 	for (int32 FlagIndex = 0; FlagIndex < Delegates.Num(); ++FlagIndex)
 	{
 		if (Flags.GetSingleFlag(FlagIndex) && ObservedFlags.GetSingleFlag(FlagIndex) && Delegates[FlagIndex].Num() > 0)
@@ -94,7 +97,7 @@ void UDebugDrawService::Draw(const FEngineShowFlags Flags, UCanvas* Canvas)
 
 				if (Delegate.IsBound())
 				{
-					Delegate.Execute(Canvas, NULL);
+					Delegate.Execute(Canvas, nullptr);
 				}
 				else
 				{
