Refactoring on HLSLTree
- Remove most temp storage on FEmitContext, and instead cache data directly on FNodes
- Split HLSL generation into 2 phases: ComputeTypes and Emit; simplifies code dealing with PHI nodes, allows better tracking of parameter types

diff --git a/Engine/Source/Runtime/Engine/Private/HLSLTree/HLSLTree.cpp b/Engine/Source/Runtime/Engine/Private/HLSLTree/HLSLTree.cpp
--- a/Engine/Source/Runtime/Engine/Private/HLSLTree/HLSLTree.cpp
+++ b/Engine/Source/Runtime/Engine/Private/HLSLTree/HLSLTree.cpp
@@ -32,541 +32,486 @@ const TCHAR* AllocateStringf(FMemStackBase& Allocator, const FormatType& Format,
 	return Result;
 }
 
-const TCHAR* AllocateString(FMemStackBase& Allocator, const FStringBuilderBase& StringBuilder)
+const TCHAR* AllocateString(FMemStackBase& Allocator, const TCHAR* String, int32 Length)
 {
-	const int32 Length = StringBuilder.Len();
 	TCHAR* Result = New<TCHAR>(Allocator, Length + 1);
-	FMemory::Memcpy(Result, StringBuilder.GetData(), Length * sizeof(TCHAR));
+	FMemory::Memcpy(Result, String, Length * sizeof(TCHAR));
 	Result[Length] = 0;
 	return Result;
 }
 
-FSHAHash HashString(const FStringBuilderBase& StringBuilder)
+const TCHAR* AllocateString(FMemStackBase& Allocator, const FStringBuilderBase& StringBuilder)
+{
+	return AllocateString(Allocator, StringBuilder.GetData(), StringBuilder.Len());
+}
+
+FSHAHash HashString(const TCHAR* String, int32 Length)
 {
 	FSHAHash Hash;
-	FSHA1::HashBuffer(StringBuilder.GetData(), StringBuilder.Len() * sizeof(TCHAR), Hash.Hash);
+	FSHA1::HashBuffer(String, Length * sizeof(TCHAR), Hash.Hash);
 	return Hash;
 }
 
-UE::HLSLTree::FEmitContext::FEmitContext()
+FSHAHash HashString(const FStringBuilderBase& StringBuilder)
 {
+	return HashString(StringBuilder.GetData(), StringBuilder.Len());
 }
 
-UE::HLSLTree::FEmitContext::~FEmitContext()
+UE::HLSLTree::FErrors::FErrors(FMemStackBase& InAllocator)
+	: Allocator(&InAllocator)
 {
-	for (Shader::FPreshaderData* Preshader : TempPreshaders)
-	{
-		delete Preshader;
-	}
 }
 
-UE::HLSLTree::FEmitScope* UE::HLSLTree::FEmitContext::FindScope(const FScope& Scope)
+bool UE::HLSLTree::FErrors::AddError(const FNode* InNode, FStringView InError)
 {
-	FEmitScope** FoundScope = ScopeMap.Find(&Scope);
-	return FoundScope ? *FoundScope : nullptr;
+	int a = 0;
+	return false;
 }
 
-UE::HLSLTree::FEmitScope* UE::HLSLTree::FEmitContext::AcquireScope(const FScope& Scope)
+UE::HLSLTree::FEmitContext::FEmitContext(FMemStackBase& InAllocator)
+	: Allocator(&InAllocator)
+	, Errors(InAllocator)
 {
-	FEmitScope** FoundScope = ScopeMap.Find(&Scope);
-	FEmitScope* EmitScope = nullptr;
-	if (FoundScope)
-	{
-		EmitScope = *FoundScope;
-	}
-	else
-	{
-		EmitScope = new(*Allocator) FEmitScope();
-		if (Scope.GetParentScope())
-		{
-			EmitScope->ParentScope = AcquireScope(*Scope.GetParentScope());
-		}
-		ScopeMap.Add(&Scope, EmitScope);
-	}
+}
 
-	return EmitScope;
+UE::HLSLTree::FEmitContext::~FEmitContext()
+{
 }
 
-UE::HLSLTree::FEmitScope& UE::HLSLTree::FEmitContext::GetCurrentScope()
+const TCHAR* UE::HLSLTree::FEmitContext::AcquireLocalDeclarationCode()
 {
-	return *ScopeStack.Last();
+	return AllocateStringf(*Allocator, TEXT("Local%d"), NumExpressionLocals++);
 }
 
-void UE::HLSLTree::FExpressionEmitResult::ForwardValue(FEmitContext& Context, const FEmitValue* InValue)
+// From MaterialUniformExpressions.cpp
+extern void WriteMaterialUniformAccess(UE::Shader::EValueComponentType ComponentType, uint32 NumComponents, uint32 UniformOffset, FStringBuilderBase& OutResult);
+
+void UE::HLSLTree::FEmitContext::AddPreshader(Shader::EValueType Type, const Shader::FPreshaderData& Preshader, FStringBuilderBase& OutCode)
 {
-	check(InValue);
-	EvaluationType = InValue->GetEvaluationType();
-	Type = InValue->GetExpressionType();
-	if (EvaluationType == EExpressionEvaluationType::Shader)
+	const Shader::FValueTypeDescription TypeDesc = Shader::GetValueTypeDescription(Type);
+	ensure(TypeDesc.ComponentType == Shader::EValueComponentType::Float || TypeDesc.ComponentType == Shader::EValueComponentType::Double);
+
+	FUniformExpressionSet& UniformExpressionSet = MaterialCompilationOutput->UniformExpressionSet;
+	FMaterialUniformPreshaderHeader& PreshaderHeader = UniformExpressionSet.UniformPreshaders.AddDefaulted_GetRef();
+
+	const uint32 RegisterOffset = UniformPreshaderOffset % 4;
+	if (TypeDesc.ComponentType == Shader::EValueComponentType::Float && RegisterOffset + TypeDesc.NumComponents > 4u)
 	{
-		bInline = true;
-		Code.Appendf(TEXT("%s"), Context.GetCode(InValue));
+		// If this uniform would span multiple registers, align offset to the next register to avoid this
+		UniformPreshaderOffset = Align(UniformPreshaderOffset, 4u);
+	}
+
+	PreshaderHeader.OpcodeOffset = UniformExpressionSet.UniformPreshaderData.Num();
+	UniformExpressionSet.UniformPreshaderData.Append(Preshader);
+	PreshaderHeader.OpcodeSize = Preshader.Num();
+	PreshaderHeader.BufferOffset = UniformPreshaderOffset;
+	PreshaderHeader.ComponentType = TypeDesc.ComponentType;
+	PreshaderHeader.NumComponents = TypeDesc.NumComponents;
+
+	if (TypeDesc.ComponentType == Shader::EValueComponentType::Double)
+	{
+		if (TypeDesc.NumComponents == 1)
+		{
+			OutCode.Append(TEXT("MakeLWCScalar("));
+		}
+		else
+		{
+			OutCode.Appendf(TEXT("MakeLWCVector%d("), TypeDesc.NumComponents);
+		}
+
+		WriteMaterialUniformAccess(UE::Shader::EValueComponentType::Float, TypeDesc.NumComponents, UniformPreshaderOffset, OutCode); // Tile
+		UniformPreshaderOffset += TypeDesc.NumComponents;
+		OutCode.Append(TEXT(","));
+		WriteMaterialUniformAccess(UE::Shader::EValueComponentType::Float, TypeDesc.NumComponents, UniformPreshaderOffset, OutCode); // Offset
+		UniformPreshaderOffset += TypeDesc.NumComponents;
+		OutCode.Append(TEXT(")"));
 	}
 	else
 	{
-		Context.AppendPreshader(InValue, Preshader);
+		WriteMaterialUniformAccess(TypeDesc.ComponentType, TypeDesc.NumComponents, UniformPreshaderOffset, OutCode);
+		UniformPreshaderOffset += TypeDesc.NumComponents;
 	}
 }
 
-const UE::HLSLTree::FEmitValue* UE::HLSLTree::FEmitContext::AcquireValue(FExpression* Expression)
+void UE::HLSLTree::FEmitContext::Finalize()
 {
-	if (!Expression)
-	{
-		return nullptr;
-	}
+	MaterialCompilationOutput->UniformExpressionSet.UniformPreshaderBufferSize = (UniformPreshaderOffset + 3u) / 4u;
+}
 
-	FDeclarationEntry** FoundEntry = DeclarationMap.Find(Expression);
-	FDeclarationEntry* Entry = FoundEntry ? *FoundEntry : nullptr;
-	if (!Entry)
+UE::HLSLTree::FScope* UE::HLSLTree::FScope::FindSharedParent(FScope* Lhs, FScope* Rhs)
+{
+	FScope* Scope0 = Lhs;
+	FScope* Scope1 = Rhs;
+	if (Scope1)
 	{
-		TStringBuilder<2048> LocalCodeString;
-		Shader::FPreshaderData LocalPreshader;
-		FExpressionEmitResult EmitResult(LocalCodeString, LocalPreshader);
-
-		bool bEmitResult = false;
+		while (Scope0 != Scope1)
 		{
-			bool bAlreadyPending = false;
-			const FSetElementId Id = PendingEmitValueExpressions.Add(Expression, &bAlreadyPending);
-			if (!bAlreadyPending)
+			if (Scope0->NestedLevel > Scope1->NestedLevel)
 			{
-				bEmitResult = Expression->EmitCode(*this, EmitResult);
-				PendingEmitValueExpressions.Remove(Id);
+				check(Scope0->ParentScope);
+				Scope0 = Scope0->ParentScope;
+			}
+			else
+			{
+				check(Scope1->ParentScope);
+				Scope1 = Scope1->ParentScope;
 			}
 		}
+	}
+	return Scope0;
+}
+
+bool UE::HLSLTree::FExpressionLocalPHI::UpdateType(FUpdateTypeContext& Context, int8 InRequestedNumComponents)
+{
+	Shader::EValueType TypePerValue[MaxNumPreviousScopes] = { Shader::EValueType::Void };
+	int32 NumValidTypes = 0;
+	Shader::EValueComponentType ComponentType = Shader::EValueComponentType::Void;
+	int8 NumComponents = 0;
 
-		if (bEmitResult)
+	auto UpdateValueTypes = [&]()
+	{
+		for (int32 i = 0; i < NumValues; ++i)
 		{
-			check(EmitResult.EvaluationType != EExpressionEvaluationType::None);
-			check(EmitResult.Type != Shader::EValueType::Void);
-
-			Entry = new(*Allocator) FDeclarationEntry();
-			DeclarationMap.Add(Expression, Entry);
-			Entry->Value.ExpressionType = EmitResult.Type;
-			Entry->Value.EvaluationType = EmitResult.EvaluationType;
-			if (EmitResult.EvaluationType == EExpressionEvaluationType::Constant)
-			{
-				// Evaluate the constant preshader and store its value
-				FMaterialRenderContext RenderContext(nullptr, *Material, nullptr);
-				LocalPreshader.Evaluate(nullptr, RenderContext, Entry->Value.ConstantValue);
-			}
-			else if (EmitResult.EvaluationType == EExpressionEvaluationType::Preshader)
-			{
-				// Non-constant preshader, store it
-				Shader::FPreshaderData* Preshader = new Shader::FPreshaderData(MoveTemp(LocalPreshader));
-				TempPreshaders.Add(Preshader); // TODO - dedupe? store more efficiently?
-				Entry->Value.Preshader = Preshader;
-			}
-			else
+			if (TypePerValue[i] == Shader::EValueType::Void)
 			{
-				check(EmitResult.EvaluationType == EExpressionEvaluationType::Shader);
-				if (EmitResult.bInline)
-				{
-					Entry->Value.Code = AllocateString(*Allocator, LocalCodeString);
-				}
-				else
+				TypePerValue[i] = RequestExpressionType(Context, Values[i], InRequestedNumComponents);
+				if (TypePerValue[i] != Shader::EValueType::Void)
 				{
-					FEmitScope* EmitScope = FindScope(*Expression->ParentScope);
-					check(EmitScope);
-
-					// Check to see if we've already generated code for an equivalent expression in either this scope, or any outer visible scope
-					const FSHAHash Hash = HashString(LocalCodeString);
-					const TCHAR* Declaration = nullptr;
-
-					FEmitScope* CheckScope = EmitScope;
-					while (CheckScope)
+					const Shader::FValueTypeDescription TypeDesc = Shader::GetValueTypeDescription(TypePerValue[i]);
+					if (ComponentType == Shader::EValueComponentType::Void)
 					{
-						const TCHAR** FoundDeclaration = CheckScope->ExpressionMap.Find(Hash);
-						if (FoundDeclaration)
-						{
-							// Re-use results from previous matching expression
-							Declaration = *FoundDeclaration;
-							break;
-						}
-						CheckScope = CheckScope->ParentScope;
+						ComponentType = TypeDesc.ComponentType;
+						NumComponents = TypeDesc.NumComponents;
 					}
-
-					if (!Declaration)
+					else
 					{
-						const Shader::FValueTypeDescription& TypeDesc = Shader::GetValueTypeDescription(EmitResult.Type);
-						Declaration = AcquireLocalDeclarationCode();
-						WriteStatementToScopef(*EmitScope, TEXT("const %s %s = %s;"),
-							TypeDesc.Name,
-							Declaration,
-							LocalCodeString.ToString());
-						EmitScope->ExpressionMap.Add(Hash, Declaration);
+						NumComponents = FMath::Max(NumComponents, TypeDesc.NumComponents);
+						if (ComponentType != TypeDesc.ComponentType)
+						{
+							return Context.Errors.AddError(this, TEXT("Type mismatch"));
+						}
 					}
-					Entry->Value.Code = Declaration;
+					check(NumValidTypes < NumValues);
+					NumValidTypes++;
 				}
 			}
 		}
-	}
-
-	return Entry ? &Entry->Value : nullptr;
-}
+		return true;
+	};
 
-const TCHAR* UE::HLSLTree::FEmitContext::AcquireLocalDeclarationCode()
-{
-	return AllocateStringf(*Allocator, TEXT("Local%d"), NumExpressionLocals++);
-}
+	// First try to assign all the values we can
+	if (!UpdateValueTypes())
+	{
+		return false;
+	}
 
-// From MaterialUniformExpressions.cpp
-extern void WriteMaterialUniformAccess(UE::Shader::EValueComponentType ComponentType, uint32 NumComponents, uint32 UniformOffset, FStringBuilderBase& OutResult);
+	// Assuming we have at least one value with a valid type, we use that to initialize our type
+	if (!SetType(Context, Shader::MakeValueType(ComponentType, NumComponents)))
+	{
+		return false;
+	}
 
-const TCHAR* UE::HLSLTree::FEmitContext::GetCode(const FEmitValue* Value)
-{
-	check(Value);
-	if (!Value->Code)
+	if (NumValidTypes < NumValues)
 	{
-		TStringBuilder<1024> FormattedCode;
-		if (Value->EvaluationType == EExpressionEvaluationType::Constant)
-		{
-			FormattedCode.Append(Value->ConstantValue.ToString(Shader::EValueStringFormat::HLSL));
-		}
-		else
+		// Now try to assign remaining types that failed the first iteration 
+		if (!UpdateValueTypes())
 		{
-			check(Value->EvaluationType == EExpressionEvaluationType::Preshader);
-			check(Value->Preshader);
-
-			const Shader::FValueTypeDescription TypeDesc = Shader::GetValueTypeDescription(Value->ExpressionType);
-			ensure(TypeDesc.ComponentType == Shader::EValueComponentType::Float || TypeDesc.ComponentType == Shader::EValueComponentType::Double);
-
-			FUniformExpressionSet& UniformExpressionSet = MaterialCompilationOutput->UniformExpressionSet;
-			FMaterialUniformPreshaderHeader& PreshaderHeader = UniformExpressionSet.UniformPreshaders.AddDefaulted_GetRef();
-
-			const uint32 RegisterOffset = UniformPreshaderOffset % 4;
-			if (TypeDesc.ComponentType == Shader::EValueComponentType::Float && RegisterOffset + TypeDesc.NumComponents > 4u)
-			{
-				// If this uniform would span multiple registers, align offset to the next register to avoid this
-				UniformPreshaderOffset = Align(UniformPreshaderOffset, 4u);
-			}
-
-			PreshaderHeader.OpcodeOffset = UniformExpressionSet.UniformPreshaderData.Num();
-			UniformExpressionSet.UniformPreshaderData.Append(*Value->Preshader);
-			PreshaderHeader.OpcodeSize = Value->Preshader->Num();
-			PreshaderHeader.BufferOffset = UniformPreshaderOffset;
-			PreshaderHeader.ComponentType = TypeDesc.ComponentType;
-			PreshaderHeader.NumComponents = TypeDesc.NumComponents;
-			
-			if (TypeDesc.ComponentType == Shader::EValueComponentType::Double)
-			{
-				if (TypeDesc.NumComponents == 1)
-				{
-					FormattedCode.Append(TEXT("MakeLWCScalar("));
-				}
-				else
-				{
-					FormattedCode.Appendf(TEXT("MakeLWCVector%d("), TypeDesc.NumComponents);
-				}
-
-				WriteMaterialUniformAccess(UE::Shader::EValueComponentType::Float, TypeDesc.NumComponents, UniformPreshaderOffset, FormattedCode); // Tile
-				UniformPreshaderOffset += TypeDesc.NumComponents;
-				FormattedCode.Append(TEXT(","));
-				WriteMaterialUniformAccess(UE::Shader::EValueComponentType::Float, TypeDesc.NumComponents, UniformPreshaderOffset, FormattedCode); // Offset
-				UniformPreshaderOffset += TypeDesc.NumComponents;
-				FormattedCode.Append(TEXT(")"));
-			}
-			else
-			{
-				WriteMaterialUniformAccess(TypeDesc.ComponentType, TypeDesc.NumComponents, UniformPreshaderOffset, FormattedCode);
-				UniformPreshaderOffset += TypeDesc.NumComponents;
-			}
+			return false;
 		}
-		Value->Code = AllocateString(*Allocator, FormattedCode);
 	}
 
-	return Value->Code;
+	return NumValidTypes == NumValues;
 }
 
-void UE::HLSLTree::FEmitContext::AppendPreshader(const FEmitValue* Value, Shader::FPreshaderData& InOutPreshader)
+bool UE::HLSLTree::FExpressionLocalPHI::PrepareValue(FEmitContext& Context)
 {
-	check(Value);
-	if (Value->EvaluationType == EExpressionEvaluationType::Constant)
+	const int32 LocalPHIIndex = Context.NumLocalPHIs++;
+	if (!SetValueInlineShaderf(Context, TEXT("LocalPHI%d"), LocalPHIIndex))
 	{
-		// Push the constant value
-		InOutPreshader.WriteOpcode(Shader::EPreshaderOpcode::Constant);
-		InOutPreshader.Write(Value->ConstantValue);
+		return false;
 	}
-	else
+
+	// Find the outermost scope to declare our local variable
+	FScope* DeclarationScope = ParentScope;
+	for (int32 i = 0; i < NumValues; ++i)
 	{
-		check(Value->EvaluationType == EExpressionEvaluationType::Preshader);
-		check(Value->Preshader);
-		InOutPreshader.Append(*Value->Preshader);
+		DeclarationScope = FScope::FindSharedParent(DeclarationScope, Scopes[i]);
+		if (!DeclarationScope)
+		{
+			return false;
+		}
 	}
-}
 
-FStringView UE::HLSLTree::FEmitContext::InternalAcquireInternedString(const TCHAR* InString, int32 InLength)
-{
-	TCHAR* InternedString = nullptr;
-	int32 Length = 0;
-	if (InString)
+	for (int32 i = 0; i < NumValues; ++i)
 	{
-		Length = InLength > 0 ? InLength : FCString::Strlen(InString);
-		InternedString = new(*Allocator) TCHAR[Length + 1];
-		FMemory::Memcpy(InternedString, InString, Length * sizeof(TCHAR));
-		InternedString[Length] = 0;
+		if (PrepareExpressionValue(Context, Values[i]) == EExpressionEvaluationType::None)
+		{
+			return false;
+		}
 	}
-	return FStringView(InternedString, Length);
-}
 
-FStringView UE::HLSLTree::FEmitContext::AcquireInternedString(const TCHAR* Format, ...)
-{
-	TCHAR Buffer[1024];
-
-	va_list va;
-	va_start(va, Format);
-	const int32 Length = FPlatformString::GetVarArgs(Buffer, UE_ARRAY_COUNT(Buffer), Format, va);
-	va_end(va);
-
-	check(Length >= 0 && Length < UE_ARRAY_COUNT(Buffer));
-
-	return InternalAcquireInternedString(Buffer, Length);
-}
+	const Shader::EValueType Type = GetValueType();
+	const Shader::FValueTypeDescription& TypeDesc = Shader::GetValueTypeDescription(Type);
 
-void UE::HLSLTree::FEmitContext::InternalWriteStatementToScope(FEmitScope& EmitScope, FStringView InternedCode)
-{
-	FEmitStatement* EmitStatement = new(*Allocator) FEmitStatement();
-	EmitStatement->Code = InternedCode;
-	if (EmitScope.LastStatement)
+	bool bNeedToAddDeclaration = true;
+	for (int32 i = 0; i < NumValues; ++i)
 	{
-		EmitStatement->LinkAfter(EmitScope.LastStatement);
+		FScope* ValueScope = Scopes[i];
+		const TCHAR* ValueCode = Values[i]->GetValueShader(Context);
+		if (ValueScope == DeclarationScope)
+		{
+			ValueScope->EmitDeclarationf(Context, TEXT("%s LocalPHI%d = %s;"),
+				TypeDesc.Name,
+				LocalPHIIndex,
+				ValueCode);
+			bNeedToAddDeclaration = false;
+		}
+		else
+		{
+			ValueScope->EmitStatementf(Context, TEXT("LocalPHI%d = %s;"),
+				LocalPHIIndex,
+				ValueCode);
+		}
 	}
-	else
+
+	if (bNeedToAddDeclaration)
 	{
-		EmitScope.FirstStatement = EmitStatement;
+		DeclarationScope->EmitDeclarationf(Context, TEXT("%s LocalPHI%d;"),
+			TypeDesc.Name,
+			LocalPHIIndex);
 	}
-	EmitScope.LastStatement = EmitStatement;
+
+	return true;
 }
 
-bool UE::HLSLTree::FEmitContext::InternalWriteScope(const FScope& Scope, FStringView InternedCode)
+void UE::HLSLTree::FNodeVisitor::VisitNode(FNode* Node)
 {
-	FEmitScope* EmitScope = AcquireScope(Scope);
-	if (!Scope.EmitHLSL(*this, *EmitScope))
+	if (Node)
 	{
-		return false;
+		Node->Visit(*this);
 	}
+}
 
-	FEmitScopeLink* Link = InternalWriteScopeLink(InternedCode);
-	Link->NextScope = EmitScope;
-	return true;
+UE::HLSLTree::ENodeVisitResult UE::HLSLTree::FStatement::Visit(FNodeVisitor& Visitor)
+{
+	return Visitor.OnStatement(*this);
 }
 
-UE::HLSLTree::FEmitScopeLink* UE::HLSLTree::FEmitContext::InternalWriteScopeLink(FStringView InternedCode)
+UE::HLSLTree::ENodeVisitResult UE::HLSLTree::FExpression::Visit(FNodeVisitor& Visitor)
 {
-	FEmitScope& EmitScope = GetCurrentScope();
-	FEmitScopeLink* EmitLink = new(*Allocator) FEmitScopeLink();
-	EmitLink->Code = InternedCode;
-	if (EmitScope.LastLink)
-	{
-		EmitLink->LinkAfter(EmitScope.LastLink);
-	}
-	else
-	{
-		EmitScope.FirstLink = EmitLink;
-	}
-	EmitScope.LastLink = EmitLink;
-	return EmitLink;
+	return Visitor.OnExpression(*this);
 }
 
-void UE::HLSLTree::FEmitContext::WriteDeclaration(FEmitScope& EmitScope, Shader::EValueType Type, const TCHAR* Declaration, const TCHAR* Value)
+UE::HLSLTree::ENodeVisitResult UE::HLSLTree::FTextureParameterDeclaration::Visit(FNodeVisitor& Visitor)
 {
-	FEmitDeclaration* EmitDeclaration = new(*Allocator) FEmitDeclaration();
-	EmitDeclaration->Type = Type;
-	EmitDeclaration->Declaration = Declaration;
-	EmitDeclaration->Value = Value;
-	EmitDeclaration->LinkHead(EmitScope.FirstDeclaration);
+	return Visitor.OnTextureParameterDeclaration(*this);
 }
 
-bool UE::HLSLTree::FEmitContext::WriteAssignment(FEmitScope& EmitScope, const TCHAR* Declaration, FExpression* Expression, Shader::EValueType& InOutType)
+UE::HLSLTree::ENodeVisitResult UE::HLSLTree::FScope::Visit(FNodeVisitor& Visitor)
 {
-	const FEmitValue* Value = AcquireValue(Expression);
-	if (Value)
-	{
-		const Shader::EValueType ValueType = Value->GetExpressionType();
-		if (InOutType == Shader::EValueType::Void)
-		{
-			InOutType = ValueType;
-		}
-		else if (InOutType != ValueType)
-		{
-			return false;
-		}
-		WriteStatementToScopef(EmitScope, TEXT("%s = %s;"), Declaration, GetCode(Value));
-	}
-	else
+	const ENodeVisitResult Result = Visitor.OnScope(*this);
+	if (ShouldVisitDependentNodes(Result))
 	{
-		FEmitAssignment* EmitAssignment = new(*Allocator) FEmitAssignment();
-		EmitAssignment->Declaration = Declaration;
-		EmitAssignment->Expression = Expression;
-		EmitAssignment->LinkHead(EmitScope.FirstAssignment);
+		Visitor.VisitNode(Statement);
 	}
-
-	return true;
+	return Result;
 }
 
-bool UE::HLSLTree::FEmitContext::FinalizeScope(FEmitScope& EmitScope)
+UE::Shader::EValueType UE::HLSLTree::RequestExpressionType(FUpdateTypeContext& Context, FExpression* InExpression, int8 InRequestedNumComponents)
 {
-	if (EmitScope.FirstAssignment)
+	if (!InExpression)
 	{
-		// Scope has pending assigmnets, flush all of them, then restart finalization
-		// (Since flushing pending assignments may generate additional pending assignments)
-		FEmitAssignment::TConstIterator It(EmitScope.FirstAssignment);
-		while (It)
-		{
-			const FEmitAssignment& EmitAssignment = *It;
-			const FEmitValue* Value = AcquireValue(EmitAssignment.Expression);
-			WriteStatementToScopef(EmitScope, TEXT("%s = %s;"), EmitAssignment.Declaration, GetCode(Value));
-			It.Next();
-		}
-		EmitScope.FirstAssignment = nullptr;
-		return false;
+		return Shader::EValueType::Void;
 	}
 
-	UE::HLSLTree::FEmitScopeLink::TConstIterator It(EmitScope.FirstLink);
-	while (It)
+	if (!InExpression->bReentryFlag && InRequestedNumComponents > InExpression->RequestedNumComponents)
 	{
-		const UE::HLSLTree::FEmitScopeLink& EmitLink = *It;
-		if (EmitLink.NextScope)
+		InExpression->bReentryFlag = true;
+		const bool bUpdateResult = InExpression->UpdateType(Context, InRequestedNumComponents);
+		InExpression->bReentryFlag = false;
+
+		if (bUpdateResult)
 		{
-			if (!FinalizeScope(*EmitLink.NextScope))
-			{
-				return false;
-			}
+			check(InExpression->ValueType != Shader::EValueType::Void);
+			InExpression->RequestedNumComponents = InRequestedNumComponents;
+		}
+		else
+		{
+			InExpression->ValueType = Shader::EValueType::Void;
 		}
-
-		It.Next();
 	}
 
-	return true;
+	return InExpression->ValueType;
 }
 
-void UE::HLSLTree::FEmitContext::Finalize()
+UE::HLSLTree::EExpressionEvaluationType UE::HLSLTree::PrepareExpressionValue(FEmitContext& Context, FExpression* InExpression)
 {
-	MaterialCompilationOutput->UniformExpressionSet.UniformPreshaderBufferSize = (UniformPreshaderOffset + 3u) / 4u;
-}
+	if (!InExpression)
+	{
+		return EExpressionEvaluationType::None;
+	}
 
-UE::HLSLTree::FScope* UE::HLSLTree::FScope::FindSharedParent(FScope* Lhs, FScope* Rhs)
-{
-	FScope* Scope0 = Lhs;
-	FScope* Scope1 = Rhs;
-	if (Scope1)
+	if (InExpression->EvaluationType == EExpressionEvaluationType::None)
 	{
-		while (Scope0 != Scope1)
+		check(InExpression->GetValueType() != Shader::EValueType::Void);
+		check(!InExpression->bReentryFlag); // Re-entry is not supported
+
+		InExpression->bReentryFlag = true;
+		const bool bPrepareResult = InExpression->PrepareValue(Context);
+		InExpression->bReentryFlag = false;
+
+		if (bPrepareResult)
 		{
-			if (Scope0->NestedLevel > Scope1->NestedLevel)
-			{
-				check(Scope0->ParentScope);
-				Scope0 = Scope0->ParentScope;
-			}
-			else
-			{
-				check(Scope1->ParentScope);
-				Scope1 = Scope1->ParentScope;
-			}
+			check(InExpression->EvaluationType != EExpressionEvaluationType::None);
+		}
+		else
+		{
+			InExpression->EvaluationType = EExpressionEvaluationType::None;
 		}
 	}
-	return Scope0;
+	return InExpression->EvaluationType;
 }
 
-bool UE::HLSLTree::FExpressionLocalPHI::EmitCode(FEmitContext& Context, FExpressionEmitResult& OutResult) const
+bool UE::HLSLTree::FExpression::InternalSetValueShader(FEmitContext& Context, const TCHAR* InCode, int32 InLength, bool bInline)
 {
-	FEmitScope* EmitScopes[MaxNumPreviousScopes] = { nullptr };
-	const FEmitValue* EmitValues[MaxNumPreviousScopes] = { nullptr };
-	Shader::EValueType CombinedValueType = Shader::EValueType::Void;
+	check(!Code);
+	check(EvaluationType == EExpressionEvaluationType::None);
 
-	const TCHAR* Declaration = Context.AcquireLocalDeclarationCode();
-
-	// Find the outermost scope to declare our local variable
-	FScope* DeclarationScope = ParentScope;
-	for (int32 i = 0; i < NumValues; ++i)
+	EvaluationType = EExpressionEvaluationType::Shader;
+	if (bInline)
 	{
-		DeclarationScope = FScope::FindSharedParent(DeclarationScope, Scopes[i]);
-		if (!DeclarationScope)
-		{
-			return false;
-		}
+		Code = AllocateString(*Context.Allocator, InCode, InLength);
 	}
-
-	bool bNeedToAddDeclaration = true;
-	for (int32 i = 0; i < NumValues; ++i)
+	else
 	{
-		EmitScopes[i] = Context.AcquireScope(*Scopes[i]);
-		if (Scopes[i] == DeclarationScope)
-		{
-			const FEmitValue* Value = Context.AcquireValue(Values[i]);
-			CombinedValueType = Value->GetExpressionType();
-			Context.WriteDeclaration(*EmitScopes[i], CombinedValueType, Declaration, Context.GetCode(Value));
-			bNeedToAddDeclaration = false;
-		}
-		else
+		// Check to see if we've already generated code for an equivalent expression in either this scope, or any outer visible scope
+		const FSHAHash Hash = HashString(InCode, InLength);
+		const TCHAR* Declaration = nullptr;
+
+		FScope* CheckScope = ParentScope;
+		while (CheckScope)
 		{
-			if (!Context.WriteAssignment(*EmitScopes[i], Declaration, Values[i], CombinedValueType))
+			const TCHAR** FoundDeclaration = CheckScope->ExpressionCodeMap.Find(Hash);
+			if (FoundDeclaration)
 			{
-				return false;
+				// Re-use results from previous matching expression
+				Declaration = *FoundDeclaration;
+				break;
 			}
+			CheckScope = CheckScope->ParentScope;
 		}
-	}
-
-	if (CombinedValueType == Shader::EValueType::Void)
-	{
-		return false;
-	}
 
-	if (bNeedToAddDeclaration)
-	{
-		FEmitScope* EmitDeclarationScope = Context.FindScope(*DeclarationScope);
-		Context.WriteDeclaration(*EmitDeclarationScope, CombinedValueType, Declaration);
+		if (!Declaration)
+		{
+			Declaration = Context.AcquireLocalDeclarationCode();
+			ParentScope->ExpressionCodeMap.Add(Hash, Declaration);
+			const Shader::FValueTypeDescription& TypeDesc = Shader::GetValueTypeDescription(ValueType);
+			ParentScope->EmitStatementf(Context, TEXT("const %s %s = %s;"),
+				TypeDesc.Name,
+				Declaration,
+				InCode);
+		}
+		Code = Declaration;
 	}
 
-	OutResult.EvaluationType = EExpressionEvaluationType::Shader;
-	OutResult.Type = CombinedValueType;
-	OutResult.bInline = true;
-	OutResult.Code.Appendf(TEXT("%s"), Declaration);
 	return true;
 }
 
-void UE::HLSLTree::FNodeVisitor::VisitNode(FNode* Node)
+bool UE::HLSLTree::FExpression::SetValuePreshader(FEmitContext& Context, Shader::FPreshaderData& InPreshader)
 {
-	if (Node)
-	{
-		Node->Visit(*this);
-	}
+	check(!Preshader);
+	check(EvaluationType == EExpressionEvaluationType::None);
+
+	EvaluationType = EExpressionEvaluationType::Preshader;
+	Preshader = new Shader::FPreshaderData(MoveTemp(InPreshader));
+	return true;
 }
 
-UE::HLSLTree::ENodeVisitResult UE::HLSLTree::FStatement::Visit(FNodeVisitor& Visitor)
+bool UE::HLSLTree::FExpression::SetValueConstant(FEmitContext& Context, const Shader::FValue& InValue)
 {
-	return Visitor.OnStatement(*this);
+	check(EvaluationType == EExpressionEvaluationType::None);
+
+	EvaluationType = EExpressionEvaluationType::Constant;
+	ConstantValue = InValue;
+	return true;
 }
 
-UE::HLSLTree::ENodeVisitResult UE::HLSLTree::FExpression::Visit(FNodeVisitor& Visitor)
+bool UE::HLSLTree::FExpression::SetValueForward(FEmitContext& Context, FExpression* Source)
 {
-	return Visitor.OnExpression(*this);
+	check(EvaluationType == EExpressionEvaluationType::None);
+
+	EvaluationType = PrepareExpressionValue(Context, Source);
+	switch (EvaluationType)
+	{
+	case EExpressionEvaluationType::None: return false;
+	case EExpressionEvaluationType::Shader: Code = Source->Code; break;
+	case EExpressionEvaluationType::Preshader: Preshader = Source->Preshader; break;
+	case EExpressionEvaluationType::Constant: ConstantValue = Source->ConstantValue; break;
+	default: checkNoEntry(); return false;
+	}
+	return true;
 }
 
-UE::HLSLTree::ENodeVisitResult UE::HLSLTree::FTextureParameterDeclaration::Visit(FNodeVisitor& Visitor)
+bool UE::HLSLTree::FExpression::SetValuePreshader(FEmitContext& Context, EExpressionEvaluationType InEvaluationType, Shader::FPreshaderData& InPreshader)
 {
-	return Visitor.OnTextureParameterDeclaration(*this);
+	if (InEvaluationType == EExpressionEvaluationType::Constant)
+	{
+		// Evaluate the constant preshader and store its value
+		FMaterialRenderContext RenderContext(nullptr, *Context.Material, nullptr);
+		Shader::FValue Value;
+		InPreshader.Evaluate(nullptr, RenderContext, Value);
+		return SetValueConstant(Context, Value);
+	}
+	else
+	{
+		check(InEvaluationType == EExpressionEvaluationType::Preshader);
+		return SetValuePreshader(Context, InPreshader);
+	}
 }
 
-UE::HLSLTree::ENodeVisitResult UE::HLSLTree::FScope::Visit(FNodeVisitor& Visitor)
+const TCHAR* UE::HLSLTree::FExpression::GetValueShader(FEmitContext& Context)
 {
-	const ENodeVisitResult Result = Visitor.OnScope(*this);
-	if (ShouldVisitDependentNodes(Result))
+	check(EvaluationType != EExpressionEvaluationType::None);
+	if (!Code)
 	{
-		Visitor.VisitNode(Statement);
+		TStringBuilder<1024> FormattedCode;
+		if (EvaluationType == EExpressionEvaluationType::Constant)
+		{
+			FormattedCode.Append(ConstantValue.ToString(Shader::EValueStringFormat::HLSL));
+		}
+		else
+		{
+			check(EvaluationType == EExpressionEvaluationType::Preshader);
+			check(Preshader);
+			Context.AddPreshader(GetValueType(), *Preshader, FormattedCode);
+		}
+		Code = AllocateString(*Context.Allocator, FormattedCode);
 	}
-	return Result;
+
+	return Code;
 }
 
-bool UE::HLSLTree::FScope::EmitHLSL(FEmitContext& Context, FEmitScope& Scope) const
+void UE::HLSLTree::FExpression::GetValuePreshader(FEmitContext& Context, Shader::FPreshaderData& OutPreshader)
 {
-	bool bResult = true;
-	if (Statement)
+	switch (EvaluationType)
 	{
-		Context.ScopeStack.Add(&Scope);
-		bResult = Statement->EmitHLSL(Context);
-		verify(Context.ScopeStack.Pop(false) == &Scope);
+	case EExpressionEvaluationType::Preshader:
+		check(Preshader);
+		OutPreshader.Append(*Preshader);
+		break;
+	case EExpressionEvaluationType::Constant:
+		OutPreshader.WriteOpcode(Shader::EPreshaderOpcode::Constant).Write(ConstantValue);
+		break;
+	default:
+		check(false);
+		break;
 	}
-	return bResult;
+}
+
+UE::Shader::FValue UE::HLSLTree::FExpression::GetValueConstant(FEmitContext& Context)
+{
+	check(EvaluationType == EExpressionEvaluationType::Constant);
+	return ConstantValue;
 }
 
 bool UE::HLSLTree::FScope::HasParentScope(const FScope& InParentScope) const
@@ -623,6 +568,43 @@ void UE::HLSLTree::FScope::UseExpression(FExpression* Expression)
 	Visitor.VisitNode(Expression);
 }
 
+void UE::HLSLTree::FScope::InternalEmitCode(FEmitContext& Context, FCodeList& List, FScope* NestedScope, const TCHAR* String, int32 Length)
+{
+	if (NestedScope && NestedScope->Statement)
+	{
+		NestedScope->Statement->EmitHLSL(Context);
+	}
+
+	const int32 SizeofString = sizeof(TCHAR) * Length;
+	void* Memory = Context.Allocator->Alloc(sizeof(FCodeEntry) + SizeofString, alignof(FCodeEntry));
+	FCodeEntry* CodeEntry = new(Memory) FCodeEntry();
+	FMemory::Memcpy(CodeEntry->String, String, SizeofString);
+	CodeEntry->String[Length] = 0;
+	CodeEntry->Length = Length;
+	CodeEntry->NestedScope = NestedScope;
+	CodeEntry->Next = nullptr;
+
+	if (!List.First)
+	{
+		List.First = CodeEntry;
+		List.Last = CodeEntry;
+	}
+	else
+	{
+		List.Last->Next = CodeEntry;
+		List.Last = CodeEntry;
+	}
+	List.Num++;
+}
+
+void UE::HLSLTree::RequestScopeTypes(FUpdateTypeContext& Context, const FScope* InScope)
+{
+	if (InScope && InScope->Statement)
+	{
+		InScope->Statement->RequestTypes(Context);
+	}
+}
+
 UE::HLSLTree::FTree* UE::HLSLTree::FTree::Create(FMemStackBase& Allocator)
 {
 	FTree* Tree = new(Allocator) FTree();
@@ -656,84 +638,59 @@ static void WriteIndent(int32 IndentLevel, FStringBuilderBase& InOutString)
 	}
 }
 
-static void WriteScope(const UE::HLSLTree::FEmitScope& EmitScope, int32 IndentLevel, FStringBuilderBase& InOutString)
+void UE::HLSLTree::FScope::WriteHLSL(int32 Indent, FStringBuilderBase& OutString) const
 {
 	{
-		UE::HLSLTree::FEmitDeclaration::TConstIterator It(EmitScope.FirstDeclaration);
-		while (It)
+		const FCodeEntry* CodeDeclaration = Declarations.First;
+		while (CodeDeclaration)
 		{
-			const UE::HLSLTree::FEmitDeclaration& EmitDeclaration = *It;
-			const UE::Shader::FValueTypeDescription TypeDesc = UE::Shader::GetValueTypeDescription(EmitDeclaration.Type);
-			WriteIndent(IndentLevel, InOutString);
-			if (EmitDeclaration.Value)
-			{
-				InOutString.Appendf(TEXT("%s %s = %s;\n"), TypeDesc.Name, EmitDeclaration.Declaration, EmitDeclaration.Value);
-			}
-			else
-			{
-				InOutString.Appendf(TEXT("%s %s;\n"), TypeDesc.Name, EmitDeclaration.Declaration);
-			}
-			It.Next();
-		}
-	}
-
-	{
-		UE::HLSLTree::FEmitStatement::TConstIterator It(EmitScope.FirstStatement);
-		while (It)
-		{
-			const UE::HLSLTree::FEmitStatement& EmitStatement = *It;
-			WriteIndent(IndentLevel, InOutString);
-			InOutString.Append(EmitStatement.Code);
-			InOutString.Append(TCHAR('\n'));
-			It.Next();
+			check(!CodeDeclaration->NestedScope);
+			WriteIndent(Indent, OutString);
+			OutString.Append(CodeDeclaration->String, CodeDeclaration->Length);
+			OutString.Append(TEXT('\n'));
+			CodeDeclaration = CodeDeclaration->Next;
 		}
 	}
 
 	{
-		UE::HLSLTree::FEmitScopeLink::TConstIterator It(EmitScope.FirstLink);
-		while (It)
+		const FCodeEntry* CodeStatement = Statements.First;
+		while (CodeStatement)
 		{
-			const UE::HLSLTree::FEmitScopeLink& EmitLink = *It;
-			if (EmitLink.Code.Len() > 0)
+			if (CodeStatement->Length > 0)
 			{
-				WriteIndent(IndentLevel, InOutString);
-				InOutString.Append(EmitLink.Code);
-				InOutString.Append(TCHAR('\n'));
+				WriteIndent(Indent, OutString);
+				OutString.Append(CodeStatement->String, CodeStatement->Length);
+				OutString.Append(TEXT('\n'));
 			}
-
-			if (EmitLink.NextScope)
+			if (CodeStatement->NestedScope)
 			{
-				WriteIndent(IndentLevel, InOutString);
-				InOutString.Append(TEXT("{\n"));
-				WriteScope(*EmitLink.NextScope, IndentLevel + 1, InOutString);
-				WriteIndent(IndentLevel, InOutString);
-				InOutString.Append(TEXT("}\n"));
+				WriteIndent(Indent, OutString);
+				OutString.Append(TEXT("{\n"));
+				CodeStatement->NestedScope->WriteHLSL(Indent + 1, OutString);
+				WriteIndent(Indent, OutString);
+				OutString.Append(TEXT("}\n"));
 			}
-
-			It.Next();
+			CodeStatement = CodeStatement->Next;
 		}
 	}
 }
 
 bool UE::HLSLTree::FTree::EmitHLSL(UE::HLSLTree::FEmitContext& Context, FStringBuilderBase& Writer) const
 {
-	FEmitScope* EmitRootScope = Context.AcquireScope(*RootScope);
-	if (RootScope->EmitHLSL(Context, *EmitRootScope))
 	{
-		// Need to continue iterating FinalizeScope until it succeeds
-		bool bFinalizeResult = false;
-		while (!bFinalizeResult)
-		{
-			bFinalizeResult = Context.FinalizeScope(*EmitRootScope);
-		}
+		FUpdateTypeContext UpdateTypeContext(Context.Errors);
+		RequestScopeTypes(UpdateTypeContext, RootScope);
+	}
 
-		Context.Finalize();
+	if (RootScope->Statement)
+	{
+		RootScope->Statement->EmitHLSL(Context);
+	}
 
-		WriteScope(*EmitRootScope, 0, Writer);
+	Context.Finalize();
 
-		return true;
-	}
-	return false;
+	RootScope->WriteHLSL(0, Writer);
+	return true;
 }
 
 void UE::HLSLTree::FTree::RegisterExpression(FScope& Scope, FExpression* Expression)
diff --git a/Engine/Source/Runtime/Engine/Private/HLSLTree/HLSLTreeCommon.cpp b/Engine/Source/Runtime/Engine/Private/HLSLTree/HLSLTreeCommon.cpp
--- a/Engine/Source/Runtime/Engine/Private/HLSLTree/HLSLTreeCommon.cpp
+++ b/Engine/Source/Runtime/Engine/Private/HLSLTree/HLSLTreeCommon.cpp
@@ -73,18 +73,19 @@ UE::HLSLTree::FSwizzleParameters UE::HLSLTree::MakeSwizzleMask(bool bInR, bool b
 }
 
 
-bool UE::HLSLTree::FExpressionConstant::EmitCode(FEmitContext& Context, FExpressionEmitResult& OutResult) const
+bool UE::HLSLTree::FExpressionConstant::PrepareValue(FEmitContext& Context)
 {
-	OutResult.EvaluationType = EExpressionEvaluationType::Constant;
-	OutResult.Type = Value.GetType();
-	OutResult.Preshader.WriteOpcode(Shader::EPreshaderOpcode::Constant);
-	OutResult.Preshader.Write(Value);
-	return true;
+	return SetValueConstant(Context, Value);
+}
+
+bool UE::HLSLTree::FExpressionMaterialParameter::UpdateType(FUpdateTypeContext& Context, int8 InRequestedNumComponents)
+{
+	//return Shader::MakeValueTypeWithRequestedNumComponents(GetShaderValueType(ParameterType), InRequestedNumComponents);
+	return SetType(Context, GetShaderValueType(ParameterType));
 }
 
-bool UE::HLSLTree::FExpressionMaterialParameter::EmitCode(FEmitContext& Context, FExpressionEmitResult& OutResult) const
+bool UE::HLSLTree::FExpressionMaterialParameter::PrepareValue(FEmitContext& Context)
 {
-	OutResult.Type = GetShaderValueType(ParameterType);
 	if (ParameterType == EMaterialParameterType::StaticSwitch)
 	{
 		const FMaterialParameterInfo ParameterInfo(ParameterName);
@@ -97,8 +98,8 @@ bool UE::HLSLTree::FExpressionMaterialParameter::EmitCode(FEmitContext& Context,
 				break;
 			}
 		}
-		OutResult.EvaluationType = EExpressionEvaluationType::Constant;
-		OutResult.Preshader.WriteOpcode(Shader::EPreshaderOpcode::Constant).Write(Shader::FValue(bValue));
+
+		return SetValueConstant(Context, bValue);
 	}
 	else
 	{
@@ -115,36 +116,43 @@ bool UE::HLSLTree::FExpressionMaterialParameter::EmitCode(FEmitContext& Context,
 		}
 		const int32 ParameterIndex = Context.MaterialCompilationOutput->UniformExpressionSet.FindOrAddNumericParameter(ParameterType, ParameterName, DefaultOffset);
 		check(ParameterIndex >= 0 && ParameterIndex <= 0xffff);
-		OutResult.EvaluationType = EExpressionEvaluationType::Preshader;
-		OutResult.Preshader.WriteOpcode(Shader::EPreshaderOpcode::Parameter).Write((uint16)ParameterIndex); 
+
+		Shader::FPreshaderData LocalPreshader;
+		LocalPreshader.WriteOpcode(Shader::EPreshaderOpcode::Parameter).Write((uint16)ParameterIndex);
+		return SetValuePreshader(Context, LocalPreshader);
 	}
 	return true;
 }
 
-bool UE::HLSLTree::FExpressionExternalInput::EmitCode(FEmitContext& Context, FExpressionEmitResult& OutResult) const
+bool UE::HLSLTree::FExpressionExternalInput::PrepareValue(FEmitContext& Context)
 {
 	const int32 TypeIndex = (int32)InputType;
 	const int32 TexCoordIndex = TypeIndex - (int32)UE::HLSLTree::EExternalInputType::TexCoord0;
 
-	OutResult.EvaluationType = EExpressionEvaluationType::Shader;
-	OutResult.Type = GetInputExpressionType(InputType);
-	OutResult.bInline = true;
 	if (TexCoordIndex >= 0 && TexCoordIndex < 8)
 	{
 		Context.NumTexCoords = FMath::Max(Context.NumTexCoords, TexCoordIndex + 1);
-		OutResult.Code.Appendf(TEXT("Parameters.TexCoords[%u].xy"), TexCoordIndex);
+		return SetValueInlineShaderf(Context, TEXT("Parameters.TexCoords[%u].xy"), TexCoordIndex);
 	}
 	else
 	{
-		checkNoEntry();
+		return Context.Errors.AddError(this, TEXT("Invalid texcoord"));
 	}
 	return true;
 }
 
-bool UE::HLSLTree::FExpressionTextureSample::EmitCode(FEmitContext& Context, FExpressionEmitResult& OutResult) const
+bool UE::HLSLTree::FExpressionTextureSample::UpdateType(FUpdateTypeContext& Context, int8 InRequestedNumComponents)
 {
-	const FEmitValue* TexCoordValue = Context.AcquireValue(TexCoordExpression);
-	if (!TexCoordExpression)
+	if (RequestExpressionType(Context, TexCoordExpression, 2) == Shader::EValueType::Void)
+	{
+		return false;
+	}
+	return SetType(Context, Shader::EValueType::Float4);
+}
+
+bool UE::HLSLTree::FExpressionTextureSample::PrepareValue(FEmitContext& Context)
+{
+	if (PrepareExpressionValue(Context, TexCoordExpression) == EExpressionEvaluationType::None)
 	{
 		return false;
 	}
@@ -273,96 +281,144 @@ bool UE::HLSLTree::FExpressionTextureSample::EmitCode(FEmitContext& Context, FEx
 		}
 	}
 
-	OutResult.EvaluationType = EExpressionEvaluationType::Shader;
-	OutResult.Type = Shader::EValueType::Float4;
-	OutResult.Code.Appendf(TEXT("%s(%s(%s, %s, %s))"), SamplerTypeFunction, SampleFunctionName, *TextureName, *SamplerStateCode, Context.GetCode(TexCoordValue));
-	return true;
+	return SetValueShaderf(Context, TEXT("%s(%s(%s, %s, %s))"),
+		SamplerTypeFunction,
+		SampleFunctionName,
+		*TextureName,
+		*SamplerStateCode,
+		TexCoordExpression->GetValueShader(Context));
 }
 
-bool UE::HLSLTree::FExpressionDefaultMaterialAttributes::EmitCode(FEmitContext& Context, FExpressionEmitResult& OutResult) const
+bool UE::HLSLTree::FExpressionDefaultMaterialAttributes::PrepareValue(FEmitContext& Context)
 {
-	OutResult.EvaluationType = EExpressionEvaluationType::Shader;
-	OutResult.Type = Shader::EValueType::MaterialAttributes;
-	OutResult.bInline = true;
-	OutResult.Code.Append(TEXT("DefaultMaterialAttributes"));
-	return true;
+	return SetValueInlineShaderf(Context, TEXT("DefaultMaterialAttributes"));
+}
+
+bool UE::HLSLTree::FExpressionSetMaterialAttribute::UpdateType(FUpdateTypeContext& Context, int8 InRequestedNumComponents)
+{
+	const EMaterialValueType AttributeType = FMaterialAttributeDefinitionMap::GetValueType(AttributeID);
+	const uint32 NumComponents = GetNumComponents(AttributeType);
+	if (RequestExpressionType(Context, AttributesExpression, 4) != Shader::EValueType::MaterialAttributes)
+	{
+		return Context.Errors.AddError(this, TEXT("Expected type MaterialAttributes"));
+	}
+	if (RequestExpressionType(Context, ValueExpression, NumComponents) == Shader::EValueType::Void)
+	{
+		return false;
+	}
+	return SetType(Context, Shader::EValueType::MaterialAttributes);
 }
 
-bool UE::HLSLTree::FExpressionSetMaterialAttribute::EmitCode(FEmitContext& Context, FExpressionEmitResult& OutResult) const
+bool UE::HLSLTree::FExpressionSetMaterialAttribute::PrepareValue(FEmitContext& Context)
 {
-	const FEmitValue* AttributresValue = Context.AcquireValue(AttributesExpression);
-	const FEmitValue* Value = Context.AcquireValue(ValueExpression);
-	if (!AttributresValue || !Value)
+	if (PrepareExpressionValue(Context, AttributesExpression) == EExpressionEvaluationType::None)
+	{
+		return false;
+	}
+	if (PrepareExpressionValue(Context, ValueExpression) == EExpressionEvaluationType::None)
 	{
 		return false;
 	}
 
 	const FString PropertyName = FMaterialAttributeDefinitionMap::GetAttributeName(AttributeID);
-	OutResult.EvaluationType = EExpressionEvaluationType::Shader;
-	OutResult.Type = Shader::EValueType::MaterialAttributes;
-	OutResult.Code.Appendf(TEXT("FMaterialAttributes_Set%s(%s, %s)"), *PropertyName, Context.GetCode(AttributresValue), Context.GetCode(Value));
-	return true;
+	return SetValueShaderf(Context, TEXT("FMaterialAttributes_Set%s(%s, %s)"),
+		*PropertyName,
+		AttributesExpression->GetValueShader(Context),
+		ValueExpression->GetValueShader(Context));
 }
 
-bool UE::HLSLTree::FExpressionSelect::EmitCode(FEmitContext& Context, FExpressionEmitResult& OutResult) const
+bool UE::HLSLTree::FExpressionSelect::UpdateType(FUpdateTypeContext& Context, int8 InRequestedNumComponents)
 {
-	const FEmitValue* ConditionValue = Context.AcquireValue(ConditionExpression);
-	if (!ConditionValue)
+	const Shader::EValueType ConditionType = RequestExpressionType(Context, ConditionExpression, 1);
+	const Shader::EValueType LhsType = RequestExpressionType(Context, FalseExpression, InRequestedNumComponents);
+	const Shader::EValueType RhsType = RequestExpressionType(Context, TrueExpression, InRequestedNumComponents);
+	if (ConditionType == Shader::EValueType::Void || LhsType == Shader::EValueType::Void || RhsType == Shader::EValueType::Void)
 	{
 		return false;
 	}
 
-	if (ConditionValue->GetEvaluationType() == EExpressionEvaluationType::Constant)
+	const Shader::FValueTypeDescription LhsTypeDesc = Shader::GetValueTypeDescription(LhsType);
+	const Shader::FValueTypeDescription RhsTypeDesc = Shader::GetValueTypeDescription(RhsType);
+	if (LhsTypeDesc.ComponentType != RhsTypeDesc.ComponentType)
 	{
-		const bool bCondition = ConditionValue->GetConstantValue().AsBoolScalar();
-		if (!bCondition)
-		{
-			const FEmitValue* Value = Context.AcquireValue(FalseExpression);
-			OutResult.ForwardValue(Context, Value);
-		}
-		else
-		{
-			const FEmitValue* Value = Context.AcquireValue(TrueExpression);
-			OutResult.ForwardValue(Context, Value);
-		}
+		return Context.Errors.AddError(this, TEXT("Type mismatch"));
+	}
+
+	const int8 NumComponents = FMath::Max(LhsTypeDesc.NumComponents, RhsTypeDesc.NumComponents);
+	return SetType(Context, Shader::MakeValueType(LhsTypeDesc.ComponentType, NumComponents));
+}
+
+bool UE::HLSLTree::FExpressionSelect::PrepareValue(FEmitContext& Context)
+{
+	const EExpressionEvaluationType ConditionEvaluation = PrepareExpressionValue(Context, ConditionExpression);
+	if (ConditionEvaluation == EExpressionEvaluationType::None)
+	{
+		return false;
+	}
+
+	if (ConditionEvaluation == EExpressionEvaluationType::Constant)
+	{
+		const bool bCondition = ConditionExpression->GetValueConstant(Context).AsBoolScalar();
+		return SetValueForward(Context, bCondition ? TrueExpression : FalseExpression);
 	}
 	else
 	{
-		const FEmitValue* TrueValue = Context.AcquireValue(TrueExpression);
-		const FEmitValue* FalseValue = Context.AcquireValue(FalseExpression);
-		OutResult.EvaluationType = EExpressionEvaluationType::Shader; // TODO - preshader
-		OutResult.Type = TrueValue->GetExpressionType(); // TODO - FalseValue?
-		OutResult.Code.Appendf(TEXT("(%s ? %s : %s)"),
-			Context.GetCode(ConditionValue),
-			Context.GetCode(TrueValue),
-			Context.GetCode(FalseValue));
+		const EExpressionEvaluationType TrueEvaluation = PrepareExpressionValue(Context, TrueExpression);
+		const EExpressionEvaluationType FalseEvaluation = PrepareExpressionValue(Context, FalseExpression);
+		if (TrueEvaluation == EExpressionEvaluationType::None || FalseEvaluation == EExpressionEvaluationType::None)
+		{
+			return false;
+		}
+
+		// TODO - preshader
+		return SetValueShaderf(Context, TEXT("(%s ? %s : %s)"),
+			ConditionExpression->GetValueShader(Context),
+			TrueExpression->GetValueShader(Context),
+			FalseExpression->GetValueShader(Context));
 	}
 	return true;
 }
 
-bool UE::HLSLTree::FExpressionBinaryOp::EmitCode(FEmitContext& Context, FExpressionEmitResult& OutResult) const
+bool UE::HLSLTree::FExpressionBinaryOp::UpdateType(FUpdateTypeContext& Context, int8 InRequestedNumComponents)
 {
-	const FEmitValue* LhsValue = Context.AcquireValue(Lhs);
-	const FEmitValue* RhsValue = Context.AcquireValue(Rhs);
-	if (!LhsValue || !RhsValue)
+	const Shader::EValueType LhsType = RequestExpressionType(Context, Lhs, InRequestedNumComponents);
+	const Shader::EValueType RhsType = RequestExpressionType(Context, Rhs, InRequestedNumComponents);
+	if (LhsType == Shader::EValueType::Void || RhsType == Shader::EValueType::Void)
 	{
 		return false;
 	}
 
 	FString ErrorMessage;
+	Shader::EValueType ResultType;
 	switch (Op)
 	{
 	case EBinaryOp::Less:
-		OutResult.Type = Shader::MakeComparisonResultType(LhsValue->GetExpressionType(), RhsValue->GetExpressionType(), ErrorMessage);
+		ResultType = Shader::MakeComparisonResultType(LhsType, RhsType, ErrorMessage);
 		break;
 	default:
-		OutResult.Type = Shader::MakeArithmeticResultType(LhsValue->GetExpressionType(), RhsValue->GetExpressionType(), ErrorMessage);
+		ResultType = Shader::MakeArithmeticResultType(LhsType, RhsType, ErrorMessage);
 		break;
 	}
+	if (!ErrorMessage.IsEmpty())
+	{
+		return Context.Errors.AddError(this, ErrorMessage);
+	}
+
+	return SetType(Context, ResultType);
+}
+
+bool UE::HLSLTree::FExpressionBinaryOp::PrepareValue(FEmitContext& Context)
+{
+	const EExpressionEvaluationType LhsEvaluation = PrepareExpressionValue(Context, Lhs);
+	const EExpressionEvaluationType RhsEvaluation = PrepareExpressionValue(Context, Rhs);
+	if (LhsEvaluation == EExpressionEvaluationType::None || RhsEvaluation == EExpressionEvaluationType::None)
+	{
+		return false;
+	}
 
-	const EExpressionEvaluationType EvaluationType = CombineEvaluationTypes(LhsValue->GetEvaluationType(), RhsValue->GetEvaluationType());
+	const EExpressionEvaluationType Evaluation = CombineEvaluationTypes(LhsEvaluation, RhsEvaluation);
 	Shader::EPreshaderOpcode Opcode = Shader::EPreshaderOpcode::Nop;
-	if (EvaluationType != EExpressionEvaluationType::Shader)
+	if (Evaluation != EExpressionEvaluationType::Shader)
 	{
 		switch (Op)
 		{
@@ -375,31 +431,36 @@ bool UE::HLSLTree::FExpressionBinaryOp::EmitCode(FEmitContext& Context, FExpress
 	if (Opcode == Shader::EPreshaderOpcode::Nop)
 	{
 		// Either preshader not supported for our inputs, or not for this operation type
-		OutResult.EvaluationType = EExpressionEvaluationType::Shader;
-		const TCHAR* LhsCode = Context.GetCode(LhsValue);
-		const TCHAR* RhsCode = Context.GetCode(RhsValue);
+		const TCHAR* LhsCode = Lhs->GetValueShader(Context);
+		const TCHAR* RhsCode = Rhs->GetValueShader(Context);
 		switch (Op)
 		{
-		case EBinaryOp::Add: OutResult.Code.Appendf(TEXT("(%s + %s)"), LhsCode, RhsCode); break;
-		case EBinaryOp::Mul: OutResult.Code.Appendf(TEXT("(%s * %s)"), LhsCode, RhsCode); break;
-		case EBinaryOp::Less: OutResult.Code.Appendf(TEXT("(%s < %s)"), LhsCode, RhsCode); break;
-		default: checkNoEntry(); break;
+		case EBinaryOp::Add: return SetValueShaderf(Context, TEXT("(%s + %s)"), LhsCode, RhsCode);
+		case EBinaryOp::Mul: return SetValueShaderf(Context, TEXT("(%s * %s)"), LhsCode, RhsCode);
+		case EBinaryOp::Less: return SetValueShaderf(Context, TEXT("(%s < %s)"), LhsCode, RhsCode);
+		default: checkNoEntry(); return false;
 		}
 	}
 	else
 	{
-		Context.AppendPreshader(LhsValue, OutResult.Preshader);
-		Context.AppendPreshader(RhsValue, OutResult.Preshader);
-		OutResult.EvaluationType = EvaluationType;
-		OutResult.Preshader.WriteOpcode(Opcode);
+		Shader::FPreshaderData LocalPreshader;
+		Lhs->GetValuePreshader(Context, LocalPreshader);
+		Rhs->GetValuePreshader(Context, LocalPreshader);
+		LocalPreshader.WriteOpcode(Opcode);
+		return SetValuePreshader(Context, Evaluation, LocalPreshader);
 	}
-	return true;
 }
 
-bool UE::HLSLTree::FExpressionSwizzle::EmitCode(FEmitContext& Context, FExpressionEmitResult& OutResult) const
+bool UE::HLSLTree::FExpressionSwizzle::UpdateType(FUpdateTypeContext& Context, int8 InRequestedNumComponents)
 {
-	const FEmitValue* InputValue = Context.AcquireValue(Input);
-	if (!InputValue)
+	const Shader::EValueType InputType = RequestExpressionType(Context, Input, InRequestedNumComponents);
+	return SetType(Context, Shader::MakeValueType(InputType, Parameters.NumComponents));
+}
+
+bool UE::HLSLTree::FExpressionSwizzle::PrepareValue(FEmitContext& Context)
+{
+	const EExpressionEvaluationType InputEvaluation = PrepareExpressionValue(Context, Input);
+	if (InputEvaluation == EExpressionEvaluationType::None)
 	{
 		return false;
 	}
@@ -413,91 +474,92 @@ bool UE::HLSLTree::FExpressionSwizzle::EmitCode(FEmitContext& Context, FExpressi
 		Swizzle[i] = ComponentName[ComponentIndex];
 	}
 
-	OutResult.bInline = true;
-	OutResult.EvaluationType = InputValue->GetEvaluationType();
-	OutResult.Type = Shader::MakeValueType(InputValue->GetExpressionType(), Parameters.NumComponents);
-	if (OutResult.EvaluationType == EExpressionEvaluationType::Shader)
+	if (InputEvaluation == EExpressionEvaluationType::Shader)
 	{
-		OutResult.Code.Appendf(TEXT("%s.%s"), Context.GetCode(InputValue), Swizzle);
+		return SetValueInlineShaderf(Context, TEXT("%s.%s"), Input->GetValueShader(Context), Swizzle);
 	}
 	else
 	{
-		Context.AppendPreshader(InputValue, OutResult.Preshader);
-		OutResult.Preshader.WriteOpcode(Shader::EPreshaderOpcode::ComponentSwizzle)
+		Shader::FPreshaderData LocalPreshader;
+		Input->GetValuePreshader(Context, LocalPreshader);
+		LocalPreshader.WriteOpcode(Shader::EPreshaderOpcode::ComponentSwizzle)
 			.Write((uint8)Parameters.NumComponents)
 			.Write((uint8)Parameters.ComponentIndex[0])
 			.Write((uint8)Parameters.ComponentIndex[1])
 			.Write((uint8)Parameters.ComponentIndex[2])
 			.Write((uint8)Parameters.ComponentIndex[3]);
+		return SetValuePreshader(Context, InputEvaluation, LocalPreshader);
 	}
-	return true;
 }
 
-bool UE::HLSLTree::FExpressionAppend::EmitCode(FEmitContext& Context, FExpressionEmitResult& OutResult) const
+bool UE::HLSLTree::FExpressionAppend::UpdateType(FUpdateTypeContext& Context, int8 InRequestedNumComponents)
 {
-	const FEmitValue* LhsValue = Context.AcquireValue(Lhs);
-	const FEmitValue* RhsValue = Context.AcquireValue(Rhs);
-	if (!LhsValue || !RhsValue)
-	{
-		return false;
-	}
-
-	const Shader::FValueTypeDescription LhsTypeDesc = Shader::GetValueTypeDescription(LhsValue->GetExpressionType());
-	const Shader::FValueTypeDescription RhsTypeDesc = Shader::GetValueTypeDescription(RhsValue->GetExpressionType());
+	const Shader::EValueType LhsType = RequestExpressionType(Context, Lhs, 4);
+	const Shader::EValueType RhsType = RequestExpressionType(Context, Rhs, 4);
+	const Shader::FValueTypeDescription LhsTypeDesc = Shader::GetValueTypeDescription(LhsType);
+	const Shader::FValueTypeDescription RhsTypeDesc = Shader::GetValueTypeDescription(RhsType);
 	if (LhsTypeDesc.ComponentType != RhsTypeDesc.ComponentType)
 	{
-		return false;
+		return Context.Errors.AddError(this, TEXT("Type mismatch"));
 	}
 
-	const int32 NumComponents = LhsTypeDesc.NumComponents + RhsTypeDesc.NumComponents;
-	if (NumComponents > 4)
+	const int32 NumComponents = FMath::Min<int32>(LhsTypeDesc.NumComponents + RhsTypeDesc.NumComponents, InRequestedNumComponents);
+	return SetType(Context, Shader::MakeValueType(LhsTypeDesc.ComponentType, NumComponents));
+}
+
+bool UE::HLSLTree::FExpressionAppend::PrepareValue(FEmitContext& Context)
+{
+	const EExpressionEvaluationType LhsEvaluation = PrepareExpressionValue(Context, Lhs);
+	const EExpressionEvaluationType RhsEvaluation = PrepareExpressionValue(Context, Rhs);
+	if (LhsEvaluation == EExpressionEvaluationType::None || RhsEvaluation == EExpressionEvaluationType::None)
 	{
 		return false;
 	}
 
-	OutResult.EvaluationType = CombineEvaluationTypes(LhsValue->GetEvaluationType(), RhsValue->GetEvaluationType());
-	OutResult.Type = Shader::MakeValueType(LhsTypeDesc.ComponentType, NumComponents);
-	if (OutResult.EvaluationType == EExpressionEvaluationType::Shader)
+	const EExpressionEvaluationType Evaluation = CombineEvaluationTypes(LhsEvaluation, RhsEvaluation);
+	if (Evaluation == EExpressionEvaluationType::Shader)
 	{
-		const Shader::FValueTypeDescription ResultTypeDesc = Shader::GetValueTypeDescription(OutResult.Type);
-		OutResult.Code.Appendf(TEXT("%s(%s, %s)"), ResultTypeDesc.Name, Context.GetCode(LhsValue), Context.GetCode(RhsValue));
+		const Shader::FValueTypeDescription ResultTypeDesc = Shader::GetValueTypeDescription(GetValueType());
+		return SetValueShaderf(Context, TEXT("%s(%s, %s)"), ResultTypeDesc.Name, Lhs->GetValueShader(Context), Rhs->GetValueShader(Context));
 	}
 	else
 	{
-		Context.AppendPreshader(LhsValue, OutResult.Preshader);
-		Context.AppendPreshader(RhsValue, OutResult.Preshader);
-		OutResult.Preshader.WriteOpcode(Shader::EPreshaderOpcode::AppendVector);
+		Shader::FPreshaderData LocalPreshader;
+		Lhs->GetValuePreshader(Context, LocalPreshader);
+		Rhs->GetValuePreshader(Context, LocalPreshader);
+		LocalPreshader.WriteOpcode(Shader::EPreshaderOpcode::AppendVector);
+		return SetValuePreshader(Context, Evaluation, LocalPreshader);
 	}
-	return true;
 }
 
-bool UE::HLSLTree::FExpressionCast::EmitCode(FEmitContext& Context, FExpressionEmitResult& OutResult) const
+bool UE::HLSLTree::FExpressionCast::UpdateType(FUpdateTypeContext& Context, int8 InRequestedNumComponents)
 {
-	const FEmitValue* InputValue = Context.AcquireValue(Input);
-	if (!InputValue)
+	RequestExpressionType(Context, Input, Shader::GetValueTypeDescription(Type).NumComponents);
+	return SetType(Context, Type);
+}
+
+bool UE::HLSLTree::FExpressionCast::PrepareValue(FEmitContext& Context)
+{
+	const EExpressionEvaluationType InputEvaluation = PrepareExpressionValue(Context, Input);
+	if (InputEvaluation == EExpressionEvaluationType::None)
 	{
 		return false;
 	}
 
-	const Shader::EValueType InputType = InputValue->GetExpressionType();
+	const Shader::EValueType InputType = Input->GetValueType();
 	if (InputType == Type)
 	{
 		// No cast required (nop)
-		OutResult.ForwardValue(Context, InputValue);
-		return true;
+		return SetValueForward(Context, Input);
 	}
 
-	OutResult.bInline = true;
-	OutResult.Type = Type;
-	OutResult.EvaluationType = EExpressionEvaluationType::Shader; // TODO - preshader
-
-	const TCHAR* InputCode = Context.GetCode(InputValue);
+	const TCHAR* InputCode = Input->GetValueShader(Context);
 	const Shader::FValueTypeDescription OutputTypeDesc = Shader::GetValueTypeDescription(Type);
 	const Shader::FValueTypeDescription InputTypeDesc = Shader::GetValueTypeDescription(InputType);
 	if (InputTypeDesc.NumComponents == OutputTypeDesc.NumComponents)
 	{
 		// Cast between different underlying types with the same number of components, use a C-style cast
-		OutResult.Code.Appendf(TEXT("((%s)%s)"), OutputTypeDesc.Name, InputCode);
+		return SetValueInlineShaderf(Context, TEXT("((%s)%s)"), OutputTypeDesc.Name, InputCode);
 	}
 	else if (OutputTypeDesc.NumComponents < InputTypeDesc.NumComponents)
 	{
@@ -506,12 +568,12 @@ bool UE::HLSLTree::FExpressionCast::EmitCode(FEmitContext& Context, FExpressionE
 		if (OutputTypeDesc.ComponentType == InputTypeDesc.ComponentType)
 		{
 			// underlying types are the same, no C-style cast needed
-			OutResult.Code.Appendf(TEXT("%s.%s"), InputCode, SwizzleName[OutputTypeDesc.NumComponents]);
+			SetValueInlineShaderf(Context, TEXT("%s.%s"), InputCode, SwizzleName[OutputTypeDesc.NumComponents]);
 		}
 		else
 		{
 			// Different underlying types, add a C-style cast as well
-			OutResult.Code.Appendf(TEXT("((%s)%s.%s)"), OutputTypeDesc.Name, InputCode, SwizzleName[OutputTypeDesc.NumComponents]);
+			SetValueInlineShaderf(Context, TEXT("((%s)%s.%s)"), OutputTypeDesc.Name, InputCode, SwizzleName[OutputTypeDesc.NumComponents]);
 		}
 	}
 	else
@@ -521,72 +583,83 @@ bool UE::HLSLTree::FExpressionCast::EmitCode(FEmitContext& Context, FExpressionE
 		{
 			// Replicating a scalar value, use a swizzle
 			static const TCHAR* ScalarSwizzleName[] = { TEXT("x"), TEXT("xx"), TEXT("xxx"), TEXT("xxxx") };
-			OutResult.Code.Appendf(TEXT("%s.%s"), InputCode, ScalarSwizzleName[OutputTypeDesc.NumComponents]);
+			SetValueInlineShaderf(Context, TEXT("%s.%s"), InputCode, ScalarSwizzleName[OutputTypeDesc.NumComponents]);
 		}
 		else
 		{
 			// Padding with 0s, use a constructor
 			const Shader::FValue ZeroValue(OutputTypeDesc.ComponentType, 1);
-			OutResult.Code.Appendf(TEXT("%s(%s"), OutputTypeDesc.Name, InputCode);
+			TStringBuilder<2048> LocalCode;
+
+			LocalCode.Appendf(TEXT("%s(%s"), OutputTypeDesc.Name, InputCode);
 			for (int32 i = InputTypeDesc.NumComponents; i < OutputTypeDesc.NumComponents; ++i)
 			{
-				OutResult.Code.Appendf(TEXT(", %s"), *ZeroValue.ToString(UE::Shader::EValueStringFormat::HLSL));
+				LocalCode.Appendf(TEXT(", %s"), *ZeroValue.ToString(UE::Shader::EValueStringFormat::HLSL));
 			}
-			OutResult.Code.Append(TEXT(")"));
+			LocalCode.Append(TEXT(")"));
+			return SetValueInlineShader(Context, LocalCode);
 		}
 	}
 	return true;
 }
 
-bool UE::HLSLTree::FExpressionReflectionVector::EmitCode(FEmitContext& Context, FExpressionEmitResult& OutResult) const
+bool UE::HLSLTree::FExpressionReflectionVector::PrepareValue(FEmitContext& Context)
 {
-	OutResult.bInline = true;
-	OutResult.Type = Shader::EValueType::Float3;
-	OutResult.EvaluationType = EExpressionEvaluationType::Shader;
-	OutResult.Code.Append(TEXT("Parameters.ReflectionVector"));
-	return true;
+	return SetValueInlineShaderf(Context, TEXT("Parameters.ReflectionVector"));
 }
 
-bool UE::HLSLTree::FStatementBreak::EmitHLSL(FEmitContext& Context) const
+void UE::HLSLTree::FStatementBreak::EmitHLSL(FEmitContext& Context) const
 {
-	Context.WriteScopeTerminatorf(TEXT("break;"));
-	return true;
+	ParentScope->EmitStatementf(Context, TEXT("break;"));
+}
+
+void UE::HLSLTree::FStatementReturn::RequestTypes(FUpdateTypeContext& Context) const
+{
+	RequestExpressionType(Context, Expression, 4);
 }
 
-bool UE::HLSLTree::FStatementReturn::EmitHLSL(FEmitContext& Context) const
+void UE::HLSLTree::FStatementReturn::EmitHLSL(FEmitContext& Context) const
 {
-	const FEmitValue* Value = Context.AcquireValue(Expression);
-	if (!Value)
+	if (PrepareExpressionValue(Context, Expression) == EExpressionEvaluationType::None)
 	{
-		return false;
+		return;
 	}
 
-	Context.WriteScopeTerminatorf(TEXT("return %s;"), Context.GetCode(Value));
-	return true;
+	ParentScope->EmitStatementf(Context, TEXT("return %s;"), Expression->GetValueShader(Context));
 }
 
-bool UE::HLSLTree::FStatementIf::EmitHLSL(FEmitContext& Context) const
+void UE::HLSLTree::FStatementIf::RequestTypes(FUpdateTypeContext& Context) const
 {
-	const FEmitValue* ConditionValue = Context.AcquireValue(ConditionExpression);
-	if (!ConditionValue)
+	RequestExpressionType(Context, ConditionExpression, 1);
+	RequestScopeTypes(Context, ThenScope);
+	RequestScopeTypes(Context, ElseScope);
+	RequestScopeTypes(Context, NextScope);
+}
+
+void UE::HLSLTree::FStatementIf::EmitHLSL(FEmitContext& Context) const
+{
+	if (PrepareExpressionValue(Context, ConditionExpression) == EExpressionEvaluationType::None)
 	{
-		return false;
+		return;
 	}
 
-	Context.WriteScopef(*ThenScope, TEXT("if (%s)"), Context.GetCode(ConditionValue));
+	ParentScope->EmitNestedScopef(Context, ThenScope, TEXT("if (%s)"), ConditionExpression->GetValueShader(Context));
 	if (ElseScope)
 	{
-		Context.WriteScopef(*ElseScope, TEXT("else"));
-	}
-	if (NextScope)
-	{
-		Context.WriteScope(*NextScope);
+		ParentScope->EmitNestedScopef(Context, ElseScope, TEXT("else"));
 	}
-	return true;
+
+	ParentScope->EmitNestedScope(Context, NextScope);
+}
+
+void UE::HLSLTree::FStatementLoop::RequestTypes(FUpdateTypeContext& Context) const
+{
+	RequestScopeTypes(Context, LoopScope);
+	RequestScopeTypes(Context, NextScope);
 }
 
-bool UE::HLSLTree::FStatementLoop::EmitHLSL(FEmitContext& Context) const
+void UE::HLSLTree::FStatementLoop::EmitHLSL(FEmitContext& Context) const
 {
-	Context.WriteScopef(*LoopScope, TEXT("while (true)"));
-	return Context.WriteScope(*NextScope);
+	ParentScope->EmitNestedScopef(Context, LoopScope, TEXT("while (true)"));
+	ParentScope->EmitNestedScope(Context, NextScope);
 }
diff --git a/Engine/Source/Runtime/Engine/Private/Materials/HLSLMaterialTranslator.cpp b/Engine/Source/Runtime/Engine/Private/Materials/HLSLMaterialTranslator.cpp
--- a/Engine/Source/Runtime/Engine/Private/Materials/HLSLMaterialTranslator.cpp
+++ b/Engine/Source/Runtime/Engine/Private/Materials/HLSLMaterialTranslator.cpp
@@ -45,24 +45,6 @@ static inline bool IsDebugTextureSampleEnabled()
 	return IsAnalyticDerivEnabled() && (GDebugTextureSampleEnabled != 0);
 }
 
-/** @return the number of components in a vector type. */
-uint32 GetNumComponents(EMaterialValueType Type)
-{
-	switch(Type)
-	{
-		case MCT_Float:
-		case MCT_Float1: return 1;
-		case MCT_Float2: return 2;
-		case MCT_Float3: return 3;
-		case MCT_Float4: return 4;
-		case MCT_LWCScalar: return 1;
-		case MCT_LWCVector2: return 2;
-		case MCT_LWCVector3: return 3;
-		case MCT_LWCVector4: return 4;
-		default: return 0;
-	}
-}
-
 /** @return the vector type containing a given number of components. */
 static inline EMaterialValueType GetVectorType(uint32 NumComponents)
 {
diff --git a/Engine/Source/Runtime/Engine/Private/Materials/HLSLMaterialTranslator.h b/Engine/Source/Runtime/Engine/Private/Materials/HLSLMaterialTranslator.h
--- a/Engine/Source/Runtime/Engine/Private/Materials/HLSLMaterialTranslator.h
+++ b/Engine/Source/Runtime/Engine/Private/Materials/HLSLMaterialTranslator.h
@@ -69,8 +69,6 @@ public:
 	virtual EMaterialExpressionVisitResult Visit(UMaterialExpression* InExpression) = 0;
 };
 
-uint32 GetNumComponents(EMaterialValueType Type);
-
 struct FShaderCodeChunk
 {
 	/**
diff --git a/Engine/Source/Runtime/Engine/Private/Materials/MaterialHLSLEmitter.cpp b/Engine/Source/Runtime/Engine/Private/Materials/MaterialHLSLEmitter.cpp
--- a/Engine/Source/Runtime/Engine/Private/Materials/MaterialHLSLEmitter.cpp
+++ b/Engine/Source/Runtime/Engine/Private/Materials/MaterialHLSLEmitter.cpp
@@ -676,8 +676,7 @@ bool MaterialEmitHLSL(const FMaterialCompileTargetParameters& InCompilerTarget,
 
 	TStringBuilder<32 * 1024> Code;
 
-	UE::HLSLTree::FEmitContext EmitContext;
-	EmitContext.Allocator = &Allocator;
+	UE::HLSLTree::FEmitContext EmitContext(Allocator);
 	EmitContext.Material = &InMaterial;
 	EmitContext.StaticParameters = &InStaticParameters;
 	EmitContext.MaterialCompilationOutput = &OutCompilationOutput;
diff --git a/Engine/Source/Runtime/Engine/Private/Shader/ShaderValue.cpp b/Engine/Source/Runtime/Engine/Private/Shader/ShaderValue.cpp
--- a/Engine/Source/Runtime/Engine/Private/Shader/ShaderValue.cpp
+++ b/Engine/Source/Runtime/Engine/Private/Shader/ShaderValue.cpp
@@ -398,6 +398,12 @@ UE::Shader::EValueType UE::Shader::MakeValueType(EValueType BaseType, int32 NumC
 	return MakeValueType(GetValueTypeDescription(BaseType).ComponentType, NumComponents);
 }
 
+UE::Shader::EValueType UE::Shader::MakeValueTypeWithRequestedNumComponents(EValueType BaseType, int8 RequestedNumComponents)
+{
+	const FValueTypeDescription TypeDesc = GetValueTypeDescription(BaseType);
+	return MakeValueType(TypeDesc.ComponentType, FMath::Min(TypeDesc.NumComponents, RequestedNumComponents));
+}
+
 UE::Shader::EValueType UE::Shader::MakeArithmeticResultType(EValueType Lhs, EValueType Rhs, FString& OutErrorMessage)
 {
 	const FValueTypeDescription LhsDesc = GetValueTypeDescription(Lhs);
diff --git a/Engine/Source/Runtime/Engine/Public/HLSLTree/HLSLTree.h b/Engine/Source/Runtime/Engine/Public/HLSLTree/HLSLTree.h
--- a/Engine/Source/Runtime/Engine/Public/HLSLTree/HLSLTree.h
+++ b/Engine/Source/Runtime/Engine/Public/HLSLTree/HLSLTree.h
@@ -30,204 +30,67 @@ class FPreshaderData;
 namespace HLSLTree
 {
 
-class FExpressionLocalPHI;
-
-static constexpr int32 MaxNumPreviousScopes = 2;
-
 class FScope;
-class FEmitScope;
 
-class FEmitDeclaration : public TIntrusiveLinkedList<FEmitDeclaration>
-{
-public:
-	const TCHAR* Declaration;
-	const TCHAR* Value;
-	Shader::EValueType Type;
-};
+static constexpr int32 MaxNumPreviousScopes = 2;
 
-class FEmitAssignment : public TIntrusiveLinkedList<FEmitAssignment>
+struct FError
 {
-public:
-	const TCHAR* Declaration;
-	FExpression* Expression;
+	const FError* Next;
+	const FNode* Node;
+	int32 MessageLength;
+	const TCHAR Message[1];
 };
 
-class FEmitStatement : public TIntrusiveLinkedList<FEmitStatement>
+class FErrors
 {
 public:
-	FStringView Code;
-};
+	explicit FErrors(FMemStackBase& InAllocator);
 
-class FEmitScopeLink : public TIntrusiveLinkedList<FEmitScopeLink>
-{
-public:
-	FEmitScope* NextScope;
-	FStringView Code;
-};
+	bool AddError(const FNode* InNode, FStringView InError);
 
-class FEmitScope
-{
-public:
-	FEmitScope* ParentScope = nullptr;
-	FEmitDeclaration* FirstDeclaration = nullptr;
-	FEmitAssignment* FirstAssignment = nullptr;
-	FEmitStatement* FirstStatement = nullptr;
-	FEmitStatement* LastStatement = nullptr;
-	FEmitScopeLink* FirstLink = nullptr;
-	FEmitScopeLink* LastLink = nullptr;
-	TMap<FSHAHash, const TCHAR*> ExpressionMap;
+private:
+	FMemStackBase* Allocator = nullptr;
+	const FError* FirstError = nullptr;
 };
 
-class FEmitValue
+class FUpdateTypeContext
 {
 public:
-	EExpressionEvaluationType GetEvaluationType() const { return EvaluationType; }
-	Shader::EValueType GetExpressionType() const { return ExpressionType; }
-	const Shader::FValue& GetConstantValue() const { return ConstantValue; }
+	explicit FUpdateTypeContext(FErrors& InErrors) : Errors(InErrors) {}
 
-private:
-	mutable const TCHAR* Code = nullptr;
-	const Shader::FPreshaderData* Preshader = nullptr;
-	EExpressionEvaluationType EvaluationType = EExpressionEvaluationType::None;
-	Shader::EValueType ExpressionType = Shader::EValueType::Void;
-	Shader::FValue ConstantValue;
-
-	friend class FEmitContext;
+	FErrors& Errors;
 };
 
 /** Tracks shared state while emitting HLSL code */
 class FEmitContext
 {
 public:
-	FEmitContext();
+	explicit FEmitContext(FMemStackBase& InAllocator);
 	~FEmitContext();
 
-	/** Returns a value that references the given expression */
-	const FEmitValue* AcquireValue(FExpression* Expression);
+	void Finalize();
 
 	/** Get a unique local variable name */
 	const TCHAR* AcquireLocalDeclarationCode();
 
-	FStringView InternalAcquireInternedString(const TCHAR* String, int32 Length);
-
-	FStringView AcquireInternedString(const TCHAR* Format, ...);
-
-	template<int32 ArrayLength>
-	FStringView AcquireInternedString(TCHAR (&String)[ArrayLength])
-	{
-		return InternalAcquireInternedString(String, ArrayLength);
-	}
-
-	/** Gets HLSL code that references the given value */
-	const TCHAR* GetCode(const FEmitValue* Value);
-
-	/** Append preshader bytecode that represents the given value */
-	void AppendPreshader(const FEmitValue* Value, Shader::FPreshaderData& InOutPreshader);
-
-	FEmitScope* FindScope(const FScope& Scope);
-	FEmitScope* AcquireScope(const FScope& Scope);
-
-	FEmitScope& GetCurrentScope();
-
-	void InternalWriteStatementToScope(FEmitScope& EmitScope, FStringView InternedCode);
-	bool InternalWriteScope(const FScope& Scope, FStringView InternedCode);
-	FEmitScopeLink* InternalWriteScopeLink(FStringView InternedCode);
-
-	template<typename StringType>
-	void WriteStatementToScope(FEmitScope& EmitScope, const StringType& String)
-	{
-		InternalWriteStatementToScope(EmitScope, AcquireInternedString(String));
-	}
-
-	template<typename FormatType, typename... Types>
-	void WriteStatementToScopef(FEmitScope& EmitScope, const FormatType& Format, Types... Args)
-	{
-		InternalWriteStatementToScope(EmitScope, AcquireInternedString(Format, Forward<Types>(Args)...));
-	}
-
-	template<typename StringType>
-	void WriteStatement(const StringType& String)
-	{
-		InternalWriteStatementToScope(GetCurrentScope(), AcquireInternedString(String));
-	}
-
-	template<typename FormatType, typename... Types>
-	void WriteStatementf(const FormatType& Format, Types... Args)
-	{
-		InternalWriteStatementToScope(GetCurrentScope(), AcquireInternedString(Format, Forward<Types>(Args)...));
-	}
-
-	bool WriteScope(const FScope& Scope)
-	{
-		return InternalWriteScope(Scope, FStringView());
-	}
-
-	template<typename FormatType, typename... Types>
-	bool WriteScopef(const FScope& Scope, const FormatType& Format, Types... Args)
-	{
-		return InternalWriteScope(Scope, AcquireInternedString(Format, Forward<Types>(Args)...));
-	}
-
-	template<typename FormatType, typename... Types>
-	void WriteScopeTerminatorf(const FormatType& Format, Types... Args)
-	{
-		InternalWriteScopeLink(AcquireInternedString(Format, Forward<Types>(Args)...));
-	}
-
-	void WriteDeclaration(FEmitScope& EmitScope, Shader::EValueType Type, const TCHAR* Declaration, const TCHAR* Value = nullptr);
-
-	bool WriteAssignment(FEmitScope& EmitScope, const TCHAR* Declaration, FExpression* Expression, Shader::EValueType& InOutType);
-
-	bool FinalizeScope(FEmitScope& EmitScope);
-	void Finalize();
-
-	struct FDeclarationEntry
-	{
-		FEmitValue Value;
-	};
-
-	struct FFunctionCallEntry
-	{
-		const FEmitValue* OutputValues;
-		int32 NumOutputs;
-	};
-
-	TArray<FEmitScope*> ScopeStack;
-	TMap<const FScope*, FEmitScope*> ScopeMap;
+	void AddPreshader(Shader::EValueType Type, const Shader::FPreshaderData& Preshader, FStringBuilderBase& OutCode);
 
-	TMap<FNode*, FDeclarationEntry*> DeclarationMap;
-	TArray<Shader::FPreshaderData*> TempPreshaders;
-	TSet<const FExpression*> PendingEmitValueExpressions;
 	FMemStackBase* Allocator = nullptr;
-	const FMaterial* Material = nullptr; // TODO - remove preshader material dependency
+	FErrors Errors;
+
+	// TODO - remove preshader material dependency
+	const FMaterial* Material = nullptr;
 	const FStaticParameterSet* StaticParameters = nullptr;
 	FMaterialCompilationOutput* MaterialCompilationOutput = nullptr;
 	TMap<Shader::FValue, uint32> DefaultUniformValues;
 	uint32 UniformPreshaderOffset = 0u;
-	int32 TotalCodeLength = 0;
+
 	int32 NumExpressionLocals = 0;
+	int32 NumLocalPHIs = 0;
 	int32 NumTexCoords = 0;
 };
 
-struct FExpressionEmitResult
-{
-	FExpressionEmitResult(FStringBuilderBase& InCode, Shader::FPreshaderData& InPreshader)
-		: Code(InCode)
-		, Preshader(InPreshader)
-		, EvaluationType(EExpressionEvaluationType::None)
-		, Type(Shader::EValueType::Void)
-		, bInline(false)
-	{}
-
-	void ForwardValue(FEmitContext& Context, const FEmitValue* InValue);
-
-	FStringBuilderBase& Code;
-	Shader::FPreshaderData& Preshader;
-	EExpressionEvaluationType EvaluationType;
-	Shader::EValueType Type;
-	bool bInline;
-};
-
 enum class ENodeVisitResult
 {
 	VisitDependentNodes,
@@ -271,8 +134,8 @@ class FStatement : public FNode
 public:
 	virtual ENodeVisitResult Visit(FNodeVisitor& Visitor) override;
 
-	/** Emits HLSL code for the statement. The generated code should include any required semi-colons and newlines */
-	virtual bool EmitHLSL(FEmitContext& Context) const = 0;
+	virtual void RequestTypes(FUpdateTypeContext& Context) const = 0;
+	virtual void EmitHLSL(FEmitContext& Context) const = 0;
 
 	FScope* ParentScope = nullptr;
 };
@@ -287,14 +150,64 @@ public:
 class FExpression : public FNode
 {
 public:
+	Shader::EValueType GetValueType() const { return ValueType; }
+
 	virtual ENodeVisitResult Visit(FNodeVisitor& Visitor) override;
 
-	/** Emits code for the given expression, either HLSL code or preshader bytecode */
-	virtual bool EmitCode(FEmitContext& Context, FExpressionEmitResult& OutResult) const = 0;
+	friend Shader::EValueType RequestExpressionType(FUpdateTypeContext& Context, FExpression* InExpression, int8 InRequestedNumComponents);
+	friend EExpressionEvaluationType PrepareExpressionValue(FEmitContext& Context, FExpression* InExpression);
+
+	const TCHAR* GetValueShader(FEmitContext& Context);
+	void GetValuePreshader(FEmitContext& Context, Shader::FPreshaderData& OutPreshader);
+	Shader::FValue GetValueConstant(FEmitContext& Context);
 
 	FScope* ParentScope = nullptr;
+
+protected:
+	virtual bool UpdateType(FUpdateTypeContext& Context, int8 InRequestedNumComponents) = 0;
+	virtual bool PrepareValue(FEmitContext& Context) = 0;
+
+	bool SetType(FUpdateTypeContext& Context, Shader::EValueType InType) { ValueType = InType; return true; }
+
+	bool InternalSetValueShader(FEmitContext& Context, const TCHAR* InCode, int32 InLength, bool bInline);
+
+	bool SetValuePreshader(FEmitContext& Context, Shader::FPreshaderData& InPreshader);
+	bool SetValueConstant(FEmitContext& Context, const Shader::FValue& InValue);
+	bool SetValueForward(FEmitContext& Context, FExpression* Source);
+
+	bool SetValuePreshader(FEmitContext& Context, EExpressionEvaluationType InEvaluationType, Shader::FPreshaderData& InPreshader);
+
+	template<typename FormatType, typename... Types>
+	bool SetValueShaderf(FEmitContext& Context, const FormatType& Format, Types... Args)
+	{
+		TStringBuilder<2048> String;
+		String.Appendf(Format, Forward<Types>(Args)...);
+		return InternalSetValueShader(Context, String.GetData(), String.Len(), false);
+	}
+
+	template<typename FormatType, typename... Types>
+	bool SetValueInlineShaderf(FEmitContext& Context, const FormatType& Format, Types... Args)
+	{
+		TStringBuilder<2048> String;
+		String.Appendf(Format, Forward<Types>(Args)...);
+		return InternalSetValueShader(Context, String.GetData(), String.Len(), true);
+	}
+
+	inline bool SetValueShader(FEmitContext& Context, const FStringBuilderBase& String) { return InternalSetValueShader(Context, String.GetData(), String.Len(), false); }
+	inline bool SetValueInlineShader(FEmitContext& Context, const FStringBuilderBase& String) { return InternalSetValueShader(Context, String.GetData(), String.Len(), true); }
+
+private:
+	const TCHAR* LocalVariableName = nullptr;
+	const TCHAR* Code = nullptr;
+	const Shader::FPreshaderData* Preshader = nullptr;
+	Shader::FValue ConstantValue;
+	EExpressionEvaluationType EvaluationType = EExpressionEvaluationType::None;
+	Shader::EValueType ValueType = Shader::EValueType::Void;
+	int8 RequestedNumComponents = 0;
+	bool bReentryFlag = false;
 };
 
+
 /**
  * Represents a phi node (see various topics on single static assignment)
  * A phi node takes on a value based on the previous scope that was executed.
@@ -303,7 +216,8 @@ public:
 class FExpressionLocalPHI final : public FExpression
 {
 public:
-	virtual bool EmitCode(FEmitContext& Context, FExpressionEmitResult& OutResult) const override;
+	virtual bool UpdateType(FUpdateTypeContext& Context, int8 InRequestedNumComponents) override;
+	virtual bool PrepareValue(FEmitContext& Context) override;
 
 	FName LocalName;
 	FScope* Scopes[MaxNumPreviousScopes];
@@ -344,19 +258,74 @@ public:
 	bool HasParentScope(const FScope& ParentScope) const;
 
 	virtual ENodeVisitResult Visit(FNodeVisitor& Visitor) override;
-	bool EmitHLSL(FEmitContext& Context, FEmitScope& Scope) const;
 
 	void AddPreviousScope(FScope& Scope);
 
 	void UseExpression(FExpression* Expression);
 
+	friend void RequestScopeTypes(FUpdateTypeContext& Context, const FScope* InScope);
+
+	template<typename FormatType, typename... Types>
+	void EmitDeclarationf(FEmitContext& Context, const FormatType& Format, Types... Args)
+	{
+		InternalEmitCodef(Context, Declarations, nullptr, Format, Forward<Types>(Args)...);
+	}
+
+	template<typename FormatType, typename... Types>
+	void EmitStatementf(FEmitContext& Context, const FormatType& Format, Types... Args)
+	{
+		InternalEmitCodef(Context, Statements, nullptr, Format, Forward<Types>(Args)...);
+	}
+
+	void EmitNestedScope(FEmitContext& Context, FScope* NestedScope)
+	{
+		InternalEmitCode(Context, Statements, NestedScope, nullptr, 0);
+	}
+
+	template<typename FormatType, typename... Types>
+	void EmitNestedScopef(FEmitContext& Context, FScope* NestedScope, const FormatType& Format, Types... Args)
+	{
+		InternalEmitCodef(Context, Statements, NestedScope, Format, Forward<Types>(Args)...);
+	}
+
+	void WriteHLSL(int32 Indent, FStringBuilderBase& OutString) const;
+
 private:
 	friend class FTree;
+	friend class FExpression;
 	friend class FNodeVisitor_MoveToScope;
 
+	struct FCodeEntry
+	{
+		FCodeEntry* Next;
+		FScope* NestedScope;
+		int32 Length;
+		TCHAR String[1];
+	};
+
+	struct FCodeList
+	{
+		FCodeEntry* First = nullptr;
+		FCodeEntry* Last = nullptr;
+		int32 Num = 0;
+	};
+
+	void InternalEmitCode(FEmitContext& Context, FCodeList& List, FScope* NestedScope, const TCHAR* String, int32 Length);
+
+	template<typename FormatType, typename... Types>
+	void InternalEmitCodef(FEmitContext& Context, FCodeList& List, FScope* NestedScope, const FormatType& Format, Types... Args)
+	{
+		TStringBuilder<2048> String;
+		String.Appendf(Format, Forward<Types>(Args)...);
+		InternalEmitCode(Context, List, NestedScope, String.GetData(), String.Len());
+	}
+
 	FScope* ParentScope = nullptr;
 	FStatement* Statement = nullptr;
 	FScope* PreviousScope[MaxNumPreviousScopes];
+	TMap<FSHAHash, const TCHAR*> ExpressionCodeMap;
+	FCodeList Declarations;
+	FCodeList Statements;
 	int32 NumPreviousScopes = 0;
 	int32 NestedLevel = 0;
 };
diff --git a/Engine/Source/Runtime/Engine/Public/HLSLTree/HLSLTreeCommon.h b/Engine/Source/Runtime/Engine/Public/HLSLTree/HLSLTreeCommon.h
--- a/Engine/Source/Runtime/Engine/Public/HLSLTree/HLSLTreeCommon.h
+++ b/Engine/Source/Runtime/Engine/Public/HLSLTree/HLSLTreeCommon.h
@@ -45,7 +45,8 @@ public:
 
 	Shader::FValue Value;
 
-	virtual bool EmitCode(FEmitContext& Context, FExpressionEmitResult& OutResult) const override;
+	virtual bool UpdateType(FUpdateTypeContext& Context, int8 InRequestedNumComponents) override { return SetType(Context, Shader::MakeValueTypeWithRequestedNumComponents(Value.GetType(), InRequestedNumComponents)); }
+	virtual bool PrepareValue(FEmitContext& Context) override;
 };
 
 class FExpressionMaterialParameter : public FExpression
@@ -59,7 +60,8 @@ public:
 	Shader::FValue DefaultValue;
 	EMaterialParameterType ParameterType;
 
-	virtual bool EmitCode(FEmitContext& Context, FExpressionEmitResult& OutResult) const override;
+	virtual bool UpdateType(FUpdateTypeContext& Context, int8 InRequestedNumComponents) override;
+	virtual bool PrepareValue(FEmitContext& Context) override;
 };
 
 enum class EExternalInputType
@@ -90,7 +92,8 @@ public:
 
 	EExternalInputType InputType;
 
-	virtual bool EmitCode(FEmitContext& Context, FExpressionEmitResult& OutResult) const override;
+	virtual bool UpdateType(FUpdateTypeContext& Context, int8 InRequestedNumComponents) override { return SetType(Context, GetInputExpressionType(InputType)); }
+	virtual bool PrepareValue(FEmitContext& Context) override;
 };
 
 class FExpressionTextureSample : public FExpression
@@ -119,7 +122,8 @@ public:
 		return Result;
 	}
 
-	virtual bool EmitCode(FEmitContext& Context, FExpressionEmitResult& OutResult) const override;
+	virtual bool UpdateType(FUpdateTypeContext& Context, int8 InRequestedNumComponents) override;
+	virtual bool PrepareValue(FEmitContext& Context) override;
 };
 
 class FExpressionDefaultMaterialAttributes : public FExpression
@@ -127,7 +131,8 @@ class FExpressionDefaultMaterialAttributes : public FExpression
 public:
 	FExpressionDefaultMaterialAttributes() {}
 
-	virtual bool EmitCode(FEmitContext& Context, FExpressionEmitResult& OutResult) const override;
+	virtual bool UpdateType(FUpdateTypeContext& Context, int8 InRequestedNumComponents) override { return SetType(Context, Shader::EValueType::MaterialAttributes); }
+	virtual bool PrepareValue(FEmitContext& Context) override;
 };
 
 class FExpressionSetMaterialAttribute : public FExpression
@@ -150,7 +155,8 @@ public:
 		return Result;
 	}
 
-	virtual bool EmitCode(FEmitContext& Context, FExpressionEmitResult& OutResult) const override;
+	virtual bool UpdateType(FUpdateTypeContext& Context, int8 InRequestedNumComponents) override;
+	virtual bool PrepareValue(FEmitContext& Context) override;
 };
 
 class FExpressionSelect : public FExpression
@@ -178,7 +184,8 @@ public:
 		return Result;
 	}
 
-	virtual bool EmitCode(FEmitContext& Context, FExpressionEmitResult& OutResult) const override;
+	virtual bool UpdateType(FUpdateTypeContext& Context, int8 InRequestedNumComponents) override;
+	virtual bool PrepareValue(FEmitContext& Context) override;
 };
 
 class FExpressionBinaryOp : public FExpression
@@ -205,7 +212,8 @@ public:
 		return Result;
 	}
 
-	virtual bool EmitCode(FEmitContext& Context, FExpressionEmitResult& OutResult) const override;
+	virtual bool UpdateType(FUpdateTypeContext& Context, int8 InRequestedNumComponents) override;
+	virtual bool PrepareValue(FEmitContext& Context) override;
 };
 
 struct FSwizzleParameters
@@ -239,7 +247,8 @@ public:
 		return Result;
 	}
 
-	virtual bool EmitCode(FEmitContext& Context, FExpressionEmitResult& OutResult) const override;
+	virtual bool UpdateType(FUpdateTypeContext& Context, int8 InRequestedNumComponents) override;
+	virtual bool PrepareValue(FEmitContext& Context) override;
 };
 
 class FExpressionAppend : public FExpression
@@ -264,7 +273,8 @@ public:
 		return Result;
 	}
 
-	virtual bool EmitCode(FEmitContext& Context, FExpressionEmitResult& OutResult) const override;
+	virtual bool UpdateType(FUpdateTypeContext& Context, int8 InRequestedNumComponents) override;
+	virtual bool PrepareValue(FEmitContext& Context) override;
 };
 
 class FExpressionCast : public FExpression
@@ -290,13 +300,15 @@ public:
 		return Result;
 	}
 
-	virtual bool EmitCode(FEmitContext& Context, FExpressionEmitResult& OutResult) const override;
+	virtual bool UpdateType(FUpdateTypeContext& Context, int8 InRequestedNumComponents) override;
+	virtual bool PrepareValue(FEmitContext& Context) override;
 };
 
 class FExpressionReflectionVector : public FExpression
 {
 public:
-	virtual bool EmitCode(FEmitContext& Context, FExpressionEmitResult& OutResult) const override;
+	virtual bool UpdateType(FUpdateTypeContext& Context, int8 InRequestedNumComponents) override { return SetType(Context, Shader::EValueType::Float3); }
+	virtual bool PrepareValue(FEmitContext& Context) override;
 };
 
 class FStatementReturn : public FStatement
@@ -314,13 +326,15 @@ public:
 		return Result;
 	}
 
-	virtual bool EmitHLSL(FEmitContext& Context) const override;
+	virtual void RequestTypes(FUpdateTypeContext& Context) const override;
+	virtual void EmitHLSL(FEmitContext& Context) const override;
 };
 
 class FStatementBreak : public FStatement
 {
 public:
-	virtual bool EmitHLSL(FEmitContext& Context) const override;
+	virtual void RequestTypes(FUpdateTypeContext& Context) const override {}
+	virtual void EmitHLSL(FEmitContext& Context) const override;
 };
 
 class FStatementIf : public FStatement
@@ -344,7 +358,8 @@ public:
 		return Result;
 	}
 
-	virtual bool EmitHLSL(FEmitContext& Context) const override;
+	virtual void RequestTypes(FUpdateTypeContext& Context) const override;
+	virtual void EmitHLSL(FEmitContext& Context) const override;
 };
 
 class FStatementLoop : public FStatement
@@ -364,7 +379,8 @@ public:
 		return Result;
 	}
 
-	virtual bool EmitHLSL(FEmitContext& Context) const override;
+	virtual void RequestTypes(FUpdateTypeContext& Context) const override;
+	virtual void EmitHLSL(FEmitContext& Context) const override;
 };
 
 }
diff --git a/Engine/Source/Runtime/Engine/Public/MaterialShared.h b/Engine/Source/Runtime/Engine/Public/MaterialShared.h
--- a/Engine/Source/Runtime/Engine/Public/MaterialShared.h
+++ b/Engine/Source/Runtime/Engine/Public/MaterialShared.h
@@ -173,6 +173,24 @@ enum EMaterialValueType
 	MCT_VoidStatement         = 1 << 22,
 };
 
+/** @return the number of components in a vector type. */
+inline uint32 GetNumComponents(EMaterialValueType Type)
+{
+	switch (Type)
+	{
+	case MCT_Float:
+	case MCT_Float1: return 1;
+	case MCT_Float2: return 2;
+	case MCT_Float3: return 3;
+	case MCT_Float4: return 4;
+	case MCT_LWCScalar: return 1;
+	case MCT_LWCVector2: return 2;
+	case MCT_LWCVector3: return 3;
+	case MCT_LWCVector4: return 4;
+	default: return 0;
+	}
+}
+
 inline bool IsLWCType(EMaterialValueType InType)
 {
 	return (InType & MCT_LWCType);
diff --git a/Engine/Source/Runtime/Engine/Public/Shader/ShaderTypes.h b/Engine/Source/Runtime/Engine/Public/Shader/ShaderTypes.h
--- a/Engine/Source/Runtime/Engine/Public/Shader/ShaderTypes.h
+++ b/Engine/Source/Runtime/Engine/Public/Shader/ShaderTypes.h
@@ -60,6 +60,7 @@ uint32 GetComponentTypeSizeInBytes(EValueComponentType Type);
 FValueTypeDescription GetValueTypeDescription(EValueType Type);
 EValueType MakeValueType(EValueComponentType ComponentType, int32 NumComponents);
 EValueType MakeValueType(EValueType BaseType, int32 NumComponents);
+EValueType MakeValueTypeWithRequestedNumComponents(EValueType BaseType, int8 RequestedNumComponents);
 EValueType MakeArithmeticResultType(EValueType Lhs, EValueType Rhs, FString& OutErrorMessage);
 EValueType MakeComparisonResultType(EValueType Lhs, EValueType Rhs, FString& OutErrorMessage);
 
