[VisualLogger] improve flushing of entries
- Entries are now flushed once per frame by a tick function. This improves the output for client/server logs in PIE and removes the need to perform a manual refresh to get the last entries in the UI (e.g. unit test, editor command, end of simulation)
- Log entry name is now prefixed with NetMode for client/server worlds (unchanged for Standalone)
- Fixed FVisualLoggerBinaryFileDevice that was not flushing its cache properly

diff --git a/Engine/Source/Runtime/Engine/Private/VisualLogger/VisualLogger.cpp b/Engine/Source/Runtime/Engine/Private/VisualLogger/VisualLogger.cpp
--- a/Engine/Source/Runtime/Engine/Private/VisualLogger/VisualLogger.cpp
+++ b/Engine/Source/Runtime/Engine/Private/VisualLogger/VisualLogger.cpp
@@ -160,49 +160,48 @@ FVisualLogEntry* FVisualLogger::GetLastEntryForObject(const UObject* Object)
 	return CurrentEntryPerObject.Contains(LogOwner) ? &CurrentEntryPerObject[LogOwner] : nullptr;
 }
 
-FVisualLogEntry* FVisualLogger::GetEntryToWrite(const UObject* Object, float TimeStamp, ECreateIfNeeded ShouldCreate)
+FVisualLogEntry* FVisualLogger::GetEntryToWrite(const UObject* Object, const float TimeStamp, const ECreateIfNeeded ShouldCreate)
 {
-	FVisualLogEntry* CurrentEntry = nullptr;
-	UObject * LogOwner = FVisualLogger::FindRedirection(Object);
+	const UObject* const LogOwner = FindRedirection(Object);
 	if (LogOwner == nullptr)
 	{
 		return nullptr;
 	}
 
-	bool bInitializeNewEntry = false;
+	// Entry can be created or reused (after being flushed) and will need to be initialized
+	bool bInitializeEntry = false;
+	FVisualLogEntry* CurrentEntry = nullptr;
 
-	const UWorld* World = GetWorldForVisualLogger(LogOwner);
 	if (CurrentEntryPerObject.Contains(LogOwner))
 	{
 		CurrentEntry = &CurrentEntryPerObject[LogOwner];
-		if (CurrentEntry->bIsAllowedToLog)
+
+		// We serialize and reinitialize entries only when allowed to log and parameter
+		// indicates that new entry can be created. Otherwise we simply return current entry.  
+		if (CurrentEntry->bIsAllowedToLog && ShouldCreate == ECreateIfNeeded::Create)
 		{
-			bInitializeNewEntry = (TimeStamp > CurrentEntry->TimeStamp) && (ShouldCreate == ECreateIfNeeded::Create);
-			if (World && IsInGameThread())
+			// Same LogOwner can be used for logs at different time in the frame so need to flush entry right away
+			// Other entries will be flushed in batch using FlushEntries
+			if (CurrentEntry->bIsInitialized && TimeStamp > CurrentEntry->TimeStamp)
 			{
-				for (auto& CurrentPair : CurrentEntryPerObject)
-				{
-					FVisualLogEntry* Entry = &CurrentPair.Value;
-					if (Entry->TimeStamp >= 0 && Entry->TimeStamp < TimeStamp)
-					{
-						for (FVisualLogDevice* Device : OutputDevices)
-						{
-							Device->Serialize(CurrentPair.Key.ResolveObjectPtrEvenIfPendingKill(), ObjectToNameMap[CurrentPair.Key], ObjectToClassNameMap[CurrentPair.Key], *Entry);
-						}
-						Entry->Reset();
-					}
-				}
+				FlushEntry(*CurrentEntry, LogOwner);
 			}
+	
+			bInitializeEntry = !CurrentEntry->bIsInitialized;
 		}
 	}
-
-	if (CurrentEntry == nullptr)
+	else if (ShouldCreate == ECreateIfNeeded::Create)
 	{
 		// It's first and only one usage of LogOwner as regular object to get names. We assume once that LogOwner is correct here and only here.
 		CurrentEntry = &CurrentEntryPerObject.Add(LogOwner);
-		const FName LogName(bForceUniqueLogNames  
-			? *FString::Printf(TEXT("%s [%d]"), *LogOwner->GetName(), LogOwner->GetUniqueID()) 
-			: *FString::Printf(TEXT("%s"), *LogOwner->GetName()));
+
+		const UWorld* World = GetWorldForVisualLogger(LogOwner);
+		const bool bIsStandalone = (World == nullptr || World->GetNetMode() == NM_Standalone);
+		const FName LogName(*FString::Printf(TEXT("%s%s%s"),
+			bIsStandalone ? TEXT("") : *FString::Printf(TEXT("(%s) "), *ToString(World->GetNetMode())),
+			*LogOwner->GetName(),
+			bForceUniqueLogNames ? *FString::Printf(TEXT(" [%d]"), LogOwner->GetUniqueID()) : TEXT("")));
+
 		ObjectToNameMap.Add(LogOwner, LogName);
 		ObjectToClassNameMap.Add(LogOwner, *(LogOwner->GetClass()->GetName()));
 		ObjectToWorldMap.Add(LogOwner, World);
@@ -213,16 +212,16 @@ FVisualLogEntry* FVisualLogger::GetEntryToWrite(const UObject* Object, float Tim
 		CurrentEntry->bPassedObjectAllowList = IsObjectAllowed(LogOwner);
 		CurrentEntry->UpdateAllowedToLog();
 
-		bInitializeNewEntry = CurrentEntry->bIsAllowedToLog;
+		bInitializeEntry = CurrentEntry->bIsAllowedToLog;
 	}
 
-	if (bInitializeNewEntry)
+	if (bInitializeEntry)
 	{
 		CurrentEntry->Reset();
 		CurrentEntry->TimeStamp = TimeStamp;
+		CurrentEntry->bIsInitialized = true;
 
-		const AActor* ObjectAsActor = Cast<AActor>(LogOwner);
-		if (ObjectAsActor)
+		if (const AActor* ObjectAsActor = Cast<AActor>(LogOwner))
 		{
 			CurrentEntry->Location = ObjectAsActor->GetActorLocation();
 		}
@@ -230,55 +229,70 @@ FVisualLogEntry* FVisualLogger::GetEntryToWrite(const UObject* Object, float Tim
 		FOwnerToChildrenRedirectionMap& RedirectionMap = GetRedirectionMap(LogOwner);
 		if (RedirectionMap.Contains(LogOwner))
 		{
-			if (CurrentEntryPerObject.Contains(LogOwner) && LogOwner->IsValidLowLevel())
+			if (const IVisualLoggerDebugSnapshotInterface* DebugSnapshotInterface = Cast<const IVisualLoggerDebugSnapshotInterface>(LogOwner))
 			{
-				const IVisualLoggerDebugSnapshotInterface* DebugSnapshotInterface = Cast<const IVisualLoggerDebugSnapshotInterface>(LogOwner);
-				if (DebugSnapshotInterface)
-				{
-					DebugSnapshotInterface->GrabDebugSnapshot(CurrentEntry);
-				}
+				DebugSnapshotInterface->GrabDebugSnapshot(CurrentEntry);
 			}
 			for (TWeakObjectPtr<const UObject>& Child : RedirectionMap[LogOwner])
 			{
-				if (Child.IsValid())
+				if (const IVisualLoggerDebugSnapshotInterface* DebugSnapshotInterface = Cast<const IVisualLoggerDebugSnapshotInterface>(Child.Get()))
 				{
-					const IVisualLoggerDebugSnapshotInterface* DebugSnapshotInterface = Cast<const IVisualLoggerDebugSnapshotInterface>(Child.Get());
-					if (DebugSnapshotInterface)
-					{
-						DebugSnapshotInterface->GrabDebugSnapshot(CurrentEntry);
-					}
+					DebugSnapshotInterface->GrabDebugSnapshot(CurrentEntry);
 				}
 			}
 		}
 		else
 		{
-			const IVisualLoggerDebugSnapshotInterface* DebugSnapshotInterface = Cast<const IVisualLoggerDebugSnapshotInterface>(LogOwner);
-			if (DebugSnapshotInterface)
+			if (const IVisualLoggerDebugSnapshotInterface* DebugSnapshotInterface = Cast<const IVisualLoggerDebugSnapshotInterface>(LogOwner))
 			{
 				DebugSnapshotInterface->GrabDebugSnapshot(CurrentEntry);
 			}
 		}
 	}
 
-	return CurrentEntry->bIsAllowedToLog ? CurrentEntry : nullptr;
+	if (CurrentEntry->bIsAllowedToLog)
+	{
+		bIsFlushRequired = true;
+	}
+	else
+	{
+		CurrentEntry = nullptr;
+	}
+
+	return CurrentEntry;
 }
 
+void FVisualLogger::Tick(float DeltaTime)
+{
+	if (bIsFlushRequired)
+	{
+		Flush();
+		bIsFlushRequired = false;
+	}
+}
 
 void FVisualLogger::Flush()
 {
 	for (auto &CurrentEntry : CurrentEntryPerObject)
 	{
-		if (CurrentEntry.Value.TimeStamp >= 0)
+		if (CurrentEntry.Value.bIsInitialized)
 		{
-			for (FVisualLogDevice* Device : OutputDevices)
-			{
-				Device->Serialize(CurrentEntry.Key.ResolveObjectPtrEvenIfPendingKill(), ObjectToNameMap[CurrentEntry.Key], ObjectToClassNameMap[CurrentEntry.Key], CurrentEntry.Value);
-			}
-			CurrentEntry.Value.Reset();
+			FlushEntry(CurrentEntry.Value, CurrentEntry.Key);
 		}
 	}
 }
 
+void FVisualLogger::FlushEntry(FVisualLogEntry& Entry, const FObjectKey& ObjectKey)
+{
+	ensureMsgf(Entry.bIsInitialized, TEXT("FlushEntry should only be called with an initialized entry."));
+
+	const UObject* OwnerObject = ObjectKey.ResolveObjectPtrEvenIfPendingKill();
+	for (FVisualLogDevice* Device : OutputDevices)
+	{
+		Device->Serialize(OwnerObject, ObjectToNameMap[ObjectKey], ObjectToClassNameMap[ObjectKey], Entry);
+	}
+	Entry.Reset();
+}
 
 void FVisualLogger::EventLog(const UObject* Object, const FName EventTag1, const FVisualLogEventBase& Event1, const FVisualLogEventBase& Event2, const FVisualLogEventBase& Event3, const FVisualLogEventBase& Event4, const FVisualLogEventBase& Event5, const FVisualLogEventBase& Event6)
 {
@@ -382,6 +396,7 @@ FVisualLogger::FVisualLogger()
 	bForceUniqueLogNames = true;
 	bIsRecordingToFile = false;
 	bIsRecordingToTrace = false;
+	bIsFlushRequired = false;
 
 	BlockAllCategories(false);
 	AddDevice(&FVisualLoggerBinaryFileDevice::Get());
@@ -393,6 +408,15 @@ FVisualLogger::FVisualLogger()
 		SetIsRecording(true);
 		SetIsRecordingToFile(true);
 	}
+
+	TickerHandle = FTSTicker::GetCoreTicker().AddTicker(TEXT("VisualLogger"), 0.0f, [this](const float DeltaTime)
+	{
+		QUICK_SCOPE_CYCLE_COUNTER(STAT_FVisualLogger_Tick);
+
+		Tick(DeltaTime);
+
+		return true;
+	});
 }
 
 void FVisualLogger::Shutdown()
@@ -402,7 +426,7 @@ void FVisualLogger::Shutdown()
 	RemoveDevice(&FVisualLoggerBinaryFileDevice::Get());
 }
 
-void FVisualLogger::Cleanup(UWorld* OldWorld, bool bReleaseMemory)
+void FVisualLogger::Cleanup(UWorld* OldWorld, const bool bReleaseMemory)
 {
 	const bool WasRecordingToFile = IsRecordingToFile();
 	if (WasRecordingToFile)
@@ -473,7 +497,7 @@ void FVisualLogger::Cleanup(UWorld* OldWorld, bool bReleaseMemory)
 	}
 }
 
-int32 FVisualLogger::GetUniqueId(float Timestamp)
+int32 FVisualLogger::GetUniqueId(const float Timestamp)
 {
 	return LastUniqueIds.FindOrAdd(Timestamp)++;
 }
@@ -551,7 +575,7 @@ UObject* FVisualLogger::FindRedirection(const UObject* Object)
 	return const_cast<UObject*>(TargetWeakPtr.Get(/*bEvenIfPendingKill*/true));
 }
 
-void FVisualLogger::SetIsRecording(bool InIsRecording)
+void FVisualLogger::SetIsRecording(const bool InIsRecording)
 {
 	if (InIsRecording == false && InIsRecording != !!bIsRecording && FParse::Param(FCommandLine::Get(), TEXT("LogNavOctree")))
 	{
@@ -564,7 +588,7 @@ void FVisualLogger::SetIsRecording(bool InIsRecording)
 	bIsRecording = InIsRecording;
 };
 
-void FVisualLogger::SetIsRecordingToFile(bool InIsRecording)
+void FVisualLogger::SetIsRecordingToFile(const bool InIsRecording)
 {
 	if (!bIsRecording && InIsRecording)
 	{
@@ -604,7 +628,7 @@ void FVisualLogger::SetIsRecordingToFile(bool InIsRecording)
 	bIsRecordingToFile = InIsRecording;
 }
 
-void FVisualLogger::SetIsRecordingToTrace(bool InIsRecording)
+void FVisualLogger::SetIsRecordingToTrace(const bool InIsRecording)
 {
 	if (!bIsRecording && InIsRecording)
 	{
diff --git a/Engine/Source/Runtime/Engine/Private/VisualLogger/VisualLoggerTypes.cpp b/Engine/Source/Runtime/Engine/Private/VisualLogger/VisualLoggerTypes.cpp
--- a/Engine/Source/Runtime/Engine/Private/VisualLogger/VisualLoggerTypes.cpp
+++ b/Engine/Source/Runtime/Engine/Private/VisualLogger/VisualLoggerTypes.cpp
@@ -133,6 +133,7 @@ void FVisualLogEntry::Reset()
 	ElementsToDraw.Reset();
 	HistogramSamples.Reset();
 	DataBlocks.Reset();
+	bIsInitialized = false;
 }
 
 void FVisualLogEntry::UpdateAllowedToLog()
diff --git a/Engine/Source/Runtime/Engine/Public/VisualLogger/VisualLogger.h b/Engine/Source/Runtime/Engine/Public/VisualLogger/VisualLogger.h
--- a/Engine/Source/Runtime/Engine/Public/VisualLogger/VisualLogger.h
+++ b/Engine/Source/Runtime/Engine/Public/VisualLogger/VisualLogger.h
@@ -12,6 +12,7 @@
 #include "Templates/AndOrNot.h"
 #include "Templates/IsArrayOrRefOfType.h"
 #include "UObject/ObjectKey.h"
+#include "Containers/Ticker.h"
 
 #if ENABLE_VISUAL_LOG
 
@@ -43,8 +44,8 @@
 #define UE_VLOG_OBOX(LogOwner, CategoryName, Verbosity, Box, Matrix, Color, Format, ...) if(FVisualLogger::IsRecording()) FVisualLogger::GeometryBoxLogf(LogOwner, CategoryName, ELogVerbosity::Verbosity, Box, Matrix, Color, Format, ##__VA_ARGS__)
 #define UE_CVLOG_OBOX(Condition, LogOwner, CategoryName, Verbosity, Box, Matrix, Color, Format, ...)  if(FVisualLogger::IsRecording() && Condition) {UE_VLOG_OBOX(LogOwner, CategoryName, Verbosity, Box, Matrix, Color, Format, ##__VA_ARGS__);} 
 // Cone shape
-#define UE_VLOG_CONE(LogOwner, CategoryName, Verbosity, Orgin, Direction, Length, Angle, Color, Format, ...) if(FVisualLogger::IsRecording()) FVisualLogger::GeometryShapeLogf(LogOwner, CategoryName, ELogVerbosity::Verbosity, Orgin, Direction, Length, Angle, Color, Format, ##__VA_ARGS__)
-#define UE_CVLOG_CONE(Condition, LogOwner, CategoryName, Verbosity, Orgin, Direction, Length, Angle, Color, Format, ...)  if(FVisualLogger::IsRecording() && Condition) {UE_VLOG_CONE(LogOwner, CategoryName, Verbosity, Orgin, Direction, Length, Angle, Color, Format, ##__VA_ARGS__);} 
+#define UE_VLOG_CONE(LogOwner, CategoryName, Verbosity, Origin, Direction, Length, Angle, Color, Format, ...) if(FVisualLogger::IsRecording()) FVisualLogger::GeometryShapeLogf(LogOwner, CategoryName, ELogVerbosity::Verbosity, Orgin, Direction, Length, Angle, Color, Format, ##__VA_ARGS__)
+#define UE_CVLOG_CONE(Condition, LogOwner, CategoryName, Verbosity, Origin, Direction, Length, Angle, Color, Format, ...)  if(FVisualLogger::IsRecording() && Condition) {UE_VLOG_CONE(LogOwner, CategoryName, Verbosity, Orgin, Direction, Length, Angle, Color, Format, ##__VA_ARGS__);} 
 // Cylinder shape
 #define UE_VLOG_CYLINDER(LogOwner, CategoryName, Verbosity, Start, End, Radius, Color, Format, ...) if(FVisualLogger::IsRecording()) FVisualLogger::GeometryShapeLogf(LogOwner, CategoryName, ELogVerbosity::Verbosity, Start, End, Radius, Color, Format, ##__VA_ARGS__)
 #define UE_CVLOG_CYLINDER(Condition, LogOwner, CategoryName, Verbosity, Start, End, Radius, Color, Format, ...)  if(FVisualLogger::IsRecording() && Condition) {UE_VLOG_CYLINDER(LogOwner, CategoryName, Verbosity, Start, End, Radius, Color, Format, ##__VA_ARGS__);} 
@@ -455,7 +456,7 @@ public:
 	static UObject* FindRedirection(const UObject* Object);
 
 	/** blocks all categories from logging. It can be bypassed with the category allow list */
-	void BlockAllCategories(bool InBlock) { bBlockedAllCategories = InBlock; }
+	void BlockAllCategories(const bool bInBlock) { bBlockedAllCategories = bInBlock; }
 
 	/** checks if all categories are blocked */
 	bool IsBlockedForAllCategories() const { return !!bBlockedAllCategories; }
@@ -473,7 +474,7 @@ public:
 	int32 GetUniqueId(float Timestamp);
 
 	/** Starts visual log collecting and recording */
-	void SetIsRecording(bool InIsRecording);
+	void SetIsRecording(const bool bInIsRecording);
 	/** return information is vlog recording is enabled or not */
 	FORCEINLINE static bool IsRecording() { return !!bIsRecording; }
 
@@ -485,13 +486,13 @@ public:
 	void DiscardRecordingToFile();
 
 	/** Starts visual log collecting and recording to insights traces (for Rewind Debugger)*/
-	void SetIsRecordingToTrace(bool InIsRecording);
+	void SetIsRecordingToTrace(const bool bInIsRecording);
 
-	void SetIsRecordingOnServer(bool IsRecording) { bIsRecordingOnServer = IsRecording; }
+	void SetIsRecordingOnServer(const bool bInIsRecording) { bIsRecordingOnServer = bInIsRecording; }
 	bool IsRecordingOnServer() const { return !!bIsRecordingOnServer; }
 
-	/** COnfigure whether VisLog should be using decorated, unique names */
-	void SetUseUniqueNames(bool bEnable) { bForceUniqueLogNames = bEnable; }
+	/** Configure whether VisLog should be using decorated, unique names */
+	void SetUseUniqueNames(const bool bEnable) { bForceUniqueLogNames = bEnable; }
 
 	/** Add visual logger output device */
 	void AddDevice(FVisualLogDevice* InDevice) { OutputDevices.AddUnique(InDevice); }
@@ -501,7 +502,7 @@ public:
 	const TArray<FVisualLogDevice*>& GetDevices() const { return OutputDevices; }
 	/** Check if log category can be recorded, verify before using GetEntryToWrite! */
 	bool IsCategoryLogged(const FLogCategoryBase& Category) const;
-	/** Returns  current entry for given TimeStap or creates another one  but first it serialize previous 
+	/** Returns  current entry for given TimeStamp or creates another one  but first it serialize previous 
 	 *	entry as completed to vislog devices. Use VisualLogger::DontCreate to get current entry without serialization
 	 *	@note this function can return null */
 	FVisualLogEntry* GetEntryToWrite(const UObject* Object, float TimeStamp, ECreateIfNeeded ShouldCreate = ECreateIfNeeded::Create);
@@ -511,7 +512,7 @@ public:
 	/** flush and serialize data if timestamp allows it */
 	virtual void Flush() override;
 
-	/** FileName getter to set project specific file name for vlogs - highly encouradged to use FVisualLogFilenameGetterDelegate::CreateUObject with this */
+	/** FileName getter to set project specific file name for vlogs - highly encouraged to use FVisualLogFilenameGetterDelegate::CreateUObject with this */
 	void SetLogFileNameGetter(const FVisualLogFilenameGetterDelegate& InLogFileNameGetter) { LogFileNameGetter = InLogFileNameGetter; }
 
 	/** Register extension to use by LogVisualizer  */
@@ -565,6 +566,18 @@ private:
 	virtual void Serialize(const TCHAR* V, ELogVerbosity::Type Verbosity, const FName& Category) override { ensureMsgf(0, TEXT("Regular serialize is forbiden for visual logs")); }
 
 protected:
+	/** Flushes entries recorded in the frame */
+	void Tick(float DeltaTime);
+
+	/**
+	 * Serializes a single entry and resets it.
+	 * Method expects an initialized entry and will ensure otherwise. 
+	 */
+	void FlushEntry(FVisualLogEntry& Entry, const FObjectKey& ObjectKey);
+
+	/** Handle to the registered ticker to flush entries */
+	FTSTicker::FDelegateHandle TickerHandle;
+	
 	/** Array of output devices to redirect to */
 	TArray<FVisualLogDevice*> OutputDevices;
 	// Map for inter-objects redirections
@@ -595,7 +608,7 @@ protected:
 	// Cached map to world information because it's just raw pointer and not real object
 	FObjectToWorldMapType ObjectToWorldMap;
 	// for any object that has requested redirection this map holds where we should
-	// redirect the trafic to
+	// redirect the traffic to
 	FChildToOwnerRedirectionMap ChildToOwnerMap;
 	// if set all categories are blocked from logging
 	bool bBlockedAllCategories : 1;
@@ -608,6 +621,8 @@ protected:
 	// controls how we generate log names. When set to TRUE there's a lower 
 	// chance of name conflict, but it's more expensive
 	bool bForceUniqueLogNames : 1;
+	/** Indicates that entries were added/updated and that a flush is required */
+	bool bIsFlushRequired : 1;
 	// start recording time
 	float StartRecordingToFileTime;
 	/** Delegate to set project specific file name for vlogs */
diff --git a/Engine/Source/Runtime/Engine/Public/VisualLogger/VisualLoggerTypes.h b/Engine/Source/Runtime/Engine/Public/VisualLogger/VisualLoggerTypes.h
--- a/Engine/Source/Runtime/Engine/Public/VisualLogger/VisualLoggerTypes.h
+++ b/Engine/Source/Runtime/Engine/Public/VisualLogger/VisualLoggerTypes.h
@@ -166,6 +166,7 @@ struct ENGINE_API FVisualLogEntry
 	uint8 bPassedClassAllowList : 1;
 	uint8 bPassedObjectAllowList : 1;	
 	uint8 bIsAllowedToLog : 1;
+	uint8 bIsInitialized : 1;
 
 	TArray<FVisualLogEvent> Events;
 	TArray<FVisualLogLine> LogLines;
