Update the way casting is handled within HLSLTree
No more explicit cast nodes
Instead, expression code can be requested with a specific type, and a cast will be generated if needed

diff --git a/Engine/Source/Runtime/Engine/Private/HLSLTree/HLSLTree.cpp b/Engine/Source/Runtime/Engine/Private/HLSLTree/HLSLTree.cpp
--- a/Engine/Source/Runtime/Engine/Private/HLSLTree/HLSLTree.cpp
+++ b/Engine/Source/Runtime/Engine/Private/HLSLTree/HLSLTree.cpp
@@ -20,24 +20,17 @@ UE::HLSLTree::EExpressionEvaluationType UE::HLSLTree::CombineEvaluationTypes(EEx
 	return EExpressionEvaluationType::Preshader;
 }
 
-template <typename FormatType, typename... ArgTypes>
-const TCHAR* AllocateStringf(FMemStackBase& Allocator, const FormatType& Format, ArgTypes... Args)
+const TCHAR* AllocateString(FMemStackBase& Allocator, const TCHAR* String, int32 Length)
 {
-	TCHAR Buffer[1024];
-	const int32 Length = FCString::Snprintf(Buffer, 1024, Format, Forward<ArgTypes>(Args)...);
-	check(Length > 0);
 	TCHAR* Result = New<TCHAR>(Allocator, Length + 1);
-	FMemory::Memcpy(Result, Buffer, Length * sizeof(TCHAR));
+	FMemory::Memcpy(Result, String, Length * sizeof(TCHAR));
 	Result[Length] = 0;
 	return Result;
 }
 
-const TCHAR* AllocateString(FMemStackBase& Allocator, const TCHAR* String, int32 Length)
+const TCHAR* AllocateString(FMemStackBase& Allocator, FStringView String)
 {
-	TCHAR* Result = New<TCHAR>(Allocator, Length + 1);
-	FMemory::Memcpy(Result, String, Length * sizeof(TCHAR));
-	Result[Length] = 0;
-	return Result;
+	return AllocateString(Allocator, String.GetData(), String.Len());
 }
 
 const TCHAR* AllocateString(FMemStackBase& Allocator, const FStringBuilderBase& StringBuilder)
@@ -45,6 +38,14 @@ const TCHAR* AllocateString(FMemStackBase& Allocator, const FStringBuilderBase&
 	return AllocateString(Allocator, StringBuilder.GetData(), StringBuilder.Len());
 }
 
+template <typename FormatType, typename... ArgTypes>
+const TCHAR* AllocateStringf(FMemStackBase& Allocator, const FormatType& Format, ArgTypes... Args)
+{
+	TStringBuilder<1024> String;
+	String.Appendf(Format, Forward<ArgTypes>(Args)...);
+	return AllocateString(Allocator, String);
+}
+
 FSHAHash HashString(const TCHAR* String, int32 Length)
 {
 	FSHAHash Hash;
@@ -83,6 +84,138 @@ const TCHAR* UE::HLSLTree::FEmitContext::AcquireLocalDeclarationCode()
 	return AllocateStringf(*Allocator, TEXT("Local%d"), NumExpressionLocals++);
 }
 
+const TCHAR* UE::HLSLTree::FEmitContext::CastShaderValue(const FNode* Node, const TCHAR* Code, Shader::EValueType SourceType, Shader::EValueType DestType, ECastFlags Flags)
+{
+	if (SourceType == DestType)
+	{
+		return Code;
+	}
+
+	const bool bAllowTruncate = EnumHasAnyFlags(Flags, ECastFlags::AllowTruncate);
+	const bool bAllowAppendZeroes = EnumHasAnyFlags(Flags, ECastFlags::AllowAppendZeroes);
+	bool bReplicateScalar = EnumHasAnyFlags(Flags, ECastFlags::ReplicateScalar);
+
+	const Shader::FValueTypeDescription SourceTypeDesc = Shader::GetValueTypeDescription(SourceType);
+	const Shader::FValueTypeDescription DestTypeDesc = Shader::GetValueTypeDescription(DestType);
+
+	if (SourceTypeDesc.NumComponents > 0 && DestTypeDesc.NumComponents > 0)
+	{
+		if (SourceTypeDesc.NumComponents != 1)
+		{
+			bReplicateScalar = false;
+		}
+		if (!bReplicateScalar && !bAllowAppendZeroes && DestTypeDesc.NumComponents > SourceTypeDesc.NumComponents)
+		{
+			Errors.AddErrorf(Node, TEXT("Cannot cast from smaller type %s to larger type %s."), SourceTypeDesc.Name, DestTypeDesc.Name);
+			return TEXT("");
+		}
+		if (!bReplicateScalar && !bAllowTruncate && DestTypeDesc.NumComponents < SourceTypeDesc.NumComponents)
+		{
+			Errors.AddErrorf(Node, TEXT("Cannot cast from larger type %s to smaller type %s."), SourceTypeDesc.Name, DestTypeDesc.Name);
+			return TEXT("");
+		}
+
+		const bool bIsSourceLWC = SourceTypeDesc.ComponentType == Shader::EValueComponentType::Double;
+		const bool bIsLWC = DestTypeDesc.ComponentType == Shader::EValueComponentType::Double;
+		if (bIsLWC != bIsSourceLWC)
+		{
+			if (bIsLWC)
+			{
+				// float->LWC
+				const TCHAR* CodeAsFloat = CastShaderValue(Node, Code, SourceType, Shader::MakeValueType(Shader::EValueComponentType::Float, DestTypeDesc.NumComponents), Flags);
+				return AllocateStringf(*Allocator, TEXT("LWCPromote(%s)"), CodeAsFloat);
+			}
+			else
+			{
+				//LWC->float
+				const TCHAR* CodeAsFloat = AllocateStringf(*Allocator, TEXT("LWCToFloat(%s)"), Code);
+				return CastShaderValue(Node, CodeAsFloat, Shader::MakeValueType(Shader::EValueComponentType::Float, SourceTypeDesc.NumComponents), DestType, Flags);
+			}
+		}
+
+		TStringBuilder<1024> Result;
+		int32 NumComponents = 0;
+		bool bNeedClosingParen = false;
+		if (bIsLWC)
+		{
+			Result.Append(TEXT("MakeLWCVector("));
+			bNeedClosingParen = true;
+		}
+		else
+		{
+			if (bReplicateScalar || SourceTypeDesc.NumComponents == DestTypeDesc.NumComponents)
+			{
+				NumComponents = DestTypeDesc.NumComponents;
+				// Cast the scalar to the correct type, HLSL language will replicate the scalar if needed when performing this cast
+				Result.Appendf(TEXT("((%s)%s)"), DestTypeDesc.Name, Code);
+			}
+			else
+			{
+				NumComponents = FMath::Min(SourceTypeDesc.NumComponents, DestTypeDesc.NumComponents);
+				if (NumComponents < DestTypeDesc.NumComponents)
+				{
+					Result.Appendf(TEXT("%s("), DestTypeDesc.Name);
+					bNeedClosingParen = true;
+				}
+				if (NumComponents == SourceTypeDesc.NumComponents && SourceTypeDesc.ComponentType == DestTypeDesc.ComponentType)
+				{
+					// If we're taking all the components from the source, can avoid adding a swizzle
+					Result.Append(Code);
+				}
+				else
+				{
+					// Use a cast to truncate the source to the correct number of types
+					const Shader::EValueType IntermediateType = Shader::MakeValueType(DestTypeDesc.ComponentType, NumComponents);
+					const Shader::FValueTypeDescription IntermediateTypeDesc = Shader::GetValueTypeDescription(IntermediateType);
+					Result.Appendf(TEXT("((%s)%s)"), IntermediateTypeDesc.Name, Code);
+				}
+			}
+		}
+
+		if (bNeedClosingParen)
+		{
+			const Shader::FValue ZeroValue(DestTypeDesc.ComponentType, 1);
+			for (int32 ComponentIndex = NumComponents; ComponentIndex < DestTypeDesc.NumComponents; ++ComponentIndex)
+			{
+				if (ComponentIndex > 0u)
+				{
+					Result.Append(TEXT(","));
+				}
+				if (bIsLWC)
+				{
+					if (!bReplicateScalar && ComponentIndex >= SourceTypeDesc.NumComponents)
+					{
+						check(bAllowAppendZeroes);
+						Result.Append(TEXT("LWCPromote(0.0f)"));
+					}
+					else
+					{
+						Result.Appendf(TEXT("LWCGetComponent(%s, %d)"), Code, bReplicateScalar ? 0 : ComponentIndex);
+					}
+				}
+				else
+				{
+					// Non-LWC case should only be zero-filling here, other cases should have already been handled
+					check(bAllowAppendZeroes);
+					check(!bReplicateScalar);
+					check(ComponentIndex >= SourceTypeDesc.NumComponents);
+					ZeroValue.ToString(Shader::EValueStringFormat::HLSL, Result);
+				}
+			}
+			NumComponents = DestTypeDesc.NumComponents;
+			Result.Append(TEXT(")"));
+		}
+
+		check(NumComponents == DestTypeDesc.NumComponents);
+		return AllocateString(*Allocator, Result);
+	}
+	else
+	{
+		Errors.AddErrorf(Node, TEXT("Cannot cast between non-numeric types %s to %s."), SourceTypeDesc.Name, DestTypeDesc.Name);
+		return TEXT("");
+	}
+}
+
 // From MaterialUniformExpressions.cpp
 extern void WriteMaterialUniformAccess(UE::Shader::EValueComponentType ComponentType, uint32 NumComponents, uint32 UniformOffset, FStringBuilderBase& OutResult);
 
@@ -257,7 +390,7 @@ bool UE::HLSLTree::FExpressionLocalPHI::PrepareValue(FEmitContext& Context)
 	for (int32 i = 0; i < NumValues; ++i)
 	{
 		FScope* ValueScope = Scopes[i];
-		const TCHAR* ValueCode = Values[i]->GetValueShader(Context);
+		const TCHAR* ValueCode = Values[i]->GetValueShader(Context, Type);
 		if (ValueScope == DeclarationScope)
 		{
 			ValueScope->EmitDeclarationf(Context, TEXT("%s LocalPHI%d = %s;"),
@@ -403,6 +536,7 @@ bool UE::HLSLTree::FExpression::InternalSetValueShader(FEmitContext& Context, co
 
 		if (!Declaration)
 		{
+			check(ParentScope);
 			Declaration = Context.AcquireLocalDeclarationCode();
 			ParentScope->ExpressionCodeMap.Add(Hash, Declaration);
 			const Shader::FValueTypeDescription& TypeDesc = Shader::GetValueTypeDescription(ValueType);
@@ -477,13 +611,13 @@ const TCHAR* UE::HLSLTree::FExpression::GetValueShader(FEmitContext& Context)
 		TStringBuilder<1024> FormattedCode;
 		if (EvaluationType == EExpressionEvaluationType::Constant)
 		{
-			FormattedCode.Append(ConstantValue.ToString(Shader::EValueStringFormat::HLSL));
+			ConstantValue.ToString(Shader::EValueStringFormat::HLSL, FormattedCode);
 		}
 		else
 		{
 			check(EvaluationType == EExpressionEvaluationType::Preshader);
 			check(Preshader);
-			Context.AddPreshader(GetValueType(), *Preshader, FormattedCode);
+			Context.AddPreshader(ValueType, *Preshader, FormattedCode);
 		}
 		Code = AllocateString(*Context.Allocator, FormattedCode);
 	}
@@ -491,6 +625,24 @@ const TCHAR* UE::HLSLTree::FExpression::GetValueShader(FEmitContext& Context)
 	return Code;
 }
 
+const TCHAR* UE::HLSLTree::FExpression::GetValueShader(FEmitContext& Context, Shader::EValueType Type)
+{
+	check(EvaluationType != EExpressionEvaluationType::None);
+
+	if (EvaluationType == EExpressionEvaluationType::Constant && Type != ValueType)
+	{
+		// If we need to cast a constant value, perform the cast first, then convert the result to HLSL
+		const Shader::FValue CastConstantValue = Shader::Cast(ConstantValue, Type);
+
+		TStringBuilder<1024> FormattedCode;
+		CastConstantValue.ToString(Shader::EValueStringFormat::HLSL, FormattedCode);
+		return AllocateString(*Context.Allocator, FormattedCode);
+	}
+
+	const ECastFlags CastFlags = ECastFlags::ReplicateScalar | ECastFlags::AllowAppendZeroes | ECastFlags::AllowTruncate;
+	return Context.CastShaderValue(this, GetValueShader(Context), ValueType, Type, CastFlags);
+}
+
 void UE::HLSLTree::FExpression::GetValuePreshader(FEmitContext& Context, Shader::FPreshaderData& OutPreshader)
 {
 	switch (EvaluationType)
diff --git a/Engine/Source/Runtime/Engine/Private/HLSLTree/HLSLTreeCommon.cpp b/Engine/Source/Runtime/Engine/Private/HLSLTree/HLSLTreeCommon.cpp
--- a/Engine/Source/Runtime/Engine/Private/HLSLTree/HLSLTreeCommon.cpp
+++ b/Engine/Source/Runtime/Engine/Private/HLSLTree/HLSLTreeCommon.cpp
@@ -4,6 +4,26 @@
 #include "MaterialShared.h"
 #include "Engine/Texture.h"
 
+static UE::Shader::EValueType GetShaderType(EMaterialValueType MaterialType)
+{
+	switch (MaterialType)
+	{
+	case MCT_Float1: return UE::Shader::EValueType::Float1;
+	case MCT_Float2: return UE::Shader::EValueType::Float2;
+	case MCT_Float3: return UE::Shader::EValueType::Float3;
+	case MCT_Float4: return UE::Shader::EValueType::Float4;
+	case MCT_Float: return UE::Shader::EValueType::Float1;
+	case MCT_StaticBool: return UE::Shader::EValueType::Bool1;
+	case MCT_MaterialAttributes: return UE::Shader::EValueType::MaterialAttributes;
+	case MCT_ShadingModel: return UE::Shader::EValueType::Int1;
+	case MCT_LWCScalar: return UE::Shader::EValueType::Double1;
+	case MCT_LWCVector2: return UE::Shader::EValueType::Double2;
+	case MCT_LWCVector3: return UE::Shader::EValueType::Double3;
+	case MCT_LWCVector4: return UE::Shader::EValueType::Double4;
+	default: checkNoEntry(); return UE::Shader::EValueType::Void;
+	}
+}
+
 UE::HLSLTree::FBinaryOpDescription UE::HLSLTree::GetBinaryOpDesription(EBinaryOp Op)
 {
 	switch (Op)
@@ -286,7 +306,7 @@ bool UE::HLSLTree::FExpressionTextureSample::PrepareValue(FEmitContext& Context)
 		SampleFunctionName,
 		*TextureName,
 		*SamplerStateCode,
-		TexCoordExpression->GetValueShader(Context));
+		TexCoordExpression->GetValueShader(Context, Shader::EValueType::Float2));
 }
 
 bool UE::HLSLTree::FExpressionDefaultMaterialAttributes::PrepareValue(FEmitContext& Context)
@@ -320,11 +340,12 @@ bool UE::HLSLTree::FExpressionSetMaterialAttribute::PrepareValue(FEmitContext& C
 		return false;
 	}
 
+	const EMaterialValueType AttributeType = FMaterialAttributeDefinitionMap::GetValueType(AttributeID);
 	const FString PropertyName = FMaterialAttributeDefinitionMap::GetAttributeName(AttributeID);
 	return SetValueShaderf(Context, TEXT("FMaterialAttributes_Set%s(%s, %s)"),
 		*PropertyName,
-		AttributesExpression->GetValueShader(Context),
-		ValueExpression->GetValueShader(Context));
+		AttributesExpression->GetValueShader(Context, Shader::EValueType::MaterialAttributes),
+		ValueExpression->GetValueShader(Context, GetShaderType(AttributeType)));
 }
 
 bool UE::HLSLTree::FExpressionSelect::UpdateType(FUpdateTypeContext& Context, int8 InRequestedNumComponents)
@@ -372,9 +393,9 @@ bool UE::HLSLTree::FExpressionSelect::PrepareValue(FEmitContext& Context)
 
 		// TODO - preshader
 		return SetValueShaderf(Context, TEXT("(%s ? %s : %s)"),
-			ConditionExpression->GetValueShader(Context),
-			TrueExpression->GetValueShader(Context),
-			FalseExpression->GetValueShader(Context));
+			ConditionExpression->GetValueShader(Context, Shader::EValueType::Bool1),
+			TrueExpression->GetValueShader(Context, GetValueType()),
+			FalseExpression->GetValueShader(Context, GetValueType()));
 	}
 	return true;
 }
@@ -453,8 +474,15 @@ bool UE::HLSLTree::FExpressionBinaryOp::PrepareValue(FEmitContext& Context)
 
 bool UE::HLSLTree::FExpressionSwizzle::UpdateType(FUpdateTypeContext& Context, int8 InRequestedNumComponents)
 {
-	const Shader::EValueType InputType = RequestExpressionType(Context, Input, InRequestedNumComponents);
-	return SetType(Context, Shader::MakeValueType(InputType, Parameters.NumComponents));
+	const int32 NumSwizzleComponents = FMath::Min<int32>(InRequestedNumComponents, Parameters.NumComponents);
+	int32 MaxSwizzleIndex = 0;
+	for (int32 i = 0; i < NumSwizzleComponents; ++i)
+	{
+		MaxSwizzleIndex = FMath::Max<int32>(MaxSwizzleIndex, Parameters.ComponentIndex[i]);
+	}
+
+	const Shader::EValueType InputType = RequestExpressionType(Context, Input, MaxSwizzleIndex + 1);
+	return SetType(Context, Shader::MakeValueType(InputType, NumSwizzleComponents));
 }
 
 bool UE::HLSLTree::FExpressionSwizzle::PrepareValue(FEmitContext& Context)
@@ -465,18 +493,37 @@ bool UE::HLSLTree::FExpressionSwizzle::PrepareValue(FEmitContext& Context)
 		return false;
 	}
 
+	const Shader::FValueTypeDescription TypeDesc = Shader::GetValueTypeDescription(GetValueType());
+	const int32 NumSwizzleComponents = FMath::Min<int32>(TypeDesc.NumComponents, Parameters.NumComponents);
+
 	static const TCHAR ComponentName[] = { 'x', 'y', 'z', 'w' };
 	TCHAR Swizzle[5] = TEXT("");
-	for (int32 i = 0; i < Parameters.NumComponents; ++i)
+	bool bHasSwizzleReorder = false;
+
+	for (int32 i = 0; i < NumSwizzleComponents; ++i)
 	{
 		const int32 ComponentIndex = Parameters.ComponentIndex[i];
 		check(ComponentIndex >= 0 && ComponentIndex < 4);
 		Swizzle[i] = ComponentName[ComponentIndex];
+		if (ComponentIndex != i)
+		{
+			bHasSwizzleReorder = true;
+		}
 	}
 
 	if (InputEvaluation == EExpressionEvaluationType::Shader)
 	{
-		return SetValueInlineShaderf(Context, TEXT("%s.%s"), Input->GetValueShader(Context), Swizzle);
+		if (bHasSwizzleReorder)
+		{
+			return SetValueInlineShaderf(Context, TEXT("%s.%s"),
+				Input->GetValueShader(Context),
+				Swizzle);
+		}
+		else
+		{
+			const Shader::EValueType InputType = Shader::MakeValueType(GetValueType(), NumSwizzleComponents);
+			return SetValueInlineShaderf(Context, TEXT("%s"), Input->GetValueShader(Context, InputType));
+		}
 	}
 	else
 	{
@@ -520,7 +567,10 @@ bool UE::HLSLTree::FExpressionAppend::PrepareValue(FEmitContext& Context)
 	if (Evaluation == EExpressionEvaluationType::Shader)
 	{
 		const Shader::FValueTypeDescription ResultTypeDesc = Shader::GetValueTypeDescription(GetValueType());
-		return SetValueShaderf(Context, TEXT("%s(%s, %s)"), ResultTypeDesc.Name, Lhs->GetValueShader(Context), Rhs->GetValueShader(Context));
+		return SetValueShaderf(Context, TEXT("%s(%s, %s)"),
+			ResultTypeDesc.Name,
+			Lhs->GetValueShader(Context),
+			Rhs->GetValueShader(Context));
 	}
 	else
 	{
@@ -532,77 +582,6 @@ bool UE::HLSLTree::FExpressionAppend::PrepareValue(FEmitContext& Context)
 	}
 }
 
-bool UE::HLSLTree::FExpressionCast::UpdateType(FUpdateTypeContext& Context, int8 InRequestedNumComponents)
-{
-	RequestExpressionType(Context, Input, Shader::GetValueTypeDescription(Type).NumComponents);
-	return SetType(Context, Type);
-}
-
-bool UE::HLSLTree::FExpressionCast::PrepareValue(FEmitContext& Context)
-{
-	const EExpressionEvaluationType InputEvaluation = PrepareExpressionValue(Context, Input);
-	if (InputEvaluation == EExpressionEvaluationType::None)
-	{
-		return false;
-	}
-
-	const Shader::EValueType InputType = Input->GetValueType();
-	if (InputType == Type)
-	{
-		// No cast required (nop)
-		return SetValueForward(Context, Input);
-	}
-
-	const TCHAR* InputCode = Input->GetValueShader(Context);
-	const Shader::FValueTypeDescription OutputTypeDesc = Shader::GetValueTypeDescription(Type);
-	const Shader::FValueTypeDescription InputTypeDesc = Shader::GetValueTypeDescription(InputType);
-	if (InputTypeDesc.NumComponents == OutputTypeDesc.NumComponents)
-	{
-		// Cast between different underlying types with the same number of components, use a C-style cast
-		return SetValueInlineShaderf(Context, TEXT("((%s)%s)"), OutputTypeDesc.Name, InputCode);
-	}
-	else if (OutputTypeDesc.NumComponents < InputTypeDesc.NumComponents)
-	{
-		static const TCHAR* SwizzleName[] = { TEXT("x"), TEXT("xy"), TEXT("xyz"), TEXT("xyzw") };
-		// Masking off some input components, use a swizzle
-		if (OutputTypeDesc.ComponentType == InputTypeDesc.ComponentType)
-		{
-			// underlying types are the same, no C-style cast needed
-			SetValueInlineShaderf(Context, TEXT("%s.%s"), InputCode, SwizzleName[OutputTypeDesc.NumComponents]);
-		}
-		else
-		{
-			// Different underlying types, add a C-style cast as well
-			SetValueInlineShaderf(Context, TEXT("((%s)%s.%s)"), OutputTypeDesc.Name, InputCode, SwizzleName[OutputTypeDesc.NumComponents]);
-		}
-	}
-	else
-	{
-		// Padding out components
-		if (InputTypeDesc.NumComponents == 1 && EnumHasAllFlags(Flags, ECastFlags::ReplicateScalar))
-		{
-			// Replicating a scalar value, use a swizzle
-			static const TCHAR* ScalarSwizzleName[] = { TEXT("x"), TEXT("xx"), TEXT("xxx"), TEXT("xxxx") };
-			SetValueInlineShaderf(Context, TEXT("%s.%s"), InputCode, ScalarSwizzleName[OutputTypeDesc.NumComponents]);
-		}
-		else
-		{
-			// Padding with 0s, use a constructor
-			const Shader::FValue ZeroValue(OutputTypeDesc.ComponentType, 1);
-			TStringBuilder<2048> LocalCode;
-
-			LocalCode.Appendf(TEXT("%s(%s"), OutputTypeDesc.Name, InputCode);
-			for (int32 i = InputTypeDesc.NumComponents; i < OutputTypeDesc.NumComponents; ++i)
-			{
-				LocalCode.Appendf(TEXT(", %s"), *ZeroValue.ToString(UE::Shader::EValueStringFormat::HLSL));
-			}
-			LocalCode.Append(TEXT(")"));
-			return SetValueInlineShader(Context, LocalCode);
-		}
-	}
-	return true;
-}
-
 bool UE::HLSLTree::FExpressionReflectionVector::PrepareValue(FEmitContext& Context)
 {
 	return SetValueInlineShaderf(Context, TEXT("Parameters.ReflectionVector"));
@@ -643,7 +622,8 @@ void UE::HLSLTree::FStatementIf::EmitHLSL(FEmitContext& Context) const
 		return;
 	}
 
-	ParentScope->EmitNestedScopef(Context, ThenScope, TEXT("if (%s)"), ConditionExpression->GetValueShader(Context));
+	ParentScope->EmitNestedScopef(Context, ThenScope, TEXT("if (%s)"),
+		ConditionExpression->GetValueShader(Context, Shader::EValueType::Bool1));
 	if (ElseScope)
 	{
 		ParentScope->EmitNestedScopef(Context, ElseScope, TEXT("else"));
diff --git a/Engine/Source/Runtime/Engine/Private/Materials/HLSLMaterialTranslator.cpp b/Engine/Source/Runtime/Engine/Private/Materials/HLSLMaterialTranslator.cpp
--- a/Engine/Source/Runtime/Engine/Private/Materials/HLSLMaterialTranslator.cpp
+++ b/Engine/Source/Runtime/Engine/Private/Materials/HLSLMaterialTranslator.cpp
@@ -4137,8 +4137,9 @@ int32 FHLSLMaterialTranslator::Constant4(float X,float Y,float Z,float W)
 
 int32 FHLSLMaterialTranslator::GenericConstant(const UE::Shader::FValue& Value)
 {
-	const FString Code = Value.ToString(UE::Shader::EValueStringFormat::HLSL);
-	return AddUniformExpression(new FMaterialUniformExpressionGenericConstant(Value), GetMaterialValueType(Value.GetType()), *Code);
+	TStringBuilder<1024> String;
+	Value.ToString(UE::Shader::EValueStringFormat::HLSL, String);
+	return AddUniformExpression(new FMaterialUniformExpressionGenericConstant(Value), GetMaterialValueType(Value.GetType()), String.ToString());
 }
 	
 int32 FHLSLMaterialTranslator::ViewProperty(EMaterialExposedViewProperty Property, bool InvProperty)
diff --git a/Engine/Source/Runtime/Engine/Private/Materials/MaterialExpressionHLSL.cpp b/Engine/Source/Runtime/Engine/Private/Materials/MaterialExpressionHLSL.cpp
--- a/Engine/Source/Runtime/Engine/Private/Materials/MaterialExpressionHLSL.cpp
+++ b/Engine/Source/Runtime/Engine/Private/Materials/MaterialExpressionHLSL.cpp
@@ -414,13 +414,13 @@ EMaterialGenerateHLSLStatus UMaterialExpressionForLoop::GenerateHLSLStatements(F
 		return Generator.Error(TEXT("Missing LoopBody connection"));
 	}
 
-	UE::HLSLTree::FExpression* StartExpression = StartIndex.AcquireHLSLExpressionWithCast(Generator, Scope, UE::Shader::EValueType::Int1);
+	UE::HLSLTree::FExpression* StartExpression = StartIndex.AcquireHLSLExpression(Generator, Scope);
 	if (!StartExpression)
 	{
 		return Generator.Error(TEXT("Missing StartIndex connection"));
 	}
 
-	UE::HLSLTree::FExpression* EndExpression = EndIndex.AcquireHLSLExpressionWithCast(Generator, Scope, UE::Shader::EValueType::Int1);
+	UE::HLSLTree::FExpression* EndExpression = EndIndex.AcquireHLSLExpression(Generator, Scope);
 	if (!EndExpression)
 	{
 		return Generator.Error(TEXT("Missing EndIndex connection"));
@@ -430,7 +430,7 @@ EMaterialGenerateHLSLStatus UMaterialExpressionForLoop::GenerateHLSLStatements(F
 	FExpressionDataForLoop* ExpressionData = Generator.NewExpressionData<FExpressionDataForLoop>(this);
 	ExpressionData->LocalName = *FString::Printf(TEXT("ForLoopControl%d"), GlobalData->NumLoops++);
 
-	UE::HLSLTree::FExpression* StepExpression = IndexStep.GetTracedInput().Expression ? IndexStep.AcquireHLSLExpressionWithCast(Generator, Scope, UE::Shader::EValueType::Int1) : Generator.NewConstant(1);
+	UE::HLSLTree::FExpression* StepExpression = IndexStep.GetTracedInput().Expression ? IndexStep.AcquireHLSLExpression(Generator, Scope) : Generator.NewConstant(1);
 
 	Generator.GenerateAssignLocal(Scope, ExpressionData->LocalName, StartExpression);
 
diff --git a/Engine/Source/Runtime/Engine/Private/Materials/MaterialExpressions.cpp b/Engine/Source/Runtime/Engine/Private/Materials/MaterialExpressions.cpp
--- a/Engine/Source/Runtime/Engine/Private/Materials/MaterialExpressions.cpp
+++ b/Engine/Source/Runtime/Engine/Private/Materials/MaterialExpressions.cpp
@@ -4263,7 +4263,9 @@ int32 UMaterialExpressionGenericConstant::Compile(class FMaterialCompiler* Compi
 
 void UMaterialExpressionGenericConstant::GetCaption(TArray<FString>& OutCaptions) const
 {
-	OutCaptions.Add(GetConstantValue().ToString(UE::Shader::EValueStringFormat::Description));
+	TStringBuilder<1024> String;
+	GetConstantValue().ToString(UE::Shader::EValueStringFormat::Description, String);
+	OutCaptions.Add(FString(String.ToView()));
 }
 
 FString UMaterialExpressionGenericConstant::GetDescription() const
diff --git a/Engine/Source/Runtime/Engine/Private/Materials/MaterialHLSLGenerator.cpp b/Engine/Source/Runtime/Engine/Private/Materials/MaterialHLSLGenerator.cpp
--- a/Engine/Source/Runtime/Engine/Private/Materials/MaterialHLSLGenerator.cpp
+++ b/Engine/Source/Runtime/Engine/Private/Materials/MaterialHLSLGenerator.cpp
@@ -142,7 +142,7 @@ static UE::HLSLTree::FExpression* CompileMaterialInput(FMaterialHLSLGenerator& G
 					const int32 TexCoordIndex = (int32)InputProperty - MP_CustomizedUVs0;
 					if (TexCoordIndex < Material->NumCustomizedUVs)
 					{
-						Expression = InputDescription.Input->AcquireHLSLExpressionWithCast(Generator, Scope, InputDescription.Type);
+						Expression = InputDescription.Input->AcquireHLSLExpression(Generator, Scope);
 					}
 					if (!Expression)
 					{
@@ -151,7 +151,7 @@ static UE::HLSLTree::FExpression* CompileMaterialInput(FMaterialHLSLGenerator& G
 				}
 				else
 				{
-					Expression = InputDescription.Input->AcquireHLSLExpressionWithCast(Generator, Scope, InputDescription.Type);
+					Expression = InputDescription.Input->AcquireHLSLExpression(Generator, Scope);
 				}
 			}
 		}
@@ -252,12 +252,6 @@ UE::HLSLTree::FExpressionSwizzle* FMaterialHLSLGenerator::NewSwizzle(UE::HLSLTre
 	return Expression;
 }
 
-UE::HLSLTree::FExpression* FMaterialHLSLGenerator::NewCast(UE::HLSLTree::FScope& Scope, UE::Shader::EValueType Type, UE::HLSLTree::FExpression* Input, UE::HLSLTree::ECastFlags Flags)
-{
-	UE::HLSLTree::FExpressionCast* Expression = HLSLTree->NewExpression<UE::HLSLTree::FExpressionCast>(Scope, Type, Input, Flags);
-	return Expression;
-}
-
 UE::HLSLTree::FTextureParameterDeclaration* FMaterialHLSLGenerator::AcquireTextureDeclaration(const UE::HLSLTree::FTextureDescription& Value)
 {
 	FString SamplerTypeError;
diff --git a/Engine/Source/Runtime/Engine/Private/Materials/MaterialShared.cpp b/Engine/Source/Runtime/Engine/Private/Materials/MaterialShared.cpp
--- a/Engine/Source/Runtime/Engine/Private/Materials/MaterialShared.cpp
+++ b/Engine/Source/Runtime/Engine/Private/Materials/MaterialShared.cpp
@@ -289,12 +289,6 @@ UE::HLSLTree::FExpression* FExpressionInput::AcquireHLSLExpression(FMaterialHLSL
 	return Result;
 }
 
-UE::HLSLTree::FExpression* FExpressionInput::AcquireHLSLExpressionWithCast(FMaterialHLSLGenerator& Generator, UE::HLSLTree::FScope& Scope, UE::Shader::EValueType Type) const
-{
-	UE::HLSLTree::FExpression* Result = AcquireHLSLExpression(Generator, Scope);
-	return Result ? Generator.NewCast(Scope, Type, Result) : nullptr;
-}
-
 UE::HLSLTree::FTextureParameterDeclaration* FExpressionInput::AcquireHLSLTexture(FMaterialHLSLGenerator& Generator, UE::HLSLTree::FScope& Scope) const
 {
 	UE::HLSLTree::FTextureParameterDeclaration* Result = nullptr;
diff --git a/Engine/Source/Runtime/Engine/Private/Shader/ShaderValue.cpp b/Engine/Source/Runtime/Engine/Private/Shader/ShaderValue.cpp
--- a/Engine/Source/Runtime/Engine/Private/Shader/ShaderValue.cpp
+++ b/Engine/Source/Runtime/Engine/Private/Shader/ShaderValue.cpp
@@ -75,11 +75,11 @@ struct FCastBool
 };
 
 template<typename Operation, typename ResultType>
-void CastValue(const Operation& Op, const FValue& Value, ResultType& OutResult)
+void AsType(const Operation& Op, const FValue& Value, ResultType& OutResult)
 {
 	using FComponentType = typename Operation::FComponentType;
-	const int32 NumComponents = Value.NumComponents;
 	const EValueComponentType ComponentType = Value.ComponentType;
+	const int32 NumComponents = Value.NumComponents;
 	if (NumComponents == 1)
 	{
 		const FComponentType Component = Op(ComponentType, Value.Component[0]);
@@ -101,6 +101,27 @@ void CastValue(const Operation& Op, const FValue& Value, ResultType& OutResult)
 	}
 }
 
+template<typename Operation>
+void Cast(const Operation& Op, const FValue& Value, FValue& OutResult)
+{
+	using FComponentType = typename Operation::FComponentType;
+	const EValueComponentType ComponentType = Value.ComponentType;
+	const int32 NumResultComponents = OutResult.NumComponents;
+	const int32 NumCopyComponents = FMath::Min<int32>(Value.NumComponents, NumResultComponents);
+	for (int32 i = 0; i < NumCopyComponents; ++i)
+	{
+		OutResult.Component[i] = Op(ComponentType, Value.Component[i]);
+	}
+
+	if (NumCopyComponents == 1 && NumResultComponents > 1)
+	{
+		for (int32 i = 1; i < NumResultComponents; ++i)
+		{
+			OutResult.Component[i] = OutResult.Component[0];
+		}
+	}
+}
+
 void FormatComponent_Double(double Value, int32 NumComponents, EValueStringFormat Format, FStringBuilderBase& OutResult)
 {
 	if (Format == EValueStringFormat::HLSL)
@@ -165,14 +186,14 @@ UE::Shader::FMemoryImageValue UE::Shader::FValue::AsMemoryImage() const
 UE::Shader::FFloatValue UE::Shader::FValue::AsFloat() const
 {
 	FFloatValue Result;
-	Private::CastValue(Private::FCastFloat(), *this, Result);
+	Private::AsType(Private::FCastFloat(), *this, Result);
 	return Result;
 }
 
 UE::Shader::FDoubleValue UE::Shader::FValue::AsDouble() const
 {
 	FDoubleValue Result;
-	Private::CastValue(Private::FCastDouble(), *this, Result);
+	Private::AsType(Private::FCastDouble(), *this, Result);
 	return Result;
 }
 
@@ -191,21 +212,21 @@ FVector4d UE::Shader::FValue::AsVector4d() const
 UE::Shader::FIntValue UE::Shader::FValue::AsInt() const
 {
 	FIntValue Result;
-	Private::CastValue(Private::FCastInt(), *this, Result);
+	Private::AsType(Private::FCastInt(), *this, Result);
 	return Result;
 }
 
 UE::Shader::FBoolValue UE::Shader::FValue::AsBool() const
 {
 	FBoolValue Result;
-	Private::CastValue(Private::FCastBool(), *this, Result);
+	Private::AsType(Private::FCastBool(), *this, Result);
 	return Result;
 }
 
 float UE::Shader::FValue::AsFloatScalar() const
 {
 	FFloatValue Result;
-	Private::CastValue(Private::FCastFloat(), *this, Result);
+	Private::AsType(Private::FCastFloat(), *this, Result);
 	return Result[0];
 }
 
@@ -236,7 +257,7 @@ uint32 UE::Shader::GetComponentTypeSizeInBytes(EValueComponentType Type)
 	}
 }
 
-FString UE::Shader::FValue::ToString(EValueStringFormat Format) const
+const TCHAR* UE::Shader::FValue::ToString(EValueStringFormat Format, FStringBuilderBase& OutString) const
 {
 	if (Format == EValueStringFormat::HLSL && ComponentType == EValueComponentType::Double)
 	{
@@ -258,16 +279,15 @@ FString UE::Shader::FValue::ToString(EValueStringFormat Format) const
 
 		if (NumComponents > 1)
 		{
-			return FString::Printf(TEXT("MakeLWCVector%d(float%d(%s), float%d(%s))"), NumComponents, NumComponents, TileValue.ToString(), NumComponents, OffsetValue.ToString());
+			OutString.Appendf(TEXT("MakeLWCVector%d(float%d(%s), float%d(%s))"), NumComponents, NumComponents, TileValue.ToString(), NumComponents, OffsetValue.ToString());
 		}
 		else
 		{
-			return FString::Printf(TEXT("MakeLWCScalar(%s, %s)"), TileValue.ToString(), OffsetValue.ToString());
+			OutString.Appendf(TEXT("MakeLWCScalar(%s, %s)"), TileValue.ToString(), OffsetValue.ToString());
 		}
 	}
 	else
 	{
-		TStringBuilder<1024> Result;
 		if (Format == EValueStringFormat::HLSL)
 		{
 			const TCHAR* ComponentName = nullptr;
@@ -280,11 +300,11 @@ FString UE::Shader::FValue::ToString(EValueStringFormat Format) const
 			}
 			if (NumComponents > 1)
 			{
-				Result.Appendf(TEXT("%s%d("), ComponentName, NumComponents);
+				OutString.Appendf(TEXT("%s%d("), ComponentName, NumComponents);
 			}
 			else
 			{
-				Result.Appendf(TEXT("%s("), ComponentName);
+				OutString.Appendf(TEXT("%s("), ComponentName);
 			}
 		}
 
@@ -292,25 +312,26 @@ FString UE::Shader::FValue::ToString(EValueStringFormat Format) const
 		{
 			if (i > 0)
 			{
-				Result.Append(TEXT(", "));
+				OutString.Append(TEXT(", "));
 			}
 
 			switch (ComponentType)
 			{
-			case EValueComponentType::Int: Result.Appendf(TEXT("%d"), Component[i].Int); break;
-			case EValueComponentType::Bool: Result.Append(Component[i].Bool ? TEXT("true") : TEXT("false")); break;
-			case EValueComponentType::Float: Private::FormatComponent_Double((double)Component[i].Float, NumComponents, Format, Result); break;
-			case EValueComponentType::Double: Private::FormatComponent_Double(Component[i].Double, NumComponents, Format, Result); break;
+			case EValueComponentType::Int: OutString.Appendf(TEXT("%d"), Component[i].Int); break;
+			case EValueComponentType::Bool: OutString.Append(Component[i].Bool ? TEXT("true") : TEXT("false")); break;
+			case EValueComponentType::Float: Private::FormatComponent_Double((double)Component[i].Float, NumComponents, Format, OutString); break;
+			case EValueComponentType::Double: Private::FormatComponent_Double(Component[i].Double, NumComponents, Format, OutString); break;
 			default: checkNoEntry(); break;
 			}
 		}
 
 		if (Format == EValueStringFormat::HLSL)
 		{
-			Result.Append(TEXT(")"));
+			OutString.Append(TEXT(")"));
 		}
-		return Result.ToString();
 	}
+
+	return OutString.ToString();
 }
 
 UE::Shader::FValueTypeDescription UE::Shader::GetValueTypeDescription(EValueType Type)
@@ -994,3 +1015,25 @@ UE::Shader::FValue UE::Shader::Append(const FValue& Lhs, const FValue& Rhs)
 	Result.NumComponents = NumComponents;
 	return Result;
 }
+
+UE::Shader::FValue UE::Shader::Cast(const FValue& Value, EValueType Type)
+{
+	const EValueType SourceType = Value.GetType();
+	if (Type == SourceType)
+	{
+		return Value;
+	}
+
+	FValue Result(Type);
+	switch (Result.ComponentType)
+	{
+	case EValueComponentType::Float: Private::Cast(Private::FCastFloat(), Value, Result); break;
+	case EValueComponentType::Double: Private::Cast(Private::FCastDouble(), Value, Result); break;
+	case EValueComponentType::Int: Private::Cast(Private::FCastInt(), Value, Result); break;
+	case EValueComponentType::Bool: Private::Cast(Private::FCastBool(), Value, Result); break;
+	default: checkNoEntry(); break;
+	}
+
+	return Result;
+}
+
diff --git a/Engine/Source/Runtime/Engine/Public/HLSLTree/HLSLTree.h b/Engine/Source/Runtime/Engine/Public/HLSLTree/HLSLTree.h
--- a/Engine/Source/Runtime/Engine/Public/HLSLTree/HLSLTree.h
+++ b/Engine/Source/Runtime/Engine/Public/HLSLTree/HLSLTree.h
@@ -30,6 +30,7 @@ class FPreshaderData;
 namespace HLSLTree
 {
 
+class FNode;
 class FScope;
 
 static constexpr int32 MaxNumPreviousScopes = 2;
@@ -49,6 +50,14 @@ public:
 
 	bool AddError(const FNode* InNode, FStringView InError);
 
+	template<typename FormatType, typename... Types>
+	bool AddErrorf(const FNode* InNode, const FormatType& Format, Types... Args)
+	{
+		TStringBuilder<2048> String;
+		String.Appendf(Format, Forward<Types>(Args)...);
+		return AddError(InNode, FStringView(String.GetData(), String.Len()));
+	}
+
 private:
 	FMemStackBase* Allocator = nullptr;
 	const FError* FirstError = nullptr;
@@ -62,6 +71,17 @@ public:
 	FErrors& Errors;
 };
 
+enum class ECastFlags : uint32
+{
+	None = 0u,
+	ReplicateScalar = (1u << 0),
+	AllowTruncate = (1u << 1),
+	AllowAppendZeroes = (1u << 2),
+
+	ValidCast = ReplicateScalar | AllowTruncate,
+};
+ENUM_CLASS_FLAGS(ECastFlags);
+
 /** Tracks shared state while emitting HLSL code */
 class FEmitContext
 {
@@ -74,6 +94,8 @@ public:
 	/** Get a unique local variable name */
 	const TCHAR* AcquireLocalDeclarationCode();
 
+	const TCHAR* CastShaderValue(const FNode* Node, const TCHAR* Code, Shader::EValueType SourceType, Shader::EValueType DestType, ECastFlags Flags);
+
 	void AddPreshader(Shader::EValueType Type, const Shader::FPreshaderData& Preshader, FStringBuilderBase& OutCode);
 
 	FMemStackBase* Allocator = nullptr;
@@ -158,6 +180,7 @@ public:
 	friend EExpressionEvaluationType PrepareExpressionValue(FEmitContext& Context, FExpression* InExpression);
 
 	const TCHAR* GetValueShader(FEmitContext& Context);
+	const TCHAR* GetValueShader(FEmitContext& Context, Shader::EValueType Type);
 	void GetValuePreshader(FEmitContext& Context, Shader::FPreshaderData& OutPreshader);
 	Shader::FValue GetValueConstant(FEmitContext& Context);
 
diff --git a/Engine/Source/Runtime/Engine/Public/HLSLTree/HLSLTreeCommon.h b/Engine/Source/Runtime/Engine/Public/HLSLTree/HLSLTreeCommon.h
--- a/Engine/Source/Runtime/Engine/Public/HLSLTree/HLSLTreeCommon.h
+++ b/Engine/Source/Runtime/Engine/Public/HLSLTree/HLSLTreeCommon.h
@@ -277,33 +277,6 @@ public:
 	virtual bool PrepareValue(FEmitContext& Context) override;
 };
 
-class FExpressionCast : public FExpression
-{
-public:
-	FExpressionCast(Shader::EValueType InType, FExpression* InInput, ECastFlags InFlags = ECastFlags::None)
-		: Type(InType)
-		, Input(InInput)
-		, Flags(InFlags)
-	{}
-
-	Shader::EValueType Type;
-	FExpression* Input;
-	ECastFlags Flags;
-
-	virtual ENodeVisitResult Visit(FNodeVisitor& Visitor) override
-	{
-		const ENodeVisitResult Result = FExpression::Visit(Visitor);
-		if (ShouldVisitDependentNodes(Result))
-		{
-			Visitor.VisitNode(Input);
-		}
-		return Result;
-	}
-
-	virtual bool UpdateType(FUpdateTypeContext& Context, int8 InRequestedNumComponents) override;
-	virtual bool PrepareValue(FEmitContext& Context) override;
-};
-
 class FExpressionReflectionVector : public FExpression
 {
 public:
diff --git a/Engine/Source/Runtime/Engine/Public/HLSLTree/HLSLTreeTypes.h b/Engine/Source/Runtime/Engine/Public/HLSLTree/HLSLTreeTypes.h
--- a/Engine/Source/Runtime/Engine/Public/HLSLTree/HLSLTreeTypes.h
+++ b/Engine/Source/Runtime/Engine/Public/HLSLTree/HLSLTreeTypes.h
@@ -64,12 +64,5 @@ inline uint32 GetTypeHash(const FTextureDescription& Ref)
 	return HashCombine(GetTypeHash(Ref.Texture), GetTypeHash(Ref.SamplerType));
 }
 
-enum class ECastFlags : uint32
-{
-	None = 0u,
-	ReplicateScalar = (1u << 0),
-};
-ENUM_CLASS_FLAGS(ECastFlags);
-
 } // namespace HLSLTree
 } // namespace UE
diff --git a/Engine/Source/Runtime/Engine/Public/MaterialExpressionIO.h b/Engine/Source/Runtime/Engine/Public/MaterialExpressionIO.h
--- a/Engine/Source/Runtime/Engine/Public/MaterialExpressionIO.h
+++ b/Engine/Source/Runtime/Engine/Public/MaterialExpressionIO.h
@@ -78,7 +78,6 @@ struct FExpressionInput
 #if WITH_EDITOR
 	ENGINE_API int32 Compile(class FMaterialCompiler* Compiler);
 	ENGINE_API UE::HLSLTree::FExpression* AcquireHLSLExpression(FMaterialHLSLGenerator& Generator, UE::HLSLTree::FScope& Scope) const;
-	ENGINE_API UE::HLSLTree::FExpression* AcquireHLSLExpressionWithCast(FMaterialHLSLGenerator& Generator, UE::HLSLTree::FScope& Scope, UE::Shader::EValueType Type) const;
 	ENGINE_API UE::HLSLTree::FTextureParameterDeclaration* AcquireHLSLTexture(FMaterialHLSLGenerator& Generator, UE::HLSLTree::FScope& Scope) const;
 #endif // WITH_EDITOR
 
diff --git a/Engine/Source/Runtime/Engine/Public/MaterialHLSLGenerator.h b/Engine/Source/Runtime/Engine/Public/MaterialHLSLGenerator.h
--- a/Engine/Source/Runtime/Engine/Public/MaterialHLSLGenerator.h
+++ b/Engine/Source/Runtime/Engine/Public/MaterialHLSLGenerator.h
@@ -33,7 +33,6 @@ struct FSwizzleParameters;
 class FExpressionConstant;
 class FExpressionExternalInput;
 class FExpressionSwizzle;
-class FExpressionCast;
 class FExpressionLocalPHI;
 }
 }
@@ -86,7 +85,6 @@ public:
 	UE::HLSLTree::FExpressionConstant* NewConstant(const UE::Shader::FValue& Value);
 	UE::HLSLTree::FExpressionExternalInput* NewTexCoord(UE::HLSLTree::FScope& Scope, int32 Index);
 	UE::HLSLTree::FExpressionSwizzle* NewSwizzle(UE::HLSLTree::FScope& Scope, const UE::HLSLTree::FSwizzleParameters& Params, UE::HLSLTree::FExpression* Input);
-	UE::HLSLTree::FExpression* NewCast(UE::HLSLTree::FScope& Scope, UE::Shader::EValueType Type, UE::HLSLTree::FExpression* Input, UE::HLSLTree::ECastFlags Flags = UE::HLSLTree::ECastFlags::None);
 
 	/** Returns a declaration to access the given texture, with no parameter */
 	UE::HLSLTree::FTextureParameterDeclaration* AcquireTextureDeclaration(const UE::HLSLTree::FTextureDescription& Value);
diff --git a/Engine/Source/Runtime/Engine/Public/Shader/ShaderTypes.h b/Engine/Source/Runtime/Engine/Public/Shader/ShaderTypes.h
--- a/Engine/Source/Runtime/Engine/Public/Shader/ShaderTypes.h
+++ b/Engine/Source/Runtime/Engine/Public/Shader/ShaderTypes.h
@@ -93,6 +93,10 @@ struct FMemoryImageValue
 union FValueComponent
 {
 	FValueComponent() : Packed(0u) {}
+	FValueComponent(float v) : Packed(0u) { Float = v; }
+	FValueComponent(double v) : Packed(0u) { Double = v; }
+	FValueComponent(int32 v) : Packed(0u) { Int = v; }
+	FValueComponent(bool v) : Packed(0u) { Bool = v ? 1 : 0; }
 
 	// 'Bool' is stored as uint8 to avoid changing on different compilers
 	bool AsBool() const { return Bool != 0u; }
@@ -120,124 +124,122 @@ struct FValue
 
 	inline FValue(float v) : ComponentType(EValueComponentType::Float), NumComponents(1)
 	{
-		Component[0].Float = v;
+		Component[0] = v;
 	}
 
 	inline FValue(float X, float Y) : ComponentType(EValueComponentType::Float), NumComponents(2)
 	{
-		Component[0].Float = X;
-		Component[1].Float = Y;
+		Component[0] = X;
+		Component[1] = Y;
 	}
 
 	inline FValue(float X, float Y, float Z) : ComponentType(EValueComponentType::Float), NumComponents(3)
 	{
-		Component[0].Float = X;
-		Component[1].Float = Y;
-		Component[2].Float = Z;
+		Component[0] = X;
+		Component[1] = Y;
+		Component[2] = Z;
 	}
 
 	inline FValue(float X, float Y, float Z, float W) : ComponentType(EValueComponentType::Float), NumComponents(4)
 	{
-		Component[0].Float = X;
-		Component[1].Float = Y;
-		Component[2].Float = Z;
-		Component[3].Float = W;
+		Component[0] = X;
+		Component[1] = Y;
+		Component[2] = Z;
+		Component[3] = W;
 	}
 
 	inline FValue(double v) : ComponentType(EValueComponentType::Double), NumComponents(1)
 	{
-		Component[0].Double = v;
+		Component[0] = v;
 	}
 
 	inline FValue(double X, double Y) : ComponentType(EValueComponentType::Double), NumComponents(2)
 	{
-		Component[0].Double = X;
-		Component[1].Double = Y;
+		Component[0] = X;
+		Component[1] = Y;
 	}
 
 	inline FValue(double X, double Y, double Z) : ComponentType(EValueComponentType::Double), NumComponents(3)
 	{
-		Component[0].Double = X;
-		Component[1].Double = Y;
-		Component[2].Double = Z;
+		Component[0] = X;
+		Component[1] = Y;
+		Component[2] = Z;
 	}
 
 	inline FValue(double X, double Y, double Z, double W) : ComponentType(EValueComponentType::Double), NumComponents(4)
 	{
-		Component[0].Double = X;
-		Component[1].Double = Y;
-		Component[2].Double = Z;
-		Component[3].Double = W;
+		Component[0] = X;
+		Component[1] = Y;
+		Component[2] = Z;
+		Component[3] = W;
 	}
 
 	inline FValue(const FLinearColor& Value) : ComponentType(EValueComponentType::Float), NumComponents(4)
 	{
-		Component[0].Float = Value.R;
-		Component[1].Float = Value.G;
-		Component[2].Float = Value.B;
-		Component[3].Float = Value.A;
+		Component[0] = Value.R;
+		Component[1] = Value.G;
+		Component[2] = Value.B;
+		Component[3] = Value.A;
 	}
 
 	inline FValue(const FVector2D& Value) : ComponentType(EValueComponentType::Float), NumComponents(2)
 	{
-		Component[0].Float = Value.X;
-		Component[1].Float = Value.Y;
+		Component[0] = Value.X;
+		Component[1] = Value.Y;
 	}
 
 	inline FValue(const FVector3f& Value) : ComponentType(EValueComponentType::Float), NumComponents(3)
 	{
-		Component[0].Float = Value.X;
-		Component[1].Float = Value.Y;
-		Component[2].Float = Value.Z;
+		Component[0] = Value.X;
+		Component[1] = Value.Y;
+		Component[2] = Value.Z;
 	}
 
 	inline FValue(const FVector3d& Value) : ComponentType(EValueComponentType::Double), NumComponents(3)
 	{
-		Component[0].Double = Value.X;
-		Component[1].Double = Value.Y;
-		Component[2].Double = Value.Z;
+		Component[0] = Value.X;
+		Component[1] = Value.Y;
+		Component[2] = Value.Z;
 	}
 
 	inline FValue(const FVector4f& Value) : ComponentType(EValueComponentType::Float), NumComponents(4)
 	{
-		Component[0].Float = Value.X;
-		Component[1].Float = Value.Y;
-		Component[2].Float = Value.Z;
-		Component[3].Float = Value.W;
+		Component[0] = Value.X;
+		Component[1] = Value.Y;
+		Component[2] = Value.Z;
+		Component[3] = Value.W;
 	}
 
 	inline FValue(const FVector4d& Value) : ComponentType(EValueComponentType::Double), NumComponents(4)
 	{
-		Component[0].Double = Value.X;
-		Component[1].Double = Value.Y;
-		Component[2].Double = Value.Z;
-		Component[3].Double = Value.W;
+		Component[0] = Value.X;
+		Component[1] = Value.Y;
+		Component[2] = Value.Z;
+		Component[3] = Value.W;
 	}
 
 	inline FValue(bool v) : ComponentType(EValueComponentType::Bool), NumComponents(1)
 	{
-		Component[0].Bool = v;
+		Component[0] = v;
 	}
 
 	inline FValue(bool X, bool Y, bool Z, bool W) : ComponentType(EValueComponentType::Bool), NumComponents(4)
 	{
-		Component[0].Bool = X;
-		Component[1].Bool = Y;
-		Component[2].Bool = Z;
-		Component[3].Bool = W;
+		Component[0] = X;
+		Component[1] = Y;
+		Component[2] = Z;
+		Component[3] = W;
 	}
 
 	inline FValue(int32 v) : ComponentType(EValueComponentType::Int), NumComponents(1)
 	{
-		Component[0].Int = v;
+		Component[0] = v;
 	}
 
 	inline EValueType GetType() const { return MakeValueType(ComponentType, NumComponents); }
-
 	inline const FValueComponent& GetComponent(int32 i) const { check(i >= 0 && i < NumComponents); return Component[i]; }
 
 	static FValue FromMemoryImage(EValueType Type, const void* Data, uint32* OutSizeInBytes = nullptr);
-
 	FMemoryImageValue AsMemoryImage() const;
 
 	FFloatValue AsFloat() const;
@@ -250,7 +252,7 @@ struct FValue
 	float AsFloatScalar() const;
 	bool AsBoolScalar() const;
 
-	FString ToString(EValueStringFormat Format = EValueStringFormat::Description) const;
+	const TCHAR* ToString(EValueStringFormat Format, FStringBuilderBase& OutString) const;
 
 	FValueComponent Component[4];
 	EValueComponentType ComponentType;
@@ -293,6 +295,8 @@ ENGINE_API FValue Atan2(const FValue& Lhs, const FValue& Rhs);
 ENGINE_API FValue Dot(const FValue& Lhs, const FValue& Rhs);
 ENGINE_API FValue Cross(const FValue& Lhs, const FValue& Rhs);
 ENGINE_API FValue Append(const FValue& Lhs, const FValue& Rhs);
+
+ENGINE_API FValue Cast(const FValue& Value, EValueType Type);
 }
 }
 
