[DebugDrawService] updated FRWAccessDetector to FRWRecursiveAccessDetector since delegates can register from the `Draw` call. (reentrance was properly supported)

diff --git a/Engine/Source/Runtime/Engine/Classes/Debug/DebugDrawService.h b/Engine/Source/Runtime/Engine/Classes/Debug/DebugDrawService.h
--- a/Engine/Source/Runtime/Engine/Classes/Debug/DebugDrawService.h
+++ b/Engine/Source/Runtime/Engine/Classes/Debug/DebugDrawService.h
@@ -45,7 +45,7 @@ private:
 	 */
 	static FCriticalSection DelegatesLock;
 #if ENABLE_MT_DETECTOR
-	static FRWAccessDetector DelegatesDetector;
+	static FRWRecursiveAccessDetector DelegatesDetector;
 #endif
 	static TArray<TArray<FDebugDrawDelegate> > Delegates;
 	static FEngineShowFlags ObservedFlags;
diff --git a/Engine/Source/Runtime/Engine/Private/Debug/DebugDrawService.cpp b/Engine/Source/Runtime/Engine/Private/Debug/DebugDrawService.cpp
--- a/Engine/Source/Runtime/Engine/Private/Debug/DebugDrawService.cpp
+++ b/Engine/Source/Runtime/Engine/Private/Debug/DebugDrawService.cpp
@@ -6,7 +6,7 @@
 
 FCriticalSection UDebugDrawService::DelegatesLock;
 #if ENABLE_MT_DETECTOR
-FRWAccessDetector UDebugDrawService::DelegatesDetector;
+FRWRecursiveAccessDetector UDebugDrawService::DelegatesDetector;
 #endif
 
 TArray<TArray<FDebugDrawDelegate> > UDebugDrawService::Delegates;
@@ -86,7 +86,7 @@ void UDebugDrawService::Draw(const FEngineShowFlags Flags, UCanvas* Canvas)
 		return;
 	}
 
-	UE_MT_SCOPED_READ_ACCESS(DelegatesDetector);
+	UE_MT_SCOPED_WRITE_ACCESS(DelegatesDetector);
 	for (int32 FlagIndex = 0; FlagIndex < Delegates.Num(); ++FlagIndex)
 	{
 		if (Flags.GetSingleFlag(FlagIndex) && ObservedFlags.GetSingleFlag(FlagIndex) && Delegates[FlagIndex].Num() > 0)
