Don't crash if more than 2 exec inputs are set in material editor
(At some point, should automatically disconnect previous expresisons to avoid this case)

diff --git a/Engine/Source/Runtime/Engine/Private/Materials/MaterialHLSLGenerator.cpp b/Engine/Source/Runtime/Engine/Private/Materials/MaterialHLSLGenerator.cpp
--- a/Engine/Source/Runtime/Engine/Private/Materials/MaterialHLSLGenerator.cpp
+++ b/Engine/Source/Runtime/Engine/Private/Materials/MaterialHLSLGenerator.cpp
@@ -363,10 +363,15 @@ UE::HLSLTree::FTextureParameterDeclaration* FMaterialHLSLGenerator::AcquireTextu
 bool FMaterialHLSLGenerator::GenerateStatements(UE::HLSLTree::FScope& Scope, UMaterialExpression* MaterialExpression)
 {
 	FStatementEntry& Entry = StatementMap.FindOrAdd(MaterialExpression);
+	check(Entry.NumInputs >= 0);
+	check(Entry.NumInputs < MaterialExpression->NumExecutionInputs);
+	if (Entry.NumInputs == MaxNumPreviousScopes)
+	{
+		Error(TEXT("Bad control flow"));
+		return false;
+	}
 
 	Entry.PreviousScope[Entry.NumInputs++] = &Scope;
-	check(Entry.NumInputs <= MaxNumPreviousScopes);
-	check(Entry.NumInputs <= MaterialExpression->NumExecutionInputs);
 
 	bool bResult = true;
 	if (Entry.NumInputs == MaterialExpression->NumExecutionInputs)
