HLSL LocalPHI nodes don't generate local variable if value is the same in all previous scopes

diff --git a/Engine/Source/Runtime/Engine/Private/HLSLTree/HLSLTree.cpp b/Engine/Source/Runtime/Engine/Private/HLSLTree/HLSLTree.cpp
--- a/Engine/Source/Runtime/Engine/Private/HLSLTree/HLSLTree.cpp
+++ b/Engine/Source/Runtime/Engine/Private/HLSLTree/HLSLTree.cpp
@@ -358,6 +358,22 @@ bool UE::HLSLTree::FExpressionLocalPHI::UpdateType(FUpdateTypeContext& Context,
 
 bool UE::HLSLTree::FExpressionLocalPHI::PrepareValue(FEmitContext& Context)
 {
+	FExpression* ForwardExpression = Values[0];
+	for (int32 i = 1; i < NumValues; ++i)
+	{
+		if (Values[i] != ForwardExpression)
+		{
+			ForwardExpression = nullptr;
+			break;
+		}
+	}
+
+	// If we have the same value in all of our scopes, just forward that value directly; no need to create/assign a local variable
+	if (ForwardExpression)
+	{
+		return SetValueForward(Context, ForwardExpression);
+	}
+
 	const int32 LocalPHIIndex = Context.NumLocalPHIs++;
 	if (!SetValueInlineShaderf(Context, TEXT("LocalPHI%d"), LocalPHIIndex))
 	{
