Implemented new input value controller that features only clamping and interpolation (relevant for most gameplay-anim functionality)

diff --git a/Engine/Source/Runtime/Engine/Classes/Animation/InputScaleBias.h b/Engine/Source/Runtime/Engine/Classes/Animation/InputScaleBias.h
--- a/Engine/Source/Runtime/Engine/Classes/Animation/InputScaleBias.h
+++ b/Engine/Source/Runtime/Engine/Classes/Animation/InputScaleBias.h
@@ -11,7 +11,7 @@
 
 struct FAnimationUpdateContext;
 
-// Input scaling struct
+// Input modifier with scaling and biasing
 USTRUCT(BlueprintType)
 struct ENGINE_API FInputScaleBias
 {
@@ -30,10 +30,10 @@ public:
 	{
 	}
 
-	// Apply scale and bias to value
 	float ApplyTo(float Value) const;
-
+#if WITH_EDITOR
 	FText GetFriendlyName(FText InFriendlyName) const;
+#endif
 };
 
 USTRUCT(BlueprintType)
@@ -69,6 +69,7 @@ public:
 	}
 };
 
+// Input modifier with remapping, scaling, biasing, clamping, and interpolation
 USTRUCT(BlueprintType)
 struct ENGINE_API FInputScaleBiasClamp
 {
@@ -127,15 +128,62 @@ public:
 	{
 	}
 
-	// Apply scale, bias, and clamp to value
 	float ApplyTo(float Value, float InDeltaTime) const;
+
 	void Reinitialize() { bInitialized = false; }
+#if WITH_EDITOR
+	FText GetFriendlyName(FText InFriendlyName) const;
+#endif
+};
+
+// Input modifier with clamping and interpolation
+USTRUCT(BlueprintType)
+struct ENGINE_API FInputClampConstants
+{
+	GENERATED_BODY()
+
+	UPROPERTY(EditAnywhere, BlueprintReadWrite, Category = Settings)
+	bool bClampResult = false;
+
+	UPROPERTY(EditAnywhere, BlueprintReadWrite, Category = Settings)
+	bool bInterpResult = false;
+
+	UPROPERTY(EditAnywhere, BlueprintReadWrite, Category = Settings, meta = (EditCondition = "bClampResult"))
+	float ClampMin = 0.f;
+
+	UPROPERTY(EditAnywhere, BlueprintReadWrite, Category = Settings, meta = (EditCondition = "bClampResult"))
+	float ClampMax = 1.f;
+
+	UPROPERTY(EditAnywhere, BlueprintReadWrite, Category = Settings, meta = (EditCondition = "bInterpResult"))
+	float InterpSpeedIncreasing = 10.f;
+
+	UPROPERTY(EditAnywhere, BlueprintReadWrite, Category = Settings, meta = (EditCondition = "bInterpResult"))
+	float InterpSpeedDecreasing = 10.f;
 
 #if WITH_EDITOR
+	// Get a friendly name to display on a pin
 	FText GetFriendlyName(FText InFriendlyName) const;
 #endif
 };
 
+// Mutable state struct to be used with FInputClampConstants
+USTRUCT(BlueprintType)
+struct ENGINE_API FInputClampState
+{
+	GENERATED_BODY()
+
+	// The interpolated result
+	float InterpolatedResult = 0.f;
+
+	// Whether this state is initialized
+	bool bInitialized = false;
+
+	// Apply scale, bias, and clamp to value
+	float ApplyTo(const FInputClampConstants& InConstants, float InValue, float InDeltaTime);
+
+	void Reinitialize() { bInitialized = false; }
+};
+
 USTRUCT(BlueprintType)
 struct ENGINE_API FInputScaleBiasClampConstants
 {
@@ -266,6 +314,7 @@ struct ENGINE_API FInputAlphaBoolBlend
 	{}
 
 	float ApplyTo(bool bEnabled, float InDeltaTime);
+
 	void Reinitialize() { bInitialized = false; }
 };
 
diff --git a/Engine/Source/Runtime/Engine/Private/Animation/InputScaleBias.cpp b/Engine/Source/Runtime/Engine/Private/Animation/InputScaleBias.cpp
--- a/Engine/Source/Runtime/Engine/Private/Animation/InputScaleBias.cpp
+++ b/Engine/Source/Runtime/Engine/Private/Animation/InputScaleBias.cpp
@@ -12,6 +12,7 @@ float FInputScaleBias::ApplyTo(float Value) const
 	return FMath::Clamp<float>( Value * Scale + Bias, 0.0f, 1.0f );
 }
 
+#if WITH_EDITOR
 FText FInputScaleBias::GetFriendlyName(FText InFriendlyName) const
 {
 	FText OutFriendlyName = InFriendlyName;
@@ -51,6 +52,63 @@ FText FInputScaleBias::GetFriendlyName(FText InFriendlyName) const
 
 	return OutFriendlyName;
 }
+#endif
+
+/////////////////////////////////////////////////////
+// FInputClamp
+
+#if WITH_EDITOR
+FText FInputClampConstants::GetFriendlyName(FText InFriendlyName) const
+{
+	FText OutFriendlyName = InFriendlyName;
+
+	// Clamp
+	if (bClampResult)
+	{
+		FFormatNamedArguments Args;
+		Args.Add(TEXT("PinFriendlyName"), OutFriendlyName);
+		Args.Add(TEXT("ClampMin"), ClampMin);
+		Args.Add(TEXT("ClampMax"), ClampMax);
+		OutFriendlyName = FText::Format(LOCTEXT("FInputScaleBias_Clamp", "Clamp({PinFriendlyName}, {ClampMin}, {ClampMax})"), Args);
+	}
+
+	// Interp
+	if (bInterpResult)
+	{
+		FFormatNamedArguments Args;
+		Args.Add(TEXT("PinFriendlyName"), OutFriendlyName);
+		Args.Add(TEXT("InterpSpeedIncreasing"), InterpSpeedIncreasing);
+		Args.Add(TEXT("InterpSpeedDecreasing"), InterpSpeedDecreasing);
+		OutFriendlyName = FText::Format(LOCTEXT("FInputScaleBias_Interp", "FInterp({PinFriendlyName}, ({InterpSpeedIncreasing}:{InterpSpeedDecreasing}))"), Args);
+	}
+
+	return OutFriendlyName;
+}
+#endif
+
+float FInputClampState::ApplyTo(const FInputClampConstants& InConstants, float InValue, float InDeltaTime)
+{
+	float Result = InValue;
+
+	if (InConstants.bClampResult)
+	{
+		Result = FMath::Clamp<float>(Result, InConstants.ClampMin, InConstants.ClampMax);
+	}
+
+	if (InConstants.bInterpResult)
+	{
+		if (bInitialized)
+		{
+			const float InterpSpeed = (Result >= InterpolatedResult) ? InConstants.InterpSpeedIncreasing : InConstants.InterpSpeedDecreasing;
+			Result = FMath::FInterpTo(InterpolatedResult, Result, InDeltaTime, InterpSpeed);
+		}
+
+		InterpolatedResult = Result;
+	}
+
+	bInitialized = true;
+	return Result;
+}
 
 /////////////////////////////////////////////////////
 // FInputScaleBiasClamp
