Fix CIS issue # 118132, 118082

diff --git a/Engine/Source/Runtime/Engine/Private/VisualLogger/VisualLogger.cpp b/Engine/Source/Runtime/Engine/Private/VisualLogger/VisualLogger.cpp
--- a/Engine/Source/Runtime/Engine/Private/VisualLogger/VisualLogger.cpp
+++ b/Engine/Source/Runtime/Engine/Private/VisualLogger/VisualLogger.cpp
@@ -250,7 +250,7 @@ FVisualLogEntry* FVisualLogger::GetEntryToWrite(const UObject* Object, const flo
 		}
 	}
 
-	if (CurrentEntry->bIsAllowedToLog)
+	if (CurrentEntry != nullptr && CurrentEntry->bIsAllowedToLog)
 	{
 		bIsFlushRequired = true;
 	}
