Allow SpringInterpolator constructor to behave similarly in functionality to SetDefaultSpringConstants for UPROPERTY default initialization

diff --git a/Engine/Source/Runtime/Engine/Classes/Engine/SpringInterpolator.h b/Engine/Source/Runtime/Engine/Classes/Engine/SpringInterpolator.h
--- a/Engine/Source/Runtime/Engine/Classes/Engine/SpringInterpolator.h
+++ b/Engine/Source/Runtime/Engine/Classes/Engine/SpringInterpolator.h
@@ -145,9 +145,9 @@ protected:
 	RK4Integrator::FRK4State<T> State;
 
 public:
-	FRK4SpringInterpolator()
-		: StiffnessConstant(1.f)
-		, DampeningRatio(1.f)
+	FRK4SpringInterpolator(float InSpringStiffness = 1.f, float InDampeningRatio = 1.f)
+		: StiffnessConstant(InSpringStiffness)
+		, DampeningRatio(InDampeningRatio)
 		, bIsInitialized(false)
 		, bIsInMotion(false)
 		, TimeRemaining(0.f)
