Slight refactoring of HLSLTree generation
No longer generate any code for empty scopes, allows dead-stripping control flow

diff --git a/Engine/Source/Runtime/Engine/Private/HLSLTree/HLSLTree.cpp b/Engine/Source/Runtime/Engine/Private/HLSLTree/HLSLTree.cpp
--- a/Engine/Source/Runtime/Engine/Private/HLSLTree/HLSLTree.cpp
+++ b/Engine/Source/Runtime/Engine/Private/HLSLTree/HLSLTree.cpp
@@ -65,7 +65,17 @@ UE::HLSLTree::FErrors::FErrors(FMemStackBase& InAllocator)
 
 bool UE::HLSLTree::FErrors::AddError(const FNode* InNode, FStringView InError)
 {
-	int a = 0;
+	const int32 SizeofString = InError.Len() * sizeof(TCHAR);
+	void* Memory = Allocator->Alloc(sizeof(FError) + SizeofString, alignof(FError));
+	FError* Error = new(Memory) FError();
+	FMemory::Memcpy(Error->Message, InError.GetData(), SizeofString);
+	Error->Message[InError.Len()] = 0;
+	Error->MessageLength = InError.Len();
+	Error->Node = InNode;
+	Error->Next = FirstError;
+	FirstError = Error;
+	NumErrors++;
+
 	return false;
 }
 
@@ -407,6 +417,7 @@ bool UE::HLSLTree::FExpressionLocalPHI::PrepareValue(FEmitContext& Context)
 	{
 		FScope* ValueScope = Scopes[i];
 		const TCHAR* ValueCode = Values[i]->GetValueShader(Context, Type);
+		ValueScope->MarkLiveRecursive();
 		if (ValueScope == DeclarationScope)
 		{
 			ValueScope->EmitDeclarationf(Context, TEXT("%s LocalPHI%d = %s;"),
@@ -425,6 +436,7 @@ bool UE::HLSLTree::FExpressionLocalPHI::PrepareValue(FEmitContext& Context)
 
 	if (bNeedToAddDeclaration)
 	{
+		check(DeclarationScope->IsLive());
 		DeclarationScope->EmitDeclarationf(Context, TEXT("%s LocalPHI%d;"),
 			TypeDesc.Name,
 			LocalPHIIndex);
@@ -500,9 +512,14 @@ UE::HLSLTree::EExpressionEvaluationType UE::HLSLTree::PrepareExpressionValue(FEm
 		return EExpressionEvaluationType::None;
 	}
 
+	if (InExpression->ValueType == Shader::EValueType::Void)
+	{
+		// Can't prepare value if we didn't successfully prepare a type
+		return EExpressionEvaluationType::None;
+	}
+
 	if (InExpression->EvaluationType == EExpressionEvaluationType::None)
 	{
-		check(InExpression->GetValueType() != Shader::EValueType::Void);
 		check(!InExpression->bReentryFlag); // Re-entry is not supported
 
 		InExpression->bReentryFlag = true;
@@ -556,6 +573,7 @@ bool UE::HLSLTree::FExpression::InternalSetValueShader(FEmitContext& Context, co
 			Declaration = Context.AcquireLocalDeclarationCode();
 			ParentScope->ExpressionCodeMap.Add(Hash, Declaration);
 			const Shader::FValueTypeDescription& TypeDesc = Shader::GetValueTypeDescription(ValueType);
+			ParentScope->MarkLiveRecursive();
 			ParentScope->EmitStatementf(Context, TEXT("const %s %s = %s;"),
 				TypeDesc.Name,
 				Declaration,
@@ -589,6 +607,7 @@ bool UE::HLSLTree::FExpression::SetValueConstant(FEmitContext& Context, const Sh
 bool UE::HLSLTree::FExpression::SetValueForward(FEmitContext& Context, FExpression* Source)
 {
 	check(EvaluationType == EExpressionEvaluationType::None);
+	check(Source);
 
 	EvaluationType = PrepareExpressionValue(Context, Source);
 	switch (EvaluationType)
@@ -738,8 +757,9 @@ void UE::HLSLTree::FScope::UseExpression(FExpression* Expression)
 
 void UE::HLSLTree::FScope::InternalEmitCode(FEmitContext& Context, FCodeList& List, FScope* NestedScope, const TCHAR* String, int32 Length)
 {
-	if (NestedScope && NestedScope->Statement)
+	if (NestedScope && NestedScope->Statement && !NestedScope->Statement->bEmitHLSL)
 	{
+		NestedScope->Statement->bEmitHLSL = true;
 		NestedScope->Statement->EmitHLSL(Context);
 	}
 
@@ -800,9 +820,26 @@ void UE::HLSLTree::FTree::Destroy(FTree* Tree)
 //
 static void WriteIndent(int32 IndentLevel, FStringBuilderBase& InOutString)
 {
+	const int32 Offset = InOutString.AddUninitialized(IndentLevel);
+	TCHAR* Result = InOutString.GetData() + Offset;
 	for (int32 i = 0; i < IndentLevel; ++i)
 	{
-		InOutString.Append(TCHAR('\t'));
+		*Result++ = TCHAR('\t');
+	}
+}
+
+void UE::HLSLTree::FScope::MarkLive()
+{
+	bLive = true;
+}
+
+void UE::HLSLTree::FScope::MarkLiveRecursive()
+{
+	FScope* Scope = this;
+	while (Scope)
+	{
+		Scope->bLive = true;
+		Scope = Scope->ParentScope;
 	}
 }
 
@@ -843,13 +880,25 @@ void UE::HLSLTree::FScope::WriteHLSL(int32 Indent, FStringBuilderBase& OutString
 	}
 }
 
-bool UE::HLSLTree::FTree::EmitHLSL(UE::HLSLTree::FEmitContext& Context, FStringBuilderBase& Writer) const
+bool UE::HLSLTree::FTree::EmitHLSL(FEmitContext& Context, FStringBuilderBase& Writer) const
 {
+	if (!ResultStatement)
+	{
+		return false;
+	}
+
 	{
 		FUpdateTypeContext UpdateTypeContext(Context.Errors);
 		RequestScopeTypes(UpdateTypeContext, RootScope);
 	}
 
+	ResultStatement->bEmitHLSL = true;
+	ResultStatement->EmitHLSL(Context);
+	if (!RootScope->IsLive())
+	{
+		return false;
+	}
+
 	if (RootScope->Statement)
 	{
 		RootScope->Statement->EmitHLSL(Context);
@@ -857,7 +906,7 @@ bool UE::HLSLTree::FTree::EmitHLSL(UE::HLSLTree::FEmitContext& Context, FStringB
 
 	Context.Finalize();
 
-	RootScope->WriteHLSL(0, Writer);
+	RootScope->WriteHLSL(1, Writer);
 	return true;
 }
 
@@ -890,3 +939,9 @@ UE::HLSLTree::FTextureParameterDeclaration* UE::HLSLTree::FTree::NewTextureParam
 	return Declaration;
 }
 
+void UE::HLSLTree::FTree::SetResult(FStatement& InResult)
+{
+	check(!ResultStatement);
+	ResultStatement = &InResult;
+}
+
diff --git a/Engine/Source/Runtime/Engine/Private/HLSLTree/HLSLTreeCommon.cpp b/Engine/Source/Runtime/Engine/Private/HLSLTree/HLSLTreeCommon.cpp
--- a/Engine/Source/Runtime/Engine/Private/HLSLTree/HLSLTreeCommon.cpp
+++ b/Engine/Source/Runtime/Engine/Private/HLSLTree/HLSLTreeCommon.cpp
@@ -589,6 +589,7 @@ bool UE::HLSLTree::FExpressionReflectionVector::PrepareValue(FEmitContext& Conte
 
 void UE::HLSLTree::FStatementBreak::EmitHLSL(FEmitContext& Context) const
 {
+	ParentScope->MarkLive();
 	ParentScope->EmitStatementf(Context, TEXT("break;"));
 }
 
@@ -604,6 +605,7 @@ void UE::HLSLTree::FStatementReturn::EmitHLSL(FEmitContext& Context) const
 		return;
 	}
 
+	ParentScope->MarkLiveRecursive();
 	ParentScope->EmitStatementf(Context, TEXT("return %s;"), Expression->GetValueShader(Context));
 }
 
@@ -617,18 +619,26 @@ void UE::HLSLTree::FStatementIf::RequestTypes(FUpdateTypeContext& Context) const
 
 void UE::HLSLTree::FStatementIf::EmitHLSL(FEmitContext& Context) const
 {
-	if (PrepareExpressionValue(Context, ConditionExpression) == EExpressionEvaluationType::None)
+	bool bIfLive = ThenScope->IsLive();
+	if (ElseScope && ElseScope->IsLive())
 	{
-		return;
+		bIfLive = true;
 	}
 
-	ParentScope->EmitNestedScopef(Context, ThenScope, TEXT("if (%s)"),
-		ConditionExpression->GetValueShader(Context, Shader::EValueType::Bool1));
-	if (ElseScope)
+	if (bIfLive)
 	{
-		ParentScope->EmitNestedScopef(Context, ElseScope, TEXT("else"));
-	}
+		if (PrepareExpressionValue(Context, ConditionExpression) == EExpressionEvaluationType::None)
+		{
+			return;
+		}
 
+		ParentScope->EmitNestedScopef(Context, ThenScope, TEXT("if (%s)"),
+			ConditionExpression->GetValueShader(Context, Shader::EValueType::Bool1));
+		if (ElseScope && ElseScope->IsLive())
+		{
+			ParentScope->EmitNestedScopef(Context, ElseScope, TEXT("else"));
+		}
+	}
 	ParentScope->EmitNestedScope(Context, NextScope);
 }
 
@@ -640,6 +650,9 @@ void UE::HLSLTree::FStatementLoop::RequestTypes(FUpdateTypeContext& Context) con
 
 void UE::HLSLTree::FStatementLoop::EmitHLSL(FEmitContext& Context) const
 {
-	ParentScope->EmitNestedScopef(Context, LoopScope, TEXT("while (true)"));
+	if (LoopScope->IsLive())
+	{
+		ParentScope->EmitNestedScopef(Context, LoopScope, TEXT("while (true)"));
+	}
 	ParentScope->EmitNestedScope(Context, NextScope);
 }
diff --git a/Engine/Source/Runtime/Engine/Private/Materials/MaterialHLSLGenerator.cpp b/Engine/Source/Runtime/Engine/Private/Materials/MaterialHLSLGenerator.cpp
--- a/Engine/Source/Runtime/Engine/Private/Materials/MaterialHLSLGenerator.cpp
+++ b/Engine/Source/Runtime/Engine/Private/Materials/MaterialHLSLGenerator.cpp
@@ -68,6 +68,11 @@ bool FMaterialHLSLGenerator::Finalize()
 		for (int32 i = 0; i < Expression->NumValues; ++i)
 		{
 			Expression->Values[i] = InternalAcquireLocalValue(*Expression->Scopes[i], Expression->LocalName);
+			if (!Expression->Values[i])
+			{
+				Errorf(TEXT("Local %s is not assigned on all control paths"), *Expression->LocalName.ToString());
+				return false;
+			}
 		}
 	}
 
@@ -204,6 +209,7 @@ bool FMaterialHLSLGenerator::GenerateResult(UE::HLSLTree::FScope& Scope)
 			{
 				UE::HLSLTree::FStatementReturn* ReturnStatement = HLSLTree->NewStatement<UE::HLSLTree::FStatementReturn>(Scope);
 				ReturnStatement->Expression = AttributesExpression;
+				HLSLTree->SetResult(*ReturnStatement);
 				bResult = true;
 			}
 		}
diff --git a/Engine/Source/Runtime/Engine/Public/HLSLTree/HLSLTree.h b/Engine/Source/Runtime/Engine/Public/HLSLTree/HLSLTree.h
--- a/Engine/Source/Runtime/Engine/Public/HLSLTree/HLSLTree.h
+++ b/Engine/Source/Runtime/Engine/Public/HLSLTree/HLSLTree.h
@@ -40,7 +40,7 @@ struct FError
 	const FError* Next;
 	const FNode* Node;
 	int32 MessageLength;
-	const TCHAR Message[1];
+	TCHAR Message[1];
 };
 
 class FErrors
@@ -48,6 +48,8 @@ class FErrors
 public:
 	explicit FErrors(FMemStackBase& InAllocator);
 
+	int32 Num() const { return NumErrors; }
+
 	bool AddError(const FNode* InNode, FStringView InError);
 
 	template<typename FormatType, typename... Types>
@@ -61,6 +63,7 @@ public:
 private:
 	FMemStackBase* Allocator = nullptr;
 	const FError* FirstError = nullptr;
+	int32 NumErrors = 0;
 };
 
 class FUpdateTypeContext
@@ -154,12 +157,16 @@ public:
 class FStatement : public FNode
 {
 public:
+	static constexpr bool MarkScopeLive = false;
+	static constexpr bool MarkScopeLiveRecursive = false;
+
 	virtual ENodeVisitResult Visit(FNodeVisitor& Visitor) override;
 
 	virtual void RequestTypes(FUpdateTypeContext& Context) const = 0;
 	virtual void EmitHLSL(FEmitContext& Context) const = 0;
 
 	FScope* ParentScope = nullptr;
+	bool bEmitHLSL = false;
 };
 
 /**
@@ -271,6 +278,7 @@ public:
 	static FScope* FindSharedParent(FScope* Lhs, FScope* Rhs);
 
 	inline FScope* GetParentScope() const { return ParentScope; }
+	inline bool IsLive() const { return bLive; }
 
 	inline TArrayView<FScope*> GetPreviousScopes() const
 	{
@@ -311,6 +319,9 @@ public:
 		InternalEmitCodef(Context, Statements, NestedScope, Format, Forward<Types>(Args)...);
 	}
 
+	void MarkLive();
+	void MarkLiveRecursive();
+
 	void WriteHLSL(int32 Indent, FStringBuilderBase& OutString) const;
 
 private:
@@ -351,6 +362,7 @@ private:
 	FCodeList Statements;
 	int32 NumPreviousScopes = 0;
 	int32 NestedLevel = 0;
+	bool bLive = false;
 };
 
 /**
@@ -381,11 +393,21 @@ public:
 	{
 		T* Statement = NewNode<T>(Forward<ArgTypes>(Args)...);
 		RegisterStatement(Scope, Statement);
+		if constexpr (T::MarkScopeLiveRecursive)
+		{
+			Scope.MarkLiveRecursive();
+		}
+		else if constexpr (T::MarkScopeLive)
+		{
+			Scope.MarkLive();
+		}
+
 		return Statement;
 	}
 
 	FScope* NewScope(FScope& Scope);
 	FTextureParameterDeclaration* NewTextureParameterDeclaration(const FName& Name, const FTextureDescription& DefaultValue);
+	void SetResult(FStatement& InResult);
 
 private:
 	template<typename T, typename... ArgTypes>
@@ -400,9 +422,10 @@ private:
 	void RegisterExpression(FScope& Scope, FExpression* Expression);
 	void RegisterStatement(FScope& Scope, FStatement* Statement);
 
-	FMemStackBase* Allocator;
-	FNode* Nodes;
-	FScope* RootScope;
+	FMemStackBase* Allocator = nullptr;
+	FNode* Nodes = nullptr;
+	FScope* RootScope = nullptr;
+	FStatement* ResultStatement = nullptr;
 };
 
 } // namespace HLSLTree
diff --git a/Engine/Source/Runtime/Engine/Public/HLSLTree/HLSLTreeCommon.h b/Engine/Source/Runtime/Engine/Public/HLSLTree/HLSLTreeCommon.h
--- a/Engine/Source/Runtime/Engine/Public/HLSLTree/HLSLTreeCommon.h
+++ b/Engine/Source/Runtime/Engine/Public/HLSLTree/HLSLTreeCommon.h
@@ -287,6 +287,7 @@ public:
 class FStatementReturn : public FStatement
 {
 public:
+	static constexpr bool MarkScopeLiveRecursive = true;
 	FExpression* Expression;
 
 	virtual ENodeVisitResult Visit(FNodeVisitor& Visitor) override
@@ -306,6 +307,8 @@ public:
 class FStatementBreak : public FStatement
 {
 public:
+	static constexpr bool MarkScopeLive = true;
+
 	virtual void RequestTypes(FUpdateTypeContext& Context) const override {}
 	virtual void EmitHLSL(FEmitContext& Context) const override;
 };
