Minor cleanup of HLSL constants

diff --git a/Engine/Source/Runtime/Engine/Private/HLSLTree/HLSLTree.cpp b/Engine/Source/Runtime/Engine/Private/HLSLTree/HLSLTree.cpp
--- a/Engine/Source/Runtime/Engine/Private/HLSLTree/HLSLTree.cpp
+++ b/Engine/Source/Runtime/Engine/Private/HLSLTree/HLSLTree.cpp
@@ -600,7 +600,7 @@ bool UE::HLSLTree::FExpression::SetValueConstant(FEmitContext& Context, const Sh
 	check(EvaluationType == EExpressionEvaluationType::None);
 
 	EvaluationType = EExpressionEvaluationType::Constant;
-	ConstantValue = InValue;
+	ConstantValue = Shader::Cast(InValue, ValueType);
 	return true;
 }
 
diff --git a/Engine/Source/Runtime/Engine/Private/HLSLTree/HLSLTreeCommon.cpp b/Engine/Source/Runtime/Engine/Private/HLSLTree/HLSLTreeCommon.cpp
--- a/Engine/Source/Runtime/Engine/Private/HLSLTree/HLSLTreeCommon.cpp
+++ b/Engine/Source/Runtime/Engine/Private/HLSLTree/HLSLTreeCommon.cpp
@@ -100,8 +100,8 @@ bool UE::HLSLTree::FExpressionConstant::PrepareValue(FEmitContext& Context)
 
 bool UE::HLSLTree::FExpressionMaterialParameter::UpdateType(FUpdateTypeContext& Context, int8 InRequestedNumComponents)
 {
-	//return Shader::MakeValueTypeWithRequestedNumComponents(GetShaderValueType(ParameterType), InRequestedNumComponents);
-	return SetType(Context, GetShaderValueType(ParameterType));
+	const Shader::EValueType Type = Shader::MakeValueTypeWithRequestedNumComponents(GetShaderValueType(ParameterType), InRequestedNumComponents);
+	return SetType(Context, Type);
 }
 
 bool UE::HLSLTree::FExpressionMaterialParameter::PrepareValue(FEmitContext& Context)
@@ -109,17 +109,17 @@ bool UE::HLSLTree::FExpressionMaterialParameter::PrepareValue(FEmitContext& Cont
 	if (ParameterType == EMaterialParameterType::StaticSwitch)
 	{
 		const FMaterialParameterInfo ParameterInfo(ParameterName);
-		bool bValue = DefaultValue.Component[0].AsBool();
+		Shader::FValue Value = DefaultValue;
 		for (const FStaticSwitchParameter& Parameter : Context.StaticParameters->StaticSwitchParameters)
 		{
 			if (Parameter.ParameterInfo == ParameterInfo)
 			{
-				bValue = Parameter.Value;
+				Value = Parameter.Value;
 				break;
 			}
 		}
 
-		return SetValueConstant(Context, bValue);
+		return SetValueConstant(Context, Value);
 	}
 	else
 	{
