fix windows compile of hlsltree

diff --git a/Engine/Source/Runtime/Engine/Public/HLSLTree/HLSLTree.h b/Engine/Source/Runtime/Engine/Public/HLSLTree/HLSLTree.h
--- a/Engine/Source/Runtime/Engine/Public/HLSLTree/HLSLTree.h
+++ b/Engine/Source/Runtime/Engine/Public/HLSLTree/HLSLTree.h
@@ -428,5 +428,9 @@ private:
 	FStatement* ResultStatement = nullptr;
 };
 
+Shader::EValueType RequestExpressionType(FUpdateTypeContext& Context, FExpression* InExpression, int8 InRequestedNumComponents); // friend of FExpression
+EExpressionEvaluationType PrepareExpressionValue(FEmitContext& Context, FExpression* InExpression); // friend of FExpression
+void RequestScopeTypes(FUpdateTypeContext& Context, const FScope* InScope); // friend of FScope
+
 } // namespace HLSLTree
 } // namespace UE
