Fixed a crash with -NoThreading in SkeletalMeshCompiler

diff --git a/Engine/Source/Runtime/Engine/Private/SkeletalMeshCompiler.cpp b/Engine/Source/Runtime/Engine/Private/SkeletalMeshCompiler.cpp
--- a/Engine/Source/Runtime/Engine/Private/SkeletalMeshCompiler.cpp
+++ b/Engine/Source/Runtime/Engine/Private/SkeletalMeshCompiler.cpp
@@ -99,7 +99,7 @@ EQueuedWorkPriority FSkeletalMeshCompilingManager::GetBasePriority(USkeletalMesh
 FQueuedThreadPool* FSkeletalMeshCompilingManager::GetThreadPool() const
 {
 	static FQueuedThreadPoolDynamicWrapper* GSkeletalMeshThreadPool = nullptr;
-	if (GSkeletalMeshThreadPool == nullptr)
+	if (GSkeletalMeshThreadPool == nullptr && FAssetCompilingManager::Get().GetThreadPool() != nullptr)
 	{
 		SkeletalMeshCompilingManagerImpl::EnsureInitializedCVars();
 
