Merging historical #UE4 Fix for missing sockets in merged meshes in cooked builds. Now assigning new reference skeleton before building sockets.

diff --git a/Engine/Source/Runtime/Engine/Private/SkeletalMeshMerge.cpp b/Engine/Source/Runtime/Engine/Private/SkeletalMeshMerge.cpp
--- a/Engine/Source/Runtime/Engine/Private/SkeletalMeshMerge.cpp
+++ b/Engine/Source/Runtime/Engine/Private/SkeletalMeshMerge.cpp
@@ -84,6 +84,10 @@ void FSkeletalMeshMerge::MergeSkeleton(const TArray<FRefPoseOverride>* RefPoseOv
 	// Build the reference skeleton & sockets.
 
 	BuildReferenceSkeleton(SrcMeshList, NewRefSkeleton, MergeMesh->GetSkeleton());
+
+	// Assign new referencer skeleton.
+	MergeMesh->RefSkeleton = NewRefSkeleton;
+
 	BuildSockets(SrcMeshList);
 
 	// Override the reference bone poses & sockets, if specified.
@@ -94,10 +98,6 @@ void FSkeletalMeshMerge::MergeSkeleton(const TArray<FRefPoseOverride>* RefPoseOv
 		OverrideMergedSockets(*RefPoseOverrides);
 	}
 
-	// Assign new referencer skeleton.
-
-	MergeMesh->SetRefSkeleton(NewRefSkeleton);
-
 	// Rebuild inverse ref pose matrices here as some access patterns 
 	// may need to access these matrices before FinalizeMesh is called
 	// (which would *normally* rebuild the inv ref matrices).
