Fix CIS build
#rnx

diff --git a/Engine/Source/Runtime/Engine/Private/SkeletalMeshMerge.cpp b/Engine/Source/Runtime/Engine/Private/SkeletalMeshMerge.cpp
--- a/Engine/Source/Runtime/Engine/Private/SkeletalMeshMerge.cpp
+++ b/Engine/Source/Runtime/Engine/Private/SkeletalMeshMerge.cpp
@@ -86,7 +86,7 @@ void FSkeletalMeshMerge::MergeSkeleton(const TArray<FRefPoseOverride>* RefPoseOv
 	BuildReferenceSkeleton(SrcMeshList, NewRefSkeleton, MergeMesh->GetSkeleton());
 
 	// Assign new referencer skeleton.
-	MergeMesh->RefSkeleton = NewRefSkeleton;
+	MergeMesh->SetRefSkeleton(NewRefSkeleton);
 
 	BuildSockets(SrcMeshList);
 
