DebugDrawService: prevent access to null canvas. The canvas provided is provided by the View (i.e. DebugCanvas) and it can be null.
Kept the call to UDebugDrawService::Draw in FEditorViewportClient::Draw since some code relies on the creation of the "DebugCanvasObject"

diff --git a/Engine/Source/Runtime/Engine/Private/Debug/DebugDrawService.cpp b/Engine/Source/Runtime/Engine/Private/Debug/DebugDrawService.cpp
--- a/Engine/Source/Runtime/Engine/Private/Debug/DebugDrawService.cpp
+++ b/Engine/Source/Runtime/Engine/Private/Debug/DebugDrawService.cpp
@@ -81,7 +81,7 @@ void UDebugDrawService::Draw(const FEngineShowFlags Flags, FViewport* Viewport,
 
 void UDebugDrawService::Draw(const FEngineShowFlags Flags, UCanvas* Canvas)
 {
-	if (Canvas == nullptr)
+	if (Canvas == nullptr || Canvas->Canvas == nullptr)
 	{
 		return;
 	}
