[VisualLogger] fixed unit test that should expect a null entry when calling GetEntryToWrite with ECreateIfNeeded::DontCreate

diff --git a/Engine/Source/Runtime/Engine/Private/VisualLogger/VisualLoggerAutomationTests.cpp b/Engine/Source/Runtime/Engine/Private/VisualLogger/VisualLoggerAutomationTests.cpp
--- a/Engine/Source/Runtime/Engine/Private/VisualLogger/VisualLoggerAutomationTests.cpp
+++ b/Engine/Source/Runtime/Engine/Private/VisualLogger/VisualLoggerAutomationTests.cpp
@@ -242,9 +242,12 @@ bool FVisualLogEventsTest::RunTest(const FString& Parameters)
 	CHECK_SUCCESS(EventTest3.Name == TEXT("EventTest3"));
 	CHECK_SUCCESS(EventTest3.FriendlyDesc == TEXT("Third simple event for vlog tests"));
 
-	FVisualLogEntry* CurrentEntry = FVisualLogger::Get().GetEntryToWrite(World, World->TimeSeconds, ECreateIfNeeded::DontCreate);
 	float CurrentTimestamp = World->TimeSeconds;
+	FVisualLogEntry* CurrentEntry = FVisualLogger::Get().GetEntryToWrite(World, CurrentTimestamp, ECreateIfNeeded::DontCreate);
+	CHECK_SUCCESS(CurrentEntry == nullptr);
+
 	UE_VLOG_EVENTS(World, NAME_None, EventTest);
+	CurrentEntry = FVisualLogger::Get().GetEntryToWrite(World, CurrentTimestamp, ECreateIfNeeded::DontCreate);
 	CHECK_NOT_NULL(CurrentEntry);
 	CHECK_SUCCESS(CurrentEntry->TimeStamp == CurrentTimestamp);
 	CHECK_SUCCESS(CurrentEntry->Events.Num() == 1);
