When complaining about long compile times, sort by worker compile time to show the bigger offenders first. Removing instruction count from this display as well as it's not really a relevant stat.

diff --git a/Engine/Source/Runtime/Engine/Private/ShaderCompiler/ShaderCompiler.cpp b/Engine/Source/Runtime/Engine/Private/ShaderCompiler/ShaderCompiler.cpp
--- a/Engine/Source/Runtime/Engine/Private/ShaderCompiler/ShaderCompiler.cpp
+++ b/Engine/Source/Runtime/Engine/Private/ShaderCompiler/ShaderCompiler.cpp
@@ -2116,32 +2116,44 @@ void FShaderCompileThreadRunnable::PushCompletedJobsToManager()
 			// Log if requested or if there was an exceptionally slow batch, to see the offender easily
 			if (Manager->bLogJobCompletionTimes || ElapsedTime > 60.0f)
 			{
+				TArray<FShaderCommonCompileJobPtr> SortedJobs = CurrentWorkerInfo.QueuedJobs;
+				SortedJobs.Sort([](const FShaderCommonCompileJobPtr& JobA, const FShaderCommonCompileJobPtr& JobB)
+					{
+						const FShaderCompileJob* SingleJobA = JobA->GetSingleShaderJob();
+						const FShaderCompileJob* SingleJobB = JobB->GetSingleShaderJob();
+
+						const float TimeA = SingleJobA ? SingleJobA->Output.CompileTime : 0.0f;
+						const float TimeB = SingleJobB ? SingleJobB->Output.CompileTime : 0.0f;
+
+						return TimeA > TimeB;
+					});
+
 				FString JobNames;
 
-				for (int32 JobIndex = 0; JobIndex < CurrentWorkerInfo.QueuedJobs.Num(); JobIndex++)
+				for (int32 JobIndex = 0; JobIndex < SortedJobs.Num(); JobIndex++)
 				{
-					const FShaderCommonCompileJob& Job = *CurrentWorkerInfo.QueuedJobs[JobIndex];
-					if (auto* SingleJob = Job.GetSingleShaderJob())
+					const FShaderCommonCompileJob& Job = *SortedJobs[JobIndex];
+					if (const FShaderCompileJob* SingleJob = Job.GetSingleShaderJob())
 					{
 						const TCHAR* JobName = Manager->bLogJobCompletionTimes ? *SingleJob->Input.DebugGroupName : SingleJob->Key.ShaderType->GetName();
-						JobNames += FString::Printf(TEXT("%s [Instructions=%d WorkerTime=%.3fs]"), JobName, SingleJob->Output.NumInstructions, SingleJob->Output.CompileTime);
+						JobNames += FString::Printf(TEXT("%s [WorkerTime=%.3fs]"), JobName, SingleJob->Output.CompileTime);
 					}
 					else
 					{
-						auto* PipelineJob = Job.GetShaderPipelineJob();
+						const FShaderPipelineCompileJob* PipelineJob = Job.GetShaderPipelineJob();
 						JobNames += FString(PipelineJob->Key.ShaderPipeline->GetName());
 						if (PipelineJob->bFailedRemovingUnused)
 						{
 							JobNames += FString(TEXT("(failed to optimize)"));
 						}
 					}
-					if (JobIndex < CurrentWorkerInfo.QueuedJobs.Num() - 1)
+					if (JobIndex < SortedJobs.Num() - 1)
 					{
 						JobNames += TEXT(", ");
 					}
 				}
 
-				UE_LOG(LogShaders, Display, TEXT("Worker (%d/%d) finished batch of %u jobs in %.3fs, %s"), WorkerIndex + 1, WorkerInfos.Num(), CurrentWorkerInfo.QueuedJobs.Num(), ElapsedTime, *JobNames);
+				UE_LOG(LogShaders, Display, TEXT("Worker (%d/%d) finished batch of %u jobs in %.3fs, %s"), WorkerIndex + 1, WorkerInfos.Num(), SortedJobs.Num(), ElapsedTime, *JobNames);
 			}
 
 			CurrentWorkerInfo.FinishTime = FPlatformTime::Seconds();
