Strata - Fixed new legacyconversion code to have correct default (e.g. see FMaterialAttributeDefinitionMap::InitializeAttributeMap for root node).
For instance this fixes water transmittance, thin translucent, coat and disney default black albedo when default values are used.

diff --git a/Engine/Source/Runtime/Engine/Private/Materials/MaterialExpressions.cpp b/Engine/Source/Runtime/Engine/Private/Materials/MaterialExpressions.cpp
--- a/Engine/Source/Runtime/Engine/Private/Materials/MaterialExpressions.cpp
+++ b/Engine/Source/Runtime/Engine/Private/Materials/MaterialExpressions.cpp
@@ -20291,27 +20291,27 @@ int32 UMaterialExpressionStrataLegacyConversion::Compile(class FMaterialCompiler
 	// We probably need to do something along these line as well :::
 	int32 OutputCodeChunk = Compiler->StrataConversionFromLegacy(
 		// Metalness workflow
-		CompileWithDefaultFloat3(Compiler, BaseColor, 0.18f, 0.18f, 0.18f),
+		CompileWithDefaultFloat3(Compiler, BaseColor, 0.0f, 0.0f, 0.0f),
 		CompileWithDefaultFloat1(Compiler, Specular, 0.5f),
 		CompileWithDefaultFloat1(Compiler, Metallic,  0.0f),
 		// Roughness
 		RoughnessCodeChunk,
 		AnisotropyCodeChunk,
 		// SSS
-		CompileWithDefaultFloat3(Compiler, SubSurfaceColor, 0.0f, 0.0f, 0.0f),
+		CompileWithDefaultFloat3(Compiler, SubSurfaceColor, 1.0f, 1.0f, 1.0f),
 		SSSProfileCodeChunk != INDEX_NONE ? SSSProfileCodeChunk : Compiler->Constant(0.0f),	
 		// Clear Coat / Custom
-		CompileWithDefaultFloat1(Compiler, ClearCoat, 0.0f),
-		CompileWithDefaultFloat1(Compiler, ClearCoatRoughness, 0.0f),
+		CompileWithDefaultFloat1(Compiler, ClearCoat, 1.0f),
+		CompileWithDefaultFloat1(Compiler, ClearCoatRoughness, 0.1f),
 		// Misc
 		CompileWithDefaultFloat3(Compiler, EmissiveColor, 0.0f, 0.0f, 0.0f),
 		CompileWithDefaultFloat1(Compiler, Opacity, 1.0f),
-		CompileWithDefaultFloat3(Compiler, TransmittanceColor, 0.0f, 0.0f, 0.0f),
+		CompileWithDefaultFloat3(Compiler, TransmittanceColor, 0.5f, 0.5f, 0.5f),
 		// Water
 		CompileWithDefaultFloat3(Compiler, WaterScatteringCoefficients, 0.0f, 0.0f, 0.0f),
 		CompileWithDefaultFloat3(Compiler, WaterAbsorptionCoefficients, 0.0f, 0.0f, 0.0f),
 		CompileWithDefaultFloat1(Compiler, WaterPhaseG, 0.0f),
-		CompileWithDefaultFloat3(Compiler, ColorScaleBehindWater, 0.0f, 0.0f, 0.0f),
+		CompileWithDefaultFloat3(Compiler, ColorScaleBehindWater, 1.0f, 1.0f, 1.0f),
 		// Shading model
 		CompileWithDefaultFloat1(Compiler, ShadingModel, 0.0f),
 		NormalCodeChunk,
