Fixed thin translucent input type to be float3 which helps when promoting it to a material parameter.

diff --git a/Engine/Source/Runtime/Engine/Classes/Materials/MaterialExpressionThinTranslucentMaterialOutput.h b/Engine/Source/Runtime/Engine/Classes/Materials/MaterialExpressionThinTranslucentMaterialOutput.h
--- a/Engine/Source/Runtime/Engine/Classes/Materials/MaterialExpressionThinTranslucentMaterialOutput.h
+++ b/Engine/Source/Runtime/Engine/Classes/Materials/MaterialExpressionThinTranslucentMaterialOutput.h
@@ -22,6 +22,7 @@ public:
 	//~ Begin UMaterialExpression Interface
 	virtual int32 Compile(class FMaterialCompiler* Compiler, int32 OutputIndex) override;
 	virtual void GetCaption(TArray<FString>& OutCaptions) const override;
+	virtual uint32 GetInputType(int32 InputIndex) override;
 	//~ End UMaterialExpression Interface
 #endif
 
diff --git a/Engine/Source/Runtime/Engine/Private/Materials/MaterialExpressions.cpp b/Engine/Source/Runtime/Engine/Private/Materials/MaterialExpressions.cpp
--- a/Engine/Source/Runtime/Engine/Private/Materials/MaterialExpressions.cpp
+++ b/Engine/Source/Runtime/Engine/Private/Materials/MaterialExpressions.cpp
@@ -19968,6 +19968,19 @@ void UMaterialExpressionThinTranslucentMaterialOutput::GetCaption(TArray<FString
 	OutCaptions.Add(FString(TEXT("Thin Translucent Material")));
 }
 
+uint32 UMaterialExpressionThinTranslucentMaterialOutput::GetInputType(int32 InputIndex)
+{
+	switch (InputIndex)
+	{
+	case 0:
+		return MCT_Float3;
+		break;
+	}
+
+	check(false);
+	return MCT_Float1;
+}
+
 #endif // WITH_EDITOR
 
 int32 UMaterialExpressionThinTranslucentMaterialOutput::GetNumOutputs() const
