Refactor cooked ODSC to use the same data structures between the ShaderCompiler, CookGlobalShaderCommandlet, and the CookingServer.

## Refactoring
- FShaderRecompileData is the de-factor data structure for holding recompile inputs/outputs (these should be separated).
- FShaderRecompileData has moved from being a CookingServer only data structure to living in ShaderCompiler (Engine runtime).
- UE::Cook::FRecompileShaderRequest was basically a duplicate struct of FShaderRecompileData.  Now it just holds a FShaderRecompileData struct and a callback.
- RecompileShadersForRemote(), which previously accepted a grab bag of arguments, now just takes the FShaderRecompileData struct.
- ShaderPlatform is no longer a int32 but a EShaderPlatform to be more explicit.
- The downside to this approach is all these systems are more tightly coupled, but the sharing of code makes this easier to modify in the future.

## Bug Fixes
- recompileshaders changed would only sometimes work depending on what materials the cooking server had loaded.  Now the command issued by the client will explicitly say which materials to load and compile for.
- In RecompileShadersForRemote remove the call to GetOutdatedShaderTypes because in the cooking server it doesn't work.  That function only works when the shader maps exist, but in the cooking server when loading materials we do not create the maps.  This creates a chicken/egg issue.  Where we need to maps to determine what's out of date and to recompile...but we don't have the maps because we haven't done a compile yet.

## Features
- Display how long the shader request took on the cooking server.
- Print shader stats out at the end of recompile request.

diff --git a/Engine/Source/Editor/UnrealEd/Private/Commandlets/CookGlobalShaderCommandlet.cpp b/Engine/Source/Editor/UnrealEd/Private/Commandlets/CookGlobalShaderCommandlet.cpp
--- a/Engine/Source/Editor/UnrealEd/Private/Commandlets/CookGlobalShaderCommandlet.cpp
+++ b/Engine/Source/Editor/UnrealEd/Private/Commandlets/CookGlobalShaderCommandlet.cpp
@@ -174,9 +174,9 @@ int32 UCookGlobalShadersCommandlet::Main(const FString& Params)
 		// Cook shaders
 		UE_LOG(LogCookGlobalShaders, Log, TEXT("Cooking Global Shaders..."));
 		FString OutputDir = FPaths::ProjectSavedDir() / TEXT("CookGlobalShaders") / PlatformName;
-		const TArray<FString> MaterialsToLoad;
-		const TArray<FODSCRequestPayload> ShadersToLoad;
-		RecompileShadersForRemote(PlatformName, SP_NumPlatforms, OutputDir, MaterialsToLoad, ShadersToLoad, nullptr, nullptr, nullptr);
+		TArray<uint8> OutGlobalShaderMap;
+		FShaderRecompileData Arguments(PlatformName, SP_NumPlatforms, ODSCRecompileCommand::Global, nullptr, nullptr, &OutGlobalShaderMap);
+		RecompileShadersForRemote(Arguments, OutputDir);
 
 		bool bCopySucceeded = false;
 
diff --git a/Engine/Source/Editor/UnrealEd/Private/CookOnTheFlyServer.cpp b/Engine/Source/Editor/UnrealEd/Private/CookOnTheFlyServer.cpp
--- a/Engine/Source/Editor/UnrealEd/Private/CookOnTheFlyServer.cpp
+++ b/Engine/Source/Editor/UnrealEd/Private/CookOnTheFlyServer.cpp
@@ -384,7 +384,7 @@ public:
 	virtual bool EnqueueRecompileShaderRequest(UE::Cook::FRecompileShaderRequest RecompileShaderRequest) override
 	{
 		UE_LOG(LogCook, Verbose, TEXT("Enqueing recompile shader(s) request, Platform='%s', ShaderPlatform='%d'"),
-			*RecompileShaderRequest.PlatformName.ToString(), RecompileShaderRequest.ShaderPlatformToCompile);
+			*RecompileShaderRequest.RecompileArguments.PlatformName, RecompileShaderRequest.RecompileArguments.ShaderPlatform);
 
 		Cooker.PackageTracker->RecompileRequests.Enqueue(MoveTemp(RecompileShaderRequest));
 
@@ -665,26 +665,12 @@ void UCookOnTheFlyServer::InitializeShadersForCookOnTheFly(const TArrayView<ITar
 {
 	for (const ITargetPlatform* TargetPlatform : NewTargetPlatforms)
 	{
-		UE_LOG(LogCook, Display, TEXT("Initializing shaders for cook-on-the-fly on platform '%s'"), *TargetPlatform->PlatformName());
-
 		const FString& PlatformName = TargetPlatform->PlatformName();
-		FString OutputDir = GetSandboxDirectory(PlatformName);
+		UE_LOG(LogCook, Display, TEXT("Initializing shaders for cook-on-the-fly on platform '%s'"), *PlatformName);
 
 		TArray<uint8> GlobalShaderMap;
-		UE::Cook::FRecompileShaderRequest RecompileShaderRequest;
-		RecompileShaderRequest.CommandType = ODSCRecompileCommand::Global;
-		RecompileShaderRequest.GlobalShaderMap = &GlobalShaderMap;
-
-		RecompileShadersForRemote(
-			PlatformName, 
-			SP_NumPlatforms,
-			OutputDir, 
-			RecompileShaderRequest.MaterialsToLoad, 
-			RecompileShaderRequest.ShadersToRecompile,
-			RecompileShaderRequest.MeshMaterialMaps, 
-			RecompileShaderRequest.GlobalShaderMap,
-			RecompileShaderRequest.ModifiedFiles,
-			RecompileShaderRequest.CommandType);
+		FShaderRecompileData RecompileArgs(PlatformName, SP_NumPlatforms, ODSCRecompileCommand::Global, nullptr, nullptr, &GlobalShaderMap);
+		RecompileShadersForRemote(RecompileArgs, GetSandboxDirectory(PlatformName));
 	}
 }
 
@@ -3993,20 +3979,7 @@ void UCookOnTheFlyServer::TickRecompileShaderRequests()
 	UE::Cook::FRecompileShaderRequest RecompileShaderRequest;
 	if (PackageTracker->RecompileRequests.Dequeue(&RecompileShaderRequest))
 	{
-		FString PlatformName = RecompileShaderRequest.PlatformName.ToString(); 
-		FString OutputDir = GetSandboxDirectory(PlatformName);
-
-		RecompileShadersForRemote(
-			PlatformName, 
-			RecompileShaderRequest.ShaderPlatformToCompile == -1 ? SP_NumPlatforms : (EShaderPlatform)RecompileShaderRequest.ShaderPlatformToCompile,
-			OutputDir, 
-			RecompileShaderRequest.MaterialsToLoad, 
-			RecompileShaderRequest.ShadersToRecompile,
-			RecompileShaderRequest.MeshMaterialMaps,
-			RecompileShaderRequest.GlobalShaderMap,
-			RecompileShaderRequest.ModifiedFiles,
-			RecompileShaderRequest.CommandType);
-
+		RecompileShadersForRemote(RecompileShaderRequest.RecompileArguments, GetSandboxDirectory(RecompileShaderRequest.RecompileArguments.PlatformName));
 		RecompileShaderRequest.CompletionCallback();
 	}
 }
@@ -6607,17 +6580,11 @@ void UCookOnTheFlyServer::SaveGlobalShaderMapFiles(const TArrayView<const ITarge
 	check( !IsCookingDLC() );
 	for (int32 Index = 0; Index < Platforms.Num(); Index++)
 	{
-		// make sure global shaders are up to date!
 		TArray<FString> Files;
-		FShaderRecompileData RecompileData;
-		RecompileData.PlatformName = Platforms[Index]->PlatformName();
-		// Compile for all platforms
-		RecompileData.ShaderPlatform = -1;
-		RecompileData.ModifiedFiles = &Files;
-		RecompileData.MeshMaterialMaps = nullptr;
 		TArray<uint8> GlobalShaderMap;
-		RecompileData.GlobalShaderMap = &GlobalShaderMap;
-		RecompileData.CommandType = ODSCRecompileCommand::Changed;
+
+		// make sure global shaders are up to date!
+		FShaderRecompileData RecompileData(Platforms[Index]->PlatformName(), SP_NumPlatforms, ODSCRecompileCommand::Changed, &Files, nullptr, &GlobalShaderMap);
 
 		check( IsInGameThread() );
 
@@ -6625,15 +6592,7 @@ void UCookOnTheFlyServer::SaveGlobalShaderMapFiles(const TArrayView<const ITarge
 
 		UE_LOG(LogCook, Display, TEXT("Checking global shaders for platform %s"), *RecompileData.PlatformName);
 
-		RecompileShadersForRemote
-			(RecompileData.PlatformName, 
-			RecompileData.ShaderPlatform == -1 ? SP_NumPlatforms : (EShaderPlatform)RecompileData.ShaderPlatform, //-V547
-			OutputDir, 
-			RecompileData.MaterialsToLoad, 
-			RecompileData.ShadersToRecompile,
-			RecompileData.MeshMaterialMaps,
-			RecompileData.GlobalShaderMap,
-			RecompileData.ModifiedFiles);
+		RecompileShadersForRemote(RecompileData, OutputDir);
 	}
 }
 
diff --git a/Engine/Source/Editor/UnrealEd/Private/Cooker/CookOnTheFlyServerInterface.h b/Engine/Source/Editor/UnrealEd/Private/Cooker/CookOnTheFlyServerInterface.h
--- a/Engine/Source/Editor/UnrealEd/Private/Cooker/CookOnTheFlyServerInterface.h
+++ b/Engine/Source/Editor/UnrealEd/Private/Cooker/CookOnTheFlyServerInterface.h
@@ -4,6 +4,7 @@
 
 #include "CoreTypes.h"
 #include "CookTypes.h"
+#include "ShaderCompiler.h"
 
 struct FODSCRequestPayload;
 class ICookedPackageWriter;
@@ -36,22 +37,8 @@ struct FCookPackageRequest
  */
 struct FRecompileShaderRequest
 {
-	/** The platform name to compile for. */
-	FName PlatformName;
-	/** Shader platform, corresponds to EShaderPlatform, -1 indicates to compile for all target shader platforms. */
-	int32 ShaderPlatformToCompile = -1; // 
-	/** Materials to laod. */
-	TArray<FString> MaterialsToLoad;
-	/** On-demand shader compiler payload.  */
-	TArray<FODSCRequestPayload> ShadersToRecompile;
-	/** Mesh materials, returned to the caller.  */
-	TArray<uint8>* MeshMaterialMaps = nullptr;
-	/** Global shader map, returned to the caller.  */
-	TArray<uint8>* GlobalShaderMap = nullptr;
-	/** All filenames that have been changed during the shader compilation. */
-	TArray<FString>* ModifiedFiles = nullptr;
-	/** What type of shaders to recompile. All, Changed, Global, or Material? */
-	ODSCRecompileCommand CommandType = ODSCRecompileCommand::Changed;
+	/** The arguments to configure shader compilation for this request. */
+	FShaderRecompileData RecompileArguments;
 	/** Completion callback. */
 	FRecompileShaderCompletedCallback CompletionCallback;
 };
diff --git a/Engine/Source/Editor/UnrealEd/Private/Cooker/IoStoreCookOnTheFlyRequestManager.cpp b/Engine/Source/Editor/UnrealEd/Private/Cooker/IoStoreCookOnTheFlyRequestManager.cpp
--- a/Engine/Source/Editor/UnrealEd/Private/Cooker/IoStoreCookOnTheFlyRequestManager.cpp
+++ b/Engine/Source/Editor/UnrealEd/Private/Cooker/IoStoreCookOnTheFlyRequestManager.cpp
@@ -534,16 +534,12 @@ private:
 		TArray<uint8> MeshMaterialMaps;
 		TArray<uint8> GlobalShaderMap;
 
-		FShaderRecompileData RecompileData;
-		RecompileData.PlatformName = Client.PlatformName.ToString();
-		RecompileData.ModifiedFiles = &RecompileModifiedFiles;
-		RecompileData.MeshMaterialMaps = &MeshMaterialMaps;
-		RecompileData.GlobalShaderMap = &GlobalShaderMap;
-
+		FShaderRecompileData RecompileData(Client.PlatformName.ToString(), &RecompileModifiedFiles, &MeshMaterialMaps, &GlobalShaderMap);
 		{
+			int32 iShaderPlatform = static_cast<int32>(RecompileData.ShaderPlatform);
 			TUniquePtr<FArchive> Ar = Request.ReadBody();
 			*Ar << RecompileData.MaterialsToLoad;
-			*Ar << RecompileData.ShaderPlatform;
+			*Ar << iShaderPlatform;
 			*Ar << RecompileData.CommandType;
 			*Ar << RecompileData.ShadersToRecompile;
 		}
@@ -571,14 +567,7 @@ private:
 
 		const bool bEnqueued = CookOnTheFlyServer.EnqueueRecompileShaderRequest(UE::Cook::FRecompileShaderRequest
 		{
-			FName(*RecompileData.PlatformName),
-			RecompileData.ShaderPlatform,
-			RecompileData.MaterialsToLoad,
-			RecompileData.ShadersToRecompile,
-			RecompileData.MeshMaterialMaps,
-			RecompileData.GlobalShaderMap,
-			RecompileData.ModifiedFiles,
-			RecompileData.CommandType,
+			RecompileData, 
 			MoveTemp(RecompileCompleted)
 		});
 		check(bEnqueued);
diff --git a/Engine/Source/Editor/UnrealEd/Private/Cooker/NetworkFileCookOnTheFlyRequestManager.cpp b/Engine/Source/Editor/UnrealEd/Private/Cooker/NetworkFileCookOnTheFlyRequestManager.cpp
--- a/Engine/Source/Editor/UnrealEd/Private/Cooker/NetworkFileCookOnTheFlyRequestManager.cpp
+++ b/Engine/Source/Editor/UnrealEd/Private/Cooker/NetworkFileCookOnTheFlyRequestManager.cpp
@@ -92,14 +92,7 @@ private:
 
 		const bool bEnqueued = CookOnTheFlyServer.EnqueueRecompileShaderRequest(UE::Cook::FRecompileShaderRequest
 		{
-			FName(*RecompileData.PlatformName),
-			RecompileData.ShaderPlatform,
-			RecompileData.MaterialsToLoad,
-			RecompileData.ShadersToRecompile,
-			RecompileData.MeshMaterialMaps,
-			RecompileData.GlobalShaderMap,
-			RecompileData.ModifiedFiles,
-			RecompileData.CommandType,
+			RecompileData,
 			MoveTemp(RecompileCompleted)
 		});
 		check(bEnqueued);
diff --git a/Engine/Source/Runtime/Engine/Private/ShaderCompiler/ShaderCompiler.cpp b/Engine/Source/Runtime/Engine/Private/ShaderCompiler/ShaderCompiler.cpp
--- a/Engine/Source/Runtime/Engine/Private/ShaderCompiler/ShaderCompiler.cpp
+++ b/Engine/Source/Runtime/Engine/Private/ShaderCompiler/ShaderCompiler.cpp
@@ -5956,6 +5956,12 @@ namespace
 		else if (!CmdName.IsEmpty() && FCString::Stricmp(*CmdName, TEXT("Changed")) == 0)
 		{
 			CommandType = ODSCRecompileCommand::Changed;
+
+			// Compile all the shaders that have changed for the materials we have loaded.
+			for (TObjectIterator<UMaterialInterface> It; It; ++It)
+			{
+				OutMaterialsToLoad.Add(It->GetPathName());
+			}
 		}
 		else
 		{
@@ -7056,27 +7062,41 @@ FArchive& operator<<(FArchive& Ar, FODSCRequestPayload& Elem)
 	return Ar;
 }
 
+FShaderRecompileData::FShaderRecompileData(const FString& InPlatformName, TArray<FString>* OutModifiedFiles, TArray<uint8>* OutMeshMaterialMaps, TArray<uint8>* OutGlobalShaderMap)
+: PlatformName(InPlatformName),
+  ModifiedFiles(OutModifiedFiles),
+  MeshMaterialMaps(OutMeshMaterialMaps),
+  GlobalShaderMap(OutGlobalShaderMap)
+{
+}
+
+FShaderRecompileData::FShaderRecompileData(const FString& InPlatformName, EShaderPlatform InShaderPlatform, ODSCRecompileCommand InCommandType, TArray<FString>* OutModifiedFiles, TArray<uint8>* OutMeshMaterialMaps, TArray<uint8>* OutGlobalShaderMap)
+: PlatformName(InPlatformName),
+  ShaderPlatform(InShaderPlatform),
+  ModifiedFiles(OutModifiedFiles),
+  MeshMaterialMaps(OutMeshMaterialMaps),
+  CommandType(InCommandType),
+  GlobalShaderMap(OutGlobalShaderMap)
+{
+}
+
+
 #if WITH_EDITOR
 void RecompileShadersForRemote(
-	const FString& PlatformName,
-	EShaderPlatform ShaderPlatformToCompile,
-	const FString& OutputDirectory,
-	const TArray<FString>& MaterialsToLoad,
-	const TArray<FODSCRequestPayload>& ShadersToRecompile,
-	TArray<uint8>* MeshMaterialMaps,
-	TArray<uint8>* GlobalShaderMap,
-	TArray<FString>* ModifiedFiles,
-	ODSCRecompileCommand RecompileCommandType)
+	FShaderRecompileData& Args,
+	const FString& OutputDirectory)
 {
 	// figure out what shader platforms to recompile
 	ITargetPlatformManagerModule* TPM = GetTargetPlatformManager();
-	ITargetPlatform* TargetPlatform = TPM->FindTargetPlatform(PlatformName);
+	ITargetPlatform* TargetPlatform = TPM->FindTargetPlatform(Args.PlatformName);
 	if (TargetPlatform == NULL)
 	{
-		UE_LOG(LogShaders, Display, TEXT("Failed to find target platform module for %s"), *PlatformName);
+		UE_LOG(LogShaders, Display, TEXT("Failed to find target platform module for %s"), *Args.PlatformName);
 		return;
 	}
 
+	const double StartTime = FPlatformTime::Seconds();
+
 	UE_LOG(LogShaders, Display, TEXT(""));
 	UE_LOG(LogShaders, Display, TEXT("********************************"));
 	UE_LOG(LogShaders, Display, TEXT("Received compile shader request."));
@@ -7087,14 +7107,14 @@ void RecompileShadersForRemote(
 	TArray<FName> DesiredShaderFormats;
 	TargetPlatform->GetAllTargetedShaderFormats(DesiredShaderFormats);
 
-	UE_LOG(LogShaders, Verbose, TEXT("Loading %d materials..."), MaterialsToLoad.Num());
+	UE_LOG(LogShaders, Verbose, TEXT("Loading %d materials..."), Args.MaterialsToLoad.Num());
 	// make sure all materials the client has loaded will be processed
 	TArray<UMaterialInterface*> MaterialsToCompile;
 
-	for (int32 Index = 0; Index < MaterialsToLoad.Num(); Index++)
+	for (int32 Index = 0; Index < Args.MaterialsToLoad.Num(); Index++)
 	{
-		UE_LOG(LogShaders, Verbose, TEXT("   --> %s"), *MaterialsToLoad[Index]);
-		MaterialsToCompile.Add(LoadObject<UMaterialInterface>(NULL, *MaterialsToLoad[Index]));
+		UE_LOG(LogShaders, Verbose, TEXT("   --> %s"), *Args.MaterialsToLoad[Index]);
+		MaterialsToCompile.Add(LoadObject<UMaterialInterface>(NULL, *Args.MaterialsToLoad[Index]));
 	}
 
 	UE_LOG(LogShaders, Verbose, TEXT("  Done!"))
@@ -7107,19 +7127,12 @@ void RecompileShadersForRemote(
 	// Pick up new changes to shader files
 	FlushShaderFileCache();
 
-	const bool bCompileChangedShaders = RecompileCommandType == ODSCRecompileCommand::Changed;
-	if (bCompileChangedShaders)
+	if (Args.ShadersToRecompile.Num())
 	{
-		GetOutdatedShaderTypes(OutdatedShaderTypes, OutdatedShaderPipelineTypes, OutdatedFactoryTypes);
-		UE_LOG(LogShaders, Display, TEXT("We found %d out of date shader types, %d outdated pipeline types, and %d out of date VF types!"), OutdatedShaderTypes.Num(), OutdatedShaderPipelineTypes.Num(), OutdatedFactoryTypes.Num());
+		UE_LOG(LogShaders, Display, TEXT("Received %d shaders to compile."), Args.ShadersToRecompile.Num());
 	}
 
-	if (ShadersToRecompile.Num())
-	{
-		UE_LOG(LogShaders, Display, TEXT("Received %d shaders to compile."), ShadersToRecompile.Num());
-	}
-
-	for (const FODSCRequestPayload& payload: ShadersToRecompile)
+	for (const FODSCRequestPayload& payload: Args.ShadersToRecompile)
 	{
 		UE_LOG(LogShaders, Display, TEXT(""));
 		UE_LOG(LogShaders, Display, TEXT("\tMaterial:    %s "), *payload.MaterialName);
@@ -7158,11 +7171,11 @@ void RecompileShadersForRemote(
 			const EShaderPlatform ShaderPlatform = ShaderFormatToLegacyShaderPlatform(DesiredShaderFormats[FormatIndex]);
 
 			// Only compile for the desired platform if requested
-			if (ShaderPlatform == ShaderPlatformToCompile || ShaderPlatformToCompile == SP_NumPlatforms)
+			if (ShaderPlatform == Args.ShaderPlatform || Args.ShaderPlatform == SP_NumPlatforms)
 			{
 				// If we are explicitly wanting to recompile global or if shaders have changed.
-				if (RecompileCommandType == ODSCRecompileCommand::Global ||
-					RecompileCommandType == ODSCRecompileCommand::Changed)
+				if (Args.CommandType == ODSCRecompileCommand::Global ||
+					Args.CommandType == ODSCRecompileCommand::Changed)
 				{
 					UE_LOG(LogShaders, Display, TEXT("Recompiling global shaders."));
 
@@ -7182,7 +7195,7 @@ void RecompileShadersForRemote(
 					FinishRecompileGlobalShaders();
 
 					// write the shader compilation info to memory, converting fnames to strings
-					FMemoryWriter MemWriter(*GlobalShaderMap, true);
+					FMemoryWriter MemWriter(*Args.GlobalShaderMap, true);
 					FNameAsStringProxyArchive Ar(MemWriter);
 					Ar.SetCookingTarget(TargetPlatform);
 
@@ -7190,16 +7203,15 @@ void RecompileShadersForRemote(
 					SaveGlobalShadersForRemoteRecompile(Ar, ShaderPlatform);
 				}
 
-				// we only want to actually compile mesh shaders if a client directly requested it, and there's actually some work to do
-				if ((RecompileCommandType == ODSCRecompileCommand::Material || RecompileCommandType == ODSCRecompileCommand::Changed) &&
-					MeshMaterialMaps != nullptr && 
-					(OutdatedShaderTypes.Num() || OutdatedFactoryTypes.Num() || bCompileChangedShaders == false))
+				// we only want to actually compile mesh shaders if a client directly requested it
+				if ((Args.CommandType == ODSCRecompileCommand::Material || Args.CommandType == ODSCRecompileCommand::Changed) &&
+					Args.MeshMaterialMaps != nullptr)
 				{
 					TMap<FString, TArray<TRefCountPtr<FMaterialShaderMap> > > CompiledShaderMaps;
 					UMaterial::CompileMaterialsForRemoteRecompile(MaterialsToCompile, ShaderPlatform, TargetPlatform, CompiledShaderMaps);
 
 					// write the shader compilation info to memory, converting fnames to strings
-					FMemoryWriter MemWriter(*MeshMaterialMaps, true);
+					FMemoryWriter MemWriter(*Args.MeshMaterialMaps, true);
 					FNameAsStringProxyArchive Ar(MemWriter);
 					Ar.SetCookingTarget(TargetPlatform);
 
@@ -7211,21 +7223,24 @@ void RecompileShadersForRemote(
 				FString GlobalShaderFilename = SaveGlobalShaderFile(ShaderPlatform, OutputDirectory, TargetPlatform);
 
 				// add this to the list of files to tell the other end about
-				if (ModifiedFiles)
+				if (Args.ModifiedFiles)
 				{
 					// need to put it in non-sandbox terms
 					FString SandboxPath(GlobalShaderFilename);
 					check(SandboxPath.StartsWith(OutputDirectory));
 					SandboxPath.ReplaceInline(*OutputDirectory, TEXT("../../../"));
 					FPaths::NormalizeFilename(SandboxPath);
-					ModifiedFiles->Add(SandboxPath);
+					Args.ModifiedFiles->Add(SandboxPath);
 				}
 			}
 		}
 	}
 
 	UE_LOG(LogShaders, Display, TEXT(""));
-	UE_LOG(LogShaders, Display, TEXT("Finished shader compile request."));
+	UE_LOG(LogShaders, Display, TEXT("Finished shader compile request in %.2f seconds."), FPlatformTime::Seconds() - StartTime);
+
+	// Provide a log of what happened.
+	GShaderCompilingManager->PrintStats(true);
 
 	// Restore compilation state.
 	GShaderCompilingManager->SkipShaderCompilation(bPreviousState);
diff --git a/Engine/Source/Runtime/Engine/Public/ShaderCompiler.h b/Engine/Source/Runtime/Engine/Public/ShaderCompiler.h
--- a/Engine/Source/Runtime/Engine/Public/ShaderCompiler.h
+++ b/Engine/Source/Runtime/Engine/Public/ShaderCompiler.h
@@ -1287,28 +1287,66 @@ enum class ODSCRecompileCommand
 	Material
 };
 
+struct FShaderRecompileData
+{
+	/** The platform name to compile for. */
+	FString PlatformName;
+
+	/** Shader platform */
+	EShaderPlatform ShaderPlatform = SP_NumPlatforms;
+
+	/** All filenames that have been changed during the shader compilation. */
+	TArray<FString>* ModifiedFiles = nullptr;
+
+	/** Mesh materials, returned to the caller.  */
+	TArray<uint8>* MeshMaterialMaps = nullptr;
+
+	/** Materials to load. */
+	TArray<FString> MaterialsToLoad;
+
+	/** What type of shaders to recompile. All, Changed, Global, or Material? */
+	ODSCRecompileCommand CommandType = ODSCRecompileCommand::Changed;
+
+	/** Global shader map, returned to the caller.  */
+	TArray<uint8>* GlobalShaderMap = nullptr;
+
+	/** On-demand shader compiler payload.  */
+	TArray<FODSCRequestPayload> ShadersToRecompile;
+
+	/** Default constructor. */
+	ENGINE_API FShaderRecompileData() {};
+
+	/** Recompile all the changed shaders for the current platform. */
+	ENGINE_API FShaderRecompileData(const FString& InPlatformName, TArray<FString>* OutModifiedFiles, TArray<uint8>* OutMeshMaterialMaps, TArray<uint8>* OutGlobalShaderMap);
+
+	/** For recompiling just global shaders. */
+	ENGINE_API FShaderRecompileData(const FString& InPlatformName, EShaderPlatform InShaderPlatform, ODSCRecompileCommand InCommandType, TArray<FString>* OutModifiedFiles, TArray<uint8>* OutMeshMaterialMaps, TArray<uint8>* OutGlobalShaderMap);
+
+	FShaderRecompileData& operator=(const FShaderRecompileData& Other)
+	{
+		PlatformName = Other.PlatformName;
+		ShaderPlatform = Other.ShaderPlatform;
+		ModifiedFiles = Other.ModifiedFiles;
+		MeshMaterialMaps = Other.MeshMaterialMaps;
+		MaterialsToLoad = Other.MaterialsToLoad;
+		CommandType = Other.CommandType;
+		GlobalShaderMap = Other.GlobalShaderMap;
+
+		ShadersToRecompile = Other.ShadersToRecompile;
+
+		return *this;
+	}
+};
+
 #if WITH_EDITOR
 
 /**
 * Recompiles global shaders
 *
-* @param PlatformName					Name of the Platform the shaders are compiled for
-* @param OutputDirectory				The directory the compiled data will be stored to
-* @param MaterialsToLoad				List of Materials that need to be loaded and compiled
-* @param MeshMaterialMaps				Mesh material maps
-* @param ModifiedFiles					Returns the list of modified files if not NULL
-* @param bCompileChangedShaders		Whether to compile all changed shaders or the specific material that is passed
+* @param Args					Arguments and configuration for issuing recompiles.
+* @param OutputDirectory		The directory the compiled data will be stored to
 **/
-extern ENGINE_API void RecompileShadersForRemote(
-	const FString& PlatformName,
-	EShaderPlatform ShaderPlatform,
-	const FString& OutputDirectory,
-	const TArray<FString>& MaterialsToLoad,
-	const TArray<FODSCRequestPayload>& ShadersToRecompile,
-	TArray<uint8>* MeshMaterialMaps,
-	TArray<uint8>* GlobalShaderMap,
-	TArray<FString>* ModifiedFiles,
-	ODSCRecompileCommand RecompileCommandType = ODSCRecompileCommand::Changed);
+extern ENGINE_API void RecompileShadersForRemote(FShaderRecompileData& Args, const FString& OutputDirectory);
 
 #endif // WITH_EDITOR
 
diff --git a/Engine/Source/Runtime/NetworkFileSystem/Private/NetworkFileServerConnection.cpp b/Engine/Source/Runtime/NetworkFileSystem/Private/NetworkFileServerConnection.cpp
--- a/Engine/Source/Runtime/NetworkFileSystem/Private/NetworkFileServerConnection.cpp
+++ b/Engine/Source/Runtime/NetworkFileSystem/Private/NetworkFileServerConnection.cpp
@@ -1014,15 +1014,8 @@ bool FNetworkFileServerClientConnection::ProcessGetFileList( FArchive& In, FArch
 	// make sure the global shaders are up to date before letting the client read any shaders
 	// @todo: This will probably add about 1/2 second to the boot-up time of the client while the server does this
 	// @note: We assume the delegate will write to the proper sandbox directory, should we pass in SandboxDirectory, or Sandbox?
-	FShaderRecompileData RecompileData;
-	RecompileData.PlatformName = ConnectedPlatformName;
-	// All target platforms
-	RecompileData.ShaderPlatform = -1;
-	RecompileData.ModifiedFiles = nullptr;
-	RecompileData.MeshMaterialMaps = nullptr;
 	TArray<uint8> GlobalShaderMap;
-	RecompileData.GlobalShaderMap = &GlobalShaderMap;
-	RecompileData.CommandType = ODSCRecompileCommand::Global;
+	FShaderRecompileData RecompileData(ConnectedPlatformName, SP_NumPlatforms, ODSCRecompileCommand::Global, nullptr, nullptr, &GlobalShaderMap);
 	NetworkFileDelegates->RecompileShadersDelegate.ExecuteIfBound(RecompileData);
 
 	UE_LOG(LogFileServer, Display, TEXT("Getting files for %d directories, game = %s, platform = %s"), RootDirectories.Num(), *GameName, *ConnectedPlatformName);
@@ -1283,15 +1276,13 @@ void FNetworkFileServerClientConnection::ProcessRecompileShaders( FArchive& In,
 	TArray<FString> RecompileModifiedFiles;
 	TArray<uint8> MeshMaterialMaps;
 	TArray<uint8> GlobalShaderMap;
-	FShaderRecompileData RecompileData;
-	RecompileData.PlatformName = ConnectedPlatformName;
-	RecompileData.ModifiedFiles = &RecompileModifiedFiles;
-	RecompileData.MeshMaterialMaps = &MeshMaterialMaps;
-	RecompileData.GlobalShaderMap = &GlobalShaderMap;
+	FShaderRecompileData RecompileData(ConnectedPlatformName, SP_NumPlatforms, ODSCRecompileCommand::Changed, &RecompileModifiedFiles, &MeshMaterialMaps, &GlobalShaderMap);
 
 	// tell other side all the materials to load, by pathname
 	In << RecompileData.MaterialsToLoad;
-	In << RecompileData.ShaderPlatform;
+
+	int32 iShaderPlatform = static_cast<int32>(RecompileData.ShaderPlatform);
+	In << iShaderPlatform;
 	In << RecompileData.CommandType;
 	In << RecompileData.ShadersToRecompile;
 
diff --git a/Engine/Source/Runtime/NetworkFileSystem/Public/INetworkFileSystemModule.h b/Engine/Source/Runtime/NetworkFileSystem/Public/INetworkFileSystemModule.h
--- a/Engine/Source/Runtime/NetworkFileSystem/Public/INetworkFileSystemModule.h
+++ b/Engine/Source/Runtime/NetworkFileSystem/Public/INetworkFileSystemModule.h
@@ -16,40 +16,6 @@ class INetworkFileServer;
  */
 DECLARE_DELEGATE_ThreeParams(FFileRequestDelegate, FString&, const FString&, TArray<FString>&);
 
-struct FShaderRecompileData
-{
-	FString PlatformName;
-	/** The platform to compile shaders for, corresponds to EShaderPlatform, but a value of -1 indicates to compile for all target shader platforms. */
-	int32 ShaderPlatform;
-	TArray<FString>* ModifiedFiles = nullptr;
-	TArray<uint8>* MeshMaterialMaps = nullptr;
-	TArray<FString> MaterialsToLoad;
-	ODSCRecompileCommand CommandType = ODSCRecompileCommand::Changed;
-	TArray<uint8>* GlobalShaderMap = nullptr;
-
-	TArray<FODSCRequestPayload> ShadersToRecompile;
-
-	FShaderRecompileData() :
-		ShaderPlatform(-1)
-	{}
-
-	FShaderRecompileData& operator=(const FShaderRecompileData& Other)
-	{
-		PlatformName = Other.PlatformName;
-		ShaderPlatform = Other.ShaderPlatform;
-		ModifiedFiles = Other.ModifiedFiles;
-		MeshMaterialMaps = Other.MeshMaterialMaps;
-		MaterialsToLoad = Other.MaterialsToLoad;
-		CommandType = Other.CommandType;
-		GlobalShaderMap = Other.GlobalShaderMap;
-
-		ShadersToRecompile = Other.ShadersToRecompile;
-
-		return *this;
-	}
-};
-
-
 /**
  * Delegate type for handling shader recompilation requests from a network client.
  */
