Remove copies of animation curve data (names and types) from BoneContainer

diff --git a/Engine/Plugins/Animation/ControlRig/Source/ControlRig/Private/AnimNode_ControlRig.cpp b/Engine/Plugins/Animation/ControlRig/Source/ControlRig/Private/AnimNode_ControlRig.cpp
--- a/Engine/Plugins/Animation/ControlRig/Source/ControlRig/Private/AnimNode_ControlRig.cpp
+++ b/Engine/Plugins/Animation/ControlRig/Source/ControlRig/Private/AnimNode_ControlRig.cpp
@@ -132,10 +132,10 @@ void FAnimNode_ControlRig::CacheBones_AnyThread(const FAnimationCacheBonesContex
 	if(RequiredBones.IsValid())
 	{
 		RefPoseSetterHash.Reset();
-		
-		TArray<FName> const& UIDToNameLookUpTable = RequiredBones.GetUIDToNameLookupTable();
 
-		auto CacheMapping = [&](const TMap<FName, FName>& Mapping, TArray<FName> const& InUIDToNameLookUpTable, 
+		const FSmartNameMapping* CurveMapping = RequiredBones.GetSkeletonAsset()->GetSmartNameContainer(USkeleton::AnimCurveMappingName);
+
+		auto CacheMapping = [&](const TMap<FName, FName>& Mapping, const FSmartNameMapping* CurveNameMapping, 
 			const FAnimationCacheBonesContext& InContext, URigHierarchy* InHierarchy)
 		{
 			for (auto Iter = Mapping.CreateConstIterator(); Iter; ++Iter)
@@ -146,8 +146,8 @@ void FAnimNode_ControlRig::CacheBones_AnyThread(const FAnimationCacheBonesContex
 
 				if (SourcePath != NAME_None && TargetPath != NAME_None)
 				{
-					int32 Found = InUIDToNameLookUpTable.Find(TargetPath);
-					if (Found != INDEX_NONE)
+					const SmartName::UID_Type Found = CurveNameMapping->FindUID(TargetPath);
+					if (Found != SmartName::MaxUID)
 					{
 						// set value - sound should be UID
 						InputToCurveMappingUIDs.Add(Iter.Value()) = Found;
@@ -177,8 +177,8 @@ void FAnimNode_ControlRig::CacheBones_AnyThread(const FAnimationCacheBonesContex
 			Hierarchy = CurrentControlRig->GetHierarchy();
 		}
 
-		CacheMapping(InputMapping, UIDToNameLookUpTable, Context, Hierarchy);
-		CacheMapping(OutputMapping, UIDToNameLookUpTable, Context, Hierarchy);
+		CacheMapping(InputMapping, CurveMapping, Context, Hierarchy);
+		CacheMapping(OutputMapping, CurveMapping, Context, Hierarchy);
 	}
 }
 
diff --git a/Engine/Plugins/Animation/ControlRig/Source/ControlRig/Private/AnimNode_ControlRigBase.cpp b/Engine/Plugins/Animation/ControlRig/Source/ControlRig/Private/AnimNode_ControlRigBase.cpp
--- a/Engine/Plugins/Animation/ControlRig/Source/ControlRig/Private/AnimNode_ControlRigBase.cpp
+++ b/Engine/Plugins/Animation/ControlRig/Source/ControlRig/Private/AnimNode_ControlRigBase.cpp
@@ -643,15 +643,23 @@ void FAnimNode_ControlRigBase::CacheBones_AnyThread(const FAnimationCacheBonesCo
 			}
 			
 			// we just support curves by name only
-			TArray<FName> const& CurveNames = RequiredBones.GetUIDToNameLookupTable();
-			URigHierarchy* Hierarchy = ControlRig->GetHierarchy();
-			for (uint16 Index = 0; Index < CurveNames.Num(); ++Index)
 			{
-				// see if the curve name exists in the control rig
-				if (Hierarchy->GetIndex(FRigElementKey(CurveNames[Index], ERigElementType::Curve)) != INDEX_NONE)
+				const FSmartNameMapping* CurveMapping = RequiredBones.GetSkeletonAsset()->GetSmartNameContainer(USkeleton::AnimCurveMappingName);
+				URigHierarchy* Hierarchy = ControlRig->GetHierarchy();
+				CurveMapping->Iterate([Hierarchy, this](const FSmartNameMapping* Mapping, SmartName::UID_Type ID)
 				{
-					ControlRigCurveMappingByName.Add(CurveNames[Index], Index);
-				}
+					// see if the curve name exists in the control rig
+					FName CurveName;
+					if (Mapping->GetName(ID, CurveName))
+					{
+						if (Hierarchy->GetIndex(FRigElementKey(CurveName, ERigElementType::Curve)) != INDEX_NONE)
+						{
+							ControlRigCurveMappingByName.Add(CurveName, ID);
+						}
+					}
+				});
+				
+					
 			}
 
 			// check if we can switch the bones to an index based mapping.
diff --git a/Engine/Source/Runtime/Engine/Classes/Animation/SmartName.h b/Engine/Source/Runtime/Engine/Classes/Animation/SmartName.h
--- a/Engine/Source/Runtime/Engine/Classes/Animation/SmartName.h
+++ b/Engine/Source/Runtime/Engine/Classes/Animation/SmartName.h
@@ -12,6 +12,7 @@
 
 struct FSmartName;
 
+
 /** Curve Meta Data for each name
  * Unfortunately this should be linked to FName, but no GUID because we don't have GUID in run-time
  * We only add this if anything changed, by default, it is attribute curve
@@ -61,6 +62,9 @@ struct ENGINE_API FSmartNameMapping
 	GENERATED_USTRUCT_BODY();
 
 	FSmartNameMapping();
+	FSmartNameMapping(FRWLock *Lock);
+	
+	void SetLock(FRWLock *Lock) { RWLock = Lock; }
 
 	// Add a name to the mapping, fails if it already exists
 	// @param InName - The name to add
@@ -130,23 +134,11 @@ struct ENGINE_API FSmartNameMapping
 	bool FindSmartNameByUID(SmartName::UID_Type UID, FSmartName& OutName) const;
 
 	// Curve Meta Data Accessors
-	FCurveMetaData* GetCurveMetaData(FName CurveName)
-	{
-		checkSlow(Exists(CurveName));
-		return &CurveMetaDataMap.FindOrAdd(CurveName);
-	}
-
-	const FCurveMetaData* GetCurveMetaData(FName CurveName) const
-	{
-		checkSlow(Exists(CurveName));
-		return CurveMetaDataMap.Find(CurveName);
-	}
-
+	FCurveMetaData* GetCurveMetaData(FName CurveName);
+	const FCurveMetaData* GetCurveMetaData(FName CurveName) const;
+	
 #if !WITH_EDITOR
-	const FCurveMetaData& GetCurveMetaData(SmartName::UID_Type CurveUID) const
-	{
-		return CurveMetaDataList[CurveUID];
-	}
+	const FCurveMetaData& GetCurveMetaData(SmartName::UID_Type CurveUID) const;
 #endif
 
 	// Serialize this to the provided archive; required for TMap serialization
@@ -160,6 +152,9 @@ struct ENGINE_API FSmartNameMapping
 	/** Get the maximum in use UID */
 	SmartName::UID_Type GetMaxUID() const { return (SmartName::UID_Type)(CurveNameList.Num() - 1); }
 
+	/** Iterate over all Names in this Mapping */
+	void Iterate(TFunction<void(const FSmartNameMapping* Mapping, SmartName::UID_Type ID)> Callback) const;
+
 private:
 	// List of curve names, indexed by UID
 	TArray<FName> CurveNameList;
@@ -170,6 +165,8 @@ private:
 #endif
 
 	TMap<FName, FCurveMetaData> CurveMetaDataMap;
+
+	FRWLock* RWLock;
 };
 
 USTRUCT()
@@ -204,6 +201,17 @@ private:
 	// Editor copy of the data we loaded, used to preserve determinism during cooking
 	TMap<FName, FSmartNameMapping> LoadedNameMappings;
 #endif
+	
+	mutable FRWLock RWLock;
+};
+
+template<>
+struct TStructOpsTypeTraits<FSmartNameContainer> : public TStructOpsTypeTraitsBase2<FSmartNameContainer>
+{
+	enum
+	{
+		WithCopy = false,
+	};
 };
 
 USTRUCT()
diff --git a/Engine/Source/Runtime/Engine/Private/Animation/AnimInstanceProxy.cpp b/Engine/Source/Runtime/Engine/Private/Animation/AnimInstanceProxy.cpp
--- a/Engine/Source/Runtime/Engine/Private/Animation/AnimInstanceProxy.cpp
+++ b/Engine/Source/Runtime/Engine/Private/Animation/AnimInstanceProxy.cpp
@@ -3047,8 +3047,7 @@ void FAnimInstanceProxy::UpdateCurvesToEvaluationContext(const FAnimationEvaluat
 	if(InContext.Curve.UIDToArrayIndexLUT != nullptr)
 	{
 		const TArray<uint16>& UIDToArrayIndexLookupTable = RequiredBones.GetUIDToArrayLookupTable();
-		const TArray<FName>& UIDToNameLookupTable = RequiredBones.GetUIDToNameLookupTable();
-		const TArray<FAnimCurveType>& UIDToCurveTypeLookupTable = RequiredBones.GetUIDToCurveTypeLookupTable();
+		const FSmartNameMapping* Mapping = RequiredBones.GetSkeletonAsset()->GetSmartNameContainer(USkeleton::AnimCurveMappingName);
 	
 		check(UIDToArrayIndexLookupTable.Num() == InContext.Curve.UIDToArrayIndexLUT->Num());
 
@@ -3060,21 +3059,31 @@ void FAnimInstanceProxy::UpdateCurvesToEvaluationContext(const FAnimationEvaluat
 				ensureAlwaysMsgf(InContext.Curve.CurveWeights.IsValidIndex(ArrayIndex), TEXT("%s Animation Instance contains out of bound UIDList."), *AnimInstanceObject->GetClass()->GetName()) &&
 				InContext.Curve.ValidCurveWeights[ArrayIndex])
 			{
-				const FName& CurveName = UIDToNameLookupTable[CurveUID];
-				const FAnimCurveType& CurveType = UIDToCurveTypeLookupTable[CurveUID];
-				const float CurveWeight = InContext.Curve.CurveWeights[ArrayIndex];
-				
-				AnimationCurves[(uint8)EAnimCurveType::AttributeCurve].Add(CurveName, CurveWeight);
-
-				if(CurveType.bMorphtarget)
+				FName CurveName;
+				if (Mapping->GetName(CurveUID, CurveName))
 				{
-					AnimationCurves[(uint8)EAnimCurveType::MorphTargetCurve].Add(CurveName, CurveWeight);
-				}
+#if !WITH_EDITOR
+					const FCurveMetaData* CurveMetaData = &Mapping->GetCurveMetaData(CurveUID);
+#else
+					if (const FCurveMetaData* CurveMetaData = Mapping->GetCurveMetaData(CurveName))
+#endif
+					{
+						const FAnimCurveType& CurveType = CurveMetaData->Type;
+						const float CurveWeight = InContext.Curve.CurveWeights[ArrayIndex];
+					
+						AnimationCurves[(uint8)EAnimCurveType::AttributeCurve].Add(CurveName, CurveWeight);
 
-				if(CurveType.bMaterial)
-				{
-					MaterialParametersToClear.Remove(CurveName);
-					AnimationCurves[(uint8)EAnimCurveType::MaterialCurve].Add(CurveName, CurveWeight);
+						if(CurveType.bMorphtarget)
+						{
+							AnimationCurves[(uint8)EAnimCurveType::MorphTargetCurve].Add(CurveName, CurveWeight);
+						}
+
+						if(CurveType.bMaterial)
+						{
+							MaterialParametersToClear.Remove(CurveName);
+							AnimationCurves[(uint8)EAnimCurveType::MaterialCurve].Add(CurveName, CurveWeight);
+						}
+					}	
 				}
 			}
 		}
diff --git a/Engine/Source/Runtime/Engine/Private/Animation/SmartName.cpp b/Engine/Source/Runtime/Engine/Private/Animation/SmartName.cpp
--- a/Engine/Source/Runtime/Engine/Private/Animation/SmartName.cpp
+++ b/Engine/Source/Runtime/Engine/Private/Animation/SmartName.cpp
@@ -4,6 +4,7 @@
 #include "UObject/FrameworkObjectVersion.h"
 #include "Animation/Skeleton.h"
 #include "UObject/AnimPhysObjectVersion.h"
+#include "Misc/ScopeRWLock.h"
 
 ////////////////////////////////////////////////////////////////////////
 //
@@ -11,10 +12,28 @@
 //
 ///////////////////////////////////////////////////////////////////////
 FSmartNameMapping::FSmartNameMapping()
-{}
+{
+}
+
+FSmartNameMapping::FSmartNameMapping(FRWLock* Lock)
+	: RWLock(Lock)
+{
+}
+
+void FSmartNameMapping::Iterate(TFunction<void(const FSmartNameMapping* Mapping, SmartName::UID_Type ID)> Callback) const
+{
+	FReadScopeLock Lock(*RWLock);
+
+	for (int32 NameIndex = 0; NameIndex < CurveNameList.Num(); ++NameIndex)
+	{
+		Callback(this, NameIndex);	
+	}
+}
 
 FSmartName FSmartNameMapping::AddName(FName Name)
 {
+	FWriteScopeLock Lock(*RWLock);
+	
 	// Make sure we are not trying to do an invalid add
 	check(Name.IsValid() && !CurveNameList.Contains(Name));
 
@@ -29,8 +48,32 @@ FSmartName FSmartNameMapping::AddName(FName Name)
 	return NewSmartName;
 }
 
+FCurveMetaData* FSmartNameMapping::GetCurveMetaData(FName CurveName)
+{
+	FWriteScopeLock Lock(*RWLock);
+	checkSlow(Exists(CurveName));
+	return &CurveMetaDataMap.FindOrAdd(CurveName);
+}
+
+const FCurveMetaData* FSmartNameMapping::GetCurveMetaData(FName CurveName) const
+{
+	FReadScopeLock Lock(*RWLock);
+	checkSlow(Exists(CurveName));
+	return CurveMetaDataMap.Find(CurveName);
+}
+
+#if !WITH_EDITOR
+const FCurveMetaData& FSmartNameMapping::GetCurveMetaData(SmartName::UID_Type CurveUID) const
+{
+	FReadScopeLock Lock(*RWLock);
+	return CurveMetaDataList[CurveUID];
+}
+#endif
+
 bool FSmartNameMapping::GetName(const SmartName::UID_Type& Uid, FName& OutName) const
 {
+	FReadScopeLock Lock(*RWLock);
+	
 	if (CurveNameList.IsValidIndex(Uid))
 	{
 		OutName = CurveNameList[Uid];
@@ -42,6 +85,8 @@ bool FSmartNameMapping::GetName(const SmartName::UID_Type& Uid, FName& OutName)
 #if WITH_EDITOR
 bool FSmartNameMapping::Rename(const SmartName::UID_Type& Uid, FName NewName)
 {
+	FWriteScopeLock Lock(*RWLock);
+	
 	FName ExistingName;
 	if(GetName(Uid, ExistingName))
 	{
@@ -64,6 +109,8 @@ bool FSmartNameMapping::Rename(const SmartName::UID_Type& Uid, FName NewName)
 
 bool FSmartNameMapping::Remove(const SmartName::UID_Type& Uid)
 {
+	FWriteScopeLock Lock(*RWLock);
+	
 	FName ExistingName;
 	if (GetName(Uid, ExistingName))
 	{
@@ -77,6 +124,8 @@ bool FSmartNameMapping::Remove(const SmartName::UID_Type& Uid)
 
 bool FSmartNameMapping::Remove(const FName& Name)
 {
+	FWriteScopeLock Lock(*RWLock);
+	
 	const SmartName::UID_Type Uid = FindUID(Name);
 	if (Uid != SmartName::MaxUID)
 	{
@@ -125,6 +174,8 @@ void FSmartNameMapping::Serialize(FArchive& Ar)
 
 void FSmartNameMapping::FillUidArray(TArray<SmartName::UID_Type>& Array) const
 {
+	FReadScopeLock Lock(*RWLock);
+	
 	Array.Reset(CurveNameList.Num());
 	
 	for (int32 NameIndex = 0; NameIndex < CurveNameList.Num(); ++NameIndex)
@@ -143,6 +194,8 @@ void FSmartNameMapping::FillUidArray(TArray<SmartName::UID_Type>& Array) const
 
 void FSmartNameMapping::FillNameArray(TArray<FName>& Array) const
 {
+	FReadScopeLock Lock(*RWLock);
+	
 	//In editor names can be removed and so have to deal with empty slots
 #if WITH_EDITOR
 	Array.Reset(CurveNameList.Num());
@@ -165,6 +218,7 @@ void FSmartNameMapping::FillUIDToNameArray(TArray<FName>& Array) const
 
 void FSmartNameMapping::FillCurveTypeArray(TArray<FAnimCurveType>& Array) const
 {
+	FReadScopeLock Lock(*RWLock);
 	Array.Reset(CurveNameList.Num());
 
 	for (const FName& Name : CurveNameList)
@@ -186,6 +240,8 @@ void FSmartNameMapping::FillCurveTypeArray(TArray<FAnimCurveType>& Array) const
 
 void FSmartNameMapping::FillUIDToCurveTypeArray(TArray<FAnimCurveType>& Array) const
 {
+	FReadScopeLock Lock(*RWLock);
+	
 	Array.Reset(CurveNameList.Num());
 
 	for (const FName& Name : CurveNameList)
@@ -212,16 +268,19 @@ void FSmartNameMapping::FillUIDToCurveTypeArray(TArray<FAnimCurveType>& Array) c
 
 bool FSmartNameMapping::Exists(const SmartName::UID_Type& Uid) const
 {
+	FReadScopeLock Lock(*RWLock);
 	return CurveNameList.IsValidIndex(Uid) && CurveNameList[Uid] != NAME_None;
 }
 
 bool FSmartNameMapping::Exists(const FName& Name) const
 {
+	FReadScopeLock Lock(*RWLock);
 	return CurveNameList.Contains(Name);
 }
 
 SmartName::UID_Type FSmartNameMapping::FindUID(const FName& Name) const
 {
+	FReadScopeLock Lock(*RWLock);
 	return CurveNameList.IndexOfByKey(Name);
 }
 
@@ -234,6 +293,7 @@ FArchive& operator<<(FArchive& Ar, FSmartNameMapping& Elem)
 
 bool FSmartNameMapping::FindSmartName(FName Name, FSmartName& OutName) const
 {
+	FReadScopeLock Lock(*RWLock);
 	SmartName::UID_Type ExistingUID = FindUID(Name);
 	if (ExistingUID != SmartName::MaxUID)
 	{
@@ -246,6 +306,7 @@ bool FSmartNameMapping::FindSmartName(FName Name, FSmartName& OutName) const
 
 bool FSmartNameMapping::FindSmartNameByUID(SmartName::UID_Type UID, FSmartName& OutName) const
 {
+	FReadScopeLock Lock(*RWLock);
 	FName ExistingName;
 	if (GetName(UID, ExistingName))
 	{
@@ -260,6 +321,8 @@ bool FSmartNameMapping::FindSmartNameByUID(SmartName::UID_Type UID, FSmartName&
 /* initialize curve meta data for the container */
 void FSmartNameMapping::InitializeCurveMetaData(class USkeleton* Skeleton)
 {
+	FWriteScopeLock Lock(*RWLock);
+	
 	// initialize bone indices for skeleton
 	for (TPair<FName, FCurveMetaData>& Iter : CurveMetaDataMap)
 	{
@@ -274,13 +337,15 @@ void FSmartNameMapping::InitializeCurveMetaData(class USkeleton* Skeleton)
 //
 // FSmartNameContainer
 //
-///////////////////////////////////////////////////////////////////////
+//////////////////////////////////////////////////////////////////////
 FSmartNameMapping* FSmartNameContainer::AddContainer(FName NewContainerName)
 {
+	FWriteScopeLock Lock(RWLock);
+	
 	FSmartNameMapping* ExistingMapping = NameMappings.Find(NewContainerName);
 	if(ExistingMapping == nullptr)
 	{
-		return &NameMappings.Add(NewContainerName);
+		return &NameMappings.Add(NewContainerName, {&RWLock});
 	}
 
 	return ExistingMapping;
@@ -303,6 +368,7 @@ void FSmartNameContainer::Serialize(FArchive& Ar, bool bIsTemplate)
 	{
 		Ar << NameMappings;
 	}
+
 }
 
 void FSmartNameContainer::PostLoad()
@@ -311,6 +377,11 @@ void FSmartNameContainer::PostLoad()
 	// Preserve Load state for deterministic cooking
 	LoadedNameMappings = NameMappings;
 #endif
+	
+	for(auto Mapping : NameMappings)
+	{
+		Mapping.Value.SetLock(&RWLock);
+	}
 }
 
 FSmartNameMapping* FSmartNameContainer::GetContainerInternal(const FName& ContainerName)
diff --git a/Engine/Source/Runtime/Engine/Private/BoneContainer.cpp b/Engine/Source/Runtime/Engine/Private/BoneContainer.cpp
--- a/Engine/Source/Runtime/Engine/Private/BoneContainer.cpp
+++ b/Engine/Source/Runtime/Engine/Private/BoneContainer.cpp
@@ -261,7 +261,7 @@ void FBoneContainer::CacheRequiredAnimCurveUids(const FCurveEvaluationOption& Cu
 				// No smart names, nothing to do
 				return;
 			}
-
+			
 			//Init UID LUT to everything unused
 			UIDToArrayIndexLUT.AddUninitialized(MaxUID+1);
 			for (SmartName::UID_Type& Item : UIDToArrayIndexLUT)
@@ -269,19 +269,13 @@ void FBoneContainer::CacheRequiredAnimCurveUids(const FCurveEvaluationOption& Cu
 				Item = SmartName::MaxUID;
 			}
 
-			// Get Current Names / UIDs
-			Mapping->FillUIDToNameArray(UIDToNameLUT);
-
-			// Get curve types
-			Mapping->FillUIDToCurveTypeArray(UIDToCurveTypeLUT);
-
 			// if the linked joints don't exists in RequiredBones, remove itself
-			if (UIDToNameLUT.Num() > 0)
+			int32 NumAvailableUIDs = 0;
+			Mapping->Iterate([this, &CurveEvalOption, &NumAvailableUIDs](const FSmartNameMapping* Mapping, SmartName::UID_Type CurveNameIndex)
 			{
-				int32 NumAvailableUIDs = 0;
-				for (int32 CurveNameIndex = UIDToNameLUT.Num() - 1; CurveNameIndex >=0 ; --CurveNameIndex)
+				FName CurveName;
+				if (Mapping->GetName(CurveNameIndex, CurveName))
 				{
-					const FName& CurveName = UIDToNameLUT[CurveNameIndex];
 					bool bBeingUsed = true;
 					if (!CurveEvalOption.bAllowCurveEvaluation)
 					{
@@ -301,7 +295,7 @@ void FBoneContainer::CacheRequiredAnimCurveUids(const FCurveEvaluationOption& Cu
 						}
 						else
 						{
-							const FCurveMetaData* CurveMetaData = Mapping->GetCurveMetaData(UIDToNameLUT[CurveNameIndex]);
+							const FCurveMetaData* CurveMetaData = Mapping->GetCurveMetaData(CurveName);
 							if (CurveMetaData)
 							{
 								if (CurveMetaData->MaxLOD < CurveEvalOption.LODIndex)
@@ -339,7 +333,7 @@ void FBoneContainer::CacheRequiredAnimCurveUids(const FCurveEvaluationOption& Cu
 					}
 				}
 				UIDToArrayIndexLUTValidCount = NumAvailableUIDs;
-			}
+			});
 		}
 	}
 	else
diff --git a/Engine/Source/Runtime/Engine/Public/BoneContainer.h b/Engine/Source/Runtime/Engine/Public/BoneContainer.h
--- a/Engine/Source/Runtime/Engine/Public/BoneContainer.h
+++ b/Engine/Source/Runtime/Engine/Public/BoneContainer.h
@@ -17,7 +17,6 @@ struct FBoneContainer;
 struct FSkeletonRemapping;
 struct FBlendedCurve;
 
-
 /** Struct used to store per-component ref pose override */
 struct FSkelMeshRefPoseOverride
 {
@@ -240,12 +239,6 @@ private:
 	/** Number of valid entries in UIDToArrayIndexLUT. I.e. a count of entries whose value does not equal to  MAX_uint16. */
 	int32 UIDToArrayIndexLUTValidCount;
 
-	/** Look up table of UID to Name UIDToNameLUT[InUID] = Name of curve. If NAME_None, it is invalid.*/	
-	TArray<FName> UIDToNameLUT;
-
-	/** Look up table of UID to FAnimCurveType UIDToNameLUT[InUID] = FAnimCurveType of curve. */	
-	TArray<FAnimCurveType> UIDToCurveTypeLUT;
-	
 	TSharedPtr<FSkelMeshRefPoseOverride> RefPoseOverride;
 	
 	// The serial number of this bone container. This is incremented each time the container is regenerated and can
@@ -264,7 +257,7 @@ private:
 	bool bUseRAWData;
 	/** Use Source Data that is imported that are not compressed. */
 	bool bUseSourceData;
-
+	
 public:
 
 	FBoneContainer();
@@ -378,7 +371,7 @@ public:
 			return RefSkeleton->GetRefBonePose()[BoneIndicesArray[BoneIndex.GetInt()]];
 		}
 	}
-	
+
 	/** Override skeleton ref pose. */
 	void SetRefPoseOverride(const TSharedPtr<FSkelMeshRefPoseOverride>& InRefPoseOverride)
 	{
@@ -434,20 +427,24 @@ public:
 	{
 		return UIDToArrayIndexLUTValidCount;
 	}
-	
 
-	/** Get UID To Name look up table */
+	/** DEPRECATED: Get UID To Name look up table */
+	UE_DEPRECATED(5.0, "GetUIDToNameLookupTable is deprecated, please access from the SmartNameMapping directly via GetSkeletonAsset()->GetSmartNameContainer(USkeleton::AnimCurveMappingName)")
 	TArray<FName> const& GetUIDToNameLookupTable() const
 	{
-		return UIDToNameLUT;
+		static TArray<FName> Dummy;
+		return Dummy;
 	}
 
-	/** Get UID To curve type look up table */
+	/** DEPRECATED: Get UID To curve type look up table */
+	UE_DEPRECATED(5.0, "GetUIDToCurveTypeLookupTable is deprecated, please access from the SmartNameMapping directly via GetSkeletonAsset()->GetSmartNameContainer(USkeleton::AnimCurveMappingName)")
 	TArray<FAnimCurveType> const& GetUIDToCurveTypeLookupTable() const
 	{
-		return UIDToCurveTypeLUT;
+		static TArray<FAnimCurveType> Dummy;
+		return Dummy;
 	}
 
+	
 	/** Get the array that maps UIDs to array indexes. */
 	TArray<SmartName::UID_Type> const& GetUIDToArrayLookupTableBackup() const
 	{
