Flush async loading data on import to make sure latest audio is used when previewing in editor.

diff --git a/Engine/Source/Runtime/Engine/Classes/Sound/SoundWave.h b/Engine/Source/Runtime/Engine/Classes/Sound/SoundWave.h
--- a/Engine/Source/Runtime/Engine/Classes/Sound/SoundWave.h
+++ b/Engine/Source/Runtime/Engine/Classes/Sound/SoundWave.h
@@ -915,6 +915,7 @@ public:
 	 * This function can be called before playing or using a SoundWave to check if any cook settings have been modified since this SoundWave was last cooked.
 	 */
 	void InvalidateSoundWaveIfNeccessary();
+
 #endif //WITH_EDITOR
 
 
@@ -922,6 +923,9 @@ public:
 
 private:
 #if WITH_EDITOR
+	// Removes any in-progress async loading data formats. 
+	void FlushAsyncLoadingDataFormats();
+
 	void BakeFFTAnalysis();
 	void BakeEnvelopeAnalysis();
 #endif //WITH_EDITOR
diff --git a/Engine/Source/Runtime/Engine/Private/SoundWave.cpp b/Engine/Source/Runtime/Engine/Private/SoundWave.cpp
--- a/Engine/Source/Runtime/Engine/Private/SoundWave.cpp
+++ b/Engine/Source/Runtime/Engine/Private/SoundWave.cpp
@@ -1034,9 +1034,17 @@ void USoundWave::InvalidateCompressedData(bool bFreeResources, bool bRebuildStre
 	}
 
 #if WITH_EDITOR
+	if (bFreeResources)
+	{
+		// Flush any in-flight async loading data formats to make sure their results
+		// are not during the next call to GetCompressedData
+		FlushAsyncLoadingDataFormats();
+	}
+
 	if (bRebuildStreamingChunks)
 	{
 		CachePlatformData();
+
 		SoundWaveDataPtr->CurrentChunkRevision += 1;
 	}
 
@@ -1050,6 +1058,25 @@ void USoundWave::InvalidateCompressedData(bool bFreeResources, bool bRebuildStre
 #endif
 }
 
+#if WITH_EDITOR
+void USoundWave::FlushAsyncLoadingDataFormats()
+{
+	// Flush any async results so we don't leak them in the DDC
+	if (GetDerivedDataCache() && AsyncLoadingDataFormats.Num() > 0)
+	{
+		TArray<uint8> OutData;
+		for (auto AsyncLoadIt = AsyncLoadingDataFormats.CreateConstIterator(); AsyncLoadIt; ++AsyncLoadIt)
+		{
+			uint32 AsyncHandle = AsyncLoadIt.Value();
+			GetDerivedDataCacheRef().WaitAsynchronousCompletion(AsyncHandle);
+			GetDerivedDataCacheRef().GetAsynchronousResults(AsyncHandle, OutData);
+		}
+
+		AsyncLoadingDataFormats.Empty();
+	}
+}
+#endif
+
 bool USoundWave::HasStreamingChunks()
 {
 	check(SoundWaveDataPtr);
@@ -1313,18 +1340,7 @@ void USoundWave::BeginDestroy()
 
 #if WITH_EDITOR
 	// Flush any async results so we don't leak them in the DDC
-	if (GetDerivedDataCache() && AsyncLoadingDataFormats.Num() > 0)
-	{
-		TArray<uint8> OutData;
-		for (auto AsyncLoadIt = AsyncLoadingDataFormats.CreateConstIterator(); AsyncLoadIt; ++AsyncLoadIt)
-		{
-			uint32 AsyncHandle = AsyncLoadIt.Value();
-			GetDerivedDataCacheRef().WaitAsynchronousCompletion(AsyncHandle);
-			GetDerivedDataCacheRef().GetAsynchronousResults(AsyncHandle, OutData);
-		}
-
-		AsyncLoadingDataFormats.Empty();
-	}
+	FlushAsyncLoadingDataFormats();
 #endif
 
 	ReleaseCompressedAudio();
