AssetManager:
1) Warn and ignore if callers attempt to add a PrimaryAsset that is not a TopLevelAsset. Actors in external packages can be non-TopLevel Assets.
2) Cull non-TopLevelAssets from search results looking for PrimaryAssets.
3) Disallow SearchAssetRegistryPaths of the entire registry; require callers to provide at least one path. This prevents PrimaryAssetRules that are missing their paths from attempting to set every asset in the AssetRegistry as a primary asset.
#rnx

diff --git a/Engine/Source/Runtime/CoreUObject/Private/AssetRegistry/AssetData.cpp b/Engine/Source/Runtime/CoreUObject/Private/AssetRegistry/AssetData.cpp
--- a/Engine/Source/Runtime/CoreUObject/Private/AssetRegistry/AssetData.cpp
+++ b/Engine/Source/Runtime/CoreUObject/Private/AssetRegistry/AssetData.cpp
@@ -129,6 +129,27 @@ bool FAssetData::IsUAsset(UObject* InAsset)
 	return DetectIsUAssetByNames(PackageNameStrBuilder, AssetNameStrBuilder);
 }
 
+bool FAssetData::IsTopLevelAsset() const
+{
+	int32 SubObjectIndex;
+	FStringView(WriteToString<256>(ObjectPath)).FindChar(SUBOBJECT_DELIMITER_CHAR, SubObjectIndex);
+	return SubObjectIndex == INDEX_NONE;
+}
+
+bool FAssetData::IsTopLevelAsset(UObject* Object)
+{
+	if (!Object)
+	{
+		return false;
+	}
+	UObject* Outer = Object->GetOuter();
+	if (!Outer)
+	{
+		return false;
+	}
+	return Outer->IsA<UPackage>();
+}
+
 void FAssetData::SetTagsAndAssetBundles(FAssetDataTagMap&& Tags)
 {
 	using namespace UE::AssetData::Private;
diff --git a/Engine/Source/Runtime/CoreUObject/Private/UObject/SoftObjectPath.cpp b/Engine/Source/Runtime/CoreUObject/Private/UObject/SoftObjectPath.cpp
--- a/Engine/Source/Runtime/CoreUObject/Private/UObject/SoftObjectPath.cpp
+++ b/Engine/Source/Runtime/CoreUObject/Private/UObject/SoftObjectPath.cpp
@@ -40,7 +40,7 @@ FString FSoftObjectPath::ToString() const
 	// Preallocate to correct size and then append strings
 	FullPathString.Reserve(AssetPathString.Len() + SubPathString.Len() + 1);
 	FullPathString += AssetPathString;
-	FullPathString += ':';
+	FullPathString += SUBOBJECT_DELIMITER_CHAR;
 	FullPathString += SubPathString;
 	return FullPathString;
 }
@@ -54,7 +54,7 @@ void FSoftObjectPath::ToString(FStringBuilderBase& Builder) const
 
 	if (SubPathString.Len() > 0)
 	{
-		Builder << ':' << SubPathString;
+		Builder << SUBOBJECT_DELIMITER_CHAR << SubPathString;
 	}
 }
 
@@ -94,7 +94,7 @@ void FSoftObjectPath::SetPath(FWideStringView Path)
 		}
 
 		int32 ColonIndex;
-		if (Path.FindChar(':', ColonIndex))
+		if (Path.FindChar(SUBOBJECT_DELIMITER_CHAR, ColonIndex))
 		{
 			// Has a subobject, split on that then create a name from the temporary path
 			AssetPathName = FName(Path.Left(ColonIndex));
@@ -143,7 +143,7 @@ void FSoftObjectPath::SetPath(FName PathName)
 			}
 
 			int32 ColonIndex;
-			if (Path.FindChar(':', ColonIndex))
+			if (Path.FindChar(SUBOBJECT_DELIMITER_CHAR, ColonIndex))
 			{
 				// Has a subobject, split on that then create a name from the temporary path
 				AssetPathName = FName(Path.Left(ColonIndex));
diff --git a/Engine/Source/Runtime/CoreUObject/Public/AssetRegistry/AssetData.h b/Engine/Source/Runtime/CoreUObject/Public/AssetRegistry/AssetData.h
--- a/Engine/Source/Runtime/CoreUObject/Public/AssetRegistry/AssetData.h
+++ b/Engine/Source/Runtime/CoreUObject/Public/AssetRegistry/AssetData.h
@@ -82,7 +82,7 @@ public:
 	}
 
 public:
-	/** The object path for the asset in the form PackageName.AssetName. Only top level objects in a package can have AssetData */
+	/** The object path for the asset in the form PackageName.ObjectName, or PackageName.ObjectName:SubObjectName. */
 	FName ObjectPath;
 	/** The name of the package in which the asset is found, this is the full long package name such as /Game/Path/Package */
 	FName PackageName;
@@ -165,6 +165,18 @@ public:
 	/** Returns true if the given UObject is the primary asset in a package, true for maps and assets but false for secondary objects like class redirectors */
 	COREUOBJECT_API static bool IsUAsset(UObject* Object);
 
+	/**
+	 * Returns true iff the Asset is a TopLevelAsset (not a subobject, its outer is a UPackage).
+	 * Only TopLevelAssets can be PrimaryAssets in the AssetManager.
+	 */
+	COREUOBJECT_API bool IsTopLevelAsset() const;
+	
+	/**
+	 * Returns true iff the given Object, assumed to be an Asset, is a TopLevelAsset (not a subobject, its outer is a UPackage).
+	 * Only TopLevelAssets can be PrimaryAssets in the AssetManager.
+	 */
+	COREUOBJECT_API static bool IsTopLevelAsset(UObject* Object);
+
 	void Shrink()
 	{
 		ChunkIDs.Shrink();
diff --git a/Engine/Source/Runtime/Engine/Classes/Engine/AssetManager.h b/Engine/Source/Runtime/Engine/Classes/Engine/AssetManager.h
--- a/Engine/Source/Runtime/Engine/Classes/Engine/AssetManager.h
+++ b/Engine/Source/Runtime/Engine/Classes/Engine/AssetManager.h
@@ -607,8 +607,8 @@ protected:
 	/** Internal helper function that attempts to get asset data from the specified path; Accounts for possibility of blueprint classes ending in _C */
 	virtual void GetAssetDataForPathInternal(class IAssetRegistry& AssetRegistry, const FString& AssetPath, OUT FAssetData& OutAssetData) const;
 
-	/** Updates the asset data cached on the name data */
-	virtual void UpdateCachedAssetData(const FPrimaryAssetId& PrimaryAssetId, const FAssetData& NewAssetData, bool bAllowDuplicates);
+	/** Updates the asset data cached on the name data; returns false if the asset is not a valid primary asset. */
+	virtual bool TryUpdateCachedAssetData(const FPrimaryAssetId& PrimaryAssetId, const FAssetData& NewAssetData, bool bAllowDuplicates);
 
 	/** Returns the NameData for a specific type/name pair */
 	FPrimaryAssetData* GetNameData(const FPrimaryAssetId& PrimaryAssetId, bool bCheckRedirector = true);
diff --git a/Engine/Source/Runtime/Engine/Private/AssetManager.cpp b/Engine/Source/Runtime/Engine/Private/AssetManager.cpp
--- a/Engine/Source/Runtime/Engine/Private/AssetManager.cpp
+++ b/Engine/Source/Runtime/Engine/Private/AssetManager.cpp
@@ -534,7 +534,10 @@ int32 UAssetManager::SearchAssetRegistryPaths(TArray<FAssetData>& OutAssetDataLi
 	{
 		return 0;
 	}
-
+	if (Rules.AssetScanPaths.IsEmpty())
+	{
+		return 0;
+	}
 	TArray<FString> Directories, PackageNames;
 	TArray<FString> ScanPaths = Rules.AssetScanPaths;
 	// Add path info
@@ -871,6 +874,11 @@ int32 UAssetManager::ScanPathsForPrimaryAssets(FPrimaryAssetType PrimaryAssetTyp
 	// Now add to map or update as needed
 	for (FAssetData& Data : AssetDataList)
 	{
+		if (!Data.IsTopLevelAsset())
+		{
+			// Only TopLevelAssets can be PrimaryAssets
+			continue;
+		}
 		FPrimaryAssetId PrimaryAssetId = ExtractPrimaryAssetIdFromData(Data, PrimaryAssetType);
 
 		// Remove invalid or wrong type assets
@@ -899,7 +907,7 @@ int32 UAssetManager::ScanPathsForPrimaryAssets(FPrimaryAssetType PrimaryAssetTyp
 
 		NumAdded++;
 
-		UpdateCachedAssetData(PrimaryAssetId, Data, false);
+		TryUpdateCachedAssetData(PrimaryAssetId, Data, false);
 	}
 
 	if (!IsBulkScanning())
@@ -959,8 +967,10 @@ bool UAssetManager::RegisterSpecificPrimaryAsset(const FPrimaryAssetId& PrimaryA
 		return false;
 	}
 
-	// If we got this far, it will succeed but might warn
-	UpdateCachedAssetData(PrimaryAssetId, NewAssetData, false);
+	if (!TryUpdateCachedAssetData(PrimaryAssetId, NewAssetData, false))
+	{
+		return false;
+	}
 
 	if (!IsBulkScanning())
 	{
@@ -974,21 +984,31 @@ bool UAssetManager::RegisterSpecificPrimaryAsset(const FPrimaryAssetId& PrimaryA
 enum class EAssetDataCanBeSubobject { Yes, No };
 
 template<EAssetDataCanBeSubobject ScanForSubobject>
-FSoftObjectPath ToSoftObjectPath(const FAssetData& AssetData);
+bool TryToSoftObjectPath(const FAssetData& AssetData, FSoftObjectPath& OutSoftObjectPath);
 
-void UAssetManager::UpdateCachedAssetData(const FPrimaryAssetId& PrimaryAssetId, const FAssetData& NewAssetData, bool bAllowDuplicates)
+bool UAssetManager::TryUpdateCachedAssetData(const FPrimaryAssetId& PrimaryAssetId, const FAssetData& NewAssetData, bool bAllowDuplicates)
 {
 	check(PrimaryAssetId.IsValid());
 
 	const TSharedRef<FPrimaryAssetTypeData>* FoundType = AssetTypeMap.Find(PrimaryAssetId.PrimaryAssetType);
 
-	if (ensure(FoundType))
+	if (!ensure(FoundType))
+	{
+		return false;
+	}
+	else
 	{
 		FPrimaryAssetTypeData& TypeData = FoundType->Get();
 
 		FPrimaryAssetData* OldData = TypeData.AssetMap.Find(PrimaryAssetId.PrimaryAssetName);
 
-		FSoftObjectPath NewAssetPath = ToSoftObjectPath<EAssetDataCanBeSubobject::No>(NewAssetData);
+		FSoftObjectPath NewAssetPath;
+		if (!TryToSoftObjectPath<EAssetDataCanBeSubobject::No>(NewAssetData, NewAssetPath))
+		{
+			UE_LOG(LogAssetManager, Warning, TEXT("Tried to add primary asset %s, but it is not a TopLevelAsset and so cannot be a primary asset"),
+				*NewAssetData.ObjectPath.ToString())
+			return false;
+		}
 
 		ensure(NewAssetPath.IsAsset());
 
@@ -1078,6 +1098,7 @@ void UAssetManager::UpdateCachedAssetData(const FPrimaryAssetId& PrimaryAssetId,
 			CachedAssetBundles.Remove(PrimaryAssetId);
 		}
 	}
+	return true;
 }
 
 int32 UAssetManager::ScanPathForPrimaryAssets(FPrimaryAssetType PrimaryAssetType, const FString& Path, UClass* BaseClass, bool bHasBlueprintClasses, bool bIsEditorOnly, bool bForceSynchronousScan)
@@ -2357,7 +2378,7 @@ bool UAssetManager::OnAssetRegistryAvailableAfterInitialization(FName InName, FA
 					if (!bPathExclusionsExists || !IsPathExcludedFromScan(AssetData.PackageName.ToString()))
 					{
 						FPrimaryAssetId PrimaryAssetId = AssetData.GetPrimaryAssetId();
-						if (PrimaryAssetId.IsValid())
+						if (PrimaryAssetId.IsValid() && AssetData.IsTopLevelAsset())
 						{
 							TSharedRef<FPrimaryAssetTypeData>* FoundType = AssetTypeMap.Find(PrimaryAssetId.PrimaryAssetType);
 							if (FoundType)
@@ -2389,8 +2410,10 @@ bool UAssetManager::OnAssetRegistryAvailableAfterInitialization(FName InName, FA
 										}
 
 										// Check exclusion path
-										UpdateCachedAssetData(PrimaryAssetId, AssetData, false);
-										bRebuildReferenceList = true;
+										if (TryUpdateCachedAssetData(PrimaryAssetId, AssetData, false))
+										{
+											bRebuildReferenceList = true;
+										}
 									}
 								}
 							}
@@ -2724,32 +2747,42 @@ static bool ContainsSubobjectDelimiter(FName Name)
 }
 
 template<EAssetDataCanBeSubobject ScanForSubobject>
-FSoftObjectPath ToSoftObjectPath(const FAssetData& AssetData)
+bool TryToSoftObjectPath(const FAssetData& AssetData, FSoftObjectPath& OutObjectPath)
 {
 	if (!AssetData.IsValid())
 	{
-		return FSoftObjectPath();
+		OutObjectPath = FSoftObjectPath();
+		return false;
 	}
 	else if (EndsWithBlueprint(AssetData.AssetClass))
 	{
 		TStringBuilder<256> AssetPath;
 		AssetPath << AssetData.ObjectPath << TEXT("_C");
-		return FSoftObjectPath(FStringView(AssetPath));
+		OutObjectPath = FSoftObjectPath(FStringView(AssetPath));
+		return true;
 	}
 	else if (ScanForSubobject == EAssetDataCanBeSubobject::Yes)
 	{
-		return FSoftObjectPath(AssetData.ObjectPath);
+		OutObjectPath = FSoftObjectPath(AssetData.ObjectPath);
+		return true;
 	}
 	else
 	{
-		check(!ContainsSubobjectDelimiter(AssetData.ObjectPath));
-		return FSoftObjectPath(AssetData.ObjectPath, /* no subobject */ FString());
+		if (ContainsSubobjectDelimiter(AssetData.ObjectPath))
+		{
+			OutObjectPath = FSoftObjectPath();
+			return false;
+		}
+		OutObjectPath = FSoftObjectPath(AssetData.ObjectPath, /* no subobject */ FString());
+		return true;
 	}
 }
 
 FSoftObjectPath UAssetManager::GetAssetPathForData(const FAssetData& AssetData) const
 {
-	return ToSoftObjectPath<EAssetDataCanBeSubobject::Yes>(AssetData);
+	FSoftObjectPath Result;
+	TryToSoftObjectPath<EAssetDataCanBeSubobject::Yes>(AssetData, Result);
+	return Result;
 }
 
 void UAssetManager::GetAssetDataForPathInternal(IAssetRegistry& AssetRegistry, const FString& AssetPath, OUT FAssetData& OutAssetData) const
@@ -4246,7 +4279,7 @@ void UAssetManager::OnInMemoryAssetCreated(UObject *Object)
 
 			GetAssetDataForPathInternal(AssetRegistry, Object->GetPathName(), NewAssetData);
 
-			if (NewAssetData.IsValid())
+			if (NewAssetData.IsValid() && NewAssetData.IsTopLevelAsset())
 			{
 				// Make sure it's in a valid path
 				bool bFoundPath = false;
@@ -4262,9 +4295,10 @@ void UAssetManager::OnInMemoryAssetCreated(UObject *Object)
 				if (bFoundPath)
 				{
 					// Add or update asset data
-					UpdateCachedAssetData(PrimaryAssetId, NewAssetData, true);
-
-					RebuildObjectReferenceList();
+					if (TryUpdateCachedAssetData(PrimaryAssetId, NewAssetData, true))
+					{
+						RebuildObjectReferenceList();
+					}
 				}
 			}
 		}
@@ -4379,7 +4413,7 @@ void UAssetManager::RefreshAssetData(UObject* ChangedObject)
 
 		if (ensure(NewData.IsValid()))
 		{
-			UpdateCachedAssetData(PrimaryAssetId, NewData, false);
+			TryUpdateCachedAssetData(PrimaryAssetId, NewData, false);
 		}
 	}
 	else
