Add support for struct in HLSLTree
Replace code related to MaterialAttributes with generic code based on structs

diff --git a/Engine/Shaders/Private/MaterialTemplate.ush b/Engine/Shaders/Private/MaterialTemplate.ush
--- a/Engine/Shaders/Private/MaterialTemplate.ush
+++ b/Engine/Shaders/Private/MaterialTemplate.ush
@@ -241,10 +241,8 @@ float4 GetDynamicParameter(FMaterialParticleParameters Parameters, float4 Defaul
 
 }
 
-struct FMaterialAttributes
-{
+/** Material declarations */
 %s
-};
 
 /** FMaterialAttributes utilities */
 %s
@@ -2737,16 +2735,9 @@ float3 CalculateAnisotropyTangent(in out FMaterialPixelParameters Parameters, FP
 	return Tangent;
 }
 
-FMaterialAttributes FMaterialAttributes_GetDefault()
-{
-	FMaterialAttributes Result;
-%s
-	return Result;
-}
-
 FMaterialAttributes EvaluateVertexMaterialAttributes(FMaterialVertexParameters Parameters)
 {
-	FMaterialAttributes DefaultMaterialAttributes = FMaterialAttributes_GetDefault();
+	FMaterialAttributes DefaultMaterialAttributes = (FMaterialAttributes)0;
 %s
 	// Return the default, this may be skipped if the above code returns different attributes
 	return DefaultMaterialAttributes;
@@ -2754,7 +2745,7 @@ FMaterialAttributes EvaluateVertexMaterialAttributes(FMaterialVertexParameters P
 
 FMaterialAttributes EvaluatePixelMaterialAttributes(in out FMaterialPixelParameters Parameters)
 {
-	FMaterialAttributes DefaultMaterialAttributes = FMaterialAttributes_GetDefault();
+	FMaterialAttributes DefaultMaterialAttributes = (FMaterialAttributes)0;
 %s
 	// Return the default, this may be skipped if the above code returns different attributes
 	return DefaultMaterialAttributes;
diff --git a/Engine/Source/Runtime/Engine/Classes/Materials/MaterialExpressionDivide.h b/Engine/Source/Runtime/Engine/Classes/Materials/MaterialExpressionDivide.h
--- a/Engine/Source/Runtime/Engine/Classes/Materials/MaterialExpressionDivide.h
+++ b/Engine/Source/Runtime/Engine/Classes/Materials/MaterialExpressionDivide.h
@@ -32,6 +32,7 @@ class UMaterialExpressionDivide : public UMaterialExpression
 	//~ Begin UMaterialExpression Interface
 #if WITH_EDITOR
 	virtual int32 Compile(class FMaterialCompiler* Compiler, int32 OutputIndex) override;
+	virtual EMaterialGenerateHLSLStatus GenerateHLSLExpression(FMaterialHLSLGenerator& Generator, UE::HLSLTree::FScope& Scope, int32 OutputIndex, UE::HLSLTree::FExpression*& OutExpression) override;
 	virtual void GetCaption(TArray<FString>& OutCaptions) const override;
 	virtual FText GetKeywords() const override {return FText::FromString(TEXT("/"));}
 #endif // WITH_EDITOR
diff --git a/Engine/Source/Runtime/Engine/Classes/Materials/MaterialExpressionGetLocal.h b/Engine/Source/Runtime/Engine/Classes/Materials/MaterialExpressionGetLocal.h
--- a/Engine/Source/Runtime/Engine/Classes/Materials/MaterialExpressionGetLocal.h
+++ b/Engine/Source/Runtime/Engine/Classes/Materials/MaterialExpressionGetLocal.h
@@ -23,6 +23,7 @@ class UMaterialExpressionGetLocal : public UMaterialExpression
 	virtual int32 Compile(class FMaterialCompiler* Compiler, int32 OutputIndex) override;
 	virtual void GetCaption(TArray<FString>& OutCaptions) const override;
 	virtual uint32 GetOutputType(int32 InputIndex) override;
+	virtual bool IsResultMaterialAttributes(int32 OutputIndex) override;
 	virtual EMaterialGenerateHLSLStatus GenerateHLSLExpression(FMaterialHLSLGenerator& Generator, UE::HLSLTree::FScope& Scope, int32 OutputIndex, UE::HLSLTree::FExpression*& OutExpression) override;
 #endif
 	//~ End UMaterialExpression Interface
diff --git a/Engine/Source/Runtime/Engine/Classes/Materials/MaterialExpressionGetMaterialAttributes.h b/Engine/Source/Runtime/Engine/Classes/Materials/MaterialExpressionGetMaterialAttributes.h
--- a/Engine/Source/Runtime/Engine/Classes/Materials/MaterialExpressionGetMaterialAttributes.h
+++ b/Engine/Source/Runtime/Engine/Classes/Materials/MaterialExpressionGetMaterialAttributes.h
@@ -36,6 +36,8 @@ class UMaterialExpressionGetMaterialAttributes : public UMaterialExpression
 	//~ Begin UMaterialExpression Interface
 #if WITH_EDITOR
 	virtual int32 Compile(class FMaterialCompiler* Compiler, int32 OutputIndex) override;
+	virtual EMaterialGenerateHLSLStatus GenerateHLSLExpression(FMaterialHLSLGenerator& Generator, UE::HLSLTree::FScope& Scope, int32 OutputIndex, UE::HLSLTree::FExpression*& OutExpression) override;
+
 	virtual void GetCaption(TArray<FString>& OutCaptions) const override;
 	virtual const TArray<FExpressionInput*> GetInputs()override;
 	virtual FExpressionInput* GetInput(int32 InputIndex)override;
diff --git a/Engine/Source/Runtime/Engine/Private/HLSLTree/HLSLTree.cpp b/Engine/Source/Runtime/Engine/Private/HLSLTree/HLSLTree.cpp
--- a/Engine/Source/Runtime/Engine/Private/HLSLTree/HLSLTree.cpp
+++ b/Engine/Source/Runtime/Engine/Private/HLSLTree/HLSLTree.cpp
@@ -63,7 +63,7 @@ UE::HLSLTree::FErrors::FErrors(FMemStackBase& InAllocator)
 {
 }
 
-bool UE::HLSLTree::FErrors::AddError(const FNode* InNode, FStringView InError)
+void UE::HLSLTree::FErrors::AddError(const FNode* InNode, FStringView InError)
 {
 	const int32 SizeofString = InError.Len() * sizeof(TCHAR);
 	void* Memory = Allocator->Alloc(sizeof(FError) + SizeofString, alignof(FError));
@@ -76,7 +76,7 @@ bool UE::HLSLTree::FErrors::AddError(const FNode* InNode, FStringView InError)
 	FirstError = Error;
 	NumErrors++;
 
-	return false;
+	ensureMsgf(false, TEXT("%s"), Error->Message);
 }
 
 UE::HLSLTree::FEmitContext::FEmitContext(FMemStackBase& InAllocator)
@@ -94,7 +94,7 @@ const TCHAR* UE::HLSLTree::FEmitContext::AcquireLocalDeclarationCode()
 	return AllocateStringf(*Allocator, TEXT("Local%d"), NumExpressionLocals++);
 }
 
-const TCHAR* UE::HLSLTree::FEmitContext::CastShaderValue(const FNode* Node, const TCHAR* Code, Shader::EValueType SourceType, Shader::EValueType DestType, ECastFlags Flags)
+const TCHAR* UE::HLSLTree::FEmitContext::CastShaderValue(const FNode* Node, const TCHAR* Code, const FType& SourceType, const FType& DestType, ECastFlags Flags)
 {
 	if (SourceType == DestType)
 	{
@@ -281,6 +281,80 @@ void UE::HLSLTree::FEmitContext::Finalize()
 	MaterialCompilationOutput->UniformExpressionSet.UniformPreshaderBufferSize = (UniformPreshaderOffset + 3u) / 4u;
 }
 
+const TCHAR* UE::HLSLTree::FType::GetName() const
+{
+	if (StructType)
+	{
+		return StructType->Name;
+	}
+
+	check(ValueType != Shader::EValueType::Struct);
+	const Shader::FValueTypeDescription TypeDesc = Shader::GetValueTypeDescription(ValueType);
+	return TypeDesc.Name;
+}
+
+int32 UE::HLSLTree::FType::GetNumComponents() const
+{
+	if (StructType)
+	{
+		return StructType->ComponentTypes.Num();
+	}
+	const Shader::FValueTypeDescription TypeDesc = Shader::GetValueTypeDescription(ValueType);
+	return TypeDesc.NumComponents;
+}
+
+bool UE::HLSLTree::FType::Merge(const FType& OtherType)
+{
+	if (ValueType == Shader::EValueType::Void)
+	{
+		ValueType = OtherType.ValueType;
+		StructType = OtherType.StructType;
+		return true;
+	}
+
+	if (IsStruct() || OtherType.IsStruct())
+	{
+		return StructType == OtherType.StructType;
+	}
+
+	if (ValueType != OtherType.ValueType)
+	{
+		const Shader::FValueTypeDescription TypeDesc = Shader::GetValueTypeDescription(ValueType);
+		const Shader::FValueTypeDescription OtherTypeDesc = Shader::GetValueTypeDescription(OtherType);
+		if (TypeDesc.ComponentType != OtherTypeDesc.ComponentType)
+		{
+			return false;
+		}
+
+		const int8 NumComponents = FMath::Max(TypeDesc.NumComponents, OtherTypeDesc.NumComponents);
+		ValueType = Shader::MakeValueType(TypeDesc.ComponentType, NumComponents);
+	}
+	return true;
+}
+
+UE::HLSLTree::FStructFieldRef UE::HLSLTree::FStructType::FindFieldByName(const TCHAR* InName) const
+{
+	for (const FStructField& Field : Fields)
+	{
+		if (FCString::Strcmp(Field.Name, InName) == 0)
+		{
+			return FStructFieldRef(Field.Type, Field.ComponentIndex, Field.Type.GetNumComponents());
+		}
+	}
+
+	return FStructFieldRef();
+}
+
+UE::HLSLTree::FConstantValue::FConstantValue(const Shader::FValue& InValue) : Type(InValue.GetType())
+{
+	const Shader::FValueTypeDescription TypeDesc = Shader::GetValueTypeDescription(Type);
+	Component.Empty(TypeDesc.NumComponents);
+	for (int32 i = 0; i < TypeDesc.NumComponents; ++i)
+	{
+		Component.Add(InValue.Component[i]);
+	}
+}
+
 UE::HLSLTree::FScope* UE::HLSLTree::FScope::FindSharedParent(FScope* Lhs, FScope* Rhs)
 {
 	FScope* Scope0 = Lhs;
@@ -304,69 +378,50 @@ UE::HLSLTree::FScope* UE::HLSLTree::FScope::FindSharedParent(FScope* Lhs, FScope
 	return Scope0;
 }
 
-bool UE::HLSLTree::FExpressionLocalPHI::UpdateType(FUpdateTypeContext& Context, int8 InRequestedNumComponents)
+void UE::HLSLTree::FExpressionLocalPHI::UpdateType(FUpdateTypeContext& Context, const FRequestedType& RequestedType)
 {
-	Shader::EValueType TypePerValue[MaxNumPreviousScopes] = { Shader::EValueType::Void };
+	FType TypePerValue[MaxNumPreviousScopes];
 	int32 NumValidTypes = 0;
-	Shader::EValueComponentType ComponentType = Shader::EValueComponentType::Void;
-	int8 NumComponents = 0;
+	FType CurrentType;
 
 	auto UpdateValueTypes = [&]()
 	{
 		for (int32 i = 0; i < NumValues; ++i)
 		{
-			if (TypePerValue[i] == Shader::EValueType::Void)
+			if (!TypePerValue[i])
 			{
-				TypePerValue[i] = RequestExpressionType(Context, Values[i], InRequestedNumComponents);
-				if (TypePerValue[i] != Shader::EValueType::Void)
+				TypePerValue[i] = RequestExpressionType(Context, Values[i], RequestedType);
+				if (TypePerValue[i])
 				{
-					const Shader::FValueTypeDescription TypeDesc = Shader::GetValueTypeDescription(TypePerValue[i]);
-					if (ComponentType == Shader::EValueComponentType::Void)
+					if (!CurrentType.Merge(TypePerValue[i]))
 					{
-						ComponentType = TypeDesc.ComponentType;
-						NumComponents = TypeDesc.NumComponents;
-					}
-					else
-					{
-						NumComponents = FMath::Max(NumComponents, TypeDesc.NumComponents);
-						if (ComponentType != TypeDesc.ComponentType)
-						{
-							return Context.Errors.AddError(this, TEXT("Type mismatch"));
-						}
+						return Context.Errors.AddError(this, TEXT("Type mismatch"));
 					}
 					check(NumValidTypes < NumValues);
 					NumValidTypes++;
 				}
 			}
 		}
-		return true;
 	};
 
 	// First try to assign all the values we can
-	if (!UpdateValueTypes())
-	{
-		return false;
-	}
+	UpdateValueTypes();
 
 	// Assuming we have at least one value with a valid type, we use that to initialize our type
-	if (!SetType(Context, Shader::MakeValueType(ComponentType, NumComponents)))
-	{
-		return false;
-	}
+	SetType(Context, CurrentType);
 
 	if (NumValidTypes < NumValues)
 	{
 		// Now try to assign remaining types that failed the first iteration 
-		if (!UpdateValueTypes())
+		UpdateValueTypes();
+		if (NumValidTypes < NumValues)
 		{
-			return false;
+			return Context.Errors.AddError(this, TEXT("Failed to compute all types for LocalPHI"));
 		}
 	}
-
-	return NumValidTypes == NumValues;
 }
 
-bool UE::HLSLTree::FExpressionLocalPHI::PrepareValue(FEmitContext& Context)
+void UE::HLSLTree::FExpressionLocalPHI::PrepareValue(FEmitContext& Context)
 {
 	FExpression* ForwardExpression = Values[0];
 	for (int32 i = 1; i < NumValues; ++i)
@@ -385,10 +440,7 @@ bool UE::HLSLTree::FExpressionLocalPHI::PrepareValue(FEmitContext& Context)
 	}
 
 	const int32 LocalPHIIndex = Context.NumLocalPHIs++;
-	if (!SetValueInlineShaderf(Context, TEXT("LocalPHI%d"), LocalPHIIndex))
-	{
-		return false;
-	}
+	SetValueInlineShaderf(Context, TEXT("LocalPHI%d"), LocalPHIIndex);
 
 	// Find the outermost scope to declare our local variable
 	FScope* DeclarationScope = ParentScope;
@@ -397,7 +449,7 @@ bool UE::HLSLTree::FExpressionLocalPHI::PrepareValue(FEmitContext& Context)
 		DeclarationScope = FScope::FindSharedParent(DeclarationScope, Scopes[i]);
 		if (!DeclarationScope)
 		{
-			return false;
+			return Context.Errors.AddError(this, TEXT("Invalid LocalPHI"));
 		}
 	}
 
@@ -405,23 +457,22 @@ bool UE::HLSLTree::FExpressionLocalPHI::PrepareValue(FEmitContext& Context)
 	{
 		if (PrepareExpressionValue(Context, Values[i]) == EExpressionEvaluationType::None)
 		{
-			return false;
+			return Context.Errors.AddError(this, TEXT("Invalid LocalPHI"));
 		}
 	}
 
-	const Shader::EValueType Type = GetValueType();
-	const Shader::FValueTypeDescription& TypeDesc = Shader::GetValueTypeDescription(Type);
+	const FType LocalType = GetType();
 
 	bool bNeedToAddDeclaration = true;
 	for (int32 i = 0; i < NumValues; ++i)
 	{
 		FScope* ValueScope = Scopes[i];
-		const TCHAR* ValueCode = Values[i]->GetValueShader(Context, Type);
+		const TCHAR* ValueCode = Values[i]->GetValueShader(Context, LocalType);
 		ValueScope->MarkLiveRecursive();
 		if (ValueScope == DeclarationScope)
 		{
 			ValueScope->EmitDeclarationf(Context, TEXT("%s LocalPHI%d = %s;"),
-				TypeDesc.Name,
+				LocalType.GetName(),
 				LocalPHIIndex,
 				ValueCode);
 			bNeedToAddDeclaration = false;
@@ -438,11 +489,9 @@ bool UE::HLSLTree::FExpressionLocalPHI::PrepareValue(FEmitContext& Context)
 	{
 		check(DeclarationScope->IsLive());
 		DeclarationScope->EmitDeclarationf(Context, TEXT("%s LocalPHI%d;"),
-			TypeDesc.Name,
+			LocalType.GetName(),
 			LocalPHIIndex);
 	}
-
-	return true;
 }
 
 void UE::HLSLTree::FNodeVisitor::VisitNode(FNode* Node)
@@ -478,31 +527,63 @@ UE::HLSLTree::ENodeVisitResult UE::HLSLTree::FScope::Visit(FNodeVisitor& Visitor
 	return Result;
 }
 
-UE::Shader::EValueType UE::HLSLTree::RequestExpressionType(FUpdateTypeContext& Context, FExpression* InExpression, int8 InRequestedNumComponents)
+UE::HLSLTree::FType UE::HLSLTree::RequestExpressionType(FUpdateTypeContext& Context, FExpression* InExpression, const FRequestedType& RequestedType)
 {
 	if (!InExpression)
 	{
 		return Shader::EValueType::Void;
 	}
 
-	if (!InExpression->bReentryFlag && InRequestedNumComponents > InExpression->RequestedNumComponents)
+	if (InExpression->bReentryFlag)
 	{
-		InExpression->bReentryFlag = true;
-		const bool bUpdateResult = InExpression->UpdateType(Context, InRequestedNumComponents);
-		InExpression->bReentryFlag = false;
+		// Valid for this to be called reentrantly
+		// Code should ensure that the type is set before the reentrant call, otherwise type will not be valid here
+		// LocalPHI nodes rely on this to break loops
+		return InExpression->Type;
+	}
 
-		if (bUpdateResult)
+	bool bNeedToUpdateType = false;
+	if (InExpression->CurrentRequestedType.RequestedComponents.Num() == 0)
+	{
+		InExpression->CurrentRequestedType = RequestedType;
+		bNeedToUpdateType = !RequestedType.IsVoid();
+	}
+	else if(InExpression->CurrentRequestedType.StructType != RequestedType.StructType)
+	{
+		Context.Errors.AddError(InExpression, TEXT("Type mismatch"));
+		return Shader::EValueType::Void;
+	}
+	else
+	{
+		const int32 NumComponents = RequestedType.GetNumComponents();
+		InExpression->CurrentRequestedType.RequestedComponents.PadToNum(NumComponents, false);
+		for (int32 ComponentIndex = 0; ComponentIndex < NumComponents; ++ComponentIndex)
 		{
-			check(InExpression->ValueType != Shader::EValueType::Void);
-			InExpression->RequestedNumComponents = InRequestedNumComponents;
+			if (!InExpression->CurrentRequestedType.IsComponentRequested(ComponentIndex) && RequestedType.IsComponentRequested(ComponentIndex))
+			{
+				InExpression->CurrentRequestedType.RequestedComponents[ComponentIndex] = true;
+				bNeedToUpdateType = true;
+			}
 		}
-		else
+	}
+
+	if (bNeedToUpdateType)
+	{
+		check(!InExpression->CurrentRequestedType.IsVoid());
+
+		InExpression->bReentryFlag = true;
+		InExpression->UpdateType(Context, InExpression->CurrentRequestedType);
+		InExpression->bReentryFlag = false;
+
+		if (!InExpression->Type)
 		{
-			InExpression->ValueType = Shader::EValueType::Void;
+			// If we failed to assign a valid type, reset the requested type as well
+			// This ensures we'll try to compute a type again the next time we're called
+			InExpression->CurrentRequestedType.Reset();
 		}
 	}
 
-	return InExpression->ValueType;
+	return InExpression->Type;
 }
 
 UE::HLSLTree::EExpressionEvaluationType UE::HLSLTree::PrepareExpressionValue(FEmitContext& Context, FExpression* InExpression)
@@ -512,7 +593,7 @@ UE::HLSLTree::EExpressionEvaluationType UE::HLSLTree::PrepareExpressionValue(FEm
 		return EExpressionEvaluationType::None;
 	}
 
-	if (InExpression->ValueType == Shader::EValueType::Void)
+	if (!InExpression->Type)
 	{
 		// Can't prepare value if we didn't successfully prepare a type
 		return EExpressionEvaluationType::None;
@@ -523,22 +604,18 @@ UE::HLSLTree::EExpressionEvaluationType UE::HLSLTree::PrepareExpressionValue(FEm
 		check(!InExpression->bReentryFlag); // Re-entry is not supported
 
 		InExpression->bReentryFlag = true;
-		const bool bPrepareResult = InExpression->PrepareValue(Context);
+		InExpression->PrepareValue(Context);
 		InExpression->bReentryFlag = false;
-
-		if (bPrepareResult)
-		{
-			check(InExpression->EvaluationType != EExpressionEvaluationType::None);
-		}
-		else
-		{
-			InExpression->EvaluationType = EExpressionEvaluationType::None;
-		}
 	}
 	return InExpression->EvaluationType;
 }
 
-bool UE::HLSLTree::FExpression::InternalSetValueShader(FEmitContext& Context, const TCHAR* InCode, int32 InLength, bool bInline)
+void UE::HLSLTree::FExpression::SetType(FUpdateTypeContext& Context, const FType& InType)
+{
+	Type = InType;
+}
+
+void UE::HLSLTree::FExpression::InternalSetValueShader(FEmitContext& Context, const TCHAR* InCode, int32 InLength, bool bInline)
 {
 	check(!Code);
 	check(EvaluationType == EExpressionEvaluationType::None);
@@ -572,56 +649,54 @@ bool UE::HLSLTree::FExpression::InternalSetValueShader(FEmitContext& Context, co
 			check(ParentScope);
 			Declaration = Context.AcquireLocalDeclarationCode();
 			ParentScope->ExpressionCodeMap.Add(Hash, Declaration);
-			const Shader::FValueTypeDescription& TypeDesc = Shader::GetValueTypeDescription(ValueType);
 			ParentScope->MarkLiveRecursive();
 			ParentScope->EmitStatementf(Context, TEXT("const %s %s = %s;"),
-				TypeDesc.Name,
+				Type.GetName(),
 				Declaration,
 				InCode);
 		}
 		Code = Declaration;
 	}
-
-	return true;
 }
 
-bool UE::HLSLTree::FExpression::SetValuePreshader(FEmitContext& Context, Shader::FPreshaderData& InPreshader)
+void UE::HLSLTree::FExpression::SetValuePreshader(FEmitContext& Context, Shader::FPreshaderData& InPreshader)
 {
 	check(!Preshader);
 	check(EvaluationType == EExpressionEvaluationType::None);
+	check(!Type.IsStruct());
 
 	EvaluationType = EExpressionEvaluationType::Preshader;
 	Preshader = new Shader::FPreshaderData(MoveTemp(InPreshader));
-	return true;
 }
 
-bool UE::HLSLTree::FExpression::SetValueConstant(FEmitContext& Context, const Shader::FValue& InValue)
+void UE::HLSLTree::FExpression::SetValueConstant(FEmitContext& Context, const Shader::FValue& InValue)
 {
 	check(EvaluationType == EExpressionEvaluationType::None);
+	check(!Type.IsStruct());
 
 	EvaluationType = EExpressionEvaluationType::Constant;
-	ConstantValue = Shader::Cast(InValue, ValueType);
-	return true;
+	ConstantValue = Shader::Cast(InValue, Type);
 }
 
-bool UE::HLSLTree::FExpression::SetValueForward(FEmitContext& Context, FExpression* Source)
+void UE::HLSLTree::FExpression::SetValueForward(FEmitContext& Context, FExpression* Source)
 {
 	check(EvaluationType == EExpressionEvaluationType::None);
 	check(Source);
 
+	ensure(Type == Source->GetType());
+
 	EvaluationType = PrepareExpressionValue(Context, Source);
 	switch (EvaluationType)
 	{
-	case EExpressionEvaluationType::None: return false;
+	case EExpressionEvaluationType::None: break;
 	case EExpressionEvaluationType::Shader: Code = Source->Code; break;
 	case EExpressionEvaluationType::Preshader: Preshader = Source->Preshader; break;
 	case EExpressionEvaluationType::Constant: ConstantValue = Source->ConstantValue; break;
-	default: checkNoEntry(); return false;
+	default: checkNoEntry(); break;
 	}
-	return true;
 }
 
-bool UE::HLSLTree::FExpression::SetValuePreshader(FEmitContext& Context, EExpressionEvaluationType InEvaluationType, Shader::FPreshaderData& InPreshader)
+void UE::HLSLTree::FExpression::SetValuePreshader(FEmitContext& Context, EExpressionEvaluationType InEvaluationType, Shader::FPreshaderData& InPreshader)
 {
 	if (InEvaluationType == EExpressionEvaluationType::Constant)
 	{
@@ -629,12 +704,12 @@ bool UE::HLSLTree::FExpression::SetValuePreshader(FEmitContext& Context, EExpres
 		FMaterialRenderContext RenderContext(nullptr, *Context.Material, nullptr);
 		Shader::FValue Value;
 		InPreshader.Evaluate(nullptr, RenderContext, Value);
-		return SetValueConstant(Context, Value);
+		SetValueConstant(Context, Value);
 	}
 	else
 	{
 		check(InEvaluationType == EExpressionEvaluationType::Preshader);
-		return SetValuePreshader(Context, InPreshader);
+		SetValuePreshader(Context, InPreshader);
 	}
 }
 
@@ -652,7 +727,7 @@ const TCHAR* UE::HLSLTree::FExpression::GetValueShader(FEmitContext& Context)
 		{
 			check(EvaluationType == EExpressionEvaluationType::Preshader);
 			check(Preshader);
-			Context.AddPreshader(ValueType, *Preshader, FormattedCode);
+			Context.AddPreshader(Type, *Preshader, FormattedCode);
 		}
 		Code = AllocateString(*Context.Allocator, FormattedCode);
 	}
@@ -660,11 +735,11 @@ const TCHAR* UE::HLSLTree::FExpression::GetValueShader(FEmitContext& Context)
 	return Code;
 }
 
-const TCHAR* UE::HLSLTree::FExpression::GetValueShader(FEmitContext& Context, Shader::EValueType Type)
+const TCHAR* UE::HLSLTree::FExpression::GetValueShader(FEmitContext& Context, const FType& InType)
 {
 	check(EvaluationType != EExpressionEvaluationType::None);
 
-	if (EvaluationType == EExpressionEvaluationType::Constant && Type != ValueType)
+	if (EvaluationType == EExpressionEvaluationType::Constant && Type != InType)
 	{
 		// If we need to cast a constant value, perform the cast first, then convert the result to HLSL
 		const Shader::FValue CastConstantValue = Shader::Cast(ConstantValue, Type);
@@ -675,7 +750,7 @@ const TCHAR* UE::HLSLTree::FExpression::GetValueShader(FEmitContext& Context, Sh
 	}
 
 	const ECastFlags CastFlags = ECastFlags::ReplicateScalar | ECastFlags::AllowAppendZeroes | ECastFlags::AllowTruncate;
-	return Context.CastShaderValue(this, GetValueShader(Context), ValueType, Type, CastFlags);
+	return Context.CastShaderValue(this, GetValueShader(Context), Type, InType, CastFlags);
 }
 
 void UE::HLSLTree::FExpression::GetValuePreshader(FEmitContext& Context, Shader::FPreshaderData& OutPreshader)
@@ -880,7 +955,27 @@ void UE::HLSLTree::FScope::WriteHLSL(int32 Indent, FStringBuilderBase& OutString
 	}
 }
 
-bool UE::HLSLTree::FTree::EmitHLSL(FEmitContext& Context, FStringBuilderBase& Writer) const
+void UE::HLSLTree::FStructType::WriteHLSL(FStringBuilderBase& OutString) const
+{
+	OutString.Appendf(TEXT("struct %s\n"), Name);
+	OutString.Append(TEXT("{\n"));
+	for (const FStructField& Field : Fields)
+	{
+		OutString.Appendf(TEXT("\t%s %s;\n"), Field.Type.GetName(), Field.Name);
+	}
+	OutString.Append(TEXT("};\n"));
+
+	for (const FStructField& Field : Fields)
+	{
+		OutString.Appendf(TEXT("%s %s_Set%s(%s Self, %s Value) { Self.%s = Value; return Self; }\n"),
+			Name, Name, Field.Name, Name, Field.Type.GetName(), Field.Name);
+	}
+	OutString.Append(TEXT("\n"));
+}
+
+bool UE::HLSLTree::FTree::EmitHLSL(FEmitContext& Context,
+	FStringBuilderBase& OutDeclarations,
+	FStringBuilderBase& OutCode) const
 {
 	if (!ResultStatement)
 	{
@@ -890,6 +985,10 @@ bool UE::HLSLTree::FTree::EmitHLSL(FEmitContext& Context, FStringBuilderBase& Wr
 	{
 		FUpdateTypeContext UpdateTypeContext(Context.Errors);
 		RequestScopeTypes(UpdateTypeContext, RootScope);
+		if (Context.Errors.Num() > 0)
+		{
+			return false;
+		}
 	}
 
 	ResultStatement->bEmitHLSL = true;
@@ -899,15 +998,32 @@ bool UE::HLSLTree::FTree::EmitHLSL(FEmitContext& Context, FStringBuilderBase& Wr
 		return false;
 	}
 
+	if (Context.Errors.Num() > 0)
+	{
+		return false;
+	}
+
 	if (RootScope->Statement)
 	{
 		RootScope->Statement->EmitHLSL(Context);
+		if (Context.Errors.Num() > 0)
+		{
+			return false;
+		}
 	}
 
 	Context.Finalize();
 
-	RootScope->WriteHLSL(1, Writer);
-	return true;
+	RootScope->WriteHLSL(1, OutCode);
+
+	const FStructType* StructType = StructTypes;
+	while (StructType)
+	{
+		StructType->WriteHLSL(OutDeclarations);
+		StructType = StructType->NextType;
+	}
+
+	return Context.Errors.Num() == 0;
 }
 
 void UE::HLSLTree::FTree::RegisterExpression(FScope& Scope, FExpression* Expression)
@@ -939,6 +1055,68 @@ UE::HLSLTree::FTextureParameterDeclaration* UE::HLSLTree::FTree::NewTextureParam
 	return Declaration;
 }
 
+namespace UE
+{
+namespace HLSLTree
+{
+namespace Private
+{
+
+static void SetComponentTypes(Shader::EValueComponentType* ComponentTypes, int32 ComponentIndex, const HLSLTree::FType& Type)
+{
+	if (Type.IsStruct())
+	{
+		for (const FStructField& Field : Type.StructType->Fields)
+		{
+			SetComponentTypes(ComponentTypes, ComponentIndex + Field.ComponentIndex, Field.Type);
+		}
+	}
+	else
+	{
+		const Shader::FValueTypeDescription TypeDesc = Shader::GetValueTypeDescription(Type);
+		for (int32 i = 0; i < TypeDesc.NumComponents; ++i)
+		{
+			ComponentTypes[ComponentIndex + i] = TypeDesc.ComponentType;
+		}
+	}
+}
+
+}
+}
+}
+
+const UE::HLSLTree::FStructType* UE::HLSLTree::FTree::NewStructType(const FStructTypeInitializer& Initializer)
+{
+	const int32 NumFields = Initializer.Fields.Num();
+	int32 ComponentIndex = 0;
+	FStructField* Fields = new(*Allocator) FStructField[NumFields];
+	for (int32 FieldIndex = 0; FieldIndex < NumFields; ++FieldIndex)
+	{
+		const FStructFieldInitializer& FieldInitializer = Initializer.Fields[FieldIndex];
+		const FType& FieldType = FieldInitializer.Type;
+		FStructField& Field = Fields[FieldIndex];
+		Field.Name = AllocateString(*Allocator, FieldInitializer.Name);
+		Field.Type = FieldType;
+		Field.ComponentIndex = ComponentIndex;
+		ComponentIndex += FieldType.GetNumComponents();
+	}
+
+	Shader::EValueComponentType* ComponentTypes = new(*Allocator) Shader::EValueComponentType[ComponentIndex];
+	for (int32 FieldIndex = 0; FieldIndex < NumFields; ++FieldIndex)
+	{
+		const FStructField& Field = Fields[FieldIndex];
+		Private::SetComponentTypes(ComponentTypes, Field.ComponentIndex, Field.Type);
+	}
+
+	FStructType* StructType = NewNode<FStructType>();
+	StructType->Name = AllocateString(*Allocator, Initializer.Name);
+	StructType->Fields = MakeArrayView(Fields, NumFields);
+	StructType->ComponentTypes = MakeArrayView(ComponentTypes, ComponentIndex);
+	StructType->NextType = StructTypes;
+	StructTypes = StructType;
+	return StructType;
+}
+
 void UE::HLSLTree::FTree::SetResult(FStatement& InResult)
 {
 	check(!ResultStatement);
diff --git a/Engine/Source/Runtime/Engine/Private/HLSLTree/HLSLTreeCommon.cpp b/Engine/Source/Runtime/Engine/Private/HLSLTree/HLSLTreeCommon.cpp
--- a/Engine/Source/Runtime/Engine/Private/HLSLTree/HLSLTreeCommon.cpp
+++ b/Engine/Source/Runtime/Engine/Private/HLSLTree/HLSLTreeCommon.cpp
@@ -4,26 +4,6 @@
 #include "MaterialShared.h"
 #include "Engine/Texture.h"
 
-static UE::Shader::EValueType GetShaderType(EMaterialValueType MaterialType)
-{
-	switch (MaterialType)
-	{
-	case MCT_Float1: return UE::Shader::EValueType::Float1;
-	case MCT_Float2: return UE::Shader::EValueType::Float2;
-	case MCT_Float3: return UE::Shader::EValueType::Float3;
-	case MCT_Float4: return UE::Shader::EValueType::Float4;
-	case MCT_Float: return UE::Shader::EValueType::Float1;
-	case MCT_StaticBool: return UE::Shader::EValueType::Bool1;
-	case MCT_MaterialAttributes: return UE::Shader::EValueType::MaterialAttributes;
-	case MCT_ShadingModel: return UE::Shader::EValueType::Int1;
-	case MCT_LWCScalar: return UE::Shader::EValueType::Double1;
-	case MCT_LWCVector2: return UE::Shader::EValueType::Double2;
-	case MCT_LWCVector3: return UE::Shader::EValueType::Double3;
-	case MCT_LWCVector4: return UE::Shader::EValueType::Double4;
-	default: checkNoEntry(); return UE::Shader::EValueType::Void;
-	}
-}
-
 UE::HLSLTree::FBinaryOpDescription UE::HLSLTree::GetBinaryOpDesription(EBinaryOp Op)
 {
 	switch (Op)
@@ -92,19 +72,61 @@ UE::HLSLTree::FSwizzleParameters UE::HLSLTree::MakeSwizzleMask(bool bInR, bool b
 	return FSwizzleParameters(ComponentIndex[0], ComponentIndex[1], ComponentIndex[2], ComponentIndex[3]);
 }
 
+void UE::HLSLTree::FExpressionConstant::UpdateType(FUpdateTypeContext& Context, const FRequestedType& RequestedType)
+{
+	if (Value.Type.IsStruct())
+	{
+		if (Value.Type.StructType != RequestedType.GetStructType())
+		{
+			return Context.Errors.AddError(this, TEXT("Type mismatch"));
+		}
+		return SetType(Context, Value.Type);
+	}
+	else
+	{
+		return SetType(Context, Shader::MakeValueTypeWithRequestedNumComponents(Value.Type, RequestedType.GetRequestedNumComponents()));
+	}
+}
 
-bool UE::HLSLTree::FExpressionConstant::PrepareValue(FEmitContext& Context)
+
+void UE::HLSLTree::FExpressionConstant::PrepareValue(FEmitContext& Context)
 {
-	return SetValueConstant(Context, Value);
+	check(Value.Component.Num() == Value.Type.GetNumComponents());
+	if (Value.Type.IsStruct())
+	{
+		const TArrayView<const Shader::EValueComponentType> ComponentTypes = Value.Type.StructType->ComponentTypes;
+		// TODO - Support for preshader structs?
+		TStringBuilder<1024> FormattedCode;
+		FormattedCode.Append(TEXT("{ "));
+		for (int32 ComponentIndex = 0; ComponentIndex < Value.Component.Num(); ++ComponentIndex)
+		{
+			if (ComponentIndex > 0)
+			{
+				FormattedCode.Append(TEXT(", "));
+			}
+			Value.Component[ComponentIndex].ToString(ComponentTypes[ComponentIndex], FormattedCode);
+		}
+		FormattedCode.Append(TEXT(" }"));
+		return SetValueShader(Context, FormattedCode);
+	}
+	else
+	{
+		Shader::FValue BaseValue(Value.Type.ValueType);
+		for (int32 ComponentIndex = 0; ComponentIndex < Value.Component.Num(); ++ComponentIndex)
+		{
+			BaseValue.Component[ComponentIndex] = Value.Component[ComponentIndex];
+		}
+		return SetValueConstant(Context, BaseValue);
+	}
 }
 
-bool UE::HLSLTree::FExpressionMaterialParameter::UpdateType(FUpdateTypeContext& Context, int8 InRequestedNumComponents)
+void UE::HLSLTree::FExpressionMaterialParameter::UpdateType(FUpdateTypeContext& Context, const FRequestedType& RequestedType)
 {
-	const Shader::EValueType Type = Shader::MakeValueTypeWithRequestedNumComponents(GetShaderValueType(ParameterType), InRequestedNumComponents);
-	return SetType(Context, Type);
+	const Shader::EValueType CurrentType = Shader::MakeValueTypeWithRequestedNumComponents(GetShaderValueType(ParameterType), RequestedType.GetRequestedNumComponents());
+	return SetType(Context, CurrentType);
 }
 
-bool UE::HLSLTree::FExpressionMaterialParameter::PrepareValue(FEmitContext& Context)
+void UE::HLSLTree::FExpressionMaterialParameter::PrepareValue(FEmitContext& Context)
 {
 	if (ParameterType == EMaterialParameterType::StaticSwitch)
 	{
@@ -141,10 +163,9 @@ bool UE::HLSLTree::FExpressionMaterialParameter::PrepareValue(FEmitContext& Cont
 		LocalPreshader.WriteOpcode(Shader::EPreshaderOpcode::Parameter).Write((uint16)ParameterIndex);
 		return SetValuePreshader(Context, LocalPreshader);
 	}
-	return true;
 }
 
-bool UE::HLSLTree::FExpressionExternalInput::PrepareValue(FEmitContext& Context)
+void UE::HLSLTree::FExpressionExternalInput::PrepareValue(FEmitContext& Context)
 {
 	const int32 TypeIndex = (int32)InputType;
 	const int32 TexCoordIndex = TypeIndex - (int32)UE::HLSLTree::EExternalInputType::TexCoord0;
@@ -158,23 +179,19 @@ bool UE::HLSLTree::FExpressionExternalInput::PrepareValue(FEmitContext& Context)
 	{
 		return Context.Errors.AddError(this, TEXT("Invalid texcoord"));
 	}
-	return true;
 }
 
-bool UE::HLSLTree::FExpressionTextureSample::UpdateType(FUpdateTypeContext& Context, int8 InRequestedNumComponents)
+void UE::HLSLTree::FExpressionTextureSample::UpdateType(FUpdateTypeContext& Context, const FRequestedType& RequestedType)
 {
-	if (RequestExpressionType(Context, TexCoordExpression, 2) == Shader::EValueType::Void)
-	{
-		return false;
-	}
+	RequestExpressionType(Context, TexCoordExpression, 2);
 	return SetType(Context, Shader::EValueType::Float4);
 }
 
-bool UE::HLSLTree::FExpressionTextureSample::PrepareValue(FEmitContext& Context)
+void UE::HLSLTree::FExpressionTextureSample::PrepareValue(FEmitContext& Context)
 {
 	if (PrepareExpressionValue(Context, TexCoordExpression) == EExpressionEvaluationType::None)
 	{
-		return false;
+		return;
 	}
 
 	const UE::HLSLTree::FTextureDescription& Desc = Declaration->Description;
@@ -309,53 +326,100 @@ bool UE::HLSLTree::FExpressionTextureSample::PrepareValue(FEmitContext& Context)
 		TexCoordExpression->GetValueShader(Context, Shader::EValueType::Float2));
 }
 
-bool UE::HLSLTree::FExpressionDefaultMaterialAttributes::PrepareValue(FEmitContext& Context)
+void UE::HLSLTree::FExpressionGetStructField::UpdateType(FUpdateTypeContext& Context, const FRequestedType& RequestedType)
+{
+	const FStructFieldRef FieldRef = StructType->FindFieldByName(FieldName);
+	if (!FieldRef)
+	{
+		return Context.Errors.AddErrorf(this, TEXT("Invalid field %s"), FieldName);
+	}
+
+	if (RequestExpressionType(Context, StructExpression, RequestedType.MakeFieldAccess(StructType, FieldRef)) != FType(StructType))
+	{
+		return Context.Errors.AddErrorf(this, TEXT("Expected type %s"), StructType->Name);
+	}
+
+	return SetType(Context, FieldRef.Type);
+}
+
+void UE::HLSLTree::FExpressionGetStructField::PrepareValue(FEmitContext& Context)
 {
-	return SetValueInlineShaderf(Context, TEXT("DefaultMaterialAttributes"));
+	const FStructFieldRef FieldRef = StructType->FindFieldByName(FieldName);
+	check(FieldRef);
+
+	if (PrepareExpressionValue(Context, StructExpression) == EExpressionEvaluationType::None)
+	{
+		return;
+	}
+
+	return SetValueInlineShaderf(Context, TEXT("%s.%s"),
+		StructExpression->GetValueShader(Context, StructType),
+		FieldName);
 }
 
-bool UE::HLSLTree::FExpressionSetMaterialAttribute::UpdateType(FUpdateTypeContext& Context, int8 InRequestedNumComponents)
+void UE::HLSLTree::FExpressionSetStructField::UpdateType(FUpdateTypeContext& Context, const FRequestedType& RequestedType)
 {
-	const EMaterialValueType AttributeType = FMaterialAttributeDefinitionMap::GetValueType(AttributeID);
-	const uint32 NumComponents = GetNumComponents(AttributeType);
-	if (RequestExpressionType(Context, AttributesExpression, 4) != Shader::EValueType::MaterialAttributes)
+	const FStructFieldRef FieldRef = StructType->FindFieldByName(FieldName);
+	if (!FieldRef)
 	{
-		return Context.Errors.AddError(this, TEXT("Expected type MaterialAttributes"));
+		return Context.Errors.AddErrorf(this, TEXT("Invalid field %s"), FieldName);
 	}
-	if (RequestExpressionType(Context, ValueExpression, NumComponents) == Shader::EValueType::Void)
+
+	const FType StructExpressionType = RequestExpressionType(Context, StructExpression, RequestedType.CopyWithFieldRemoved(FieldRef));
+	if (StructExpressionType && StructExpressionType.StructType != StructType)
 	{
-		return false;
+		return Context.Errors.AddErrorf(this, TEXT("Expected type %s"), StructType->Name);
 	}
-	return SetType(Context, Shader::EValueType::MaterialAttributes);
+
+	RequestExpressionType(Context, FieldExpression, RequestedType.GetField(FieldRef));
+	return SetType(Context, StructType);
 }
 
-bool UE::HLSLTree::FExpressionSetMaterialAttribute::PrepareValue(FEmitContext& Context)
+void UE::HLSLTree::FExpressionSetStructField::PrepareValue(FEmitContext& Context)
 {
-	if (PrepareExpressionValue(Context, AttributesExpression) == EExpressionEvaluationType::None)
+	const EExpressionEvaluationType StructEvaluation = PrepareExpressionValue(Context, StructExpression);
+	const EExpressionEvaluationType FieldEvaluation = PrepareExpressionValue(Context, FieldExpression);
+	
+	if (FieldEvaluation != EExpressionEvaluationType::None)
 	{
-		return false;
+		const FStructFieldRef FieldRef = StructType->FindFieldByName(FieldName);
+		check(FieldRef);
+
+		if (StructEvaluation != EExpressionEvaluationType::None)
+		{
+			return SetValueShaderf(Context, TEXT("%s_Set%s(%s, %s)"),
+				StructType->Name,
+				FieldName,
+				StructExpression->GetValueShader(Context, StructType),
+				FieldExpression->GetValueShader(Context, FieldRef.Type));
+		}
+		else
+		{
+			// StructExpression is not used, so default to a zero-initialized struct
+			// This will happen if all the accessed struct fields are explicitly defined
+			return SetValueShaderf(Context, TEXT("%s_Set%s((%s)0, %s)"),
+				StructType->Name,
+				FieldName,
+				StructType->Name,
+				FieldExpression->GetValueShader(Context, FieldRef.Type));
+		}
 	}
-	if (PrepareExpressionValue(Context, ValueExpression) == EExpressionEvaluationType::None)
+	else
 	{
-		return false;
+		// Skip assigning the field, and just forward the struct value
+		check(StructEvaluation != EExpressionEvaluationType::None);
+		return SetValueForward(Context, StructExpression);
 	}
-
-	const EMaterialValueType AttributeType = FMaterialAttributeDefinitionMap::GetValueType(AttributeID);
-	const FString PropertyName = FMaterialAttributeDefinitionMap::GetAttributeName(AttributeID);
-	return SetValueShaderf(Context, TEXT("FMaterialAttributes_Set%s(%s, %s)"),
-		*PropertyName,
-		AttributesExpression->GetValueShader(Context, Shader::EValueType::MaterialAttributes),
-		ValueExpression->GetValueShader(Context, GetShaderType(AttributeType)));
 }
 
-bool UE::HLSLTree::FExpressionSelect::UpdateType(FUpdateTypeContext& Context, int8 InRequestedNumComponents)
+void UE::HLSLTree::FExpressionSelect::UpdateType(FUpdateTypeContext& Context, const FRequestedType& RequestedType)
 {
-	const Shader::EValueType ConditionType = RequestExpressionType(Context, ConditionExpression, 1);
-	const Shader::EValueType LhsType = RequestExpressionType(Context, FalseExpression, InRequestedNumComponents);
-	const Shader::EValueType RhsType = RequestExpressionType(Context, TrueExpression, InRequestedNumComponents);
-	if (ConditionType == Shader::EValueType::Void || LhsType == Shader::EValueType::Void || RhsType == Shader::EValueType::Void)
+	const FType ConditionType = RequestExpressionType(Context, ConditionExpression, 1);
+	const FType LhsType = RequestExpressionType(Context, FalseExpression, RequestedType);
+	const FType RhsType = RequestExpressionType(Context, TrueExpression, RequestedType);
+	if (!ConditionType || !LhsType || !RhsType)
 	{
-		return false;
+		return;
 	}
 
 	const Shader::FValueTypeDescription LhsTypeDesc = Shader::GetValueTypeDescription(LhsType);
@@ -365,16 +429,23 @@ bool UE::HLSLTree::FExpressionSelect::UpdateType(FUpdateTypeContext& Context, in
 		return Context.Errors.AddError(this, TEXT("Type mismatch"));
 	}
 
-	const int8 NumComponents = FMath::Max(LhsTypeDesc.NumComponents, RhsTypeDesc.NumComponents);
-	return SetType(Context, Shader::MakeValueType(LhsTypeDesc.ComponentType, NumComponents));
+	if (LhsType.IsStruct())
+	{
+		return SetType(Context, LhsType);
+	}
+	else
+	{
+		const int8 NumComponents = FMath::Max(LhsTypeDesc.NumComponents, RhsTypeDesc.NumComponents);
+		return SetType(Context, Shader::MakeValueType(LhsTypeDesc.ComponentType, NumComponents));
+	}
 }
 
-bool UE::HLSLTree::FExpressionSelect::PrepareValue(FEmitContext& Context)
+void UE::HLSLTree::FExpressionSelect::PrepareValue(FEmitContext& Context)
 {
 	const EExpressionEvaluationType ConditionEvaluation = PrepareExpressionValue(Context, ConditionExpression);
 	if (ConditionEvaluation == EExpressionEvaluationType::None)
 	{
-		return false;
+		return;
 	}
 
 	if (ConditionEvaluation == EExpressionEvaluationType::Constant)
@@ -388,25 +459,24 @@ bool UE::HLSLTree::FExpressionSelect::PrepareValue(FEmitContext& Context)
 		const EExpressionEvaluationType FalseEvaluation = PrepareExpressionValue(Context, FalseExpression);
 		if (TrueEvaluation == EExpressionEvaluationType::None || FalseEvaluation == EExpressionEvaluationType::None)
 		{
-			return false;
+			return;
 		}
 
 		// TODO - preshader
 		return SetValueShaderf(Context, TEXT("(%s ? %s : %s)"),
 			ConditionExpression->GetValueShader(Context, Shader::EValueType::Bool1),
-			TrueExpression->GetValueShader(Context, GetValueType()),
-			FalseExpression->GetValueShader(Context, GetValueType()));
+			TrueExpression->GetValueShader(Context, GetType()),
+			FalseExpression->GetValueShader(Context, GetType()));
 	}
-	return true;
 }
 
-bool UE::HLSLTree::FExpressionBinaryOp::UpdateType(FUpdateTypeContext& Context, int8 InRequestedNumComponents)
+void UE::HLSLTree::FExpressionBinaryOp::UpdateType(FUpdateTypeContext& Context, const FRequestedType& RequestedType)
 {
-	const Shader::EValueType LhsType = RequestExpressionType(Context, Lhs, InRequestedNumComponents);
-	const Shader::EValueType RhsType = RequestExpressionType(Context, Rhs, InRequestedNumComponents);
-	if (LhsType == Shader::EValueType::Void || RhsType == Shader::EValueType::Void)
+	const FType LhsType = RequestExpressionType(Context, Lhs, RequestedType);
+	const FType RhsType = RequestExpressionType(Context, Rhs, RequestedType);
+	if (!LhsType || !RhsType)
 	{
-		return false;
+		return;
 	}
 
 	FString ErrorMessage;
@@ -428,13 +498,13 @@ bool UE::HLSLTree::FExpressionBinaryOp::UpdateType(FUpdateTypeContext& Context,
 	return SetType(Context, ResultType);
 }
 
-bool UE::HLSLTree::FExpressionBinaryOp::PrepareValue(FEmitContext& Context)
+void UE::HLSLTree::FExpressionBinaryOp::PrepareValue(FEmitContext& Context)
 {
 	const EExpressionEvaluationType LhsEvaluation = PrepareExpressionValue(Context, Lhs);
 	const EExpressionEvaluationType RhsEvaluation = PrepareExpressionValue(Context, Rhs);
 	if (LhsEvaluation == EExpressionEvaluationType::None || RhsEvaluation == EExpressionEvaluationType::None)
 	{
-		return false;
+		return;
 	}
 
 	const EExpressionEvaluationType Evaluation = CombineEvaluationTypes(LhsEvaluation, RhsEvaluation);
@@ -444,7 +514,9 @@ bool UE::HLSLTree::FExpressionBinaryOp::PrepareValue(FEmitContext& Context)
 		switch (Op)
 		{
 		case UE::HLSLTree::EBinaryOp::Add: Opcode = Shader::EPreshaderOpcode::Add; break;
+		case UE::HLSLTree::EBinaryOp::Sub: Opcode = Shader::EPreshaderOpcode::Add; break;
 		case UE::HLSLTree::EBinaryOp::Mul: Opcode = Shader::EPreshaderOpcode::Mul; break;
+		case UE::HLSLTree::EBinaryOp::Div: Opcode = Shader::EPreshaderOpcode::Div; break;
 		default: break;
 		}
 	}
@@ -457,9 +529,11 @@ bool UE::HLSLTree::FExpressionBinaryOp::PrepareValue(FEmitContext& Context)
 		switch (Op)
 		{
 		case EBinaryOp::Add: return SetValueShaderf(Context, TEXT("(%s + %s)"), LhsCode, RhsCode);
+		case EBinaryOp::Sub: return SetValueShaderf(Context, TEXT("(%s - %s)"), LhsCode, RhsCode);
 		case EBinaryOp::Mul: return SetValueShaderf(Context, TEXT("(%s * %s)"), LhsCode, RhsCode);
+		case EBinaryOp::Div: return SetValueShaderf(Context, TEXT("(%s / %s)"), LhsCode, RhsCode);
 		case EBinaryOp::Less: return SetValueShaderf(Context, TEXT("(%s < %s)"), LhsCode, RhsCode);
-		default: checkNoEntry(); return false;
+		default: checkNoEntry(); return;
 		}
 	}
 	else
@@ -472,28 +546,33 @@ bool UE::HLSLTree::FExpressionBinaryOp::PrepareValue(FEmitContext& Context)
 	}
 }
 
-bool UE::HLSLTree::FExpressionSwizzle::UpdateType(FUpdateTypeContext& Context, int8 InRequestedNumComponents)
+void UE::HLSLTree::FExpressionSwizzle::UpdateType(FUpdateTypeContext& Context, const FRequestedType& RequestedType)
 {
-	const int32 NumSwizzleComponents = FMath::Min<int32>(InRequestedNumComponents, Parameters.NumComponents);
-	int32 MaxSwizzleIndex = 0;
-	for (int32 i = 0; i < NumSwizzleComponents; ++i)
+	FRequestedType RequestedInputType(4, false);
+	int32 MaxComponentIndex = 0;
+	for (int32 ComponentIndex = 0; ComponentIndex < Parameters.NumComponents; ++ComponentIndex)
 	{
-		MaxSwizzleIndex = FMath::Max<int32>(MaxSwizzleIndex, Parameters.ComponentIndex[i]);
+		if (RequestedType.IsComponentRequested(ComponentIndex))
+		{
+			const int32 SwizzledComponentIndex = Parameters.ComponentIndex[ComponentIndex];
+			RequestedInputType.RequestedComponents[SwizzledComponentIndex] = true;
+			MaxComponentIndex = ComponentIndex;
+		}
 	}
 
-	const Shader::EValueType InputType = RequestExpressionType(Context, Input, MaxSwizzleIndex + 1);
-	return SetType(Context, Shader::MakeValueType(InputType, NumSwizzleComponents));
+	const Shader::EValueType InputType = RequestExpressionType(Context, Input, RequestedInputType);
+	return SetType(Context, Shader::MakeValueType(InputType, MaxComponentIndex + 1));
 }
 
-bool UE::HLSLTree::FExpressionSwizzle::PrepareValue(FEmitContext& Context)
+void UE::HLSLTree::FExpressionSwizzle::PrepareValue(FEmitContext& Context)
 {
 	const EExpressionEvaluationType InputEvaluation = PrepareExpressionValue(Context, Input);
 	if (InputEvaluation == EExpressionEvaluationType::None)
 	{
-		return false;
+		return;
 	}
 
-	const Shader::FValueTypeDescription TypeDesc = Shader::GetValueTypeDescription(GetValueType());
+	const Shader::FValueTypeDescription TypeDesc = Shader::GetValueTypeDescription(GetType());
 	const int32 NumSwizzleComponents = FMath::Min<int32>(TypeDesc.NumComponents, Parameters.NumComponents);
 
 	static const TCHAR ComponentName[] = { 'x', 'y', 'z', 'w' };
@@ -521,7 +600,7 @@ bool UE::HLSLTree::FExpressionSwizzle::PrepareValue(FEmitContext& Context)
 		}
 		else
 		{
-			const Shader::EValueType InputType = Shader::MakeValueType(GetValueType(), NumSwizzleComponents);
+			const Shader::EValueType InputType = Shader::MakeValueType(GetType(), NumSwizzleComponents);
 			return SetValueInlineShaderf(Context, TEXT("%s"), Input->GetValueShader(Context, InputType));
 		}
 	}
@@ -539,10 +618,20 @@ bool UE::HLSLTree::FExpressionSwizzle::PrepareValue(FEmitContext& Context)
 	}
 }
 
-bool UE::HLSLTree::FExpressionAppend::UpdateType(FUpdateTypeContext& Context, int8 InRequestedNumComponents)
+void UE::HLSLTree::FExpressionAppend::UpdateType(FUpdateTypeContext& Context, const FRequestedType& RequestedType)
 {
-	const Shader::EValueType LhsType = RequestExpressionType(Context, Lhs, 4);
-	const Shader::EValueType RhsType = RequestExpressionType(Context, Rhs, 4);
+	const FType LhsType = RequestExpressionType(Context, Lhs, 4);
+	const FType RhsType = RequestExpressionType(Context, Rhs, 4);
+	if (!LhsType || !RhsType)
+	{
+		return;
+	}
+
+	if (LhsType.IsStruct() || RhsType.IsStruct())
+	{
+		return Context.Errors.AddError(this, TEXT("Not supported for structs"));
+	}
+
 	const Shader::FValueTypeDescription LhsTypeDesc = Shader::GetValueTypeDescription(LhsType);
 	const Shader::FValueTypeDescription RhsTypeDesc = Shader::GetValueTypeDescription(RhsType);
 	if (LhsTypeDesc.ComponentType != RhsTypeDesc.ComponentType)
@@ -550,23 +639,23 @@ bool UE::HLSLTree::FExpressionAppend::UpdateType(FUpdateTypeContext& Context, in
 		return Context.Errors.AddError(this, TEXT("Type mismatch"));
 	}
 
-	const int32 NumComponents = FMath::Min<int32>(LhsTypeDesc.NumComponents + RhsTypeDesc.NumComponents, InRequestedNumComponents);
+	const int32 NumComponents = FMath::Min<int32>(LhsTypeDesc.NumComponents + RhsTypeDesc.NumComponents, RequestedType.GetRequestedNumComponents());
 	return SetType(Context, Shader::MakeValueType(LhsTypeDesc.ComponentType, NumComponents));
 }
 
-bool UE::HLSLTree::FExpressionAppend::PrepareValue(FEmitContext& Context)
+void UE::HLSLTree::FExpressionAppend::PrepareValue(FEmitContext& Context)
 {
 	const EExpressionEvaluationType LhsEvaluation = PrepareExpressionValue(Context, Lhs);
 	const EExpressionEvaluationType RhsEvaluation = PrepareExpressionValue(Context, Rhs);
 	if (LhsEvaluation == EExpressionEvaluationType::None || RhsEvaluation == EExpressionEvaluationType::None)
 	{
-		return false;
+		return;
 	}
 
 	const EExpressionEvaluationType Evaluation = CombineEvaluationTypes(LhsEvaluation, RhsEvaluation);
 	if (Evaluation == EExpressionEvaluationType::Shader)
 	{
-		const Shader::FValueTypeDescription ResultTypeDesc = Shader::GetValueTypeDescription(GetValueType());
+		const Shader::FValueTypeDescription ResultTypeDesc = Shader::GetValueTypeDescription(GetType());
 		return SetValueShaderf(Context, TEXT("%s(%s, %s)"),
 			ResultTypeDesc.Name,
 			Lhs->GetValueShader(Context),
@@ -582,7 +671,7 @@ bool UE::HLSLTree::FExpressionAppend::PrepareValue(FEmitContext& Context)
 	}
 }
 
-bool UE::HLSLTree::FExpressionReflectionVector::PrepareValue(FEmitContext& Context)
+void UE::HLSLTree::FExpressionReflectionVector::PrepareValue(FEmitContext& Context)
 {
 	return SetValueInlineShaderf(Context, TEXT("Parameters.ReflectionVector"));
 }
@@ -595,7 +684,7 @@ void UE::HLSLTree::FStatementBreak::EmitHLSL(FEmitContext& Context) const
 
 void UE::HLSLTree::FStatementReturn::RequestTypes(FUpdateTypeContext& Context) const
 {
-	RequestExpressionType(Context, Expression, 4);
+	RequestExpressionType(Context, Expression, Type);
 }
 
 void UE::HLSLTree::FStatementReturn::EmitHLSL(FEmitContext& Context) const
diff --git a/Engine/Source/Runtime/Engine/Private/Materials/HLSLMaterialTranslator.cpp b/Engine/Source/Runtime/Engine/Private/Materials/HLSLMaterialTranslator.cpp
--- a/Engine/Source/Runtime/Engine/Private/Materials/HLSLMaterialTranslator.cpp
+++ b/Engine/Source/Runtime/Engine/Private/Materials/HLSLMaterialTranslator.cpp
@@ -96,7 +96,7 @@ static inline EMaterialValueType GetMaterialValueType(UE::Shader::EValueType Typ
 	case UE::Shader::EValueType::Bool3: return MCT_Float3;
 	case UE::Shader::EValueType::Bool4: return MCT_Float4;
 
-	case UE::Shader::EValueType::MaterialAttributes: return MCT_MaterialAttributes;
+	case UE::Shader::EValueType::Struct: return MCT_MaterialAttributes;
 	default: checkNoEntry(); return MCT_Unknown;
 	}
 }
@@ -2032,10 +2032,11 @@ FString FHLSLMaterialTranslator::GetMaterialShaderCode()
 
 	FString MaterialAttributesDeclaration;
 	FString MaterialAttributesUtilities;
-	FString MaterialAttributesDefault;
 
 	const EMaterialShadingModel DefaultShadingModel = Material->GetShadingModels().GetFirstShadingModel();
 
+	MaterialAttributesDeclaration += TEXT("struct FMaterialAttributes\n{\n");
+
 	const TArray<FGuid>& OrderedVisibleAttributes = FMaterialAttributeDefinitionMap::GetOrderedVisibleAttributeList();
 	for (const FGuid& AttributeID : OrderedVisibleAttributes)
 	{
@@ -2063,21 +2064,11 @@ FString FHLSLMaterialTranslator::GetMaterialShaderCode()
 			// Chainable method to set the attribute
 			MaterialAttributesUtilities += FString::Printf(TEXT("FMaterialAttributes FMaterialAttributes_Set%s(FMaterialAttributes InAttributes, %s InValue) { InAttributes.%s = InValue; return InAttributes; }") LINE_TERMINATOR,
 				*PropertyName, HLSLType, *PropertyName);
-
-			// Standard value type
-			switch (PropertyType)
-			{
-			case MCT_Float: case MCT_Float1: MaterialAttributesDefault += FString::Printf(TEXT("\tResult.%s = %0.8f;") LINE_TERMINATOR, *PropertyName, DefaultValue.X); break;
-			case MCT_Float2: MaterialAttributesDefault += FString::Printf(TEXT("\tResult.%s = MaterialFloat2(%0.8f,%0.8f);") LINE_TERMINATOR, *PropertyName, DefaultValue.X, DefaultValue.Y); break;
-			case MCT_Float3: MaterialAttributesDefault += FString::Printf(TEXT("\tResult.%s = MaterialFloat3(%0.8f,%0.8f,%0.8f);") LINE_TERMINATOR, *PropertyName, DefaultValue.X, DefaultValue.Y, DefaultValue.Z); break;
-			case MCT_Float4: MaterialAttributesDefault += FString::Printf(TEXT("\tResult.%s = MaterialFloat4(%0.8f,%0.8f,%0.8f,%0.8f);") LINE_TERMINATOR, *PropertyName, DefaultValue.X, DefaultValue.Y, DefaultValue.Z, DefaultValue.W); break;
-			case MCT_ShadingModel: MaterialAttributesDefault += FString::Printf(TEXT("\tResult.%s = %d;") LINE_TERMINATOR, *PropertyName, (int32)DefaultShadingModel); break;
-			case MCT_Strata: MaterialAttributesDefault += FString::Printf(TEXT("\tResult.%s = GetInitialisedStrataData();") LINE_TERMINATOR, *PropertyName); break; // TODO
-			default: checkNoEntry(); break;
-			}
 		}
 	}
 
+	MaterialAttributesDeclaration += TEXT("};\n");
+
 	LazyPrintf.PushParam(*MaterialAttributesDeclaration);
 	LazyPrintf.PushParam(*MaterialAttributesUtilities);
 
@@ -2206,8 +2197,6 @@ FString FHLSLMaterialTranslator::GetMaterialShaderCode()
 
 	LazyPrintf.PushParam(*CustomInterpolatorAssignments);
 
-	LazyPrintf.PushParam(*MaterialAttributesDefault);
-
 	if (bEnableExecutionFlow)
 	{
 		FString EvaluateVertexCode;
diff --git a/Engine/Source/Runtime/Engine/Private/Materials/Material.cpp b/Engine/Source/Runtime/Engine/Private/Materials/Material.cpp
--- a/Engine/Source/Runtime/Engine/Private/Materials/Material.cpp
+++ b/Engine/Source/Runtime/Engine/Private/Materials/Material.cpp
@@ -4484,7 +4484,7 @@ static void SetMaterialInputDescription(FShadingModelMaterialInput& Input, bool
 
 static void SetMaterialInputDescription(FMaterialAttributesInput& Input, bool bHidden, FMaterialInputDescription& OutDescription)
 {
-	OutDescription.Type = UE::Shader::EValueType::MaterialAttributes;
+	OutDescription.Type = UE::Shader::EValueType::Struct;
 	OutDescription.Input = &Input;
 	OutDescription.bUseConstant = false;
 	OutDescription.bHidden = bHidden;
diff --git a/Engine/Source/Runtime/Engine/Private/Materials/MaterialExpressionHLSL.cpp b/Engine/Source/Runtime/Engine/Private/Materials/MaterialExpressionHLSL.cpp
--- a/Engine/Source/Runtime/Engine/Private/Materials/MaterialExpressionHLSL.cpp
+++ b/Engine/Source/Runtime/Engine/Private/Materials/MaterialExpressionHLSL.cpp
@@ -31,8 +31,10 @@
 #include "Materials/MaterialExpressionGetLocal.h"
 #include "Materials/MaterialExpressionAdd.h"
 #include "Materials/MaterialExpressionMultiply.h"
+#include "Materials/MaterialExpressionDivide.h"
 #include "Materials/MaterialExpressionBinaryOp.h"
 #include "Materials/MaterialExpressionAppendVector.h"
+#include "Materials/MaterialExpressionGetMaterialAttributes.h"
 #include "Materials/MaterialExpressionSetMaterialAttributes.h"
 #include "Materials/MaterialExpressionReflectionVectorWS.h"
 #include "Materials/MaterialExpressionSetLocal.h"
@@ -235,6 +237,18 @@ EMaterialGenerateHLSLStatus UMaterialExpressionMultiply::GenerateHLSLExpression(
 	return EMaterialGenerateHLSLStatus::Success;
 }
 
+EMaterialGenerateHLSLStatus UMaterialExpressionDivide::GenerateHLSLExpression(FMaterialHLSLGenerator& Generator, UE::HLSLTree::FScope& Scope, int32 OutputIndex, UE::HLSLTree::FExpression*& OutExpression)
+{
+	UE::HLSLTree::FExpression* Lhs = A.GetTracedInput().Expression ? A.AcquireHLSLExpression(Generator, Scope) : Generator.NewConstant(ConstA);
+	UE::HLSLTree::FExpression* Rhs = B.GetTracedInput().Expression ? B.AcquireHLSLExpression(Generator, Scope) : Generator.NewConstant(ConstB);
+	if (!Lhs || !Rhs)
+	{
+		return EMaterialGenerateHLSLStatus::Error;
+	}
+	OutExpression = Generator.GetTree().NewExpression<UE::HLSLTree::FExpressionBinaryOp>(Scope, UE::HLSLTree::EBinaryOp::Div, Lhs, Rhs);
+	return EMaterialGenerateHLSLStatus::Success;
+}
+
 EMaterialGenerateHLSLStatus UMaterialExpressionAppendVector::GenerateHLSLExpression(FMaterialHLSLGenerator& Generator, UE::HLSLTree::FScope& Scope, int32 OutputIndex, UE::HLSLTree::FExpression*& OutExpression)
 {
 	UE::HLSLTree::FExpression* Lhs = A.AcquireHLSLExpression(Generator, Scope);
@@ -247,6 +261,37 @@ EMaterialGenerateHLSLStatus UMaterialExpressionAppendVector::GenerateHLSLExpress
 	return EMaterialGenerateHLSLStatus::Success;
 }
 
+EMaterialGenerateHLSLStatus UMaterialExpressionGetMaterialAttributes::GenerateHLSLExpression(FMaterialHLSLGenerator& Generator, UE::HLSLTree::FScope& Scope, int32 OutputIndex, UE::HLSLTree::FExpression*& OutExpression)
+{
+	UE::HLSLTree::FExpression* AttributesExpression = MaterialAttributes.AcquireHLSLExpression(Generator, Scope);
+	if (!AttributesExpression)
+	{
+		return EMaterialGenerateHLSLStatus::Error;
+	}
+	if (OutputIndex == 0)
+	{
+		OutExpression = AttributesExpression;
+		return EMaterialGenerateHLSLStatus::Success;
+	}
+	const int32 AttributeIndex = OutputIndex - 1;
+	if (!AttributeGetTypes.IsValidIndex(AttributeIndex))
+	{
+		return Generator.Error(TEXT("Invalid attribute"));
+	}
+
+	const FGuid& AttributeID = AttributeGetTypes[AttributeIndex];
+	const FString& AttributeName = FMaterialAttributeDefinitionMap::GetAttributeName(AttributeID);
+
+	UE::HLSLTree::FExpressionGetStructField* SetAttributeExpression = Generator.GetTree().NewExpression<UE::HLSLTree::FExpressionGetStructField>(Scope);
+	SetAttributeExpression->StructType = Generator.GetMaterialAttributesType();
+	SetAttributeExpression->StructExpression = AttributesExpression;
+	SetAttributeExpression->FieldName = *AttributeName;
+	OutExpression = SetAttributeExpression;
+
+	return EMaterialGenerateHLSLStatus::Success;
+
+}
+
 EMaterialGenerateHLSLStatus UMaterialExpressionSetMaterialAttributes::GenerateHLSLExpression(FMaterialHLSLGenerator& Generator, UE::HLSLTree::FScope& Scope, int32 OutputIndex, UE::HLSLTree::FExpression*& OutExpression)
 {
 	UE::HLSLTree::FExpression* AttributesExpression = nullptr;
@@ -256,7 +301,7 @@ EMaterialGenerateHLSLStatus UMaterialExpressionSetMaterialAttributes::GenerateHL
 	}
 	else
 	{
-		AttributesExpression = Generator.GetTree().NewExpression<UE::HLSLTree::FExpressionDefaultMaterialAttributes>(Scope);
+		AttributesExpression = Generator.GetTree().NewExpression<UE::HLSLTree::FExpressionConstant>(Scope, Generator.GetMaterialAttributesDefaultValue());
 	}
 
 	for (int32 PinIndex = 0; PinIndex < AttributeSetTypes.Num(); ++PinIndex)
@@ -267,15 +312,17 @@ EMaterialGenerateHLSLStatus UMaterialExpressionSetMaterialAttributes::GenerateHL
 			const FGuid& AttributeID = AttributeSetTypes[PinIndex];
 			// Only compile code to set attributes of the current shader frequency
 			const EShaderFrequency AttributeFrequency = FMaterialAttributeDefinitionMap::GetShaderFrequency(AttributeID);
+			const FString& AttributeName = FMaterialAttributeDefinitionMap::GetAttributeName(AttributeID);
 			//if (AttributeFrequency == Compiler->GetCurrentShaderFrequency())
 			{
 				UE::HLSLTree::FExpression* ValueExpression = AttributeInput.AcquireHLSLExpression(Generator, Scope);
 				if (ValueExpression)
 				{
-					UE::HLSLTree::FExpressionSetMaterialAttribute* SetAttributeExpression = Generator.GetTree().NewExpression<UE::HLSLTree::FExpressionSetMaterialAttribute>(Scope);
-					SetAttributeExpression->AttributeID = AttributeID;
-					SetAttributeExpression->AttributesExpression = AttributesExpression;
-					SetAttributeExpression->ValueExpression = ValueExpression;
+					UE::HLSLTree::FExpressionSetStructField* SetAttributeExpression = Generator.GetTree().NewExpression<UE::HLSLTree::FExpressionSetStructField>(Scope);
+					SetAttributeExpression->StructType = Generator.GetMaterialAttributesType();
+					SetAttributeExpression->FieldName = *AttributeName;
+					SetAttributeExpression->StructExpression = AttributesExpression;
+					SetAttributeExpression->FieldExpression = ValueExpression;
 					AttributesExpression = SetAttributeExpression;
 				}
 			}
diff --git a/Engine/Source/Runtime/Engine/Private/Materials/MaterialExpressions.cpp b/Engine/Source/Runtime/Engine/Private/Materials/MaterialExpressions.cpp
--- a/Engine/Source/Runtime/Engine/Private/Materials/MaterialExpressions.cpp
+++ b/Engine/Source/Runtime/Engine/Private/Materials/MaterialExpressions.cpp
@@ -22294,7 +22294,12 @@ void UMaterialExpressionGetLocal::GetCaption(TArray<FString>& OutCaptions) const
 
 uint32 UMaterialExpressionGetLocal::GetOutputType(int32 InputIndex)
 {
-	return MCT_Float1;
+	return MCT_Unknown;
+}
+
+bool UMaterialExpressionGetLocal::IsResultMaterialAttributes(int32 OutputIndex)
+{
+	return false;
 }
 #endif
 
@@ -22317,11 +22322,7 @@ UMaterialExpressionSetLocal::UMaterialExpressionSetLocal(const FObjectInitialize
 #if WITH_EDITOR
 uint32 UMaterialExpressionSetLocal::GetInputType(int32 InputIndex)
 {
-	switch (InputIndex)
-	{
-	case 0: return MCT_Float1;
-	default: checkNoEntry(); return 0u;
-	}
+	return MCT_Unknown;
 }
 
 int32 UMaterialExpressionSetLocal::Compile(class FMaterialCompiler* Compiler, int32 OutputIndex)
diff --git a/Engine/Source/Runtime/Engine/Private/Materials/MaterialHLSLEmitter.cpp b/Engine/Source/Runtime/Engine/Private/Materials/MaterialHLSLEmitter.cpp
--- a/Engine/Source/Runtime/Engine/Private/Materials/MaterialHLSLEmitter.cpp
+++ b/Engine/Source/Runtime/Engine/Private/Materials/MaterialHLSLEmitter.cpp
@@ -44,6 +44,7 @@ static const TCHAR* HLSLTypeString(EMaterialValueType Type)
 static FString GenerateMaterialTemplateHLSL(EShaderPlatform ShaderPlatform,
 	const FMaterial& Material,
 	const UE::HLSLTree::FEmitContext& EmitContext,
+	const TCHAR* DeclarationsCode,
 	const TCHAR* PixelShaderCode)
 {
 	FString MaterialTemplateSource;
@@ -131,8 +132,11 @@ static FString GenerateMaterialTemplateHLSL(EShaderPlatform ShaderPlatform,
 		}
 	}
 
-	LazyPrintf.PushParam(*MaterialAttributesDeclaration);
-	LazyPrintf.PushParam(*MaterialAttributesUtilities);
+	//LazyPrintf.PushParam(*MaterialAttributesDeclaration);
+	//LazyPrintf.PushParam(*MaterialAttributesUtilities);
+
+	LazyPrintf.PushParam(DeclarationsCode);
+	LazyPrintf.PushParam(TEXT(""));
 
 	// Stores the shared shader results member declarations
 	// PixelMembersDeclaration should be the same for all variations, but might change in the future. There are cases where work is shared
@@ -276,7 +280,7 @@ static FString GenerateMaterialTemplateHLSL(EShaderPlatform ShaderPlatform,
 
 	LazyPrintf.PushParam(*CustomInterpolatorAssignments);
 
-	LazyPrintf.PushParam(*MaterialAttributesDefault);
+	//LazyPrintf.PushParam(*MaterialAttributesDefault);
 
 	//if (bEnableExecutionFlow)
 	{
@@ -310,7 +314,12 @@ static FString GenerateMaterialTemplateHLSL(EShaderPlatform ShaderPlatform,
 			// Special case MP_SubsurfaceColor as the actual property is a combination of the color and the profile but we don't want to expose the profile
 			const FString PropertyName = FMaterialAttributeDefinitionMap::GetAttributeName(Property);
 
-			if (PropertyIndex == MP_SubsurfaceColor)
+			if (PropertyIndex == MP_FrontMaterial)
+			{
+				// TODO - Strata
+				EvaluateMaterialAttributesCode += FString::Printf("    PixelMaterialInputs.%s = GetInitialisedStrataData();" LINE_TERMINATOR, *PropertyName);
+			}
+			else if (PropertyIndex == MP_SubsurfaceColor)
 			{
 				// TODO - properly handle subsurface profile
 				EvaluateMaterialAttributesCode += FString::Printf("    PixelMaterialInputs.Subsurface = float4(MaterialAttributes.%s, 0.0f);" LINE_TERMINATOR, *PropertyName);
@@ -646,6 +655,12 @@ static void GetMaterialEnvironment(EShaderPlatform InPlatform,
 	OutEnvironment.SetDefine(TEXT("TEXTURE_SAMPLE_DEBUG"), false);// IsDebugTextureSampleEnabled() ? TEXT("1") : TEXT("0"));
 }
 
+class FStringBuilderMemstack : public TStringBuilderBase<TCHAR>
+{
+public:
+	FStringBuilderMemstack(FMemStackBase& Allocator, int32 InSize) : TStringBuilderBase<TCHAR>((TCHAR*)Allocator.Alloc(sizeof(TCHAR)* InSize, alignof(TCHAR)), InSize) {}
+};
+
 bool MaterialEmitHLSL(const FMaterialCompileTargetParameters& InCompilerTarget,
 	const FMaterial& InMaterial,
 	const FStaticParameterSet& InStaticParameters,
@@ -674,18 +689,19 @@ bool MaterialEmitHLSL(const FMaterialCompileTargetParameters& InCompilerTarget,
 	FMemStackBase Allocator;
 	FMemMark MemMark(Allocator);
 
-	TStringBuilder<32 * 1024> Code;
+	FStringBuilderMemstack Declarations(Allocator, 64 * 1024);
+	FStringBuilderMemstack Code(Allocator, 1024 * 1024);
 
 	UE::HLSLTree::FEmitContext EmitContext(Allocator);
 	EmitContext.Material = &InMaterial;
 	EmitContext.StaticParameters = &InStaticParameters;
 	EmitContext.MaterialCompilationOutput = &OutCompilationOutput;
-	InTree.EmitHLSL(EmitContext, Code);
+	InTree.EmitHLSL(EmitContext, Declarations, Code);
 
 	//InOutMaterial.CompileErrors = MoveTemp(CompileErrors);
 	//InOutMaterial.ErrorExpressions = MoveTemp(ErrorExpressions);
 
-	const FString MaterialTemplateSource = GenerateMaterialTemplateHLSL(InCompilerTarget.ShaderPlatform, InMaterial, EmitContext, Code.ToString());
+	const FString MaterialTemplateSource = GenerateMaterialTemplateHLSL(InCompilerTarget.ShaderPlatform, InMaterial, EmitContext, Declarations.ToString(), Code.ToString());
 
 	OutMaterialEnvironment = new FSharedShaderCompilerEnvironment();
 	OutMaterialEnvironment->TargetPlatform = InCompilerTarget.TargetPlatform;
diff --git a/Engine/Source/Runtime/Engine/Private/Materials/MaterialHLSLGenerator.cpp b/Engine/Source/Runtime/Engine/Private/Materials/MaterialHLSLGenerator.cpp
--- a/Engine/Source/Runtime/Engine/Private/Materials/MaterialHLSLGenerator.cpp
+++ b/Engine/Source/Runtime/Engine/Private/Materials/MaterialHLSLGenerator.cpp
@@ -16,15 +16,73 @@
 #include "HLSLTree/HLSLTreeCommon.h"
 #include "Containers/LazyPrintf.h"
 
+static UE::Shader::EValueType GetShaderType(EMaterialValueType MaterialType)
+{
+	switch (MaterialType)
+	{
+	case MCT_Float1: return UE::Shader::EValueType::Float1;
+	case MCT_Float2: return UE::Shader::EValueType::Float2;
+	case MCT_Float3: return UE::Shader::EValueType::Float3;
+	case MCT_Float4: return UE::Shader::EValueType::Float4;
+	case MCT_Float: return UE::Shader::EValueType::Float1;
+	case MCT_StaticBool: return UE::Shader::EValueType::Bool1;
+	case MCT_MaterialAttributes: return UE::Shader::EValueType::Struct;
+	case MCT_ShadingModel: return UE::Shader::EValueType::Int1;
+	case MCT_LWCScalar: return UE::Shader::EValueType::Double1;
+	case MCT_LWCVector2: return UE::Shader::EValueType::Double2;
+	case MCT_LWCVector3: return UE::Shader::EValueType::Double3;
+	case MCT_LWCVector4: return UE::Shader::EValueType::Double4;
+	default:return UE::Shader::EValueType::Void;
+	}
+}
+
 FMaterialHLSLGenerator::FMaterialHLSLGenerator(UMaterial* InTargetMaterial, const FMaterialCompileTargetParameters& InCompileTarget, UE::HLSLTree::FTree& InOutTree)
 	: CompileTarget(InCompileTarget)
 	, TargetMaterial(InTargetMaterial)
 	, HLSLTree(&InOutTree)
 	, bGeneratedResult(false)
 {
+	const EMaterialShadingModel DefaultShadingModel = InTargetMaterial->GetShadingModels().GetFirstShadingModel();
+
 	FFunctionCallEntry* RootFunctionEntry = new(InOutTree.GetAllocator()) FFunctionCallEntry();
 	FunctionCallStack.Add(RootFunctionEntry);
 
+	TArray<UE::HLSLTree::FStructFieldInitializer, TInlineAllocator<MP_MAX>> MaterialAttributeFields;
+
+	const TArray<FGuid>& OrderedVisibleAttributes = FMaterialAttributeDefinitionMap::GetOrderedVisibleAttributeList();
+	for (const FGuid& AttributeID : OrderedVisibleAttributes)
+	{
+		const FString& PropertyName = FMaterialAttributeDefinitionMap::GetAttributeName(AttributeID);
+		const EMaterialValueType PropertyType = FMaterialAttributeDefinitionMap::GetValueType(AttributeID);
+		const UE::Shader::EValueType ValueType = GetShaderType(PropertyType);
+		
+		if (ValueType != UE::Shader::EValueType::Void &&
+			ValueType != UE::Shader::EValueType::Struct)
+		{
+			MaterialAttributeFields.Emplace(PropertyName, ValueType);
+
+			if (PropertyType == MCT_ShadingModel)
+			{
+				check(ValueType == UE::Shader::EValueType::Int1);
+				MaterialAttributesDefaultValue.Component.Add((int32)DefaultShadingModel);
+			}
+			else
+			{
+				const UE::Shader::FValue DefaultValue = UE::Shader::Cast(FMaterialAttributeDefinitionMap::GetDefaultValue(AttributeID), ValueType);
+				for (int32 i = 0; i < DefaultValue.NumComponents; ++i)
+				{
+					MaterialAttributesDefaultValue.Component.Add(DefaultValue.Component[i]);
+				}
+			}
+		}
+	}
+
+	UE::HLSLTree::FStructTypeInitializer MaterialAttributesInitializer;
+	MaterialAttributesInitializer.Name = TEXT("FMaterialAttributes");
+	MaterialAttributesInitializer.Fields = MaterialAttributeFields;
+	MaterialAttributesType = InOutTree.NewStructType(MaterialAttributesInitializer);
+
+	MaterialAttributesDefaultValue.Type = MaterialAttributesType;
 }
 
 void FMaterialHLSLGenerator::AcquireErrors(TArray<FString>& OutCompileErrors, TArray<UMaterialExpression*>& OutErrorExpressions)
@@ -124,6 +182,7 @@ static UE::HLSLTree::FExpression* CompileMaterialInput(FMaterialHLSLGenerator& G
 	UMaterial* Material,
 	UE::HLSLTree::FExpression* AttributesExpression)
 {
+#if 0
 	UE::HLSLTree::FExpression* Expression = nullptr;
 	if (Material->IsPropertyActive(InputProperty))
 	{
@@ -170,7 +229,7 @@ static UE::HLSLTree::FExpression* CompileMaterialInput(FMaterialHLSLGenerator& G
 		SetAttributeExpression->ValueExpression = Expression;
 		return SetAttributeExpression;
 	}
-
+#endif
 	return AttributesExpression;
 }
 
@@ -191,13 +250,13 @@ bool FMaterialHLSLGenerator::GenerateResult(UE::HLSLTree::FScope& Scope)
 				FMaterialInputDescription InputDescription;
 				if (TargetMaterial->GetExpressionInputDescription(MP_MaterialAttributes, InputDescription))
 				{
-					check(InputDescription.Type == UE::Shader::EValueType::MaterialAttributes);
+					check(InputDescription.Type == UE::Shader::EValueType::Struct);
 					AttributesExpression = InputDescription.Input->AcquireHLSLExpression(*this, Scope);
 				}
 			}
 			else
 			{
-				AttributesExpression = HLSLTree->NewExpression<UE::HLSLTree::FExpressionDefaultMaterialAttributes>(Scope);
+				AttributesExpression = HLSLTree->NewExpression<UE::HLSLTree::FExpressionConstant>(Scope, MaterialAttributesDefaultValue);
 				for (int32 PropertyIndex = 0; PropertyIndex < MP_MAX; ++PropertyIndex)
 				{
 					const EMaterialProperty Property = (EMaterialProperty)PropertyIndex;
@@ -208,6 +267,7 @@ bool FMaterialHLSLGenerator::GenerateResult(UE::HLSLTree::FScope& Scope)
 			if (AttributesExpression)
 			{
 				UE::HLSLTree::FStatementReturn* ReturnStatement = HLSLTree->NewStatement<UE::HLSLTree::FStatementReturn>(Scope);
+				ReturnStatement->Type = MaterialAttributesType;
 				ReturnStatement->Expression = AttributesExpression;
 				HLSLTree->SetResult(*ReturnStatement);
 				bResult = true;
diff --git a/Engine/Source/Runtime/Engine/Private/Shader/ShaderValue.cpp b/Engine/Source/Runtime/Engine/Private/Shader/ShaderValue.cpp
--- a/Engine/Source/Runtime/Engine/Private/Shader/ShaderValue.cpp
+++ b/Engine/Source/Runtime/Engine/Private/Shader/ShaderValue.cpp
@@ -252,11 +252,22 @@ uint32 UE::Shader::GetComponentTypeSizeInBytes(EValueComponentType Type)
 	case EValueComponentType::Double: return sizeof(double);
 	case EValueComponentType::Int: return sizeof(int32);
 	case EValueComponentType::Bool: return 1u;
-	case EValueComponentType::MaterialAttributes: return 0u;
 	default: checkNoEntry(); return 0u;
 	}
 }
 
+const TCHAR* UE::Shader::FValueComponent::ToString(EValueComponentType Type, FStringBuilderBase& OutString) const
+{
+	switch (Type)
+	{
+	case EValueComponentType::Int: OutString.Appendf(TEXT("%d"), Int); break;
+	case EValueComponentType::Bool: OutString.Append(AsBool() ? TEXT("true") : TEXT("false")); break;
+	case EValueComponentType::Float: OutString.Appendf(TEXT("%#.9gf"), Float); break;
+	default: checkNoEntry(); break; // TODO - double
+	}
+	return OutString.ToString();
+}
+
 const TCHAR* UE::Shader::FValue::ToString(EValueStringFormat Format, FStringBuilderBase& OutString) const
 {
 	if (Format == EValueStringFormat::HLSL && ComponentType == EValueComponentType::Double)
@@ -355,21 +366,20 @@ UE::Shader::FValueTypeDescription UE::Shader::GetValueTypeDescription(EValueType
 	case EValueType::Bool2: return FValueTypeDescription(TEXT("bool2"), EValueComponentType::Bool, 2);
 	case EValueType::Bool3: return FValueTypeDescription(TEXT("bool3"), EValueComponentType::Bool, 3);
 	case EValueType::Bool4: return FValueTypeDescription(TEXT("bool4"), EValueComponentType::Bool, 4);
-	case EValueType::MaterialAttributes: return FValueTypeDescription(TEXT("FMaterialAttributes"), EValueComponentType::MaterialAttributes, 0);
-	default: checkNoEntry(); return FValueTypeDescription();
+	case EValueType::Struct: return FValueTypeDescription(TEXT("struct"), EValueComponentType::Void, 0);
+	default: checkNoEntry(); return FValueTypeDescription(TEXT("<INVALID>"), EValueComponentType::Void, 0);
 	}
 }
 
 UE::Shader::EValueType UE::Shader::MakeValueType(EValueComponentType ComponentType, int32 NumComponents)
 {
-	switch (ComponentType)
+	if (NumComponents == 0)
 	{
-	case EValueComponentType::Void:
-		check(NumComponents == 0);
 		return EValueType::Void;
-	case EValueComponentType::MaterialAttributes:
-		check(NumComponents == 0);
-		return EValueType::MaterialAttributes;
+	}
+
+	switch (ComponentType)
+	{
 	case EValueComponentType::Float:
 		switch (NumComponents)
 		{
diff --git a/Engine/Source/Runtime/Engine/Public/HLSLTree/HLSLTree.h b/Engine/Source/Runtime/Engine/Public/HLSLTree/HLSLTree.h
--- a/Engine/Source/Runtime/Engine/Public/HLSLTree/HLSLTree.h
+++ b/Engine/Source/Runtime/Engine/Public/HLSLTree/HLSLTree.h
@@ -5,6 +5,7 @@
 #include "Engine/EngineTypes.h"
 #include "Containers/ArrayView.h"
 #include "Containers/StringView.h"
+#include "Containers/BitArray.h"
 #include "Containers/List.h"
 #include "Misc/StringBuilder.h"
 #include "Misc/MemStack.h"
@@ -35,6 +36,50 @@ class FScope;
 
 static constexpr int32 MaxNumPreviousScopes = 2;
 
+struct FType
+{
+	FType() : StructType(nullptr), ValueType(Shader::EValueType::Void) {}
+	FType(Shader::EValueType InValueType) : StructType(nullptr), ValueType(InValueType) {}
+	FType(const FStructType* InStruct) : StructType(InStruct), ValueType(Shader::EValueType::Struct) { check(InStruct); }
+
+	const TCHAR* GetName() const;
+	bool IsStruct() const { return (bool)StructType; }
+	int32 GetNumComponents() const;
+	bool Merge(const FType& OtherType);
+
+	inline operator Shader::EValueType() const { return ValueType; }
+	inline operator bool() const { return ValueType != Shader::EValueType::Void; }
+	inline bool operator!() const { return ValueType == Shader::EValueType::Void; }
+
+	const FStructType* StructType;
+	Shader::EValueType ValueType;
+};
+
+inline bool operator==(const FType& Lhs, const FType& Rhs)
+{
+	if (Lhs.ValueType != Rhs.ValueType) return false;
+	if (Lhs.ValueType == Shader::EValueType::Struct && Lhs.StructType != Rhs.StructType) return false;
+	return true;
+}
+inline bool operator!=(const FType& Lhs, const FType& Rhs)
+{
+	return !operator==(Lhs, Rhs);
+}
+
+struct FConstantValue
+{
+	FConstantValue() = default;
+	FConstantValue(const Shader::FValue& InValue);
+
+	FType Type;
+
+	/**
+	 * For basic types, this will include 1-4 components
+	 * For struct types, will match the flattened list of struct component types
+	 */
+	TArray<Shader::FValueComponent, TInlineAllocator<16>> Component;
+};
+
 struct FError
 {
 	const FError* Next;
@@ -50,14 +95,14 @@ public:
 
 	int32 Num() const { return NumErrors; }
 
-	bool AddError(const FNode* InNode, FStringView InError);
+	void AddError(const FNode* InNode, FStringView InError);
 
 	template<typename FormatType, typename... Types>
-	bool AddErrorf(const FNode* InNode, const FormatType& Format, Types... Args)
+	void AddErrorf(const FNode* InNode, const FormatType& Format, Types... Args)
 	{
 		TStringBuilder<2048> String;
 		String.Appendf(Format, Forward<Types>(Args)...);
-		return AddError(InNode, FStringView(String.GetData(), String.Len()));
+		AddError(InNode, FStringView(String.ToString(), String.Len()));
 	}
 
 private:
@@ -97,7 +142,7 @@ public:
 	/** Get a unique local variable name */
 	const TCHAR* AcquireLocalDeclarationCode();
 
-	const TCHAR* CastShaderValue(const FNode* Node, const TCHAR* Code, Shader::EValueType SourceType, Shader::EValueType DestType, ECastFlags Flags);
+	const TCHAR* CastShaderValue(const FNode* Node, const TCHAR* Code, const FType& SourceType, const FType& DestType, ECastFlags Flags);
 
 	void AddPreshader(Shader::EValueType Type, const Shader::FPreshaderData& Preshader, FStringBuilderBase& OutCode);
 
@@ -149,6 +194,61 @@ public:
 	FNode* NextNode = nullptr;
 };
 
+struct FStructField
+{
+	const TCHAR* Name;
+	FType Type;
+	int32 ComponentIndex;
+};
+
+struct FStructFieldRef
+{
+	FStructFieldRef() = default;
+	FStructFieldRef(const FType& InType, int32 InIndex, int32 InNum) : Type(InType), ComponentIndex(InIndex), ComponentNum(InNum) {}
+
+	FType Type;
+	int32 ComponentIndex = INDEX_NONE;
+	int32 ComponentNum = 0;
+
+	inline operator bool() const { return ComponentNum > 0; }
+	inline bool operator!() const { return ComponentNum <= 0; }
+};
+
+class FStructType : public FNode
+{
+public:
+	FStructType* NextType;
+	const TCHAR* Name;
+	TArrayView<const FStructField> Fields;
+
+	/**
+	 * Most code working with HLSLTree views struct types as a flat list of components
+	 * Fields with basic types are represented directly. Fields with struct types are recursively flattened into this list
+	 */
+	TArrayView<const Shader::EValueComponentType> ComponentTypes;
+
+	virtual ENodeVisitResult Visit(FNodeVisitor& Visitor) override { return ENodeVisitResult::VisitDependentNodes; }
+
+	FStructFieldRef FindFieldByName(const TCHAR* InName) const;
+
+	void WriteHLSL(FStringBuilderBase& OutWriter) const;
+};
+
+struct FStructFieldInitializer
+{
+	FStructFieldInitializer() = default;
+	FStructFieldInitializer(const FStringView& InName, const FType& InType) : Name(InName), Type(InType) {}
+
+	FStringView Name;
+	FType Type;
+};
+
+struct FStructTypeInitializer
+{
+	FStringView Name;
+	TArrayView<const FStructFieldInitializer> Fields;
+};
+
 /**
  * Represents an HLSL statement.  This is a piece of code that doesn't evaluate to any value, but instead should be executed sequentially, and likely has side-effects.
  * Examples include assigning a value, or various control flow structures (if, for, while, etc)
@@ -169,6 +269,57 @@ public:
 	bool bEmitHLSL = false;
 };
 
+class FRequestedType
+{
+public:
+	FRequestedType() = default;
+	FRequestedType(int32 NumComponents, bool bDefaultRequest = true) : RequestedComponents(bDefaultRequest, NumComponents) {}
+	FRequestedType(const FType& InType, bool bDefaultRequest = true) : StructType(InType.StructType), RequestedComponents(bDefaultRequest, InType.GetNumComponents()) {}
+
+	bool IsStruct() const { return (bool)StructType; }
+	const FStructType* GetStructType() const { return StructType; }
+	int32 GetNumComponents() const { return RequestedComponents.Num(); }
+	int32 GetRequestedNumComponents() const { const int32 LastIndex = RequestedComponents.FindLast(true); return LastIndex + 1; }
+	bool IsComponentRequested(int32 Index) const { return RequestedComponents.IsValidIndex(Index) ? (bool)RequestedComponents[Index] : false; }
+	bool IsVoid() const { return RequestedComponents.Find(true) == INDEX_NONE; }
+
+	void Reset()
+	{
+		StructType = nullptr;
+		RequestedComponents.Reset();
+	}
+
+	/** Returns a type that represents access of FieldRef from a struct on given type, using 'this' */
+	FRequestedType MakeFieldAccess(const FStructType* InStructType, const FStructFieldRef& FieldRef) const
+	{
+		FRequestedType Result(InStructType, false);
+		Result.RequestedComponents.SetRangeFromRange(FieldRef.ComponentIndex, FieldRef.ComponentNum, RequestedComponents, 0);
+		return Result;
+	}
+
+	/** Returns copy of 'this', with request for given field removed */
+	FRequestedType CopyWithFieldRemoved(const FStructFieldRef& FieldRef) const
+	{
+		FRequestedType Result(*this);
+		Result.RequestedComponents.SetRange(FieldRef.ComponentIndex, FieldRef.ComponentNum, false);
+		return Result;
+	}
+
+	/** Returns a type that represents access of the given field from 'this' */
+	FRequestedType GetField(const FStructFieldRef& FieldRef) const
+	{
+		FRequestedType Result(FieldRef.Type);
+		Result.RequestedComponents.SetRangeFromRange(0, FieldRef.ComponentNum, RequestedComponents, FieldRef.ComponentIndex);
+		return Result;
+	}
+
+	/** The struct type we're requesting, or nullptr if we're requesting a basic type */
+	const FStructType* StructType = nullptr;
+
+	/** 1 bit per component, a value of 'true' means the specified component is requsted */
+	TBitArray<> RequestedComponents;
+};
+
 /**
  * Represents an HLSL expression.  This is a piece of code that evaluates to a value, but has no side effects.
  * Unlike statements, expressions are not expected to execute in any particular order.  They may be cached (or not) in generated code, without the underlying implementation needing to care.
@@ -179,52 +330,52 @@ public:
 class FExpression : public FNode
 {
 public:
-	Shader::EValueType GetValueType() const { return ValueType; }
+	const FType& GetType() const { return Type; }
 
 	virtual ENodeVisitResult Visit(FNodeVisitor& Visitor) override;
 
-	friend Shader::EValueType RequestExpressionType(FUpdateTypeContext& Context, FExpression* InExpression, int8 InRequestedNumComponents);
+	friend FType RequestExpressionType(FUpdateTypeContext& Context, FExpression* InExpression, const FRequestedType& RequestedType);
 	friend EExpressionEvaluationType PrepareExpressionValue(FEmitContext& Context, FExpression* InExpression);
 
 	const TCHAR* GetValueShader(FEmitContext& Context);
-	const TCHAR* GetValueShader(FEmitContext& Context, Shader::EValueType Type);
+	const TCHAR* GetValueShader(FEmitContext& Context, const FType& InType);
 	void GetValuePreshader(FEmitContext& Context, Shader::FPreshaderData& OutPreshader);
 	Shader::FValue GetValueConstant(FEmitContext& Context);
 
 	FScope* ParentScope = nullptr;
 
 protected:
-	virtual bool UpdateType(FUpdateTypeContext& Context, int8 InRequestedNumComponents) = 0;
-	virtual bool PrepareValue(FEmitContext& Context) = 0;
+	virtual void UpdateType(FUpdateTypeContext& Context, const FRequestedType& RequestedType) = 0;
+	virtual void PrepareValue(FEmitContext& Context) = 0;
 
-	bool SetType(FUpdateTypeContext& Context, Shader::EValueType InType) { ValueType = InType; return true; }
+	void SetType(FUpdateTypeContext& Context, const FType& InType);
 
-	bool InternalSetValueShader(FEmitContext& Context, const TCHAR* InCode, int32 InLength, bool bInline);
+	void InternalSetValueShader(FEmitContext& Context, const TCHAR* InCode, int32 InLength, bool bInline);
 
-	bool SetValuePreshader(FEmitContext& Context, Shader::FPreshaderData& InPreshader);
-	bool SetValueConstant(FEmitContext& Context, const Shader::FValue& InValue);
-	bool SetValueForward(FEmitContext& Context, FExpression* Source);
+	void SetValuePreshader(FEmitContext& Context, Shader::FPreshaderData& InPreshader);
+	void SetValueConstant(FEmitContext& Context, const Shader::FValue& InValue);
+	void SetValueForward(FEmitContext& Context, FExpression* Source);
 
-	bool SetValuePreshader(FEmitContext& Context, EExpressionEvaluationType InEvaluationType, Shader::FPreshaderData& InPreshader);
+	void SetValuePreshader(FEmitContext& Context, EExpressionEvaluationType InEvaluationType, Shader::FPreshaderData& InPreshader);
 
 	template<typename FormatType, typename... Types>
-	bool SetValueShaderf(FEmitContext& Context, const FormatType& Format, Types... Args)
+	void SetValueShaderf(FEmitContext& Context, const FormatType& Format, Types... Args)
 	{
 		TStringBuilder<2048> String;
 		String.Appendf(Format, Forward<Types>(Args)...);
-		return InternalSetValueShader(Context, String.GetData(), String.Len(), false);
+		InternalSetValueShader(Context, String.ToString(), String.Len(), false);
 	}
 
 	template<typename FormatType, typename... Types>
-	bool SetValueInlineShaderf(FEmitContext& Context, const FormatType& Format, Types... Args)
+	void SetValueInlineShaderf(FEmitContext& Context, const FormatType& Format, Types... Args)
 	{
 		TStringBuilder<2048> String;
 		String.Appendf(Format, Forward<Types>(Args)...);
-		return InternalSetValueShader(Context, String.GetData(), String.Len(), true);
+		InternalSetValueShader(Context, String.ToString(), String.Len(), true);
 	}
 
-	inline bool SetValueShader(FEmitContext& Context, const FStringBuilderBase& String) { return InternalSetValueShader(Context, String.GetData(), String.Len(), false); }
-	inline bool SetValueInlineShader(FEmitContext& Context, const FStringBuilderBase& String) { return InternalSetValueShader(Context, String.GetData(), String.Len(), true); }
+	inline void SetValueShader(FEmitContext& Context, const FStringBuilderBase& String) { InternalSetValueShader(Context, String.ToString(), String.Len(), false); }
+	inline void SetValueInlineShader(FEmitContext& Context, const FStringBuilderBase& String) { InternalSetValueShader(Context, String.ToString(), String.Len(), true); }
 
 private:
 	const TCHAR* LocalVariableName = nullptr;
@@ -232,8 +383,8 @@ private:
 	const Shader::FPreshaderData* Preshader = nullptr;
 	Shader::FValue ConstantValue;
 	EExpressionEvaluationType EvaluationType = EExpressionEvaluationType::None;
-	Shader::EValueType ValueType = Shader::EValueType::Void;
-	int8 RequestedNumComponents = 0;
+	FType Type;
+	FRequestedType CurrentRequestedType;
 	bool bReentryFlag = false;
 };
 
@@ -246,8 +397,8 @@ private:
 class FExpressionLocalPHI final : public FExpression
 {
 public:
-	virtual bool UpdateType(FUpdateTypeContext& Context, int8 InRequestedNumComponents) override;
-	virtual bool PrepareValue(FEmitContext& Context) override;
+	virtual void UpdateType(FUpdateTypeContext& Context, const FRequestedType& RequestedType) override;
+	virtual void PrepareValue(FEmitContext& Context) override;
 
 	FName LocalName;
 	FScope* Scopes[MaxNumPreviousScopes];
@@ -351,7 +502,7 @@ private:
 	{
 		TStringBuilder<2048> String;
 		String.Appendf(Format, Forward<Types>(Args)...);
-		InternalEmitCode(Context, List, NestedScope, String.GetData(), String.Len());
+		InternalEmitCode(Context, List, NestedScope, String.ToString(), String.Len());
 	}
 
 	FScope* ParentScope = nullptr;
@@ -376,7 +527,9 @@ public:
 
 	FMemStackBase& GetAllocator() { return *Allocator; }
 
-	bool EmitHLSL(FEmitContext& Context, FStringBuilderBase& Writer) const;
+	bool EmitHLSL(FEmitContext& Context,
+		FStringBuilderBase& OutDeclarations,
+		FStringBuilderBase& OutCode) const;
 
 	FScope& GetRootScope() const { return *RootScope; }
 
@@ -407,6 +560,9 @@ public:
 
 	FScope* NewScope(FScope& Scope);
 	FTextureParameterDeclaration* NewTextureParameterDeclaration(const FName& Name, const FTextureDescription& DefaultValue);
+
+	const FStructType* NewStructType(const FStructTypeInitializer& Initializer);
+
 	void SetResult(FStatement& InResult);
 
 private:
@@ -424,13 +580,14 @@ private:
 
 	FMemStackBase* Allocator = nullptr;
 	FNode* Nodes = nullptr;
+	FStructType* StructTypes = nullptr;
 	FScope* RootScope = nullptr;
 	FStatement* ResultStatement = nullptr;
 };
 
-Shader::EValueType RequestExpressionType(FUpdateTypeContext& Context, FExpression* InExpression, int8 InRequestedNumComponents); // friend of FExpression
-EExpressionEvaluationType PrepareExpressionValue(FEmitContext& Context, FExpression* InExpression); // friend of FExpression
-void RequestScopeTypes(FUpdateTypeContext& Context, const FScope* InScope); // friend of FScope
+//Shader::EValueType RequestExpressionType(FUpdateTypeContext& Context, FExpression* InExpression, int8 InRequestedNumComponents); // friend of FExpression
+//EExpressionEvaluationType PrepareExpressionValue(FEmitContext& Context, FExpression* InExpression); // friend of FExpression
+//void RequestScopeTypes(FUpdateTypeContext& Context, const FScope* InScope); // friend of FScope
 
 } // namespace HLSLTree
 } // namespace UE
diff --git a/Engine/Source/Runtime/Engine/Public/HLSLTree/HLSLTreeCommon.h b/Engine/Source/Runtime/Engine/Public/HLSLTree/HLSLTreeCommon.h
--- a/Engine/Source/Runtime/Engine/Public/HLSLTree/HLSLTreeCommon.h
+++ b/Engine/Source/Runtime/Engine/Public/HLSLTree/HLSLTreeCommon.h
@@ -39,14 +39,14 @@ FBinaryOpDescription GetBinaryOpDesription(EBinaryOp Op);
 class FExpressionConstant : public FExpression
 {
 public:
-	explicit FExpressionConstant(const Shader::FValue& InValue)
+	explicit FExpressionConstant(const FConstantValue& InValue)
 		: Value(InValue)
 	{}
 
-	Shader::FValue Value;
+	FConstantValue Value;
 
-	virtual bool UpdateType(FUpdateTypeContext& Context, int8 InRequestedNumComponents) override { return SetType(Context, Shader::MakeValueTypeWithRequestedNumComponents(Value.GetType(), InRequestedNumComponents)); }
-	virtual bool PrepareValue(FEmitContext& Context) override;
+	virtual void UpdateType(FUpdateTypeContext& Context, const FRequestedType& RequestedType) override;
+	virtual void PrepareValue(FEmitContext& Context) override;
 };
 
 class FExpressionMaterialParameter : public FExpression
@@ -60,8 +60,8 @@ public:
 	Shader::FValue DefaultValue;
 	EMaterialParameterType ParameterType;
 
-	virtual bool UpdateType(FUpdateTypeContext& Context, int8 InRequestedNumComponents) override;
-	virtual bool PrepareValue(FEmitContext& Context) override;
+	virtual void UpdateType(FUpdateTypeContext& Context, const FRequestedType& RequestedType) override;
+	virtual void PrepareValue(FEmitContext& Context) override;
 };
 
 enum class EExternalInputType
@@ -92,8 +92,8 @@ public:
 
 	EExternalInputType InputType;
 
-	virtual bool UpdateType(FUpdateTypeContext& Context, int8 InRequestedNumComponents) override { return SetType(Context, GetInputExpressionType(InputType)); }
-	virtual bool PrepareValue(FEmitContext& Context) override;
+	virtual void UpdateType(FUpdateTypeContext& Context, const FRequestedType& RequestedType) override { SetType(Context, GetInputExpressionType(InputType)); }
+	virtual void PrepareValue(FEmitContext& Context) override;
 };
 
 class FExpressionTextureSample : public FExpression
@@ -122,41 +122,56 @@ public:
 		return Result;
 	}
 
-	virtual bool UpdateType(FUpdateTypeContext& Context, int8 InRequestedNumComponents) override;
-	virtual bool PrepareValue(FEmitContext& Context) override;
+	virtual void UpdateType(FUpdateTypeContext& Context, const FRequestedType& RequestedType) override;
+	virtual void PrepareValue(FEmitContext& Context) override;
 };
 
-class FExpressionDefaultMaterialAttributes : public FExpression
+class FExpressionGetStructField : public FExpression
 {
 public:
-	FExpressionDefaultMaterialAttributes() {}
+	FExpressionGetStructField() {}
 
-	virtual bool UpdateType(FUpdateTypeContext& Context, int8 InRequestedNumComponents) override { return SetType(Context, Shader::EValueType::MaterialAttributes); }
-	virtual bool PrepareValue(FEmitContext& Context) override;
+	const FStructType* StructType;
+	const TCHAR* FieldName;
+	FExpression* StructExpression;
+
+	virtual ENodeVisitResult Visit(FNodeVisitor& Visitor) override
+	{
+		const ENodeVisitResult Result = FExpression::Visit(Visitor);
+		if (ShouldVisitDependentNodes(Result))
+		{
+			Visitor.VisitNode(StructExpression);
+		}
+		return Result;
+	}
+
+	virtual void UpdateType(FUpdateTypeContext& Context, const FRequestedType& RequestedType) override;
+	virtual void PrepareValue(FEmitContext& Context) override;
 };
 
-class FExpressionSetMaterialAttribute : public FExpression
+class FExpressionSetStructField : public FExpression
 {
 public:
-	FExpressionSetMaterialAttribute() {}
-
-	FGuid AttributeID;
-	FExpression* AttributesExpression;
-	FExpression* ValueExpression;
+	FExpressionSetStructField() {}
 
+	const FStructType* StructType;
+	const TCHAR* FieldName;
+	FExpression* StructExpression;
+	FExpression* FieldExpression;
+	
 	virtual ENodeVisitResult Visit(FNodeVisitor& Visitor) override
 	{
 		const ENodeVisitResult Result = FExpression::Visit(Visitor);
 		if (ShouldVisitDependentNodes(Result))
 		{
-			Visitor.VisitNode(AttributesExpression);
-			Visitor.VisitNode(ValueExpression);
+			Visitor.VisitNode(StructExpression);
+			Visitor.VisitNode(FieldExpression);
 		}
 		return Result;
 	}
 
-	virtual bool UpdateType(FUpdateTypeContext& Context, int8 InRequestedNumComponents) override;
-	virtual bool PrepareValue(FEmitContext& Context) override;
+	virtual void UpdateType(FUpdateTypeContext& Context, const FRequestedType& RequestedType) override;
+	virtual void PrepareValue(FEmitContext& Context) override;
 };
 
 class FExpressionSelect : public FExpression
@@ -184,8 +199,8 @@ public:
 		return Result;
 	}
 
-	virtual bool UpdateType(FUpdateTypeContext& Context, int8 InRequestedNumComponents) override;
-	virtual bool PrepareValue(FEmitContext& Context) override;
+	virtual void UpdateType(FUpdateTypeContext& Context, const FRequestedType& RequestedType) override;
+	virtual void PrepareValue(FEmitContext& Context) override;
 };
 
 class FExpressionBinaryOp : public FExpression
@@ -212,8 +227,8 @@ public:
 		return Result;
 	}
 
-	virtual bool UpdateType(FUpdateTypeContext& Context, int8 InRequestedNumComponents) override;
-	virtual bool PrepareValue(FEmitContext& Context) override;
+	virtual void UpdateType(FUpdateTypeContext& Context, const FRequestedType& RequestedType) override;
+	virtual void PrepareValue(FEmitContext& Context) override;
 };
 
 struct FSwizzleParameters
@@ -247,8 +262,8 @@ public:
 		return Result;
 	}
 
-	virtual bool UpdateType(FUpdateTypeContext& Context, int8 InRequestedNumComponents) override;
-	virtual bool PrepareValue(FEmitContext& Context) override;
+	virtual void UpdateType(FUpdateTypeContext& Context, const FRequestedType& RequestedType) override;
+	virtual void PrepareValue(FEmitContext& Context) override;
 };
 
 class FExpressionAppend : public FExpression
@@ -273,21 +288,23 @@ public:
 		return Result;
 	}
 
-	virtual bool UpdateType(FUpdateTypeContext& Context, int8 InRequestedNumComponents) override;
-	virtual bool PrepareValue(FEmitContext& Context) override;
+	virtual void UpdateType(FUpdateTypeContext& Context, const FRequestedType& RequestedType) override;
+	virtual void PrepareValue(FEmitContext& Context) override;
 };
 
 class FExpressionReflectionVector : public FExpression
 {
 public:
-	virtual bool UpdateType(FUpdateTypeContext& Context, int8 InRequestedNumComponents) override { return SetType(Context, Shader::EValueType::Float3); }
-	virtual bool PrepareValue(FEmitContext& Context) override;
+	virtual void UpdateType(FUpdateTypeContext& Context, const FRequestedType& RequestedType) override { SetType(Context, Shader::EValueType::Float3); }
+	virtual void PrepareValue(FEmitContext& Context) override;
 };
 
 class FStatementReturn : public FStatement
 {
 public:
 	static constexpr bool MarkScopeLiveRecursive = true;
+
+	FType Type;
 	FExpression* Expression;
 
 	virtual ENodeVisitResult Visit(FNodeVisitor& Visitor) override
diff --git a/Engine/Source/Runtime/Engine/Public/HLSLTree/HLSLTreeTypes.h b/Engine/Source/Runtime/Engine/Public/HLSLTree/HLSLTreeTypes.h
--- a/Engine/Source/Runtime/Engine/Public/HLSLTree/HLSLTreeTypes.h
+++ b/Engine/Source/Runtime/Engine/Public/HLSLTree/HLSLTreeTypes.h
@@ -2,6 +2,7 @@
 #pragma once
 
 #include "Misc/EnumClassFlags.h"
+#include "Containers/BitArray.h"
 #include "Shader/ShaderTypes.h"
 
 class UTexture;
@@ -12,12 +13,12 @@ namespace HLSLTree
 {
 
 class FNode;
+class FStructType;
 class FExpression;
 class FTextureParameterDeclaration;
 class FStatement;
 class FScope;
 class FTree;
-class FCodeWriter;
 
 /**
  * Describes how a given expression needs to be evaluated */
diff --git a/Engine/Source/Runtime/Engine/Public/MaterialHLSLGenerator.h b/Engine/Source/Runtime/Engine/Public/MaterialHLSLGenerator.h
--- a/Engine/Source/Runtime/Engine/Public/MaterialHLSLGenerator.h
+++ b/Engine/Source/Runtime/Engine/Public/MaterialHLSLGenerator.h
@@ -132,6 +132,9 @@ public:
 		return Data;
 	}
 
+	const UE::HLSLTree::FStructType* GetMaterialAttributesType() const { return MaterialAttributesType; }
+	const UE::HLSLTree::FConstantValue& GetMaterialAttributesDefaultValue() const { return MaterialAttributesDefaultValue; }
+
 private:
 	static constexpr int32 MaxNumPreviousScopes = UE::HLSLTree::MaxNumPreviousScopes;
 	struct FExpressionKey
@@ -208,6 +211,9 @@ private:
 	const FMaterialCompileTargetParameters& CompileTarget;
 	UMaterial* TargetMaterial;
 
+	const UE::HLSLTree::FStructType* MaterialAttributesType;
+	UE::HLSLTree::FConstantValue MaterialAttributesDefaultValue;
+
 	UE::HLSLTree::FTree* HLSLTree;
 	TArray<FExpressionKey> ExpressionStack;
 	TArray<FFunctionCallEntry*> FunctionCallStack;
diff --git a/Engine/Source/Runtime/Engine/Public/MaterialShared.h b/Engine/Source/Runtime/Engine/Public/MaterialShared.h
--- a/Engine/Source/Runtime/Engine/Public/MaterialShared.h
+++ b/Engine/Source/Runtime/Engine/Public/MaterialShared.h
@@ -2964,14 +2964,14 @@ public:
 	}
 
 	/** Returns the display name of a material attribute */
-	ENGINE_API static FString GetAttributeName(EMaterialProperty Property)
+	ENGINE_API static const FString& GetAttributeName(EMaterialProperty Property)
 	{
 		FMaterialAttributeDefintion* Attribute = GMaterialPropertyAttributesMap.Find(Property);
 		return Attribute->AttributeName;
 	}
 
 	/** Returns the display name of a material attribute */
-	ENGINE_API static FString GetAttributeName(const FGuid& AttributeID)
+	ENGINE_API static const FString& GetAttributeName(const FGuid& AttributeID)
 	{
 		FMaterialAttributeDefintion* Attribute = GMaterialPropertyAttributesMap.Find(AttributeID);
 		return Attribute->AttributeName;
diff --git a/Engine/Source/Runtime/Engine/Public/Shader/ShaderTypes.h b/Engine/Source/Runtime/Engine/Public/Shader/ShaderTypes.h
--- a/Engine/Source/Runtime/Engine/Public/Shader/ShaderTypes.h
+++ b/Engine/Source/Runtime/Engine/Public/Shader/ShaderTypes.h
@@ -16,7 +16,6 @@ enum class EValueComponentType : uint8
 	Double,
 	Int,
 	Bool,
-	MaterialAttributes,
 };
 
 struct FValueTypeDescription
@@ -53,7 +52,7 @@ enum class EValueType : uint8
 	Bool3,
 	Bool4,
 
-	MaterialAttributes,
+	Struct,
 };
 
 uint32 GetComponentTypeSizeInBytes(EValueComponentType Type);
@@ -101,6 +100,8 @@ union FValueComponent
 	// 'Bool' is stored as uint8 to avoid changing on different compilers
 	bool AsBool() const { return Bool != 0u; }
 
+	const TCHAR* ToString(EValueComponentType Type, FStringBuilderBase& OutString) const;
+
 	uint64 Packed;
 	double Double;
 	float Float;
