Fix deprecation

diff --git a/Engine/Source/Runtime/Engine/Private/MorphTools.cpp b/Engine/Source/Runtime/Engine/Private/MorphTools.cpp
--- a/Engine/Source/Runtime/Engine/Private/MorphTools.cpp
+++ b/Engine/Source/Runtime/Engine/Private/MorphTools.cpp
@@ -215,7 +215,7 @@ void FFinishBuildMorphTargetData::ApplyEditorData(USkeletalMesh * SkeletalMesh)
 		
 		//Move the unused asset in the transient package and mark it pending kill
 		ToDeleteMorphTarget->Rename(nullptr, GetTransientPackage(), REN_ForceNoResetLoaders | REN_DoNotDirty | REN_DontCreateRedirectors | REN_NonTransactional);
-		ToDeleteMorphTarget->MarkPendingKill();
+		ToDeleteMorphTarget->MarkAsGarbage();
 	}
 }
 
