Update code style, wrap cpp files in namespace {}  blocks, rather than explicitly scoping all idenfiers

diff --git a/Engine/Source/Runtime/Engine/Private/HLSLTree/HLSLTree.cpp b/Engine/Source/Runtime/Engine/Private/HLSLTree/HLSLTree.cpp
--- a/Engine/Source/Runtime/Engine/Private/HLSLTree/HLSLTree.cpp
+++ b/Engine/Source/Runtime/Engine/Private/HLSLTree/HLSLTree.cpp
@@ -4,23 +4,10 @@
 #include "Misc/MemStack.h"
 #include "MaterialShared.h" // TODO - split preshader out into its own module
 
-UE::HLSLTree::EExpressionEvaluationType UE::HLSLTree::CombineEvaluationTypes(EExpressionEvaluationType Lhs, EExpressionEvaluationType Rhs)
-{
-	if (Lhs == EExpressionEvaluationType::Constant && Rhs == EExpressionEvaluationType::Constant)
-	{
-		// 2 constants make a constant
-		return EExpressionEvaluationType::Constant;
-	}
-	else if (Lhs == EExpressionEvaluationType::Shader || Rhs == EExpressionEvaluationType::Shader)
-	{
-		// If either requires shader, shader is required
-		return EExpressionEvaluationType::Shader;
-	}
-	// Any combination of constants/preshader can make a preshader
-	return EExpressionEvaluationType::Preshader;
-}
+// From MaterialUniformExpressions.cpp
+extern void WriteMaterialUniformAccess(UE::Shader::EValueComponentType ComponentType, uint32 NumComponents, uint32 UniformOffset, FStringBuilderBase& OutResult);
 
-const TCHAR* AllocateString(FMemStackBase& Allocator, const TCHAR* String, int32 Length)
+static const TCHAR* AllocateString(FMemStackBase& Allocator, const TCHAR* String, int32 Length)
 {
 	TCHAR* Result = New<TCHAR>(Allocator, Length + 1);
 	FMemory::Memcpy(Result, String, Length * sizeof(TCHAR));
@@ -28,42 +15,63 @@ const TCHAR* AllocateString(FMemStackBase& Allocator, const TCHAR* String, int32
 	return Result;
 }
 
-const TCHAR* AllocateString(FMemStackBase& Allocator, FStringView String)
+static const TCHAR* AllocateString(FMemStackBase& Allocator, FStringView String)
 {
 	return AllocateString(Allocator, String.GetData(), String.Len());
 }
 
-const TCHAR* AllocateString(FMemStackBase& Allocator, const FStringBuilderBase& StringBuilder)
+static const TCHAR* AllocateString(FMemStackBase& Allocator, const FStringBuilderBase& StringBuilder)
 {
 	return AllocateString(Allocator, StringBuilder.GetData(), StringBuilder.Len());
 }
 
 template <typename FormatType, typename... ArgTypes>
-const TCHAR* AllocateStringf(FMemStackBase& Allocator, const FormatType& Format, ArgTypes... Args)
+static const TCHAR* AllocateStringf(FMemStackBase& Allocator, const FormatType& Format, ArgTypes... Args)
 {
 	TStringBuilder<1024> String;
 	String.Appendf(Format, Forward<ArgTypes>(Args)...);
 	return AllocateString(Allocator, String);
 }
 
-FSHAHash HashString(const TCHAR* String, int32 Length)
+static FSHAHash HashString(const TCHAR* String, int32 Length)
 {
 	FSHAHash Hash;
 	FSHA1::HashBuffer(String, Length * sizeof(TCHAR), Hash.Hash);
 	return Hash;
 }
 
-FSHAHash HashString(const FStringBuilderBase& StringBuilder)
+static FSHAHash HashString(const FStringBuilderBase& StringBuilder)
 {
 	return HashString(StringBuilder.GetData(), StringBuilder.Len());
 }
 
-UE::HLSLTree::FErrors::FErrors(FMemStackBase& InAllocator)
+namespace UE
+{
+namespace HLSLTree
+{
+
+EExpressionEvaluationType CombineEvaluationTypes(EExpressionEvaluationType Lhs, EExpressionEvaluationType Rhs)
+{
+	if (Lhs == EExpressionEvaluationType::Constant && Rhs == EExpressionEvaluationType::Constant)
+	{
+		// 2 constants make a constant
+		return EExpressionEvaluationType::Constant;
+	}
+	else if (Lhs == EExpressionEvaluationType::Shader || Rhs == EExpressionEvaluationType::Shader)
+	{
+		// If either requires shader, shader is required
+		return EExpressionEvaluationType::Shader;
+	}
+	// Any combination of constants/preshader can make a preshader
+	return EExpressionEvaluationType::Preshader;
+}
+
+FErrors::FErrors(FMemStackBase& InAllocator)
 	: Allocator(&InAllocator)
 {
 }
 
-void UE::HLSLTree::FErrors::AddError(const FNode* InNode, FStringView InError)
+void FErrors::AddError(const FNode* InNode, FStringView InError)
 {
 	const int32 SizeofString = InError.Len() * sizeof(TCHAR);
 	void* Memory = Allocator->Alloc(sizeof(FError) + SizeofString, alignof(FError));
@@ -79,22 +87,22 @@ void UE::HLSLTree::FErrors::AddError(const FNode* InNode, FStringView InError)
 	ensureMsgf(false, TEXT("%s"), Error->Message);
 }
 
-UE::HLSLTree::FEmitContext::FEmitContext(FMemStackBase& InAllocator)
+FEmitContext::FEmitContext(FMemStackBase& InAllocator)
 	: Allocator(&InAllocator)
 	, Errors(InAllocator)
 {
 }
 
-UE::HLSLTree::FEmitContext::~FEmitContext()
+FEmitContext::~FEmitContext()
 {
 }
 
-const TCHAR* UE::HLSLTree::FEmitContext::AcquireLocalDeclarationCode()
+const TCHAR* FEmitContext::AcquireLocalDeclarationCode()
 {
 	return AllocateStringf(*Allocator, TEXT("Local%d"), NumExpressionLocals++);
 }
 
-const TCHAR* UE::HLSLTree::FEmitContext::CastShaderValue(const FNode* Node, const TCHAR* Code, const FType& SourceType, const FType& DestType, ECastFlags Flags)
+const TCHAR* FEmitContext::CastShaderValue(const FNode* Node, const TCHAR* Code, const FType& SourceType, const FType& DestType, ECastFlags Flags)
 {
 	if (SourceType == DestType)
 	{
@@ -226,10 +234,7 @@ const TCHAR* UE::HLSLTree::FEmitContext::CastShaderValue(const FNode* Node, cons
 	}
 }
 
-// From MaterialUniformExpressions.cpp
-extern void WriteMaterialUniformAccess(UE::Shader::EValueComponentType ComponentType, uint32 NumComponents, uint32 UniformOffset, FStringBuilderBase& OutResult);
-
-void UE::HLSLTree::FEmitContext::AddPreshader(Shader::EValueType Type, const Shader::FPreshaderData& Preshader, FStringBuilderBase& OutCode)
+void FEmitContext::AddPreshader(Shader::EValueType Type, const Shader::FPreshaderData& Preshader, FStringBuilderBase& OutCode)
 {
 	const Shader::FValueTypeDescription TypeDesc = Shader::GetValueTypeDescription(Type);
 	ensure(TypeDesc.ComponentType == Shader::EValueComponentType::Float || TypeDesc.ComponentType == Shader::EValueComponentType::Double);
@@ -262,10 +267,10 @@ void UE::HLSLTree::FEmitContext::AddPreshader(Shader::EValueType Type, const Sha
 			OutCode.Appendf(TEXT("MakeLWCVector%d("), TypeDesc.NumComponents);
 		}
 
-		WriteMaterialUniformAccess(UE::Shader::EValueComponentType::Float, TypeDesc.NumComponents, UniformPreshaderOffset, OutCode); // Tile
+		WriteMaterialUniformAccess(Shader::EValueComponentType::Float, TypeDesc.NumComponents, UniformPreshaderOffset, OutCode); // Tile
 		UniformPreshaderOffset += TypeDesc.NumComponents;
 		OutCode.Append(TEXT(","));
-		WriteMaterialUniformAccess(UE::Shader::EValueComponentType::Float, TypeDesc.NumComponents, UniformPreshaderOffset, OutCode); // Offset
+		WriteMaterialUniformAccess(Shader::EValueComponentType::Float, TypeDesc.NumComponents, UniformPreshaderOffset, OutCode); // Offset
 		UniformPreshaderOffset += TypeDesc.NumComponents;
 		OutCode.Append(TEXT(")"));
 	}
@@ -276,12 +281,12 @@ void UE::HLSLTree::FEmitContext::AddPreshader(Shader::EValueType Type, const Sha
 	}
 }
 
-void UE::HLSLTree::FEmitContext::Finalize()
+void FEmitContext::Finalize()
 {
 	MaterialCompilationOutput->UniformExpressionSet.UniformPreshaderBufferSize = (UniformPreshaderOffset + 3u) / 4u;
 }
 
-const TCHAR* UE::HLSLTree::FType::GetName() const
+const TCHAR* FType::GetName() const
 {
 	if (StructType)
 	{
@@ -293,7 +298,7 @@ const TCHAR* UE::HLSLTree::FType::GetName() const
 	return TypeDesc.Name;
 }
 
-int32 UE::HLSLTree::FType::GetNumComponents() const
+int32 FType::GetNumComponents() const
 {
 	if (StructType)
 	{
@@ -303,7 +308,7 @@ int32 UE::HLSLTree::FType::GetNumComponents() const
 	return TypeDesc.NumComponents;
 }
 
-bool UE::HLSLTree::FType::Merge(const FType& OtherType)
+bool FType::Merge(const FType& OtherType)
 {
 	if (ValueType == Shader::EValueType::Void)
 	{
@@ -332,7 +337,7 @@ bool UE::HLSLTree::FType::Merge(const FType& OtherType)
 	return true;
 }
 
-UE::HLSLTree::FStructFieldRef UE::HLSLTree::FStructType::FindFieldByName(const TCHAR* InName) const
+FStructFieldRef FStructType::FindFieldByName(const TCHAR* InName) const
 {
 	for (const FStructField& Field : Fields)
 	{
@@ -345,7 +350,7 @@ UE::HLSLTree::FStructFieldRef UE::HLSLTree::FStructType::FindFieldByName(const T
 	return FStructFieldRef();
 }
 
-UE::HLSLTree::FConstantValue::FConstantValue(const Shader::FValue& InValue) : Type(InValue.GetType())
+FConstantValue::FConstantValue(const Shader::FValue& InValue) : Type(InValue.GetType())
 {
 	const Shader::FValueTypeDescription TypeDesc = Shader::GetValueTypeDescription(Type);
 	Component.Empty(TypeDesc.NumComponents);
@@ -355,7 +360,7 @@ UE::HLSLTree::FConstantValue::FConstantValue(const Shader::FValue& InValue) : Ty
 	}
 }
 
-UE::HLSLTree::FScope* UE::HLSLTree::FScope::FindSharedParent(FScope* Lhs, FScope* Rhs)
+FScope* FScope::FindSharedParent(FScope* Lhs, FScope* Rhs)
 {
 	FScope* Scope0 = Lhs;
 	FScope* Scope1 = Rhs;
@@ -378,7 +383,7 @@ UE::HLSLTree::FScope* UE::HLSLTree::FScope::FindSharedParent(FScope* Lhs, FScope
 	return Scope0;
 }
 
-void UE::HLSLTree::FExpressionLocalPHI::UpdateType(FUpdateTypeContext& Context, const FRequestedType& RequestedType)
+void FExpressionLocalPHI::UpdateType(FUpdateTypeContext& Context, const FRequestedType& RequestedType)
 {
 	FType TypePerValue[MaxNumPreviousScopes];
 	int32 NumValidTypes = 0;
@@ -421,7 +426,7 @@ void UE::HLSLTree::FExpressionLocalPHI::UpdateType(FUpdateTypeContext& Context,
 	}
 }
 
-void UE::HLSLTree::FExpressionLocalPHI::PrepareValue(FEmitContext& Context)
+void FExpressionLocalPHI::PrepareValue(FEmitContext& Context)
 {
 	FExpression* ForwardExpression = Values[0];
 	for (int32 i = 1; i < NumValues; ++i)
@@ -494,7 +499,7 @@ void UE::HLSLTree::FExpressionLocalPHI::PrepareValue(FEmitContext& Context)
 	}
 }
 
-void UE::HLSLTree::FNodeVisitor::VisitNode(FNode* Node)
+void FNodeVisitor::VisitNode(FNode* Node)
 {
 	if (Node)
 	{
@@ -502,22 +507,22 @@ void UE::HLSLTree::FNodeVisitor::VisitNode(FNode* Node)
 	}
 }
 
-UE::HLSLTree::ENodeVisitResult UE::HLSLTree::FStatement::Visit(FNodeVisitor& Visitor)
+ENodeVisitResult FStatement::Visit(FNodeVisitor& Visitor)
 {
 	return Visitor.OnStatement(*this);
 }
 
-UE::HLSLTree::ENodeVisitResult UE::HLSLTree::FExpression::Visit(FNodeVisitor& Visitor)
+ENodeVisitResult FExpression::Visit(FNodeVisitor& Visitor)
 {
 	return Visitor.OnExpression(*this);
 }
 
-UE::HLSLTree::ENodeVisitResult UE::HLSLTree::FTextureParameterDeclaration::Visit(FNodeVisitor& Visitor)
+ENodeVisitResult FTextureParameterDeclaration::Visit(FNodeVisitor& Visitor)
 {
 	return Visitor.OnTextureParameterDeclaration(*this);
 }
 
-UE::HLSLTree::ENodeVisitResult UE::HLSLTree::FScope::Visit(FNodeVisitor& Visitor)
+ENodeVisitResult FScope::Visit(FNodeVisitor& Visitor)
 {
 	const ENodeVisitResult Result = Visitor.OnScope(*this);
 	if (ShouldVisitDependentNodes(Result))
@@ -527,7 +532,7 @@ UE::HLSLTree::ENodeVisitResult UE::HLSLTree::FScope::Visit(FNodeVisitor& Visitor
 	return Result;
 }
 
-UE::HLSLTree::FType UE::HLSLTree::RequestExpressionType(FUpdateTypeContext& Context, FExpression* InExpression, const FRequestedType& RequestedType)
+FType RequestExpressionType(FUpdateTypeContext& Context, FExpression* InExpression, const FRequestedType& RequestedType)
 {
 	if (!InExpression)
 	{
@@ -586,7 +591,7 @@ UE::HLSLTree::FType UE::HLSLTree::RequestExpressionType(FUpdateTypeContext& Cont
 	return InExpression->Type;
 }
 
-UE::HLSLTree::EExpressionEvaluationType UE::HLSLTree::PrepareExpressionValue(FEmitContext& Context, FExpression* InExpression)
+EExpressionEvaluationType PrepareExpressionValue(FEmitContext& Context, FExpression* InExpression)
 {
 	if (!InExpression)
 	{
@@ -610,12 +615,12 @@ UE::HLSLTree::EExpressionEvaluationType UE::HLSLTree::PrepareExpressionValue(FEm
 	return InExpression->EvaluationType;
 }
 
-void UE::HLSLTree::FExpression::SetType(FUpdateTypeContext& Context, const FType& InType)
+void FExpression::SetType(FUpdateTypeContext& Context, const FType& InType)
 {
 	Type = InType;
 }
 
-void UE::HLSLTree::FExpression::InternalSetValueShader(FEmitContext& Context, const TCHAR* InCode, int32 InLength, bool bInline)
+void FExpression::InternalSetValueShader(FEmitContext& Context, const TCHAR* InCode, int32 InLength, bool bInline)
 {
 	check(!Code);
 	check(EvaluationType == EExpressionEvaluationType::None);
@@ -659,7 +664,7 @@ void UE::HLSLTree::FExpression::InternalSetValueShader(FEmitContext& Context, co
 	}
 }
 
-void UE::HLSLTree::FExpression::SetValuePreshader(FEmitContext& Context, Shader::FPreshaderData& InPreshader)
+void FExpression::SetValuePreshader(FEmitContext& Context, Shader::FPreshaderData& InPreshader)
 {
 	check(!Preshader);
 	check(EvaluationType == EExpressionEvaluationType::None);
@@ -669,7 +674,7 @@ void UE::HLSLTree::FExpression::SetValuePreshader(FEmitContext& Context, Shader:
 	Preshader = new Shader::FPreshaderData(MoveTemp(InPreshader));
 }
 
-void UE::HLSLTree::FExpression::SetValueConstant(FEmitContext& Context, const Shader::FValue& InValue)
+void FExpression::SetValueConstant(FEmitContext& Context, const Shader::FValue& InValue)
 {
 	check(EvaluationType == EExpressionEvaluationType::None);
 	check(!Type.IsStruct());
@@ -678,7 +683,7 @@ void UE::HLSLTree::FExpression::SetValueConstant(FEmitContext& Context, const Sh
 	ConstantValue = Shader::Cast(InValue, Type);
 }
 
-void UE::HLSLTree::FExpression::SetValueForward(FEmitContext& Context, FExpression* Source)
+void FExpression::SetValueForward(FEmitContext& Context, FExpression* Source)
 {
 	check(EvaluationType == EExpressionEvaluationType::None);
 	check(Source);
@@ -696,7 +701,7 @@ void UE::HLSLTree::FExpression::SetValueForward(FEmitContext& Context, FExpressi
 	}
 }
 
-void UE::HLSLTree::FExpression::SetValuePreshader(FEmitContext& Context, EExpressionEvaluationType InEvaluationType, Shader::FPreshaderData& InPreshader)
+void FExpression::SetValuePreshader(FEmitContext& Context, EExpressionEvaluationType InEvaluationType, Shader::FPreshaderData& InPreshader)
 {
 	if (InEvaluationType == EExpressionEvaluationType::Constant)
 	{
@@ -713,7 +718,7 @@ void UE::HLSLTree::FExpression::SetValuePreshader(FEmitContext& Context, EExpres
 	}
 }
 
-const TCHAR* UE::HLSLTree::FExpression::GetValueShader(FEmitContext& Context)
+const TCHAR* FExpression::GetValueShader(FEmitContext& Context)
 {
 	check(EvaluationType != EExpressionEvaluationType::None);
 	if (!Code)
@@ -735,7 +740,7 @@ const TCHAR* UE::HLSLTree::FExpression::GetValueShader(FEmitContext& Context)
 	return Code;
 }
 
-const TCHAR* UE::HLSLTree::FExpression::GetValueShader(FEmitContext& Context, const FType& InType)
+const TCHAR* FExpression::GetValueShader(FEmitContext& Context, const FType& InType)
 {
 	check(EvaluationType != EExpressionEvaluationType::None);
 
@@ -753,7 +758,7 @@ const TCHAR* UE::HLSLTree::FExpression::GetValueShader(FEmitContext& Context, co
 	return Context.CastShaderValue(this, GetValueShader(Context), Type, InType, CastFlags);
 }
 
-void UE::HLSLTree::FExpression::GetValuePreshader(FEmitContext& Context, Shader::FPreshaderData& OutPreshader)
+void FExpression::GetValuePreshader(FEmitContext& Context, Shader::FPreshaderData& OutPreshader)
 {
 	switch (EvaluationType)
 	{
@@ -770,13 +775,13 @@ void UE::HLSLTree::FExpression::GetValuePreshader(FEmitContext& Context, Shader:
 	}
 }
 
-UE::Shader::FValue UE::HLSLTree::FExpression::GetValueConstant(FEmitContext& Context)
+Shader::FValue FExpression::GetValueConstant(FEmitContext& Context)
 {
 	check(EvaluationType == EExpressionEvaluationType::Constant);
 	return ConstantValue;
 }
 
-bool UE::HLSLTree::FScope::HasParentScope(const FScope& InParentScope) const
+bool FScope::HasParentScope(const FScope& InParentScope) const
 {
 	const FScope* CurrentScope = this;
 	while (CurrentScope)
@@ -790,16 +795,12 @@ bool UE::HLSLTree::FScope::HasParentScope(const FScope& InParentScope) const
 	return false;
 }
 
-void UE::HLSLTree::FScope::AddPreviousScope(FScope& Scope)
+void FScope::AddPreviousScope(FScope& Scope)
 {
 	check(NumPreviousScopes < MaxNumPreviousScopes);
 	PreviousScope[NumPreviousScopes++] = &Scope;
 }
 
-namespace UE
-{
-namespace HLSLTree
-{
 class FNodeVisitor_MoveToScope : public FNodeVisitor
 {
 public:
@@ -820,17 +821,15 @@ public:
 
 	FScope* Scope;
 };
-} // namespace HLSLTree
-} // namespace UE
 
-void UE::HLSLTree::FScope::UseExpression(FExpression* Expression)
+void FScope::UseExpression(FExpression* Expression)
 {
 	// Need to move all dependent expressions/declarations into this scope
 	FNodeVisitor_MoveToScope Visitor(this);
 	Visitor.VisitNode(Expression);
 }
 
-void UE::HLSLTree::FScope::InternalEmitCode(FEmitContext& Context, FCodeList& List, FScope* NestedScope, const TCHAR* String, int32 Length)
+void FScope::InternalEmitCode(FEmitContext& Context, FCodeList& List, FScope* NestedScope, const TCHAR* String, int32 Length)
 {
 	if (NestedScope && NestedScope->Statement && !NestedScope->Statement->bEmitHLSL)
 	{
@@ -860,7 +859,7 @@ void UE::HLSLTree::FScope::InternalEmitCode(FEmitContext& Context, FCodeList& Li
 	List.Num++;
 }
 
-void UE::HLSLTree::RequestScopeTypes(FUpdateTypeContext& Context, const FScope* InScope)
+void RequestScopeTypes(FUpdateTypeContext& Context, const FScope* InScope)
 {
 	if (InScope && InScope->Statement)
 	{
@@ -868,7 +867,7 @@ void UE::HLSLTree::RequestScopeTypes(FUpdateTypeContext& Context, const FScope*
 	}
 }
 
-UE::HLSLTree::FTree* UE::HLSLTree::FTree::Create(FMemStackBase& Allocator)
+FTree* FTree::Create(FMemStackBase& Allocator)
 {
 	FTree* Tree = new(Allocator) FTree();
 	Tree->Allocator = &Allocator;
@@ -876,7 +875,7 @@ UE::HLSLTree::FTree* UE::HLSLTree::FTree::Create(FMemStackBase& Allocator)
 	return Tree;
 }
 
-void UE::HLSLTree::FTree::Destroy(FTree* Tree)
+void FTree::Destroy(FTree* Tree)
 {
 	if (Tree)
 	{
@@ -903,12 +902,12 @@ static void WriteIndent(int32 IndentLevel, FStringBuilderBase& InOutString)
 	}
 }
 
-void UE::HLSLTree::FScope::MarkLive()
+void FScope::MarkLive()
 {
 	bLive = true;
 }
 
-void UE::HLSLTree::FScope::MarkLiveRecursive()
+void FScope::MarkLiveRecursive()
 {
 	FScope* Scope = this;
 	while (Scope)
@@ -918,7 +917,7 @@ void UE::HLSLTree::FScope::MarkLiveRecursive()
 	}
 }
 
-void UE::HLSLTree::FScope::WriteHLSL(int32 Indent, FStringBuilderBase& OutString) const
+void FScope::WriteHLSL(int32 Indent, FStringBuilderBase& OutString) const
 {
 	{
 		const FCodeEntry* CodeDeclaration = Declarations.First;
@@ -955,7 +954,7 @@ void UE::HLSLTree::FScope::WriteHLSL(int32 Indent, FStringBuilderBase& OutString
 	}
 }
 
-void UE::HLSLTree::FStructType::WriteHLSL(FStringBuilderBase& OutString) const
+void FStructType::WriteHLSL(FStringBuilderBase& OutString) const
 {
 	OutString.Appendf(TEXT("struct %s\n"), Name);
 	OutString.Append(TEXT("{\n"));
@@ -973,7 +972,7 @@ void UE::HLSLTree::FStructType::WriteHLSL(FStringBuilderBase& OutString) const
 	OutString.Append(TEXT("\n"));
 }
 
-bool UE::HLSLTree::FTree::EmitHLSL(FEmitContext& Context,
+bool FTree::EmitHLSL(FEmitContext& Context,
 	FStringBuilderBase& OutDeclarations,
 	FStringBuilderBase& OutCode) const
 {
@@ -1026,13 +1025,13 @@ bool UE::HLSLTree::FTree::EmitHLSL(FEmitContext& Context,
 	return Context.Errors.Num() == 0;
 }
 
-void UE::HLSLTree::FTree::RegisterExpression(FScope& Scope, FExpression* Expression)
+void FTree::RegisterExpression(FScope& Scope, FExpression* Expression)
 {
 	check(!Expression->ParentScope);
 	Expression->ParentScope = &Scope;
 }
 
-void UE::HLSLTree::FTree::RegisterStatement(FScope& Scope, FStatement* Statement)
+void FTree::RegisterStatement(FScope& Scope, FStatement* Statement)
 {
 	check(!Scope.Statement)
 	check(!Statement->ParentScope);
@@ -1040,7 +1039,7 @@ void UE::HLSLTree::FTree::RegisterStatement(FScope& Scope, FStatement* Statement
 	Scope.Statement = Statement;
 }
 
-UE::HLSLTree::FScope* UE::HLSLTree::FTree::NewScope(FScope& Scope)
+FScope* FTree::NewScope(FScope& Scope)
 {
 	FScope* NewScope = NewNode<FScope>();
 	NewScope->ParentScope = &Scope;
@@ -1049,20 +1048,15 @@ UE::HLSLTree::FScope* UE::HLSLTree::FTree::NewScope(FScope& Scope)
 	return NewScope;
 }
 
-UE::HLSLTree::FTextureParameterDeclaration* UE::HLSLTree::FTree::NewTextureParameterDeclaration(const FName& Name, const FTextureDescription& DefaultValue)
+FTextureParameterDeclaration* FTree::NewTextureParameterDeclaration(const FName& Name, const FTextureDescription& DefaultValue)
 {
 	FTextureParameterDeclaration* Declaration = NewNode<FTextureParameterDeclaration>(Name, DefaultValue);
 	return Declaration;
 }
 
-namespace UE
-{
-namespace HLSLTree
-{
 namespace Private
 {
-
-static void SetComponentTypes(Shader::EValueComponentType* ComponentTypes, int32 ComponentIndex, const HLSLTree::FType& Type)
+void SetComponentTypes(Shader::EValueComponentType* ComponentTypes, int32 ComponentIndex, const HLSLTree::FType& Type)
 {
 	if (Type.IsStruct())
 	{
@@ -1080,12 +1074,9 @@ static void SetComponentTypes(Shader::EValueComponentType* ComponentTypes, int32
 		}
 	}
 }
+} // namespace Private
 
-}
-}
-}
-
-const UE::HLSLTree::FStructType* UE::HLSLTree::FTree::NewStructType(const FStructTypeInitializer& Initializer)
+const FStructType* FTree::NewStructType(const FStructTypeInitializer& Initializer)
 {
 	const int32 NumFields = Initializer.Fields.Num();
 	int32 ComponentIndex = 0;
@@ -1117,9 +1108,11 @@ const UE::HLSLTree::FStructType* UE::HLSLTree::FTree::NewStructType(const FStruc
 	return StructType;
 }
 
-void UE::HLSLTree::FTree::SetResult(FStatement& InResult)
+void FTree::SetResult(FStatement& InResult)
 {
 	check(!ResultStatement);
 	ResultStatement = &InResult;
 }
 
+} // namespace HLSLTree
+} // namespace UE
\ No newline at end of file
diff --git a/Engine/Source/Runtime/Engine/Private/HLSLTree/HLSLTreeCommon.cpp b/Engine/Source/Runtime/Engine/Private/HLSLTree/HLSLTreeCommon.cpp
--- a/Engine/Source/Runtime/Engine/Private/HLSLTree/HLSLTreeCommon.cpp
+++ b/Engine/Source/Runtime/Engine/Private/HLSLTree/HLSLTreeCommon.cpp
@@ -4,7 +4,12 @@
 #include "MaterialShared.h"
 #include "Engine/Texture.h"
 
-UE::HLSLTree::FBinaryOpDescription UE::HLSLTree::GetBinaryOpDesription(EBinaryOp Op)
+namespace UE
+{
+namespace HLSLTree
+{
+
+FBinaryOpDescription GetBinaryOpDesription(EBinaryOp Op)
 {
 	switch (Op)
 	{
@@ -18,7 +23,7 @@ UE::HLSLTree::FBinaryOpDescription UE::HLSLTree::GetBinaryOpDesription(EBinaryOp
 	}
 }
 
-UE::HLSLTree::FSwizzleParameters::FSwizzleParameters(int8 InR, int8 InG, int8 InB, int8 InA) : NumComponents(0)
+FSwizzleParameters::FSwizzleParameters(int8 InR, int8 InG, int8 InB, int8 InA) : NumComponents(0)
 {
 	ComponentIndex[0] = InR;
 	ComponentIndex[1] = InG;
@@ -49,7 +54,7 @@ UE::HLSLTree::FSwizzleParameters::FSwizzleParameters(int8 InR, int8 InG, int8 In
 	++NumComponents;
 }
 
-UE::HLSLTree::FSwizzleParameters UE::HLSLTree::MakeSwizzleMask(bool bInR, bool bInG, bool bInB, bool bInA)
+FSwizzleParameters MakeSwizzleMask(bool bInR, bool bInG, bool bInB, bool bInA)
 {
 	int8 ComponentIndex[4] = { INDEX_NONE, INDEX_NONE, INDEX_NONE, INDEX_NONE };
 	int8 CurrentComponent = 0;
@@ -72,7 +77,7 @@ UE::HLSLTree::FSwizzleParameters UE::HLSLTree::MakeSwizzleMask(bool bInR, bool b
 	return FSwizzleParameters(ComponentIndex[0], ComponentIndex[1], ComponentIndex[2], ComponentIndex[3]);
 }
 
-void UE::HLSLTree::FExpressionConstant::UpdateType(FUpdateTypeContext& Context, const FRequestedType& RequestedType)
+void FExpressionConstant::UpdateType(FUpdateTypeContext& Context, const FRequestedType& RequestedType)
 {
 	if (Value.Type.IsStruct())
 	{
@@ -89,7 +94,7 @@ void UE::HLSLTree::FExpressionConstant::UpdateType(FUpdateTypeContext& Context,
 }
 
 
-void UE::HLSLTree::FExpressionConstant::PrepareValue(FEmitContext& Context)
+void FExpressionConstant::PrepareValue(FEmitContext& Context)
 {
 	check(Value.Component.Num() == Value.Type.GetNumComponents());
 	if (Value.Type.IsStruct())
@@ -120,13 +125,13 @@ void UE::HLSLTree::FExpressionConstant::PrepareValue(FEmitContext& Context)
 	}
 }
 
-void UE::HLSLTree::FExpressionMaterialParameter::UpdateType(FUpdateTypeContext& Context, const FRequestedType& RequestedType)
+void FExpressionMaterialParameter::UpdateType(FUpdateTypeContext& Context, const FRequestedType& RequestedType)
 {
 	const Shader::EValueType CurrentType = Shader::MakeValueTypeWithRequestedNumComponents(GetShaderValueType(ParameterType), RequestedType.GetRequestedNumComponents());
 	return SetType(Context, CurrentType);
 }
 
-void UE::HLSLTree::FExpressionMaterialParameter::PrepareValue(FEmitContext& Context)
+void FExpressionMaterialParameter::PrepareValue(FEmitContext& Context)
 {
 	if (ParameterType == EMaterialParameterType::StaticSwitch)
 	{
@@ -165,10 +170,10 @@ void UE::HLSLTree::FExpressionMaterialParameter::PrepareValue(FEmitContext& Cont
 	}
 }
 
-void UE::HLSLTree::FExpressionExternalInput::PrepareValue(FEmitContext& Context)
+void FExpressionExternalInput::PrepareValue(FEmitContext& Context)
 {
 	const int32 TypeIndex = (int32)InputType;
-	const int32 TexCoordIndex = TypeIndex - (int32)UE::HLSLTree::EExternalInputType::TexCoord0;
+	const int32 TexCoordIndex = TypeIndex - (int32)EExternalInputType::TexCoord0;
 
 	if (TexCoordIndex >= 0 && TexCoordIndex < 8)
 	{
@@ -181,20 +186,20 @@ void UE::HLSLTree::FExpressionExternalInput::PrepareValue(FEmitContext& Context)
 	}
 }
 
-void UE::HLSLTree::FExpressionTextureSample::UpdateType(FUpdateTypeContext& Context, const FRequestedType& RequestedType)
+void FExpressionTextureSample::UpdateType(FUpdateTypeContext& Context, const FRequestedType& RequestedType)
 {
 	RequestExpressionType(Context, TexCoordExpression, 2);
 	return SetType(Context, Shader::EValueType::Float4);
 }
 
-void UE::HLSLTree::FExpressionTextureSample::PrepareValue(FEmitContext& Context)
+void FExpressionTextureSample::PrepareValue(FEmitContext& Context)
 {
 	if (PrepareExpressionValue(Context, TexCoordExpression) == EExpressionEvaluationType::None)
 	{
 		return;
 	}
 
-	const UE::HLSLTree::FTextureDescription& Desc = Declaration->Description;
+	const FTextureDescription& Desc = Declaration->Description;
 	const EMaterialValueType MaterialValueType = Desc.Texture->GetMaterialType();
 	EMaterialTextureParameterType TextureType = EMaterialTextureParameterType::Count;
 	bool bVirtualTexture = false;
@@ -226,7 +231,7 @@ void UE::HLSLTree::FExpressionTextureSample::PrepareValue(FEmitContext& Context)
 		// TODO
 		TextureTypeName = TEXT("ExternalTexture");
 		SampleFunctionName = TEXT("TextureExternalSample");
-		break; 
+		break;
 	case MCT_TextureVirtual:
 		// TODO
 		TextureType = EMaterialTextureParameterType::Virtual;
@@ -326,7 +331,7 @@ void UE::HLSLTree::FExpressionTextureSample::PrepareValue(FEmitContext& Context)
 		TexCoordExpression->GetValueShader(Context, Shader::EValueType::Float2));
 }
 
-void UE::HLSLTree::FExpressionGetStructField::UpdateType(FUpdateTypeContext& Context, const FRequestedType& RequestedType)
+void FExpressionGetStructField::UpdateType(FUpdateTypeContext& Context, const FRequestedType& RequestedType)
 {
 	const FStructFieldRef FieldRef = StructType->FindFieldByName(FieldName);
 	if (!FieldRef)
@@ -342,7 +347,7 @@ void UE::HLSLTree::FExpressionGetStructField::UpdateType(FUpdateTypeContext& Con
 	return SetType(Context, FieldRef.Type);
 }
 
-void UE::HLSLTree::FExpressionGetStructField::PrepareValue(FEmitContext& Context)
+void FExpressionGetStructField::PrepareValue(FEmitContext& Context)
 {
 	const FStructFieldRef FieldRef = StructType->FindFieldByName(FieldName);
 	check(FieldRef);
@@ -357,7 +362,7 @@ void UE::HLSLTree::FExpressionGetStructField::PrepareValue(FEmitContext& Context
 		FieldName);
 }
 
-void UE::HLSLTree::FExpressionSetStructField::UpdateType(FUpdateTypeContext& Context, const FRequestedType& RequestedType)
+void FExpressionSetStructField::UpdateType(FUpdateTypeContext& Context, const FRequestedType& RequestedType)
 {
 	const FStructFieldRef FieldRef = StructType->FindFieldByName(FieldName);
 	if (!FieldRef)
@@ -375,11 +380,11 @@ void UE::HLSLTree::FExpressionSetStructField::UpdateType(FUpdateTypeContext& Con
 	return SetType(Context, StructType);
 }
 
-void UE::HLSLTree::FExpressionSetStructField::PrepareValue(FEmitContext& Context)
+void FExpressionSetStructField::PrepareValue(FEmitContext& Context)
 {
 	const EExpressionEvaluationType StructEvaluation = PrepareExpressionValue(Context, StructExpression);
 	const EExpressionEvaluationType FieldEvaluation = PrepareExpressionValue(Context, FieldExpression);
-	
+
 	if (FieldEvaluation != EExpressionEvaluationType::None)
 	{
 		const FStructFieldRef FieldRef = StructType->FindFieldByName(FieldName);
@@ -412,7 +417,7 @@ void UE::HLSLTree::FExpressionSetStructField::PrepareValue(FEmitContext& Context
 	}
 }
 
-void UE::HLSLTree::FExpressionSelect::UpdateType(FUpdateTypeContext& Context, const FRequestedType& RequestedType)
+void FExpressionSelect::UpdateType(FUpdateTypeContext& Context, const FRequestedType& RequestedType)
 {
 	const FType ConditionType = RequestExpressionType(Context, ConditionExpression, 1);
 	const FType LhsType = RequestExpressionType(Context, FalseExpression, RequestedType);
@@ -440,7 +445,7 @@ void UE::HLSLTree::FExpressionSelect::UpdateType(FUpdateTypeContext& Context, co
 	}
 }
 
-void UE::HLSLTree::FExpressionSelect::PrepareValue(FEmitContext& Context)
+void FExpressionSelect::PrepareValue(FEmitContext& Context)
 {
 	const EExpressionEvaluationType ConditionEvaluation = PrepareExpressionValue(Context, ConditionExpression);
 	if (ConditionEvaluation == EExpressionEvaluationType::None)
@@ -470,7 +475,7 @@ void UE::HLSLTree::FExpressionSelect::PrepareValue(FEmitContext& Context)
 	}
 }
 
-void UE::HLSLTree::FExpressionBinaryOp::UpdateType(FUpdateTypeContext& Context, const FRequestedType& RequestedType)
+void FExpressionBinaryOp::UpdateType(FUpdateTypeContext& Context, const FRequestedType& RequestedType)
 {
 	const FType LhsType = RequestExpressionType(Context, Lhs, RequestedType);
 	const FType RhsType = RequestExpressionType(Context, Rhs, RequestedType);
@@ -498,7 +503,7 @@ void UE::HLSLTree::FExpressionBinaryOp::UpdateType(FUpdateTypeContext& Context,
 	return SetType(Context, ResultType);
 }
 
-void UE::HLSLTree::FExpressionBinaryOp::PrepareValue(FEmitContext& Context)
+void FExpressionBinaryOp::PrepareValue(FEmitContext& Context)
 {
 	const EExpressionEvaluationType LhsEvaluation = PrepareExpressionValue(Context, Lhs);
 	const EExpressionEvaluationType RhsEvaluation = PrepareExpressionValue(Context, Rhs);
@@ -513,10 +518,10 @@ void UE::HLSLTree::FExpressionBinaryOp::PrepareValue(FEmitContext& Context)
 	{
 		switch (Op)
 		{
-		case UE::HLSLTree::EBinaryOp::Add: Opcode = Shader::EPreshaderOpcode::Add; break;
-		case UE::HLSLTree::EBinaryOp::Sub: Opcode = Shader::EPreshaderOpcode::Add; break;
-		case UE::HLSLTree::EBinaryOp::Mul: Opcode = Shader::EPreshaderOpcode::Mul; break;
-		case UE::HLSLTree::EBinaryOp::Div: Opcode = Shader::EPreshaderOpcode::Div; break;
+		case EBinaryOp::Add: Opcode = Shader::EPreshaderOpcode::Add; break;
+		case EBinaryOp::Sub: Opcode = Shader::EPreshaderOpcode::Add; break;
+		case EBinaryOp::Mul: Opcode = Shader::EPreshaderOpcode::Mul; break;
+		case EBinaryOp::Div: Opcode = Shader::EPreshaderOpcode::Div; break;
 		default: break;
 		}
 	}
@@ -546,7 +551,7 @@ void UE::HLSLTree::FExpressionBinaryOp::PrepareValue(FEmitContext& Context)
 	}
 }
 
-void UE::HLSLTree::FExpressionSwizzle::UpdateType(FUpdateTypeContext& Context, const FRequestedType& RequestedType)
+void FExpressionSwizzle::UpdateType(FUpdateTypeContext& Context, const FRequestedType& RequestedType)
 {
 	FRequestedType RequestedInputType(4, false);
 	int32 MaxComponentIndex = 0;
@@ -564,7 +569,7 @@ void UE::HLSLTree::FExpressionSwizzle::UpdateType(FUpdateTypeContext& Context, c
 	return SetType(Context, Shader::MakeValueType(InputType, MaxComponentIndex + 1));
 }
 
-void UE::HLSLTree::FExpressionSwizzle::PrepareValue(FEmitContext& Context)
+void FExpressionSwizzle::PrepareValue(FEmitContext& Context)
 {
 	const EExpressionEvaluationType InputEvaluation = PrepareExpressionValue(Context, Input);
 	if (InputEvaluation == EExpressionEvaluationType::None)
@@ -618,7 +623,7 @@ void UE::HLSLTree::FExpressionSwizzle::PrepareValue(FEmitContext& Context)
 	}
 }
 
-void UE::HLSLTree::FExpressionAppend::UpdateType(FUpdateTypeContext& Context, const FRequestedType& RequestedType)
+void FExpressionAppend::UpdateType(FUpdateTypeContext& Context, const FRequestedType& RequestedType)
 {
 	const FType LhsType = RequestExpressionType(Context, Lhs, 4);
 	const FType RhsType = RequestExpressionType(Context, Rhs, 4);
@@ -643,7 +648,7 @@ void UE::HLSLTree::FExpressionAppend::UpdateType(FUpdateTypeContext& Context, co
 	return SetType(Context, Shader::MakeValueType(LhsTypeDesc.ComponentType, NumComponents));
 }
 
-void UE::HLSLTree::FExpressionAppend::PrepareValue(FEmitContext& Context)
+void FExpressionAppend::PrepareValue(FEmitContext& Context)
 {
 	const EExpressionEvaluationType LhsEvaluation = PrepareExpressionValue(Context, Lhs);
 	const EExpressionEvaluationType RhsEvaluation = PrepareExpressionValue(Context, Rhs);
@@ -671,23 +676,23 @@ void UE::HLSLTree::FExpressionAppend::PrepareValue(FEmitContext& Context)
 	}
 }
 
-void UE::HLSLTree::FExpressionReflectionVector::PrepareValue(FEmitContext& Context)
+void FExpressionReflectionVector::PrepareValue(FEmitContext& Context)
 {
 	return SetValueInlineShaderf(Context, TEXT("Parameters.ReflectionVector"));
 }
 
-void UE::HLSLTree::FStatementBreak::EmitHLSL(FEmitContext& Context) const
+void FStatementBreak::EmitHLSL(FEmitContext& Context) const
 {
 	ParentScope->MarkLive();
 	ParentScope->EmitStatementf(Context, TEXT("break;"));
 }
 
-void UE::HLSLTree::FStatementReturn::RequestTypes(FUpdateTypeContext& Context) const
+void FStatementReturn::RequestTypes(FUpdateTypeContext& Context) const
 {
 	RequestExpressionType(Context, Expression, Type);
 }
 
-void UE::HLSLTree::FStatementReturn::EmitHLSL(FEmitContext& Context) const
+void FStatementReturn::EmitHLSL(FEmitContext& Context) const
 {
 	if (PrepareExpressionValue(Context, Expression) == EExpressionEvaluationType::None)
 	{
@@ -698,7 +703,7 @@ void UE::HLSLTree::FStatementReturn::EmitHLSL(FEmitContext& Context) const
 	ParentScope->EmitStatementf(Context, TEXT("return %s;"), Expression->GetValueShader(Context));
 }
 
-void UE::HLSLTree::FStatementIf::RequestTypes(FUpdateTypeContext& Context) const
+void FStatementIf::RequestTypes(FUpdateTypeContext& Context) const
 {
 	RequestExpressionType(Context, ConditionExpression, 1);
 	RequestScopeTypes(Context, ThenScope);
@@ -706,7 +711,7 @@ void UE::HLSLTree::FStatementIf::RequestTypes(FUpdateTypeContext& Context) const
 	RequestScopeTypes(Context, NextScope);
 }
 
-void UE::HLSLTree::FStatementIf::EmitHLSL(FEmitContext& Context) const
+void FStatementIf::EmitHLSL(FEmitContext& Context) const
 {
 	bool bIfLive = ThenScope->IsLive();
 	if (ElseScope && ElseScope->IsLive())
@@ -731,13 +736,13 @@ void UE::HLSLTree::FStatementIf::EmitHLSL(FEmitContext& Context) const
 	ParentScope->EmitNestedScope(Context, NextScope);
 }
 
-void UE::HLSLTree::FStatementLoop::RequestTypes(FUpdateTypeContext& Context) const
+void FStatementLoop::RequestTypes(FUpdateTypeContext& Context) const
 {
 	RequestScopeTypes(Context, LoopScope);
 	RequestScopeTypes(Context, NextScope);
 }
 
-void UE::HLSLTree::FStatementLoop::EmitHLSL(FEmitContext& Context) const
+void FStatementLoop::EmitHLSL(FEmitContext& Context) const
 {
 	if (LoopScope->IsLive())
 	{
@@ -745,3 +750,6 @@ void UE::HLSLTree::FStatementLoop::EmitHLSL(FEmitContext& Context) const
 	}
 	ParentScope->EmitNestedScope(Context, NextScope);
 }
+
+} // namespace HLSLTree
+} // namespace UE
diff --git a/Engine/Source/Runtime/Engine/Private/Shader/ShaderValue.cpp b/Engine/Source/Runtime/Engine/Private/Shader/ShaderValue.cpp
--- a/Engine/Source/Runtime/Engine/Private/Shader/ShaderValue.cpp
+++ b/Engine/Source/Runtime/Engine/Private/Shader/ShaderValue.cpp
@@ -7,9 +7,9 @@ namespace UE
 {
 namespace Shader
 {
+
 namespace Private
 {
-
 struct FCastFloat
 {
 	using FComponentType = float;
@@ -142,10 +142,8 @@ void FormatComponent_Double(double Value, int32 NumComponents, EValueStringForma
 }
 
 } // namespace Private
-} // namespace Shader
-} // namespace UE
 
-UE::Shader::FValue UE::Shader::FValue::FromMemoryImage(EValueType Type, const void* Data, uint32* OutSizeInBytes)
+FValue FValue::FromMemoryImage(EValueType Type, const void* Data, uint32* OutSizeInBytes)
 {
 	FValue Result(Type);
 	const uint8* Bytes = static_cast<const uint8*>(Data);
@@ -165,7 +163,7 @@ UE::Shader::FValue UE::Shader::FValue::FromMemoryImage(EValueType Type, const vo
 	return Result;
 }
 
-UE::Shader::FMemoryImageValue UE::Shader::FValue::AsMemoryImage() const
+FMemoryImageValue FValue::AsMemoryImage() const
 {
 	FMemoryImageValue Result;
 	uint8* Bytes = Result.Bytes;
@@ -183,54 +181,54 @@ UE::Shader::FMemoryImageValue UE::Shader::FValue::AsMemoryImage() const
 	return Result;
 }
 
-UE::Shader::FFloatValue UE::Shader::FValue::AsFloat() const
+FFloatValue FValue::AsFloat() const
 {
 	FFloatValue Result;
 	Private::AsType(Private::FCastFloat(), *this, Result);
 	return Result;
 }
 
-UE::Shader::FDoubleValue UE::Shader::FValue::AsDouble() const
+FDoubleValue FValue::AsDouble() const
 {
 	FDoubleValue Result;
 	Private::AsType(Private::FCastDouble(), *this, Result);
 	return Result;
 }
 
-FLinearColor UE::Shader::FValue::AsLinearColor() const
+FLinearColor FValue::AsLinearColor() const
 {
 	const FFloatValue Result = AsFloat();
 	return FLinearColor(Result.Component[0], Result.Component[1], Result.Component[2], Result.Component[3]);
 }
 
-FVector4d UE::Shader::FValue::AsVector4d() const
+FVector4d FValue::AsVector4d() const
 {
 	const FDoubleValue Result = AsDouble();
 	return FVector4d(Result.Component[0], Result.Component[1], Result.Component[2], Result.Component[3]);
 }
 
-UE::Shader::FIntValue UE::Shader::FValue::AsInt() const
+FIntValue FValue::AsInt() const
 {
 	FIntValue Result;
 	Private::AsType(Private::FCastInt(), *this, Result);
 	return Result;
 }
 
-UE::Shader::FBoolValue UE::Shader::FValue::AsBool() const
+FBoolValue FValue::AsBool() const
 {
 	FBoolValue Result;
 	Private::AsType(Private::FCastBool(), *this, Result);
 	return Result;
 }
 
-float UE::Shader::FValue::AsFloatScalar() const
+float FValue::AsFloatScalar() const
 {
 	FFloatValue Result;
 	Private::AsType(Private::FCastFloat(), *this, Result);
 	return Result[0];
 }
 
-bool UE::Shader::FValue::AsBoolScalar() const
+bool FValue::AsBoolScalar() const
 {
 	const FBoolValue Result = AsBool();
 	for (int32 i = 0; i < NumComponents; ++i)
@@ -243,7 +241,7 @@ bool UE::Shader::FValue::AsBoolScalar() const
 	return false;
 }
 
-uint32 UE::Shader::GetComponentTypeSizeInBytes(EValueComponentType Type)
+uint32 GetComponentTypeSizeInBytes(EValueComponentType Type)
 {
 	switch (Type)
 	{
@@ -256,7 +254,7 @@ uint32 UE::Shader::GetComponentTypeSizeInBytes(EValueComponentType Type)
 	}
 }
 
-const TCHAR* UE::Shader::FValueComponent::ToString(EValueComponentType Type, FStringBuilderBase& OutString) const
+const TCHAR* FValueComponent::ToString(EValueComponentType Type, FStringBuilderBase& OutString) const
 {
 	switch (Type)
 	{
@@ -268,7 +266,7 @@ const TCHAR* UE::Shader::FValueComponent::ToString(EValueComponentType Type, FSt
 	return OutString.ToString();
 }
 
-const TCHAR* UE::Shader::FValue::ToString(EValueStringFormat Format, FStringBuilderBase& OutString) const
+const TCHAR* FValue::ToString(EValueStringFormat Format, FStringBuilderBase& OutString) const
 {
 	if (Format == EValueStringFormat::HLSL && ComponentType == EValueComponentType::Double)
 	{
@@ -345,7 +343,7 @@ const TCHAR* UE::Shader::FValue::ToString(EValueStringFormat Format, FStringBuil
 	return OutString.ToString();
 }
 
-UE::Shader::FValueTypeDescription UE::Shader::GetValueTypeDescription(EValueType Type)
+FValueTypeDescription GetValueTypeDescription(EValueType Type)
 {
 	switch (Type)
 	{
@@ -371,7 +369,7 @@ UE::Shader::FValueTypeDescription UE::Shader::GetValueTypeDescription(EValueType
 	}
 }
 
-UE::Shader::EValueType UE::Shader::MakeValueType(EValueComponentType ComponentType, int32 NumComponents)
+EValueType MakeValueType(EValueComponentType ComponentType, int32 NumComponents)
 {
 	if (NumComponents == 0)
 	{
@@ -424,18 +422,18 @@ UE::Shader::EValueType UE::Shader::MakeValueType(EValueComponentType ComponentTy
 	return EValueType::Void;
 }
 
-UE::Shader::EValueType UE::Shader::MakeValueType(EValueType BaseType, int32 NumComponents)
+EValueType MakeValueType(EValueType BaseType, int32 NumComponents)
 {
 	return MakeValueType(GetValueTypeDescription(BaseType).ComponentType, NumComponents);
 }
 
-UE::Shader::EValueType UE::Shader::MakeValueTypeWithRequestedNumComponents(EValueType BaseType, int8 RequestedNumComponents)
+EValueType MakeValueTypeWithRequestedNumComponents(EValueType BaseType, int8 RequestedNumComponents)
 {
 	const FValueTypeDescription TypeDesc = GetValueTypeDescription(BaseType);
 	return MakeValueType(TypeDesc.ComponentType, FMath::Min(TypeDesc.NumComponents, RequestedNumComponents));
 }
 
-UE::Shader::EValueType UE::Shader::MakeArithmeticResultType(EValueType Lhs, EValueType Rhs, FString& OutErrorMessage)
+EValueType MakeArithmeticResultType(EValueType Lhs, EValueType Rhs, FString& OutErrorMessage)
 {
 	const FValueTypeDescription LhsDesc = GetValueTypeDescription(Lhs);
 	const FValueTypeDescription RhsDesc = GetValueTypeDescription(Rhs);
@@ -488,7 +486,7 @@ UE::Shader::EValueType UE::Shader::MakeArithmeticResultType(EValueType Lhs, EVal
 	return EValueType::Void;
 }
 
-UE::Shader::EValueType UE::Shader::MakeComparisonResultType(EValueType Lhs, EValueType Rhs, FString& OutErrorMessage)
+EValueType MakeComparisonResultType(EValueType Lhs, EValueType Rhs, FString& OutErrorMessage)
 {
 	const FValueTypeDescription LhsDesc = GetValueTypeDescription(Lhs);
 	const FValueTypeDescription RhsDesc = GetValueTypeDescription(Rhs);
@@ -512,10 +510,6 @@ UE::Shader::EValueType UE::Shader::MakeComparisonResultType(EValueType Lhs, EVal
 	return EValueType::Void;
 }
 
-namespace UE
-{
-namespace Shader
-{
 namespace Private
 {
 
@@ -607,9 +601,9 @@ inline FValue UnaryOp(const Operation& Op, const FValue& Value)
 
 	if constexpr (Operation::SupportsDouble)
 	{
-		if (Value.ComponentType == UE::Shader::EValueComponentType::Double)
+		if (Value.ComponentType == EValueComponentType::Double)
 		{
-			Result.ComponentType = UE::Shader::EValueComponentType::Double;
+			Result.ComponentType = EValueComponentType::Double;
 			const FDoubleValue Cast = Value.AsDouble();
 			for (int32 i = 0; i < NumComponents; ++i)
 			{
@@ -621,9 +615,9 @@ inline FValue UnaryOp(const Operation& Op, const FValue& Value)
 
 	if constexpr (Operation::SupportsInt)
 	{
-		if (Value.ComponentType != UE::Shader::EValueComponentType::Float)
+		if (Value.ComponentType != EValueComponentType::Float)
 		{
-			Result.ComponentType = UE::Shader::EValueComponentType::Int;
+			Result.ComponentType = EValueComponentType::Int;
 			const FIntValue Cast = Value.AsInt();
 			for (int32 i = 0; i < NumComponents; ++i)
 			{
@@ -633,7 +627,7 @@ inline FValue UnaryOp(const Operation& Op, const FValue& Value)
 		}
 	}
 
-	Result.ComponentType = UE::Shader::EValueComponentType::Float;
+	Result.ComponentType = EValueComponentType::Float;
 	const FFloatValue Cast = Value.AsFloat();
 	for (int32 i = 0; i < NumComponents; ++i)
 	{
@@ -659,9 +653,9 @@ inline FValue BinaryOp(const Operation& Op, const FValue& Lhs, const FValue& Rhs
 
 	if constexpr (Operation::SupportsDouble)
 	{
-		if (Lhs.ComponentType == UE::Shader::EValueComponentType::Double || Rhs.ComponentType == UE::Shader::EValueComponentType::Double)
+		if (Lhs.ComponentType == EValueComponentType::Double || Rhs.ComponentType == EValueComponentType::Double)
 		{
-			Result.ComponentType = UE::Shader::EValueComponentType::Double;
+			Result.ComponentType = EValueComponentType::Double;
 			const FDoubleValue LhsCast = Lhs.AsDouble();
 			const FDoubleValue RhsCast = Rhs.AsDouble();
 			for (int32 i = 0; i < NumComponents; ++i)
@@ -674,9 +668,9 @@ inline FValue BinaryOp(const Operation& Op, const FValue& Lhs, const FValue& Rhs
 
 	if constexpr (Operation::SupportsInt)
 	{
-		if (Lhs.ComponentType != UE::Shader::EValueComponentType::Float && Rhs.ComponentType != UE::Shader::EValueComponentType::Float)
+		if (Lhs.ComponentType != EValueComponentType::Float && Rhs.ComponentType != EValueComponentType::Float)
 		{
-			Result.ComponentType = UE::Shader::EValueComponentType::Int;
+			Result.ComponentType = EValueComponentType::Int;
 			const FIntValue LhsCast = Lhs.AsInt();
 			const FIntValue RhsCast = Rhs.AsInt();
 			for (int32 i = 0; i < NumComponents; ++i)
@@ -687,7 +681,7 @@ inline FValue BinaryOp(const Operation& Op, const FValue& Lhs, const FValue& Rhs
 		}
 	}
 
-	Result.ComponentType = UE::Shader::EValueComponentType::Float;
+	Result.ComponentType = EValueComponentType::Float;
 	const FFloatValue LhsCast = Lhs.AsFloat();
 	const FFloatValue RhsCast = Rhs.AsFloat();
 	for (int32 i = 0; i < NumComponents; ++i)
@@ -698,10 +692,8 @@ inline FValue BinaryOp(const Operation& Op, const FValue& Lhs, const FValue& Rhs
 }
 
 } // namespace Private
-} // namespace Shader
-} // namespace UE
 
-bool UE::Shader::operator==(const FValue& Lhs, const FValue& Rhs)
+bool operator==(const FValue& Lhs, const FValue& Rhs)
 {
 	if (Lhs.ComponentType != Rhs.ComponentType || Lhs.NumComponents != Rhs.NumComponents)
 	{
@@ -717,7 +709,7 @@ bool UE::Shader::operator==(const FValue& Lhs, const FValue& Rhs)
 	return true;
 }
 
-uint32 UE::Shader::GetTypeHash(const FValue& Value)
+uint32 GetTypeHash(const FValue& Value)
 {
 	uint32 Result = ::GetTypeHash(Value.ComponentType);
 	Result = HashCombine(Result, ::GetTypeHash(Value.NumComponents));
@@ -738,156 +730,156 @@ uint32 UE::Shader::GetTypeHash(const FValue& Value)
 	return Result;
 }
 
-UE::Shader::FValue UE::Shader::Abs(const FValue& Value)
+FValue Abs(const FValue& Value)
 {
 	return Private::UnaryOp(Private::FOpAbs(), Value);
 }
 
-UE::Shader::FValue UE::Shader::Saturate(const FValue& Value)
+FValue Saturate(const FValue& Value)
 {
 	return Private::UnaryOp(Private::FOpSaturate(), Value);
 }
 
-UE::Shader::FValue UE::Shader::Floor(const FValue& Value)
+FValue Floor(const FValue& Value)
 {
 	return Private::UnaryOp(Private::FOpFloor(), Value);
 }
 
-UE::Shader::FValue UE::Shader::Ceil(const FValue& Value)
+FValue Ceil(const FValue& Value)
 {
 	return Private::UnaryOp(Private::FOpCeil(), Value);
 }
 
-UE::Shader::FValue UE::Shader::Round(const FValue& Value)
+FValue Round(const FValue& Value)
 {
 	return Private::UnaryOp(Private::FOpRound(), Value);
 }
 
-UE::Shader::FValue UE::Shader::Trunc(const FValue& Value)
+FValue Trunc(const FValue& Value)
 {
 	return Private::UnaryOp(Private::FOpTrunc(), Value);
 }
 
-UE::Shader::FValue UE::Shader::Sign(const FValue& Value)
+FValue Sign(const FValue& Value)
 {
 	return Private::UnaryOp(Private::FOpSign(), Value);
 }
 
-UE::Shader::FValue UE::Shader::Frac(const FValue& Value)
+FValue Frac(const FValue& Value)
 {
 	return Private::UnaryOp(Private::FOpFrac(), Value);
 }
 
-UE::Shader::FValue UE::Shader::Fractional(const FValue& Value)
+FValue Fractional(const FValue& Value)
 {
 	return Private::UnaryOp(Private::FOpFractional(), Value);
 }
 
-UE::Shader::FValue UE::Shader::Sqrt(const FValue& Value)
+FValue Sqrt(const FValue& Value)
 {
 	return Private::UnaryOp(Private::FOpSqrt(), Value);
 }
 
-UE::Shader::FValue UE::Shader::Rcp(const FValue& Value)
+FValue Rcp(const FValue& Value)
 {
 	return Private::UnaryOp(Private::FOpRcp(), Value);
 }
 
-UE::Shader::FValue UE::Shader::Log2(const FValue& Value)
+FValue Log2(const FValue& Value)
 {
 	return Private::UnaryOp(Private::FOpLog2(), Value);
 }
 
-UE::Shader::FValue UE::Shader::Log10(const FValue& Value)
+FValue Log10(const FValue& Value)
 {
 	return Private::UnaryOp(Private::FOpLog10(), Value);
 }
 
-UE::Shader::FValue UE::Shader::Sin(const FValue& Value)
+FValue Sin(const FValue& Value)
 {
 	return Private::UnaryOp(Private::FOpSin(), Value);
 }
 
-UE::Shader::FValue UE::Shader::Cos(const FValue& Value)
+FValue Cos(const FValue& Value)
 {
 	return Private::UnaryOp(Private::FOpCos(), Value);
 }
 
-UE::Shader::FValue UE::Shader::Tan(const FValue& Value)
+FValue Tan(const FValue& Value)
 {
 	return Private::UnaryOp(Private::FOpTan(), Value);
 }
 
-UE::Shader::FValue UE::Shader::Asin(const FValue& Value)
+FValue Asin(const FValue& Value)
 {
 	return Private::UnaryOp(Private::FOpAsin(), Value);
 }
 
-UE::Shader::FValue UE::Shader::Acos(const FValue& Value)
+FValue Acos(const FValue& Value)
 {
 	return Private::UnaryOp(Private::FOpAcos(), Value);
 }
 
-UE::Shader::FValue UE::Shader::Atan(const FValue& Value)
+FValue Atan(const FValue& Value)
 {
 	return Private::UnaryOp(Private::FOpAtan(), Value);
 }
 
-UE::Shader::FValue UE::Shader::Add(const UE::Shader::FValue& Lhs, const UE::Shader::FValue& Rhs)
+FValue Add(const FValue& Lhs, const FValue& Rhs)
 {
 	return Private::BinaryOp(Private::FOpAdd(), Lhs, Rhs);
 }
 
-UE::Shader::FValue UE::Shader::Sub(const UE::Shader::FValue& Lhs, const UE::Shader::FValue& Rhs)
+FValue Sub(const FValue& Lhs, const FValue& Rhs)
 {
 	return Private::BinaryOp(Private::FOpSub(), Lhs, Rhs);
 }
 
-UE::Shader::FValue UE::Shader::Mul(const UE::Shader::FValue& Lhs, const UE::Shader::FValue& Rhs)
+FValue Mul(const FValue& Lhs, const FValue& Rhs)
 {
 	return Private::BinaryOp(Private::FOpMul(), Lhs, Rhs);
 }
 
-UE::Shader::FValue UE::Shader::Div(const UE::Shader::FValue& Lhs, const UE::Shader::FValue& Rhs)
+FValue Div(const FValue& Lhs, const FValue& Rhs)
 {
 	return Private::BinaryOp(Private::FOpDiv(), Lhs, Rhs);
 }
 
-UE::Shader::FValue UE::Shader::Min(const UE::Shader::FValue& Lhs, const UE::Shader::FValue& Rhs)
+FValue Min(const FValue& Lhs, const FValue& Rhs)
 {
 	return Private::BinaryOp(Private::FOpMin(), Lhs, Rhs);
 }
 
-UE::Shader::FValue UE::Shader::Max(const UE::Shader::FValue& Lhs, const UE::Shader::FValue& Rhs)
+FValue Max(const FValue& Lhs, const FValue& Rhs)
 {
 	return Private::BinaryOp(Private::FOpMax(), Lhs, Rhs);
 }
 
-UE::Shader::FValue UE::Shader::Fmod(const UE::Shader::FValue& Lhs, const UE::Shader::FValue& Rhs)
+FValue Fmod(const FValue& Lhs, const FValue& Rhs)
 {
 	return Private::BinaryOp(Private::FOpFmod(), Lhs, Rhs);
 }
 
-UE::Shader::FValue UE::Shader::Atan2(const FValue& Lhs, const FValue& Rhs)
+FValue Atan2(const FValue& Lhs, const FValue& Rhs)
 {
 	return Private::BinaryOp(Private::FOpAtan2(), Lhs, Rhs);
 }
 
-UE::Shader::FValue UE::Shader::Clamp(const FValue& Value, const FValue& Low, const FValue& High)
+FValue Clamp(const FValue& Value, const FValue& Low, const FValue& High)
 {
 	return Min(Max(Value, Low), High);
 }
 
-UE::Shader::FValue UE::Shader::Dot(const FValue& Lhs, const FValue& Rhs)
+FValue Dot(const FValue& Lhs, const FValue& Rhs)
 {
 	const int8 NumComponents = Private::GetNumComponentsResult(Lhs.NumComponents, Rhs.NumComponents);
 
 	FValue Result;
 	Result.NumComponents = 1;
 
-	if (Lhs.ComponentType == UE::Shader::EValueComponentType::Double || Rhs.ComponentType == UE::Shader::EValueComponentType::Double)
+	if (Lhs.ComponentType == EValueComponentType::Double || Rhs.ComponentType == EValueComponentType::Double)
 	{
-		Result.ComponentType = UE::Shader::EValueComponentType::Double;
+		Result.ComponentType = EValueComponentType::Double;
 		const FDoubleValue LhsValue = Lhs.AsDouble();
 		const FDoubleValue RhsValue = Rhs.AsDouble();
 		double ComponentValue = 0.0;
@@ -897,9 +889,9 @@ UE::Shader::FValue UE::Shader::Dot(const FValue& Lhs, const FValue& Rhs)
 		}
 		Result.Component[0].Double = ComponentValue;
 	}
-	else if (Lhs.ComponentType == UE::Shader::EValueComponentType::Float || Rhs.ComponentType == UE::Shader::EValueComponentType::Float)
+	else if (Lhs.ComponentType == EValueComponentType::Float || Rhs.ComponentType == EValueComponentType::Float)
 	{
-		Result.ComponentType = UE::Shader::EValueComponentType::Float;
+		Result.ComponentType = EValueComponentType::Float;
 		const FFloatValue LhsValue = Lhs.AsFloat();
 		const FFloatValue RhsValue = Rhs.AsFloat();
 		float ComponentValue = 0.0f;
@@ -911,7 +903,7 @@ UE::Shader::FValue UE::Shader::Dot(const FValue& Lhs, const FValue& Rhs)
 	}
 	else
 	{
-		Result.ComponentType = UE::Shader::EValueComponentType::Int;
+		Result.ComponentType = EValueComponentType::Int;
 		const FIntValue LhsValue = Lhs.AsInt();
 		const FIntValue RhsValue = Rhs.AsInt();
 		int32 ComponentValue = 0;
@@ -924,14 +916,14 @@ UE::Shader::FValue UE::Shader::Dot(const FValue& Lhs, const FValue& Rhs)
 	return Result;
 }
 
-UE::Shader::FValue UE::Shader::Cross(const FValue& Lhs, const FValue& Rhs)
+FValue Cross(const FValue& Lhs, const FValue& Rhs)
 {
 	FValue Result;
 	Result.NumComponents = 3;
 
-	if (Lhs.ComponentType == UE::Shader::EValueComponentType::Double || Rhs.ComponentType == UE::Shader::EValueComponentType::Double)
+	if (Lhs.ComponentType == EValueComponentType::Double || Rhs.ComponentType == EValueComponentType::Double)
 	{
-		Result.ComponentType = UE::Shader::EValueComponentType::Double;
+		Result.ComponentType = EValueComponentType::Double;
 		const FDoubleValue LhsValue = Lhs.AsDouble();
 		const FDoubleValue RhsValue = Rhs.AsDouble();
 
@@ -939,9 +931,9 @@ UE::Shader::FValue UE::Shader::Cross(const FValue& Lhs, const FValue& Rhs)
 		Result.Component[1].Double = LhsValue.Component[2] * RhsValue.Component[0] - LhsValue.Component[0] * RhsValue.Component[2];
 		Result.Component[2].Double = LhsValue.Component[0] * RhsValue.Component[1] - LhsValue.Component[1] * RhsValue.Component[0];
 	}
-	else if (Lhs.ComponentType == UE::Shader::EValueComponentType::Float || Rhs.ComponentType == UE::Shader::EValueComponentType::Float)
+	else if (Lhs.ComponentType == EValueComponentType::Float || Rhs.ComponentType == EValueComponentType::Float)
 	{
-		Result.ComponentType = UE::Shader::EValueComponentType::Float;
+		Result.ComponentType = EValueComponentType::Float;
 		const FFloatValue LhsValue = Lhs.AsFloat();
 		const FFloatValue RhsValue = Rhs.AsFloat();
 		
@@ -951,7 +943,7 @@ UE::Shader::FValue UE::Shader::Cross(const FValue& Lhs, const FValue& Rhs)
 	}
 	else
 	{
-		Result.ComponentType = UE::Shader::EValueComponentType::Int;
+		Result.ComponentType = EValueComponentType::Int;
 		const FIntValue LhsValue = Lhs.AsInt();
 		const FIntValue RhsValue = Rhs.AsInt();
 
@@ -962,7 +954,7 @@ UE::Shader::FValue UE::Shader::Cross(const FValue& Lhs, const FValue& Rhs)
 	return Result;
 }
 
-UE::Shader::FValue UE::Shader::Append(const FValue& Lhs, const FValue& Rhs)
+FValue Append(const FValue& Lhs, const FValue& Rhs)
 {
 	FValue Result;
 	int32 NumComponents = 0;
@@ -980,9 +972,9 @@ UE::Shader::FValue UE::Shader::Append(const FValue& Lhs, const FValue& Rhs)
 			Result.Component[NumComponents++] = Rhs.Component[i];
 		}
 	}
-	else if (Lhs.ComponentType == UE::Shader::EValueComponentType::Double || Rhs.ComponentType == UE::Shader::EValueComponentType::Double)
+	else if (Lhs.ComponentType == EValueComponentType::Double || Rhs.ComponentType == EValueComponentType::Double)
 	{
-		Result.ComponentType = UE::Shader::EValueComponentType::Double;
+		Result.ComponentType = EValueComponentType::Double;
 		const FDoubleValue LhsValue = Lhs.AsDouble();
 		const FDoubleValue RhsValue = Rhs.AsDouble();
 		for (int32 i = 0; i < Lhs.NumComponents; ++i)
@@ -994,9 +986,9 @@ UE::Shader::FValue UE::Shader::Append(const FValue& Lhs, const FValue& Rhs)
 			Result.Component[NumComponents++].Double = RhsValue.Component[i];
 		}
 	}
-	else if (Lhs.ComponentType == UE::Shader::EValueComponentType::Float || Rhs.ComponentType == UE::Shader::EValueComponentType::Float)
+	else if (Lhs.ComponentType == EValueComponentType::Float || Rhs.ComponentType == EValueComponentType::Float)
 	{
-		Result.ComponentType = UE::Shader::EValueComponentType::Float;
+		Result.ComponentType = EValueComponentType::Float;
 		const FFloatValue LhsValue = Lhs.AsFloat();
 		const FFloatValue RhsValue = Rhs.AsFloat();
 		for (int32 i = 0; i < Lhs.NumComponents; ++i)
@@ -1010,7 +1002,7 @@ UE::Shader::FValue UE::Shader::Append(const FValue& Lhs, const FValue& Rhs)
 	}
 	else
 	{
-		Result.ComponentType = UE::Shader::EValueComponentType::Int;
+		Result.ComponentType = EValueComponentType::Int;
 		const FIntValue LhsValue = Lhs.AsInt();
 		const FIntValue RhsValue = Rhs.AsInt();
 		for (int32 i = 0; i < Lhs.NumComponents; ++i)
@@ -1026,7 +1018,7 @@ UE::Shader::FValue UE::Shader::Append(const FValue& Lhs, const FValue& Rhs)
 	return Result;
 }
 
-UE::Shader::FValue UE::Shader::Cast(const FValue& Value, EValueType Type)
+FValue Cast(const FValue& Value, EValueType Type)
 {
 	const EValueType SourceType = Value.GetType();
 	if (Type == SourceType)
@@ -1047,3 +1039,5 @@ UE::Shader::FValue UE::Shader::Cast(const FValue& Value, EValueType Type)
 	return Result;
 }
 
+} // namespace Shader
+} // namespace UE
