Fix some trivial UWorld const issues reported by a licensee

diff --git a/Engine/Source/Runtime/Engine/Classes/Engine/World.h b/Engine/Source/Runtime/Engine/Classes/Engine/World.h
--- a/Engine/Source/Runtime/Engine/Classes/Engine/World.h
+++ b/Engine/Source/Runtime/Engine/Classes/Engine/World.h
@@ -2798,7 +2798,7 @@ public:
 	 *
 	 * @return true if level load requests are allowed, false otherwise.
 	 */
-	bool AllowLevelLoadRequests();
+	bool AllowLevelLoadRequests() const;
 
 	/** Creates instances for each parameter collection in memory.  Called when a world is created. */
 	void SetupParameterCollectionInstances();
@@ -2807,7 +2807,7 @@ public:
 	void AddParameterCollectionInstance(class UMaterialParameterCollection* Collection, bool bUpdateScene);
 
 	/** Gets this world's instance for a given collection. */
-	UMaterialParameterCollectionInstance* GetParameterCollectionInstance(const UMaterialParameterCollection* Collection);
+	UMaterialParameterCollectionInstance* GetParameterCollectionInstance(const UMaterialParameterCollection* Collection) const;
 
 	/** Updates this world's scene with the list of instances, and optionally updates each instance's uniform buffer. */
 	void UpdateParameterCollectionInstances(bool bUpdateInstanceUniformBuffers, bool bRecreateUniformBuffer);
@@ -3566,10 +3566,10 @@ public:
 	FOnBeginTearingDownEvent& OnBeginTearingDown() { return BeginTearingDownEvent; }
 
 	/** Returns the actor count. */
-	int32 GetProgressDenominator();
+	int32 GetProgressDenominator() const;
 	
 	/** Returns the actor count. */
-	int32 GetActorCount();
+	int32 GetActorCount() const;
 	
 public:
 
@@ -3582,14 +3582,14 @@ public:
 	 * @param	OutInteriorSettings		[out] Upon return, the interior settings for a camera at ViewLocation.
 	 * @return							If the settings came from an audio volume, the audio volume object is returned.
 	 */
-	class AAudioVolume* GetAudioSettings( const FVector& ViewLocation, struct FReverbSettings* OutReverbSettings, struct FInteriorSettings* OutInteriorSettings );
+	class AAudioVolume* GetAudioSettings( const FVector& ViewLocation, struct FReverbSettings* OutReverbSettings, struct FInteriorSettings* OutInteriorSettings ) const;
 
 	void SetAudioDevice(const FAudioDeviceHandle& InHandle);
 
 	/**
 	 * Get the audio device used by this world.
 	 */
-	FAudioDeviceHandle GetAudioDevice();
+	FAudioDeviceHandle GetAudioDevice() const;
 
 	/**
 	* Returns the audio device associated with this world.
@@ -3597,7 +3597,7 @@ public:
 	*
 	* @return Audio device to use with this world.
 	*/
-	class FAudioDevice* GetAudioDeviceRaw();
+	class FAudioDevice* GetAudioDeviceRaw() const;
 
 	/** Return the URL of this level on the local machine. */
 	FString GetLocalURL() const;
@@ -3695,7 +3695,7 @@ public:
 	}
 
 	/** @return whether we're currently in a seamless transition */
-	bool IsInSeamlessTravel();
+	bool IsInSeamlessTravel() const;
 
 	/** this function allows pausing the seamless travel in the middle,
 	 * right before it starts loading the destination (i.e. while in the transition level)
@@ -3706,7 +3706,7 @@ public:
 	void SetSeamlessTravelMidpointPause(bool bNowPaused);
 
 	/** @return the current detail mode, like EDetailMode but can be outside of the range */
-	int32 GetDetailMode();
+	int32 GetDetailMode() const;
 
 	/** asynchronously loads the given levels in preparation for a streaming map transition.
 	 * This codepath is designed for worlds that heavily use level streaming and GameModes where the game state should
@@ -3716,10 +3716,10 @@ public:
 	void PrepareMapChange(const TArray<FName>& LevelNames);
 
 	/** @return true if there's a map change currently in progress */
-	bool IsPreparingMapChange();
+	bool IsPreparingMapChange() const;
 
 	/** @return true if there is a map change being prepared, returns whether that change is ready to be committed, otherwise false */
-	bool IsMapChangeReady();
+	bool IsMapChangeReady() const;
 
 	/** cancels pending map change (@note: we can't cancel pending async loads, so this won't immediately free the memory) */
 	void CancelPendingMapChange();
diff --git a/Engine/Source/Runtime/Engine/Private/LevelActor.cpp b/Engine/Source/Runtime/Engine/Private/LevelActor.cpp
--- a/Engine/Source/Runtime/Engine/Private/LevelActor.cpp
+++ b/Engine/Source/Runtime/Engine/Private/LevelActor.cpp
@@ -1720,7 +1720,7 @@ void UWorld::IssueEditorLoadWarnings()
 #endif // WITH_EDITOR
 
 
-AAudioVolume* UWorld::GetAudioSettings( const FVector& ViewLocation, FReverbSettings* OutReverbSettings, FInteriorSettings* OutInteriorSettings )
+AAudioVolume* UWorld::GetAudioSettings( const FVector& ViewLocation, FReverbSettings* OutReverbSettings, FInteriorSettings* OutInteriorSettings ) const
 {
 	// Find the highest priority volume encompassing the current view location.
 	for (AAudioVolume* Volume : AudioVolumes)
diff --git a/Engine/Source/Runtime/Engine/Private/World.cpp b/Engine/Source/Runtime/Engine/Private/World.cpp
--- a/Engine/Source/Runtime/Engine/Private/World.cpp
+++ b/Engine/Source/Runtime/Engine/Private/World.cpp
@@ -1334,7 +1334,7 @@ void UWorld::AddParameterCollectionInstance(UMaterialParameterCollection* Collec
 	}
 }
 
-UMaterialParameterCollectionInstance* UWorld::GetParameterCollectionInstance(const UMaterialParameterCollection* Collection)
+UMaterialParameterCollectionInstance* UWorld::GetParameterCollectionInstance(const UMaterialParameterCollection* Collection) const
 {
 	for (int32 InstanceIndex = 0; InstanceIndex < ParameterCollectionInstances.Num(); InstanceIndex++)
 	{
@@ -1346,7 +1346,7 @@ UMaterialParameterCollectionInstance* UWorld::GetParameterCollectionInstance(con
 
 	// Instance should always exist due to SetupParameterCollectionInstances() and UMaterialParameterCollection::PostLoad()
 	check(0);
-	return NULL;
+	return nullptr;
 }
 
 void UWorld::UpdateParameterCollectionInstances(bool bUpdateInstanceUniformBuffers, bool bRecreateUniformBuffer)
@@ -4121,7 +4121,7 @@ void UWorld::AsyncLoadAlwaysLoadedLevelsForSeamlessTravel()
 	}
 }
 
-bool UWorld::AllowLevelLoadRequests()
+bool UWorld::AllowLevelLoadRequests() const
 {
 	// Always allow level load request in the editor or when we do full streaming flush
 	if (IsGameWorld() && FlushLevelStreamingType != EFlushLevelStreamingType::Full)
@@ -6139,14 +6139,14 @@ void UWorld::PrepareMapChange(const TArray<FName>& LevelNames)
 	}
 }
 
-bool UWorld::IsPreparingMapChange()
+bool UWorld::IsPreparingMapChange() const
 {
-	return GEngine->IsPreparingMapChange(this);
+	return GEngine->IsPreparingMapChange(const_cast<UWorld*>(this));
 }
 
-bool UWorld::IsMapChangeReady()
+bool UWorld::IsMapChangeReady() const
 {
-	return GEngine->IsReadyForMapChange(this);
+	return GEngine->IsReadyForMapChange(const_cast<UWorld*>(this));
 }
 
 void UWorld::CancelPendingMapChange()
@@ -7160,9 +7160,9 @@ void UWorld::SeamlessTravel(const FString& SeamlessTravelURL, bool bAbsolute)
 }
 
 /** @return whether we're currently in a seamless transition */
-bool UWorld::IsInSeamlessTravel()
+bool UWorld::IsInSeamlessTravel() const
 {
-	FSeamlessTravelHandler &SeamlessTravelHandler = GEngine->SeamlessTravelHandlerForWorld( this );
+	FSeamlessTravelHandler& SeamlessTravelHandler = GEngine->SeamlessTravelHandlerForWorld(const_cast<UWorld*>(this));
 	return SeamlessTravelHandler.IsInTransition();
 }
 
@@ -7178,7 +7178,7 @@ void UWorld::SetSeamlessTravelMidpointPause(bool bNowPaused)
 	SeamlessTravelHandler.SetPauseAtMidpoint(bNowPaused);
 }
 
-int32 UWorld::GetDetailMode()
+int32 UWorld::GetDetailMode() const
 {
 	return GetCachedScalabilityCVars().DetailMode;
 }
@@ -7202,12 +7202,12 @@ void UWorld::UpdateConstraintActors()
 	}
 }
 
-int32 UWorld::GetProgressDenominator()
+int32 UWorld::GetProgressDenominator() const
 {
 	return GetActorCount();
 }
 
-int32 UWorld::GetActorCount()
+int32 UWorld::GetActorCount() const
 {
 	int32 TotalActorCount = 0;
 	for( int32 LevelIndex=0; LevelIndex<GetNumLevels(); LevelIndex++ )
@@ -8421,7 +8421,7 @@ void UWorld::SetAudioDevice(const FAudioDeviceHandle& InHandle)
 	}
 }
 
-FAudioDeviceHandle UWorld::GetAudioDevice()
+FAudioDeviceHandle UWorld::GetAudioDevice() const
 {
 	if (AudioDeviceHandle)
 	{
@@ -8437,7 +8437,7 @@ FAudioDeviceHandle UWorld::GetAudioDevice()
 	}
 }
 
-FAudioDevice* UWorld::GetAudioDeviceRaw()
+FAudioDevice* UWorld::GetAudioDeviceRaw() const
 {
 	if (AudioDeviceHandle)
 	{
