Fix duplicate symbol in StressUnity build

diff --git a/Engine/Source/Runtime/Engine/Private/WorldPartition/LevelInstance/LevelInstanceActorDesc.cpp b/Engine/Source/Runtime/Engine/Private/WorldPartition/LevelInstance/LevelInstanceActorDesc.cpp
--- a/Engine/Source/Runtime/Engine/Private/WorldPartition/LevelInstance/LevelInstanceActorDesc.cpp
+++ b/Engine/Source/Runtime/Engine/Private/WorldPartition/LevelInstance/LevelInstanceActorDesc.cpp
@@ -13,7 +13,7 @@
 #include "UObject/UE5MainStreamObjectVersion.h"
 
 static int32 GLevelInstanceDebugForceLevelStreaming = 0;
-static FAutoConsoleVariableRef CVarUpdateStreamingSources(
+static FAutoConsoleVariableRef CVarForceLevelStreaming(
 	TEXT("levelinstance.debug.forcelevelstreaming"),
 	GLevelInstanceDebugForceLevelStreaming,
 	TEXT("Set to 1 to force Level Instance to be streamed instead of embedded in World Partition grid."));
