Prevent UTexture calls from the render thread to avoid potential race condition

diff --git a/Engine/Source/Runtime/Engine/Private/Materials/MaterialShared.cpp b/Engine/Source/Runtime/Engine/Private/Materials/MaterialShared.cpp
--- a/Engine/Source/Runtime/Engine/Private/Materials/MaterialShared.cpp
+++ b/Engine/Source/Runtime/Engine/Private/Materials/MaterialShared.cpp
@@ -3113,37 +3113,43 @@ IAllocatedVirtualTexture* FMaterialRenderProxy::AllocateVTStack(const FMaterialR
 		{
 			continue;
 		}
-		if (!Texture->IsCurrentlyVirtualTextured())
+
+		// GetResource() is safe to call from the render thread.
+		const FTextureResource* TextureResource = Texture->GetResource();
+		if (TextureResource)
 		{
-			// The placeholder used during async texture compilation is expected to be of the wrong type since
-			// no VT infos are available until later in the compilation process. This will be resolved
-			// once the final texture resource is available.
+			const FVirtualTexture2DResource* VirtualTextureResourceForLayer = TextureResource->GetVirtualTexture2DResource();
+
+			if (VirtualTextureResourceForLayer == nullptr)
+			{
+				// The placeholder used during async texture compilation is expected to be of the wrong type since
+				// no VT infos are available until later in the compilation process. This will be resolved
+				// once the final texture resource is available.
 #if WITH_EDITOR
-			if (!Texture->IsDefaultTexture())
+				if (!TextureResource->IsProxy())
 #endif
-			{
-				UE_LOG(LogMaterial, Warning, TEXT("Material '%s' expects texture '%s' to be Virtual"),
-					*GetFriendlyName(), *Texture->GetName());
+				{
+					UE_LOG(LogMaterial, Warning, TEXT("Material '%s' expects texture '%s' to be Virtual"),
+						*GetFriendlyName(), *Texture->GetName());
+				}
+				continue;
 			}
-			continue;
-		}
-
-		const FVirtualTexture2DResource* VirtualTextureResourceForLayer = (FVirtualTexture2DResource*)Texture->GetResource();
-		if (VirtualTextureResourceForLayer != nullptr)
-		{
-			// All tile sizes need to match
-			check(!bFoundValidLayer || VTDesc.TileSize == VirtualTextureResourceForLayer->GetTileSize());
-			check(!bFoundValidLayer || VTDesc.TileBorderSize == VirtualTextureResourceForLayer->GetBorderSize());
-
-			const FVirtualTextureProducerHandle& ProducerHandle = VirtualTextureResourceForLayer->GetProducerHandle();
-			if (ProducerHandle.IsValid())
+			else
 			{
-				VTDesc.TileSize = VirtualTextureResourceForLayer->GetTileSize();
-				VTDesc.TileBorderSize = VirtualTextureResourceForLayer->GetBorderSize();
-				VTDesc.ProducerHandle[LayerIndex] = ProducerHandle;
-				VTDesc.ProducerLayerIndex[LayerIndex] = 0u;
-				GetRendererModule().AddVirtualTextureProducerDestroyedCallback(ProducerHandle, &OnVirtualTextureDestroyedCB, const_cast<FMaterialRenderProxy*>(this));
-				bFoundValidLayer = true;
+				// All tile sizes need to match
+				check(!bFoundValidLayer || VTDesc.TileSize == VirtualTextureResourceForLayer->GetTileSize());
+				check(!bFoundValidLayer || VTDesc.TileBorderSize == VirtualTextureResourceForLayer->GetBorderSize());
+
+				const FVirtualTextureProducerHandle& ProducerHandle = VirtualTextureResourceForLayer->GetProducerHandle();
+				if (ProducerHandle.IsValid())
+				{
+					VTDesc.TileSize = VirtualTextureResourceForLayer->GetTileSize();
+					VTDesc.TileBorderSize = VirtualTextureResourceForLayer->GetBorderSize();
+					VTDesc.ProducerHandle[LayerIndex] = ProducerHandle;
+					VTDesc.ProducerLayerIndex[LayerIndex] = 0u;
+					GetRendererModule().AddVirtualTextureProducerDestroyedCallback(ProducerHandle, &OnVirtualTextureDestroyedCB, const_cast<FMaterialRenderProxy*>(this));
+					bFoundValidLayer = true;
+				}
 			}
 		}
 	}
diff --git a/Engine/Source/Runtime/Engine/Public/TextureResource.h b/Engine/Source/Runtime/Engine/Public/TextureResource.h
--- a/Engine/Source/Runtime/Engine/Public/TextureResource.h
+++ b/Engine/Source/Runtime/Engine/Public/TextureResource.h
@@ -31,6 +31,7 @@ class FStreamableTextureResource;
 class FTexture2DResource;
 class FTexture3DResource;
 class FTexture2DArrayResource;
+class FVirtualTexture2DResource;
 
 /** Maximum number of slices in texture source art. */
 #define MAX_TEXTURE_SOURCE_SLICES 6
@@ -100,11 +101,13 @@ public:
 	ENGINE_API virtual FTexture3DResource* GetTexture3DResource() { return nullptr; }
 	ENGINE_API virtual FTexture2DArrayResource* GetTexture2DArrayResource() { return nullptr; }
 	ENGINE_API virtual FStreamableTextureResource* GetStreamableTextureResource() { return nullptr; }
+	ENGINE_API virtual FVirtualTexture2DResource* GetVirtualTexture2DResource() { return nullptr; }
 	// Dynamic cast methods (const).
 	ENGINE_API virtual const FTexture2DResource* GetTexture2DResource() const { return nullptr; }
 	ENGINE_API virtual const FTexture3DResource* GetTexture3DResource() const { return nullptr; }
 	ENGINE_API virtual const FTexture2DArrayResource* GetTexture2DArrayResource() const { return nullptr; }
 	ENGINE_API virtual const FStreamableTextureResource* GetStreamableTextureResource() const { return nullptr; }
+	ENGINE_API virtual const FVirtualTexture2DResource* GetVirtualTexture2DResource() const { return nullptr; }
 
 	// Current mip count. We use "current" to specify that it is not computed from SizeX() which is the size when fully streamed in.
 	FORCEINLINE int32 GetCurrentMipCount() const
@@ -156,6 +159,11 @@ public:
 
 	virtual void InitRHI() override;
 	virtual void ReleaseRHI() override;
+	
+	// Dynamic cast methods.
+	ENGINE_API virtual FVirtualTexture2DResource* GetVirtualTexture2DResource() { return this; }
+	// Dynamic cast methods (const).
+	ENGINE_API virtual const FVirtualTexture2DResource* GetVirtualTexture2DResource() const { return this; }
 
 #if WITH_EDITOR
 	void InitializeEditorResources(class IVirtualTexture* InVirtualTexture);
