Fix/Reduce to necessary shading model when converting legacy material to strata.

diff --git a/Engine/Source/Runtime/Engine/Private/Materials/Material.cpp b/Engine/Source/Runtime/Engine/Private/Materials/Material.cpp
--- a/Engine/Source/Runtime/Engine/Private/Materials/Material.cpp
+++ b/Engine/Source/Runtime/Engine/Private/Materials/Material.cpp
@@ -82,6 +82,7 @@
 #include "Materials/StrataMaterial.h"
 #include "Materials/MaterialExpressionThinTranslucentMaterialOutput.h"
 #include "Materials/MaterialExpressionClearCoatNormalCustomOutput.h"
+#include "Materials/MaterialExpressionTangentOutput.h"
 #include "Materials/MaterialExpressionConstant.h"
 #include "Materials/MaterialExpressionBreakMaterialAttributes.h"
 
@@ -2709,6 +2710,22 @@ void UMaterial::BackwardsCompatibilityDecalConversion()
 #endif // WITH_EDITOR
 }
 
+static void AddStrataShadingModelFromMaterialShadingModel(FStrataMaterialInfo& OutInfo, const FMaterialShadingModelField& InShadingModels)
+{
+	if (InShadingModels.HasShadingModel(MSM_Unlit))				{ OutInfo.AddShadingModel(EStrataShadingModel::SSM_Unlit); }
+	if (InShadingModels.HasShadingModel(MSM_DefaultLit))		{ OutInfo.AddShadingModel(EStrataShadingModel::SSM_DefaultLit); }
+	if (InShadingModels.HasShadingModel(MSM_Subsurface))		{ OutInfo.AddShadingModel(EStrataShadingModel::SSM_SubsurfaceLit); }
+	if (InShadingModels.HasShadingModel(MSM_PreintegratedSkin))	{ OutInfo.AddShadingModel(EStrataShadingModel::SSM_SubsurfaceLit); }
+	if (InShadingModels.HasShadingModel(MSM_ClearCoat))			{ OutInfo.AddShadingModel(EStrataShadingModel::SSM_DefaultLit); }
+	if (InShadingModels.HasShadingModel(MSM_SubsurfaceProfile))	{ OutInfo.AddShadingModel(EStrataShadingModel::SSM_SubsurfaceLit); }
+	if (InShadingModels.HasShadingModel(MSM_TwoSidedFoliage))	{ OutInfo.AddShadingModel(EStrataShadingModel::SSM_SubsurfaceLit); }
+	if (InShadingModels.HasShadingModel(MSM_Hair))				{ OutInfo.AddShadingModel(EStrataShadingModel::SSM_Hair); }
+	if (InShadingModels.HasShadingModel(MSM_Cloth))				{ OutInfo.AddShadingModel(EStrataShadingModel::SSM_DefaultLit); }
+	if (InShadingModels.HasShadingModel(MSM_Eye))				{ OutInfo.AddShadingModel(EStrataShadingModel::SSM_SubsurfaceLit); }
+	if (InShadingModels.HasShadingModel(MSM_SingleLayerWater))	{ OutInfo.AddShadingModel(EStrataShadingModel::SSM_SingleLayerWater); }
+	if (InShadingModels.HasShadingModel(MSM_ThinTranslucent))	{ OutInfo.AddShadingModel(EStrataShadingModel::SSM_DefaultLit); }
+}
+
 void UMaterial::ConvertMaterialToStrataMaterial()
 {
 #if WITH_EDITOR
@@ -2736,6 +2753,11 @@ void UMaterial::ConvertMaterialToStrataMaterial()
 		}
 	};
 
+	// SSS Profile
+	const bool bHasShadingModelMixture		= ShadingModels.CountShadingModels() > 1;
+	const bool bRequireSubsurfacePasses		= ShadingModels.HasShadingModel(MSM_SubsurfaceProfile) || ShadingModels.HasShadingModel(MSM_Subsurface) || ShadingModels.HasShadingModel(MSM_PreintegratedSkin) || ShadingModels.HasShadingModel(MSM_Eye);
+	const bool bRequireNoSubsurfaceProfile	= !bHasShadingModelMixture && (ShadingModel == MSM_Subsurface || ShadingModel == MSM_PreintegratedSkin); // Insure there is no profile, as this would take priority otherwise
+
 	bool bInvalidateShader = false;
 	// Connect all the legacy pin into the conversion node
 	if (bUseMaterialAttributes && MaterialAttributes.Expression && !MaterialAttributes.Expression->IsResultStrataMaterial(MaterialAttributes.OutputIndex)) // M_Rifle cause issues there
@@ -2757,6 +2779,7 @@ void UMaterial::ConvertMaterialToStrataMaterial()
 		ConvertNode->ClearCoatRoughness.Connect(13, BreakMatAtt);
 		ConvertNode->Opacity.Connect(6, BreakMatAtt);
 		ConvertNode->ShadingModel.Connect(25, BreakMatAtt);
+		ConvertNode->SubsurfaceProfile = bRequireNoSubsurfaceProfile ? nullptr : SubsurfaceProfile;
 
 		// * Remove support for material attribute
 		// * explicitely connect the Strata node to the root node
@@ -2768,25 +2791,25 @@ void UMaterial::ConvertMaterialToStrataMaterial()
 		AmbientOcclusion.Connect(14, BreakMatAtt);
 		PixelDepthOffset.Connect(24, BreakMatAtt);
 
-		// SSS Profile
-		const bool bHasShadingModelMixture = ShadingModels.CountShadingModels() > 1;
-		const bool bRequireSubsurfacePasses = ShadingModels.HasShadingModel(MSM_SubsurfaceProfile) || ShadingModels.HasShadingModel(MSM_Subsurface) || ShadingModels.HasShadingModel(MSM_PreintegratedSkin);
-		const bool bRequireNoSubsurfaceProfile = !bHasShadingModelMixture && (ShadingModel == MSM_Subsurface || ShadingModel == MSM_PreintegratedSkin); // Insure there is no profile, as this would take priority otherwise
-		if (SubsurfaceProfile != nullptr)
+		if (ShadingModel == MSM_FromMaterialExpression)
 		{
-			ShadingModel = MSM_SubsurfaceProfile;
+			ConvertNode->ConvertedStrataMaterialInfo.SetShadingModelFromExpression(true);
+			AddStrataShadingModelFromMaterialShadingModel(ConvertNode->ConvertedStrataMaterialInfo, ShadingModels);
+			check(ConvertNode->ConvertedStrataMaterialInfo.CountShadingModels() >= 1);
 		}
-		ConvertNode->SubsurfaceProfile = bRequireNoSubsurfaceProfile ? nullptr : SubsurfaceProfile;
-
-		// Shading Model
-		// Note: store this conversion type(s) into ConvertedStrataMaterialInfo for having more context when 
-		// rebuilding the final Shading model (see RebuildShadingModelField())
+		else
 		{
-			ConvertNode->ConvertedStrataMaterialInfo.AddShadingModel(SSM_Unlit);
-			ConvertNode->ConvertedStrataMaterialInfo.AddShadingModel(SSM_DefaultLit);
-			ConvertNode->ConvertedStrataMaterialInfo.AddShadingModel(SSM_SubsurfaceLit);
-			ConvertNode->ConvertedStrataMaterialInfo.AddShadingModel(SSM_Hair);
-			ConvertNode->ConvertedStrataMaterialInfo.SetShadingModelFromExpression(true);
+			check(!bHasShadingModelMixture);
+
+			// Add constant for the shading model
+			UMaterialExpressionConstant* ShadingModelNode = NewObject<UMaterialExpressionConstant>(this);
+			ShadingModelNode->SetParameterName(FName(TEXT("ConstantShadingModel")));
+			ShadingModelNode->R = ShadingModel;
+			ConvertNode->ShadingModel.Connect(0, ShadingModelNode);
+
+			// Store strata shading model of the converted material. 
+			AddStrataShadingModelFromMaterialShadingModel(ConvertNode->ConvertedStrataMaterialInfo, ShadingModels);
+			check(ConvertNode->ConvertedStrataMaterialInfo.CountShadingModels() == 1);
 		}
 
 		bInvalidateShader = true;
@@ -2806,6 +2829,8 @@ void UMaterial::ConvertMaterialToStrataMaterial()
 			UMaterialExpressionThinTranslucentMaterialOutput* ThinTranslucentOutput = nullptr;
 			UMaterialExpressionSingleLayerWaterMaterialOutput* SingleLayerWaterOutput = nullptr;
 			UMaterialExpressionClearCoatNormalCustomOutput* ClearCoatOutput = nullptr;
+			UMaterialExpressionTangentOutput* TangentOutput = nullptr;
+
 			for (UMaterialExpressionCustomOutput* Expression : CustomOutputExpressions)
 			{
 				// Gather custom output for thin translucency
@@ -2826,19 +2851,16 @@ void UMaterial::ConvertMaterialToStrataMaterial()
 					ClearCoatOutput = Cast<UMaterialExpressionClearCoatNormalCustomOutput>(Expression);
 				}
 
-				if (ThinTranslucentOutput && SingleLayerWaterOutput && ClearCoatOutput)
+				// Gather custom output for tangent (unused atm)
+				if (TangentOutput == nullptr && Cast<UMaterialExpressionTangentOutput>(Expression))
 				{
-					break;
+					TangentOutput = Cast<UMaterialExpressionTangentOutput>(Expression);
 				}
-			}
 
-			// SSS Profile
-			const bool bHasShadingModelMixture = ShadingModels.CountShadingModels() > 1;
-			const bool bRequireSubsurfacePasses = ShadingModels.HasShadingModel(MSM_SubsurfaceProfile) || ShadingModels.HasShadingModel(MSM_Subsurface) || ShadingModels.HasShadingModel(MSM_PreintegratedSkin) || ShadingModels.HasShadingModel(MSM_Eye);
-			const bool bRequireNoSubsurfaceProfile = !bHasShadingModelMixture && (ShadingModel == MSM_Subsurface || ShadingModel == MSM_PreintegratedSkin); // Insure there is no profile, as this would take priority otherwise
-			if (SubsurfaceProfile != nullptr)
-			{
-				ShadingModel = MSM_SubsurfaceProfile;
+				if (ThinTranslucentOutput && SingleLayerWaterOutput && ClearCoatOutput && TangentOutput)
+				{
+					break;
+				}
 			}
 
 			UMaterialExpressionStrataLegacyConversion* ConvertNode = NewObject<UMaterialExpressionStrataLegacyConversion>(this);
@@ -2877,21 +2899,22 @@ void UMaterial::ConvertMaterialToStrataMaterial()
 			// 
 			// Note: store this conversion type(s) into ConvertedStrataMaterialInfo for having more context when 
 			// rebuilding the final Shading model (see RebuildShadingModelField())
-			if (ShadingModelFromMaterialExpression.IsConnected())
+			if (ShadingModel == MSM_FromMaterialExpression)
 			{
+				check(ShadingModelFromMaterialExpression.IsConnected());
+
 				// Reconnect the shading model expression
 				MoveConnectionTo(ShadingModelFromMaterialExpression, ConvertNode, 18);
 
 				// Store strata shading model of the converted material. 
-				ConvertNode->ConvertedStrataMaterialInfo.AddShadingModel(SSM_Unlit);
-				ConvertNode->ConvertedStrataMaterialInfo.AddShadingModel(SSM_DefaultLit);
-				ConvertNode->ConvertedStrataMaterialInfo.AddShadingModel(SSM_SubsurfaceLit);
-				ConvertNode->ConvertedStrataMaterialInfo.AddShadingModel(SSM_Hair);
 				if (SingleLayerWaterOutput)
 				{
 					ConvertNode->ConvertedStrataMaterialInfo.AddShadingModel(SSM_SingleLayerWater);
 				}
+
 				ConvertNode->ConvertedStrataMaterialInfo.SetShadingModelFromExpression(true);
+				AddStrataShadingModelFromMaterialShadingModel(ConvertNode->ConvertedStrataMaterialInfo, ShadingModels);
+				check(ConvertNode->ConvertedStrataMaterialInfo.CountShadingModels() >= 1);
 			}
 			else
 			{
@@ -2902,26 +2925,9 @@ void UMaterial::ConvertMaterialToStrataMaterial()
 				ShadingModelNode->SetParameterName(FName(TEXT("ConstantShadingModel")));
 				ShadingModelNode->R = ShadingModel;
 				ConvertNode->ShadingModel.Connect(0, ShadingModelNode);
-				
-				// Store strata shading model of the converted material. 
-				EStrataShadingModel StrataShadingModel = EStrataShadingModel::SSM_NUM;
-				switch (ShadingModel)
-				{
-					case MSM_Unlit:				StrataShadingModel = EStrataShadingModel::SSM_Unlit; break;
-					case MSM_DefaultLit:		StrataShadingModel = EStrataShadingModel::SSM_DefaultLit; break;
-					case MSM_Subsurface:		StrataShadingModel = EStrataShadingModel::SSM_SubsurfaceLit; break;
-					case MSM_PreintegratedSkin: StrataShadingModel = EStrataShadingModel::SSM_SubsurfaceLit; break;
-					case MSM_ClearCoat:			StrataShadingModel = EStrataShadingModel::SSM_DefaultLit; break;
-					case MSM_SubsurfaceProfile: StrataShadingModel = EStrataShadingModel::SSM_SubsurfaceLit; break;
-					case MSM_TwoSidedFoliage:	StrataShadingModel = EStrataShadingModel::SSM_SubsurfaceLit; break;
-					case MSM_Hair:				StrataShadingModel = EStrataShadingModel::SSM_Hair; break;
-					case MSM_Cloth:				StrataShadingModel = EStrataShadingModel::SSM_DefaultLit; break;
-					case MSM_Eye:				StrataShadingModel = EStrataShadingModel::SSM_SubsurfaceLit; break;
-					case MSM_SingleLayerWater:	StrataShadingModel = EStrataShadingModel::SSM_SingleLayerWater; break;
-					case MSM_ThinTranslucent:	StrataShadingModel = EStrataShadingModel::SSM_DefaultLit; break;
-				};
-				check(StrataShadingModel != EStrataShadingModel::SSM_NUM);
-				ConvertNode->ConvertedStrataMaterialInfo.AddShadingModel(StrataShadingModel);
+
+				AddStrataShadingModelFromMaterialShadingModel(ConvertNode->ConvertedStrataMaterialInfo, ShadingModels);
+				check(ConvertNode->ConvertedStrataMaterialInfo.CountShadingModels() == 1);
 			}
 
 			FrontMaterial.Connect(0, ConvertNode);
@@ -3903,19 +3909,28 @@ void UMaterial::RebuildShadingModelField()
 		{
 			MaterialDomain = EMaterialDomain::MD_Surface;
 
-			// Shading models
-			if (StrataMaterialInfo.HasShadingModel(SSM_Unlit))				{ ShadingModels.AddShadingModel(MSM_Unlit); }
-			if (StrataMaterialInfo.HasShadingModel(SSM_DefaultLit))			{ ShadingModels.AddShadingModel(MSM_DefaultLit); }
-			if (StrataMaterialInfo.HasShadingModel(SSM_SubsurfaceLit))		{ ShadingModels.AddShadingModel(MSM_SubsurfaceProfile); }
-			if (StrataMaterialInfo.HasShadingModel(SSM_Hair))				{ ShadingModels.AddShadingModel(MSM_Hair); }
-			if (StrataMaterialInfo.HasShadingModel(SSM_SingleLayerWater))	{ ShadingModels.AddShadingModel(MSM_SingleLayerWater); }
+			check(ShadingModel == MSM_FromMaterialExpression);
+			{
+				TArray<UMaterialExpressionShadingModel*> ShadingModelExpressions;
+				GetAllExpressionsInMaterialAndFunctionsOfType(ShadingModelExpressions);
+
+				for (UMaterialExpressionShadingModel* MatExpr : ShadingModelExpressions)
+				{
+					ShadingModels.AddShadingModel(MatExpr->ShadingModel);
+				}
+
+				// If no expressions have been found, set a default
+				if (!ShadingModels.IsValid())
+				{
+					ShadingModels.AddShadingModel(MSM_DefaultLit);
+				}
+			}
 
 			// Blend mode
 			// Unclear what best fallback it should be
 			if (BlendMode != EBlendMode::BLEND_Opaque && BlendMode != EBlendMode::BLEND_Masked)
 			{
 				BlendMode = EBlendMode::BLEND_Translucent; // This is to be able to use dual-source blending
-				//BlendMode = EBlendMode::BLEND_Opaque;
 			}
 
 			// Subsurface profil
