Optimized memory usage for skeletal remapping using a lazy evaluation approach.
The curve remapping backup table is now only copied when skeleton remapping of curves is active.
This should save some memory in case this isn't used (most cases).

diff --git a/Engine/Source/Runtime/Engine/Private/Animation/AnimSequence.cpp b/Engine/Source/Runtime/Engine/Private/Animation/AnimSequence.cpp
--- a/Engine/Source/Runtime/Engine/Private/Animation/AnimSequence.cpp
+++ b/Engine/Source/Runtime/Engine/Private/Animation/AnimSequence.cpp
@@ -1482,7 +1482,7 @@ void UAnimSequence::GetBonePose(FAnimationPoseData& OutAnimationPoseData, const
 
 	FCompactPose& OutPose = OutAnimationPoseData.GetPose();
 
-	const FBoneContainer& RequiredBones = OutPose.GetBoneContainer();
+	FBoneContainer& RequiredBones = OutPose.GetBoneContainer();
 	const bool bUseRawDataForPoseExtraction = (CanEvaluateRawAnimationData() && bForceUseRawData) || UseRawDataForPoseExtraction(RequiredBones);
 
 	const bool bIsBakedAdditive = !bUseRawDataForPoseExtraction && IsValidAdditive();
diff --git a/Engine/Source/Runtime/Engine/Private/BoneContainer.cpp b/Engine/Source/Runtime/Engine/Private/BoneContainer.cpp
--- a/Engine/Source/Runtime/Engine/Private/BoneContainer.cpp
+++ b/Engine/Source/Runtime/Engine/Private/BoneContainer.cpp
@@ -15,7 +15,7 @@ DEFINE_LOG_CATEGORY(LogSkeletalControl);
 //////////////////////////////////////////////////////////////////////////
 // FSkeletonRemappingCurve
 
-FSkeletonRemappingCurve::FSkeletonRemappingCurve(FBlendedCurve& InCurve, const FBoneContainer& InBoneContainer, const FSkeletonRemapping* InSkeletonRemapping)
+FSkeletonRemappingCurve::FSkeletonRemappingCurve(FBlendedCurve& InCurve, FBoneContainer& InBoneContainer, const FSkeletonRemapping* InSkeletonRemapping)
 	: Curve(InCurve)
 	, BoneContainer(InBoneContainer)
 	, bIsRemapping(true)
@@ -25,7 +25,7 @@ FSkeletonRemappingCurve::FSkeletonRemappingCurve(FBlendedCurve& InCurve, const F
 	Curve.UIDToArrayIndexLUT = &CurveMapping.UIDToArrayIndices;
 }
 
-FSkeletonRemappingCurve::FSkeletonRemappingCurve(FBlendedCurve& InCurve, const FBoneContainer& InBoneContainer, const USkeleton* SourceSkeleton)
+FSkeletonRemappingCurve::FSkeletonRemappingCurve(FBlendedCurve& InCurve, FBoneContainer& InBoneContainer, const USkeleton* SourceSkeleton)
 	: Curve(InCurve)
 	, BoneContainer(InBoneContainer)
 {
@@ -244,6 +244,11 @@ void FBoneContainer::Initialize(const FCurveEvaluationOption& CurveEvalOption)
 void FBoneContainer::CacheRequiredAnimCurveUids(const FCurveEvaluationOption& CurveEvalOption)
 {
 	LLM_SCOPE_BYNAME(TEXT("Animation/BoneContainer"));
+
+	// Clear the skeleton remapping backup LUT.
+	// We will initialize this only when needed, using a lazy approach.
+	UIDToArrayIndexLUTBackup.Reset();
+
 	if (AssetSkeleton.IsValid())
 	{
 		// this is placeholder. In the future, this will change to work with linked joint of curve meta data
@@ -348,9 +353,6 @@ void FBoneContainer::CacheRequiredAnimCurveUids(const FCurveEvaluationOption& Cu
 		UIDToArrayIndexLUTValidCount = 0;
 	}
 
-	// Make a backup, used for skeleton remapping of curves.
-	UIDToArrayIndexLUTBackup = UIDToArrayIndexLUT;
-
 	// Make sure we regenerate our cached curve mappings next time they are requested.
 	MarkAllCachedCurveMappingsDirty();
 
@@ -377,7 +379,7 @@ void FBoneContainer::MarkAllCachedCurveMappingsDirty()
 	}
 }
 
-const FCachedSkeletonCurveMapping& FBoneContainer::GetOrCreateCachedCurveMapping(const FSkeletonRemapping* SkeletonRemapping) const
+const FCachedSkeletonCurveMapping& FBoneContainer::GetOrCreateCachedCurveMapping(const FSkeletonRemapping* SkeletonRemapping)
 {
 	LLM_SCOPE_BYNAME(TEXT("Animation/BoneContainer"));
 	
@@ -386,6 +388,12 @@ const FCachedSkeletonCurveMapping& FBoneContainer::GetOrCreateCachedCurveMapping
 	const USkeleton* SourceSkeleton = SkeletonRemapping->GetSourceSkeleton().Get();
 	check(SourceSkeleton);
 
+	// Make a backup, used for skeleton remapping of curves.
+	if (UIDToArrayIndexLUTBackup.IsEmpty())
+	{
+		UIDToArrayIndexLUTBackup = UIDToArrayIndexLUT;
+	}
+
 	// Check if we already have some cached data for this skeleton.
 	FCachedSkeletonCurveMapping* CachedData = CachedCurveMappingTable.Find(SourceSkeleton);
 	if (CachedData)
diff --git a/Engine/Source/Runtime/Engine/Public/BoneContainer.h b/Engine/Source/Runtime/Engine/Public/BoneContainer.h
--- a/Engine/Source/Runtime/Engine/Public/BoneContainer.h
+++ b/Engine/Source/Runtime/Engine/Public/BoneContainer.h
@@ -139,8 +139,8 @@ public:
 	FSkeletonRemappingCurve() = delete;
 	FSkeletonRemappingCurve(const FSkeletonRemappingCurve&) = delete;
 	FSkeletonRemappingCurve(FSkeletonRemappingCurve&&) = delete;
-	FSkeletonRemappingCurve(FBlendedCurve& InCurve, const FBoneContainer& InBoneContainer, const FSkeletonRemapping* InSkeletonMapping);
-	FSkeletonRemappingCurve(FBlendedCurve& InCurve, const FBoneContainer& InBoneContainer, const USkeleton* SourceSkeleton);
+	FSkeletonRemappingCurve(FBlendedCurve& InCurve, FBoneContainer& InBoneContainer, const FSkeletonRemapping* InSkeletonMapping);
+	FSkeletonRemappingCurve(FBlendedCurve& InCurve, FBoneContainer& InBoneContainer, const USkeleton* SourceSkeleton);
 	~FSkeletonRemappingCurve();
 
 	FSkeletonRemappingCurve& operator = (const FSkeletonRemappingCurve&) = delete;
@@ -151,7 +151,7 @@ public:
 
 private:
 	FBlendedCurve& Curve;
-	const FBoneContainer& BoneContainer;
+	FBoneContainer& BoneContainer;
 	bool bIsRemapping = false;
 };
 
@@ -648,7 +648,7 @@ public:
 #endif
 	
 	// Curve remapping
-	const FCachedSkeletonCurveMapping& GetOrCreateCachedCurveMapping(const FSkeletonRemapping* SkeletonRemapping) const;
+	const FCachedSkeletonCurveMapping& GetOrCreateCachedCurveMapping(const FSkeletonRemapping* SkeletonRemapping);
 	void MarkAllCachedCurveMappingsDirty();
 
 	// Get the serial number of this bone container. @see SerialNumber
diff --git a/Engine/Source/Runtime/Engine/Public/BonePose.h b/Engine/Source/Runtime/Engine/Public/BonePose.h
--- a/Engine/Source/Runtime/Engine/Public/BonePose.h
+++ b/Engine/Source/Runtime/Engine/Public/BonePose.h
@@ -140,6 +140,12 @@ public:
 		return *BoneContainer;
 	}
 
+	FBoneContainer& GetBoneContainer()
+	{
+		checkSlow(BoneContainer && BoneContainer->IsValid());
+		return *const_cast<FBoneContainer*>(BoneContainer);
+	}
+
 	void SetBoneContainer(const FBoneContainer* InBoneContainer)
 	{
 		check(InBoneContainer && InBoneContainer->IsValid());
