Remove skinning pre/post offset functionality.
This were added in 4.26 specifically and temporarily for the MLDeformer prototype. (The public functions have always been marked as deprecated.)
MLDeformer now uses an Optimus deformer.

diff --git a/Engine/Shaders/Private/GpuSkinCacheComputeShader.usf b/Engine/Shaders/Private/GpuSkinCacheComputeShader.usf
--- a/Engine/Shaders/Private/GpuSkinCacheComputeShader.usf
+++ b/Engine/Shaders/Private/GpuSkinCacheComputeShader.usf
@@ -23,35 +23,6 @@
 
 STRONG_TYPE Buffer<float4> BoneMatrices;
 
-uint VertexOffsetUsage;
-STRONG_TYPE Buffer<float> PreSkinOffsets;
-STRONG_TYPE Buffer<float> PostSkinOffsets;
-
-void GetVertexOffsets(uint VertexIndex, out float3 PreSkinOffset, out float3 PostSkinOffset)
- {
-	uint Offset = VertexIndex * 3;
-
-	PreSkinOffset = 0;
-	if (VertexOffsetUsage & 1)
-	{
-		PreSkinOffset = float3(
-			PreSkinOffsets[Offset + 0],
-			PreSkinOffsets[Offset + 1],
-			PreSkinOffsets[Offset + 2]
-			);
-	}
-
-	PostSkinOffset = 0;
-	if (VertexOffsetUsage & 2)
-	{
-		PostSkinOffset = float3(
-			PostSkinOffsets[Offset + 0],
-			PostSkinOffsets[Offset + 1],
-			PostSkinOffsets[Offset + 2]
-			);
-	}
-}
-
 #if GPUSKIN_MORPH_BLEND
 // float3 DeltaPosition, PackedNormal, see FMorphGPUSkinVertex
 STRONG_TYPE Buffer<float> MorphBuffer;
@@ -148,9 +119,6 @@ struct FVertexUnpacked
 	half3	DeltaTangentZ;
 #endif
 
-	float3 PreSkinOffset;
-	float3 PostSkinOffset;
-
 #if GPUSKIN_APEX_CLOTH
     FVertexTriangleInfluence Influences[NUM_INFLUENCES_PER_VERTEX];
 #endif
@@ -293,8 +261,6 @@ FVertexUnpacked UnpackedVertex(uint MorphIndex, uint ClothIndex, uint ClothIndex
 	GetMorphBuffer(MorphIndex, Unpacked.DeltaPosition, Unpacked.DeltaTangentZ);
 #endif
 
-	GetVertexOffsets(VertexIndex, Unpacked.PreSkinOffset, Unpacked.PostSkinOffset);
-
 #if GPUSKIN_APEX_CLOTH
     GetClothBuffer(ClothIndex, ClothIndexOffset, Unpacked.Influences);
 #endif
@@ -460,9 +426,7 @@ void SkinCacheUpdateBatchCS(uint3 GroupID : SV_GroupID,
 #endif // GPUSKIN_APEX_CLOTH
 #endif // GPUSKIN_MORPH_BLEND
 
-	Intermediates.UnpackedPosition += Unpacked.PreSkinOffset;
 	float3 SPos = SkinPosition(Intermediates) * SkinCacheDebug;
-	SPos += Unpacked.PostSkinOffset;
 
 #if GPUSKIN_APEX_CLOTH
     float4 TransformedSkinPos = mul(float4(SPos, 1), ClothLocalToWorld);
diff --git a/Engine/Shaders/Private/GpuSkinVertexFactory.ush b/Engine/Shaders/Private/GpuSkinVertexFactory.ush
--- a/Engine/Shaders/Private/GpuSkinVertexFactory.ush
+++ b/Engine/Shaders/Private/GpuSkinVertexFactory.ush
@@ -106,9 +106,6 @@ struct FVertexFactoryInput
 	float3	DeltaTangentZ	: ATTRIBUTE10;	//NORMAL1;
 #endif
 
-	float3 PreSkinOffset : ATTRIBUTE11;
-	float3 PostSkinOffset : ATTRIBUTE12;
-
 #if GPUSKIN_MORPH_BLEND || GPUSKIN_APEX_CLOTH
 	uint VertexID : SV_VertexID;
 #endif
@@ -333,9 +330,6 @@ FMaterialVertexParameters GetMaterialVertexParameters(FVertexFactoryInput Input,
 	Result.PreSkinnedPosition = Intermediates.UnpackedPosition.xyz;
 	Result.PreSkinnedNormal = Input.TangentZ.xyz;
 
-	Result.PreSkinOffset = Input.PreSkinOffset;
-	Result.PostSkinOffset = Input.PostSkinOffset;
-
 	// Assumes no instacing
 	Result.PrevFrameLocalToWorld = Intermediates.PrevLocalToWorld;
 
@@ -638,14 +632,10 @@ float3 SkinPosition( FVertexFactoryInput Input, FVertexFactoryIntermediates Inte
 	float3 Position = Intermediates.UnpackedPosition;
 #endif
 
-	Position += Input.PreSkinOffset;
-
 	// Note the use of mul(Matrix,Vector), bone matrices are stored transposed
 	// for tighter packing.
 	Position = mul(Intermediates.BlendMatrix, float4(Position, 1));
 
-	Position += Input.PostSkinOffset;
-
 	return Position;
 }
 
@@ -667,14 +657,10 @@ float3 SkinPreviousPosition( FVertexFactoryInput Input, FVertexFactoryIntermedia
 		BlendMatrix = CalcPreviousBoneMatrix( Input );
 	}
 
-	Position += Input.PreSkinOffset;
-
 	// Note the use of mul(Matrix,Vector), bone matrices are stored transposed
 	// for tighter packing.
 	Position = mul(BlendMatrix, float4(Position, 1));
 
-	Position += Input.PostSkinOffset;
-
 	return Position;
 }
 /** transform the tangent basis vectors */
diff --git a/Engine/Shaders/Private/MaterialTemplate.ush b/Engine/Shaders/Private/MaterialTemplate.ush
--- a/Engine/Shaders/Private/MaterialTemplate.ush
+++ b/Engine/Shaders/Private/MaterialTemplate.ush
@@ -501,11 +501,6 @@ struct FMaterialVertexParameters
 	float3 PreSkinnedPosition;
 	float3 PreSkinnedNormal;
 
-#if GPU_SKINNED_MESH_FACTORY
-	float3 PreSkinOffset;
-	float3 PostSkinOffset;
-#endif
-
 	half4 VertexColor;
 #if NUM_MATERIAL_TEXCOORDS_VERTEX
 	float2 TexCoords[NUM_MATERIAL_TEXCOORDS_VERTEX];
@@ -1586,24 +1581,6 @@ float MaterialExpressionVolumeSampleShadowSampleDistance(FMaterialPixelParameter
 #endif
 }
 
-float3 MaterialExpressionPreSkinOffset(FMaterialVertexParameters Parameters)
-{
-#if GPU_SKINNED_MESH_FACTORY
-	return Parameters.PreSkinOffset;
-#else
-	return 0;
-#endif
-}
-
-float3 MaterialExpressionPostSkinOffset(FMaterialVertexParameters Parameters)
-{
-#if GPU_SKINNED_MESH_FACTORY
-	return Parameters.PostSkinOffset;
-#else
-	return 0;
-#endif
-}
-
 // We use a forward declaration towards the real implementation of functions we need. This is always defined for base passes (opaque and translucent)
 #if defined(IS_BASE_PASS) && IS_MOBILE_BASE_PASS==0
 float3 GetSkyLightReflectionSupportingBlend(float3 ReflectionVector, float Roughness, out float OutSkyAverageBrightness);
diff --git a/Engine/Source/Developer/MaterialBaking/Private/ExportMaterialProxy.h b/Engine/Source/Developer/MaterialBaking/Private/ExportMaterialProxy.h
--- a/Engine/Source/Developer/MaterialBaking/Private/ExportMaterialProxy.h
+++ b/Engine/Source/Developer/MaterialBaking/Private/ExportMaterialProxy.h
@@ -132,16 +132,6 @@ struct FExportMaterialCompiler : public FProxyMaterialCompiler
 		return Compiler->VertexColor();
 	}
 
-	virtual int32 PreSkinVertexOffset() override
-	{
-		return Compiler->PreSkinVertexOffset();
-	}
-
-	virtual int32 PostSkinVertexOffset() override
-	{
-		return Compiler->PostSkinVertexOffset();
-	}
-
 	virtual int32 PreSkinnedPosition() override
 	{
 		return Compiler->PreSkinnedPosition();
diff --git a/Engine/Source/Developer/MaterialUtilities/Private/MaterialUtilities.cpp b/Engine/Source/Developer/MaterialUtilities/Private/MaterialUtilities.cpp
--- a/Engine/Source/Developer/MaterialUtilities/Private/MaterialUtilities.cpp
+++ b/Engine/Source/Developer/MaterialUtilities/Private/MaterialUtilities.cpp
@@ -307,16 +307,6 @@ struct FExportMaterialCompiler : public FProxyMaterialCompiler
 		return Compiler->VertexColor(); 
 	}
 
-	virtual int32 PreSkinVertexOffset() override
-	{
-		return Compiler->PreSkinVertexOffset();
-	}
-
-	virtual int32 PostSkinVertexOffset() override
-	{
-		return Compiler->PostSkinVertexOffset();
-	}
-
 	virtual int32 PreSkinnedPosition() override
 	{
 		return Compiler->PreSkinnedPosition();
diff --git a/Engine/Source/Editor/UnrealEd/Private/Lightmass/LightmassRender.cpp b/Engine/Source/Editor/UnrealEd/Private/Lightmass/LightmassRender.cpp
--- a/Engine/Source/Editor/UnrealEd/Private/Lightmass/LightmassRender.cpp
+++ b/Engine/Source/Editor/UnrealEd/Private/Lightmass/LightmassRender.cpp
@@ -180,16 +180,6 @@ struct FLightmassMaterialCompiler : public FProxyMaterialCompiler
 		return Compiler->Constant4(1.0f,1.0f,1.0f,1.0f);
 	}
 
-	virtual int32 PreSkinVertexOffset() override
-	{
-		return Compiler->Constant3(0.f, 0.f, 0.f);
-	}
-
-	virtual int32 PostSkinVertexOffset() override
-	{
-		return Compiler->Constant3(0.f, 0.f, 0.f);
-	}
-
 	virtual int32 PreSkinnedPosition() override
 	{
 		return Compiler->Constant3(0.f,0.f,0.f);
diff --git a/Engine/Source/Runtime/Engine/Classes/Components/SkinnedMeshComponent.h b/Engine/Source/Runtime/Engine/Classes/Components/SkinnedMeshComponent.h
--- a/Engine/Source/Runtime/Engine/Classes/Components/SkinnedMeshComponent.h
+++ b/Engine/Source/Runtime/Engine/Classes/Components/SkinnedMeshComponent.h
@@ -106,15 +106,6 @@ namespace EBoneSpaces
 	};
 }
 
-UENUM(BlueprintType, meta = (Bitflags, UseEnumValuesAsMaskValuesInEditor = "true"))
-enum class EVertexOffsetUsageType : uint8
-{
-	None = 0,
-	PreSkinningOffset = (1 << 0),
-	PostSkinningOffset = (1 << 1),
-};
-ENUM_CLASS_FLAGS(EVertexOffsetUsageType);
-
 /** Struct used to indicate one active morph target that should be applied to this USkeletalMesh when rendered. */
 struct FActiveMorphTarget
 {
@@ -178,9 +169,6 @@ struct ENGINE_API FSkelMeshComponentLODInfo
 	/** Vertex buffer used to override skin weights from one of the profiles */
 	FSkinWeightVertexBuffer* OverrideProfileSkinWeights;
 
-	TArray<FVector3f> PreSkinningOffsets;
-	TArray<FVector3f> PostSkinningOffsets;
-
 	FSkelMeshComponentLODInfo();
 	~FSkelMeshComponentLODInfo();
 
@@ -244,9 +232,6 @@ class ENGINE_API USkinnedMeshComponent : public UMeshComponent, public ILODSyncI
 	UPROPERTY(EditAnywhere, BlueprintReadOnly, Category = "Mesh")
 	TArray<ESkinCacheUsage> SkinCacheUsage;
 
-	UPROPERTY(EditAnywhere, BlueprintReadOnly, Category = "Mesh", meta=(DisplayName="Pre/Post Skin Deltas Usage"))
-	TArray<FVertexOffsetUsage> VertexOffsetUsage;
-
 	/** const getters for previous transform idea */
 	const TArray<uint8>& GetPreviousBoneVisibilityStates() const
 	{
@@ -1043,19 +1028,19 @@ public:
 
 	UE_DEPRECATED(4.26, "GetVertexOffsetUsage() has been deprecated. Support will be dropped in the future.")
 	UFUNCTION(BlueprintCallable, Category = "Components|SkinnedMesh")
-	int32 GetVertexOffsetUsage(int32 LODIndex) const;
+	int32 GetVertexOffsetUsage(int32 LODIndex) const { return 0; }
 
 	UE_DEPRECATED(4.26, "SetVertexOffsetUsage() has been deprecated. Support will be dropped in the future.")
 	UFUNCTION(BlueprintCallable, Category = "Components|SkinnedMesh")
-	void SetVertexOffsetUsage(int32 LODIndex, int32 Usage);
+	void SetVertexOffsetUsage(int32 LODIndex, int32 Usage) {}
 
 	UE_DEPRECATED(4.26, "SetPreSkinningOffsets() has been deprecated. Support will be dropped in the future.")
 	UFUNCTION(BlueprintCallable, Category = "Components|SkinnedMesh")
-	void SetPreSkinningOffsets(int32 LODIndex, TArray<FVector> Offsets);
+	void SetPreSkinningOffsets(int32 LODIndex, TArray<FVector> Offsets) {}
 
 	UE_DEPRECATED(4.26, "SetPostSkinningOffsets() has been deprecated. Support will be dropped in the future.")
 	UFUNCTION(BlueprintCallable, Category = "Components|SkinnedMesh")
-	void SetPostSkinningOffsets(int32 LODIndex, TArray<FVector> Offsets);
+	void SetPostSkinningOffsets(int32 LODIndex, TArray<FVector> Offsets) {}
 
 	/** Check whether or not a Skin Weight Profile is currently pending load / create */
 	bool IsSkinWeightProfilePending() const { return bSkinWeightProfilePending == 1; }
diff --git a/Engine/Source/Runtime/Engine/Classes/Materials/MaterialExpressionSkinningVertexOffsets.h b/Engine/Source/Runtime/Engine/Classes/Materials/MaterialExpressionSkinningVertexOffsets.h
deleted file mode 100644
diff --git a/Engine/Source/Runtime/Engine/Private/Components/SkinnedMeshComponent.cpp b/Engine/Source/Runtime/Engine/Private/Components/SkinnedMeshComponent.cpp
--- a/Engine/Source/Runtime/Engine/Private/Components/SkinnedMeshComponent.cpp
+++ b/Engine/Source/Runtime/Engine/Private/Components/SkinnedMeshComponent.cpp
@@ -3965,89 +3965,6 @@ void USkinnedMeshComponent::SetRenderStatic(bool bNewValue)
 	}
 }
 
-int32 USkinnedMeshComponent::GetVertexOffsetUsage(int32 LODIndex) const
-{
-	if (LODInfo.IsValidIndex(LODIndex))
-	{
-		return (LODIndex < VertexOffsetUsage.Num()) ? VertexOffsetUsage[LODIndex].Usage : 0;
-	}
-
-	return 0;
-}
-
-void USkinnedMeshComponent::SetVertexOffsetUsage(int32 LODIndex, int32 Usage)
-{
-	InitLODInfos();
-
-	if (LODInfo.IsValidIndex(LODIndex))
-	{
-		if (LODIndex >= VertexOffsetUsage.Num())
-		{
-			VertexOffsetUsage.SetNumZeroed(LODIndex+1);
-		}
-
-		VertexOffsetUsage[LODIndex].Usage = Usage;
-
-		if ((Usage & int32(EVertexOffsetUsageType::PreSkinningOffset)) == 0)
-		{
-			LODInfo[LODIndex].PreSkinningOffsets.Empty();
-		}
-
-		if ((Usage & int32(EVertexOffsetUsageType::PostSkinningOffset)) == 0)
-		{
-			LODInfo[LODIndex].PostSkinningOffsets.Empty();
-		}
-
-		MarkRenderStateDirty();
-	}
-}
-
-void USkinnedMeshComponent::SetPreSkinningOffsets(int32 LODIndex, TArray<FVector> Offsets)
-{
-	InitLODInfos();
-
-	FSkeletalMeshRenderData* SkelMeshRenderData = GetSkeletalMeshRenderData();
-
-	// If we have a render resource, and the requested LODIndex is valid (for both component and mesh, though these should be the same)
-	if (SkelMeshRenderData && LODInfo.IsValidIndex(LODIndex) && SkelMeshRenderData->LODRenderData.IsValidIndex(LODIndex))
-	{
-		ensure(LODInfo.Num() == SkelMeshRenderData->LODRenderData.Num());
-
-		FSkelMeshComponentLODInfo& Info = LODInfo[LODIndex];
-		FSkeletalMeshLODRenderData& LODData = SkelMeshRenderData->LODRenderData[LODIndex];
-
-		uint32 VertexCount = LODData.GetNumVertices();
-		Offsets.SetNumZeroed(VertexCount);
-
-		Info.PreSkinningOffsets = LWC::ConvertArrayType<FVector3f>(Offsets);	// LWC_TODO: Perf pessimization - was MoveTemp!
-
-		MarkRenderDynamicDataDirty();
-	}
-}
-
-void USkinnedMeshComponent::SetPostSkinningOffsets(int32 LODIndex, TArray<FVector> Offsets)
-{
-	InitLODInfos();
-
-	FSkeletalMeshRenderData* SkelMeshRenderData = GetSkeletalMeshRenderData();
-
-	// If we have a render resource, and the requested LODIndex is valid (for both component and mesh, though these should be the same)
-	if (SkelMeshRenderData && LODInfo.IsValidIndex(LODIndex) && SkelMeshRenderData->LODRenderData.IsValidIndex(LODIndex))
-	{
-		ensure(LODInfo.Num() == SkelMeshRenderData->LODRenderData.Num());
-
-		FSkelMeshComponentLODInfo& Info = LODInfo[LODIndex];
-		FSkeletalMeshLODRenderData& LODData = SkelMeshRenderData->LODRenderData[LODIndex];
-
-		uint32 VertexCount = LODData.GetNumVertices();
-		Offsets.SetNumZeroed(VertexCount);
-
-		Info.PostSkinningOffsets = LWC::ConvertArrayType<FVector3f>(Offsets);	// LWC_TODO: Perf pessimization - was MoveTemp!
-
-		MarkRenderDynamicDataDirty();
-	}
-}
-
 #if WITH_EDITOR
 
 bool USkinnedMeshComponent::IsCompiling() const
diff --git a/Engine/Source/Runtime/Engine/Private/GPUSkinCache.cpp b/Engine/Source/Runtime/Engine/Private/GPUSkinCache.cpp
--- a/Engine/Source/Runtime/Engine/Private/GPUSkinCache.cpp
+++ b/Engine/Source/Runtime/Engine/Private/GPUSkinCache.cpp
@@ -547,11 +547,8 @@ protected:
 	bool bUse16BitBoneIndex;
 	uint32 InputWeightIndexSize;
 	uint32 InputWeightStride;
-	uint32 VertexOffsetUsage = 0;
 	FShaderResourceViewRHIRef InputWeightStreamSRV;
 	FShaderResourceViewRHIRef InputWeightLookupStreamSRV;
-	FRHIShaderResourceView* PreSkinningVertexOffsetSRV = nullptr;
-	FRHIShaderResourceView* PostSkinningVertexOffsetSRV = nullptr;
 	FRHIShaderResourceView* MorphBuffer;
 	FShaderResourceViewRHIRef ClothBuffer;
 	int32 LOD;
@@ -578,10 +575,6 @@ public:
 		TangentInputBuffer.Bind(Initializer.ParameterMap, TEXT("TangentInputBuffer"));
 		PositionInputBuffer.Bind(Initializer.ParameterMap, TEXT("PositionInputBuffer"));
 
-		VertexOffsetUsage.Bind(Initializer.ParameterMap, TEXT("VertexOffsetUsage"));
-		PreSkinOffsets.Bind(Initializer.ParameterMap, TEXT("PreSkinOffsets"));
-		PostSkinOffsets.Bind(Initializer.ParameterMap, TEXT("PostSkinOffsets"));
-
 		InputStreamStart.Bind(Initializer.ParameterMap, TEXT("InputStreamStart"));
 
 		NumBoneInfluences.Bind(Initializer.ParameterMap, TEXT("NumBoneInfluences"));
@@ -625,10 +618,6 @@ public:
 		SetSRVParameter(RHICmdList, ShaderRHI, TangentInputBuffer, DispatchData.TangentBufferSRV);
 		SetSRVParameter(RHICmdList, ShaderRHI, PositionInputBuffer, DispatchData.PositionBufferSRV);
 
-		SetShaderValue(RHICmdList, ShaderRHI, VertexOffsetUsage, Entry->VertexOffsetUsage);
-		SetSRVParameter(RHICmdList, ShaderRHI, PreSkinOffsets, Entry->PreSkinningVertexOffsetSRV ? Entry->PreSkinningVertexOffsetSRV : GNullVertexBuffer.VertexBufferSRV.GetReference());
-		SetSRVParameter(RHICmdList, ShaderRHI, PostSkinOffsets, Entry->PostSkinningVertexOffsetSRV ? Entry->PostSkinningVertexOffsetSRV : GNullVertexBuffer.VertexBufferSRV.GetReference());
-
 		SetShaderValue(RHICmdList, ShaderRHI, NumBoneInfluences, DispatchData.NumBoneInfluences);
 		SetShaderValue(RHICmdList, ShaderRHI, InputWeightIndexSize, Entry->InputWeightIndexSize);
 		SetShaderValue(RHICmdList, ShaderRHI, InputWeightStart, DispatchData.InputWeightStart);
@@ -688,10 +677,6 @@ private:
 	LAYOUT_FIELD(FShaderResourceParameter, PositionBufferUAV)
 	LAYOUT_FIELD(FShaderResourceParameter, TangentBufferUAV)
 
-	LAYOUT_FIELD(FShaderParameter, VertexOffsetUsage)
-	LAYOUT_FIELD(FShaderResourceParameter, PreSkinOffsets)
-	LAYOUT_FIELD(FShaderResourceParameter, PostSkinOffsets)
-
 	LAYOUT_FIELD(FShaderParameter, NumBoneInfluences);
 	LAYOUT_FIELD(FShaderParameter, InputWeightIndexSize);
 	LAYOUT_FIELD(FShaderParameter, InputWeightStart)
@@ -1510,7 +1495,6 @@ bool FGPUSkinCache::ProcessEntry(
 	FGPUSkinPassthroughVertexFactory* TargetVertexFactory, 
 	const FSkelMeshRenderSection& BatchElement, 
 	FSkeletalMeshObjectGPUSkin* Skin,
-	FVertexOffsetBuffers* VertexOffsetBuffers,
 	const FMorphVertexBuffer* MorphVertexBuffer,
 	const FSkeletalMeshVertexClothBuffer* ClothVertexBuffer, 
 	const FClothSimulData* SimData,
@@ -1622,10 +1606,6 @@ bool FGPUSkinCache::ProcessEntry(
 		Entries.Add(InOutEntry);
 	}
 
-	InOutEntry->VertexOffsetUsage = VertexOffsetBuffers->GetUsage();
-	InOutEntry->PreSkinningVertexOffsetSRV = VertexOffsetBuffers->PreSkinningOffsetsVertexBuffer.GetSRV();
-	InOutEntry->PostSkinningVertexOffsetSRV = VertexOffsetBuffers->PostSkinningOffsetsVertexBuffer.GetSRV();
-
 	const bool bMorph = MorphVertexBuffer && MorphVertexBuffer->SectionIds.Contains(Section);
 	if (bMorph)
 	{
diff --git a/Engine/Source/Runtime/Engine/Private/GPUSkinVertexFactory.cpp b/Engine/Source/Runtime/Engine/Private/GPUSkinVertexFactory.cpp
--- a/Engine/Source/Runtime/Engine/Private/GPUSkinVertexFactory.cpp
+++ b/Engine/Source/Runtime/Engine/Private/GPUSkinVertexFactory.cpp
@@ -578,28 +578,6 @@ void TGPUSkinVertexFactory<BoneInfluenceType>::AddVertexElements(FDataType& InDa
 		}
 	}
 
-	// Pre-skinning offsets
-	if (InData.PreSkinningOffsets.VertexBuffer)
-	{
-		OutElements.Add(AccessStreamComponent(InData.PreSkinningOffsets, 11));
-	}
-	else
-	{
-		FVertexStreamComponent NullDeltaComponent(&GNullVertexBuffer, 0, 0, VET_Float3);
-		OutElements.Add(AccessStreamComponent(NullDeltaComponent, 11));
-	}
-
-	// Post-skinning offsets
-	if (InData.PostSkinningOffsets.VertexBuffer)
-	{
-		OutElements.Add(AccessStreamComponent(InData.PostSkinningOffsets, 12));
-	}
-	else
-	{
-		FVertexStreamComponent NullDeltaComponent(&GNullVertexBuffer, 0, 0, VET_Float3);
-		OutElements.Add(AccessStreamComponent(NullDeltaComponent, 12));
-	}
-
 	// Primitive Id
 	AddPrimitiveIdStreamElement(EVertexInputStreamType::Default, OutElements, 16, 0xff);
 }
diff --git a/Engine/Source/Runtime/Engine/Private/Materials/HLSLMaterialTranslator.cpp b/Engine/Source/Runtime/Engine/Private/Materials/HLSLMaterialTranslator.cpp
--- a/Engine/Source/Runtime/Engine/Private/Materials/HLSLMaterialTranslator.cpp
+++ b/Engine/Source/Runtime/Engine/Private/Materials/HLSLMaterialTranslator.cpp
@@ -6946,26 +6946,6 @@ int32 FHLSLMaterialTranslator::VertexColor()
 	}
 }
 
-int32 FHLSLMaterialTranslator::PreSkinVertexOffset()
-{
-	if (ShaderFrequency != SF_Vertex)
-	{
-		return Errorf(TEXT("Pre Skin Offset only available in the vertex shader, pass through custom interpolators if needed."));
-	}
-
-	return AddCodeChunk(MCT_Float3, TEXT("MaterialExpressionPreSkinOffset(Parameters)"));
-}
-
-int32 FHLSLMaterialTranslator::PostSkinVertexOffset()
-{
-	if (ShaderFrequency != SF_Vertex)
-	{
-		return Errorf(TEXT("Post Skin Offset only available in the vertex shader, pass through custom interpolators if needed."));
-	}
-
-	return AddCodeChunk(MCT_Float3, TEXT("MaterialExpressionPostSkinOffset(Parameters)"));
-}
-
 int32 FHLSLMaterialTranslator::PreSkinnedPosition()
 {
 	if (ShaderFrequency != SF_Vertex)
diff --git a/Engine/Source/Runtime/Engine/Private/Materials/HLSLMaterialTranslator.h b/Engine/Source/Runtime/Engine/Private/Materials/HLSLMaterialTranslator.h
--- a/Engine/Source/Runtime/Engine/Private/Materials/HLSLMaterialTranslator.h
+++ b/Engine/Source/Runtime/Engine/Private/Materials/HLSLMaterialTranslator.h
@@ -868,9 +868,6 @@ protected:
 
 	virtual int32 VertexColor() override;
 
-	virtual int32 PreSkinVertexOffset() override;
-	virtual int32 PostSkinVertexOffset() override;
-
 	virtual int32 PreSkinnedPosition() override;
 	virtual int32 PreSkinnedNormal() override;
 
diff --git a/Engine/Source/Runtime/Engine/Private/Materials/MaterialExpressions.cpp b/Engine/Source/Runtime/Engine/Private/Materials/MaterialExpressions.cpp
--- a/Engine/Source/Runtime/Engine/Private/Materials/MaterialExpressions.cpp
+++ b/Engine/Source/Runtime/Engine/Private/Materials/MaterialExpressions.cpp
@@ -172,7 +172,6 @@
 #include "Materials/MaterialExpressionPixelDepth.h"
 #include "Materials/MaterialExpressionPixelNormalWS.h"
 #include "Materials/MaterialExpressionPower.h"
-#include "Materials/MaterialExpressionSkinningVertexOffsets.h"
 #include "Materials/MaterialExpressionPreSkinnedNormal.h"
 #include "Materials/MaterialExpressionPreSkinnedPosition.h"
 #include "Materials/MaterialExpressionQualitySwitch.h"
@@ -18994,72 +18993,6 @@ void UMaterialExpressionSkyAtmosphereDistantLightScatteredLuminance::GetCaption(
 }
 #endif // WITH_EDITOR
 
-///////////////////////////////////////////////////////////////////////////////
-// UMaterialExpressionSkinningVertexOffsets
-///////////////////////////////////////////////////////////////////////////////
-PRAGMA_DISABLE_DEPRECATION_WARNINGS
-UMaterialExpressionSkinningVertexOffsets::UMaterialExpressionSkinningVertexOffsets(const FObjectInitializer& ObjectInitializer)
-	: Super(ObjectInitializer)
-{
-#if WITH_EDITORONLY_DATA
-	// Structure to hold one-time initialization
-	struct FConstructorStatics
-	{
-		FText NAME_Constants;
-		FConstructorStatics()
-			: NAME_Constants(LOCTEXT("Vectors", "Vectors"))
-		{
-		}
-	};
-	static FConstructorStatics ConstructorStatics;
-
-	MenuCategories.Add(ConstructorStatics.NAME_Constants);
-
-	Outputs.Reset();
-	Outputs.Add(FExpressionOutput(TEXT("Pre Skin Offset"), 1, 1, 1, 1, 0));
-	Outputs.Add(FExpressionOutput(TEXT("Post Skin Offset"), 1, 1, 1, 1, 0));
-	bShaderInputData = true;
-	bShowOutputNameOnPin = true;
-#endif
-}
-PRAGMA_ENABLE_DEPRECATION_WARNINGS
-
-#if WITH_EDITOR
-int32 UMaterialExpressionSkinningVertexOffsets::Compile(class FMaterialCompiler* Compiler, int32 OutputIndex)
-{
-	if (Compiler->GetCurrentShaderFrequency() != SF_Vertex)
-	{
-		return Compiler->Errorf(
-			TEXT("%s only available in the vertex shader, pass through custom interpolators if needed."), 
-			(OutputIndex == 0) ? TEXT("Pre Skin Offset") : TEXT("Post Skin Offset")
-			);
-	}
-
-	switch (OutputIndex)
-	{
-	default:
-		return Compiler->Constant3(0, 0, 0);
-
-	case 0:
-		return Compiler->PreSkinVertexOffset();
-
-	case 1:
-		return Compiler->PostSkinVertexOffset();
-	}
-}
-
-void UMaterialExpressionSkinningVertexOffsets::GetCaption(TArray<FString>& OutCaptions) const
-{
-	OutCaptions.Add(TEXT("Skin Vertex Offset"));
-}
-
-void UMaterialExpressionSkinningVertexOffsets::GetExpressionToolTip(TArray<FString>& OutToolTip)
-{
-	ConvertToMultilineToolTip(TEXT("Returns pre or post skinned offset applied to each vertex for a skeletal mesh, usable in vertex shader only."
-		"Returns <0, 0, 0> for non-skeletal meshes or when disabled."), 40, OutToolTip);
-}
-#endif // WITH_EDITOR
-
 ///////////////////////////////////////////////////////////////////////////////
 // UMaterialExpressionPreSkinnedPosition
 ///////////////////////////////////////////////////////////////////////////////
diff --git a/Engine/Source/Runtime/Engine/Private/SkeletalRenderGPUSkin.cpp b/Engine/Source/Runtime/Engine/Private/SkeletalRenderGPUSkin.cpp
--- a/Engine/Source/Runtime/Engine/Private/SkeletalRenderGPUSkin.cpp
+++ b/Engine/Source/Runtime/Engine/Private/SkeletalRenderGPUSkin.cpp
@@ -298,9 +298,7 @@ void FSkeletalMeshObjectGPUSkin::InitResources(USkinnedMeshComponent* InMeshComp
 				CompLODInfo = &InMeshComponent->LODInfo[LODIndex];
 			}
 
-			PRAGMA_DISABLE_DEPRECATION_WARNINGS
-			SkelLOD.InitResources(InMeshComponent->GetVertexOffsetUsage(LODIndex), MeshLODInfo, CompLODInfo, FeatureLevel);
-			PRAGMA_ENABLE_DEPRECATION_WARNINGS
+			SkelLOD.InitResources(MeshLODInfo, CompLODInfo, FeatureLevel);
 		}
 	}
 
@@ -637,45 +635,6 @@ void FSkeletalMeshObjectGPUSkin::ProcessUpdatedDynamicData(EGPUSkinCacheEntryMod
 		FGPUSkinCache::Release(SkinCacheEntry);
 	}
 
-	if (DynamicData->PreSkinningOffsets.Num() > 0)
-	{
-		FPositionVertexBuffer& PositionBuffer = LOD.VertexOffsetVertexBuffers.PreSkinningOffsetsVertexBuffer;
-
-		check(PositionBuffer.GetNumVertices() == DynamicData->PreSkinningOffsets.Num());
-
-		uint32 SizeInBytes = PositionBuffer.GetNumVertices() * sizeof(FVector3f);
-
-		void* Buffer = RHICmdList.LockBuffer(
-			PositionBuffer.VertexBufferRHI,
-			0,
-			SizeInBytes,
-			RLM_WriteOnly
-			);
-
-		FMemory::Memcpy(Buffer, DynamicData->PreSkinningOffsets.GetData(), SizeInBytes);
-
-		RHICmdList.UnlockBuffer(PositionBuffer.VertexBufferRHI);
-	}
-
-	if (DynamicData->PostSkinningOffsets.Num() > 0)
-	{
-		FPositionVertexBuffer& PositionBuffer = LOD.VertexOffsetVertexBuffers.PostSkinningOffsetsVertexBuffer;
-
-		check(PositionBuffer.GetNumVertices() == DynamicData->PostSkinningOffsets.Num());
-
-		uint32 SizeInBytes = PositionBuffer.GetNumVertices() * sizeof(FVector3f);
-
-		void* Buffer = RHICmdList.LockBuffer(
-			PositionBuffer.VertexBufferRHI,
-			0,
-			SizeInBytes,
-			RLM_WriteOnly
-			);
-
-		FMemory::Memcpy(Buffer, DynamicData->PostSkinningOffsets.GetData(), SizeInBytes);
-
-		RHICmdList.UnlockBuffer(PositionBuffer.VertexBufferRHI);
-	}
 
 	if (MorphVertexBuffer.bNeedsInitialClear && !(bMorph && bMorphNeedsUpdate))
 	{
@@ -811,7 +770,6 @@ void FSkeletalMeshObjectGPUSkin::ProcessUpdatedDynamicData(EGPUSkinCacheEntryMod
 						VertexFactoryData.PassthroughVertexFactories[SectionIdx].Get(),
 						Section,
 						this,
-						&LOD.VertexOffsetVertexBuffers,
 						bMorph ? &MorphVertexBuffer : 0,
 						bClothFactory ? &LODData.ClothVertexBuffer : 0,
 						bClothFactory ? DynamicData->ClothingSimData.Find(Section.CorrespondClothAssetIndex) : 0,
@@ -1413,31 +1371,6 @@ void InitGPUSkinVertexFactoryComponents(typename VertexFactoryType::FDataType* V
 		VertexFactoryData->ColorComponentsSRV = nullptr;
 		VertexFactoryData->ColorIndexMask = 0;
 	}
-
-	if (VertexBuffers.VertexOffsetVertexBuffers)
-	{
-		FVertexOffsetBuffers* VertexOffsetVertexBuffers = VertexBuffers.VertexOffsetVertexBuffers;
-
-		if (VertexOffsetVertexBuffers->PreSkinningOffsetsVertexBuffer.VertexBufferRHI)
-		{
-			VertexFactoryData->PreSkinningOffsets = FVertexStreamComponent(
-				&VertexOffsetVertexBuffers->PreSkinningOffsetsVertexBuffer,
-				0,
-				VertexOffsetVertexBuffers->PreSkinningOffsetsVertexBuffer.GetStride(),
-				VET_Float3
-				);
-		}
-
-		if (VertexOffsetVertexBuffers->PostSkinningOffsetsVertexBuffer.VertexBufferRHI)
-		{
-			VertexFactoryData->PostSkinningOffsets = FVertexStreamComponent(
-				&VertexOffsetVertexBuffers->PostSkinningOffsetsVertexBuffer,
-				0,
-				VertexOffsetVertexBuffers->PostSkinningOffsetsVertexBuffer.GetStride(),
-				VET_Float3
-				);
-		}
-	}
 }
 
 /** 
@@ -1705,7 +1638,6 @@ void FSkeletalMeshObjectGPUSkin::FSkeletalMeshObjectLOD::GetVertexBuffers(FVerte
 	OutVertexBuffers.SkinWeightVertexBuffer = MeshObjectWeightBuffer;
 	OutVertexBuffers.MorphVertexBufferPool = &MorphVertexBufferPool;
 	OutVertexBuffers.APEXClothVertexBuffer = &LODData.ClothVertexBuffer;
-	OutVertexBuffers.VertexOffsetVertexBuffers = &VertexOffsetVertexBuffers;
 	OutVertexBuffers.NumVertices = LODData.GetNumVertices();
 }
 
@@ -1912,7 +1844,7 @@ void FSkeletalMeshObjectGPUSkin::FVertexFactoryData::UpdateVertexFactoryData(con
 	}
 }
 
-void FSkeletalMeshObjectGPUSkin::FSkeletalMeshObjectLOD::InitResources(uint32 VertexOffsetUsage, const FSkelMeshObjectLODInfo& MeshLODInfo, FSkelMeshComponentLODInfo* CompLODInfo, ERHIFeatureLevel::Type InFeatureLevel)
+void FSkeletalMeshObjectGPUSkin::FSkeletalMeshObjectLOD::InitResources(const FSkelMeshObjectLODInfo& MeshLODInfo, FSkelMeshComponentLODInfo* CompLODInfo, ERHIFeatureLevel::Type InFeatureLevel)
 {
 	check(SkelMeshRenderData);
 	check(SkelMeshRenderData->LODRenderData.IsValidIndex(LODIndex));
@@ -1951,22 +1883,6 @@ void FSkeletalMeshObjectGPUSkin::FSkeletalMeshObjectLOD::InitResources(uint32 Ve
 		MeshObjectColorBuffer = &LODData.StaticVertexBuffers.ColorVertexBuffer;
 	}
 
-	if ((VertexOffsetUsage & uint32(EVertexOffsetUsageType::PreSkinningOffset)) && CompLODInfo && CompLODInfo->PreSkinningOffsets.Num() == 0)
-	{
-		CompLODInfo->PreSkinningOffsets.SetNumZeroed(LODData.GetNumVertices());
-	}
-
-	if ((VertexOffsetUsage & uint32(EVertexOffsetUsageType::PostSkinningOffset)) && CompLODInfo && CompLODInfo->PostSkinningOffsets.Num() == 0)
-	{
-		CompLODInfo->PostSkinningOffsets.SetNumZeroed(LODData.GetNumVertices());
-	}
-
-	if (CompLODInfo)
-	{
-		VertexOffsetVertexBuffers.Init(VertexOffsetUsage, CompLODInfo->PreSkinningOffsets, CompLODInfo->PostSkinningOffsets);
-		VertexOffsetVertexBuffers.BeginInitResource();
-	}
-
 	// Vertex buffers available for the LOD
 	FVertexFactoryBuffers VertexBuffers;
 	GetVertexBuffers(VertexBuffers, LODData);
@@ -1984,8 +1900,6 @@ void FSkeletalMeshObjectGPUSkin::FSkeletalMeshObjectLOD::InitResources(uint32 Ve
  */
 void FSkeletalMeshObjectGPUSkin::FSkeletalMeshObjectLOD::ReleaseResources()
 {	
-	VertexOffsetVertexBuffers.BeginReleaseResource();
-
 	// Release gpu skin vertex factories
 	GPUSkinVertexFactories.ReleaseVertexFactories();
 
@@ -2259,9 +2173,6 @@ void FDynamicSkelMeshObjectDataGPUSkin::InitDynamicSkelMeshObjectDataGPUSkin(
 		bAnySegmentUsesWorldPositionOffset = SkeletalMeshProxy->bAnySegmentUsesWorldPositionOffset;
 	}
 #endif
-
-	PreSkinningOffsets = MoveTemp(InMeshComponent->LODInfo[InLODIndex].PreSkinningOffsets);
-	PostSkinningOffsets = MoveTemp(InMeshComponent->LODInfo[InLODIndex].PostSkinningOffsets);
 }
 
 bool FDynamicSkelMeshObjectDataGPUSkin::ActiveMorphTargetsEqual( const TArray<FActiveMorphTarget>& CompareActiveMorphTargets, const TArray<float>& CompareMorphTargetWeights)
diff --git a/Engine/Source/Runtime/Engine/Private/SkeletalRenderGPUSkin.h b/Engine/Source/Runtime/Engine/Private/SkeletalRenderGPUSkin.h
--- a/Engine/Source/Runtime/Engine/Private/SkeletalRenderGPUSkin.h
+++ b/Engine/Source/Runtime/Engine/Private/SkeletalRenderGPUSkin.h
@@ -101,9 +101,6 @@ public:
 	/** a weight factor to blend between simulated positions and skinned positions */	
 	float ClothBlendWeight;
 
-	TArray<FVector3f> PreSkinningOffsets;
-	TArray<FVector3f> PostSkinningOffsets;
-
 	/**
 	* Compare the given set of active morph targets with the current list to check if different
 	* @param CompareActiveMorphTargets - array of morphs to compare
@@ -306,82 +303,6 @@ private:
 	bool bDoubleBuffer = false;
 };
 
-/**
- * 
- */
-class FVertexOffsetBuffers
-{
-public:
-
-	void Init(uint32 Usage, const TArray<FVector3f>& PreSkinningOffsets, const TArray<FVector3f>& PostSkinningOffsets)
-	{
-		if (Usage & uint32(EVertexOffsetUsageType::PreSkinningOffset))
-		{
-			ensure(PreSkinningOffsets.Num() > 0);
-			PreSkinningOffsetsVertexBuffer.Init(PreSkinningOffsets);
-		}
-
-		if (Usage & uint32(EVertexOffsetUsageType::PostSkinningOffset))
-		{
-			ensure(PostSkinningOffsets.Num() > 0);
-			PostSkinningOffsetsVertexBuffer.Init(PostSkinningOffsets);
-		}
-	}
-
-	/**
-	 * Get Resource Size : mostly copied from InitDynamicRHI - how much they allocate when initialize
-	 */
-	SIZE_T GetResourceSize() const
-	{
-		SIZE_T ResourceSize = sizeof(*this);
-
-		if (PreSkinningOffsetsVertexBuffer.VertexBufferRHI)
-		{
-			ResourceSize += PreSkinningOffsetsVertexBuffer.GetNumVertices() * sizeof(FVector3f);
-		}
-
-		if (PostSkinningOffsetsVertexBuffer.VertexBufferRHI)
-		{
-			ResourceSize += PostSkinningOffsetsVertexBuffer.GetNumVertices() * sizeof(FVector3f);
-		}
-
-		return ResourceSize;
-	}
-
-
-	uint32 GetUsage() const
-	{
-		uint32 Usage = 0;
-
-		if (PreSkinningOffsetsVertexBuffer.GetNumVertices() > 0)
-		{
-			Usage |= uint32(EVertexOffsetUsageType::PreSkinningOffset);
-		}
-
-		if (PostSkinningOffsetsVertexBuffer.GetNumVertices() > 0)
-		{
-			Usage |= uint32(EVertexOffsetUsageType::PostSkinningOffset);
-		}
-
-		return Usage;
-	}
-
-	void BeginInitResource()
-	{
-		::BeginInitResource(&PreSkinningOffsetsVertexBuffer);
-		::BeginInitResource(&PostSkinningOffsetsVertexBuffer);
-	}
-
-	void BeginReleaseResource()
-	{
-		::BeginReleaseResource(&PostSkinningOffsetsVertexBuffer);
-		::BeginReleaseResource(&PreSkinningOffsetsVertexBuffer);
-	}
-
-	FPositionVertexBuffer PreSkinningOffsetsVertexBuffer;
-	FPositionVertexBuffer PostSkinningOffsetsVertexBuffer;
-};
-
 /**
  * Render data for a GPU skinned mesh
  */
@@ -598,7 +519,7 @@ protected:
 		 * @param MeshLODInfo - information about the state of the bone influence swapping
 		 * @param CompLODInfo - information about this LOD from the skeletal component 
 		 */
-		void InitResources(uint32 VertexOffsetUsage, const FSkelMeshObjectLODInfo& MeshLODInfo, FSkelMeshComponentLODInfo* CompLODInfo, ERHIFeatureLevel::Type FeatureLevel);
+		void InitResources(const FSkelMeshObjectLODInfo& MeshLODInfo, FSkelMeshComponentLODInfo* CompLODInfo, ERHIFeatureLevel::Type FeatureLevel);
 
 		/** 
 		 * Release rendering resources for this LOD 
@@ -623,7 +544,6 @@ protected:
 		void GetResourceSizeEx(FResourceSizeEx& CumulativeResourceSize)
 		{
 			CumulativeResourceSize.AddUnknownMemoryBytes(MorphVertexBufferPool.GetResourceSize());
-			CumulativeResourceSize.AddUnknownMemoryBytes(VertexOffsetVertexBuffers.GetResourceSize());
 			CumulativeResourceSize.AddUnknownMemoryBytes(GPUSkinVertexFactories.GetResourceSize());
 		}
 
@@ -636,8 +556,6 @@ protected:
 		/** Pooled vertex buffers that store the morph target vertex deltas. */
 		FMorphVertexBufferPool	MorphVertexBufferPool;
 
-		FVertexOffsetBuffers VertexOffsetVertexBuffers;
-
 		/** Default GPU skinning vertex factories and matrices */
 		FVertexFactoryData GPUSkinVertexFactories;
 
diff --git a/Engine/Source/Runtime/Engine/Public/GPUSkinCache.h b/Engine/Source/Runtime/Engine/Public/GPUSkinCache.h
--- a/Engine/Source/Runtime/Engine/Public/GPUSkinCache.h
+++ b/Engine/Source/Runtime/Engine/Public/GPUSkinCache.h
@@ -167,7 +167,6 @@ public:
 		FGPUSkinPassthroughVertexFactory* TargetVertexFactory, 
 		const FSkelMeshRenderSection& BatchElement, 
 		FSkeletalMeshObjectGPUSkin* Skin,
-		FVertexOffsetBuffers* VertexOffsetBuffers,
 		const FMorphVertexBuffer* MorphVertexBuffer, 
 		const FSkeletalMeshVertexClothBuffer* ClothVertexBuffer, 
 		const FClothSimulData* SimData,
diff --git a/Engine/Source/Runtime/Engine/Public/GPUSkinVertexFactory.h b/Engine/Source/Runtime/Engine/Public/GPUSkinVertexFactory.h
--- a/Engine/Source/Runtime/Engine/Public/GPUSkinVertexFactory.h
+++ b/Engine/Source/Runtime/Engine/Public/GPUSkinVertexFactory.h
@@ -377,12 +377,6 @@ public:
 		/** The stream to read the blend stream offset and num of influences from */
 		FVertexStreamComponent BlendOffsetCount;
 
-		/** The stream to read for vertex offsets pre skinning */
-		FVertexStreamComponent PreSkinningOffsets;
-
-		/** The stream to read for vertex offsets post skinning */
-		FVertexStreamComponent PostSkinningOffsets;
-
 		/** Number of bone influences */
 		uint32 NumBoneInfluences = 0;
 
diff --git a/Engine/Source/Runtime/Engine/Public/MaterialCompiler.h b/Engine/Source/Runtime/Engine/Public/MaterialCompiler.h
--- a/Engine/Source/Runtime/Engine/Public/MaterialCompiler.h
+++ b/Engine/Source/Runtime/Engine/Public/MaterialCompiler.h
@@ -307,9 +307,6 @@ public:
 
 	virtual int32 VertexColor() = 0;
 
-	virtual int32 PreSkinVertexOffset() = 0;
-	virtual int32 PostSkinVertexOffset() = 0;
-
 	virtual int32 PreSkinnedPosition() = 0;
 	virtual int32 PreSkinnedNormal() = 0;
 
