Fix typo in tooltip

diff --git a/Engine/Source/Runtime/Engine/Classes/Sound/AudioSettings.h b/Engine/Source/Runtime/Engine/Classes/Sound/AudioSettings.h
--- a/Engine/Source/Runtime/Engine/Classes/Sound/AudioSettings.h
+++ b/Engine/Source/Runtime/Engine/Classes/Sound/AudioSettings.h
@@ -192,7 +192,7 @@ class ENGINE_API UAudioSettings : public UDeveloperSettings
 	EPanningMethod PanningMethod;
 
 	/**
-	* The upmixing method for mono sound sources. Defines up mono channels are up-mixed to stereo channels.
+	* The upmixing method for mono sound sources. Defines how mono channels are up-mixed to stereo channels.
 	*/
 	UPROPERTY(config, EditAnywhere, Category = "Quality", AdvancedDisplay)
 	EMonoChannelUpmixMethod MonoChannelUpmixMethod;
