Fixing tooltip

diff --git a/Engine/Source/Runtime/Engine/Classes/Sound/SoundSubmix.h b/Engine/Source/Runtime/Engine/Classes/Sound/SoundSubmix.h
--- a/Engine/Source/Runtime/Engine/Classes/Sound/SoundSubmix.h
+++ b/Engine/Source/Runtime/Engine/Classes/Sound/SoundSubmix.h
@@ -423,13 +423,11 @@ public:
 	UPROPERTY(EditAnywhere, AssetRegistrySearchable, Category = Soundfield)
 	FName SoundfieldEncodingFormat;
 
-	//TODO: Make this editable only if SoundfieldEncodingFormat is non-default,
-	// and filter classes based on ISoundfieldFactory::GetCustomSettingsClass().
+	/** Which encoding settings to use the sound field. */
 	UPROPERTY(EditAnywhere, Category = Soundfield)
 	TObjectPtr<USoundfieldEncodingSettingsBase> EncodingSettings;
 
-	// TODO: make this editable only if SoundfieldEncodingFormat is non-default
-	// and filter classes based on USoundfieldProcessorBase::SupportsFormat.
+	/** Soundfield effect chain to use for the sound field. */
 	UPROPERTY(EditAnywhere, Category = Soundfield)
 	TArray<TObjectPtr<USoundfieldEffectBase>> SoundfieldEffectChain;
 
