Refactor HLSL expression API, First phase renamed from UpdateType to PreapreValue.  Now responsible for setting EvaluationMode in addition to type.
Second phase split into EmitValueShader/EmitValuePreshader, which are called depending on the EvaluationMode set in PrepareValue.
This allows checks for constant values to be moved to PrepareValue, which will facilitate moving checks for reading material attributes to PrepareValue as well.
This is important since PrepareValue is designed to be able to be called multiple times, but EmitValue* should only be called once

diff --git a/Engine/Source/Runtime/Engine/Private/HLSLTree/HLSLTree.cpp b/Engine/Source/Runtime/Engine/Private/HLSLTree/HLSLTree.cpp
--- a/Engine/Source/Runtime/Engine/Private/HLSLTree/HLSLTree.cpp
+++ b/Engine/Source/Runtime/Engine/Private/HLSLTree/HLSLTree.cpp
@@ -383,8 +383,34 @@ FScope* FScope::FindSharedParent(FScope* Lhs, FScope* Rhs)
 	return Scope0;
 }
 
-void FExpressionLocalPHI::UpdateType(FUpdateTypeContext& Context, const FRequestedType& RequestedType)
+void FExpressionLocalPHI::PrepareValue(FEmitContext& Context, const FRequestedType& RequestedType)
 {
+	FExpression* ForwardExpression = nullptr;
+	bool bForwardExpressionValid = false;
+	for (int32 i = 0; i < NumValues; ++i)
+	{
+		// Ignore values in scopes that have already been marked as dead
+		if (!Scopes[i]->IsDead())
+		{
+			FExpression* ScopeExpression = Values[i];
+			if (!ForwardExpression)
+			{
+				ForwardExpression = ScopeExpression;
+				bForwardExpressionValid = true;
+			}
+			else if (ForwardExpression != ScopeExpression)
+			{
+				bForwardExpressionValid = false;
+			}
+		}
+	}
+
+	if (bForwardExpressionValid)
+	{
+		check(ForwardExpression);
+		return SetForwardValue(Context, ForwardExpression, RequestedType);
+	}
+
 	FType TypePerValue[MaxNumPreviousScopes];
 	int32 NumValidTypes = 0;
 	FType CurrentType;
@@ -395,10 +421,11 @@ void FExpressionLocalPHI::UpdateType(FUpdateTypeContext& Context, const FRequest
 		{
 			if (!TypePerValue[i])
 			{
-				TypePerValue[i] = RequestExpressionType(Context, Values[i], RequestedType);
-				if (TypePerValue[i])
+				const FPrepareValueResult ValueResult = PrepareExpressionValue(Context, Values[i], RequestedType);
+				if (ValueResult.Type)
 				{
-					if (!CurrentType.Merge(TypePerValue[i]))
+					TypePerValue[i] = ValueResult.Type;
+					if (!CurrentType.Merge(ValueResult.Type))
 					{
 						return Context.Errors.AddError(this, TEXT("Type mismatch"));
 					}
@@ -413,7 +440,7 @@ void FExpressionLocalPHI::UpdateType(FUpdateTypeContext& Context, const FRequest
 	UpdateValueTypes();
 
 	// Assuming we have at least one value with a valid type, we use that to initialize our type
-	SetType(Context, CurrentType);
+	SetType(Context, EExpressionEvaluationType::Shader, CurrentType);
 
 	if (NumValidTypes < NumValues)
 	{
@@ -426,27 +453,16 @@ void FExpressionLocalPHI::UpdateType(FUpdateTypeContext& Context, const FRequest
 	}
 }
 
-void FExpressionLocalPHI::PrepareValue(FEmitContext& Context)
+void FExpressionLocalPHI::EmitValueShader(FEmitContext& Context, FShaderValue& OutShader) const
 {
-	FExpression* ForwardExpression = Values[0];
-	for (int32 i = 1; i < NumValues; ++i)
-	{
-		if (Values[i] != ForwardExpression)
-		{
-			ForwardExpression = nullptr;
-			break;
-		}
-	}
-
-	// If we have the same value in all of our scopes, just forward that value directly; no need to create/assign a local variable
-	if (ForwardExpression)
-	{
-		return SetValueForward(Context, ForwardExpression);
-	}
-
 	const int32 LocalPHIIndex = Context.NumLocalPHIs++;
-	SetValueInlineShaderf(Context, TEXT("LocalPHI%d"), LocalPHIIndex);
+	OutShader.Code.Appendf(TEXT("LocalPHI%d"), LocalPHIIndex);
+	OutShader.bInline = true;
+	OutShader.bHasDependencies = true;
+}
 
+void FExpressionLocalPHI::EmitShaderDependencies(FEmitContext& Context, const FShaderValue& Shader) const
+{
 	// Find the outermost scope to declare our local variable
 	FScope* DeclarationScope = ParentScope;
 	for (int32 i = 0; i < NumValues; ++i)
@@ -458,14 +474,6 @@ void FExpressionLocalPHI::PrepareValue(FEmitContext& Context)
 		}
 	}
 
-	for (int32 i = 0; i < NumValues; ++i)
-	{
-		if (PrepareExpressionValue(Context, Values[i]) == EExpressionEvaluationType::None)
-		{
-			return Context.Errors.AddError(this, TEXT("Invalid LocalPHI"));
-		}
-	}
-
 	const FType LocalType = GetType();
 
 	bool bNeedToAddDeclaration = true;
@@ -476,16 +484,16 @@ void FExpressionLocalPHI::PrepareValue(FEmitContext& Context)
 		ValueScope->MarkLiveRecursive();
 		if (ValueScope == DeclarationScope)
 		{
-			ValueScope->EmitDeclarationf(Context, TEXT("%s LocalPHI%d = %s;"),
+			ValueScope->EmitDeclarationf(Context, TEXT("%s %s = %s;"),
 				LocalType.GetName(),
-				LocalPHIIndex,
+				Shader.Code.ToString(),
 				ValueCode);
 			bNeedToAddDeclaration = false;
 		}
 		else
 		{
-			ValueScope->EmitStatementf(Context, TEXT("LocalPHI%d = %s;"),
-				LocalPHIIndex,
+			ValueScope->EmitStatementf(Context, TEXT("%s = %s;"),
+				Shader.Code.ToString(),
 				ValueCode);
 		}
 	}
@@ -493,9 +501,9 @@ void FExpressionLocalPHI::PrepareValue(FEmitContext& Context)
 	if (bNeedToAddDeclaration)
 	{
 		check(DeclarationScope->IsLive());
-		DeclarationScope->EmitDeclarationf(Context, TEXT("%s LocalPHI%d;"),
+		DeclarationScope->EmitDeclarationf(Context, TEXT("%s %s;"),
 			LocalType.GetName(),
-			LocalPHIIndex);
+			Shader.Code.ToString());
 	}
 }
 
@@ -532,11 +540,11 @@ ENodeVisitResult FScope::Visit(FNodeVisitor& Visitor)
 	return Result;
 }
 
-FType RequestExpressionType(FUpdateTypeContext& Context, FExpression* InExpression, const FRequestedType& RequestedType)
+FPrepareValueResult PrepareExpressionValue(FEmitContext& Context, FExpression* InExpression, const FRequestedType& RequestedType)
 {
 	if (!InExpression)
 	{
-		return Shader::EValueType::Void;
+		return FPrepareValueResult();
 	}
 
 	if (InExpression->bReentryFlag)
@@ -544,7 +552,7 @@ FType RequestExpressionType(FUpdateTypeContext& Context, FExpression* InExpressi
 		// Valid for this to be called reentrantly
 		// Code should ensure that the type is set before the reentrant call, otherwise type will not be valid here
 		// LocalPHI nodes rely on this to break loops
-		return InExpression->Type;
+		return InExpression->PrepareValueResult;
 	}
 
 	bool bNeedToUpdateType = false;
@@ -556,7 +564,7 @@ FType RequestExpressionType(FUpdateTypeContext& Context, FExpression* InExpressi
 	else if(InExpression->CurrentRequestedType.StructType != RequestedType.StructType)
 	{
 		Context.Errors.AddError(InExpression, TEXT("Type mismatch"));
-		return Shader::EValueType::Void;
+		return FPrepareValueResult();
 	}
 	else
 	{
@@ -577,197 +585,183 @@ FType RequestExpressionType(FUpdateTypeContext& Context, FExpression* InExpressi
 		check(!InExpression->CurrentRequestedType.IsVoid());
 
 		InExpression->bReentryFlag = true;
-		InExpression->UpdateType(Context, InExpression->CurrentRequestedType);
+		InExpression->PrepareValue(Context, InExpression->CurrentRequestedType);
 		InExpression->bReentryFlag = false;
 
-		if (!InExpression->Type)
+		if (InExpression->PrepareValueResult.EvaluationType == EExpressionEvaluationType::None)
 		{
 			// If we failed to assign a valid type, reset the requested type as well
 			// This ensures we'll try to compute a type again the next time we're called
 			InExpression->CurrentRequestedType.Reset();
 		}
+		else if (InExpression->PrepareValueResult.EvaluationType == EExpressionEvaluationType::Constant &&
+			InExpression->PrepareValueResult.Type != FType(InExpression->PrepareValueResult.ConstantValue.GetType()))
+		{
+			Shader::FPreshaderData ConstantPreshader;
+			InExpression->bReentryFlag = true;
+			InExpression->EmitValuePreshader(Context, ConstantPreshader);
+			InExpression->bReentryFlag = false;
+
+			// Evaluate the constant preshader and store its value
+			FMaterialRenderContext RenderContext(nullptr, *Context.Material, nullptr);
+			Shader::FValue Value;
+			ConstantPreshader.Evaluate(nullptr, RenderContext, Value);
+			InExpression->PrepareValueResult.ConstantValue = Value;
+		}
 	}
 
-	return InExpression->Type;
-}
-
-EExpressionEvaluationType PrepareExpressionValue(FEmitContext& Context, FExpression* InExpression)
-{
-	if (!InExpression)
-	{
-		return EExpressionEvaluationType::None;
-	}
-
-	if (!InExpression->Type)
-	{
-		// Can't prepare value if we didn't successfully prepare a type
-		return EExpressionEvaluationType::None;
-	}
-
-	if (InExpression->EvaluationType == EExpressionEvaluationType::None)
-	{
-		check(!InExpression->bReentryFlag); // Re-entry is not supported
-
-		InExpression->bReentryFlag = true;
-		InExpression->PrepareValue(Context);
-		InExpression->bReentryFlag = false;
-	}
-	return InExpression->EvaluationType;
+	return InExpression->PrepareValueResult;
 }
 
-void FExpression::SetType(FUpdateTypeContext& Context, const FType& InType)
+void FExpression::SetType(FEmitContext& Context, EExpressionEvaluationType InEvaluationType, const FType& InType)
 {
-	Type = InType;
+	PrepareValueResult.EvaluationType = InEvaluationType;
+	PrepareValueResult.Type = InType;
+	ForwardedValue = nullptr;
 }
 
-void FExpression::InternalSetValueShader(FEmitContext& Context, const TCHAR* InCode, int32 InLength, bool bInline)
+void FExpression::SetForwardValue(FEmitContext& Context, FExpression* InForwardValue, const FRequestedType& RequestedType)
 {
-	check(!Code);
-	check(EvaluationType == EExpressionEvaluationType::None);
-
-	EvaluationType = EExpressionEvaluationType::Shader;
-	if (bInline)
-	{
-		Code = AllocateString(*Context.Allocator, InCode, InLength);
-	}
-	else
-	{
-		// Check to see if we've already generated code for an equivalent expression in either this scope, or any outer visible scope
-		const FSHAHash Hash = HashString(InCode, InLength);
-		const TCHAR* Declaration = nullptr;
-
-		FScope* CheckScope = ParentScope;
-		while (CheckScope)
-		{
-			const TCHAR** FoundDeclaration = CheckScope->ExpressionCodeMap.Find(Hash);
-			if (FoundDeclaration)
-			{
-				// Re-use results from previous matching expression
-				Declaration = *FoundDeclaration;
-				break;
-			}
-			CheckScope = CheckScope->ParentScope;
-		}
-
-		if (!Declaration)
-		{
-			check(ParentScope);
-			Declaration = Context.AcquireLocalDeclarationCode();
-			ParentScope->ExpressionCodeMap.Add(Hash, Declaration);
-			ParentScope->MarkLiveRecursive();
-			ParentScope->EmitStatementf(Context, TEXT("const %s %s = %s;"),
-				Type.GetName(),
-				Declaration,
-				InCode);
-		}
-		Code = Declaration;
-	}
+	check(InForwardValue);
+	PrepareValueResult = PrepareExpressionValue(Context, InForwardValue, RequestedType);
+	ForwardedValue = InForwardValue;
 }
 
-void FExpression::SetValuePreshader(FEmitContext& Context, Shader::FPreshaderData& InPreshader)
+void FExpression::EmitValueShader(FEmitContext& Context, FShaderValue& OutShader) const
 {
-	check(!Preshader);
-	check(EvaluationType == EExpressionEvaluationType::None);
-	check(!Type.IsStruct());
-
-	EvaluationType = EExpressionEvaluationType::Preshader;
-	Preshader = new Shader::FPreshaderData(MoveTemp(InPreshader));
+	check(false);
 }
 
-void FExpression::SetValueConstant(FEmitContext& Context, const Shader::FValue& InValue)
+void FExpression::EmitShaderDependencies(FEmitContext& Context, const FShaderValue& Shader) const
 {
-	check(EvaluationType == EExpressionEvaluationType::None);
-	check(!Type.IsStruct());
-
-	EvaluationType = EExpressionEvaluationType::Constant;
-	ConstantValue = Shader::Cast(InValue, Type);
+	check(false);
 }
 
-void FExpression::SetValueForward(FEmitContext& Context, FExpression* Source)
+void FExpression::EmitValuePreshader(FEmitContext& Context, Shader::FPreshaderData& OutPreshader) const
 {
-	check(EvaluationType == EExpressionEvaluationType::None);
-	check(Source);
-
-	ensure(Type == Source->GetType());
-
-	EvaluationType = PrepareExpressionValue(Context, Source);
-	switch (EvaluationType)
-	{
-	case EExpressionEvaluationType::None: break;
-	case EExpressionEvaluationType::Shader: Code = Source->Code; break;
-	case EExpressionEvaluationType::Preshader: Preshader = Source->Preshader; break;
-	case EExpressionEvaluationType::Constant: ConstantValue = Source->ConstantValue; break;
-	default: checkNoEntry(); break;
-	}
+	check(false);
 }
 
-void FExpression::SetValuePreshader(FEmitContext& Context, EExpressionEvaluationType InEvaluationType, Shader::FPreshaderData& InPreshader)
+const TCHAR* FExpression::GetValueShader(FEmitContext& Context)
 {
-	if (InEvaluationType == EExpressionEvaluationType::Constant)
+	check(PrepareValueResult.EvaluationType != EExpressionEvaluationType::None);
+	if (ForwardedValue)
 	{
-		// Evaluate the constant preshader and store its value
-		FMaterialRenderContext RenderContext(nullptr, *Context.Material, nullptr);
-		Shader::FValue Value;
-		InPreshader.Evaluate(nullptr, RenderContext, Value);
-		SetValueConstant(Context, Value);
+		return ForwardedValue->GetValueShader(Context);
 	}
-	else
-	{
-		check(InEvaluationType == EExpressionEvaluationType::Preshader);
-		SetValuePreshader(Context, InPreshader);
-	}
-}
 
-const TCHAR* FExpression::GetValueShader(FEmitContext& Context)
-{
-	check(EvaluationType != EExpressionEvaluationType::None);
 	if (!Code)
 	{
+		check(!bReentryFlag);
+
 		TStringBuilder<1024> FormattedCode;
-		if (EvaluationType == EExpressionEvaluationType::Constant)
+		if (PrepareValueResult.EvaluationType == EExpressionEvaluationType::Shader)
+		{
+			FShaderValue ShaderValue(FormattedCode);
+
+			bReentryFlag = true;
+			EmitValueShader(Context, ShaderValue);
+			bReentryFlag = false;
+
+			if (ShaderValue.bInline)
+			{
+				Code = AllocateString(*Context.Allocator, FormattedCode);
+			}
+			else
+			{
+				// Check to see if we've already generated code for an equivalent expression in either this scope, or any outer visible scope
+				const FSHAHash Hash = HashString(FormattedCode);
+				const TCHAR* Declaration = nullptr;
+
+				FScope* CheckScope = ParentScope;
+				while (CheckScope)
+				{
+					const TCHAR** FoundDeclaration = CheckScope->ExpressionCodeMap.Find(Hash);
+					if (FoundDeclaration)
+					{
+						// Re-use results from previous matching expression
+						Declaration = *FoundDeclaration;
+						break;
+					}
+					CheckScope = CheckScope->ParentScope;
+				}
+
+				if (!Declaration)
+				{
+					check(ParentScope);
+					Declaration = Context.AcquireLocalDeclarationCode();
+					ParentScope->ExpressionCodeMap.Add(Hash, Declaration);
+					ParentScope->MarkLiveRecursive();
+					ParentScope->EmitStatementf(Context, TEXT("const %s %s = %s;"),
+						PrepareValueResult.Type.GetName(),
+						Declaration,
+						FormattedCode.ToString());
+				}
+				Code = Declaration;
+			}
+			check(Code);
+			if (ShaderValue.bHasDependencies)
+			{
+				bReentryFlag = true;
+				EmitShaderDependencies(Context, ShaderValue);
+				bReentryFlag = false;
+			}
+		}
+		else if (PrepareValueResult.EvaluationType == EExpressionEvaluationType::Preshader)
 		{
-			ConstantValue.ToString(Shader::EValueStringFormat::HLSL, FormattedCode);
+			Shader::FPreshaderData Preshader;
+
+			bReentryFlag = true;
+			EmitValuePreshader(Context, Preshader);
+			bReentryFlag = false;
+
+			Context.AddPreshader(PrepareValueResult.Type, Preshader, FormattedCode);
+			Code = AllocateString(*Context.Allocator, FormattedCode);
 		}
 		else
 		{
-			check(EvaluationType == EExpressionEvaluationType::Preshader);
-			check(Preshader);
-			Context.AddPreshader(Type, *Preshader, FormattedCode);
+			check(PrepareValueResult.EvaluationType == EExpressionEvaluationType::Constant);
+			PrepareValueResult.ConstantValue.ToString(Shader::EValueStringFormat::HLSL, FormattedCode);
+			Code = AllocateString(*Context.Allocator, FormattedCode);
 		}
-		Code = AllocateString(*Context.Allocator, FormattedCode);
 	}
 
+	check(Code);
 	return Code;
 }
 
 const TCHAR* FExpression::GetValueShader(FEmitContext& Context, const FType& InType)
 {
-	check(EvaluationType != EExpressionEvaluationType::None);
-
-	if (EvaluationType == EExpressionEvaluationType::Constant && Type != InType)
+	if (PrepareValueResult.EvaluationType == EExpressionEvaluationType::Constant && PrepareValueResult.Type != InType)
 	{
 		// If we need to cast a constant value, perform the cast first, then convert the result to HLSL
-		const Shader::FValue CastConstantValue = Shader::Cast(ConstantValue, Type);
+		const Shader::FValue ConstantValue = Shader::Cast(PrepareValueResult.ConstantValue, InType);
 
 		TStringBuilder<1024> FormattedCode;
-		CastConstantValue.ToString(Shader::EValueStringFormat::HLSL, FormattedCode);
+		ConstantValue.ToString(Shader::EValueStringFormat::HLSL, FormattedCode);
 		return AllocateString(*Context.Allocator, FormattedCode);
 	}
 
 	const ECastFlags CastFlags = ECastFlags::ReplicateScalar | ECastFlags::AllowAppendZeroes | ECastFlags::AllowTruncate;
-	return Context.CastShaderValue(this, GetValueShader(Context), Type, InType, CastFlags);
+	return Context.CastShaderValue(this, GetValueShader(Context), PrepareValueResult.Type, InType, CastFlags);
 }
 
 void FExpression::GetValuePreshader(FEmitContext& Context, Shader::FPreshaderData& OutPreshader)
 {
-	switch (EvaluationType)
+	switch (PrepareValueResult.EvaluationType)
 	{
 	case EExpressionEvaluationType::Preshader:
-		check(Preshader);
-		OutPreshader.Append(*Preshader);
+		if (ForwardedValue)
+		{
+			ForwardedValue->EmitValuePreshader(Context, OutPreshader);
+		}
+		else
+		{
+			EmitValuePreshader(Context, OutPreshader);
+		}
 		break;
 	case EExpressionEvaluationType::Constant:
-		OutPreshader.WriteOpcode(Shader::EPreshaderOpcode::Constant).Write(ConstantValue);
+		OutPreshader.WriteOpcode(Shader::EPreshaderOpcode::Constant).Write(PrepareValueResult.ConstantValue);
 		break;
 	default:
 		check(false);
@@ -777,8 +771,8 @@ void FExpression::GetValuePreshader(FEmitContext& Context, Shader::FPreshaderDat
 
 Shader::FValue FExpression::GetValueConstant(FEmitContext& Context)
 {
-	check(EvaluationType == EExpressionEvaluationType::Constant);
-	return ConstantValue;
+	check(PrepareValueResult.EvaluationType == EExpressionEvaluationType::Constant);
+	return PrepareValueResult.ConstantValue;
 }
 
 bool FScope::HasParentScope(const FScope& InParentScope) const
@@ -829,12 +823,12 @@ void FScope::UseExpression(FExpression* Expression)
 	Visitor.VisitNode(Expression);
 }
 
-void FScope::InternalEmitCode(FEmitContext& Context, FCodeList& List, FScope* NestedScope, const TCHAR* String, int32 Length)
+void FScope::InternalEmitCode(FEmitContext& Context, FCodeList& List, ENextScopeFormat ScopeFormat, FScope* NextScope, const TCHAR* String, int32 Length)
 {
-	if (NestedScope && NestedScope->Statement && !NestedScope->Statement->bEmitHLSL)
+	if (NextScope && NextScope->Statement && !NextScope->Statement->bEmitHLSL)
 	{
-		NestedScope->Statement->bEmitHLSL = true;
-		NestedScope->Statement->EmitHLSL(Context);
+		NextScope->Statement->bEmitHLSL = true;
+		NextScope->Statement->EmitHLSL(Context);
 	}
 
 	const int32 SizeofString = sizeof(TCHAR) * Length;
@@ -843,7 +837,8 @@ void FScope::InternalEmitCode(FEmitContext& Context, FCodeList& List, FScope* Ne
 	FMemory::Memcpy(CodeEntry->String, String, SizeofString);
 	CodeEntry->String[Length] = 0;
 	CodeEntry->Length = Length;
-	CodeEntry->NestedScope = NestedScope;
+	CodeEntry->NextScope = NextScope;
+	CodeEntry->ScopeFormat = ScopeFormat;
 	CodeEntry->Next = nullptr;
 
 	if (!List.First)
@@ -859,11 +854,11 @@ void FScope::InternalEmitCode(FEmitContext& Context, FCodeList& List, FScope* Ne
 	List.Num++;
 }
 
-void RequestScopeTypes(FUpdateTypeContext& Context, const FScope* InScope)
+void PrepareScopeValues(FEmitContext& Context, const FScope* InScope)
 {
 	if (InScope && InScope->Statement)
 	{
-		InScope->Statement->RequestTypes(Context);
+		InScope->Statement->PrepareValues(Context);
 	}
 }
 
@@ -904,26 +899,35 @@ static void WriteIndent(int32 IndentLevel, FStringBuilderBase& InOutString)
 
 void FScope::MarkLive()
 {
-	bLive = true;
+	if (State == EScopeState::Uninitialized)
+	{
+		State = EScopeState::Live;
+	}
 }
 
 void FScope::MarkLiveRecursive()
 {
 	FScope* Scope = this;
-	while (Scope)
+	while (Scope && Scope->State == EScopeState::Uninitialized)
 	{
-		Scope->bLive = true;
+		Scope->State = EScopeState::Live;
 		Scope = Scope->ParentScope;
 	}
 }
 
+void FScope::MarkDead()
+{
+	// TODO - mark child scopes as dead as well
+	State = EScopeState::Dead;
+}
+
 void FScope::WriteHLSL(int32 Indent, FStringBuilderBase& OutString) const
 {
 	{
 		const FCodeEntry* CodeDeclaration = Declarations.First;
 		while (CodeDeclaration)
 		{
-			check(!CodeDeclaration->NestedScope);
+			check(!CodeDeclaration->NextScope);
 			WriteIndent(Indent, OutString);
 			OutString.Append(CodeDeclaration->String, CodeDeclaration->Length);
 			OutString.Append(TEXT('\n'));
@@ -941,13 +945,24 @@ void FScope::WriteHLSL(int32 Indent, FStringBuilderBase& OutString) const
 				OutString.Append(CodeStatement->String, CodeStatement->Length);
 				OutString.Append(TEXT('\n'));
 			}
-			if (CodeStatement->NestedScope)
+			if (CodeStatement->NextScope)
 			{
-				WriteIndent(Indent, OutString);
-				OutString.Append(TEXT("{\n"));
-				CodeStatement->NestedScope->WriteHLSL(Indent + 1, OutString);
-				WriteIndent(Indent, OutString);
-				OutString.Append(TEXT("}\n"));
+				int32 NextIndent = Indent;
+				bool bNeedToCloseScope = false;
+				if (CodeStatement->ScopeFormat == ENextScopeFormat::Scoped)
+				{
+					WriteIndent(Indent, OutString);
+					OutString.Append(TEXT("{\n"));
+					NextIndent++;
+					bNeedToCloseScope = true;
+				}
+
+				CodeStatement->NextScope->WriteHLSL(NextIndent, OutString);
+				if (bNeedToCloseScope)
+				{
+					WriteIndent(Indent, OutString);
+					OutString.Append(TEXT("}\n"));
+				}
 			}
 			CodeStatement = CodeStatement->Next;
 		}
@@ -981,13 +996,10 @@ bool FTree::EmitHLSL(FEmitContext& Context,
 		return false;
 	}
 
+	PrepareScopeValues(Context, RootScope);
+	if (Context.Errors.Num() > 0)
 	{
-		FUpdateTypeContext UpdateTypeContext(Context.Errors);
-		RequestScopeTypes(UpdateTypeContext, RootScope);
-		if (Context.Errors.Num() > 0)
-		{
-			return false;
-		}
+		return false;
 	}
 
 	ResultStatement->bEmitHLSL = true;
diff --git a/Engine/Source/Runtime/Engine/Private/HLSLTree/HLSLTreeCommon.cpp b/Engine/Source/Runtime/Engine/Private/HLSLTree/HLSLTreeCommon.cpp
--- a/Engine/Source/Runtime/Engine/Private/HLSLTree/HLSLTreeCommon.cpp
+++ b/Engine/Source/Runtime/Engine/Private/HLSLTree/HLSLTreeCommon.cpp
@@ -9,16 +9,24 @@ namespace UE
 namespace HLSLTree
 {
 
+FBinaryOpDescription::FBinaryOpDescription()
+	: Name(nullptr), Operator(nullptr), PreshaderOpcode(Shader::EPreshaderOpcode::Nop)
+{}
+
+FBinaryOpDescription::FBinaryOpDescription(const TCHAR* InName, const TCHAR* InOperator, Shader::EPreshaderOpcode InOpcode)
+	: Name(InName), Operator(InOperator), PreshaderOpcode(InOpcode)
+{}
+
 FBinaryOpDescription GetBinaryOpDesription(EBinaryOp Op)
 {
 	switch (Op)
 	{
-	case EBinaryOp::None: return FBinaryOpDescription(TEXT("None"), TEXT("")); break;
-	case EBinaryOp::Add: return FBinaryOpDescription(TEXT("Add"), TEXT("+")); break;
-	case EBinaryOp::Sub: return FBinaryOpDescription(TEXT("Subtract"), TEXT("-")); break;
-	case EBinaryOp::Mul: return FBinaryOpDescription(TEXT("Multiply"), TEXT("*")); break;
-	case EBinaryOp::Div: return FBinaryOpDescription(TEXT("Divide"), TEXT("/")); break;
-	case EBinaryOp::Less: return FBinaryOpDescription(TEXT("Less"), TEXT("<")); break;
+	case EBinaryOp::None: return FBinaryOpDescription(TEXT("None"), TEXT(""), Shader::EPreshaderOpcode::Nop); break;
+	case EBinaryOp::Add: return FBinaryOpDescription(TEXT("Add"), TEXT("+"), Shader::EPreshaderOpcode::Add); break;
+	case EBinaryOp::Sub: return FBinaryOpDescription(TEXT("Subtract"), TEXT("-"), Shader::EPreshaderOpcode::Sub); break;
+	case EBinaryOp::Mul: return FBinaryOpDescription(TEXT("Multiply"), TEXT("*"), Shader::EPreshaderOpcode::Mul); break;
+	case EBinaryOp::Div: return FBinaryOpDescription(TEXT("Divide"), TEXT("/"), Shader::EPreshaderOpcode::Div); break;
+	case EBinaryOp::Less: return FBinaryOpDescription(TEXT("Less"), TEXT("<"), Shader::EPreshaderOpcode::Nop); break;
 	default: checkNoEntry(); return FBinaryOpDescription();
 	}
 }
@@ -77,7 +85,7 @@ FSwizzleParameters MakeSwizzleMask(bool bInR, bool bInG, bool bInB, bool bInA)
 	return FSwizzleParameters(ComponentIndex[0], ComponentIndex[1], ComponentIndex[2], ComponentIndex[3]);
 }
 
-void FExpressionConstant::UpdateType(FUpdateTypeContext& Context, const FRequestedType& RequestedType)
+void FExpressionConstant::PrepareValue(FEmitContext& Context, const FRequestedType& RequestedType)
 {
 	if (Value.Type.IsStruct())
 	{
@@ -85,53 +93,49 @@ void FExpressionConstant::UpdateType(FUpdateTypeContext& Context, const FRequest
 		{
 			return Context.Errors.AddError(this, TEXT("Type mismatch"));
 		}
-		return SetType(Context, Value.Type);
+
+		// TODO - Support for preshader structs?
+		return SetType(Context, EExpressionEvaluationType::Shader, Value.Type);
 	}
 	else
 	{
-		return SetType(Context, Shader::MakeValueTypeWithRequestedNumComponents(Value.Type, RequestedType.GetRequestedNumComponents()));
+		return SetType(Context, EExpressionEvaluationType::Constant, Shader::MakeValueTypeWithRequestedNumComponents(Value.Type, RequestedType.GetRequestedNumComponents()));
 	}
 }
 
-
-void FExpressionConstant::PrepareValue(FEmitContext& Context)
+void FExpressionConstant::EmitValueShader(FEmitContext& Context, FShaderValue& OutShader) const
 {
-	check(Value.Component.Num() == Value.Type.GetNumComponents());
-	if (Value.Type.IsStruct())
+	const TArrayView<const Shader::EValueComponentType> ComponentTypes = Value.Type.StructType->ComponentTypes;
+	OutShader.Code.Append(TEXT("{ "));
+	for (int32 ComponentIndex = 0; ComponentIndex < Value.Component.Num(); ++ComponentIndex)
 	{
-		const TArrayView<const Shader::EValueComponentType> ComponentTypes = Value.Type.StructType->ComponentTypes;
-		// TODO - Support for preshader structs?
-		TStringBuilder<1024> FormattedCode;
-		FormattedCode.Append(TEXT("{ "));
-		for (int32 ComponentIndex = 0; ComponentIndex < Value.Component.Num(); ++ComponentIndex)
+		if (ComponentIndex > 0)
 		{
-			if (ComponentIndex > 0)
-			{
-				FormattedCode.Append(TEXT(", "));
-			}
-			Value.Component[ComponentIndex].ToString(ComponentTypes[ComponentIndex], FormattedCode);
+			OutShader.Code.Append(TEXT(", "));
 		}
-		FormattedCode.Append(TEXT(" }"));
-		return SetValueShader(Context, FormattedCode);
+		Value.Component[ComponentIndex].ToString(ComponentTypes[ComponentIndex], OutShader.Code);
 	}
-	else
+	OutShader.Code.Append(TEXT(" }"));
+}
+
+void FExpressionConstant::EmitValuePreshader(FEmitContext& Context, Shader::FPreshaderData& OutPreshader) const
+{
+	Shader::FValue BaseValue(Value.Type.ValueType);
+	for (int32 ComponentIndex = 0; ComponentIndex < Value.Component.Num(); ++ComponentIndex)
 	{
-		Shader::FValue BaseValue(Value.Type.ValueType);
-		for (int32 ComponentIndex = 0; ComponentIndex < Value.Component.Num(); ++ComponentIndex)
-		{
-			BaseValue.Component[ComponentIndex] = Value.Component[ComponentIndex];
-		}
-		return SetValueConstant(Context, BaseValue);
+		BaseValue.Component[ComponentIndex] = Value.Component[ComponentIndex];
 	}
+	OutPreshader.WriteOpcode(Shader::EPreshaderOpcode::Constant).Write(BaseValue);
 }
 
-void FExpressionMaterialParameter::UpdateType(FUpdateTypeContext& Context, const FRequestedType& RequestedType)
+void FExpressionMaterialParameter::PrepareValue(FEmitContext& Context, const FRequestedType& RequestedType)
 {
 	const Shader::EValueType CurrentType = Shader::MakeValueTypeWithRequestedNumComponents(GetShaderValueType(ParameterType), RequestedType.GetRequestedNumComponents());
-	return SetType(Context, CurrentType);
+	const EExpressionEvaluationType EvaluationType = IsStaticMaterialParameter(ParameterType) ? EExpressionEvaluationType::Constant : EExpressionEvaluationType::Preshader;
+	return SetType(Context, EvaluationType, CurrentType);
 }
 
-void FExpressionMaterialParameter::PrepareValue(FEmitContext& Context)
+void FExpressionMaterialParameter::EmitValuePreshader(FEmitContext& Context, Shader::FPreshaderData& OutPreshader) const
 {
 	if (ParameterType == EMaterialParameterType::StaticSwitch)
 	{
@@ -145,8 +149,7 @@ void FExpressionMaterialParameter::PrepareValue(FEmitContext& Context)
 				break;
 			}
 		}
-
-		return SetValueConstant(Context, Value);
+		OutPreshader.WriteOpcode(Shader::EPreshaderOpcode::Constant).Write(Value);
 	}
 	else
 	{
@@ -163,14 +166,11 @@ void FExpressionMaterialParameter::PrepareValue(FEmitContext& Context)
 		}
 		const int32 ParameterIndex = Context.MaterialCompilationOutput->UniformExpressionSet.FindOrAddNumericParameter(ParameterType, ParameterName, DefaultOffset);
 		check(ParameterIndex >= 0 && ParameterIndex <= 0xffff);
-
-		Shader::FPreshaderData LocalPreshader;
-		LocalPreshader.WriteOpcode(Shader::EPreshaderOpcode::Parameter).Write((uint16)ParameterIndex);
-		return SetValuePreshader(Context, LocalPreshader);
+		OutPreshader.WriteOpcode(Shader::EPreshaderOpcode::Parameter).Write((uint16)ParameterIndex);
 	}
 }
 
-void FExpressionExternalInput::PrepareValue(FEmitContext& Context)
+void FExpressionExternalInput::EmitValueShader(FEmitContext& Context, FShaderValue& OutShader) const
 {
 	const int32 TypeIndex = (int32)InputType;
 	const int32 TexCoordIndex = TypeIndex - (int32)EExternalInputType::TexCoord0;
@@ -178,7 +178,8 @@ void FExpressionExternalInput::PrepareValue(FEmitContext& Context)
 	if (TexCoordIndex >= 0 && TexCoordIndex < 8)
 	{
 		Context.NumTexCoords = FMath::Max(Context.NumTexCoords, TexCoordIndex + 1);
-		return SetValueInlineShaderf(Context, TEXT("Parameters.TexCoords[%u].xy"), TexCoordIndex);
+		OutShader.bInline = true;
+		OutShader.Code.Appendf(TEXT("Parameters.TexCoords[%u].xy"), TexCoordIndex);
 	}
 	else
 	{
@@ -186,19 +187,14 @@ void FExpressionExternalInput::PrepareValue(FEmitContext& Context)
 	}
 }
 
-void FExpressionTextureSample::UpdateType(FUpdateTypeContext& Context, const FRequestedType& RequestedType)
+void FExpressionTextureSample::PrepareValue(FEmitContext& Context, const FRequestedType& RequestedType)
 {
-	RequestExpressionType(Context, TexCoordExpression, 2);
-	return SetType(Context, Shader::EValueType::Float4);
+	PrepareExpressionValue(Context, TexCoordExpression, 2);
+	return SetType(Context, EExpressionEvaluationType::Shader, Shader::EValueType::Float4);
 }
 
-void FExpressionTextureSample::PrepareValue(FEmitContext& Context)
+void FExpressionTextureSample::EmitValueShader(FEmitContext& Context, FShaderValue& OutShader) const
 {
-	if (PrepareExpressionValue(Context, TexCoordExpression) == EExpressionEvaluationType::None)
-	{
-		return;
-	}
-
 	const FTextureDescription& Desc = Declaration->Description;
 	const EMaterialValueType MaterialValueType = Desc.Texture->GetMaterialType();
 	EMaterialTextureParameterType TextureType = EMaterialTextureParameterType::Count;
@@ -323,7 +319,7 @@ void FExpressionTextureSample::PrepareValue(FEmitContext& Context)
 		}
 	}
 
-	return SetValueShaderf(Context, TEXT("%s(%s(%s, %s, %s))"),
+	OutShader.Code.Appendf(TEXT("%s(%s(%s, %s, %s))"),
 		SamplerTypeFunction,
 		SampleFunctionName,
 		*TextureName,
@@ -331,7 +327,7 @@ void FExpressionTextureSample::PrepareValue(FEmitContext& Context)
 		TexCoordExpression->GetValueShader(Context, Shader::EValueType::Float2));
 }
 
-void FExpressionGetStructField::UpdateType(FUpdateTypeContext& Context, const FRequestedType& RequestedType)
+void FExpressionGetStructField::PrepareValue(FEmitContext& Context, const FRequestedType& RequestedType)
 {
 	const FStructFieldRef FieldRef = StructType->FindFieldByName(FieldName);
 	if (!FieldRef)
@@ -339,30 +335,25 @@ void FExpressionGetStructField::UpdateType(FUpdateTypeContext& Context, const FR
 		return Context.Errors.AddErrorf(this, TEXT("Invalid field %s"), FieldName);
 	}
 
-	if (RequestExpressionType(Context, StructExpression, RequestedType.MakeFieldAccess(StructType, FieldRef)) != FType(StructType))
+	const FPrepareValueResult StructResult = PrepareExpressionValue(Context, StructExpression, RequestedType.MakeFieldAccess(StructType, FieldRef));
+	if (StructResult && StructResult.Type != FType(StructType))
 	{
 		return Context.Errors.AddErrorf(this, TEXT("Expected type %s"), StructType->Name);
 	}
 
-	return SetType(Context, FieldRef.Type);
+	// TODO - preshader structs
+	return SetType(Context, EExpressionEvaluationType::Shader, FieldRef.Type);
 }
 
-void FExpressionGetStructField::PrepareValue(FEmitContext& Context)
+void FExpressionGetStructField::EmitValueShader(FEmitContext& Context, FShaderValue& OutShader) const
 {
-	const FStructFieldRef FieldRef = StructType->FindFieldByName(FieldName);
-	check(FieldRef);
-
-	if (PrepareExpressionValue(Context, StructExpression) == EExpressionEvaluationType::None)
-	{
-		return;
-	}
-
-	return SetValueInlineShaderf(Context, TEXT("%s.%s"),
+	OutShader.bInline = true;
+	OutShader.Code.Appendf(TEXT("%s.%s"),
 		StructExpression->GetValueShader(Context, StructType),
 		FieldName);
 }
 
-void FExpressionSetStructField::UpdateType(FUpdateTypeContext& Context, const FRequestedType& RequestedType)
+void FExpressionSetStructField::PrepareValue(FEmitContext& Context, const FRequestedType& RequestedType)
 {
 	const FStructFieldRef FieldRef = StructType->FindFieldByName(FieldName);
 	if (!FieldRef)
@@ -370,116 +361,98 @@ void FExpressionSetStructField::UpdateType(FUpdateTypeContext& Context, const FR
 		return Context.Errors.AddErrorf(this, TEXT("Invalid field %s"), FieldName);
 	}
 
-	const FType StructExpressionType = RequestExpressionType(Context, StructExpression, RequestedType.CopyWithFieldRemoved(FieldRef));
-	if (StructExpressionType && StructExpressionType.StructType != StructType)
+	const FPrepareValueResult StructResult = PrepareExpressionValue(Context, StructExpression, RequestedType.CopyWithFieldRemoved(FieldRef));
+	if (StructResult && StructResult.Type != FType(StructType))
 	{
 		return Context.Errors.AddErrorf(this, TEXT("Expected type %s"), StructType->Name);
 	}
 
-	RequestExpressionType(Context, FieldExpression, RequestedType.GetField(FieldRef));
-	return SetType(Context, StructType);
+	const FPrepareValueResult FieldResult = PrepareExpressionValue(Context, FieldExpression, RequestedType.GetField(FieldRef));
+	if (!FieldResult)
+	{
+		return SetForwardValue(Context, StructExpression, RequestedType);
+	}
+
+	// TODO - preshader structs
+	return SetType(Context, EExpressionEvaluationType::Shader, StructType);
 }
 
-void FExpressionSetStructField::PrepareValue(FEmitContext& Context)
+void FExpressionSetStructField::EmitValueShader(FEmitContext& Context, FShaderValue& OutShader) const
 {
-	const EExpressionEvaluationType StructEvaluation = PrepareExpressionValue(Context, StructExpression);
-	const EExpressionEvaluationType FieldEvaluation = PrepareExpressionValue(Context, FieldExpression);
-
-	if (FieldEvaluation != EExpressionEvaluationType::None)
+	const FStructFieldRef FieldRef = StructType->FindFieldByName(FieldName);
+	check(FieldRef);
+	if (StructExpression->GetEvaluationType() != EExpressionEvaluationType::None)
 	{
-		const FStructFieldRef FieldRef = StructType->FindFieldByName(FieldName);
-		check(FieldRef);
-
-		if (StructEvaluation != EExpressionEvaluationType::None)
-		{
-			return SetValueShaderf(Context, TEXT("%s_Set%s(%s, %s)"),
-				StructType->Name,
-				FieldName,
-				StructExpression->GetValueShader(Context, StructType),
-				FieldExpression->GetValueShader(Context, FieldRef.Type));
-		}
-		else
-		{
-			// StructExpression is not used, so default to a zero-initialized struct
-			// This will happen if all the accessed struct fields are explicitly defined
-			return SetValueShaderf(Context, TEXT("%s_Set%s((%s)0, %s)"),
-				StructType->Name,
-				FieldName,
-				StructType->Name,
-				FieldExpression->GetValueShader(Context, FieldRef.Type));
-		}
+		OutShader.Code.Appendf(TEXT("%s_Set%s(%s, %s)"),
+			StructType->Name,
+			FieldName,
+			StructExpression->GetValueShader(Context, StructType),
+			FieldExpression->GetValueShader(Context, FieldRef.Type));
 	}
 	else
 	{
-		// Skip assigning the field, and just forward the struct value
-		check(StructEvaluation != EExpressionEvaluationType::None);
-		return SetValueForward(Context, StructExpression);
+		// StructExpression is not used, so default to a zero-initialized struct
+		// This will happen if all the accessed struct fields are explicitly defined
+		OutShader.Code.Appendf(TEXT("%s_Set%s((%s)0, %s)"),
+			StructType->Name,
+			FieldName,
+			StructType->Name,
+			FieldExpression->GetValueShader(Context, FieldRef.Type));
 	}
 }
 
-void FExpressionSelect::UpdateType(FUpdateTypeContext& Context, const FRequestedType& RequestedType)
+void FExpressionSelect::PrepareValue(FEmitContext& Context, const FRequestedType& RequestedType)
 {
-	const FType ConditionType = RequestExpressionType(Context, ConditionExpression, 1);
-	const FType LhsType = RequestExpressionType(Context, FalseExpression, RequestedType);
-	const FType RhsType = RequestExpressionType(Context, TrueExpression, RequestedType);
-	if (!ConditionType || !LhsType || !RhsType)
+	const FPrepareValueResult ConditionResult = PrepareExpressionValue(Context, ConditionExpression, 1);
+	if (ConditionResult.EvaluationType == EExpressionEvaluationType::Constant)
+	{
+		const bool bCondition = ConditionResult.ConstantValue.AsBoolScalar();
+		return SetForwardValue(Context, bCondition ? TrueExpression : FalseExpression, RequestedType);
+	}
+
+	const FPrepareValueResult LhsResult = PrepareExpressionValue(Context, FalseExpression, RequestedType);
+	const FPrepareValueResult RhsResult = PrepareExpressionValue(Context, TrueExpression, RequestedType);
+	if (!ConditionResult || !LhsResult || !RhsResult)
 	{
 		return;
 	}
 
-	const Shader::FValueTypeDescription LhsTypeDesc = Shader::GetValueTypeDescription(LhsType);
-	const Shader::FValueTypeDescription RhsTypeDesc = Shader::GetValueTypeDescription(RhsType);
+	const Shader::FValueTypeDescription LhsTypeDesc = Shader::GetValueTypeDescription(LhsResult.Type);
+	const Shader::FValueTypeDescription RhsTypeDesc = Shader::GetValueTypeDescription(RhsResult.Type);
 	if (LhsTypeDesc.ComponentType != RhsTypeDesc.ComponentType)
 	{
 		return Context.Errors.AddError(this, TEXT("Type mismatch"));
 	}
 
-	if (LhsType.IsStruct())
+	if (LhsResult.Type.IsStruct())
 	{
-		return SetType(Context, LhsType);
+		return SetType(Context, EExpressionEvaluationType::Shader, LhsResult.Type);
 	}
 	else
 	{
 		const int8 NumComponents = FMath::Max(LhsTypeDesc.NumComponents, RhsTypeDesc.NumComponents);
-		return SetType(Context, Shader::MakeValueType(LhsTypeDesc.ComponentType, NumComponents));
+		return SetType(Context, EExpressionEvaluationType::Shader, Shader::MakeValueType(LhsTypeDesc.ComponentType, NumComponents));
 	}
 }
 
-void FExpressionSelect::PrepareValue(FEmitContext& Context)
+void FExpressionSelect::EmitValueShader(FEmitContext& Context, FShaderValue& OutShader) const
 {
-	const EExpressionEvaluationType ConditionEvaluation = PrepareExpressionValue(Context, ConditionExpression);
-	if (ConditionEvaluation == EExpressionEvaluationType::None)
-	{
-		return;
-	}
-
-	if (ConditionEvaluation == EExpressionEvaluationType::Constant)
-	{
-		const bool bCondition = ConditionExpression->GetValueConstant(Context).AsBoolScalar();
-		return SetValueForward(Context, bCondition ? TrueExpression : FalseExpression);
-	}
-	else
-	{
-		const EExpressionEvaluationType TrueEvaluation = PrepareExpressionValue(Context, TrueExpression);
-		const EExpressionEvaluationType FalseEvaluation = PrepareExpressionValue(Context, FalseExpression);
-		if (TrueEvaluation == EExpressionEvaluationType::None || FalseEvaluation == EExpressionEvaluationType::None)
-		{
-			return;
-		}
+	OutShader.Code.Appendf(TEXT("(%s ? %s : %s)"),
+		ConditionExpression->GetValueShader(Context, Shader::EValueType::Bool1),
+		TrueExpression->GetValueShader(Context, GetType()),
+		FalseExpression->GetValueShader(Context, GetType()));
+}
 
-		// TODO - preshader
-		return SetValueShaderf(Context, TEXT("(%s ? %s : %s)"),
-			ConditionExpression->GetValueShader(Context, Shader::EValueType::Bool1),
-			TrueExpression->GetValueShader(Context, GetType()),
-			FalseExpression->GetValueShader(Context, GetType()));
-	}
+void FExpressionSelect::EmitValuePreshader(FEmitContext& Context, Shader::FPreshaderData& OutPreshader) const
+{
+	check(false); // TODO
 }
 
-void FExpressionBinaryOp::UpdateType(FUpdateTypeContext& Context, const FRequestedType& RequestedType)
+void FExpressionBinaryOp::PrepareValue(FEmitContext& Context, const FRequestedType& RequestedType)
 {
-	const FType LhsType = RequestExpressionType(Context, Lhs, RequestedType);
-	const FType RhsType = RequestExpressionType(Context, Rhs, RequestedType);
-	if (!LhsType || !RhsType)
+	const FPrepareValueResult LhsResult = PrepareExpressionValue(Context, Lhs, RequestedType);
+	const FPrepareValueResult RhsResult = PrepareExpressionValue(Context, Rhs, RequestedType);
+	if (!LhsResult || !RhsResult)
 	{
 		return;
 	}
@@ -489,10 +462,10 @@ void FExpressionBinaryOp::UpdateType(FUpdateTypeContext& Context, const FRequest
 	switch (Op)
 	{
 	case EBinaryOp::Less:
-		ResultType = Shader::MakeComparisonResultType(LhsType, RhsType, ErrorMessage);
+		ResultType = Shader::MakeComparisonResultType(LhsResult.Type, RhsResult.Type, ErrorMessage);
 		break;
 	default:
-		ResultType = Shader::MakeArithmeticResultType(LhsType, RhsType, ErrorMessage);
+		ResultType = Shader::MakeArithmeticResultType(LhsResult.Type, RhsResult.Type, ErrorMessage);
 		break;
 	}
 	if (!ErrorMessage.IsEmpty())
@@ -500,58 +473,39 @@ void FExpressionBinaryOp::UpdateType(FUpdateTypeContext& Context, const FRequest
 		return Context.Errors.AddError(this, ErrorMessage);
 	}
 
-	return SetType(Context, ResultType);
+	const FBinaryOpDescription OpDesc = GetBinaryOpDesription(Op);
+	const EExpressionEvaluationType EvaluationType = OpDesc.PreshaderOpcode != Shader::EPreshaderOpcode::Nop
+		? CombineEvaluationTypes(LhsResult.EvaluationType, RhsResult.EvaluationType)
+		: EExpressionEvaluationType::Shader;
+	return SetType(Context, EvaluationType, ResultType);
 }
 
-void FExpressionBinaryOp::PrepareValue(FEmitContext& Context)
+void FExpressionBinaryOp::EmitValueShader(FEmitContext& Context, FShaderValue& OutShader) const
 {
-	const EExpressionEvaluationType LhsEvaluation = PrepareExpressionValue(Context, Lhs);
-	const EExpressionEvaluationType RhsEvaluation = PrepareExpressionValue(Context, Rhs);
-	if (LhsEvaluation == EExpressionEvaluationType::None || RhsEvaluation == EExpressionEvaluationType::None)
+	const TCHAR* LhsCode = Lhs->GetValueShader(Context);
+	const TCHAR* RhsCode = Rhs->GetValueShader(Context);
+	switch (Op)
 	{
-		return;
+	case EBinaryOp::Add: OutShader.Code.Appendf(TEXT("(%s + %s)"), LhsCode, RhsCode); break;
+	case EBinaryOp::Sub: OutShader.Code.Appendf(TEXT("(%s - %s)"), LhsCode, RhsCode); break;
+	case EBinaryOp::Mul: OutShader.Code.Appendf(TEXT("(%s * %s)"), LhsCode, RhsCode); break;
+	case EBinaryOp::Div: OutShader.Code.Appendf(TEXT("(%s / %s)"), LhsCode, RhsCode); break;
+	case EBinaryOp::Less: OutShader.Code.Appendf(TEXT("(%s < %s)"), LhsCode, RhsCode); break;
+	default: checkNoEntry(); break;
 	}
+}
 
-	const EExpressionEvaluationType Evaluation = CombineEvaluationTypes(LhsEvaluation, RhsEvaluation);
-	Shader::EPreshaderOpcode Opcode = Shader::EPreshaderOpcode::Nop;
-	if (Evaluation != EExpressionEvaluationType::Shader)
-	{
-		switch (Op)
-		{
-		case EBinaryOp::Add: Opcode = Shader::EPreshaderOpcode::Add; break;
-		case EBinaryOp::Sub: Opcode = Shader::EPreshaderOpcode::Add; break;
-		case EBinaryOp::Mul: Opcode = Shader::EPreshaderOpcode::Mul; break;
-		case EBinaryOp::Div: Opcode = Shader::EPreshaderOpcode::Div; break;
-		default: break;
-		}
-	}
+void FExpressionBinaryOp::EmitValuePreshader(FEmitContext& Context, Shader::FPreshaderData& OutPreshader) const
+{
+	const FBinaryOpDescription OpDesc = GetBinaryOpDesription(Op);
+	check(OpDesc.PreshaderOpcode != Shader::EPreshaderOpcode::Nop);
 
-	if (Opcode == Shader::EPreshaderOpcode::Nop)
-	{
-		// Either preshader not supported for our inputs, or not for this operation type
-		const TCHAR* LhsCode = Lhs->GetValueShader(Context);
-		const TCHAR* RhsCode = Rhs->GetValueShader(Context);
-		switch (Op)
-		{
-		case EBinaryOp::Add: return SetValueShaderf(Context, TEXT("(%s + %s)"), LhsCode, RhsCode);
-		case EBinaryOp::Sub: return SetValueShaderf(Context, TEXT("(%s - %s)"), LhsCode, RhsCode);
-		case EBinaryOp::Mul: return SetValueShaderf(Context, TEXT("(%s * %s)"), LhsCode, RhsCode);
-		case EBinaryOp::Div: return SetValueShaderf(Context, TEXT("(%s / %s)"), LhsCode, RhsCode);
-		case EBinaryOp::Less: return SetValueShaderf(Context, TEXT("(%s < %s)"), LhsCode, RhsCode);
-		default: checkNoEntry(); return;
-		}
-	}
-	else
-	{
-		Shader::FPreshaderData LocalPreshader;
-		Lhs->GetValuePreshader(Context, LocalPreshader);
-		Rhs->GetValuePreshader(Context, LocalPreshader);
-		LocalPreshader.WriteOpcode(Opcode);
-		return SetValuePreshader(Context, Evaluation, LocalPreshader);
-	}
+	Lhs->GetValuePreshader(Context, OutPreshader);
+	Rhs->GetValuePreshader(Context, OutPreshader);
+	OutPreshader.WriteOpcode(OpDesc.PreshaderOpcode);
 }
 
-void FExpressionSwizzle::UpdateType(FUpdateTypeContext& Context, const FRequestedType& RequestedType)
+void FExpressionSwizzle::PrepareValue(FEmitContext& Context, const FRequestedType& RequestedType)
 {
 	FRequestedType RequestedInputType(4, false);
 	int32 MaxComponentIndex = 0;
@@ -565,18 +519,12 @@ void FExpressionSwizzle::UpdateType(FUpdateTypeContext& Context, const FRequeste
 		}
 	}
 
-	const Shader::EValueType InputType = RequestExpressionType(Context, Input, RequestedInputType);
-	return SetType(Context, Shader::MakeValueType(InputType, MaxComponentIndex + 1));
+	const FPrepareValueResult InputResult = PrepareExpressionValue(Context, Input, RequestedInputType);
+	return SetType(Context, InputResult.EvaluationType, Shader::MakeValueType(InputResult.Type, MaxComponentIndex + 1));
 }
 
-void FExpressionSwizzle::PrepareValue(FEmitContext& Context)
+void FExpressionSwizzle::EmitValueShader(FEmitContext& Context, FShaderValue& OutShader) const
 {
-	const EExpressionEvaluationType InputEvaluation = PrepareExpressionValue(Context, Input);
-	if (InputEvaluation == EExpressionEvaluationType::None)
-	{
-		return;
-	}
-
 	const Shader::FValueTypeDescription TypeDesc = Shader::GetValueTypeDescription(GetType());
 	const int32 NumSwizzleComponents = FMath::Min<int32>(TypeDesc.NumComponents, Parameters.NumComponents);
 
@@ -595,90 +543,84 @@ void FExpressionSwizzle::PrepareValue(FEmitContext& Context)
 		}
 	}
 
-	if (InputEvaluation == EExpressionEvaluationType::Shader)
+	OutShader.bInline = true;
+	if (bHasSwizzleReorder)
 	{
-		if (bHasSwizzleReorder)
-		{
-			return SetValueInlineShaderf(Context, TEXT("%s.%s"),
-				Input->GetValueShader(Context),
-				Swizzle);
-		}
-		else
-		{
-			const Shader::EValueType InputType = Shader::MakeValueType(GetType(), NumSwizzleComponents);
-			return SetValueInlineShaderf(Context, TEXT("%s"), Input->GetValueShader(Context, InputType));
-		}
+		OutShader.Code.Appendf(TEXT("%s.%s"),
+			Input->GetValueShader(Context),
+			Swizzle);
 	}
 	else
 	{
-		Shader::FPreshaderData LocalPreshader;
-		Input->GetValuePreshader(Context, LocalPreshader);
-		LocalPreshader.WriteOpcode(Shader::EPreshaderOpcode::ComponentSwizzle)
-			.Write((uint8)Parameters.NumComponents)
-			.Write((uint8)Parameters.ComponentIndex[0])
-			.Write((uint8)Parameters.ComponentIndex[1])
-			.Write((uint8)Parameters.ComponentIndex[2])
-			.Write((uint8)Parameters.ComponentIndex[3]);
-		return SetValuePreshader(Context, InputEvaluation, LocalPreshader);
+		const Shader::EValueType InputType = Shader::MakeValueType(GetType(), NumSwizzleComponents);
+		OutShader.Code.Appendf(TEXT("%s"), Input->GetValueShader(Context, InputType));
 	}
 }
 
-void FExpressionAppend::UpdateType(FUpdateTypeContext& Context, const FRequestedType& RequestedType)
+void FExpressionSwizzle::EmitValuePreshader(FEmitContext& Context, Shader::FPreshaderData& OutPreshader) const
+{
+	Input->GetValuePreshader(Context, OutPreshader);
+	OutPreshader.WriteOpcode(Shader::EPreshaderOpcode::ComponentSwizzle)
+		.Write((uint8)Parameters.NumComponents)
+		.Write((uint8)Parameters.ComponentIndex[0])
+		.Write((uint8)Parameters.ComponentIndex[1])
+		.Write((uint8)Parameters.ComponentIndex[2])
+		.Write((uint8)Parameters.ComponentIndex[3]);
+}
+
+void FExpressionAppend::PrepareValue(FEmitContext& Context, const FRequestedType& RequestedType)
 {
-	const FType LhsType = RequestExpressionType(Context, Lhs, 4);
-	const FType RhsType = RequestExpressionType(Context, Rhs, 4);
-	if (!LhsType || !RhsType)
+	const FPrepareValueResult LhsResult = PrepareExpressionValue(Context, Lhs, 4);
+	const FPrepareValueResult RhsResult = PrepareExpressionValue(Context, Rhs, 4);
+	if (!LhsResult || !RhsResult)
 	{
 		return;
 	}
 
-	if (LhsType.IsStruct() || RhsType.IsStruct())
+	if (LhsResult.Type.IsStruct() || RhsResult.Type.IsStruct())
 	{
 		return Context.Errors.AddError(this, TEXT("Not supported for structs"));
 	}
 
-	const Shader::FValueTypeDescription LhsTypeDesc = Shader::GetValueTypeDescription(LhsType);
-	const Shader::FValueTypeDescription RhsTypeDesc = Shader::GetValueTypeDescription(RhsType);
+	const Shader::FValueTypeDescription LhsTypeDesc = Shader::GetValueTypeDescription(LhsResult.Type);
+	const Shader::FValueTypeDescription RhsTypeDesc = Shader::GetValueTypeDescription(RhsResult.Type);
 	if (LhsTypeDesc.ComponentType != RhsTypeDesc.ComponentType)
 	{
 		return Context.Errors.AddError(this, TEXT("Type mismatch"));
 	}
 
+	const EExpressionEvaluationType EvaluationType = CombineEvaluationTypes(LhsResult.EvaluationType, RhsResult.EvaluationType);
 	const int32 NumComponents = FMath::Min<int32>(LhsTypeDesc.NumComponents + RhsTypeDesc.NumComponents, RequestedType.GetRequestedNumComponents());
-	return SetType(Context, Shader::MakeValueType(LhsTypeDesc.ComponentType, NumComponents));
+	return SetType(Context, EvaluationType, Shader::MakeValueType(LhsTypeDesc.ComponentType, NumComponents));
 }
 
-void FExpressionAppend::PrepareValue(FEmitContext& Context)
+void FExpressionAppend::EmitValueShader(FEmitContext& Context, FShaderValue& OutShader) const
 {
-	const EExpressionEvaluationType LhsEvaluation = PrepareExpressionValue(Context, Lhs);
-	const EExpressionEvaluationType RhsEvaluation = PrepareExpressionValue(Context, Rhs);
-	if (LhsEvaluation == EExpressionEvaluationType::None || RhsEvaluation == EExpressionEvaluationType::None)
-	{
-		return;
-	}
+	const Shader::FValueTypeDescription ResultTypeDesc = Shader::GetValueTypeDescription(GetType());
+	OutShader.bInline = true;
+	OutShader.Code.Appendf(TEXT("%s(%s, %s)"),
+		ResultTypeDesc.Name,
+		Lhs->GetValueShader(Context),
+		Rhs->GetValueShader(Context));
+}
 
-	const EExpressionEvaluationType Evaluation = CombineEvaluationTypes(LhsEvaluation, RhsEvaluation);
-	if (Evaluation == EExpressionEvaluationType::Shader)
-	{
-		const Shader::FValueTypeDescription ResultTypeDesc = Shader::GetValueTypeDescription(GetType());
-		return SetValueShaderf(Context, TEXT("%s(%s, %s)"),
-			ResultTypeDesc.Name,
-			Lhs->GetValueShader(Context),
-			Rhs->GetValueShader(Context));
-	}
-	else
-	{
-		Shader::FPreshaderData LocalPreshader;
-		Lhs->GetValuePreshader(Context, LocalPreshader);
-		Rhs->GetValuePreshader(Context, LocalPreshader);
-		LocalPreshader.WriteOpcode(Shader::EPreshaderOpcode::AppendVector);
-		return SetValuePreshader(Context, Evaluation, LocalPreshader);
-	}
+void FExpressionAppend::EmitValuePreshader(FEmitContext& Context, Shader::FPreshaderData& OutPreshader) const
+{
+	Lhs->GetValuePreshader(Context, OutPreshader);
+	Rhs->GetValuePreshader(Context, OutPreshader);
+	OutPreshader.WriteOpcode(Shader::EPreshaderOpcode::AppendVector);
+}
+
+void FExpressionReflectionVector::PrepareValue(FEmitContext& Context, const FRequestedType& RequestedType)
+{
+	Context.bReadMaterialNormal = true;
+	return SetType(Context, EExpressionEvaluationType::Shader, Shader::EValueType::Float3);
 }
 
-void FExpressionReflectionVector::PrepareValue(FEmitContext& Context)
+void FExpressionReflectionVector::EmitValueShader(FEmitContext& Context, FShaderValue& OutShader) const
 {
-	return SetValueInlineShaderf(Context, TEXT("Parameters.ReflectionVector"));
+	OutShader.bInline = true;
+	OutShader.Code.Append(TEXT("Parameters.ReflectionVector"));
 }
 
 void FStatementBreak::EmitHLSL(FEmitContext& Context) const
@@ -687,68 +629,85 @@ void FStatementBreak::EmitHLSL(FEmitContext& Context) const
 	ParentScope->EmitStatementf(Context, TEXT("break;"));
 }
 
-void FStatementReturn::RequestTypes(FUpdateTypeContext& Context) const
+void FStatementReturn::PrepareValues(FEmitContext& Context) const
 {
-	RequestExpressionType(Context, Expression, Type);
+	PrepareExpressionValue(Context, Expression, Type);
 }
 
 void FStatementReturn::EmitHLSL(FEmitContext& Context) const
 {
-	if (PrepareExpressionValue(Context, Expression) == EExpressionEvaluationType::None)
-	{
-		return;
-	}
-
 	ParentScope->MarkLiveRecursive();
 	ParentScope->EmitStatementf(Context, TEXT("return %s;"), Expression->GetValueShader(Context));
 }
 
-void FStatementIf::RequestTypes(FUpdateTypeContext& Context) const
+void FStatementIf::PrepareValues(FEmitContext& Context) const
 {
-	RequestExpressionType(Context, ConditionExpression, 1);
-	RequestScopeTypes(Context, ThenScope);
-	RequestScopeTypes(Context, ElseScope);
-	RequestScopeTypes(Context, NextScope);
+	const FPrepareValueResult ConditionResult = PrepareExpressionValue(Context, ConditionExpression, 1);
+	if (ConditionResult.EvaluationType == EExpressionEvaluationType::Constant)
+	{
+		const bool bCondition = ConditionResult.ConstantValue.AsBoolScalar();
+		if (bCondition)
+		{
+			MarkScopeDead(ElseScope);
+			PrepareScopeValues(Context, ThenScope);
+		}
+		else
+		{
+			MarkScopeDead(ThenScope);
+			PrepareScopeValues(Context, ElseScope);
+		}
+	}
+	else
+	{
+		PrepareScopeValues(Context, ThenScope);
+		PrepareScopeValues(Context, ElseScope);
+	}
+	PrepareScopeValues(Context, NextScope);
 }
 
 void FStatementIf::EmitHLSL(FEmitContext& Context) const
 {
-	bool bIfLive = ThenScope->IsLive();
-	if (ElseScope && ElseScope->IsLive())
-	{
-		bIfLive = true;
-	}
-
-	if (bIfLive)
+	if (ConditionExpression->GetEvaluationType() == EExpressionEvaluationType::Constant)
 	{
-		if (PrepareExpressionValue(Context, ConditionExpression) == EExpressionEvaluationType::None)
+		const bool bCondition = ConditionExpression->GetValueConstant(Context).AsBoolScalar();
+		if (bCondition && IsScopeLive(ThenScope))
 		{
-			return;
+			ParentScope->EmitNextScope(Context, ThenScope);
 		}
-
-		ParentScope->EmitNestedScopef(Context, ThenScope, TEXT("if (%s)"),
-			ConditionExpression->GetValueShader(Context, Shader::EValueType::Bool1));
-		if (ElseScope && ElseScope->IsLive())
+		else if(!bCondition && IsScopeLive(ElseScope))
 		{
-			ParentScope->EmitNestedScopef(Context, ElseScope, TEXT("else"));
+			ParentScope->EmitNextScope(Context, ElseScope);
 		}
 	}
-	ParentScope->EmitNestedScope(Context, NextScope);
+	else
+	{
+		if (IsScopeLive(ThenScope) || IsScopeLive(ElseScope))
+		{
+			ParentScope->EmitNestedScopef(Context, ThenScope, TEXT("if (%s)"),
+				ConditionExpression->GetValueShader(Context, Shader::EValueType::Bool1));
+			if (IsScopeLive(ElseScope))
+			{
+				ParentScope->EmitNestedScopef(Context, ElseScope, TEXT("else"));
+			}
+		}
+	}
+
+	ParentScope->EmitNextScope(Context, NextScope);
 }
 
-void FStatementLoop::RequestTypes(FUpdateTypeContext& Context) const
+void FStatementLoop::PrepareValues(FEmitContext& Context) const
 {
-	RequestScopeTypes(Context, LoopScope);
-	RequestScopeTypes(Context, NextScope);
+	PrepareScopeValues(Context, LoopScope);
+	PrepareScopeValues(Context, NextScope);
 }
 
 void FStatementLoop::EmitHLSL(FEmitContext& Context) const
 {
-	if (LoopScope->IsLive())
+	if (IsScopeLive(LoopScope))
 	{
 		ParentScope->EmitNestedScopef(Context, LoopScope, TEXT("while (true)"));
 	}
-	ParentScope->EmitNestedScope(Context, NextScope);
+	ParentScope->EmitNextScope(Context, NextScope);
 }
 
 } // namespace HLSLTree
diff --git a/Engine/Source/Runtime/Engine/Public/HLSLTree/HLSLTree.h b/Engine/Source/Runtime/Engine/Public/HLSLTree/HLSLTree.h
--- a/Engine/Source/Runtime/Engine/Public/HLSLTree/HLSLTree.h
+++ b/Engine/Source/Runtime/Engine/Public/HLSLTree/HLSLTree.h
@@ -111,14 +111,6 @@ private:
 	int32 NumErrors = 0;
 };
 
-class FUpdateTypeContext
-{
-public:
-	explicit FUpdateTypeContext(FErrors& InErrors) : Errors(InErrors) {}
-
-	FErrors& Errors;
-};
-
 enum class ECastFlags : uint32
 {
 	None = 0u,
@@ -155,6 +147,7 @@ public:
 	FMaterialCompilationOutput* MaterialCompilationOutput = nullptr;
 	TMap<Shader::FValue, uint32> DefaultUniformValues;
 	uint32 UniformPreshaderOffset = 0u;
+	bool bReadMaterialNormal = false;
 
 	int32 NumExpressionLocals = 0;
 	int32 NumLocalPHIs = 0;
@@ -262,7 +255,7 @@ public:
 
 	virtual ENodeVisitResult Visit(FNodeVisitor& Visitor) override;
 
-	virtual void RequestTypes(FUpdateTypeContext& Context) const = 0;
+	virtual void PrepareValues(FEmitContext& Context) const = 0;
 	virtual void EmitHLSL(FEmitContext& Context) const = 0;
 
 	FScope* ParentScope = nullptr;
@@ -320,6 +313,25 @@ public:
 	TBitArray<> RequestedComponents;
 };
 
+struct FShaderValue
+{
+	explicit FShaderValue(FStringBuilderBase& InCode) : Code(InCode) {}
+
+	FStringBuilderBase& Code;
+	bool bInline = false;
+	bool bHasDependencies = false;
+};
+
+struct FPrepareValueResult
+{
+	FType Type;
+	EExpressionEvaluationType EvaluationType = EExpressionEvaluationType::None;
+	Shader::FValue ConstantValue;
+
+	inline operator bool() const { return EvaluationType != EExpressionEvaluationType::None; }
+	inline bool operator!() const { return EvaluationType == EExpressionEvaluationType::None; }
+};
+
 /**
  * Represents an HLSL expression.  This is a piece of code that evaluates to a value, but has no side effects.
  * Unlike statements, expressions are not expected to execute in any particular order.  They may be cached (or not) in generated code, without the underlying implementation needing to care.
@@ -330,12 +342,12 @@ public:
 class FExpression : public FNode
 {
 public:
-	const FType& GetType() const { return Type; }
+	const FType& GetType() const { return PrepareValueResult.Type; }
+	EExpressionEvaluationType GetEvaluationType() const { return PrepareValueResult.EvaluationType; }
 
 	virtual ENodeVisitResult Visit(FNodeVisitor& Visitor) override;
 
-	friend FType RequestExpressionType(FUpdateTypeContext& Context, FExpression* InExpression, const FRequestedType& RequestedType);
-	friend EExpressionEvaluationType PrepareExpressionValue(FEmitContext& Context, FExpression* InExpression);
+	friend FPrepareValueResult PrepareExpressionValue(FEmitContext& Context, FExpression* InExpression, const FRequestedType& RequestedType);
 
 	const TCHAR* GetValueShader(FEmitContext& Context);
 	const TCHAR* GetValueShader(FEmitContext& Context, const FType& InType);
@@ -345,46 +357,20 @@ public:
 	FScope* ParentScope = nullptr;
 
 protected:
-	virtual void UpdateType(FUpdateTypeContext& Context, const FRequestedType& RequestedType) = 0;
-	virtual void PrepareValue(FEmitContext& Context) = 0;
-
-	void SetType(FUpdateTypeContext& Context, const FType& InType);
-
-	void InternalSetValueShader(FEmitContext& Context, const TCHAR* InCode, int32 InLength, bool bInline);
-
-	void SetValuePreshader(FEmitContext& Context, Shader::FPreshaderData& InPreshader);
-	void SetValueConstant(FEmitContext& Context, const Shader::FValue& InValue);
-	void SetValueForward(FEmitContext& Context, FExpression* Source);
+	virtual void PrepareValue(FEmitContext& Context, const FRequestedType& RequestedType) = 0;
+	virtual void EmitValueShader(FEmitContext& Context, FShaderValue& OutShader) const;
+	virtual void EmitShaderDependencies(FEmitContext& Context, const FShaderValue& Shader) const;
+	virtual void EmitValuePreshader(FEmitContext& Context, Shader::FPreshaderData& OutPreshader) const;
 
-	void SetValuePreshader(FEmitContext& Context, EExpressionEvaluationType InEvaluationType, Shader::FPreshaderData& InPreshader);
-
-	template<typename FormatType, typename... Types>
-	void SetValueShaderf(FEmitContext& Context, const FormatType& Format, Types... Args)
-	{
-		TStringBuilder<2048> String;
-		String.Appendf(Format, Forward<Types>(Args)...);
-		InternalSetValueShader(Context, String.ToString(), String.Len(), false);
-	}
-
-	template<typename FormatType, typename... Types>
-	void SetValueInlineShaderf(FEmitContext& Context, const FormatType& Format, Types... Args)
-	{
-		TStringBuilder<2048> String;
-		String.Appendf(Format, Forward<Types>(Args)...);
-		InternalSetValueShader(Context, String.ToString(), String.Len(), true);
-	}
-
-	inline void SetValueShader(FEmitContext& Context, const FStringBuilderBase& String) { InternalSetValueShader(Context, String.ToString(), String.Len(), false); }
-	inline void SetValueInlineShader(FEmitContext& Context, const FStringBuilderBase& String) { InternalSetValueShader(Context, String.ToString(), String.Len(), true); }
+	void SetType(FEmitContext& Context, EExpressionEvaluationType InEvaluationType, const FType& InType);
+	void SetForwardValue(FEmitContext& Context, FExpression* InForwardValue, const FRequestedType& RequestedType);
 
 private:
 	const TCHAR* LocalVariableName = nullptr;
 	const TCHAR* Code = nullptr;
-	const Shader::FPreshaderData* Preshader = nullptr;
-	Shader::FValue ConstantValue;
-	EExpressionEvaluationType EvaluationType = EExpressionEvaluationType::None;
-	FType Type;
 	FRequestedType CurrentRequestedType;
+	FPrepareValueResult PrepareValueResult;
+	FExpression* ForwardedValue = nullptr;
 	bool bReentryFlag = false;
 };
 
@@ -397,8 +383,9 @@ private:
 class FExpressionLocalPHI final : public FExpression
 {
 public:
-	virtual void UpdateType(FUpdateTypeContext& Context, const FRequestedType& RequestedType) override;
-	virtual void PrepareValue(FEmitContext& Context) override;
+	virtual void PrepareValue(FEmitContext& Context, const FRequestedType& RequestedType) override;
+	virtual void EmitValueShader(FEmitContext& Context, FShaderValue& OutShader) const override;
+	virtual void EmitShaderDependencies(FEmitContext& Context, const FShaderValue& Shader) const override;
 
 	FName LocalName;
 	FScope* Scopes[MaxNumPreviousScopes];
@@ -420,6 +407,13 @@ public:
 	FTextureDescription Description;
 };
 
+enum class EScopeState : uint8
+{
+	Uninitialized,
+	Live,
+	Dead,
+};
+
 /**
  * Represents an HLSL scope.  A scope contains a single statement, along with any expressions required by that statement
  */
@@ -429,7 +423,8 @@ public:
 	static FScope* FindSharedParent(FScope* Lhs, FScope* Rhs);
 
 	inline FScope* GetParentScope() const { return ParentScope; }
-	inline bool IsLive() const { return bLive; }
+	inline bool IsLive() const { return State == EScopeState::Live; }
+	inline bool IsDead() const { return State == EScopeState::Dead; }
 
 	inline TArrayView<FScope*> GetPreviousScopes() const
 	{
@@ -445,33 +440,34 @@ public:
 
 	void UseExpression(FExpression* Expression);
 
-	friend void RequestScopeTypes(FUpdateTypeContext& Context, const FScope* InScope);
+	friend void PrepareScopeValues(FEmitContext& Context, const FScope* InScope);
 
 	template<typename FormatType, typename... Types>
 	void EmitDeclarationf(FEmitContext& Context, const FormatType& Format, Types... Args)
 	{
-		InternalEmitCodef(Context, Declarations, nullptr, Format, Forward<Types>(Args)...);
+		InternalEmitCodef(Context, Declarations, ENextScopeFormat::None, nullptr, Format, Forward<Types>(Args)...);
 	}
 
 	template<typename FormatType, typename... Types>
 	void EmitStatementf(FEmitContext& Context, const FormatType& Format, Types... Args)
 	{
-		InternalEmitCodef(Context, Statements, nullptr, Format, Forward<Types>(Args)...);
+		InternalEmitCodef(Context, Statements, ENextScopeFormat::None, nullptr, Format, Forward<Types>(Args)...);
 	}
 
-	void EmitNestedScope(FEmitContext& Context, FScope* NestedScope)
+	void EmitNextScope(FEmitContext& Context, FScope* NestedScope)
 	{
-		InternalEmitCode(Context, Statements, NestedScope, nullptr, 0);
+		InternalEmitCode(Context, Statements, ENextScopeFormat::Unscoped, NestedScope, nullptr, 0);
 	}
 
 	template<typename FormatType, typename... Types>
 	void EmitNestedScopef(FEmitContext& Context, FScope* NestedScope, const FormatType& Format, Types... Args)
 	{
-		InternalEmitCodef(Context, Statements, NestedScope, Format, Forward<Types>(Args)...);
+		InternalEmitCodef(Context, Statements, ENextScopeFormat::Scoped, NestedScope, Format, Forward<Types>(Args)...);
 	}
 
 	void MarkLive();
 	void MarkLiveRecursive();
+	void MarkDead();
 
 	void WriteHLSL(int32 Indent, FStringBuilderBase& OutString) const;
 
@@ -480,11 +476,19 @@ private:
 	friend class FExpression;
 	friend class FNodeVisitor_MoveToScope;
 
+	enum class ENextScopeFormat : uint8
+	{
+		None,
+		Unscoped,
+		Scoped,
+	};
+
 	struct FCodeEntry
 	{
 		FCodeEntry* Next;
-		FScope* NestedScope;
+		FScope* NextScope;
 		int32 Length;
+		ENextScopeFormat ScopeFormat;
 		TCHAR String[1];
 	};
 
@@ -495,14 +499,14 @@ private:
 		int32 Num = 0;
 	};
 
-	void InternalEmitCode(FEmitContext& Context, FCodeList& List, FScope* NestedScope, const TCHAR* String, int32 Length);
+	void InternalEmitCode(FEmitContext& Context, FCodeList& List, ENextScopeFormat ScopeFormat, FScope* NextScope, const TCHAR* String, int32 Length);
 
 	template<typename FormatType, typename... Types>
-	void InternalEmitCodef(FEmitContext& Context, FCodeList& List, FScope* NestedScope, const FormatType& Format, Types... Args)
+	void InternalEmitCodef(FEmitContext& Context, FCodeList& List, ENextScopeFormat ScopeFormat, FScope* NextScope, const FormatType& Format, Types... Args)
 	{
 		TStringBuilder<2048> String;
 		String.Appendf(Format, Forward<Types>(Args)...);
-		InternalEmitCode(Context, List, NestedScope, String.ToString(), String.Len());
+		InternalEmitCode(Context, List, ScopeFormat, NextScope, String.ToString(), String.Len());
 	}
 
 	FScope* ParentScope = nullptr;
@@ -513,9 +517,22 @@ private:
 	FCodeList Statements;
 	int32 NumPreviousScopes = 0;
 	int32 NestedLevel = 0;
-	bool bLive = false;
+	EScopeState State = EScopeState::Uninitialized;
 };
 
+inline bool IsScopeLive(const FScope* InScope)
+{
+	return (bool)InScope && InScope->IsLive();
+}
+
+inline void MarkScopeDead(FScope* InScope)
+{
+	if (InScope)
+	{
+		InScope->MarkDead();
+	}
+}
+
 /**
  * The HLSL AST.  Basically a wrapper around the root scope, with some helper methods
  */
@@ -585,9 +602,9 @@ private:
 	FStatement* ResultStatement = nullptr;
 };
 
-//Shader::EValueType RequestExpressionType(FUpdateTypeContext& Context, FExpression* InExpression, int8 InRequestedNumComponents); // friend of FExpression
+//Shader::EValueType RequestExpressionType(FEmitContext& Context, FExpression* InExpression, int8 InRequestedNumComponents); // friend of FExpression
 //EExpressionEvaluationType PrepareExpressionValue(FEmitContext& Context, FExpression* InExpression); // friend of FExpression
-//void RequestScopeTypes(FUpdateTypeContext& Context, const FScope* InScope); // friend of FScope
+//void PrepareScopeValues(FEmitContext& Context, const FScope* InScope); // friend of FScope
 
 } // namespace HLSLTree
 } // namespace UE
diff --git a/Engine/Source/Runtime/Engine/Public/HLSLTree/HLSLTreeCommon.h b/Engine/Source/Runtime/Engine/Public/HLSLTree/HLSLTreeCommon.h
--- a/Engine/Source/Runtime/Engine/Public/HLSLTree/HLSLTreeCommon.h
+++ b/Engine/Source/Runtime/Engine/Public/HLSLTree/HLSLTreeCommon.h
@@ -7,6 +7,12 @@ enum class EMaterialParameterType : uint8;
 
 namespace UE
 {
+
+namespace Shader
+{
+enum class EPreshaderOpcode : uint8;
+} // namespace Shader
+
 namespace HLSLTree
 {
 
@@ -22,16 +28,12 @@ enum class EBinaryOp : uint8
 
 struct FBinaryOpDescription
 {
-	FBinaryOpDescription()
-		: Name(nullptr), Operator(nullptr)
-	{}
-
-	FBinaryOpDescription(const TCHAR* InName, const TCHAR* InOperator)
-		: Name(InName), Operator(InOperator)
-	{}
+	FBinaryOpDescription();
+	FBinaryOpDescription(const TCHAR* InName, const TCHAR* InOperator, Shader::EPreshaderOpcode InOpcode);
 
 	const TCHAR* Name;
 	const TCHAR* Operator;
+	Shader::EPreshaderOpcode PreshaderOpcode;
 };
 
 FBinaryOpDescription GetBinaryOpDesription(EBinaryOp Op);
@@ -45,8 +47,9 @@ public:
 
 	FConstantValue Value;
 
-	virtual void UpdateType(FUpdateTypeContext& Context, const FRequestedType& RequestedType) override;
-	virtual void PrepareValue(FEmitContext& Context) override;
+	virtual void PrepareValue(FEmitContext& Context, const FRequestedType& RequestedType) override;
+	virtual void EmitValueShader(FEmitContext& Context, FShaderValue& OutShader) const override;
+	virtual void EmitValuePreshader(FEmitContext& Context, Shader::FPreshaderData& OutPreshader) const override;
 };
 
 class FExpressionMaterialParameter : public FExpression
@@ -60,8 +63,8 @@ public:
 	Shader::FValue DefaultValue;
 	EMaterialParameterType ParameterType;
 
-	virtual void UpdateType(FUpdateTypeContext& Context, const FRequestedType& RequestedType) override;
-	virtual void PrepareValue(FEmitContext& Context) override;
+	virtual void PrepareValue(FEmitContext& Context, const FRequestedType& RequestedType) override;
+	virtual void EmitValuePreshader(FEmitContext& Context, Shader::FPreshaderData& OutPreshader) const override;
 };
 
 enum class EExternalInputType
@@ -92,8 +95,8 @@ public:
 
 	EExternalInputType InputType;
 
-	virtual void UpdateType(FUpdateTypeContext& Context, const FRequestedType& RequestedType) override { SetType(Context, GetInputExpressionType(InputType)); }
-	virtual void PrepareValue(FEmitContext& Context) override;
+	virtual void PrepareValue(FEmitContext& Context, const FRequestedType& RequestedType) override { SetType(Context, EExpressionEvaluationType::Shader, GetInputExpressionType(InputType)); }
+	virtual void EmitValueShader(FEmitContext& Context, FShaderValue& OutShader) const override;
 };
 
 class FExpressionTextureSample : public FExpression
@@ -122,8 +125,8 @@ public:
 		return Result;
 	}
 
-	virtual void UpdateType(FUpdateTypeContext& Context, const FRequestedType& RequestedType) override;
-	virtual void PrepareValue(FEmitContext& Context) override;
+	virtual void PrepareValue(FEmitContext& Context, const FRequestedType& RequestedType) override;
+	virtual void EmitValueShader(FEmitContext& Context, FShaderValue& OutShader) const override;
 };
 
 class FExpressionGetStructField : public FExpression
@@ -145,8 +148,8 @@ public:
 		return Result;
 	}
 
-	virtual void UpdateType(FUpdateTypeContext& Context, const FRequestedType& RequestedType) override;
-	virtual void PrepareValue(FEmitContext& Context) override;
+	virtual void PrepareValue(FEmitContext& Context, const FRequestedType& RequestedType) override;
+	virtual void EmitValueShader(FEmitContext& Context, FShaderValue& OutShader) const override;
 };
 
 class FExpressionSetStructField : public FExpression
@@ -170,8 +173,8 @@ public:
 		return Result;
 	}
 
-	virtual void UpdateType(FUpdateTypeContext& Context, const FRequestedType& RequestedType) override;
-	virtual void PrepareValue(FEmitContext& Context) override;
+	virtual void PrepareValue(FEmitContext& Context, const FRequestedType& RequestedType) override;
+	virtual void EmitValueShader(FEmitContext& Context, FShaderValue& OutShader) const override;
 };
 
 class FExpressionSelect : public FExpression
@@ -199,8 +202,9 @@ public:
 		return Result;
 	}
 
-	virtual void UpdateType(FUpdateTypeContext& Context, const FRequestedType& RequestedType) override;
-	virtual void PrepareValue(FEmitContext& Context) override;
+	virtual void PrepareValue(FEmitContext& Context, const FRequestedType& RequestedType) override;
+	virtual void EmitValueShader(FEmitContext& Context, FShaderValue& OutShader) const override;
+	virtual void EmitValuePreshader(FEmitContext& Context, Shader::FPreshaderData& OutPreshader) const override;
 };
 
 class FExpressionBinaryOp : public FExpression
@@ -227,8 +231,9 @@ public:
 		return Result;
 	}
 
-	virtual void UpdateType(FUpdateTypeContext& Context, const FRequestedType& RequestedType) override;
-	virtual void PrepareValue(FEmitContext& Context) override;
+	virtual void PrepareValue(FEmitContext& Context, const FRequestedType& RequestedType) override;
+	virtual void EmitValueShader(FEmitContext& Context, FShaderValue& OutShader) const override;
+	virtual void EmitValuePreshader(FEmitContext& Context, Shader::FPreshaderData& OutPreshader) const override;
 };
 
 struct FSwizzleParameters
@@ -262,8 +267,9 @@ public:
 		return Result;
 	}
 
-	virtual void UpdateType(FUpdateTypeContext& Context, const FRequestedType& RequestedType) override;
-	virtual void PrepareValue(FEmitContext& Context) override;
+	virtual void PrepareValue(FEmitContext& Context, const FRequestedType& RequestedType) override;
+	virtual void EmitValueShader(FEmitContext& Context, FShaderValue& OutShader) const override;
+	virtual void EmitValuePreshader(FEmitContext& Context, Shader::FPreshaderData& OutPreshader) const override;
 };
 
 class FExpressionAppend : public FExpression
@@ -288,15 +294,16 @@ public:
 		return Result;
 	}
 
-	virtual void UpdateType(FUpdateTypeContext& Context, const FRequestedType& RequestedType) override;
-	virtual void PrepareValue(FEmitContext& Context) override;
+	virtual void PrepareValue(FEmitContext& Context, const FRequestedType& RequestedType) override;
+	virtual void EmitValueShader(FEmitContext& Context, FShaderValue& OutShader) const override;
+	virtual void EmitValuePreshader(FEmitContext& Context, Shader::FPreshaderData& OutPreshader) const override;
 };
 
 class FExpressionReflectionVector : public FExpression
 {
 public:
-	virtual void UpdateType(FUpdateTypeContext& Context, const FRequestedType& RequestedType) override { SetType(Context, Shader::EValueType::Float3); }
-	virtual void PrepareValue(FEmitContext& Context) override;
+	virtual void PrepareValue(FEmitContext& Context, const FRequestedType& RequestedType) override;
+	virtual void EmitValueShader(FEmitContext& Context, FShaderValue& OutShader) const override;
 };
 
 class FStatementReturn : public FStatement
@@ -317,7 +324,7 @@ public:
 		return Result;
 	}
 
-	virtual void RequestTypes(FUpdateTypeContext& Context) const override;
+	virtual void PrepareValues(FEmitContext& Context) const override;
 	virtual void EmitHLSL(FEmitContext& Context) const override;
 };
 
@@ -326,7 +333,7 @@ class FStatementBreak : public FStatement
 public:
 	static constexpr bool MarkScopeLive = true;
 
-	virtual void RequestTypes(FUpdateTypeContext& Context) const override {}
+	virtual void PrepareValues(FEmitContext& Context) const override {}
 	virtual void EmitHLSL(FEmitContext& Context) const override;
 };
 
@@ -351,7 +358,7 @@ public:
 		return Result;
 	}
 
-	virtual void RequestTypes(FUpdateTypeContext& Context) const override;
+	virtual void PrepareValues(FEmitContext& Context) const override;
 	virtual void EmitHLSL(FEmitContext& Context) const override;
 };
 
@@ -372,7 +379,7 @@ public:
 		return Result;
 	}
 
-	virtual void RequestTypes(FUpdateTypeContext& Context) const override;
+	virtual void PrepareValues(FEmitContext& Context) const override;
 	virtual void EmitHLSL(FEmitContext& Context) const override;
 };
 
