Add 'Standalone' to 'Ray Traced (Deprecated)' to make it clear what's being deprecated

diff --git a/Engine/Source/Runtime/Engine/Classes/Engine/EngineTypes.h b/Engine/Source/Runtime/Engine/Classes/Engine/EngineTypes.h
--- a/Engine/Source/Runtime/Engine/Classes/Engine/EngineTypes.h
+++ b/Engine/Source/Runtime/Engine/Classes/Engine/EngineTypes.h
@@ -395,8 +395,8 @@ namespace EDynamicGlobalIlluminationMethod
 		/** Standalone Screen Space Global Illumination.  Low cost, but limited by screen space information. */
 		ScreenSpace UMETA(DisplayName="Screen Space (Beta)"),
 
-		/** Ray Traced Global Illumination technique.  Deprecated, use Lumen Global Illumination instead. */
-		RayTraced UMETA(DisplayName="Ray Traced (Deprecated)"),
+		/** Standalone Ray Traced Global Illumination technique.  Deprecated, use Lumen Global Illumination instead. */
+		RayTraced UMETA(DisplayName="Standalone Ray Traced (Deprecated)"),
 	};
 }
 
@@ -415,8 +415,8 @@ namespace EReflectionMethod
 		/** Standalone Screen Space Reflections.  Low cost, but limited by screen space information. */
 		ScreenSpace UMETA(DisplayName="Screen Space"),
 
-		/** Ray Traced Reflections technique.  Deprecated, use Lumen Reflections instead. */
-		RayTraced UMETA(DisplayName="Ray Traced (Deprecated)"),
+		/** Standalone Ray Traced Reflections technique.  Deprecated, use Lumen Reflections instead. */
+		RayTraced UMETA(DisplayName="Standalone Ray Traced (Deprecated)"),
 	};
 }
 
