New Engine Scalability groups Global Illumination and Reflections
Moved Lumen settings to new groups

diff --git a/Engine/Config/BaseScalability.ini b/Engine/Config/BaseScalability.ini
--- a/Engine/Config/BaseScalability.ini
+++ b/Engine/Config/BaseScalability.ini
@@ -25,6 +25,8 @@
 PerfIndexThresholds_ViewDistanceQuality="Min 18 42 105"
 PerfIndexThresholds_AntiAliasingQuality="GPU 18 42 115"
 PerfIndexThresholds_ShadowQuality="Min 18 42 105"
+PerfIndexThresholds_GlobalIlluminationQuality="GPU 18 42 115"
+PerfIndexThresholds_ReflectionQuality="GPU 18 42 115"
 PerfIndexThresholds_PostProcessQuality="GPU 18 42 115"
 PerfIndexThresholds_TextureQuality="GPU 18 42 115"
 PerfIndexThresholds_EffectsQuality="Min 18 42 105"
@@ -178,6 +180,72 @@
 
 ;-----------------------------------------------------------------------------------------------------------------
 
+[GlobalIlluminationQuality@0]
+r.Lumen.DiffuseIndirect.Allow=0
+
+[GlobalIlluminationQuality@1]
+r.Lumen.DiffuseIndirect.Allow=0
+
+[GlobalIlluminationQuality@2]
+r.Lumen.DiffuseIndirect.Allow=1
+r.Lumen.TranslucencyVolume.GridPixelSize=64
+r.LumenScene.Radiosity=0
+r.Lumen.TraceMeshSDFs.Allow=0
+r.Lumen.ScreenProbeGather.RadianceCache.ProbeResolution=16
+r.Lumen.ScreenProbeGather.ScreenSpaceBentNormal=0
+r.Lumen.ScreenProbeGather.DownsampleFactor=32
+r.Lumen.TranslucencyVolume.TraceFromVolume=0
+r.Lumen.TranslucencyVolume.RadianceCache.ProbeResolution=8
+
+[GlobalIlluminationQuality@3]
+r.Lumen.DiffuseIndirect.Allow=1
+r.Lumen.TranslucencyVolume.GridPixelSize=32
+r.LumenScene.Radiosity=1
+r.Lumen.TraceMeshSDFs.Allow=1
+r.Lumen.ScreenProbeGather.RadianceCache.ProbeResolution=32
+r.Lumen.ScreenProbeGather.ScreenSpaceBentNormal=1
+r.Lumen.ScreenProbeGather.DownsampleFactor=16
+r.Lumen.TranslucencyVolume.TraceFromVolume=1
+r.Lumen.TranslucencyVolume.TracingOctahedronResolution=3
+r.Lumen.TranslucencyVolume.RadianceCache.ProbeResolution=8
+
+[GlobalIlluminationQuality@Cine]
+r.Lumen.DiffuseIndirect.Allow=1
+r.Lumen.TranslucencyVolume.GridPixelSize=32
+r.LumenScene.Radiosity=1
+r.Lumen.TraceMeshSDFs.Allow=1
+r.Lumen.ScreenProbeGather.RadianceCache.ProbeResolution=32
+r.Lumen.ScreenProbeGather.ScreenSpaceBentNormal=1
+r.Lumen.ScreenProbeGather.DownsampleFactor=16
+r.Lumen.TranslucencyVolume.TraceFromVolume=1
+r.Lumen.TranslucencyVolume.TracingOctahedronResolution=4
+r.Lumen.TranslucencyVolume.RadianceCache.ProbeResolution=16
+
+
+;-----------------------------------------------------------------------------------------------------------------
+
+[ReflectionQuality@0]
+r.Lumen.Reflections.Allow=0
+
+[ReflectionQuality@1]
+r.Lumen.Reflections.Allow=0
+
+[ReflectionQuality@2]
+r.Lumen.Reflections.Allow=1
+r.Lumen.Reflections.DownsampleFactor=2
+
+[ReflectionQuality@3]
+r.Lumen.Reflections.Allow=1
+r.Lumen.Reflections.DownsampleFactor=1
+
+[ReflectionQuality@Cine]
+r.Lumen.Reflections.Allow=1
+r.Lumen.Reflections.DownsampleFactor=1
+
+;-----------------------------------------------------------------------------------------------------------------
+
+
+
 [PostProcessQuality@0]
 r.MotionBlurQuality=0
 r.AmbientOcclusionMipLevelFactor=1.0
@@ -350,7 +418,7 @@
 r.VT.MaxAnisotropy=4
 r.Streaming.LimitPoolSizeToVRAM=1
 r.Streaming.PoolSize=400
-r.Streaming.MaxEffectiveScreenSize=0
+r.Streaming.MaxEffectiveScreenSize=0f
 
 [TextureQuality@1]
 r.Streaming.MipBias=1
@@ -427,8 +495,6 @@
 r.SkyAtmosphere.MultiScatteringLUT.SampleCount=15.0
 r.SkyLight.RealTimeReflectionCapture=0
 fx.Niagara.QualityLevel=0
-r.Lumen.DiffuseIndirect.Allow=0
-r.Lumen.Reflections.Allow=0
 
 [EffectsQuality@1]
 r.TranslucencyLightingVolumeDim=32
@@ -459,8 +525,6 @@
 r.SkyAtmosphere.MultiScatteringLUT.SampleCount=15.0
 r.SkyLight.RealTimeReflectionCapture=0
 fx.Niagara.QualityLevel=1
-r.Lumen.DiffuseIndirect.Allow=0
-r.Lumen.Reflections.Allow=0
 
 [EffectsQuality@2]
 r.TranslucencyLightingVolumeDim=48
@@ -491,17 +555,6 @@
 r.SkyAtmosphere.MultiScatteringLUT.SampleCount=15.0
 r.SkyLight.RealTimeReflectionCapture=1
 fx.Niagara.QualityLevel=2
-r.Lumen.DiffuseIndirect.Allow=1
-r.Lumen.Reflections.Allow=1
-r.Lumen.TranslucencyVolume.GridPixelSize=64
-r.LumenScene.Radiosity=0
-r.Lumen.TraceMeshSDFs.Allow=0
-r.Lumen.ScreenProbeGather.RadianceCache.ProbeResolution=16
-r.Lumen.ScreenProbeGather.ScreenSpaceBentNormal=0
-r.Lumen.ScreenProbeGather.DownsampleFactor=32
-r.Lumen.Reflections.DownsampleFactor=2
-r.Lumen.TranslucencyVolume.TraceFromVolume=0
-r.Lumen.TranslucencyVolume.RadianceCache.ProbeResolution=8
 
 [EffectsQuality@3]
 r.TranslucencyLightingVolumeDim=64
@@ -532,18 +585,6 @@
 r.SkyAtmosphere.MultiScatteringLUT.SampleCount=15.0
 r.SkyLight.RealTimeReflectionCapture=1
 fx.Niagara.QualityLevel=3
-r.Lumen.DiffuseIndirect.Allow=1
-r.Lumen.Reflections.Allow=1
-r.Lumen.TranslucencyVolume.GridPixelSize=32
-r.LumenScene.Radiosity=1
-r.Lumen.TraceMeshSDFs.Allow=1
-r.Lumen.ScreenProbeGather.RadianceCache.ProbeResolution=32
-r.Lumen.ScreenProbeGather.ScreenSpaceBentNormal=1
-r.Lumen.ScreenProbeGather.DownsampleFactor=16
-r.Lumen.Reflections.DownsampleFactor=1
-r.Lumen.TranslucencyVolume.TraceFromVolume=1
-r.Lumen.TranslucencyVolume.TracingOctahedronResolution=3
-r.Lumen.TranslucencyVolume.RadianceCache.ProbeResolution=8
 
 [EffectsQuality@Cine]
 r.TranslucencyLightingVolumeDim=64
@@ -574,18 +615,6 @@
 r.SkyAtmosphere.MultiScatteringLUT.SampleCount=20.0
 r.SkyLight.RealTimeReflectionCapture=1
 fx.Niagara.QualityLevel=4
-r.Lumen.DiffuseIndirect.Allow=1
-r.Lumen.Reflections.Allow=1
-r.Lumen.TranslucencyVolume.GridPixelSize=32
-r.LumenScene.Radiosity=1
-r.Lumen.TraceMeshSDFs.Allow=1
-r.Lumen.ScreenProbeGather.RadianceCache.ProbeResolution=32
-r.Lumen.ScreenProbeGather.ScreenSpaceBentNormal=1
-r.Lumen.ScreenProbeGather.DownsampleFactor=16
-r.Lumen.Reflections.DownsampleFactor=1
-r.Lumen.TranslucencyVolume.TraceFromVolume=1
-r.Lumen.TranslucencyVolume.TracingOctahedronResolution=4
-r.Lumen.TranslucencyVolume.RadianceCache.ProbeResolution=16
 
 ;-----------------------------------------------------------------------------------------------------------------
 
diff --git a/Engine/Plugins/Experimental/AutomationUtils/Source/AutomationUtils/Private/AutomationUtilsBlueprintLibrary.cpp b/Engine/Plugins/Experimental/AutomationUtils/Source/AutomationUtils/Private/AutomationUtilsBlueprintLibrary.cpp
--- a/Engine/Plugins/Experimental/AutomationUtils/Source/AutomationUtils/Private/AutomationUtilsBlueprintLibrary.cpp
+++ b/Engine/Plugins/Experimental/AutomationUtils/Source/AutomationUtils/Private/AutomationUtilsBlueprintLibrary.cpp
@@ -148,6 +148,8 @@ FAutomationUtilsGameplayAutomationScreenshotInstance::FAutomationUtilsGameplayAu
 		JsonObject->SetNumberField(TEXT("viewDistanceQuality"), QualityLevels.ViewDistanceQuality);
 		JsonObject->SetNumberField(TEXT("antiAliasingQuality"), QualityLevels.AntiAliasingQuality);
 		JsonObject->SetNumberField(TEXT("shadowQuality"), QualityLevels.ShadowQuality);
+		JsonObject->SetNumberField(TEXT("globalIlluminationQuality"), QualityLevels.GlobalIlluminationQuality);
+		JsonObject->SetNumberField(TEXT("reflectionQuality"), QualityLevels.ReflectionQuality);
 		JsonObject->SetNumberField(TEXT("postProcessQuality"), QualityLevels.PostProcessQuality);
 		JsonObject->SetNumberField(TEXT("textureQuality"), QualityLevels.TextureQuality);
 		JsonObject->SetNumberField(TEXT("effectsQuality"), QualityLevels.EffectsQuality);
diff --git a/Engine/Plugins/MovieScene/MovieRenderPipeline/Source/MovieRenderPipelineCore/Private/MoviePipelineGameOverrideSetting.cpp b/Engine/Plugins/MovieScene/MovieRenderPipeline/Source/MovieRenderPipelineCore/Private/MoviePipelineGameOverrideSetting.cpp
--- a/Engine/Plugins/MovieScene/MovieRenderPipeline/Source/MovieRenderPipelineCore/Private/MoviePipelineGameOverrideSetting.cpp
+++ b/Engine/Plugins/MovieScene/MovieRenderPipeline/Source/MovieRenderPipelineCore/Private/MoviePipelineGameOverrideSetting.cpp
@@ -124,6 +124,8 @@ void UMoviePipelineGameOverrideSetting::BuildNewProcessCommandLineArgsImpl(TArra
 		InOutDeviceProfileCvars.Add(TEXT("sg.ViewDistanceQuality=4"));
 		InOutDeviceProfileCvars.Add(TEXT("sg.AntiAliasingQuality=4"));
 		InOutDeviceProfileCvars.Add(TEXT("sg.ShadowQuality=4"));
+		InOutDeviceProfileCvars.Add(TEXT("sg.GlobalIlluminationQuality=4"));
+		InOutDeviceProfileCvars.Add(TEXT("sg.ReflectionQuality=4"));
 		InOutDeviceProfileCvars.Add(TEXT("sg.PostProcessQuality=4"));
 		InOutDeviceProfileCvars.Add(TEXT("sg.TextureQuality=4"));
 		InOutDeviceProfileCvars.Add(TEXT("sg.EffectsQuality=4"));
diff --git a/Engine/Source/Editor/UnrealEd/Private/PerformanceMonitor.cpp b/Engine/Source/Editor/UnrealEd/Private/PerformanceMonitor.cpp
--- a/Engine/Source/Editor/UnrealEd/Private/PerformanceMonitor.cpp
+++ b/Engine/Source/Editor/UnrealEd/Private/PerformanceMonitor.cpp
@@ -162,6 +162,8 @@ bool FPerformanceMonitor::WillAutoScalabilityHelp() const
 	IsAutoScaleLower |= NewLevels.ViewDistanceQuality < CurrentLevels.ViewDistanceQuality;
 	IsAutoScaleLower |= NewLevels.AntiAliasingQuality < CurrentLevels.AntiAliasingQuality;
 	IsAutoScaleLower |= NewLevels.ShadowQuality < CurrentLevels.ShadowQuality;
+	IsAutoScaleLower |= NewLevels.GlobalIlluminationQuality < CurrentLevels.GlobalIlluminationQuality;
+	IsAutoScaleLower |= NewLevels.ReflectionQuality < CurrentLevels.ReflectionQuality;
 	IsAutoScaleLower |= NewLevels.PostProcessQuality < CurrentLevels.PostProcessQuality;
 	IsAutoScaleLower |= NewLevels.TextureQuality < CurrentLevels.TextureQuality;
 	IsAutoScaleLower |= NewLevels.EffectsQuality < CurrentLevels.EffectsQuality;
@@ -183,6 +185,8 @@ Scalability::FQualityLevels FPerformanceMonitor::GetAutoScalabilityQualityLevels
 	NewLevels.ViewDistanceQuality	= FMath::Min(NewLevels.ViewDistanceQuality, ExistingLevels.ViewDistanceQuality);
 	NewLevels.AntiAliasingQuality	= FMath::Min(NewLevels.AntiAliasingQuality, ExistingLevels.AntiAliasingQuality);
 	NewLevels.ShadowQuality			= FMath::Min(NewLevels.ShadowQuality, ExistingLevels.ShadowQuality);
+	NewLevels.GlobalIlluminationQuality			= FMath::Min(NewLevels.GlobalIlluminationQuality, ExistingLevels.GlobalIlluminationQuality);
+	NewLevels.ReflectionQuality			= FMath::Min(NewLevels.ReflectionQuality, ExistingLevels.ReflectionQuality);
 	NewLevels.PostProcessQuality	= FMath::Min(NewLevels.PostProcessQuality, ExistingLevels.PostProcessQuality);
 	NewLevels.TextureQuality		= FMath::Min(NewLevels.TextureQuality, ExistingLevels.TextureQuality);
 	NewLevels.EffectsQuality		= FMath::Min(NewLevels.EffectsQuality, ExistingLevels.EffectsQuality);
diff --git a/Engine/Source/Editor/UnrealEd/Private/SScalabilitySettings.cpp b/Engine/Source/Editor/UnrealEd/Private/SScalabilitySettings.cpp
--- a/Engine/Source/Editor/UnrealEd/Private/SScalabilitySettings.cpp
+++ b/Engine/Source/Editor/UnrealEd/Private/SScalabilitySettings.cpp
@@ -24,6 +24,8 @@ ECheckBoxState SScalabilitySettings::IsGroupQualityLevelSelected(const TCHAR* In
 	else if (FCString::Strcmp(InGroupName, TEXT("AntiAliasingQuality")) == 0) QualityLevel = CachedQualityLevels.AntiAliasingQuality;
 	else if (FCString::Strcmp(InGroupName, TEXT("PostProcessQuality")) == 0) QualityLevel = CachedQualityLevels.PostProcessQuality;
 	else if (FCString::Strcmp(InGroupName, TEXT("ShadowQuality")) == 0) QualityLevel = CachedQualityLevels.ShadowQuality;
+	else if (FCString::Strcmp(InGroupName, TEXT("GlobalIlluminationQuality")) == 0) QualityLevel = CachedQualityLevels.GlobalIlluminationQuality;
+	else if (FCString::Strcmp(InGroupName, TEXT("ReflectionQuality")) == 0) QualityLevel = CachedQualityLevels.ReflectionQuality;
 	else if (FCString::Strcmp(InGroupName, TEXT("TextureQuality")) == 0) QualityLevel = CachedQualityLevels.TextureQuality;
 	else if (FCString::Strcmp(InGroupName, TEXT("EffectsQuality")) == 0) QualityLevel = CachedQualityLevels.EffectsQuality;
 	else if (FCString::Strcmp(InGroupName, TEXT("FoliageQuality")) == 0) QualityLevel = CachedQualityLevels.FoliageQuality;
@@ -39,6 +41,8 @@ void SScalabilitySettings::OnGroupQualityLevelChanged(ECheckBoxState NewState, c
 	else if (FCString::Strcmp(InGroupName, TEXT("AntiAliasingQuality")) == 0) CachedQualityLevels.AntiAliasingQuality = InQualityLevel;
 	else if (FCString::Strcmp(InGroupName, TEXT("PostProcessQuality")) == 0) CachedQualityLevels.PostProcessQuality = InQualityLevel;
 	else if (FCString::Strcmp(InGroupName, TEXT("ShadowQuality")) == 0) CachedQualityLevels.ShadowQuality = InQualityLevel;
+	else if (FCString::Strcmp(InGroupName, TEXT("GlobalIlluminationQuality")) == 0) CachedQualityLevels.GlobalIlluminationQuality = InQualityLevel;
+	else if (FCString::Strcmp(InGroupName, TEXT("ReflectionQuality")) == 0) CachedQualityLevels.ReflectionQuality = InQualityLevel;
 	else if (FCString::Strcmp(InGroupName, TEXT("TextureQuality")) == 0) CachedQualityLevels.TextureQuality = InQualityLevel;
 	else if (FCString::Strcmp(InGroupName, TEXT("EffectsQuality")) == 0) CachedQualityLevels.EffectsQuality = InQualityLevel;
 	else if (FCString::Strcmp(InGroupName, TEXT("FoliageQuality")) == 0) CachedQualityLevels.FoliageQuality = InQualityLevel;
@@ -216,13 +220,15 @@ void SScalabilitySettings::Construct( const FArguments& InArgs )
 	Scalability::FQualityLevels LevelCounts = Scalability::GetQualityLevelCounts();
 	const int32 MaxLevelCount =
 		FMath::Max(LevelCounts.ShadowQuality,
+		FMath::Max(LevelCounts.GlobalIlluminationQuality,
+		FMath::Max(LevelCounts.ReflectionQuality,
 		FMath::Max(LevelCounts.TextureQuality,
 		FMath::Max(LevelCounts.ViewDistanceQuality,
 		FMath::Max(LevelCounts.EffectsQuality,
 		FMath::Max(LevelCounts.FoliageQuality,
 		FMath::Max(LevelCounts.ShadingQuality,
 		FMath::Max(LevelCounts.PostProcessQuality, LevelCounts.AntiAliasingQuality)
-		))))));
+		))))))));
 	const int32 TotalWidth = MaxLevelCount + 1;
 
 	TSharedRef<SGridPanel> ButtonMatrix = 
@@ -238,6 +244,8 @@ void SScalabilitySettings::Construct( const FArguments& InArgs )
 		+MakeGridSlot(0,7,TotalWidth,1) [ SNew (SBorder).BorderImage(FEditorStyle::GetBrush("Scalability.RowBackground")) ]
 		+MakeGridSlot(0,8,TotalWidth,1) [ SNew (SBorder).BorderImage(FEditorStyle::GetBrush("Scalability.RowBackground")) ]
 		+MakeGridSlot(0,9,TotalWidth,1) [ SNew (SBorder).BorderImage(FEditorStyle::GetBrush("Scalability.RowBackground")) ]
+		+MakeGridSlot(0,10,TotalWidth,1) [ SNew (SBorder).BorderImage(FEditorStyle::GetBrush("Scalability.RowBackground")) ]
+		+MakeGridSlot(0,11,TotalWidth,1) [ SNew (SBorder).BorderImage(FEditorStyle::GetBrush("Scalability.RowBackground")) ]
 
 		+MakeGridSlot(0,0).VAlign(VAlign_Center) [ SNew(STextBlock).Text(LOCTEXT("QualityLabel", "Quality")).Font(TitleFont) ]
 		+MakeGridSlot(1,0) [ MakeHeaderButtonWidget(NamesLow, 0, LOCTEXT("QualityLow", "Set all groups to low quality"), Scalability::EQualityLevelBehavior::EAbsolute) ]
@@ -258,22 +266,26 @@ void SScalabilitySettings::Construct( const FArguments& InArgs )
 		+MakeGridSlot(0,4) [ SNew(STextBlock).Text(LOCTEXT("PostProcessQualityLabel1", "Post Processing")).Font(GroupFont) ]
 
 		+MakeGridSlot(0,5) [ SNew(STextBlock).Text(LOCTEXT("ShadowLabel1", "Shadows")).Font(GroupFont) ]
+		+MakeGridSlot(0,6) [ SNew(STextBlock).Text(LOCTEXT("GlobalIlluminationLabel1", "Global Illumination")).Font(GroupFont) ]
+		+MakeGridSlot(0,7) [ SNew(STextBlock).Text(LOCTEXT("Reflection1", "Reflections")).Font(GroupFont) ]
 
-		+MakeGridSlot(0,6) [ SNew(STextBlock).Text(LOCTEXT("TextureQualityLabel1", "Textures")).Font(GroupFont) ]
+		+MakeGridSlot(0,8) [ SNew(STextBlock).Text(LOCTEXT("TextureQualityLabel1", "Textures")).Font(GroupFont) ]
 
-		+MakeGridSlot(0,7) [ SNew(STextBlock).Text(LOCTEXT("EffectsQualityLabel1", "Effects")).Font(GroupFont) ]
-		+MakeGridSlot(0,8) [ SNew(STextBlock).Text(LOCTEXT("FoliageQualityLabel1", "Foliage")).Font(GroupFont) ]
-		+MakeGridSlot(0,9) [ SNew(STextBlock).Text(LOCTEXT("ShadingQualityLabel1", "Shading")).Font(GroupFont) ]
+		+MakeGridSlot(0,9) [ SNew(STextBlock).Text(LOCTEXT("EffectsQualityLabel1", "Effects")).Font(GroupFont) ]
+		+MakeGridSlot(0,10) [ SNew(STextBlock).Text(LOCTEXT("FoliageQualityLabel1", "Foliage")).Font(GroupFont) ]
+		+MakeGridSlot(0,11) [ SNew(STextBlock).Text(LOCTEXT("ShadingQualityLabel1", "Shading")).Font(GroupFont) ]
 		;
 
 	AddButtonsToGrid(1, 2, ButtonMatrix, FiveDistanceNames, LevelCounts.ViewDistanceQuality, TEXT("ViewDistanceQuality"), LOCTEXT("ViewDistanceQualityTooltip", "Set view distance to {0}"));
 	AddButtonsToGrid(1, 3, ButtonMatrix, FiveNames, LevelCounts.AntiAliasingQuality, TEXT("AntiAliasingQuality"), LOCTEXT("AntiAliasingQualityTooltip", "Set anti-aliasing quality to {0}"));
 	AddButtonsToGrid(1, 4, ButtonMatrix, FiveNames, LevelCounts.PostProcessQuality, TEXT("PostProcessQuality"), LOCTEXT("PostProcessQualityTooltip", "Set post processing quality to {0}"));
 	AddButtonsToGrid(1, 5, ButtonMatrix, FiveNames, LevelCounts.ShadowQuality, TEXT("ShadowQuality"), LOCTEXT("ShadowQualityTooltip", "Set shadow quality to {0}"));
-	AddButtonsToGrid(1, 6, ButtonMatrix, FiveNames, LevelCounts.TextureQuality, TEXT("TextureQuality"), LOCTEXT("TextureQualityTooltip", "Set texture quality to {0}"));
-	AddButtonsToGrid(1, 7, ButtonMatrix, FiveNames, LevelCounts.EffectsQuality, TEXT("EffectsQuality"), LOCTEXT("EffectsQualityTooltip", "Set effects quality to {0}"));
-	AddButtonsToGrid(1, 8, ButtonMatrix, FiveNames, LevelCounts.FoliageQuality, TEXT("FoliageQuality"), LOCTEXT("FoliageQualityTooltip", "Set foliage quality to {0}"));
-	AddButtonsToGrid(1, 9, ButtonMatrix, FiveNames, LevelCounts.ShadingQuality, TEXT("ShadingQuality"), LOCTEXT("ShadingQualityTooltip", "Set shading quality to {0}"));
+	AddButtonsToGrid(1, 6, ButtonMatrix, FiveNames, LevelCounts.GlobalIlluminationQuality, TEXT("GlobalIlluminationQuality"), LOCTEXT("GlobalIlluminationQualityTooltip", "Set Global Illumination quality to {0}"));
+	AddButtonsToGrid(1, 7, ButtonMatrix, FiveNames, LevelCounts.ReflectionQuality, TEXT("ReflectionQuality"), LOCTEXT("ReflectionQualityTooltip", "Set Reflection quality to {0}"));
+	AddButtonsToGrid(1, 8, ButtonMatrix, FiveNames, LevelCounts.TextureQuality, TEXT("TextureQuality"), LOCTEXT("TextureQualityTooltip", "Set texture quality to {0}"));
+	AddButtonsToGrid(1, 9, ButtonMatrix, FiveNames, LevelCounts.EffectsQuality, TEXT("EffectsQuality"), LOCTEXT("EffectsQualityTooltip", "Set effects quality to {0}"));
+	AddButtonsToGrid(1, 10, ButtonMatrix, FiveNames, LevelCounts.FoliageQuality, TEXT("FoliageQuality"), LOCTEXT("FoliageQualityTooltip", "Set foliage quality to {0}"));
+	AddButtonsToGrid(1, 11, ButtonMatrix, FiveNames, LevelCounts.ShadingQuality, TEXT("ShadingQuality"), LOCTEXT("ShadingQualityTooltip", "Set shading quality to {0}"));
 
 
 	this->ChildSlot
diff --git a/Engine/Source/Editor/UnrealEd/Private/Settings/EditorSettings.cpp b/Engine/Source/Editor/UnrealEd/Private/Settings/EditorSettings.cpp
--- a/Engine/Source/Editor/UnrealEd/Private/Settings/EditorSettings.cpp
+++ b/Engine/Source/Editor/UnrealEd/Private/Settings/EditorSettings.cpp
@@ -122,6 +122,8 @@ void UEditorSettings::LoadScalabilityBenchmark()
 		GConfig->GetInt(Section, TEXT("ViewDistanceQuality"), Temporary.ViewDistanceQuality, GEditorSettingsIni);
 		GConfig->GetInt(Section, TEXT("AntiAliasingQuality"), Temporary.AntiAliasingQuality, GEditorSettingsIni);
 		GConfig->GetInt(Section, TEXT("ShadowQuality"), Temporary.ShadowQuality, GEditorSettingsIni);
+		GConfig->GetInt(Section, TEXT("GlobalIlluminationQuality"), Temporary.GlobalIlluminationQuality, GEditorSettingsIni);
+		GConfig->GetInt(Section, TEXT("ReflectionQuality"), Temporary.ReflectionQuality, GEditorSettingsIni);
 		GConfig->GetInt(Section, TEXT("PostProcessQuality"), Temporary.PostProcessQuality, GEditorSettingsIni);
 		GConfig->GetInt(Section, TEXT("TextureQuality"), Temporary.TextureQuality, GEditorSettingsIni);
 		GConfig->GetInt(Section, TEXT("EffectsQuality"), Temporary.EffectsQuality, GEditorSettingsIni);
@@ -146,6 +148,8 @@ void UEditorSettings::AutoApplyScalabilityBenchmark()
 	GConfig->SetInt(Section, TEXT("ViewDistanceQuality"), Temporary.ViewDistanceQuality, GEditorSettingsIni);
 	GConfig->SetInt(Section, TEXT("AntiAliasingQuality"), Temporary.AntiAliasingQuality, GEditorSettingsIni);
 	GConfig->SetInt(Section, TEXT("ShadowQuality"), Temporary.ShadowQuality, GEditorSettingsIni);
+	GConfig->SetInt(Section, TEXT("GlobalIlluminationQuality"), Temporary.GlobalIlluminationQuality, GEditorSettingsIni);
+	GConfig->SetInt(Section, TEXT("ReflectionQuality"), Temporary.ReflectionQuality, GEditorSettingsIni);
 	GConfig->SetInt(Section, TEXT("PostProcessQuality"), Temporary.PostProcessQuality, GEditorSettingsIni);
 	GConfig->SetInt(Section, TEXT("TextureQuality"), Temporary.TextureQuality, GEditorSettingsIni);
 	GConfig->SetInt(Section, TEXT("EffectsQuality"), Temporary.EffectsQuality, GEditorSettingsIni);
diff --git a/Engine/Source/Runtime/AutomationMessages/Public/AutomationWorkerMessages.h b/Engine/Source/Runtime/AutomationMessages/Public/AutomationWorkerMessages.h
--- a/Engine/Source/Runtime/AutomationMessages/Public/AutomationWorkerMessages.h
+++ b/Engine/Source/Runtime/AutomationMessages/Public/AutomationWorkerMessages.h
@@ -444,6 +444,10 @@ public:
 	UPROPERTY(EditAnywhere, Category="Message")
 	int32 ShadowQuality;
 	UPROPERTY(EditAnywhere, Category="Message")
+	int32 GlobalIlluminationQuality;
+	UPROPERTY(EditAnywhere, Category="Message")
+	int32 ReflectionQuality;
+	UPROPERTY(EditAnywhere, Category="Message")
 	int32 PostProcessQuality;
 	UPROPERTY(EditAnywhere, Category="Message")
 	int32 TextureQuality;
@@ -488,6 +492,8 @@ public:
 		, ViewDistanceQuality(0)
 		, AntiAliasingQuality(0)
 		, ShadowQuality(0)
+		, GlobalIlluminationQuality(0)
+		, ReflectionQuality(0)
 		, PostProcessQuality(0)
 		, TextureQuality(0)
 		, EffectsQuality(0)
@@ -543,6 +549,8 @@ public:
 		ViewDistanceQuality = Data.ViewDistanceQuality;
 		AntiAliasingQuality = Data.AntiAliasingQuality;
 		ShadowQuality = Data.ShadowQuality;
+		GlobalIlluminationQuality = Data.GlobalIlluminationQuality;
+		ReflectionQuality = Data.ReflectionQuality;
 		PostProcessQuality = Data.PostProcessQuality;
 		TextureQuality = Data.TextureQuality;
 		EffectsQuality = Data.EffectsQuality;
@@ -641,6 +649,16 @@ public:
 			Score += 10;
 		}
 
+		if (GlobalIlluminationQuality == OtherMetadata.GlobalIlluminationQuality)
+		{
+			Score += 10;
+		}
+
+		if (ReflectionQuality == OtherMetadata.ReflectionQuality)
+		{
+			Score += 10;
+		}
+
 		if (PostProcessQuality == OtherMetadata.PostProcessQuality)
 		{
 			Score += 10;
diff --git a/Engine/Source/Runtime/Core/Public/Misc/AutomationTest.h b/Engine/Source/Runtime/Core/Public/Misc/AutomationTest.h
--- a/Engine/Source/Runtime/Core/Public/Misc/AutomationTest.h
+++ b/Engine/Source/Runtime/Core/Public/Misc/AutomationTest.h
@@ -614,6 +614,8 @@ struct FAutomationScreenshotData
 	int32 ViewDistanceQuality;
 	int32 AntiAliasingQuality;
 	int32 ShadowQuality;
+	int32 GlobalIlluminationQuality;
+	int32 ReflectionQuality;
 	int32 PostProcessQuality;
 	int32 TextureQuality;
 	int32 EffectsQuality;
@@ -646,6 +648,8 @@ struct FAutomationScreenshotData
 		, ViewDistanceQuality(0)
 		, AntiAliasingQuality(0)
 		, ShadowQuality(0)
+		, GlobalIlluminationQuality(0)
+		, ReflectionQuality(0)
 		, PostProcessQuality(0)
 		, TextureQuality(0)
 		, EffectsQuality(0)
diff --git a/Engine/Source/Runtime/Engine/Classes/GameFramework/GameUserSettings.h b/Engine/Source/Runtime/Engine/Classes/GameFramework/GameUserSettings.h
--- a/Engine/Source/Runtime/Engine/Classes/GameFramework/GameUserSettings.h
+++ b/Engine/Source/Runtime/Engine/Classes/GameFramework/GameUserSettings.h
@@ -199,6 +199,24 @@ public:
 	UFUNCTION(BlueprintCallable, Category=Settings)
 	int32 GetShadowQuality() const;
 
+	// Sets the global illumination quality (0..4, higher is better)
+	// @param Value 0:low, 1:medium, 2:high, 3:epic, 4:cinematic (gets clamped if needed)
+	UFUNCTION(BlueprintCallable, Category=Settings)
+	void SetGlobalIlluminationQuality(int32 Value);
+
+	// Returns the global illumination quality (0..4, higher is better)
+	UFUNCTION(BlueprintCallable, Category=Settings)
+	int32 GetGlobalIlluminationQuality() const;
+
+	// Sets the reflection quality (0..4, higher is better)
+	// @param Value 0:low, 1:medium, 2:high, 3:epic, 4:cinematic (gets clamped if needed)
+	UFUNCTION(BlueprintCallable, Category=Settings)
+	void SetReflectionQuality(int32 Value);
+
+	// Returns the reflection quality (0..4, higher is better)
+	UFUNCTION(BlueprintCallable, Category=Settings)
+	int32 GetReflectionQuality() const;
+
 	// Sets the anti-aliasing quality (0..4, higher is better)
 	// @param Value 0:low, 1:medium, 2:high, 3:epic, 4:cinematic (gets clamped if needed)
 	UFUNCTION(BlueprintCallable, Category=Settings)
diff --git a/Engine/Source/Runtime/Engine/Private/ChartCreation.cpp b/Engine/Source/Runtime/Engine/Private/ChartCreation.cpp
--- a/Engine/Source/Runtime/Engine/Private/ChartCreation.cpp
+++ b/Engine/Source/Runtime/Engine/Private/ChartCreation.cpp
@@ -174,6 +174,8 @@ void FDumpFPSChartToEndpoint::HandleBasicStats()
 	PrintToEndpoint(FString::Printf(TEXT("\tView Distance Quality: %d"), ScalabilityQuality.ViewDistanceQuality));
 	PrintToEndpoint(FString::Printf(TEXT("\tAnti-Aliasing Quality: %d"), ScalabilityQuality.AntiAliasingQuality));
 	PrintToEndpoint(FString::Printf(TEXT("\tShadow Quality: %d"), ScalabilityQuality.ShadowQuality));
+	PrintToEndpoint(FString::Printf(TEXT("\tGlobal Illumination Quality: %d"), ScalabilityQuality.GlobalIlluminationQuality));
+	PrintToEndpoint(FString::Printf(TEXT("\tReflection Quality: %d"), ScalabilityQuality.ReflectionQuality));
 	PrintToEndpoint(FString::Printf(TEXT("\tPost-Process Quality: %d"), ScalabilityQuality.PostProcessQuality));
 	PrintToEndpoint(FString::Printf(TEXT("\tTexture Quality: %d"), ScalabilityQuality.TextureQuality));
 	PrintToEndpoint(FString::Printf(TEXT("\tEffects Quality: %d"), ScalabilityQuality.EffectsQuality));
@@ -377,6 +379,8 @@ protected:
 			ParamArray.Add(FAnalyticsEventAttribute(TEXT("ViewDistanceQuality"), ScalabilityQuality.ViewDistanceQuality));
 			ParamArray.Add(FAnalyticsEventAttribute(TEXT("AntiAliasingQuality"), ScalabilityQuality.AntiAliasingQuality));
 			ParamArray.Add(FAnalyticsEventAttribute(TEXT("ShadowQuality"), ScalabilityQuality.ShadowQuality));
+			ParamArray.Add(FAnalyticsEventAttribute(TEXT("GlobalIlluminationQuality"), ScalabilityQuality.GlobalIlluminationQuality));
+			ParamArray.Add(FAnalyticsEventAttribute(TEXT("ReflectionQuality"), ScalabilityQuality.ReflectionQuality));
 			ParamArray.Add(FAnalyticsEventAttribute(TEXT("PostProcessQuality"), ScalabilityQuality.PostProcessQuality));
 			ParamArray.Add(FAnalyticsEventAttribute(TEXT("TextureQuality"), ScalabilityQuality.TextureQuality));
 			ParamArray.Add(FAnalyticsEventAttribute(TEXT("FXQuality"), ScalabilityQuality.EffectsQuality));
@@ -523,6 +527,8 @@ protected:
 		FPSChartRow = FPSChartRow.Replace(TEXT("TOKEN_SETTINGS_VD"), *FString::Printf(TEXT("%d"), ScalabilityQuality.ViewDistanceQuality), ESearchCase::CaseSensitive);
 		FPSChartRow = FPSChartRow.Replace(TEXT("TOKEN_SETTINGS_AA"), *FString::Printf(TEXT("%d"), ScalabilityQuality.AntiAliasingQuality), ESearchCase::CaseSensitive);
 		FPSChartRow = FPSChartRow.Replace(TEXT("TOKEN_SETTINGS_SHADOW"), *FString::Printf(TEXT("%d"), ScalabilityQuality.ShadowQuality), ESearchCase::CaseSensitive);
+		FPSChartRow = FPSChartRow.Replace(TEXT("TOKEN_SETTINGS_GI"), *FString::Printf(TEXT("%d"), ScalabilityQuality.GlobalIlluminationQuality), ESearchCase::CaseSensitive);
+		FPSChartRow = FPSChartRow.Replace(TEXT("TOKEN_SETTINGS_REFL"), *FString::Printf(TEXT("%d"), ScalabilityQuality.ReflectionQuality), ESearchCase::CaseSensitive);
 		FPSChartRow = FPSChartRow.Replace(TEXT("TOKEN_SETTINGS_PP"), *FString::Printf(TEXT("%d"), ScalabilityQuality.PostProcessQuality), ESearchCase::CaseSensitive);
 		FPSChartRow = FPSChartRow.Replace(TEXT("TOKEN_SETTINGS_TEX"), *FString::Printf(TEXT("%d"), ScalabilityQuality.TextureQuality), ESearchCase::CaseSensitive);
 		FPSChartRow = FPSChartRow.Replace(TEXT("TOKEN_SETTINGS_FX"), *FString::Printf(TEXT("%d"), ScalabilityQuality.EffectsQuality), ESearchCase::CaseSensitive);
diff --git a/Engine/Source/Runtime/Engine/Private/GameUserSettings.cpp b/Engine/Source/Runtime/Engine/Private/GameUserSettings.cpp
--- a/Engine/Source/Runtime/Engine/Private/GameUserSettings.cpp
+++ b/Engine/Source/Runtime/Engine/Private/GameUserSettings.cpp
@@ -799,6 +799,26 @@ int32 UGameUserSettings::GetShadowQuality() const
 	return ScalabilityQuality.ShadowQuality;
 }
 
+void UGameUserSettings::SetGlobalIlluminationQuality(int32 Value)
+{
+	ScalabilityQuality.SetGlobalIlluminationQuality(Value);
+}
+
+int32 UGameUserSettings::GetGlobalIlluminationQuality() const
+{
+	return ScalabilityQuality.GlobalIlluminationQuality;
+}
+
+void UGameUserSettings::SetReflectionQuality(int32 Value)
+{
+	ScalabilityQuality.SetReflectionQuality(Value);
+}
+
+int32 UGameUserSettings::GetReflectionQuality() const
+{
+	return ScalabilityQuality.ReflectionQuality;
+}
+
 void UGameUserSettings::SetAntiAliasingQuality(int32 Value)
 {
 	ScalabilityQuality.SetAntiAliasingQuality(Value);
diff --git a/Engine/Source/Runtime/Engine/Private/Scalability.cpp b/Engine/Source/Runtime/Engine/Private/Scalability.cpp
--- a/Engine/Source/Runtime/Engine/Private/Scalability.cpp
+++ b/Engine/Source/Runtime/Engine/Private/Scalability.cpp
@@ -38,6 +38,20 @@ static TAutoConsoleVariable<int32> CVarShadowQuality(
 	TEXT(" 0:low, 1:med, 2:high, 3:epic, 4:cinematic, default: 3"),
 	ECVF_ScalabilityGroup);
 
+static TAutoConsoleVariable<int32> CVarGlobalIlluminationQuality(
+	TEXT("sg.GlobalIlluminationQuality"),
+	Scalability::DefaultQualityLevel,
+	TEXT("Scalability quality state (internally used by scalability system, ini load/save or using SCALABILITY console command)\n")
+	TEXT(" 0:low, 1:med, 2:high, 3:epic, 4:cinematic, default: 3"),
+	ECVF_ScalabilityGroup);
+
+static TAutoConsoleVariable<int32> CVarReflectionQuality(
+	TEXT("sg.ReflectionQuality"),
+	Scalability::DefaultQualityLevel,
+	TEXT("Scalability quality state (internally used by scalability system, ini load/save or using SCALABILITY console command)\n")
+	TEXT(" 0:low, 1:med, 2:high, 3:epic, 4:cinematic, default: 3"),
+	ECVF_ScalabilityGroup);
+
 static TAutoConsoleVariable<int32> CVarPostProcessQuality(
 	TEXT("sg.PostProcessQuality"),
 	Scalability::DefaultQualityLevel,
@@ -94,6 +108,20 @@ static TAutoConsoleVariable<int32> CVarShadowQuality_NumLevels(
 	TEXT(" default: 5 (0..4)"),
 	ECVF_ReadOnly);
 
+static TAutoConsoleVariable<int32> CVarGlobalIlluminationQuality_NumLevels(
+	TEXT("sg.GlobalIlluminationQuality.NumLevels"),
+	5,
+	TEXT("Number of settings quality levels in sg.GlobalIlluminationQuality\n")
+	TEXT(" default: 5 (0..4)"),
+	ECVF_ReadOnly);
+
+static TAutoConsoleVariable<int32> CVarReflectionQuality_NumLevels(
+	TEXT("sg.ReflectionQuality.NumLevels"),
+	5,
+	TEXT("Number of settings quality levels in sg.ReflectionQuality\n")
+	TEXT(" default: 5 (0..4)"),
+	ECVF_ReadOnly);
+
 static TAutoConsoleVariable<int32> CVarPostProcessQuality_NumLevels(
 	TEXT("sg.PostProcessQuality.NumLevels"),
 	5,
@@ -465,6 +493,16 @@ void OnChangeShadowQuality(IConsoleVariable* Var)
 	SetGroupQualityLevel(TEXT("ShadowQuality"), Var->GetInt(), CVarShadowQuality_NumLevels->GetInt());
 }
 
+void OnChangeGlobalIlluminationQuality(IConsoleVariable* Var)
+{
+	SetGroupQualityLevel(TEXT("GlobalIlluminationQuality"), Var->GetInt(), CVarGlobalIlluminationQuality_NumLevels->GetInt());
+}
+
+void OnChangeReflectionQuality(IConsoleVariable* Var)
+{
+	SetGroupQualityLevel(TEXT("ReflectionQuality"), Var->GetInt(), CVarReflectionQuality_NumLevels->GetInt());
+}
+
 void OnChangePostProcessQuality(IConsoleVariable* Var)
 {
 	SetGroupQualityLevel(TEXT("PostProcessQuality"), Var->GetInt(), CVarPostProcessQuality_NumLevels->GetInt());
@@ -508,6 +546,8 @@ void InitScalabilitySystem()
 	CVarViewDistanceQuality.AsVariable()->SetOnChangedCallback(FConsoleVariableDelegate::CreateStatic(&OnChangeViewDistanceQuality));
 	CVarAntiAliasingQuality.AsVariable()->SetOnChangedCallback(FConsoleVariableDelegate::CreateStatic(&OnChangeAntiAliasingQuality));
 	CVarShadowQuality.AsVariable()->SetOnChangedCallback(FConsoleVariableDelegate::CreateStatic(&OnChangeShadowQuality));
+	CVarGlobalIlluminationQuality.AsVariable()->SetOnChangedCallback(FConsoleVariableDelegate::CreateStatic(&OnChangeGlobalIlluminationQuality));
+	CVarReflectionQuality.AsVariable()->SetOnChangedCallback(FConsoleVariableDelegate::CreateStatic(&OnChangeReflectionQuality));
 	CVarPostProcessQuality.AsVariable()->SetOnChangedCallback(FConsoleVariableDelegate::CreateStatic(&OnChangePostProcessQuality));
 	CVarTextureQuality.AsVariable()->SetOnChangedCallback(FConsoleVariableDelegate::CreateStatic(&OnChangeTextureQuality));
 	CVarEffectsQuality.AsVariable()->SetOnChangedCallback(FConsoleVariableDelegate::CreateStatic(&OnChangeEffectsQuality));
@@ -566,6 +606,8 @@ FQualityLevels BenchmarkQualityLevels(uint32 WorkScale, float CPUMultiplier, flo
 	Results.ViewDistanceQuality = ComputeOptionFromPerfIndex(TEXT("ViewDistanceQuality"), CPUPerfIndex, GPUPerfIndex);
 	Results.AntiAliasingQuality = ComputeOptionFromPerfIndex(TEXT("AntiAliasingQuality"), CPUPerfIndex, GPUPerfIndex);
 	Results.ShadowQuality = ComputeOptionFromPerfIndex(TEXT("ShadowQuality"), CPUPerfIndex, GPUPerfIndex);
+	Results.GlobalIlluminationQuality = ComputeOptionFromPerfIndex(TEXT("GlobalIlluminationQuality"), CPUPerfIndex, GPUPerfIndex);
+	Results.ReflectionQuality = ComputeOptionFromPerfIndex(TEXT("ReflectionQuality"), CPUPerfIndex, GPUPerfIndex);
 	Results.PostProcessQuality = ComputeOptionFromPerfIndex(TEXT("PostProcessQuality"), CPUPerfIndex, GPUPerfIndex);
 	Results.TextureQuality = ComputeOptionFromPerfIndex(TEXT("TextureQuality"), CPUPerfIndex, GPUPerfIndex);
 	Results.EffectsQuality = ComputeOptionFromPerfIndex(TEXT("EffectsQuality"), CPUPerfIndex, GPUPerfIndex);
@@ -672,6 +714,8 @@ void ProcessCommand(const TCHAR* Cmd, FOutputDevice& Ar)
 		PrintGroupInfo(TEXT("ViewDistanceQuality"), bInfoMode);
 		PrintGroupInfo(TEXT("AntiAliasingQuality"), bInfoMode);
 		PrintGroupInfo(TEXT("ShadowQuality"), bInfoMode);
+		PrintGroupInfo(TEXT("GlobalIlluminationQuality"), bInfoMode);
+		PrintGroupInfo(TEXT("ReflectionQuality"), bInfoMode);
 		PrintGroupInfo(TEXT("PostProcessQuality"), bInfoMode);
 		PrintGroupInfo(TEXT("TextureQuality"), bInfoMode);
 		PrintGroupInfo(TEXT("EffectsQuality"), bInfoMode);
@@ -720,6 +764,8 @@ void SetQualityLevels(const FQualityLevels& QualityLevels, bool bForce/* = false
 	ClampedLevels.SetViewDistanceQuality(QualityLevels.ViewDistanceQuality);
 	ClampedLevels.SetAntiAliasingQuality(QualityLevels.AntiAliasingQuality);
 	ClampedLevels.SetShadowQuality(QualityLevels.ShadowQuality);
+	ClampedLevels.SetGlobalIlluminationQuality(QualityLevels.GlobalIlluminationQuality);
+	ClampedLevels.SetReflectionQuality(QualityLevels.ReflectionQuality);
 	ClampedLevels.SetPostProcessQuality(QualityLevels.PostProcessQuality);
 	ClampedLevels.SetTextureQuality(QualityLevels.TextureQuality);
 	ClampedLevels.SetEffectsQuality(QualityLevels.EffectsQuality);
@@ -740,6 +786,8 @@ void SetQualityLevels(const FQualityLevels& QualityLevels, bool bForce/* = false
 		SetQualityLevelCVar(CVarViewDistanceQuality, ClampedLevels.ViewDistanceQuality, GScalabilityQualityLevelsOverride.ViewDistanceQuality, bForce);
 		SetQualityLevelCVar(CVarAntiAliasingQuality, ClampedLevels.AntiAliasingQuality, GScalabilityQualityLevelsOverride.AntiAliasingQuality, bForce);
 		SetQualityLevelCVar(CVarShadowQuality, ClampedLevels.ShadowQuality, GScalabilityQualityLevelsOverride.ShadowQuality, bForce);
+		SetQualityLevelCVar(CVarGlobalIlluminationQuality, ClampedLevels.GlobalIlluminationQuality, GScalabilityQualityLevelsOverride.GlobalIlluminationQuality, bForce);
+		SetQualityLevelCVar(CVarReflectionQuality, ClampedLevels.ReflectionQuality, GScalabilityQualityLevelsOverride.ReflectionQuality, bForce);
 		SetQualityLevelCVar(CVarPostProcessQuality, ClampedLevels.PostProcessQuality, GScalabilityQualityLevelsOverride.PostProcessQuality, bForce);
 		SetQualityLevelCVar(CVarTextureQuality, ClampedLevels.TextureQuality, GScalabilityQualityLevelsOverride.TextureQuality, bForce);
 		SetQualityLevelCVar(CVarEffectsQuality, ClampedLevels.EffectsQuality, GScalabilityQualityLevelsOverride.EffectsQuality, bForce);
@@ -761,6 +809,8 @@ FQualityLevels GetQualityLevels()
 		Ret.ViewDistanceQuality = CVarViewDistanceQuality.GetValueOnGameThread();
 		Ret.AntiAliasingQuality = CVarAntiAliasingQuality.GetValueOnGameThread();
 		Ret.ShadowQuality = CVarShadowQuality.GetValueOnGameThread();
+		Ret.GlobalIlluminationQuality = CVarGlobalIlluminationQuality.GetValueOnGameThread();
+		Ret.ReflectionQuality = CVarReflectionQuality.GetValueOnGameThread();
 		Ret.PostProcessQuality = CVarPostProcessQuality.GetValueOnGameThread();
 		Ret.TextureQuality = CVarTextureQuality.GetValueOnGameThread();
 		Ret.EffectsQuality = CVarEffectsQuality.GetValueOnGameThread();
@@ -828,6 +878,8 @@ void FQualityLevels::SetFromSingleQualityLevel(int32 Value)
 	ViewDistanceQuality = FMath::Clamp(Value, 0, CVarViewDistanceQuality_NumLevels->GetInt() - 1);
 	AntiAliasingQuality = FMath::Clamp(Value, 0, CVarAntiAliasingQuality_NumLevels->GetInt() - 1);
 	ShadowQuality = FMath::Clamp(Value, 0, CVarShadowQuality_NumLevels->GetInt() - 1);
+	GlobalIlluminationQuality = FMath::Clamp(Value, 0, CVarGlobalIlluminationQuality_NumLevels->GetInt() - 1);
+	ReflectionQuality = FMath::Clamp(Value, 0, CVarReflectionQuality_NumLevels->GetInt() - 1);
 	PostProcessQuality = FMath::Clamp(Value, 0, CVarPostProcessQuality_NumLevels->GetInt() - 1);
 	TextureQuality = FMath::Clamp(Value, 0, CVarTextureQuality_NumLevels->GetInt() - 1);
 	EffectsQuality = FMath::Clamp(Value, 0, CVarEffectsQuality_NumLevels->GetInt() - 1);
@@ -845,6 +897,8 @@ void FQualityLevels::SetFromSingleQualityLevelRelativeToMax(int32 Value)
 	ViewDistanceQuality = FMath::Max(CVarViewDistanceQuality_NumLevels->GetInt() - Value, 0);
 	AntiAliasingQuality = FMath::Max(CVarAntiAliasingQuality_NumLevels->GetInt() - Value, 0);
 	ShadowQuality = FMath::Max(CVarShadowQuality_NumLevels->GetInt() - Value, 0);
+	GlobalIlluminationQuality = FMath::Max(CVarGlobalIlluminationQuality_NumLevels->GetInt() - Value, 0);
+	ReflectionQuality = FMath::Max(CVarReflectionQuality_NumLevels->GetInt() - Value, 0);
 	PostProcessQuality = FMath::Max(CVarPostProcessQuality_NumLevels->GetInt() - Value, 0);
 	TextureQuality = FMath::Max(CVarTextureQuality_NumLevels->GetInt() - Value, 0);
 	EffectsQuality = FMath::Max(CVarEffectsQuality_NumLevels->GetInt() - Value, 0);
@@ -859,7 +913,7 @@ int32 FQualityLevels::GetSingleQualityLevel() const
 	int32 Result = ViewDistanceQuality;
 
 	const int32 Target = ViewDistanceQuality;
-	if ((Target == AntiAliasingQuality) && (Target == ShadowQuality) && (Target == PostProcessQuality) && (Target == TextureQuality) && (Target == EffectsQuality) && (Target == FoliageQuality) && (Target == ShadingQuality))
+	if ((Target == AntiAliasingQuality) && (Target == ShadowQuality) && (Target == GlobalIlluminationQuality) && (Target == ReflectionQuality) && (Target == PostProcessQuality) && (Target == TextureQuality) && (Target == EffectsQuality) && (Target == FoliageQuality) && (Target == ShadingQuality))
 	{
 		if (GetRenderScaleLevelFromQualityLevel(Target) == ResolutionQuality)
 		{
@@ -876,6 +930,8 @@ int32 FQualityLevels::GetMinQualityLevel() const
 
 	Level = FMath::Min(Level, AntiAliasingQuality);
 	Level = FMath::Min(Level, ShadowQuality);
+	Level = FMath::Min(Level, GlobalIlluminationQuality);
+	Level = FMath::Min(Level, ReflectionQuality);
 	Level = FMath::Min(Level, PostProcessQuality);
 	Level = FMath::Min(Level, TextureQuality);
 	Level = FMath::Min(Level, EffectsQuality);
@@ -900,6 +956,16 @@ void FQualityLevels::SetShadowQuality(int32 Value)
 	ShadowQuality = FMath::Clamp(Value, 0, CVarShadowQuality_NumLevels->GetInt() - 1);
 }
 
+void FQualityLevels::SetGlobalIlluminationQuality(int32 Value)
+{
+	GlobalIlluminationQuality = FMath::Clamp(Value, 0, CVarGlobalIlluminationQuality_NumLevels->GetInt() - 1);
+}
+
+void FQualityLevels::SetReflectionQuality(int32 Value)
+{
+	ReflectionQuality = FMath::Clamp(Value, 0, CVarReflectionQuality_NumLevels->GetInt() - 1);
+}
+
 void FQualityLevels::SetPostProcessQuality(int32 Value)
 {
 	PostProcessQuality = FMath::Clamp(Value, 0, CVarPostProcessQuality_NumLevels->GetInt() - 1);
@@ -942,6 +1008,8 @@ void LoadState(const FString& IniName)
 	GConfig->GetInt(Section, TEXT("sg.ViewDistanceQuality"), State.ViewDistanceQuality, IniName);
 	GConfig->GetInt(Section, TEXT("sg.AntiAliasingQuality"), State.AntiAliasingQuality, IniName);
 	GConfig->GetInt(Section, TEXT("sg.ShadowQuality"), State.ShadowQuality, IniName);
+	GConfig->GetInt(Section, TEXT("sg.GlobalIlluminationQuality"), State.GlobalIlluminationQuality, IniName);
+	GConfig->GetInt(Section, TEXT("sg.ReflectionQuality"), State.ReflectionQuality, IniName);
 	GConfig->GetInt(Section, TEXT("sg.PostProcessQuality"), State.PostProcessQuality, IniName);
 	GConfig->GetInt(Section, TEXT("sg.TextureQuality"), State.TextureQuality, IniName);
 	GConfig->GetInt(Section, TEXT("sg.EffectsQuality"), State.EffectsQuality, IniName);
@@ -973,6 +1041,8 @@ void SaveState(const FString& IniName)
 	GConfig->SetInt(Section, TEXT("sg.ViewDistanceQuality"), State.ViewDistanceQuality, IniName);
 	GConfig->SetInt(Section, TEXT("sg.AntiAliasingQuality"), State.AntiAliasingQuality, IniName);
 	GConfig->SetInt(Section, TEXT("sg.ShadowQuality"), State.ShadowQuality, IniName);
+	GConfig->SetInt(Section, TEXT("sg.GlobalIlluminationQuality"), State.GlobalIlluminationQuality, IniName);
+	GConfig->SetInt(Section, TEXT("sg.ReflectionQuality"), State.ReflectionQuality, IniName);
 	GConfig->SetInt(Section, TEXT("sg.PostProcessQuality"), State.PostProcessQuality, IniName);
 	GConfig->SetInt(Section, TEXT("sg.TextureQuality"), State.TextureQuality, IniName);
 	GConfig->SetInt(Section, TEXT("sg.EffectsQuality"), State.EffectsQuality, IniName);
@@ -992,6 +1062,8 @@ void RecordQualityLevelsAnalytics(bool bAutoApplied)
 		Attributes.Add(FAnalyticsEventAttribute(TEXT("ViewDistanceQuality"), State.ViewDistanceQuality));
 		Attributes.Add(FAnalyticsEventAttribute(TEXT("AntiAliasingQuality"), State.AntiAliasingQuality));
 		Attributes.Add(FAnalyticsEventAttribute(TEXT("ShadowQuality"), State.ShadowQuality));
+		Attributes.Add(FAnalyticsEventAttribute(TEXT("GlobalIlluminationQuality"), State.GlobalIlluminationQuality));
+		Attributes.Add(FAnalyticsEventAttribute(TEXT("ReflectionQuality"), State.ReflectionQuality));
 		Attributes.Add(FAnalyticsEventAttribute(TEXT("PostProcessQuality"), State.PostProcessQuality));
 		Attributes.Add(FAnalyticsEventAttribute(TEXT("TextureQuality"), State.TextureQuality));
 		Attributes.Add(FAnalyticsEventAttribute(TEXT("EffectsQuality"), State.EffectsQuality));
@@ -1011,6 +1083,8 @@ FQualityLevels GetQualityLevelCounts()
 	Result.ViewDistanceQuality = CVarViewDistanceQuality_NumLevels->GetInt();
 	Result.AntiAliasingQuality = CVarAntiAliasingQuality_NumLevels->GetInt();
 	Result.ShadowQuality = CVarShadowQuality_NumLevels->GetInt();
+	Result.GlobalIlluminationQuality = CVarGlobalIlluminationQuality_NumLevels->GetInt();
+	Result.ReflectionQuality = CVarReflectionQuality_NumLevels->GetInt();
 	Result.PostProcessQuality = CVarPostProcessQuality_NumLevels->GetInt();
 	Result.TextureQuality = CVarTextureQuality_NumLevels->GetInt();
 	Result.EffectsQuality = CVarEffectsQuality_NumLevels->GetInt();
diff --git a/Engine/Source/Runtime/Engine/Private/Tests/AutomationCommon.cpp b/Engine/Source/Runtime/Engine/Private/Tests/AutomationCommon.cpp
--- a/Engine/Source/Runtime/Engine/Private/Tests/AutomationCommon.cpp
+++ b/Engine/Source/Runtime/Engine/Private/Tests/AutomationCommon.cpp
@@ -138,6 +138,8 @@ namespace AutomationCommon
 		Data.ViewDistanceQuality = QualityLevels.ViewDistanceQuality;
 		Data.AntiAliasingQuality = QualityLevels.AntiAliasingQuality;
 		Data.ShadowQuality = QualityLevels.ShadowQuality;
+		Data.GlobalIlluminationQuality = QualityLevels.GlobalIlluminationQuality;
+		Data.ReflectionQuality = QualityLevels.ReflectionQuality;
 		Data.PostProcessQuality = QualityLevels.PostProcessQuality;
 		Data.TextureQuality = QualityLevels.TextureQuality;
 		Data.EffectsQuality = QualityLevels.EffectsQuality;
diff --git a/Engine/Source/Runtime/Engine/Public/Scalability.h b/Engine/Source/Runtime/Engine/Public/Scalability.h
--- a/Engine/Source/Runtime/Engine/Public/Scalability.h
+++ b/Engine/Source/Runtime/Engine/Public/Scalability.h
@@ -28,6 +28,8 @@ namespace Scalability
 		int32 ViewDistanceQuality;
 		int32 AntiAliasingQuality;
 		int32 ShadowQuality;
+		int32 GlobalIlluminationQuality;
+		int32 ReflectionQuality;
 		int32 PostProcessQuality;
 		int32 TextureQuality;
 		int32 EffectsQuality;
@@ -52,6 +54,8 @@ namespace Scalability
 				ViewDistanceQuality == Other.ViewDistanceQuality &&
 				AntiAliasingQuality == Other.AntiAliasingQuality &&
 				ShadowQuality == Other.ShadowQuality &&
+				GlobalIlluminationQuality == Other.GlobalIlluminationQuality &&
+				ReflectionQuality == Other.ReflectionQuality &&
 				PostProcessQuality == Other.PostProcessQuality &&
 				TextureQuality == Other.TextureQuality &&
 				EffectsQuality == Other.EffectsQuality &&
@@ -71,6 +75,8 @@ namespace Scalability
 				FCrc::TypeCrc32<int32>(ViewDistanceQuality) ^
 				FCrc::TypeCrc32<int32>(AntiAliasingQuality) ^
 				FCrc::TypeCrc32<int32>(ShadowQuality) ^
+				FCrc::TypeCrc32<int32>(GlobalIlluminationQuality) ^
+				FCrc::TypeCrc32<int32>(ReflectionQuality) ^
 				FCrc::TypeCrc32<int32>(PostProcessQuality) ^
 				FCrc::TypeCrc32<int32>(TextureQuality) ^
 				FCrc::TypeCrc32<int32>(EffectsQuality) ^
@@ -106,6 +112,14 @@ namespace Scalability
 		// @param Value 0:low, 1:medium, 2:high, 3:epic, 4:cinematic (gets clamped if needed)
 		void SetShadowQuality(int32 Value);
 
+		// Sets shadow quality
+		// @param Value 0:low, 1:medium, 2:high, 3:epic, 4:cinematic (gets clamped if needed)
+		void SetGlobalIlluminationQuality(int32 Value);
+
+		// Sets shadow quality
+		// @param Value 0:low, 1:medium, 2:high, 3:epic, 4:cinematic (gets clamped if needed)
+		void SetReflectionQuality(int32 Value);
+
 		// Sets the post-processing quality
 		// @param Value 0:low, 1:medium, 2:high, 3:epic, 4:cinematic (gets clamped if needed)
 		void SetPostProcessQuality(int32 Value);
