Fixes UWorldPartitionRuntimeHash::CheckForErrorsInternal crash when no DataLayers are present.

#changelist validated

diff --git a/Engine/Source/Runtime/Engine/Private/WorldPartition/WorldPartitionRuntimeHash.cpp b/Engine/Source/Runtime/Engine/Private/WorldPartition/WorldPartitionRuntimeHash.cpp
--- a/Engine/Source/Runtime/Engine/Private/WorldPartition/WorldPartitionRuntimeHash.cpp
+++ b/Engine/Source/Runtime/Engine/Private/WorldPartition/WorldPartitionRuntimeHash.cpp
@@ -72,6 +72,7 @@ void UWorldPartitionRuntimeHash::CheckForErrors() const
 
 void UWorldPartitionRuntimeHash::CheckForErrorsInternal(const TMap<FGuid, FWorldPartitionActorViewProxy>& ActorDescList) const
 {
+	AWorldDataLayers* WorldDataLayers = GetWorld()->GetWorldDataLayers();
 	auto GetActorLabel = [](const FWorldPartitionActorDescView& ActorDescView) -> FString
 	{
 		const FName ActorLabel = ActorDescView.GetActorLabel();
@@ -119,24 +120,25 @@ void UWorldPartitionRuntimeHash::CheckForErrorsInternal(const TMap<FGuid, FWorld
 							->AddToken(FMapErrorToken::Create(FName(TEXT("WorldPartition_StreamedActorReferenceAlwaysLoadedActor_CheckForErrors"))));
 					}
 
-					AWorldDataLayers* WorldDataLayers = GetWorld()->GetWorldDataLayers();
-
-					TArray<FName> ActorDescLayerNames = ActorDescView.GetDataLayers();
-					ActorDescLayerNames.Sort([](const FName& A, const FName& B) { return A.FastLess(B); });
-					TArray<const UDataLayer*> ActorDescLayers = WorldDataLayers->GetDataLayerObjects(ActorDescLayerNames);
-
-					TArray<FName> ActorDescRefLayerNames = ActorDescRefView->GetDataLayers();
-					ActorDescRefLayerNames.Sort([](const FName& A, const FName& B) { return A.FastLess(B); });
-					TArray<const UDataLayer*> ActorDescRefLayers = WorldDataLayers->GetDataLayerObjects(ActorDescRefLayerNames);
-
-					if (ActorDescLayers != ActorDescRefLayers)
+					if (WorldDataLayers != nullptr)
 					{
-						TSharedRef<FTokenizedMessage> Error = FMessageLog("MapCheck").Error()
-							->AddToken(FTextToken::Create(LOCTEXT("MapCheck_WorldPartition_Actor", "Actor")))
-							->AddToken(FAssetNameToken::Create(GetActorLabel(ActorDescView)))
-							->AddToken(FTextToken::Create(LOCTEXT("MapCheck_WorldPartition_ReferenceActorInOtherDataLayers", "references an actor in a different set of data layers")))
-							->AddToken(FAssetNameToken::Create(GetActorLabel(*ActorDescRefView)))
-							->AddToken(FMapErrorToken::Create(FName(TEXT("WorldPartition_ActorReferenceActorInAnotherDataLayer_CheckForErrors"))));
+						TArray<FName> ActorDescLayerNames = ActorDescView.GetDataLayers();
+						ActorDescLayerNames.Sort([](const FName& A, const FName& B) { return A.FastLess(B); });
+						TArray<const UDataLayer*> ActorDescLayers = WorldDataLayers->GetDataLayerObjects(ActorDescLayerNames);
+
+						TArray<FName> ActorDescRefLayerNames = ActorDescRefView->GetDataLayers();
+						ActorDescRefLayerNames.Sort([](const FName& A, const FName& B) { return A.FastLess(B); });
+						TArray<const UDataLayer*> ActorDescRefLayers = WorldDataLayers->GetDataLayerObjects(ActorDescRefLayerNames);
+
+						if (ActorDescLayers != ActorDescRefLayers)
+						{
+							TSharedRef<FTokenizedMessage> Error = FMessageLog("MapCheck").Error()
+								->AddToken(FTextToken::Create(LOCTEXT("MapCheck_WorldPartition_Actor", "Actor")))
+								->AddToken(FAssetNameToken::Create(GetActorLabel(ActorDescView)))
+								->AddToken(FTextToken::Create(LOCTEXT("MapCheck_WorldPartition_ReferenceActorInOtherDataLayers", "references an actor in a different set of data layers")))
+								->AddToken(FAssetNameToken::Create(GetActorLabel(*ActorDescRefView)))
+								->AddToken(FMapErrorToken::Create(FName(TEXT("WorldPartition_ActorReferenceActorInAnotherDataLayer_CheckForErrors"))));
+						}
 					}
 				}
 				else
