Make the memory throttling consider virtual memory limits.

diff --git a/Engine/Source/Runtime/Engine/Private/AssetCompilingManager.cpp b/Engine/Source/Runtime/Engine/Private/AssetCompilingManager.cpp
--- a/Engine/Source/Runtime/Engine/Private/AssetCompilingManager.cpp
+++ b/Engine/Source/Runtime/Engine/Private/AssetCompilingManager.cpp
@@ -109,16 +109,20 @@ public:
 
 	int64 GetMemoryLimit() const
 	{
+		const FPlatformMemoryStats MemoryStats = FPlatformMemory::GetStats();
+
 		// Just make sure available physical fits in a int64, if it's bigger than that, we're not expecting to be memory limited anyway
-		int64 AvailablePhysical = (int64)FMath::Min(FPlatformMemory::GetStats().AvailablePhysical, (uint64)INT64_MAX);
+		// Also uses AvailableVirtual because the system might have plenty of physical memory but still be limited by virtual memory available in some cases.
+		//   (i.e. per-process quota, paging file size lower than actual memory available, etc.).
+		const int64 AvailableMemory = (int64)FMath::Min3(MemoryStats.AvailablePhysical, MemoryStats.AvailableVirtual, (uint64)INT64_MAX);
 
 		int64 HardMemoryLimit = GetHardMemoryLimit();
 		if (HardMemoryLimit > 0)
 		{
-			return FMath::Min(HardMemoryLimit, AvailablePhysical);
+			return FMath::Min(HardMemoryLimit, AvailableMemory);
 		}
 		
-		return AvailablePhysical;
+		return AvailableMemory;
 	}
 
 	int64 GetRequiredMemory(const IQueuedWork* InWork) const
