Static analysis error fixes.
- Added nullptr check on SourceProperty pointer usage in Kismet `SReplaceNodeReferences::FindAllReplacementsComplete`.
- Speculative fix for return of pointer to temporary local in `FVisualLogger::FindRedirection` method.
- Moved `FPurchaseFlowChildWindowImpl::Close` (virtual function) implementation to an internal private method, so to avoid calling virtual method from destructor.

diff --git a/Engine/Source/Editor/Kismet/Private/SReplaceNodeReferences.cpp b/Engine/Source/Editor/Kismet/Private/SReplaceNodeReferences.cpp
--- a/Engine/Source/Editor/Kismet/Private/SReplaceNodeReferences.cpp
+++ b/Engine/Source/Editor/Kismet/Private/SReplaceNodeReferences.cpp
@@ -690,7 +690,7 @@ void SReplaceNodeReferences::FindAllReplacementsComplete(TArray<FImaginaryFiBDat
 		}
 	}
 
-	if (SelectedTargetReferenceItem.IsValid())
+	if (HasValidSource() && SelectedTargetReferenceItem.IsValid())
 	{
 		FMemberReference SourceVariableReference;
 		SourceVariableReference.SetFromField<FProperty>(SourceProperty, SourceProperty->GetOwnerClass());
diff --git a/Engine/Source/Runtime/Engine/Private/VisualLogger/VisualLogger.cpp b/Engine/Source/Runtime/Engine/Private/VisualLogger/VisualLogger.cpp
--- a/Engine/Source/Runtime/Engine/Private/VisualLogger/VisualLogger.cpp
+++ b/Engine/Source/Runtime/Engine/Private/VisualLogger/VisualLogger.cpp
@@ -572,7 +572,8 @@ UObject* FVisualLogger::FindRedirection(const UObject* Object)
 		}
 	}
 
-	return const_cast<UObject*>(TargetWeakPtr.Get(/*bEvenIfPendingKill*/true));
+	const UObject* const Target = TargetWeakPtr.Get(/*bEvenIfPendingKill*/true);
+	return const_cast<UObject*>(Target);
 }
 
 void FVisualLogger::SetIsRecording(const bool InIsRecording)
