Crash opening anim blueprints with Pose Driver setup
#fix avoid invalid consecutive fix-up paths - SaveEditorOnlyFullPoseForPoseAsset assumes PoseAssetSupportPerBoneMask hasn't been ran as its reliant on LocalSpacePose and CurveData

diff --git a/Engine/Source/Runtime/Engine/Private/Animation/PoseAsset.cpp b/Engine/Source/Runtime/Engine/Private/Animation/PoseAsset.cpp
--- a/Engine/Source/Runtime/Engine/Private/Animation/PoseAsset.cpp
+++ b/Engine/Source/Runtime/Engine/Private/Animation/PoseAsset.cpp
@@ -864,7 +864,7 @@ void UPoseAsset::PostLoad()
 // moved to PostLoad because Skeleton is not completely loaded when we do this in Serialize
 // and we need Skeleton
 #if WITH_EDITOR
-	if (GetLinkerCustomVersion(FFrameworkObjectVersion::GUID) < FFrameworkObjectVersion::PoseAssetSupportPerBoneMask)
+	if (GetLinkerCustomVersion(FFrameworkObjectVersion::GUID) < FFrameworkObjectVersion::PoseAssetSupportPerBoneMask && GetLinkerCustomVersion(FAnimPhysObjectVersion::GUID) >= FAnimPhysObjectVersion::SaveEditorOnlyFullPoseForPoseAsset)
 	{
 		// fix curve names
 		// copy to source local data FIRST
