Fix Level Instance not updating when changes are saved outside of commit menu action (Editor SaveAll followed by Commit)

diff --git a/Engine/Source/Runtime/Engine/Private/LevelInstance/LevelInstanceActor.cpp b/Engine/Source/Runtime/Engine/Private/LevelInstance/LevelInstanceActor.cpp
--- a/Engine/Source/Runtime/Engine/Private/LevelInstance/LevelInstanceActor.cpp
+++ b/Engine/Source/Runtime/Engine/Private/LevelInstance/LevelInstanceActor.cpp
@@ -230,7 +230,7 @@ void ALevelInstance::OnEdit()
 	bEditLockLocation = true;
 }
 
-void ALevelInstance::OnCommit()
+void ALevelInstance::OnCommit(bool bChanged, bool bPromptForSave)
 {
 	bEditLockLocation = false;
 }
diff --git a/Engine/Source/Runtime/Engine/Private/LevelInstance/LevelInstanceSubsystem.cpp b/Engine/Source/Runtime/Engine/Private/LevelInstance/LevelInstanceSubsystem.cpp
--- a/Engine/Source/Runtime/Engine/Private/LevelInstance/LevelInstanceSubsystem.cpp
+++ b/Engine/Source/Runtime/Engine/Private/LevelInstance/LevelInstanceSubsystem.cpp
@@ -37,6 +37,8 @@
 #include "LevelInstance/Packed/PackedLevelInstanceBuilder.h"
 #include "Engine/LevelScriptBlueprint.h"
 #include "EdGraph/EdGraph.h"
+#include "UObject/ObjectSaveContext.h"
+#include "UObject/UObjectGlobals.h"
 #endif
 
 #include "HAL/IConsoleManager.h"
@@ -64,6 +66,12 @@ void ULevelInstanceSubsystem::Initialize(FSubsystemCollectionBase& Collection)
 	if (GEditor)
 	{
 		FModuleManager::LoadModuleChecked<ILevelInstanceEditorModule>("LevelInstanceEditor");
+
+		if (!GetWorld()->IsGameWorld())
+		{
+			FCoreUObjectDelegates::OnObjectPreSave.AddUObject(this, &ULevelInstanceSubsystem::OnObjectPreSave);
+			FEditorDelegates::OnPackageDeleted.AddUObject(this, &ULevelInstanceSubsystem::OnPackageDeleted);
+		}
 	}
 #endif
 }
@@ -71,6 +79,14 @@ void ULevelInstanceSubsystem::Initialize(FSubsystemCollectionBase& Collection)
 void ULevelInstanceSubsystem::Deinitialize()
 {
 	Super::Deinitialize();
+
+#if WITH_EDITOR
+	if (!GetWorld()->IsGameWorld())
+	{
+		FCoreUObjectDelegates::OnObjectPreSave.RemoveAll(this);
+		FEditorDelegates::OnPackageDeleted.RemoveAll(this);
+	}
+#endif
 }
 
 bool ULevelInstanceSubsystem::DoesSupportWorldType(EWorldType::Type WorldType) const
@@ -376,6 +392,43 @@ void ULevelInstanceSubsystem::Tick()
 	}
 }
 
+void ULevelInstanceSubsystem::OnPackageChanged(UPackage* Package)
+{
+	for (auto& Pair : LevelInstanceEdits)
+	{
+		FLevelInstanceEdit& LevelInstanceEdit = Pair.Value;
+		if (UWorld* EditWorld = LevelInstanceEdit.GetEditWorld())
+		{
+			if (EditWorld->GetPackage() == Package)
+			{
+				LevelInstanceEdit.bCommittedChanges = true;
+			}
+			else
+			{
+				TSet<UPackage*> Packages;
+				Packages.Append(EditWorld->GetPackage()->GetExternalPackages());
+				if (Packages.Contains(Package))
+				{
+					LevelInstanceEdit.bCommittedChanges = true;
+				}
+			}
+		}
+	}
+}
+
+void ULevelInstanceSubsystem::OnObjectPreSave(UObject* Object, FObjectPreSaveContext SaveContext)
+{
+	if (!SaveContext.IsProceduralSave() && !(SaveContext.GetSaveFlags() & SAVE_FromAutosave))
+	{
+		OnPackageChanged(Object->GetPackage());
+	}
+}
+
+void ULevelInstanceSubsystem::OnPackageDeleted(UPackage* Package)
+{
+	OnPackageChanged(Package);
+}
+
 bool ULevelInstanceSubsystem::CanPackLevelInstances() const
 {
 	return !LevelInstanceEdits.Num();
@@ -461,7 +514,7 @@ void ULevelInstanceSubsystem::PackLevelInstances()
 
 	for (APackedLevelInstance* PackedLevelInstance : PackedLevelInstancesToUpdate)
 	{
-		PackedLevelInstance->OnWorldAssetChanged();
+		PackedLevelInstance->UpdateLevelInstance();
 		UpdateProgress();
 	}
 }
@@ -1605,7 +1658,6 @@ ALevelInstance* ULevelInstanceSubsystem::CommitLevelInstance(ALevelInstance* Lev
 	UWorld* EditingWorld = LevelInstanceEdit->GetEditWorld();
 	check(EditingWorld);
 			
-	bool bChangesCommitted = false;
 	if (IsLevelInstanceEditDirty(LevelInstanceEdit) && !bDiscardEdits)
 	{
 		const bool bPromptUserToSave = bPromptForSave;
@@ -1627,21 +1679,6 @@ ALevelInstance* ULevelInstanceSubsystem::CommitLevelInstance(ALevelInstance* Lev
 		{
 			return LevelInstanceActor;
 		}
-
-		// Validate that we indeed need to refresh instances (user can cancel the changes when prompted)
-		bChangesCommitted = !IsLevelInstanceEditDirty(LevelInstanceEdit);
-
-		if(bChangesCommitted)
-		{
-			// Sync the AssetData so that the updated instances have the latest Actor Registry Data
-			IAssetRegistry& AssetRegistry = FModuleManager::LoadModuleChecked<FAssetRegistryModule>(TEXT("AssetRegistry")).Get();
-			AssetRegistry.ScanPathsSynchronous({ LevelInstanceActor->GetWorldAssetPackage() }, true);
-			// Notify 
-			LevelInstanceActor->OnWorldAssetSaved(bPromptUserToSave);
-
-			// Update pointer since BP Compilation might have invalidated LevelInstanceActor
-			LevelInstanceActor = LevelInstanceEdit->LevelStreaming->GetLevelInstanceActor();
-		}
 	}
 
 	FScopedSlowTask SlowTask(0, LOCTEXT("EndEditLevelInstance", "Unloading edit Level Instance..."), !GetWorld()->IsGameWorld());
@@ -1673,11 +1710,25 @@ ALevelInstance* ULevelInstanceSubsystem::CommitLevelInstance(ALevelInstance* Lev
 	const FString EditPackage = LevelInstanceActor->GetWorldAssetPackage();
 
 	// Remove from streaming level...
+	const bool bChangesCommitted = LevelInstanceEdit->bCommittedChanges;
 	ULevelStreamingLevelInstanceEditor::Unload(LevelInstanceEdit->LevelStreaming);
 	LevelInstanceEdits.Remove(FName(*EditPackage));
 
-	// Notify (Actor might get destroyed by this call if its a packed bp
-	LevelInstanceActor->OnCommit();
+	if (bChangesCommitted)
+	{
+		// Sync the AssetData so that the updated instances have the latest Actor Registry Data
+		IAssetRegistry& AssetRegistry = FModuleManager::LoadModuleChecked<FAssetRegistryModule>(TEXT("AssetRegistry")).Get();
+		AssetRegistry.ScanPathsSynchronous({ LevelInstanceActor->GetWorldAssetPackage() }, true);
+	}
+	
+	// Backup ID on Commit in case Actor gets recreated
+	const FLevelInstanceID LevelInstanceID = LevelInstanceActor->GetLevelInstanceID();
+
+	// Notify (Actor might get destroyed by this call if its a packed bp)
+	LevelInstanceActor->OnCommit(bChangesCommitted, bPromptForSave);
+
+	// Update pointer since BP Compilation might have invalidated LevelInstanceActor
+	LevelInstanceActor = GetLevelInstance(LevelInstanceID);
 
 	// Notify Ancestors
 	ForEachLevelInstanceAncestors(LevelInstanceActor, [bChangesCommitted](ALevelInstance* AncestorLevelInstance)
diff --git a/Engine/Source/Runtime/Engine/Private/LevelInstance/Packed/PackedLevelInstanceActor.cpp b/Engine/Source/Runtime/Engine/Private/LevelInstance/Packed/PackedLevelInstanceActor.cpp
--- a/Engine/Source/Runtime/Engine/Private/LevelInstance/Packed/PackedLevelInstanceActor.cpp
+++ b/Engine/Source/Runtime/Engine/Private/LevelInstance/Packed/PackedLevelInstanceActor.cpp
@@ -92,32 +92,19 @@ FName APackedLevelInstance::GetPackedComponentTag()
 	return PackedComponentTag;
 }
 
-void APackedLevelInstance::OnWorldAssetChanged()
+void APackedLevelInstance::UpdateLevelInstance()
 {
-	if (IsLevelInstancePathValid())
+	if (!Cast<UBlueprint>(GetClass()->ClassGeneratedBy))
 	{
-		TSharedPtr<FPackedLevelInstanceBuilder> Builder = FPackedLevelInstanceBuilder::CreateDefaultBuilder();
-		Builder->PackActor(this, GetWorldAsset());
-	}
-	else
-	{
-		DestroyPackedComponents();
-	}
-}
-
-void APackedLevelInstance::OnWorldAssetSaved(bool bPromptForSave)
-{
-	TSharedPtr<FPackedLevelInstanceBuilder> Builder = FPackedLevelInstanceBuilder::CreateDefaultBuilder();
-	
-	if (UBlueprint* GeneratedBy = Cast<UBlueprint>(GetClass()->ClassGeneratedBy))
-	{
-		check(GeneratedBy == BlueprintAsset.Get());
-		const bool bCheckoutAndSave = true;
-		Builder->UpdateBlueprint(GeneratedBy, bCheckoutAndSave, bPromptForSave);
-	}
-	else
-	{
-		Builder->PackActor(this);
+		if (IsLevelInstancePathValid())
+		{
+			TSharedPtr<FPackedLevelInstanceBuilder> Builder = FPackedLevelInstanceBuilder::CreateDefaultBuilder();
+			Builder->PackActor(this, GetWorldAsset());
+		}
+		else
+		{
+			DestroyPackedComponents();
+		}
 	}
 }
 
@@ -168,11 +155,24 @@ void APackedLevelInstance::OnEdit()
 	MarkComponentsRenderStateDirty();
 }
 
-void APackedLevelInstance::OnCommit()
+void APackedLevelInstance::OnCommit(bool bChanged, bool bPromptForSave)
 {
-	Super::OnCommit();
+	Super::OnCommit(bChanged, bPromptForSave);
 	check(bEditing);
 	bEditing = false;
+
+	if (bChanged)
+	{
+		if (UBlueprint* GeneratedBy = Cast<UBlueprint>(GetClass()->ClassGeneratedBy))
+		{
+			check(GeneratedBy == BlueprintAsset.Get());
+			const bool bCheckoutAndSave = true;
+			TSharedPtr<FPackedLevelInstanceBuilder> Builder = FPackedLevelInstanceBuilder::CreateDefaultBuilder();
+			Builder->UpdateBlueprint(GeneratedBy, bCheckoutAndSave, bPromptForSave);
+		}
+	}
+	
+	// bEditing flag changed so dirty render state
 	MarkComponentsRenderStateDirty();
 }
 
diff --git a/Engine/Source/Runtime/Engine/Public/LevelInstance/LevelInstanceActor.h b/Engine/Source/Runtime/Engine/Public/LevelInstance/LevelInstanceActor.h
--- a/Engine/Source/Runtime/Engine/Public/LevelInstance/LevelInstanceActor.h
+++ b/Engine/Source/Runtime/Engine/Public/LevelInstance/LevelInstanceActor.h
@@ -102,17 +102,16 @@ public:
 	bool IsLoaded() const;
 	void OnLevelInstanceLoaded();
 	FString GetWorldAssetPackage() const;
-	void UpdateLevelInstance();
 	bool SetWorldAsset(TSoftObjectPtr<UWorld> WorldAsset);
 	bool CheckForLoop(TSoftObjectPtr<UWorld> WorldAsset, TArray<TPair<FText, TSoftObjectPtr<UWorld>>>* LoopInfo = nullptr, const ALevelInstance** LoopStart = nullptr) const;
 	AActor* FindEditorInstanceActor() const;
 
 	virtual void OnWorldAssetChanged() { UpdateLevelInstance(); }
-	virtual void OnWorldAssetSaved(bool bPromptForSave) {}
 	virtual void OnEdit();
 	virtual void OnEditChild() {}
-	virtual void OnCommit();
+	virtual void OnCommit(bool bChanged, bool bPromptForSave);
 	virtual void OnCommitChild(bool bChanged) {}
+	virtual void UpdateLevelInstance();
 		
 	DECLARE_MULTICAST_DELEGATE_OneParam(FOnLevelInstanceActorPostLoad, ALevelInstance*);
 	static FOnLevelInstanceActorPostLoad OnLevelInstanceActorPostLoad;
diff --git a/Engine/Source/Runtime/Engine/Public/LevelInstance/LevelInstanceSubsystem.h b/Engine/Source/Runtime/Engine/Public/LevelInstance/LevelInstanceSubsystem.h
--- a/Engine/Source/Runtime/Engine/Public/LevelInstance/LevelInstanceSubsystem.h
+++ b/Engine/Source/Runtime/Engine/Public/LevelInstance/LevelInstanceSubsystem.h
@@ -47,6 +47,8 @@ public:
 
 #if WITH_EDITOR
 	void Tick();
+	void OnObjectPreSave(UObject* Object, FObjectPreSaveContext SaveContext);
+	void OnPackageDeleted(UPackage* Package);
 	void PackLevelInstances();
 	bool CanPackLevelInstances() const;
 
@@ -106,10 +108,12 @@ private:
 	void BreakLevelInstance_Impl(ALevelInstance* LevelInstanceActor, uint32 Levels, TArray<AActor*>& OutMovedActors);
 
 	static bool ShouldIgnoreDirtyPackage(UPackage* DirtyPackage, const UWorld* EditingWorld);
+	void OnPackageChanged(UPackage* Package);
 
 	struct FLevelInstanceEdit
 	{
 		ULevelStreamingLevelInstanceEditor* LevelStreaming = nullptr;
+		bool bCommittedChanges = false;
 
 		UWorld* GetEditWorld() const;
 	};
diff --git a/Engine/Source/Runtime/Engine/Public/LevelInstance/Packed/PackedLevelInstanceActor.h b/Engine/Source/Runtime/Engine/Public/LevelInstance/Packed/PackedLevelInstanceActor.h
--- a/Engine/Source/Runtime/Engine/Public/LevelInstance/Packed/PackedLevelInstanceActor.h
+++ b/Engine/Source/Runtime/Engine/Public/LevelInstance/Packed/PackedLevelInstanceActor.h
@@ -37,9 +37,8 @@ public:
 	static FName GetPackedComponentTag();
 
 	virtual void PostLoad() override;
-	virtual void OnWorldAssetChanged() override;
-	virtual void OnWorldAssetSaved(bool bPromptForSave) override;
-	virtual void OnCommit() override;
+	virtual void UpdateLevelInstance() override;
+	virtual void OnCommit(bool bChanged, bool bPromptForSave) override;
 	virtual void OnCommitChild(bool bChanged) override;
 	virtual void OnEdit() override;
 	virtual void OnEditChild() override;
