Addition of OnNavigationInitDone delegate in UNavigationSystemBase

diff --git a/Engine/Source/Runtime/Engine/Classes/AI/NavigationSystemBase.h b/Engine/Source/Runtime/Engine/Classes/AI/NavigationSystemBase.h
--- a/Engine/Source/Runtime/Engine/Classes/AI/NavigationSystemBase.h
+++ b/Engine/Source/Runtime/Engine/Classes/AI/NavigationSystemBase.h
@@ -202,6 +202,7 @@ namespace FNavigationSystem
 	DECLARE_DELEGATE_RetVal_OneParam(INavigationDataInterface*, FNavDataForActorSignature, const AActor& /*Actor*/);
 	DECLARE_DELEGATE_RetVal(TSubclassOf<AActor>, FNavDataClassFetchSignature);
 	DECLARE_DELEGATE_TwoParams(FWorldBoolBasedSignature, UWorld& /*World*/, const bool /*bShow*/);
+	DECLARE_MULTICAST_DELEGATE_OneParam(FOnNavigationInitDoneSignature, const UNavigationSystemBase&);
 }
 
 
@@ -251,6 +252,10 @@ public:
 	virtual void RemoveNavigationDataChunk(class ANavigationDataChunkActor& DataChunkActor) {}
 	virtual void FillNavigationDataChunkActor(const FBox& QueryBounds, class ANavigationDataChunkActor& DataChunkActor, FBox& OutTilesBounds) {}
 
+	virtual bool IsWorldInitDone() const PURE_VIRTUAL(UNavigationSystemBase::IsWorldInitDone, return false;);
+
+	static FNavigationSystem::FOnNavigationInitDoneSignature& OnNavigationInitDoneStaticDelegate();
+
 protected:
 	/**	Sets the Transform the Navigation System will use when converting from FromCoordType
 	 *	to ToCoordType
diff --git a/Engine/Source/Runtime/Engine/Private/AI/NavigationSystemBase.cpp b/Engine/Source/Runtime/Engine/Private/AI/NavigationSystemBase.cpp
--- a/Engine/Source/Runtime/Engine/Private/AI/NavigationSystemBase.cpp
+++ b/Engine/Source/Runtime/Engine/Private/AI/NavigationSystemBase.cpp
@@ -131,11 +131,11 @@ namespace FNavigationSystem
 		FNavDatConfigBasedSignature GetDefaultSupportedAgent;
 		FActorBooleBasedSignature UpdateActorAndComponentData;
 		FComponentBoundsChangeSignature OnComponentBoundsChanged;
-		//FNavDataForPropsSignature GetNavDataForProps;
 		FNavDataForActorSignature GetNavDataForActor;
 		FNavDataClassFetchSignature GetDefaultNavDataClass;
 		FWorldBoolBasedSignature VerifyNavigationRenderingComponents;
 		FWorldBasedSignature Build;
+		FOnNavigationInitDoneSignature OnNavigationInitDone;
 #if WITH_EDITOR
 		FWorldBasedSignature OnPIEStart;
 		FWorldBasedSignature OnPIEEnd;
@@ -164,7 +164,6 @@ namespace FNavigationSystem
 			GetDefaultSupportedAgent.BindStatic(&GetFallbackSupportedAgent);
 			UpdateActorAndComponentData.BindLambda([](AActor&, bool) {});
 			OnComponentBoundsChanged.BindLambda([](UActorComponent&, const FBox&, const FBox&) {});
-			//GetNavDataForProps.BindLambda([](const FNavAgentProperties&) { return nullptr; });
 			GetNavDataForActor.BindLambda([](const AActor&) { return nullptr; });
 			GetDefaultNavDataClass.BindLambda([]() { return AActor::StaticClass(); });
 			VerifyNavigationRenderingComponents.BindLambda([](UWorld&, bool) {});
@@ -398,6 +397,7 @@ FNavigationSystem::FNavDataForActorSignature& UNavigationSystemBase::GetNavDataF
 FNavigationSystem::FNavDataClassFetchSignature& UNavigationSystemBase::GetDefaultNavDataClassDelegate() { return FNavigationSystem::Delegates.GetDefaultNavDataClass; }
 FNavigationSystem::FWorldBoolBasedSignature& UNavigationSystemBase::VerifyNavigationRenderingComponentsDelegate() { return FNavigationSystem::Delegates.VerifyNavigationRenderingComponents; }
 FNavigationSystem::FWorldBasedSignature& UNavigationSystemBase::BuildDelegate() { return FNavigationSystem::Delegates.Build; }
+FNavigationSystem::FOnNavigationInitDoneSignature& UNavigationSystemBase::OnNavigationInitDoneStaticDelegate() { return FNavigationSystem::Delegates.OnNavigationInitDone; }
 #if WITH_EDITOR
 FNavigationSystem::FWorldBasedSignature& UNavigationSystemBase::OnPIEStartDelegate() { return FNavigationSystem::Delegates.OnPIEStart; }
 FNavigationSystem::FWorldBasedSignature& UNavigationSystemBase::OnPIEEndDelegate() { return FNavigationSystem::Delegates.OnPIEEnd; }
diff --git a/Engine/Source/Runtime/NavigationSystem/Private/NavigationSystem.cpp b/Engine/Source/Runtime/NavigationSystem/Private/NavigationSystem.cpp
--- a/Engine/Source/Runtime/NavigationSystem/Private/NavigationSystem.cpp
+++ b/Engine/Source/Runtime/NavigationSystem/Private/NavigationSystem.cpp
@@ -1025,6 +1025,7 @@ void UNavigationSystemV1::OnWorldInitDone(FNavigationSystemRunMode Mode)
 
 	bWorldInitDone = true;
 	OnNavigationInitDone.Broadcast();
+	UNavigationSystemBase::OnNavigationInitDoneStaticDelegate().Broadcast(*this);
 }
 
 void UNavigationSystemV1::RegisterNavigationDataInstances()
diff --git a/Engine/Source/Runtime/NavigationSystem/Public/NavigationSystem.h b/Engine/Source/Runtime/NavigationSystem/Public/NavigationSystem.h
--- a/Engine/Source/Runtime/NavigationSystem/Public/NavigationSystem.h
+++ b/Engine/Source/Runtime/NavigationSystem/Public/NavigationSystem.h
@@ -909,7 +909,10 @@ public:
 	/** */
 	virtual void OnWorldInitDone(FNavigationSystemRunMode Mode);
 
+	/** Returns true if world has been initialized. */
+	FORCEINLINE virtual bool IsWorldInitDone() const override { return bWorldInitDone; }
 	FORCEINLINE bool IsInitialized() const { return bWorldInitDone; }
+
 	FORCEINLINE FNavigationSystemRunMode GetRunMode() const { return OperationMode; }
 
 	/** adds BSP collisions of currently streamed in levels to octree */
