Adds ensure since WorldDataLayer actor should always be present.

#changelist validated

diff --git a/Engine/Source/Runtime/Engine/Private/WorldPartition/WorldPartitionRuntimeHash.cpp b/Engine/Source/Runtime/Engine/Private/WorldPartition/WorldPartitionRuntimeHash.cpp
--- a/Engine/Source/Runtime/Engine/Private/WorldPartition/WorldPartitionRuntimeHash.cpp
+++ b/Engine/Source/Runtime/Engine/Private/WorldPartition/WorldPartitionRuntimeHash.cpp
@@ -120,7 +120,7 @@ void UWorldPartitionRuntimeHash::CheckForErrorsInternal(const TMap<FGuid, FWorld
 							->AddToken(FMapErrorToken::Create(FName(TEXT("WorldPartition_StreamedActorReferenceAlwaysLoadedActor_CheckForErrors"))));
 					}
 
-					if (WorldDataLayers != nullptr)
+					if (ensure(WorldDataLayers != nullptr))
 					{
 						TArray<FName> ActorDescLayerNames = ActorDescView.GetDataLayers();
 						ActorDescLayerNames.Sort([](const FName& A, const FName& B) { return A.FastLess(B); });
