Add more info to a check() that's been reported to fail

diff --git a/Engine/Source/Runtime/Engine/Public/Shader/ShaderTypes.h b/Engine/Source/Runtime/Engine/Public/Shader/ShaderTypes.h
--- a/Engine/Source/Runtime/Engine/Public/Shader/ShaderTypes.h
+++ b/Engine/Source/Runtime/Engine/Public/Shader/ShaderTypes.h
@@ -237,7 +237,11 @@ struct FValue
 	}
 
 	inline EValueType GetType() const { return MakeValueType(ComponentType, NumComponents); }
-	inline const FValueComponent& GetComponent(int32 i) const { check(i >= 0 && i < NumComponents); return Component[i]; }
+	inline const FValueComponent& GetComponent(int32 i) const
+	{
+		checkf(i >= 0 && i < NumComponents, TEXT("Invalid component %d/%d, of type '%s'"), i, NumComponents, GetValueTypeDescription(GetType()).Name);
+		return Component[i];
+	}
 
 	static FValue FromMemoryImage(EValueType Type, const void* Data, uint32* OutSizeInBytes = nullptr);
 	FMemoryImageValue AsMemoryImage() const;
