LevelInstance: Avoid prompting user when creating from an unsaved level

diff --git a/Engine/Source/Runtime/Engine/Private/LevelInstance/LevelInstanceSubsystem.cpp b/Engine/Source/Runtime/Engine/Private/LevelInstance/LevelInstanceSubsystem.cpp
--- a/Engine/Source/Runtime/Engine/Private/LevelInstance/LevelInstanceSubsystem.cpp
+++ b/Engine/Source/Runtime/Engine/Private/LevelInstance/LevelInstanceSubsystem.cpp
@@ -20,6 +20,7 @@
 #include "Misc/ScopedSlowTask.h"
 #include "Misc/ITransaction.h"
 #include "Misc/Paths.h"
+#include "Misc/PackageName.h"
 #include "AssetRegistryModule.h"
 #include "AssetData.h"
 #include "FileHelpers.h"
@@ -837,7 +838,11 @@ ALevelInstance* ULevelInstanceSubsystem::CreateLevelInstanceFrom(const TArray<AA
 	// Capture Packages before Moving actors as they can get GCed in the process
 	for (AActor* ActorToMove : ActorsToMove)
 	{
-		DirtyPackages.Add(ActorToMove->GetPackage()->GetFName());
+		// Don't force saving of unsaved/temp packages onto the user.
+		if (!FPackageName::IsTempPackage(ActorToMove->GetPackage()->GetName()))
+		{
+			DirtyPackages.Add(ActorToMove->GetPackage()->GetFName());
+		}
 	}
 
 	if (!EditorLevelUtils::MoveActorsToLevel(ActorsToMove, LoadedLevel, bWarnAboutReferences, bWarnAboutRenaming, bMoveAllOrFail))
@@ -925,7 +930,12 @@ ALevelInstance* ULevelInstanceSubsystem::CreateLevelInstanceFrom(const TArray<AA
 	LevelStreaming->LevelInstanceID = NewLevelInstanceActor->GetLevelInstanceID();
 		
 	GetWorld()->SetCurrentLevel(LoadedLevel);
-	DirtyPackages.Add(NewLevelInstanceActor->GetPackage()->GetFName());
+
+	// Don't force saving of unsaved/temp packages onto the user.
+	if (!FPackageName::IsTempPackage(NewLevelInstanceActor->GetPackage()->GetName()))
+	{
+		DirtyPackages.Add(NewLevelInstanceActor->GetPackage()->GetFName());
+	}
 	// Use LevelStreaming->GetLevelInstanceActor() because OnWorldAssetSaved could've reinstanced the LevelInstanceActor
 	return CommitLevelInstance(LevelStreaming->GetLevelInstanceActor(), /*bDiscardEdits*/false, CreationParams.bPromptForSave, &DirtyPackages);
 }
