Allow conduits to be used as entry states in animation state machines

diff --git a/Engine/Source/Editor/AnimGraph/Private/AnimGraphNode_StateMachineBase.cpp b/Engine/Source/Editor/AnimGraph/Private/AnimGraphNode_StateMachineBase.cpp
--- a/Engine/Source/Editor/AnimGraph/Private/AnimGraphNode_StateMachineBase.cpp
+++ b/Engine/Source/Editor/AnimGraph/Private/AnimGraphNode_StateMachineBase.cpp
@@ -375,10 +375,12 @@ void UAnimGraphNode_StateMachineBase::OnProcessDuringCompilation(IAnimBlueprintC
 			else
 			{
 				// Make sure the entry node is a state and not a conduit
-				if (BakedMachine.States[BakedMachine.InitialState].bIsAConduit)
+				if (BakedMachine.States[BakedMachine.InitialState].bIsAConduit && !StateMachineInstance->GetNode().bAllowConduitEntryStates)
 				{
 					UEdGraphNode* StateNode = GetMachineSpecificDebugData().FindNodeFromStateIndex(BakedMachine.InitialState);
-					CompilationContext.GetMessageLog().Error(*LOCTEXT("BadStateEntryNode", "A conduit (@@) cannot be used as the entry node for a state machine").ToString(), StateNode);
+					CompilationContext.GetMessageLog().Error(*LOCTEXT("BadStateEntryNode", 
+					"A conduit (@@) cannot be used as the entry node for a state machine. To enable this, check the 'Allow conduit entry states' checkbox for StateMachine. Warning, if a valid entry state cannot be found at runtime then this will generate a reference pose!"
+					).ToString(), StateNode);
 				}
 			}
 		}
diff --git a/Engine/Source/Runtime/Engine/Classes/Animation/AnimNode_StateMachine.h b/Engine/Source/Runtime/Engine/Classes/Animation/AnimNode_StateMachine.h
--- a/Engine/Source/Runtime/Engine/Classes/Animation/AnimNode_StateMachine.h
+++ b/Engine/Source/Runtime/Engine/Classes/Animation/AnimNode_StateMachine.h
@@ -145,6 +145,11 @@ public:
 	// data has a  slight performance cost.
 	UPROPERTY(EditAnywhere, Category = Settings)
 	bool bCreateNotifyMetaData;
+
+	// Allows a conduit to be used as this state machine's entry state
+	// If a valid entry state cannot be found at runtime then this will generate a reference pose!
+	UPROPERTY(EditAnywhere, Category = Settings)
+	bool bAllowConduitEntryStates;
 private:
 	// true if it is the first update.
 	bool bFirstUpdate;
@@ -207,6 +212,7 @@ public:
 		, bSkipFirstUpdateTransition(true)
 		, bReinitializeOnBecomingRelevant(true)
 		, bCreateNotifyMetaData(true)
+		, bAllowConduitEntryStates(false)
 		, bFirstUpdate(true)
 		, CurrentState(INDEX_NONE)
 		, ElapsedTime(0.0f)
diff --git a/Engine/Source/Runtime/Engine/Private/Animation/AnimNode_StateMachine.cpp b/Engine/Source/Runtime/Engine/Private/Animation/AnimNode_StateMachine.cpp
--- a/Engine/Source/Runtime/Engine/Private/Animation/AnimNode_StateMachine.cpp
+++ b/Engine/Source/Runtime/Engine/Private/Animation/AnimNode_StateMachine.cpp
@@ -791,7 +791,22 @@ void FAnimNode_StateMachine::Evaluate_AnyThread(FPoseContext& Output)
 {
 	if (const FBakedAnimationStateMachine* Machine = GetMachineDescription())
 	{
-		if (Machine->States.Num() == 0 || !Machine->States.IsValidIndex(CurrentState))
+		const bool bCurrentStateIsConduit = IsAConduitState(CurrentState);
+#if WITH_EDITORONLY_DATA
+		if (bCurrentStateIsConduit)
+		{ 
+			UAnimBlueprint* AnimBlueprint = Output.AnimInstanceProxy->GetAnimBlueprint();
+
+			FText Message = FText::Format(LOCTEXT("CurrentStateIsConduitWarning", 
+			"Current state is a conduit and will reset to ref pose. This can happen if a conduit is an entry state and there isn't at least one valid transition to take. Conduit({0}), State Machine({1}), AnimBlueprint({2})."),
+				FText::FromName(Machine->MachineName), FText::FromName(GetCurrentStateName()), FText::FromString(GetPathNameSafe(AnimBlueprint)));
+			Output.LogMessage(EMessageSeverity::Error, Message);
+		}
+#else
+		ensure(!bCurrentStateIsConduit);
+#endif
+
+		if (bCurrentStateIsConduit || Machine->States.Num() == 0 || !Machine->States.IsValidIndex(CurrentState))
 		{
 			Output.Pose.ResetToRefPose();
 			return;
@@ -978,7 +993,7 @@ void FAnimNode_StateMachine::GatherDebugData(FNodeDebugData& DebugData)
 void FAnimNode_StateMachine::SetStateInternal(int32 NewStateIndex)
 {
 	checkSlow(PRIVATE_MachineDescription);
-	ensure(!IsAConduitState(NewStateIndex));
+	ensure(!IsAConduitState(NewStateIndex) || bAllowConduitEntryStates);
 	CurrentState = FMath::Clamp<int32>(NewStateIndex, 0, PRIVATE_MachineDescription->States.Num() - 1);
 	check(CurrentState == NewStateIndex);
 	ElapsedTime = 0.0f;
