Fix for edge case in state machines where multiple states transition in/out of a conduit with both in/out transitions being true

diff --git a/Engine/Source/Runtime/Engine/Private/Animation/AnimNode_StateMachine.cpp b/Engine/Source/Runtime/Engine/Private/Animation/AnimNode_StateMachine.cpp
--- a/Engine/Source/Runtime/Engine/Private/Animation/AnimNode_StateMachine.cpp
+++ b/Engine/Source/Runtime/Engine/Private/Animation/AnimNode_StateMachine.cpp
@@ -418,7 +418,8 @@ void FAnimNode_StateMachine::Update_AnyThread(const FAnimationUpdateContext& Con
 		}
 				
 		// If transition is valid and not waiting on other conditions
-		if (PotentialTransition.IsValid())
+		// and we're not doing a transition to self
+		if (PotentialTransition.IsValid() && PotentialTransition.TargetState != CurrentState)
 		{
 			bFoundValidTransition = true;
 
@@ -663,13 +664,6 @@ bool FAnimNode_StateMachine::FindValidTransition(const FAnimationUpdateContext&
 			continue;
 		}
 
-		const int32 NextState = GetTransitionInfo(TransitionRule.TransitionIndex).NextState;
-		// Skip this transition if we've already visited the destination state
-		if (OutVisitedStateIndices.Contains(NextState))
-		{
-			continue;
-		}
-
 		FAnimNode_TransitionResult* ResultNode = GetNodeFromPropertyIndex<FAnimNode_TransitionResult>(Context.AnimInstanceProxy->GetAnimInstanceObject(), AnimBlueprintClass, TransitionRule.CanTakeDelegateIndex);
 		float CrossfadeTimeAdjustment = 0.f;
 
@@ -707,6 +701,7 @@ bool FAnimNode_StateMachine::FindValidTransition(const FAnimationUpdateContext&
 
 		if (ResultNode->bCanEnterTransition == TransitionRule.bDesiredTransitionReturnValue)
 		{
+			const int32 NextState = GetTransitionInfo(TransitionRule.TransitionIndex).NextState;
 			const FBakedAnimationState& NextStateInfo = GetStateInfo(NextState);
 
 			// if next state is a conduit we want to check for transitions using that state as the root
