[ODSC] - Break out recompile shaders functionality into its own file.

- Break out common functionality for both ShaderCompiler.cpp and RecompileShaders.cpp into it's own shared file.

diff --git a/Engine/Source/Editor/UnrealEd/Private/Commandlets/CookGlobalShaderCommandlet.cpp b/Engine/Source/Editor/UnrealEd/Private/Commandlets/CookGlobalShaderCommandlet.cpp
--- a/Engine/Source/Editor/UnrealEd/Private/Commandlets/CookGlobalShaderCommandlet.cpp
+++ b/Engine/Source/Editor/UnrealEd/Private/Commandlets/CookGlobalShaderCommandlet.cpp
@@ -8,7 +8,7 @@
 #include "GenericPlatform/GenericPlatformFile.h"
 #include "UObject/UObjectIterator.h"
 #include "UObject/Class.h"
-#include "ShaderCompiler.h"
+#include "RecompileShaders.h"
 #include "CookOnTheSide/CookOnTheFlyServer.h"
 #include "Interfaces/ITargetPlatformManagerModule.h"
 #include "Interfaces/ITargetPlatform.h"
diff --git a/Engine/Source/Editor/UnrealEd/Private/CookOnTheFlyServer.cpp b/Engine/Source/Editor/UnrealEd/Private/CookOnTheFlyServer.cpp
--- a/Engine/Source/Editor/UnrealEd/Private/CookOnTheFlyServer.cpp
+++ b/Engine/Source/Editor/UnrealEd/Private/CookOnTheFlyServer.cpp
@@ -93,6 +93,7 @@
 #include "ProfilingDebugging/CookStats.h"
 #include "ProfilingDebugging/PlatformFileTrace.h"
 #include "ProjectDescriptor.h"
+#include "RecompileShaders.h"
 #include "SceneUtils.h"
 #include "Serialization/ArchiveStackTrace.h"
 #include "Serialization/ArchiveUObject.h"
diff --git a/Engine/Source/Editor/UnrealEd/Private/Cooker/CookOnTheFlyServerInterface.h b/Engine/Source/Editor/UnrealEd/Private/Cooker/CookOnTheFlyServerInterface.h
--- a/Engine/Source/Editor/UnrealEd/Private/Cooker/CookOnTheFlyServerInterface.h
+++ b/Engine/Source/Editor/UnrealEd/Private/Cooker/CookOnTheFlyServerInterface.h
@@ -4,7 +4,7 @@
 
 #include "CoreTypes.h"
 #include "CookTypes.h"
-#include "ShaderCompiler.h"
+#include "RecompileShaders.h"
 
 struct FODSCRequestPayload;
 class ICookedPackageWriter;
diff --git a/Engine/Source/Runtime/Engine/Private/ODSC/ODSCThread.cpp b/Engine/Source/Runtime/Engine/Private/ODSC/ODSCThread.cpp
--- a/Engine/Source/Runtime/Engine/Private/ODSC/ODSCThread.cpp
+++ b/Engine/Source/Runtime/Engine/Private/ODSC/ODSCThread.cpp
@@ -6,7 +6,7 @@
 #include "HAL/FileManager.h"
 #include "HAL/PlatformProcess.h"
 #include "HAL/RunnableThread.h"
-#include "ShaderCompiler.h"
+#include "RecompileShaders.h"
 #include "CookOnTheFly.h"
 
 FODSCRequestPayload::FODSCRequestPayload(EShaderPlatform InShaderPlatform, const FString& InMaterialName, const FString& InVertexFactoryName, const FString& InPipelineName, const TArray<FString>& InShaderTypeNames, const FString& InRequestHash)
diff --git a/Engine/Source/Runtime/Engine/Private/ODSC/ODSCThread.h b/Engine/Source/Runtime/Engine/Private/ODSC/ODSCThread.h
--- a/Engine/Source/Runtime/Engine/Private/ODSC/ODSCThread.h
+++ b/Engine/Source/Runtime/Engine/Private/ODSC/ODSCThread.h
@@ -7,7 +7,7 @@
 #include "Misc/SingleThreadRunnable.h"
 #include "GenericPlatform/GenericPlatformFile.h"
 #include "Containers/Queue.h"
-#include "ShaderCompiler.h"
+#include "RecompileShaders.h"
 #include "RHIDefinitions.h"
 
 class FEvent;
diff --git a/Engine/Source/Runtime/Engine/Private/ShaderCompiler/RecompileShaders.cpp b/Engine/Source/Runtime/Engine/Private/ShaderCompiler/RecompileShaders.cpp
new file mode 100644
--- /dev/null
+++ b/Engine/Source/Runtime/Engine/Private/ShaderCompiler/RecompileShaders.cpp
@@ -0,0 +1,711 @@
+// Copyright Epic Games, Inc. All Rights Reserved.
+
+#include "RecompileShaders.h"
+
+#include "GlobalShader.h"
+#include "Interfaces/ITargetPlatform.h"
+#include "Interfaces/ITargetPlatformManagerModule.h"
+#include "Materials/Material.h"
+#include "Materials/MaterialInterface.h"
+#include "MaterialShared.h"
+#include "Misc/FeedbackContext.h"
+#if WITH_ODSC
+#include "ODSC/ODSCManager.h"
+#endif
+#include "Serialization/NameAsStringProxyArchive.h"
+#include "ShaderCompiler.h"
+#include "ShaderCompilerShared.h"
+#include "UObject/UObjectIterator.h"
+
+/** Timer class used to report information on the 'recompileshaders' console command. */
+class FRecompileShadersTimer
+{
+public:
+	FRecompileShadersTimer(const TCHAR* InInfoStr = TEXT("Test")) :
+		InfoStr(InInfoStr),
+		bAlreadyStopped(false)
+	{
+		StartTime = FPlatformTime::Seconds();
+	}
+
+	FRecompileShadersTimer(const FString& InInfoStr) :
+		InfoStr(InInfoStr),
+		bAlreadyStopped(false)
+	{
+		StartTime = FPlatformTime::Seconds();
+	}
+
+	void Stop(bool DisplayLog = true)
+	{
+		if (!bAlreadyStopped)
+		{
+			bAlreadyStopped = true;
+			EndTime = FPlatformTime::Seconds();
+			TimeElapsed = EndTime - StartTime;
+			if (DisplayLog)
+			{
+				UE_LOG(LogShaderCompilers, Warning, TEXT("		[%s] took [%.4f] s"), *InfoStr, TimeElapsed);
+			}
+		}
+	}
+
+	~FRecompileShadersTimer()
+	{
+		Stop(true);
+	}
+
+protected:
+	double StartTime, EndTime;
+	double TimeElapsed;
+	FString InfoStr;
+	bool bAlreadyStopped;
+};
+
+namespace
+{
+	ODSCRecompileCommand ParseRecompileCommandString(const TCHAR* CmdString, TArray<FString>& OutMaterialsToLoad)
+	{
+		FString CmdName = FParse::Token(CmdString, 0);
+
+		ODSCRecompileCommand CommandType = ODSCRecompileCommand::None;
+		OutMaterialsToLoad.Empty();
+
+		if (!CmdName.IsEmpty() && FCString::Stricmp(*CmdName, TEXT("Material")) == 0)
+		{
+			CommandType = ODSCRecompileCommand::Material;
+
+			// tell other side the material to load, by pathname
+			FString RequestedMaterialName(FParse::Token(CmdString, 0));
+
+			for (TObjectIterator<UMaterialInterface> It; It; ++It)
+			{
+				UMaterial* Material = It->GetMaterial();
+
+				if (Material && Material->GetName() == RequestedMaterialName)
+				{
+					OutMaterialsToLoad.Add(It->GetPathName());
+					break;
+				}
+			}
+		}
+		else if (!CmdName.IsEmpty() && FCString::Stricmp(*CmdName, TEXT("Global")) == 0)
+		{
+			CommandType = ODSCRecompileCommand::Global;
+		}
+		else if (!CmdName.IsEmpty() && FCString::Stricmp(*CmdName, TEXT("Changed")) == 0)
+		{
+			CommandType = ODSCRecompileCommand::Changed;
+
+			// Compile all the shaders that have changed for the materials we have loaded.
+			for (TObjectIterator<UMaterialInterface> It; It; ++It)
+			{
+				OutMaterialsToLoad.Add(It->GetPathName());
+			}
+		}
+		else
+		{
+			CommandType = ODSCRecompileCommand::Material;
+
+			// tell other side all the materials to load, by pathname
+			for (TObjectIterator<UMaterialInterface> It; It; ++It)
+			{
+				OutMaterialsToLoad.Add(It->GetPathName());
+			}
+		}
+
+		return CommandType;
+	}
+}
+
+void ProcessCookOnTheFlyShaders(bool bReloadGlobalShaders, const TArray<uint8>& MeshMaterialMaps, const TArray<FString>& MaterialsToLoad, const TArray<uint8>& GlobalShaderMap)
+{
+	check(IsInGameThread());
+
+	// now we need to refresh the RHI resources
+	FlushRenderingCommands();
+
+	// reload the global shaders
+	if (bReloadGlobalShaders)
+	{
+		// Some platforms rely on global shaders to be created to implement basic RHI functionality
+		extern int32 GCreateShadersOnLoad;
+		TGuardValue<int32> Guard(GCreateShadersOnLoad, 1);
+		CompileGlobalShaderMap(true);
+	}
+
+	// load all the mesh material shaders if any were sent back
+	if (MeshMaterialMaps.Num() > 0)
+	{
+		// parse the shaders
+		FMemoryReader MemoryReader(MeshMaterialMaps, true);
+		FNameAsStringProxyArchive Ar(MemoryReader);
+
+		TArray<UMaterialInterface*> LoadedMaterials;
+		FMaterialShaderMap::LoadForRemoteRecompile(Ar, GMaxRHIShaderPlatform, LoadedMaterials);
+
+		// Only update materials if we need to.
+		if (LoadedMaterials.Num())
+		{
+			// this will stop the rendering thread, and reattach components, in the destructor
+			FMaterialUpdateContext UpdateContext(FMaterialUpdateContext::EOptions::RecreateRenderStates);
+
+			// gather the shader maps to reattach
+			for (UMaterialInterface* Material : LoadedMaterials)
+			{
+				Material->RecacheUniformExpressions(true);
+				UpdateContext.AddMaterialInterface(Material);
+			}
+		}
+	}
+
+	// load all the global shaders if any were sent back
+	if (GlobalShaderMap.Num() > 0)
+	{
+		// parse the shaders
+		FMemoryReader MemoryReader(GlobalShaderMap, true);
+		FNameAsStringProxyArchive Ar(MemoryReader);
+
+		LoadGlobalShadersForRemoteRecompile(Ar, GMaxRHIShaderPlatform);
+	}
+}
+
+/**
+* Forces a recompile of the global shaders.
+*/
+void RecompileGlobalShaders()
+{
+	if (!FPlatformProperties::RequiresCookedData())
+	{
+		// Flush pending accesses to the existing global shaders.
+		FlushRenderingCommands();
+
+		UMaterialInterface::IterateOverActiveFeatureLevels([&](ERHIFeatureLevel::Type InFeatureLevel) {
+			auto ShaderPlatform = GShaderPlatformForFeatureLevel[InFeatureLevel];
+			GetGlobalShaderMap(ShaderPlatform)->Empty();
+			VerifyGlobalShaders(ShaderPlatform, nullptr, false);
+			});
+
+		GShaderCompilingManager->ProcessAsyncResults(false, true);
+	}
+}
+
+void GetOutdatedShaderTypes(TArray<const FShaderType*>& OutdatedShaderTypes, TArray<const FShaderPipelineType*>& OutdatedShaderPipelineTypes, TArray<const FVertexFactoryType*>& OutdatedFactoryTypes)
+{
+#if WITH_EDITOR
+	for (int PlatformIndex = 0; PlatformIndex < SP_NumPlatforms; ++PlatformIndex)
+	{
+		const FGlobalShaderMap* ShaderMap = GGlobalShaderMap[PlatformIndex];
+		if (ShaderMap)
+		{
+			ShaderMap->GetOutdatedTypes(OutdatedShaderTypes, OutdatedShaderPipelineTypes, OutdatedFactoryTypes);
+		}
+	}
+
+	FMaterialShaderMap::GetAllOutdatedTypes(OutdatedShaderTypes, OutdatedShaderPipelineTypes, OutdatedFactoryTypes);
+
+	for (int32 TypeIndex = 0; TypeIndex < OutdatedShaderTypes.Num(); TypeIndex++)
+	{
+		UE_LOG(LogShaders, Warning, TEXT("		Recompiling %s"), OutdatedShaderTypes[TypeIndex]->GetName());
+	}
+	for (int32 TypeIndex = 0; TypeIndex < OutdatedShaderPipelineTypes.Num(); TypeIndex++)
+	{
+		UE_LOG(LogShaders, Warning, TEXT("		Recompiling %s"), OutdatedShaderPipelineTypes[TypeIndex]->GetName());
+	}
+	for (int32 TypeIndex = 0; TypeIndex < OutdatedFactoryTypes.Num(); TypeIndex++)
+	{
+		UE_LOG(LogShaders, Warning, TEXT("		Recompiling %s"), OutdatedFactoryTypes[TypeIndex]->GetName());
+	}
+#endif // WITH_EDITOR
+}
+
+bool RecompileShaders(const TCHAR* Cmd, FOutputDevice& Ar)
+{
+	// if this platform can't compile shaders, then we try to send a message to a file/cooker server
+	if (FPlatformProperties::RequiresCookedData())
+	{
+#if WITH_ODSC
+		TArray<FString> MaterialsToLoad;
+		ODSCRecompileCommand CommandType = ParseRecompileCommandString(Cmd, MaterialsToLoad);
+		GODSCManager->AddThreadedRequest(MaterialsToLoad, GMaxRHIShaderPlatform, CommandType);
+#endif
+		return true;
+	}
+
+	FString FlagStr(FParse::Token(Cmd, 0));
+	if (FlagStr.Len() > 0)
+	{
+		GWarn->BeginSlowTask(NSLOCTEXT("ShaderCompilingManager", "BeginRecompilingShadersTask", "Recompiling shaders"), true);
+
+		// Flush the shader file cache so that any changes to shader source files will be detected
+		FlushShaderFileCache();
+		FlushRenderingCommands();
+
+		if (FCString::Stricmp(*FlagStr, TEXT("Changed")) == 0)
+		{
+			TArray<const FShaderType*> OutdatedShaderTypes;
+			TArray<const FVertexFactoryType*> OutdatedFactoryTypes;
+			TArray<const FShaderPipelineType*> OutdatedShaderPipelineTypes;
+			{
+				FRecompileShadersTimer SearchTimer(TEXT("Searching for changed files"));
+				GetOutdatedShaderTypes(OutdatedShaderTypes, OutdatedShaderPipelineTypes, OutdatedFactoryTypes);
+			}
+
+			if (OutdatedShaderPipelineTypes.Num() > 0 || OutdatedShaderTypes.Num() > 0 || OutdatedFactoryTypes.Num() > 0)
+			{
+				FRecompileShadersTimer TestTimer(TEXT("RecompileShaders Changed"));
+
+				UMaterialInterface::IterateOverActiveFeatureLevels([&](ERHIFeatureLevel::Type InFeatureLevel) {
+					auto ShaderPlatform = GShaderPlatformForFeatureLevel[InFeatureLevel];
+					BeginRecompileGlobalShaders(OutdatedShaderTypes, OutdatedShaderPipelineTypes, ShaderPlatform);
+					});
+
+				// Block on global shaders
+				FinishRecompileGlobalShaders();
+
+				// Kick off global shader recompiles
+				UMaterialInterface::IterateOverActiveFeatureLevels([&](ERHIFeatureLevel::Type InFeatureLevel) {
+					auto ShaderPlatform = GShaderPlatformForFeatureLevel[InFeatureLevel];
+					UMaterial::UpdateMaterialShaders(OutdatedShaderTypes, OutdatedShaderPipelineTypes, OutdatedFactoryTypes, ShaderPlatform);
+					});
+
+				GWarn->StatusUpdate(0, 1, NSLOCTEXT("ShaderCompilingManager", "CompilingGlobalShaderStatus", "Compiling global shaders..."));
+			}
+			else
+			{
+				UE_LOG(LogShaderCompilers, Warning, TEXT("No Shader changes found."));
+			}
+		}
+		else if (FCString::Stricmp(*FlagStr, TEXT("Global")) == 0)
+		{
+			FRecompileShadersTimer TestTimer(TEXT("RecompileShaders Global"));
+			RecompileGlobalShaders();
+		}
+		else if (FCString::Stricmp(*FlagStr, TEXT("Material")) == 0)
+		{
+			FString RequestedMaterialName(FParse::Token(Cmd, 0));
+			FRecompileShadersTimer TestTimer(FString::Printf(TEXT("Recompile Material %s"), *RequestedMaterialName));
+
+			ITargetPlatformManagerModule& TPM = GetTargetPlatformManagerRef();
+			FString TargetPlatformName(FParse::Token(Cmd, 0));
+			const ITargetPlatform* TargetPlatform = nullptr;
+			if (TargetPlatformName.Len() > 0)
+			{
+				TargetPlatform = TPM.FindTargetPlatform(TargetPlatformName);
+			}
+
+			bool bMaterialFound = false;
+			for (TObjectIterator<UMaterialInterface> It; It; ++It)
+			{
+				UMaterialInterface* Material = *It;
+				if (Material && Material->GetName() == RequestedMaterialName)
+				{
+					bMaterialFound = true;
+#if WITH_EDITOR
+					// <Pre/Post>EditChange will force a re-creation of the resource,
+					// in turn recompiling the shader.
+					if (TargetPlatform)
+					{
+						Material->BeginCacheForCookedPlatformData(TargetPlatform);
+						while (!Material->IsCachedCookedPlatformDataLoaded(TargetPlatform))
+						{
+							FPlatformProcess::Sleep(0.1f);
+							GShaderCompilingManager->ProcessAsyncResults(false, false);
+						}
+						Material->ClearCachedCookedPlatformData(TargetPlatform);
+					}
+					else
+					{
+						Material->PreEditChange(NULL);
+						Material->PostEditChange();
+					}
+#endif // WITH_EDITOR
+					break;
+				}
+			}
+
+			if (!bMaterialFound)
+			{
+				TestTimer.Stop(false);
+				UE_LOG(LogShaderCompilers, Warning, TEXT("Couldn't find Material %s!"), *RequestedMaterialName);
+			}
+		}
+		else if (FCString::Stricmp(*FlagStr, TEXT("All")) == 0)
+		{
+			FRecompileShadersTimer TestTimer(TEXT("RecompileShaders"));
+			RecompileGlobalShaders();
+
+			FMaterialUpdateContext UpdateContext(0);
+			for (TObjectIterator<UMaterial> It; It; ++It)
+			{
+				UMaterial* Material = *It;
+				if (Material)
+				{
+					UE_LOG(LogShaderCompilers, Log, TEXT("recompiling [%s]"), *Material->GetFullName());
+					UpdateContext.AddMaterial(Material);
+#if WITH_EDITOR
+					// <Pre/Post>EditChange will force a re-creation of the resource,
+					// in turn recompiling the shader.
+					Material->PreEditChange(NULL);
+					Material->PostEditChange();
+#endif // WITH_EDITOR
+				}
+			}
+		}
+		else
+		{
+			TArray<const FShaderType*> ShaderTypes = FShaderType::GetShaderTypesByFilename(*FlagStr);
+			TArray<const FShaderPipelineType*> ShaderPipelineTypes = FShaderPipelineType::GetShaderPipelineTypesByFilename(*FlagStr);
+			if (ShaderTypes.Num() > 0 || ShaderPipelineTypes.Num() > 0)
+			{
+				FRecompileShadersTimer TestTimer(TEXT("RecompileShaders SingleShader"));
+
+				TArray<const FVertexFactoryType*> FactoryTypes;
+
+				UMaterialInterface::IterateOverActiveFeatureLevels([&](ERHIFeatureLevel::Type InFeatureLevel) {
+					auto ShaderPlatform = GShaderPlatformForFeatureLevel[InFeatureLevel];
+					BeginRecompileGlobalShaders(ShaderTypes, ShaderPipelineTypes, ShaderPlatform);
+					//UMaterial::UpdateMaterialShaders(ShaderTypes, ShaderPipelineTypes, FactoryTypes, ShaderPlatform);
+					FinishRecompileGlobalShaders();
+					});
+			}
+		}
+
+		GWarn->EndSlowTask();
+
+		return 1;
+	}
+
+	UE_LOG(LogShaderCompilers, Warning, TEXT("Invalid parameter. Options are: \n'Changed', 'Global', 'Material [name]', 'All'."));
+	return 1;
+}
+
+bool RecompileChangedShadersForPlatform(const FString& PlatformName)
+{
+	// figure out what shader platforms to recompile
+	ITargetPlatformManagerModule* TPM = GetTargetPlatformManager();
+	ITargetPlatform* TargetPlatform = TPM->FindTargetPlatform(PlatformName);
+	if (TargetPlatform == NULL)
+	{
+		UE_LOG(LogShaders, Display, TEXT("Failed to find target platform module for %s"), *PlatformName);
+		return false;
+	}
+
+	TArray<FName> DesiredShaderFormats;
+	TargetPlatform->GetAllTargetedShaderFormats(DesiredShaderFormats);
+
+
+
+	// figure out which shaders are out of date
+	TArray<const FShaderType*> OutdatedShaderTypes;
+	TArray<const FVertexFactoryType*> OutdatedFactoryTypes;
+	TArray<const FShaderPipelineType*> OutdatedShaderPipelineTypes;
+
+	// Pick up new changes to shader files
+	FlushShaderFileCache();
+
+	GetOutdatedShaderTypes(OutdatedShaderTypes, OutdatedShaderPipelineTypes, OutdatedFactoryTypes);
+	UE_LOG(LogShaders, Display, TEXT("We found %d out of date shader types, %d outdated pipeline types, and %d out of date VF types!"), OutdatedShaderTypes.Num(), OutdatedShaderPipelineTypes.Num(), OutdatedFactoryTypes.Num());
+
+	for (int32 FormatIndex = 0; FormatIndex < DesiredShaderFormats.Num(); FormatIndex++)
+	{
+		// get the shader platform enum
+		const EShaderPlatform ShaderPlatform = ShaderFormatToLegacyShaderPlatform(DesiredShaderFormats[FormatIndex]);
+
+		// Only compile for the desired platform if requested
+		// Kick off global shader recompiles
+		BeginRecompileGlobalShaders(OutdatedShaderTypes, OutdatedShaderPipelineTypes, ShaderPlatform);
+
+		// Block on global shaders
+		FinishRecompileGlobalShaders();
+#if WITH_EDITOR
+		// we only want to actually compile mesh shaders if we have out of date ones
+		if (OutdatedShaderTypes.Num() || OutdatedFactoryTypes.Num())
+		{
+			for (TObjectIterator<UMaterialInterface> It; It; ++It)
+			{
+				(*It)->ClearCachedCookedPlatformData(TargetPlatform);
+			}
+		}
+#endif
+	}
+
+	if (OutdatedFactoryTypes.Num() || OutdatedShaderTypes.Num())
+	{
+		return true;
+	}
+	return false;
+}
+
+FArchive& operator<<(FArchive& Ar, FODSCRequestPayload& Elem)
+{
+	uint32 ConvertedShaderPlatform = (uint32)Elem.ShaderPlatform;
+	Ar << ConvertedShaderPlatform;
+	Ar << Elem.MaterialName;
+	Ar << Elem.VertexFactoryName;
+	Ar << Elem.PipelineName;
+	Ar << Elem.ShaderTypeNames;
+	Ar << Elem.RequestHash;
+
+	return Ar;
+}
+
+FShaderRecompileData::FShaderRecompileData(const FString& InPlatformName, TArray<FString>* OutModifiedFiles, TArray<uint8>* OutMeshMaterialMaps, TArray<uint8>* OutGlobalShaderMap)
+	: PlatformName(InPlatformName),
+	ModifiedFiles(OutModifiedFiles),
+	MeshMaterialMaps(OutMeshMaterialMaps),
+	GlobalShaderMap(OutGlobalShaderMap)
+{
+}
+
+FShaderRecompileData::FShaderRecompileData(const FString& InPlatformName, EShaderPlatform InShaderPlatform, ODSCRecompileCommand InCommandType, TArray<FString>* OutModifiedFiles, TArray<uint8>* OutMeshMaterialMaps, TArray<uint8>* OutGlobalShaderMap)
+	: PlatformName(InPlatformName),
+	ShaderPlatform(InShaderPlatform),
+	ModifiedFiles(OutModifiedFiles),
+	MeshMaterialMaps(OutMeshMaterialMaps),
+	CommandType(InCommandType),
+	GlobalShaderMap(OutGlobalShaderMap)
+{
+}
+
+void SaveGlobalShadersForRemoteRecompile(FArchive& Ar, EShaderPlatform ShaderPlatform)
+{
+	FGlobalShaderMap* GlobalShaderMap = GetGlobalShaderMap(ShaderPlatform);
+	uint8 bIsValid = GlobalShaderMap != nullptr;
+	Ar << bIsValid;
+
+	if (GlobalShaderMap)
+	{
+		GlobalShaderMap->SaveToGlobalArchive(Ar);
+	}
+}
+
+void LoadGlobalShadersForRemoteRecompile(FArchive& Ar, EShaderPlatform ShaderPlatform)
+{
+	uint8 bIsValid = 0;
+	Ar << bIsValid;
+
+	if (bIsValid)
+	{
+		FlushRenderingCommands();
+
+		FGlobalShaderMap* NewGlobalShaderMap = new FGlobalShaderMap(ShaderPlatform);
+		if (NewGlobalShaderMap)
+		{
+			NewGlobalShaderMap->LoadFromGlobalArchive(Ar);
+
+			if (GGlobalShaderMap[ShaderPlatform])
+			{
+				GGlobalShaderMap[ShaderPlatform]->ReleaseAllSections();
+
+				delete GGlobalShaderMap[ShaderPlatform];
+				GGlobalShaderMap[ShaderPlatform] = nullptr;
+				GGlobalShaderMap[ShaderPlatform] = NewGlobalShaderMap;
+
+				VerifyGlobalShaders(ShaderPlatform, nullptr, false);
+
+				// Invalidate global bound shader states so they will be created with the new shaders the next time they are set (in SetGlobalBoundShaderState)
+				for (TLinkedList<FGlobalBoundShaderStateResource*>::TIterator It(FGlobalBoundShaderStateResource::GetGlobalBoundShaderStateList()); It; It.Next())
+				{
+					BeginUpdateResourceRHI(*It);
+				}
+
+				PropagateGlobalShadersToAllPrimitives();
+			}
+			else
+			{
+				delete NewGlobalShaderMap;
+			}
+		}
+	}
+}
+
+#if WITH_EDITOR
+void RecompileShadersForRemote(
+	FShaderRecompileData& Args,
+	const FString& OutputDirectory)
+{
+	// figure out what shader platforms to recompile
+	ITargetPlatformManagerModule* TPM = GetTargetPlatformManager();
+	ITargetPlatform* TargetPlatform = TPM->FindTargetPlatform(Args.PlatformName);
+	if (TargetPlatform == NULL)
+	{
+		UE_LOG(LogShaders, Display, TEXT("Failed to find target platform module for %s"), *Args.PlatformName);
+		return;
+	}
+
+	const double StartTime = FPlatformTime::Seconds();
+
+	UE_LOG(LogShaders, Display, TEXT(""));
+	UE_LOG(LogShaders, Display, TEXT("********************************"));
+	UE_LOG(LogShaders, Display, TEXT("Received compile shader request."));
+
+	const bool bPreviousState = GShaderCompilingManager->IsShaderCompilationSkipped();
+	GShaderCompilingManager->SkipShaderCompilation(false);
+
+	TArray<FName> DesiredShaderFormats;
+	TargetPlatform->GetAllTargetedShaderFormats(DesiredShaderFormats);
+
+	UE_LOG(LogShaders, Verbose, TEXT("Loading %d materials..."), Args.MaterialsToLoad.Num());
+	// make sure all materials the client has loaded will be processed
+	TArray<UMaterialInterface*> MaterialsToCompile;
+
+	for (int32 Index = 0; Index < Args.MaterialsToLoad.Num(); Index++)
+	{
+		UE_LOG(LogShaders, Verbose, TEXT("   --> %s"), *Args.MaterialsToLoad[Index]);
+		MaterialsToCompile.Add(LoadObject<UMaterialInterface>(NULL, *Args.MaterialsToLoad[Index]));
+	}
+
+	UE_LOG(LogShaders, Verbose, TEXT("  Done!"))
+
+		// figure out which shaders are out of date
+		TArray<const FShaderType*> OutdatedShaderTypes;
+	TArray<const FVertexFactoryType*> OutdatedFactoryTypes;
+	TArray<const FShaderPipelineType*> OutdatedShaderPipelineTypes;
+
+	// Pick up new changes to shader files
+	FlushShaderFileCache();
+
+	if (Args.ShadersToRecompile.Num())
+	{
+		UE_LOG(LogShaders, Display, TEXT("Received %d shaders to compile."), Args.ShadersToRecompile.Num());
+	}
+
+	for (const FODSCRequestPayload& payload : Args.ShadersToRecompile)
+	{
+		UE_LOG(LogShaders, Display, TEXT(""));
+		UE_LOG(LogShaders, Display, TEXT("\tMaterial:    %s "), *payload.MaterialName);
+		UE_LOG(LogShaders, Display, TEXT("\tVF Type:     %s "), *payload.VertexFactoryName);
+
+		MaterialsToCompile.Add(LoadObject<UMaterialInterface>(NULL, *payload.MaterialName));
+
+		const FVertexFactoryType* VFType = FVertexFactoryType::GetVFByName(payload.VertexFactoryName);
+		if (VFType)
+		{
+			OutdatedFactoryTypes.Add(VFType);
+		}
+
+		const FShaderPipelineType* PipelineType = FShaderPipelineType::GetShaderPipelineTypeByName(payload.PipelineName);
+		if (PipelineType)
+		{
+			OutdatedShaderPipelineTypes.Add(PipelineType);
+		}
+
+		for (const FString& ShaderTypeName : payload.ShaderTypeNames)
+		{
+			UE_LOG(LogShaders, Display, TEXT("\tShader Type: %s"), *ShaderTypeName);
+
+			const FShaderType* ShaderType = FShaderType::GetShaderTypeByName(*ShaderTypeName);
+			if (ShaderType)
+			{
+				OutdatedShaderTypes.Add(ShaderType);
+			}
+		}
+	}
+
+	{
+		for (int32 FormatIndex = 0; FormatIndex < DesiredShaderFormats.Num(); FormatIndex++)
+		{
+			// get the shader platform enum
+			const EShaderPlatform ShaderPlatform = ShaderFormatToLegacyShaderPlatform(DesiredShaderFormats[FormatIndex]);
+
+			// Only compile for the desired platform if requested
+			if (ShaderPlatform == Args.ShaderPlatform || Args.ShaderPlatform == SP_NumPlatforms)
+			{
+				// If we are explicitly wanting to recompile global or if shaders have changed.
+				if (Args.CommandType == ODSCRecompileCommand::Global ||
+					Args.CommandType == ODSCRecompileCommand::Changed)
+				{
+					UE_LOG(LogShaders, Display, TEXT("Recompiling global shaders."));
+
+					// Explicitly get outdated types for global shaders.
+					const FGlobalShaderMap* ShaderMap = GGlobalShaderMap[ShaderPlatform];
+					if (ShaderMap)
+					{
+						ShaderMap->GetOutdatedTypes(OutdatedShaderTypes, OutdatedShaderPipelineTypes, OutdatedFactoryTypes);
+					}
+
+					UE_LOG(LogShaders, Display, TEXT("\tFound %d outdated shader types."), OutdatedShaderTypes.Num() + OutdatedShaderPipelineTypes.Num());
+
+					// Kick off global shader recompiles
+					BeginRecompileGlobalShaders(OutdatedShaderTypes, OutdatedShaderPipelineTypes, ShaderPlatform, TargetPlatform);
+
+					// Block on global shaders
+					FinishRecompileGlobalShaders();
+
+					// write the shader compilation info to memory, converting fnames to strings
+					FMemoryWriter MemWriter(*Args.GlobalShaderMap, true);
+					FNameAsStringProxyArchive Ar(MemWriter);
+					Ar.SetCookingTarget(TargetPlatform);
+
+					// save out the global shader map to the byte array
+					SaveGlobalShadersForRemoteRecompile(Ar, ShaderPlatform);
+				}
+
+				// we only want to actually compile mesh shaders if a client directly requested it
+				if ((Args.CommandType == ODSCRecompileCommand::Material || Args.CommandType == ODSCRecompileCommand::Changed) &&
+					Args.MeshMaterialMaps != nullptr)
+				{
+					TMap<FString, TArray<TRefCountPtr<FMaterialShaderMap> > > CompiledShaderMaps;
+					UMaterial::CompileMaterialsForRemoteRecompile(MaterialsToCompile, ShaderPlatform, TargetPlatform, CompiledShaderMaps);
+
+					// write the shader compilation info to memory, converting fnames to strings
+					FMemoryWriter MemWriter(*Args.MeshMaterialMaps, true);
+					FNameAsStringProxyArchive Ar(MemWriter);
+					Ar.SetCookingTarget(TargetPlatform);
+
+					// save out the shader map to the byte array
+					FMaterialShaderMap::SaveForRemoteRecompile(Ar, CompiledShaderMaps);
+				}
+
+				// save it out so the client can get it (and it's up to date next time)
+				FString GlobalShaderFilename = SaveGlobalShaderFile(ShaderPlatform, OutputDirectory, TargetPlatform);
+
+				// add this to the list of files to tell the other end about
+				if (Args.ModifiedFiles)
+				{
+					// need to put it in non-sandbox terms
+					FString SandboxPath(GlobalShaderFilename);
+					check(SandboxPath.StartsWith(OutputDirectory));
+					SandboxPath.ReplaceInline(*OutputDirectory, TEXT("../../../"));
+					FPaths::NormalizeFilename(SandboxPath);
+					Args.ModifiedFiles->Add(SandboxPath);
+				}
+			}
+		}
+	}
+
+	UE_LOG(LogShaders, Display, TEXT(""));
+	UE_LOG(LogShaders, Display, TEXT("Finished shader compile request in %.2f seconds."), FPlatformTime::Seconds() - StartTime);
+
+	// Provide a log of what happened.
+	GShaderCompilingManager->PrintStats(true);
+
+	// Restore compilation state.
+	GShaderCompilingManager->SkipShaderCompilation(bPreviousState);
+}
+#endif // WITH_EDITOR
+
+void BeginRecompileGlobalShaders(const TArray<const FShaderType*>& OutdatedShaderTypes, const TArray<const FShaderPipelineType*>& OutdatedShaderPipelineTypes, EShaderPlatform ShaderPlatform, const ITargetPlatform* TargetPlatform)
+{
+	if (!FPlatformProperties::RequiresCookedData())
+	{
+		// Flush pending accesses to the existing global shaders.
+		FlushRenderingCommands();
+
+		// Calling CompileGlobalShaderMap will force starting the compile jobs if the map is empty (by calling VerifyGlobalShaders)
+		CompileGlobalShaderMap(ShaderPlatform, TargetPlatform, false);
+		FGlobalShaderMap* GlobalShaderMap = GetGlobalShaderMap(ShaderPlatform);
+
+		// Now check if there is any work to be done wrt outdates types
+		if (OutdatedShaderTypes.Num() > 0 || OutdatedShaderPipelineTypes.Num() > 0)
+		{
+			VerifyGlobalShaders(ShaderPlatform, TargetPlatform, false, &OutdatedShaderTypes, &OutdatedShaderPipelineTypes);
+		}
+	}
+}
+
+void FinishRecompileGlobalShaders()
+{
+	// Block until global shaders have been compiled and processed
+	GShaderCompilingManager->ProcessAsyncResults(false, true);
+}
\ No newline at end of file
diff --git a/Engine/Source/Runtime/Engine/Private/ShaderCompiler/ShaderCompiler.cpp b/Engine/Source/Runtime/Engine/Private/ShaderCompiler/ShaderCompiler.cpp
--- a/Engine/Source/Runtime/Engine/Private/ShaderCompiler/ShaderCompiler.cpp
+++ b/Engine/Source/Runtime/Engine/Private/ShaderCompiler/ShaderCompiler.cpp
@@ -5,12 +5,12 @@
 =============================================================================*/
 
 #include "ShaderCompiler.h"
+#include "ShaderCompilerShared.h"
 #include "GenericPlatform/GenericPlatformFile.h"
 #include "HAL/PlatformFileManager.h"
 #include "Misc/MessageDialog.h"
 #include "HAL/FileManager.h"
 #include "HAL/ExceptionHandling.h"
-#include "Serialization/NameAsStringProxyArchive.h"
 #include "Misc/CommandLine.h"
 #include "Misc/Paths.h"
 #include "Misc/ScopeLock.h"
@@ -19,7 +19,6 @@
 #include "Serialization/MemoryReader.h"
 #include "Serialization/MemoryHasher.h"
 #include "Misc/ConfigCacheIni.h"
-#include "Misc/FeedbackContext.h"
 #include "Misc/ScopedSlowTask.h"
 #include "Modules/ModuleManager.h"
 #include "UObject/UObjectHash.h"
@@ -69,9 +68,7 @@
 #include "TextureCompiler.h"
 #include "Rendering/StaticLightingSystemInterface.h"
 #endif
-#if WITH_ODSC
-#include "ODSC/ODSCManager.h"
-#endif
+
 #include "HAL/IConsoleManager.h"
 
 #define LOCTEXT_NAMESPACE "ShaderCompiler"
@@ -4162,39 +4159,6 @@ void FShaderCompilingManager::BlockOnAllShaderMapCompletion(TMap<int32, FShaderM
 	}
 }
 
-namespace
-{
-	void PropagateGlobalShadersToAllPrimitives()
-	{
-		// Re-register everything to work around FShader lifetime issues - it currently lives and dies with the
-		// shadermap it is stored in, while cached MDCs can reference its memory. Re-registering will
-		// re-create the cache.
-		TRACE_CPUPROFILER_EVENT_SCOPE(PropagateGlobalShadersToAllPrimitives);
-
-		FObjectCacheContextScope ObjectCacheScope;
-		TSet<FSceneInterface*> ScenesToUpdate;
-		TIndirectArray<FComponentRecreateRenderStateContext> ComponentContexts;
-		for (UPrimitiveComponent* PrimitiveComponent : ObjectCacheScope.GetContext().GetPrimitiveComponents())
-		{
-			if (PrimitiveComponent->IsRenderStateCreated())
-			{
-				ComponentContexts.Add(new FComponentRecreateRenderStateContext(PrimitiveComponent, &ScenesToUpdate));
-#if WITH_EDITOR
-				if (PrimitiveComponent->HasValidSettingsForStaticLighting(false))
-				{
-					FStaticLightingSystemInterface::OnPrimitiveComponentUnregistered.Broadcast(PrimitiveComponent);
-					FStaticLightingSystemInterface::OnPrimitiveComponentRegistered.Broadcast(PrimitiveComponent);
-				}
-#endif
-			}
-		}
-
-		UpdateAllPrimitiveSceneInfosForScenes(ScenesToUpdate);
-		ComponentContexts.Empty();
-		UpdateAllPrimitiveSceneInfosForScenes(ScenesToUpdate);
-	}
-}
-
 void FShaderCompilingManager::ProcessCompiledShaderMaps(
 	TMap<int32, FShaderMapFinalizeResults>& CompiledShaderMaps, 
 	float TimeBudget)
@@ -5877,368 +5841,6 @@ void GlobalBeginCompileShader(
 #endif
 }
 
-
-/** Timer class used to report information on the 'recompileshaders' console command. */
-class FRecompileShadersTimer
-{
-public:
-	FRecompileShadersTimer(const TCHAR* InInfoStr=TEXT("Test")) :
-		InfoStr(InInfoStr),
-		bAlreadyStopped(false)
-	{
-		StartTime = FPlatformTime::Seconds();
-	}
-
-	FRecompileShadersTimer(const FString& InInfoStr) :
-		InfoStr(InInfoStr),
-		bAlreadyStopped(false)
-	{
-		StartTime = FPlatformTime::Seconds();
-	}
-
-	void Stop(bool DisplayLog = true)
-	{
-		if (!bAlreadyStopped)
-		{
-			bAlreadyStopped = true;
-			EndTime = FPlatformTime::Seconds();
-			TimeElapsed = EndTime-StartTime;
-			if (DisplayLog)
-			{
-				UE_LOG(LogShaderCompilers, Warning, TEXT("		[%s] took [%.4f] s"),*InfoStr,TimeElapsed);
-			}
-		}
-	}
-
-	~FRecompileShadersTimer()
-	{
-		Stop(true);
-	}
-
-protected:
-	double StartTime,EndTime;
-	double TimeElapsed;
-	FString InfoStr;
-	bool bAlreadyStopped;
-};
-
-namespace
-{
-	ODSCRecompileCommand ParseRecompileCommandString(const TCHAR* CmdString, TArray<FString>& OutMaterialsToLoad)
-	{
-		FString CmdName = FParse::Token(CmdString, 0);
-
-		ODSCRecompileCommand CommandType = ODSCRecompileCommand::None;
-		OutMaterialsToLoad.Empty();
-
-		if( !CmdName.IsEmpty() && FCString::Stricmp(*CmdName,TEXT("Material"))==0 )
-		{
-			CommandType = ODSCRecompileCommand::Material;
-
-			// tell other side the material to load, by pathname
-			FString RequestedMaterialName( FParse::Token( CmdString, 0 ) );
-
-			for( TObjectIterator<UMaterialInterface> It; It; ++It )
-			{
-				UMaterial* Material = It->GetMaterial();
-
-				if( Material && Material->GetName() == RequestedMaterialName)
-				{
-					OutMaterialsToLoad.Add(It->GetPathName());
-					break;
-				}
-			}
-		}
-		else if (!CmdName.IsEmpty() && FCString::Stricmp(*CmdName, TEXT("Global")) == 0)
-		{
-			CommandType = ODSCRecompileCommand::Global;
-		}
-		else if (!CmdName.IsEmpty() && FCString::Stricmp(*CmdName, TEXT("Changed")) == 0)
-		{
-			CommandType = ODSCRecompileCommand::Changed;
-
-			// Compile all the shaders that have changed for the materials we have loaded.
-			for (TObjectIterator<UMaterialInterface> It; It; ++It)
-			{
-				OutMaterialsToLoad.Add(It->GetPathName());
-			}
-		}
-		else
-		{
-			CommandType = ODSCRecompileCommand::Material;
-
-			// tell other side all the materials to load, by pathname
-			for (TObjectIterator<UMaterialInterface> It; It; ++It)
-			{
-				OutMaterialsToLoad.Add(It->GetPathName());
-			}
-		}
-
-		return CommandType;
-	}
-}
-
-void ProcessCookOnTheFlyShaders(bool bReloadGlobalShaders, const TArray<uint8>& MeshMaterialMaps, const TArray<FString>& MaterialsToLoad, const TArray<uint8>& GlobalShaderMap)
-{
-	check(IsInGameThread());
-
-	// now we need to refresh the RHI resources
-	FlushRenderingCommands();
-
-	// reload the global shaders
-	if (bReloadGlobalShaders)
-	{
-		// Some platforms rely on global shaders to be created to implement basic RHI functionality
-		TGuardValue<int32> Guard(GCreateShadersOnLoad, 1);
-		CompileGlobalShaderMap(true);
-	}
-
-	// load all the mesh material shaders if any were sent back
-	if (MeshMaterialMaps.Num() > 0)
-	{
-		// parse the shaders
-		FMemoryReader MemoryReader(MeshMaterialMaps, true);
-		FNameAsStringProxyArchive Ar(MemoryReader);
-
-		TArray<UMaterialInterface*> LoadedMaterials;
-		FMaterialShaderMap::LoadForRemoteRecompile(Ar, GMaxRHIShaderPlatform, LoadedMaterials);
-
-		// Only update materials if we need to.
-		if (LoadedMaterials.Num())
-		{
-			// this will stop the rendering thread, and reattach components, in the destructor
-			FMaterialUpdateContext UpdateContext(FMaterialUpdateContext::EOptions::RecreateRenderStates);
-
-			// gather the shader maps to reattach
-			for (UMaterialInterface* Material : LoadedMaterials)
-			{
-				Material->RecacheUniformExpressions(true);
-				UpdateContext.AddMaterialInterface(Material);
-			}
-		}
-	}
-
-	// load all the global shaders if any were sent back
-	if (GlobalShaderMap.Num() > 0)
-	{
-		// parse the shaders
-		FMemoryReader MemoryReader(GlobalShaderMap, true);
-		FNameAsStringProxyArchive Ar(MemoryReader);
-
-		LoadGlobalShadersForRemoteRecompile(Ar, GMaxRHIShaderPlatform);
-	}
-}
-
-/**
-* Forces a recompile of the global shaders.
-*/
-void RecompileGlobalShaders()
-{
-	if (!FPlatformProperties::RequiresCookedData())
-	{
-		// Flush pending accesses to the existing global shaders.
-		FlushRenderingCommands();
-
-		UMaterialInterface::IterateOverActiveFeatureLevels([&](ERHIFeatureLevel::Type InFeatureLevel) {
-			auto ShaderPlatform = GShaderPlatformForFeatureLevel[InFeatureLevel];
-			GetGlobalShaderMap(ShaderPlatform)->Empty();
-			VerifyGlobalShaders(ShaderPlatform, nullptr, false);
-		});
-
-		GShaderCompilingManager->ProcessAsyncResults(false, true);
-	}
-}
-
-void GetOutdatedShaderTypes(TArray<const FShaderType*>& OutdatedShaderTypes, TArray<const FShaderPipelineType*>& OutdatedShaderPipelineTypes, TArray<const FVertexFactoryType*>& OutdatedFactoryTypes)
-{
-#if WITH_EDITOR
-	for (int PlatformIndex = 0; PlatformIndex < SP_NumPlatforms; ++PlatformIndex)
-	{
-		const FGlobalShaderMap* ShaderMap = GGlobalShaderMap[PlatformIndex];
-		if (ShaderMap)
-		{
-			ShaderMap->GetOutdatedTypes(OutdatedShaderTypes, OutdatedShaderPipelineTypes, OutdatedFactoryTypes);
-		}
-	}
-
-	FMaterialShaderMap::GetAllOutdatedTypes(OutdatedShaderTypes, OutdatedShaderPipelineTypes, OutdatedFactoryTypes);
-
-	for (int32 TypeIndex = 0; TypeIndex < OutdatedShaderTypes.Num(); TypeIndex++)
-	{
-		UE_LOG(LogShaders, Warning, TEXT("		Recompiling %s"), OutdatedShaderTypes[TypeIndex]->GetName());
-	}
-	for (int32 TypeIndex = 0; TypeIndex < OutdatedShaderPipelineTypes.Num(); TypeIndex++)
-	{
-		UE_LOG(LogShaders, Warning, TEXT("		Recompiling %s"), OutdatedShaderPipelineTypes[TypeIndex]->GetName());
-	}
-	for (int32 TypeIndex = 0; TypeIndex < OutdatedFactoryTypes.Num(); TypeIndex++)
-	{
-		UE_LOG(LogShaders, Warning, TEXT("		Recompiling %s"), OutdatedFactoryTypes[TypeIndex]->GetName());
-	}
-#endif // WITH_EDITOR
-}
-
-bool RecompileShaders(const TCHAR* Cmd, FOutputDevice& Ar)
-{
-	// if this platform can't compile shaders, then we try to send a message to a file/cooker server
-	if (FPlatformProperties::RequiresCookedData())
-	{
-#if WITH_ODSC
-		TArray<FString> MaterialsToLoad;
-		ODSCRecompileCommand CommandType = ParseRecompileCommandString(Cmd, MaterialsToLoad);
-		GODSCManager->AddThreadedRequest(MaterialsToLoad, GMaxRHIShaderPlatform, CommandType);
-#endif
-		return true;
-	}
-
-	FString FlagStr(FParse::Token(Cmd, 0));
-	if( FlagStr.Len() > 0 )
-	{
-		GWarn->BeginSlowTask( NSLOCTEXT("ShaderCompilingManager", "BeginRecompilingShadersTask", "Recompiling shaders"), true );
-
-		// Flush the shader file cache so that any changes to shader source files will be detected
-		FlushShaderFileCache();
-		FlushRenderingCommands();
-
-		if( FCString::Stricmp(*FlagStr,TEXT("Changed"))==0)
-		{
-			TArray<const FShaderType*> OutdatedShaderTypes;
-			TArray<const FVertexFactoryType*> OutdatedFactoryTypes;
-			TArray<const FShaderPipelineType*> OutdatedShaderPipelineTypes;
-			{
-				FRecompileShadersTimer SearchTimer(TEXT("Searching for changed files"));
-				GetOutdatedShaderTypes(OutdatedShaderTypes, OutdatedShaderPipelineTypes, OutdatedFactoryTypes);
-			}
-
-			if (OutdatedShaderPipelineTypes.Num() > 0 || OutdatedShaderTypes.Num() > 0 || OutdatedFactoryTypes.Num() > 0)
-			{
-				FRecompileShadersTimer TestTimer(TEXT("RecompileShaders Changed"));
-
-				UMaterialInterface::IterateOverActiveFeatureLevels([&](ERHIFeatureLevel::Type InFeatureLevel) {
-					auto ShaderPlatform = GShaderPlatformForFeatureLevel[InFeatureLevel];
-					BeginRecompileGlobalShaders(OutdatedShaderTypes, OutdatedShaderPipelineTypes, ShaderPlatform);
-				});
-
-				// Block on global shaders
-				FinishRecompileGlobalShaders();
-
-				// Kick off global shader recompiles
-				UMaterialInterface::IterateOverActiveFeatureLevels([&](ERHIFeatureLevel::Type InFeatureLevel) {
-					auto ShaderPlatform = GShaderPlatformForFeatureLevel[InFeatureLevel];
-					UMaterial::UpdateMaterialShaders(OutdatedShaderTypes, OutdatedShaderPipelineTypes, OutdatedFactoryTypes, ShaderPlatform);
-				});
-
-				GWarn->StatusUpdate(0, 1, NSLOCTEXT("ShaderCompilingManager", "CompilingGlobalShaderStatus", "Compiling global shaders..."));
-			}
-			else
-			{
-				UE_LOG(LogShaderCompilers, Warning, TEXT("No Shader changes found."));
-			}
-		}
-		else if( FCString::Stricmp(*FlagStr,TEXT("Global"))==0)
-		{
-			FRecompileShadersTimer TestTimer(TEXT("RecompileShaders Global"));
-			RecompileGlobalShaders();
-		}
-		else if( FCString::Stricmp(*FlagStr,TEXT("Material"))==0)
-		{
-			FString RequestedMaterialName(FParse::Token(Cmd, 0));
-			FRecompileShadersTimer TestTimer(FString::Printf(TEXT("Recompile Material %s"), *RequestedMaterialName));
-
-			ITargetPlatformManagerModule& TPM = GetTargetPlatformManagerRef();
-			FString TargetPlatformName(FParse::Token(Cmd, 0));
-			const ITargetPlatform* TargetPlatform = nullptr;
-			if (TargetPlatformName.Len() > 0)
-			{
-				TargetPlatform = TPM.FindTargetPlatform(TargetPlatformName);
-			}
-
-			bool bMaterialFound = false;
-			for( TObjectIterator<UMaterialInterface> It; It; ++It )
-			{
-				UMaterialInterface* Material = *It;
-				if( Material && Material->GetName() == RequestedMaterialName)
-				{
-					bMaterialFound = true;
-#if WITH_EDITOR
-					// <Pre/Post>EditChange will force a re-creation of the resource,
-					// in turn recompiling the shader.
-					if (TargetPlatform)
-					{
-						Material->BeginCacheForCookedPlatformData(TargetPlatform);
-						while (!Material->IsCachedCookedPlatformDataLoaded(TargetPlatform))
-						{
-							FPlatformProcess::Sleep(0.1f);
-							GShaderCompilingManager->ProcessAsyncResults(false, false);
-						}
-						Material->ClearCachedCookedPlatformData(TargetPlatform);
-					}
-					else
-					{
-						Material->PreEditChange(NULL);
-						Material->PostEditChange();
-					}
-#endif // WITH_EDITOR
-					break;
-				}
-			}
-
-			if (!bMaterialFound)
-			{
-				TestTimer.Stop(false);
-				UE_LOG(LogShaderCompilers, Warning, TEXT("Couldn't find Material %s!"), *RequestedMaterialName);
-			}
-		}
-		else if( FCString::Stricmp(*FlagStr,TEXT("All"))==0)
-		{
-			FRecompileShadersTimer TestTimer(TEXT("RecompileShaders"));
-			RecompileGlobalShaders();
-
-			FMaterialUpdateContext UpdateContext(0);
-			for( TObjectIterator<UMaterial> It; It; ++It )
-			{
-				UMaterial* Material = *It;
-				if( Material )
-				{
-					UE_LOG(LogShaderCompilers, Log, TEXT("recompiling [%s]"),*Material->GetFullName());
-					UpdateContext.AddMaterial(Material);
-#if WITH_EDITOR
-					// <Pre/Post>EditChange will force a re-creation of the resource,
-					// in turn recompiling the shader.
-					Material->PreEditChange(NULL);
-					Material->PostEditChange();
-#endif // WITH_EDITOR
-				}
-			}
-		}
-		else
-		{
-			TArray<const FShaderType*> ShaderTypes = FShaderType::GetShaderTypesByFilename(*FlagStr);
-			TArray<const FShaderPipelineType*> ShaderPipelineTypes = FShaderPipelineType::GetShaderPipelineTypesByFilename(*FlagStr);
-			if (ShaderTypes.Num() > 0 || ShaderPipelineTypes.Num() > 0)
-			{
-				FRecompileShadersTimer TestTimer(TEXT("RecompileShaders SingleShader"));
-				
-				TArray<const FVertexFactoryType*> FactoryTypes;
-
-				UMaterialInterface::IterateOverActiveFeatureLevels([&](ERHIFeatureLevel::Type InFeatureLevel) {
-					auto ShaderPlatform = GShaderPlatformForFeatureLevel[InFeatureLevel];
-					BeginRecompileGlobalShaders(ShaderTypes, ShaderPipelineTypes, ShaderPlatform);
-					//UMaterial::UpdateMaterialShaders(ShaderTypes, ShaderPipelineTypes, FactoryTypes, ShaderPlatform);
-					FinishRecompileGlobalShaders();
-				});
-			}
-		}
-
-		GWarn->EndSlowTask();
-
-		return 1;
-	}
-
-	UE_LOG(LogShaderCompilers, Warning, TEXT("Invalid parameter. Options are: \n'Changed', 'Global', 'Material [name]', 'All'."));
-	return 1;
-}
-
 static void PrepareGlobalShaderCompileJob(EShaderPlatform Platform,
 	EShaderPermutationFlags PermutationFlags,
 	const FShaderPipelineType* ShaderPipeline,
@@ -6992,286 +6594,6 @@ static FAutoConsoleCommand CCmdSetGlobalShaderCacheOverrideDirectory = FAutoCons
 	TEXT("Set the directory to read the override global shader map file from."),
 	FConsoleCommandWithArgsDelegate::CreateStatic(SetGlobalShaderCacheOverrideDirectory));
 
-bool RecompileChangedShadersForPlatform(const FString& PlatformName)
-{
-	// figure out what shader platforms to recompile
-	ITargetPlatformManagerModule* TPM = GetTargetPlatformManager();
-	ITargetPlatform* TargetPlatform = TPM->FindTargetPlatform(PlatformName);
-	if (TargetPlatform == NULL)
-	{
-		UE_LOG(LogShaders, Display, TEXT("Failed to find target platform module for %s"), *PlatformName);
-		return false;
-	}
-
-	TArray<FName> DesiredShaderFormats;
-	TargetPlatform->GetAllTargetedShaderFormats(DesiredShaderFormats);
-
-
-
-	// figure out which shaders are out of date
-	TArray<const FShaderType*> OutdatedShaderTypes;
-	TArray<const FVertexFactoryType*> OutdatedFactoryTypes;
-	TArray<const FShaderPipelineType*> OutdatedShaderPipelineTypes;
-
-	// Pick up new changes to shader files
-	FlushShaderFileCache();
-
-	GetOutdatedShaderTypes(OutdatedShaderTypes, OutdatedShaderPipelineTypes, OutdatedFactoryTypes);
-	UE_LOG(LogShaders, Display, TEXT("We found %d out of date shader types, %d outdated pipeline types, and %d out of date VF types!"), OutdatedShaderTypes.Num(), OutdatedShaderPipelineTypes.Num(), OutdatedFactoryTypes.Num());
-
-	for (int32 FormatIndex = 0; FormatIndex < DesiredShaderFormats.Num(); FormatIndex++)
-	{
-		// get the shader platform enum
-		const EShaderPlatform ShaderPlatform = ShaderFormatToLegacyShaderPlatform(DesiredShaderFormats[FormatIndex]);
-
-		// Only compile for the desired platform if requested
-		// Kick off global shader recompiles
-		BeginRecompileGlobalShaders(OutdatedShaderTypes, OutdatedShaderPipelineTypes, ShaderPlatform);
-
-		// Block on global shaders
-		FinishRecompileGlobalShaders();
-#if WITH_EDITOR
-		// we only want to actually compile mesh shaders if we have out of date ones
-		if (OutdatedShaderTypes.Num() || OutdatedFactoryTypes.Num())
-		{
-			for (TObjectIterator<UMaterialInterface> It; It; ++It)
-			{
-				(*It)->ClearCachedCookedPlatformData(TargetPlatform);
-			}
-		}
-#endif
-	}
-
-	if (OutdatedFactoryTypes.Num() || OutdatedShaderTypes.Num())
-	{
-		return true;
-	}
-	return false;
-}
-
-FArchive& operator<<(FArchive& Ar, FODSCRequestPayload& Elem)
-{
-	uint32 ConvertedShaderPlatform = (uint32)Elem.ShaderPlatform;
-	Ar << ConvertedShaderPlatform;
-	Ar << Elem.MaterialName;
-	Ar << Elem.VertexFactoryName;
-	Ar << Elem.PipelineName;
-	Ar << Elem.ShaderTypeNames;
-	Ar << Elem.RequestHash;
-
-	return Ar;
-}
-
-FShaderRecompileData::FShaderRecompileData(const FString& InPlatformName, TArray<FString>* OutModifiedFiles, TArray<uint8>* OutMeshMaterialMaps, TArray<uint8>* OutGlobalShaderMap)
-: PlatformName(InPlatformName),
-  ModifiedFiles(OutModifiedFiles),
-  MeshMaterialMaps(OutMeshMaterialMaps),
-  GlobalShaderMap(OutGlobalShaderMap)
-{
-}
-
-FShaderRecompileData::FShaderRecompileData(const FString& InPlatformName, EShaderPlatform InShaderPlatform, ODSCRecompileCommand InCommandType, TArray<FString>* OutModifiedFiles, TArray<uint8>* OutMeshMaterialMaps, TArray<uint8>* OutGlobalShaderMap)
-: PlatformName(InPlatformName),
-  ShaderPlatform(InShaderPlatform),
-  ModifiedFiles(OutModifiedFiles),
-  MeshMaterialMaps(OutMeshMaterialMaps),
-  CommandType(InCommandType),
-  GlobalShaderMap(OutGlobalShaderMap)
-{
-}
-
-
-#if WITH_EDITOR
-void RecompileShadersForRemote(
-	FShaderRecompileData& Args,
-	const FString& OutputDirectory)
-{
-	// figure out what shader platforms to recompile
-	ITargetPlatformManagerModule* TPM = GetTargetPlatformManager();
-	ITargetPlatform* TargetPlatform = TPM->FindTargetPlatform(Args.PlatformName);
-	if (TargetPlatform == NULL)
-	{
-		UE_LOG(LogShaders, Display, TEXT("Failed to find target platform module for %s"), *Args.PlatformName);
-		return;
-	}
-
-	const double StartTime = FPlatformTime::Seconds();
-
-	UE_LOG(LogShaders, Display, TEXT(""));
-	UE_LOG(LogShaders, Display, TEXT("********************************"));
-	UE_LOG(LogShaders, Display, TEXT("Received compile shader request."));
-
-	const bool bPreviousState = GShaderCompilingManager->IsShaderCompilationSkipped();
-	GShaderCompilingManager->SkipShaderCompilation(false);
-
-	TArray<FName> DesiredShaderFormats;
-	TargetPlatform->GetAllTargetedShaderFormats(DesiredShaderFormats);
-
-	UE_LOG(LogShaders, Verbose, TEXT("Loading %d materials..."), Args.MaterialsToLoad.Num());
-	// make sure all materials the client has loaded will be processed
-	TArray<UMaterialInterface*> MaterialsToCompile;
-
-	for (int32 Index = 0; Index < Args.MaterialsToLoad.Num(); Index++)
-	{
-		UE_LOG(LogShaders, Verbose, TEXT("   --> %s"), *Args.MaterialsToLoad[Index]);
-		MaterialsToCompile.Add(LoadObject<UMaterialInterface>(NULL, *Args.MaterialsToLoad[Index]));
-	}
-
-	UE_LOG(LogShaders, Verbose, TEXT("  Done!"))
-
-	// figure out which shaders are out of date
-	TArray<const FShaderType*> OutdatedShaderTypes;
-	TArray<const FVertexFactoryType*> OutdatedFactoryTypes;
-	TArray<const FShaderPipelineType*> OutdatedShaderPipelineTypes;
-
-	// Pick up new changes to shader files
-	FlushShaderFileCache();
-
-	if (Args.ShadersToRecompile.Num())
-	{
-		UE_LOG(LogShaders, Display, TEXT("Received %d shaders to compile."), Args.ShadersToRecompile.Num());
-	}
-
-	for (const FODSCRequestPayload& payload: Args.ShadersToRecompile)
-	{
-		UE_LOG(LogShaders, Display, TEXT(""));
-		UE_LOG(LogShaders, Display, TEXT("\tMaterial:    %s "), *payload.MaterialName);
-		UE_LOG(LogShaders, Display, TEXT("\tVF Type:     %s "), *payload.VertexFactoryName);
-
-		MaterialsToCompile.Add(LoadObject<UMaterialInterface>(NULL, *payload.MaterialName));
-
-		const FVertexFactoryType* VFType = FVertexFactoryType::GetVFByName(payload.VertexFactoryName);
-		if (VFType)
-		{
-			OutdatedFactoryTypes.Add(VFType);
-		}
-
-		const FShaderPipelineType* PipelineType = FShaderPipelineType::GetShaderPipelineTypeByName(payload.PipelineName);
-		if (PipelineType)
-		{
-			OutdatedShaderPipelineTypes.Add(PipelineType);
-		}
-
-		for (const FString& ShaderTypeName : payload.ShaderTypeNames)
-		{
-			UE_LOG(LogShaders, Display, TEXT("\tShader Type: %s"), *ShaderTypeName);
-
-			const FShaderType* ShaderType = FShaderType::GetShaderTypeByName(*ShaderTypeName);
-			if (ShaderType)
-			{
-				OutdatedShaderTypes.Add(ShaderType);
-			}
-		}
-	}
-
-	{
-		for (int32 FormatIndex = 0; FormatIndex < DesiredShaderFormats.Num(); FormatIndex++)
-		{
-			// get the shader platform enum
-			const EShaderPlatform ShaderPlatform = ShaderFormatToLegacyShaderPlatform(DesiredShaderFormats[FormatIndex]);
-
-			// Only compile for the desired platform if requested
-			if (ShaderPlatform == Args.ShaderPlatform || Args.ShaderPlatform == SP_NumPlatforms)
-			{
-				// If we are explicitly wanting to recompile global or if shaders have changed.
-				if (Args.CommandType == ODSCRecompileCommand::Global ||
-					Args.CommandType == ODSCRecompileCommand::Changed)
-				{
-					UE_LOG(LogShaders, Display, TEXT("Recompiling global shaders."));
-
-					// Explicitly get outdated types for global shaders.
-					const FGlobalShaderMap* ShaderMap = GGlobalShaderMap[ShaderPlatform];
-					if (ShaderMap)
-					{
-						ShaderMap->GetOutdatedTypes(OutdatedShaderTypes, OutdatedShaderPipelineTypes, OutdatedFactoryTypes);
-					}
-
-					UE_LOG(LogShaders, Display, TEXT("\tFound %d outdated shader types."), OutdatedShaderTypes.Num() + OutdatedShaderPipelineTypes.Num());
-
-					// Kick off global shader recompiles
-					BeginRecompileGlobalShaders(OutdatedShaderTypes, OutdatedShaderPipelineTypes, ShaderPlatform, TargetPlatform);
-
-					// Block on global shaders
-					FinishRecompileGlobalShaders();
-
-					// write the shader compilation info to memory, converting fnames to strings
-					FMemoryWriter MemWriter(*Args.GlobalShaderMap, true);
-					FNameAsStringProxyArchive Ar(MemWriter);
-					Ar.SetCookingTarget(TargetPlatform);
-
-					// save out the global shader map to the byte array
-					SaveGlobalShadersForRemoteRecompile(Ar, ShaderPlatform);
-				}
-
-				// we only want to actually compile mesh shaders if a client directly requested it
-				if ((Args.CommandType == ODSCRecompileCommand::Material || Args.CommandType == ODSCRecompileCommand::Changed) &&
-					Args.MeshMaterialMaps != nullptr)
-				{
-					TMap<FString, TArray<TRefCountPtr<FMaterialShaderMap> > > CompiledShaderMaps;
-					UMaterial::CompileMaterialsForRemoteRecompile(MaterialsToCompile, ShaderPlatform, TargetPlatform, CompiledShaderMaps);
-
-					// write the shader compilation info to memory, converting fnames to strings
-					FMemoryWriter MemWriter(*Args.MeshMaterialMaps, true);
-					FNameAsStringProxyArchive Ar(MemWriter);
-					Ar.SetCookingTarget(TargetPlatform);
-
-					// save out the shader map to the byte array
-					FMaterialShaderMap::SaveForRemoteRecompile(Ar, CompiledShaderMaps);
-				}
-
-				// save it out so the client can get it (and it's up to date next time)
-				FString GlobalShaderFilename = SaveGlobalShaderFile(ShaderPlatform, OutputDirectory, TargetPlatform);
-
-				// add this to the list of files to tell the other end about
-				if (Args.ModifiedFiles)
-				{
-					// need to put it in non-sandbox terms
-					FString SandboxPath(GlobalShaderFilename);
-					check(SandboxPath.StartsWith(OutputDirectory));
-					SandboxPath.ReplaceInline(*OutputDirectory, TEXT("../../../"));
-					FPaths::NormalizeFilename(SandboxPath);
-					Args.ModifiedFiles->Add(SandboxPath);
-				}
-			}
-		}
-	}
-
-	UE_LOG(LogShaders, Display, TEXT(""));
-	UE_LOG(LogShaders, Display, TEXT("Finished shader compile request in %.2f seconds."), FPlatformTime::Seconds() - StartTime);
-
-	// Provide a log of what happened.
-	GShaderCompilingManager->PrintStats(true);
-
-	// Restore compilation state.
-	GShaderCompilingManager->SkipShaderCompilation(bPreviousState);
-}
-#endif // WITH_EDITOR
-
-void BeginRecompileGlobalShaders(const TArray<const FShaderType*>& OutdatedShaderTypes, const TArray<const FShaderPipelineType*>& OutdatedShaderPipelineTypes, EShaderPlatform ShaderPlatform, const ITargetPlatform* TargetPlatform)
-{
-	if (!FPlatformProperties::RequiresCookedData())
-	{
-		// Flush pending accesses to the existing global shaders.
-		FlushRenderingCommands();
-
-		// Calling CompileGlobalShaderMap will force starting the compile jobs if the map is empty (by calling VerifyGlobalShaders)
-		CompileGlobalShaderMap(ShaderPlatform, TargetPlatform, false);
-		FGlobalShaderMap* GlobalShaderMap = GetGlobalShaderMap(ShaderPlatform);
-
-		// Now check if there is any work to be done wrt outdates types
-		if (OutdatedShaderTypes.Num() > 0 || OutdatedShaderPipelineTypes.Num() > 0)
-		{
-			VerifyGlobalShaders(ShaderPlatform, TargetPlatform, false, &OutdatedShaderTypes, &OutdatedShaderPipelineTypes);
-		}
-	}
-}
-
-void FinishRecompileGlobalShaders()
-{
-	// Block until global shaders have been compiled and processed
-	GShaderCompilingManager->ProcessAsyncResults(false, true);
-}
-
 static inline FShader* ProcessCompiledJob(FShaderCompileJob* SingleJob, const FShaderPipelineType* Pipeline, TArray<EShaderPlatform>& ShaderPlatformsProcessed, TArray<const FShaderPipelineType*>& OutSharedPipelines)
 {
 	const FGlobalShaderType* GlobalShaderType = SingleJob->Key.ShaderType->GetGlobalShaderType();
@@ -7401,58 +6723,6 @@ void ProcessCompiledGlobalShaders(const TArray<FShaderCommonCompileJobPtr>& Comp
 	}
 }
 
-void SaveGlobalShadersForRemoteRecompile(FArchive& Ar, EShaderPlatform ShaderPlatform)
-{
-	FGlobalShaderMap* GlobalShaderMap = GetGlobalShaderMap(ShaderPlatform);
-	uint8 bIsValid = GlobalShaderMap != nullptr;
-	Ar << bIsValid;
-
-	if (GlobalShaderMap)
-	{
-		GlobalShaderMap->SaveToGlobalArchive(Ar);
-	}
-}
-
-void LoadGlobalShadersForRemoteRecompile(FArchive& Ar, EShaderPlatform ShaderPlatform)
-{
-	uint8 bIsValid = 0;
-	Ar << bIsValid;
-
-	if (bIsValid)
-	{
-		FlushRenderingCommands();
-
-		FGlobalShaderMap* NewGlobalShaderMap = new FGlobalShaderMap(ShaderPlatform);
-		if (NewGlobalShaderMap)
-		{
-			NewGlobalShaderMap->LoadFromGlobalArchive(Ar);
-
-			if (GGlobalShaderMap[ShaderPlatform])
-			{
-				GGlobalShaderMap[ShaderPlatform]->ReleaseAllSections();
-
-				delete GGlobalShaderMap[ShaderPlatform];
-				GGlobalShaderMap[ShaderPlatform] = nullptr;
-				GGlobalShaderMap[ShaderPlatform] = NewGlobalShaderMap;
-
-				VerifyGlobalShaders(ShaderPlatform, nullptr, false);
-
-				// Invalidate global bound shader states so they will be created with the new shaders the next time they are set (in SetGlobalBoundShaderState)
-				for (TLinkedList<FGlobalBoundShaderStateResource*>::TIterator It(FGlobalBoundShaderStateResource::GetGlobalBoundShaderStateList()); It; It.Next())
-				{
-					BeginUpdateResourceRHI(*It);
-				}
-
-				PropagateGlobalShadersToAllPrimitives();
-			}
-			else
-			{
-				delete NewGlobalShaderMap;
-			}
-		}
-	}
-}
-
 FShaderCommonCompileJob::FInputHash FShaderCompileJob::GetInputHash()
 {
 	if (bInputHashSet)
diff --git a/Engine/Source/Runtime/Engine/Private/ShaderCompiler/ShaderCompilerShared.cpp b/Engine/Source/Runtime/Engine/Private/ShaderCompiler/ShaderCompilerShared.cpp
new file mode 100644
--- /dev/null
+++ b/Engine/Source/Runtime/Engine/Private/ShaderCompiler/ShaderCompilerShared.cpp
@@ -0,0 +1,39 @@
+// Copyright Epic Games, Inc. All Rights Reserved.
+
+#include "ShaderCompilerShared.h"
+
+#include "ComponentRecreateRenderStateContext.h"
+#include "ObjectCacheContext.h"
+#if WITH_EDITOR
+#include "Rendering/StaticLightingSystemInterface.h"
+#endif
+
+void PropagateGlobalShadersToAllPrimitives()
+{
+	// Re-register everything to work around FShader lifetime issues - it currently lives and dies with the
+	// shadermap it is stored in, while cached MDCs can reference its memory. Re-registering will
+	// re-create the cache.
+	TRACE_CPUPROFILER_EVENT_SCOPE(PropagateGlobalShadersToAllPrimitives);
+
+	FObjectCacheContextScope ObjectCacheScope;
+	TSet<FSceneInterface*> ScenesToUpdate;
+	TIndirectArray<FComponentRecreateRenderStateContext> ComponentContexts;
+	for (UPrimitiveComponent* PrimitiveComponent : ObjectCacheScope.GetContext().GetPrimitiveComponents())
+	{
+		if (PrimitiveComponent->IsRenderStateCreated())
+		{
+			ComponentContexts.Add(new FComponentRecreateRenderStateContext(PrimitiveComponent, &ScenesToUpdate));
+#if WITH_EDITOR
+			if (PrimitiveComponent->HasValidSettingsForStaticLighting(false))
+			{
+				FStaticLightingSystemInterface::OnPrimitiveComponentUnregistered.Broadcast(PrimitiveComponent);
+				FStaticLightingSystemInterface::OnPrimitiveComponentRegistered.Broadcast(PrimitiveComponent);
+			}
+#endif
+		}
+	}
+
+	UpdateAllPrimitiveSceneInfosForScenes(ScenesToUpdate);
+	ComponentContexts.Empty();
+	UpdateAllPrimitiveSceneInfosForScenes(ScenesToUpdate);
+}
diff --git a/Engine/Source/Runtime/Engine/Private/ShaderCompiler/ShaderCompilerShared.h b/Engine/Source/Runtime/Engine/Private/ShaderCompiler/ShaderCompilerShared.h
new file mode 100644
--- /dev/null
+++ b/Engine/Source/Runtime/Engine/Private/ShaderCompiler/ShaderCompilerShared.h
@@ -0,0 +1,5 @@
+// Copyright Epic Games, Inc. All Rights Reserved.
+
+#pragma once
+
+void PropagateGlobalShadersToAllPrimitives();
diff --git a/Engine/Source/Runtime/Engine/Public/ODSC/ODSCManager.h b/Engine/Source/Runtime/Engine/Public/ODSC/ODSCManager.h
--- a/Engine/Source/Runtime/Engine/Public/ODSC/ODSCManager.h
+++ b/Engine/Source/Runtime/Engine/Public/ODSC/ODSCManager.h
@@ -6,7 +6,7 @@
 #include "Containers/Ticker.h"
 #include "RHIDefinitions.h"
 #include "Containers/Ticker.h"
-#include "ShaderCompiler.h"
+#include "RecompileShaders.h"
 
 class FODSCThread;
 
@@ -72,4 +72,4 @@ private:
 };
 
 /** The global shader ODSC manager. */
-extern ENGINE_API FODSCManager* GODSCManager;
\ No newline at end of file
+extern ENGINE_API FODSCManager* GODSCManager;
diff --git a/Engine/Source/Runtime/Engine/Public/RecompileShaders.h b/Engine/Source/Runtime/Engine/Public/RecompileShaders.h
new file mode 100644
--- /dev/null
+++ b/Engine/Source/Runtime/Engine/Public/RecompileShaders.h
@@ -0,0 +1,127 @@
+// Copyright Epic Games, Inc. All Rights Reserved.
+
+#pragma once
+
+#include "CoreMinimal.h"
+#include "ShaderCore.h"
+#include "ShaderCompilerCore.h"
+#include "Shader.h"
+
+/**
+* Handles serializing in MeshMaterialMaps or GlobalShaderMap from a CookOnTheFly command and applying them to the in-memory shadermaps.
+*
+* @param MeshMaterialMaps				Byte array that contains the serialized material shadermap from across the network.
+* @param MaterialsToLoad				The materials contained in the MeshMaterialMaps
+* @param GlobalShaderMap				Byte array that contains the serialized global shadermap from across the network.
+**/
+extern ENGINE_API void ProcessCookOnTheFlyShaders(bool bReloadGlobalShaders, const TArray<uint8>& MeshMaterialMaps, const TArray<FString>& MaterialsToLoad, const TArray<uint8>& GlobalShaderMap);
+
+struct FODSCRequestPayload
+{
+	/** The shader platform to compile for. */
+	EShaderPlatform ShaderPlatform;
+
+	/** Which material do we compile for?. */
+	FString MaterialName;
+
+	/** The vertex factory type name to compile shaders for. */
+	FString VertexFactoryName;
+
+	/** The name of the pipeline to compile shaders for. */
+	FString PipelineName;
+
+	/** An array of shader type names for each stage in the Pipeline. */
+	TArray<FString> ShaderTypeNames;
+
+	/** A hash of the above information to uniquely identify a Request. */
+	FString RequestHash;
+
+	ENGINE_API FODSCRequestPayload() {};
+	ENGINE_API FODSCRequestPayload(EShaderPlatform InShaderPlatform, const FString& InMaterialName, const FString& InVertexFactoryName, const FString& InPipelineName, const TArray<FString>& InShaderTypeNames, const FString& InRequestHash);
+
+	/**
+	* Serializes FODSCRequestPayload value from or into this archive.
+	*
+	* @param Ar The archive to serialize to.
+	* @param Value The value to serialize.
+	* @return The archive.
+	*/
+	ENGINE_API friend FArchive& operator<<(FArchive& Ar, FODSCRequestPayload& Elem);
+};
+
+enum class ODSCRecompileCommand
+{
+	None,
+	Changed,
+	Global,
+	Material
+};
+
+struct FShaderRecompileData
+{
+	/** The platform name to compile for. */
+	FString PlatformName;
+
+	/** Shader platform */
+	EShaderPlatform ShaderPlatform = SP_NumPlatforms;
+
+	/** All filenames that have been changed during the shader compilation. */
+	TArray<FString>* ModifiedFiles = nullptr;
+
+	/** Mesh materials, returned to the caller.  */
+	TArray<uint8>* MeshMaterialMaps = nullptr;
+
+	/** Materials to load. */
+	TArray<FString> MaterialsToLoad;
+
+	/** What type of shaders to recompile. All, Changed, Global, or Material? */
+	ODSCRecompileCommand CommandType = ODSCRecompileCommand::Changed;
+
+	/** Global shader map, returned to the caller.  */
+	TArray<uint8>* GlobalShaderMap = nullptr;
+
+	/** On-demand shader compiler payload.  */
+	TArray<FODSCRequestPayload> ShadersToRecompile;
+
+	/** Default constructor. */
+	ENGINE_API FShaderRecompileData() {};
+
+	/** Recompile all the changed shaders for the current platform. */
+	ENGINE_API FShaderRecompileData(const FString& InPlatformName, TArray<FString>* OutModifiedFiles, TArray<uint8>* OutMeshMaterialMaps, TArray<uint8>* OutGlobalShaderMap);
+
+	/** For recompiling just global shaders. */
+	ENGINE_API FShaderRecompileData(const FString& InPlatformName, EShaderPlatform InShaderPlatform, ODSCRecompileCommand InCommandType, TArray<FString>* OutModifiedFiles, TArray<uint8>* OutMeshMaterialMaps, TArray<uint8>* OutGlobalShaderMap);
+
+	FShaderRecompileData& operator=(const FShaderRecompileData& Other)
+	{
+		PlatformName = Other.PlatformName;
+		ShaderPlatform = Other.ShaderPlatform;
+		ModifiedFiles = Other.ModifiedFiles;
+		MeshMaterialMaps = Other.MeshMaterialMaps;
+		MaterialsToLoad = Other.MaterialsToLoad;
+		CommandType = Other.CommandType;
+		GlobalShaderMap = Other.GlobalShaderMap;
+
+		ShadersToRecompile = Other.ShadersToRecompile;
+
+		return *this;
+	}
+};
+
+/** Serializes a global shader map to an archive (used with recompiling shaders for a remote console) */
+extern ENGINE_API void SaveGlobalShadersForRemoteRecompile(FArchive& Ar, EShaderPlatform ShaderPlatform);
+
+/** Serializes a global shader map to an archive (used with recompiling shaders for a remote console) */
+extern ENGINE_API void LoadGlobalShadersForRemoteRecompile(FArchive& Ar, EShaderPlatform ShaderPlatform);
+
+#if WITH_EDITOR
+
+/**
+* Recompiles global shaders
+*
+* @param Args					Arguments and configuration for issuing recompiles.
+* @param OutputDirectory		The directory the compiled data will be stored to
+**/
+extern ENGINE_API void RecompileShadersForRemote(FShaderRecompileData& Args, const FString& OutputDirectory);
+
+#endif // WITH_EDITOR
\ No newline at end of file
diff --git a/Engine/Source/Runtime/Engine/Public/ShaderCompiler.h b/Engine/Source/Runtime/Engine/Public/ShaderCompiler.h
--- a/Engine/Source/Runtime/Engine/Public/ShaderCompiler.h
+++ b/Engine/Source/Runtime/Engine/Public/ShaderCompiler.h
@@ -1234,122 +1234,12 @@ extern ENGINE_API void FinishRecompileGlobalShaders();
 /** Called by the shader compiler to process completed global shader jobs. */
 extern ENGINE_API void ProcessCompiledGlobalShaders(const TArray<FShaderCommonCompileJobPtr>& CompilationResults);
 
-/** Serializes a global shader map to an archive (used with recompiling shaders for a remote console) */
-extern ENGINE_API void SaveGlobalShadersForRemoteRecompile(FArchive& Ar, EShaderPlatform ShaderPlatform);
-
-/** Serializes a global shader map to an archive (used with recompiling shaders for a remote console) */
-extern ENGINE_API void LoadGlobalShadersForRemoteRecompile(FArchive& Ar, EShaderPlatform ShaderPlatform);
-
 /**
 * Saves the global shader map as a file for the target platform.
 * @return the name of the file written
 */
 extern ENGINE_API FString SaveGlobalShaderFile(EShaderPlatform Platform, FString SavePath, class ITargetPlatform* TargetPlatform = nullptr);
 
-struct FODSCRequestPayload
-{
-	/** The shader platform to compile for. */
-	EShaderPlatform ShaderPlatform;
-
-	/** Which material do we compile for?. */
-	FString MaterialName;
-
-	/** The vertex factory type name to compile shaders for. */
-	FString VertexFactoryName;
-
-	/** The name of the pipeline to compile shaders for. */
-	FString PipelineName;
-
-	/** An array of shader type names for each stage in the Pipeline. */
-	TArray<FString> ShaderTypeNames;
-
-	/** A hash of the above information to uniquely identify a Request. */
-	FString RequestHash;
-
-	ENGINE_API FODSCRequestPayload() {};
-	ENGINE_API FODSCRequestPayload(EShaderPlatform InShaderPlatform, const FString& InMaterialName, const FString& InVertexFactoryName, const FString& InPipelineName, const TArray<FString>& InShaderTypeNames, const FString& InRequestHash);
-
-	/**
-	* Serializes FODSCRequestPayload value from or into this archive.
-	*
-	* @param Ar The archive to serialize to.
-	* @param Value The value to serialize.
-	* @return The archive.
-	*/
-	ENGINE_API friend FArchive& operator<<(FArchive& Ar, FODSCRequestPayload& Elem);
-};
-
-enum class ODSCRecompileCommand
-{
-	None,
-	Changed,
-	Global,
-	Material
-};
-
-struct FShaderRecompileData
-{
-	/** The platform name to compile for. */
-	FString PlatformName;
-
-	/** Shader platform */
-	EShaderPlatform ShaderPlatform = SP_NumPlatforms;
-
-	/** All filenames that have been changed during the shader compilation. */
-	TArray<FString>* ModifiedFiles = nullptr;
-
-	/** Mesh materials, returned to the caller.  */
-	TArray<uint8>* MeshMaterialMaps = nullptr;
-
-	/** Materials to load. */
-	TArray<FString> MaterialsToLoad;
-
-	/** What type of shaders to recompile. All, Changed, Global, or Material? */
-	ODSCRecompileCommand CommandType = ODSCRecompileCommand::Changed;
-
-	/** Global shader map, returned to the caller.  */
-	TArray<uint8>* GlobalShaderMap = nullptr;
-
-	/** On-demand shader compiler payload.  */
-	TArray<FODSCRequestPayload> ShadersToRecompile;
-
-	/** Default constructor. */
-	ENGINE_API FShaderRecompileData() {};
-
-	/** Recompile all the changed shaders for the current platform. */
-	ENGINE_API FShaderRecompileData(const FString& InPlatformName, TArray<FString>* OutModifiedFiles, TArray<uint8>* OutMeshMaterialMaps, TArray<uint8>* OutGlobalShaderMap);
-
-	/** For recompiling just global shaders. */
-	ENGINE_API FShaderRecompileData(const FString& InPlatformName, EShaderPlatform InShaderPlatform, ODSCRecompileCommand InCommandType, TArray<FString>* OutModifiedFiles, TArray<uint8>* OutMeshMaterialMaps, TArray<uint8>* OutGlobalShaderMap);
-
-	FShaderRecompileData& operator=(const FShaderRecompileData& Other)
-	{
-		PlatformName = Other.PlatformName;
-		ShaderPlatform = Other.ShaderPlatform;
-		ModifiedFiles = Other.ModifiedFiles;
-		MeshMaterialMaps = Other.MeshMaterialMaps;
-		MaterialsToLoad = Other.MaterialsToLoad;
-		CommandType = Other.CommandType;
-		GlobalShaderMap = Other.GlobalShaderMap;
-
-		ShadersToRecompile = Other.ShadersToRecompile;
-
-		return *this;
-	}
-};
-
-#if WITH_EDITOR
-
-/**
-* Recompiles global shaders
-*
-* @param Args					Arguments and configuration for issuing recompiles.
-* @param OutputDirectory		The directory the compiled data will be stored to
-**/
-extern ENGINE_API void RecompileShadersForRemote(FShaderRecompileData& Args, const FString& OutputDirectory);
-
-#endif // WITH_EDITOR
-
 extern ENGINE_API void CompileGlobalShaderMap(bool bRefreshShaderMap=false);
 extern ENGINE_API void CompileGlobalShaderMap(ERHIFeatureLevel::Type InFeatureLevel, bool bRefreshShaderMap=false);
 extern ENGINE_API void CompileGlobalShaderMap(EShaderPlatform Platform, bool bRefreshShaderMap = false);
@@ -1358,12 +1248,3 @@ extern ENGINE_API void CompileGlobalShaderMap(EShaderPlatform Platform, const IT
 extern ENGINE_API FString GetGlobalShaderMapDDCKey();
 
 extern ENGINE_API FString GetMaterialShaderMapDDCKey();
-
-/**
-* Handles serializing in MeshMaterialMaps or GlobalShaderMap from a CookOnTheFly command and applying them to the in-memory shadermaps.
-*
-* @param MeshMaterialMaps				Byte array that contains the serialized material shadermap from across the network.
-* @param MaterialsToLoad				The materials contained in the MeshMaterialMaps
-* @param GlobalShaderMap				Byte array that contains the serialized global shadermap from across the network.
-**/
-extern ENGINE_API void ProcessCookOnTheFlyShaders(bool bReloadGlobalShaders, const TArray<uint8>& MeshMaterialMaps, const TArray<FString>& MaterialsToLoad, const TArray<uint8>& GlobalShaderMap);
diff --git a/Engine/Source/Runtime/NetworkFileSystem/Private/NetworkFileServerConnection.cpp b/Engine/Source/Runtime/NetworkFileSystem/Private/NetworkFileServerConnection.cpp
--- a/Engine/Source/Runtime/NetworkFileSystem/Private/NetworkFileServerConnection.cpp
+++ b/Engine/Source/Runtime/NetworkFileSystem/Private/NetworkFileServerConnection.cpp
@@ -15,7 +15,7 @@
 #include "Interfaces/ITargetPlatform.h"
 #include "HAL/PlatformTime.h"
 #include "Interfaces/IPluginManager.h"
-
+#include "RecompileShaders.h"
 
 /**
  * Helper function for resolving engine and game sandbox paths
diff --git a/Engine/Source/Runtime/NetworkFileSystem/Public/INetworkFileSystemModule.h b/Engine/Source/Runtime/NetworkFileSystem/Public/INetworkFileSystemModule.h
--- a/Engine/Source/Runtime/NetworkFileSystem/Public/INetworkFileSystemModule.h
+++ b/Engine/Source/Runtime/NetworkFileSystem/Public/INetworkFileSystemModule.h
@@ -4,9 +4,9 @@
 
 #include "CoreMinimal.h"
 #include "Modules/ModuleInterface.h"
-#include "ShaderCompiler.h"
 
 class INetworkFileServer;
+struct FShaderRecompileData;
 
 /**
  * Delegate type for handling file requests from a network client.
