Basic support for materials that access computed normal attribute in new HLSL translator
Handled by splitting generated code into 2 phases; first one that computes normal, followed by all other attributes
Code for each phase is generated from the same tree, but properly dead-stripped, so only relevant code is generated for each phase

diff --git a/Engine/Shaders/Private/MaterialTemplate.ush b/Engine/Shaders/Private/MaterialTemplate.ush
--- a/Engine/Shaders/Private/MaterialTemplate.ush
+++ b/Engine/Shaders/Private/MaterialTemplate.ush
@@ -2720,13 +2720,8 @@ FMaterialAttributes EvaluateVertexMaterialAttributes(FMaterialVertexParameters P
 	return DefaultMaterialAttributes;
 }
 
-FMaterialAttributes EvaluatePixelMaterialAttributes(in out FMaterialPixelParameters Parameters)
-{
-	FMaterialAttributes DefaultMaterialAttributes = (FMaterialAttributes)0;
+// EvaluatePixelMaterialAttributes
 %s
-	// Return the default, this may be skipped if the above code returns different attributes
-	return DefaultMaterialAttributes;
-}
 
 // Be very, very careful about changing CalcPixelMaterialInputs() or CalcPixelMaterialInputsAnalyticDerivatives(). They apply the same basic calculation,
 // but CalcPixelMaterialInputs() uses finite difference derivates from hardware whereas CalcPixelMaterialInputsAnalyticDerivatives() uses analytic
diff --git a/Engine/Source/Runtime/Engine/Classes/Materials/MaterialExpression.h b/Engine/Source/Runtime/Engine/Classes/Materials/MaterialExpression.h
--- a/Engine/Source/Runtime/Engine/Classes/Materials/MaterialExpression.h
+++ b/Engine/Source/Runtime/Engine/Classes/Materials/MaterialExpression.h
@@ -129,7 +129,7 @@ public:
 	ENGINE_API bool GenerateHLSLStatements(FMaterialHLSLGenerator& Generator, UE::HLSLTree::FScope& Scope) const;
 
 	/** Creates a new scope, and populates it with the expression connected to this input */
-	ENGINE_API UE::HLSLTree::FScope* NewScopeWithStatements(FMaterialHLSLGenerator& Generator, UE::HLSLTree::FScope& Scope) const;
+	ENGINE_API UE::HLSLTree::FScope* NewOwnedScopeWithStatements(FMaterialHLSLGenerator& Generator, UE::HLSLTree::FStatement& Owner) const;
 	ENGINE_API UE::HLSLTree::FScope* NewScopeWithStatements(FMaterialHLSLGenerator& Generator, UE::HLSLTree::FScope& Scope, EMaterialNewScopeFlag Flags) const;
 #endif // WITH_EDITOR
 private:
diff --git a/Engine/Source/Runtime/Engine/Private/HLSLTree/HLSLTree.cpp b/Engine/Source/Runtime/Engine/Private/HLSLTree/HLSLTree.cpp
--- a/Engine/Source/Runtime/Engine/Private/HLSLTree/HLSLTree.cpp
+++ b/Engine/Source/Runtime/Engine/Private/HLSLTree/HLSLTree.cpp
@@ -234,8 +234,18 @@ const TCHAR* FEmitContext::CastShaderValue(const FNode* Node, const TCHAR* Code,
 	}
 }
 
-void FEmitContext::AddPreshader(Shader::EValueType Type, const Shader::FPreshaderData& Preshader, FStringBuilderBase& OutCode)
-{
+const TCHAR* FEmitContext::AcquirePreshader(Shader::EValueType Type, const Shader::FPreshaderData& Preshader)
+{
+	FSHA1 Hasher;
+	Hasher.Update((uint8*)&Type, sizeof(Type));
+	Preshader.AppendHash(Hasher);
+	const FSHAHash Hash = Hasher.Finalize();
+	TCHAR const* const* PrevPreshader = Preshaders.Find(Hash);
+	if (PrevPreshader)
+	{
+		return *PrevPreshader;
+	}
+
 	const Shader::FValueTypeDescription TypeDesc = Shader::GetValueTypeDescription(Type);
 	ensure(TypeDesc.ComponentType == Shader::EValueComponentType::Float || TypeDesc.ComponentType == Shader::EValueComponentType::Double);
 
@@ -256,29 +266,34 @@ void FEmitContext::AddPreshader(Shader::EValueType Type, const Shader::FPreshade
 	PreshaderHeader.ComponentType = TypeDesc.ComponentType;
 	PreshaderHeader.NumComponents = TypeDesc.NumComponents;
 
+	TStringBuilder<1024> FormattedCode;
 	if (TypeDesc.ComponentType == Shader::EValueComponentType::Double)
 	{
 		if (TypeDesc.NumComponents == 1)
 		{
-			OutCode.Append(TEXT("MakeLWCScalar("));
+			FormattedCode.Append(TEXT("MakeLWCScalar("));
 		}
 		else
 		{
-			OutCode.Appendf(TEXT("MakeLWCVector%d("), TypeDesc.NumComponents);
+			FormattedCode.Appendf(TEXT("MakeLWCVector%d("), TypeDesc.NumComponents);
 		}
 
-		WriteMaterialUniformAccess(Shader::EValueComponentType::Float, TypeDesc.NumComponents, UniformPreshaderOffset, OutCode); // Tile
+		WriteMaterialUniformAccess(Shader::EValueComponentType::Float, TypeDesc.NumComponents, UniformPreshaderOffset, FormattedCode); // Tile
 		UniformPreshaderOffset += TypeDesc.NumComponents;
-		OutCode.Append(TEXT(","));
-		WriteMaterialUniformAccess(Shader::EValueComponentType::Float, TypeDesc.NumComponents, UniformPreshaderOffset, OutCode); // Offset
+		FormattedCode.Append(TEXT(","));
+		WriteMaterialUniformAccess(Shader::EValueComponentType::Float, TypeDesc.NumComponents, UniformPreshaderOffset, FormattedCode); // Offset
 		UniformPreshaderOffset += TypeDesc.NumComponents;
-		OutCode.Append(TEXT(")"));
+		FormattedCode.Append(TEXT(")"));
 	}
 	else
 	{
-		WriteMaterialUniformAccess(TypeDesc.ComponentType, TypeDesc.NumComponents, UniformPreshaderOffset, OutCode);
+		WriteMaterialUniformAccess(TypeDesc.ComponentType, TypeDesc.NumComponents, UniformPreshaderOffset, FormattedCode);
 		UniformPreshaderOffset += TypeDesc.NumComponents;
 	}
+
+	const TCHAR* Code = AllocateString(*Allocator, FormattedCode);
+	Preshaders.Add(Hash, Code);
+	return Code;
 }
 
 void FEmitContext::Finalize()
@@ -360,6 +375,14 @@ FConstantValue::FConstantValue(const Shader::FValue& InValue) : Type(InValue.Get
 	}
 }
 
+void FScope::Reset()
+{
+	State = EScopeState::Uninitialized;
+	ExpressionCodeMap.Reset();
+	Declarations = FCodeList();
+	Statements = FCodeList();
+}
+
 FScope* FScope::FindSharedParent(FScope* Lhs, FScope* Rhs)
 {
 	FScope* Scope0 = Lhs;
@@ -385,12 +408,34 @@ FScope* FScope::FindSharedParent(FScope* Lhs, FScope* Rhs)
 
 void FExpressionLocalPHI::PrepareValue(FEmitContext& Context, const FRequestedType& RequestedType)
 {
-	FExpression* ForwardExpression = nullptr;
-	bool bForwardExpressionValid = false;
+	FExpression* ForwardExpression = Values[0];
+	bool bForwardExpressionValid = true;
+
+	// There are 2 cases we want to optimize here
+	// 1) If the PHI node has the same value in all the previous scopes, we can avoid generating code for the previous scopes, and just use the value directly
+	for (int32 i = 1; i < NumValues; ++i)
+	{
+		FExpression* ScopeExpression = Values[i];
+		if (ScopeExpression != ForwardExpression)
+		{
+			ForwardExpression = nullptr;
+			bForwardExpressionValid = false;
+			break;
+		}
+	}
+
+	if (bForwardExpressionValid)
+	{
+		check(ForwardExpression);
+		return SetForwardValue(Context, ForwardExpression, RequestedType);
+	}
+
+	// 2) PHI has different values in previous scopes, but possible some previous scopes may become dead due to constant folding
+	// In this case, we check to see if the value is the same in all live scopes, and forward if possible
 	for (int32 i = 0; i < NumValues; ++i)
 	{
-		// Ignore values in scopes that have already been marked as dead
-		if (!Scopes[i]->IsDead())
+		// Ignore values in dead scopes
+		if (PrepareScopeValues(Context, Scopes[i]))
 		{
 			FExpression* ScopeExpression = Values[i];
 			if (!ForwardExpression)
@@ -419,7 +464,7 @@ void FExpressionLocalPHI::PrepareValue(FEmitContext& Context, const FRequestedTy
 	{
 		for (int32 i = 0; i < NumValues; ++i)
 		{
-			if (!TypePerValue[i])
+			if (!TypePerValue[i] && PrepareScopeValues(Context, Scopes[i]))
 			{
 				const FPrepareValueResult ValueResult = PrepareExpressionValue(Context, Values[i], RequestedType);
 				if (ValueResult.Type)
@@ -480,8 +525,9 @@ void FExpressionLocalPHI::EmitShaderDependencies(FEmitContext& Context, const FS
 	for (int32 i = 0; i < NumValues; ++i)
 	{
 		FScope* ValueScope = Scopes[i];
+		check(IsScopeLive(ValueScope));
+
 		const TCHAR* ValueCode = Values[i]->GetValueShader(Context, LocalType);
-		ValueScope->MarkLiveRecursive();
 		if (ValueScope == DeclarationScope)
 		{
 			ValueScope->EmitDeclarationf(Context, TEXT("%s %s = %s;"),
@@ -500,7 +546,7 @@ void FExpressionLocalPHI::EmitShaderDependencies(FEmitContext& Context, const FS
 
 	if (bNeedToAddDeclaration)
 	{
-		check(DeclarationScope->IsLive());
+		check(IsScopeLive(DeclarationScope));
 		DeclarationScope->EmitDeclarationf(Context, TEXT("%s %s;"),
 			LocalType.GetName(),
 			Shader.Code.ToString());
@@ -535,11 +581,25 @@ ENodeVisitResult FScope::Visit(FNodeVisitor& Visitor)
 	const ENodeVisitResult Result = Visitor.OnScope(*this);
 	if (ShouldVisitDependentNodes(Result))
 	{
-		Visitor.VisitNode(Statement);
+		Visitor.VisitNode(ContainedStatement);
 	}
 	return Result;
 }
 
+void FStatement::Reset()
+{
+	bEmitHLSL = false;
+}
+
+void FExpression::Reset()
+{
+	LocalVariableName = nullptr;
+	Code = nullptr;
+	ForwardedValue = nullptr;
+	CurrentRequestedType.Reset();
+	PrepareValueResult = FPrepareValueResult();
+}
+
 FPrepareValueResult PrepareExpressionValue(FEmitContext& Context, FExpression* InExpression, const FRequestedType& RequestedType)
 {
 	if (!InExpression)
@@ -606,7 +666,7 @@ FPrepareValueResult PrepareExpressionValue(FEmitContext& Context, FExpression* I
 			FMaterialRenderContext RenderContext(nullptr, *Context.Material, nullptr);
 			Shader::FValue Value;
 			ConstantPreshader.Evaluate(nullptr, RenderContext, Value);
-			InExpression->PrepareValueResult.ConstantValue = Value;
+			InExpression->PrepareValueResult.ConstantValue = Shader::Cast(Value, InExpression->PrepareValueResult.Type);
 		}
 	}
 
@@ -680,6 +740,7 @@ const TCHAR* FExpression::GetValueShader(FEmitContext& Context)
 					if (FoundDeclaration)
 					{
 						// Re-use results from previous matching expression
+						check(IsScopeLive(CheckScope));
 						Declaration = *FoundDeclaration;
 						break;
 					}
@@ -688,11 +749,10 @@ const TCHAR* FExpression::GetValueShader(FEmitContext& Context)
 
 				if (!Declaration)
 				{
-					check(ParentScope);
+					check(IsScopeLive(ParentScope));
 					Declaration = Context.AcquireLocalDeclarationCode();
 					ParentScope->ExpressionCodeMap.Add(Hash, Declaration);
-					ParentScope->MarkLiveRecursive();
-					ParentScope->EmitStatementf(Context, TEXT("const %s %s = %s;"),
+					ParentScope->EmitDeclarationf(Context, TEXT("const %s %s = %s;"),
 						PrepareValueResult.Type.GetName(),
 						Declaration,
 						FormattedCode.ToString());
@@ -715,8 +775,7 @@ const TCHAR* FExpression::GetValueShader(FEmitContext& Context)
 			EmitValuePreshader(Context, Preshader);
 			bReentryFlag = false;
 
-			Context.AddPreshader(PrepareValueResult.Type, Preshader, FormattedCode);
-			Code = AllocateString(*Context.Allocator, FormattedCode);
+			Code = Context.AcquirePreshader(PrepareValueResult.Type, Preshader);
 		}
 		else
 		{
@@ -825,10 +884,10 @@ void FScope::UseExpression(FExpression* Expression)
 
 void FScope::InternalEmitCode(FEmitContext& Context, FCodeList& List, ENextScopeFormat ScopeFormat, FScope* NextScope, const TCHAR* String, int32 Length)
 {
-	if (NextScope && NextScope->Statement && !NextScope->Statement->bEmitHLSL)
+	if (NextScope && NextScope->ContainedStatement && !NextScope->ContainedStatement->bEmitHLSL)
 	{
-		NextScope->Statement->bEmitHLSL = true;
-		NextScope->Statement->EmitHLSL(Context);
+		NextScope->ContainedStatement->bEmitHLSL = true;
+		NextScope->ContainedStatement->EmitHLSL(Context);
 	}
 
 	const int32 SizeofString = sizeof(TCHAR) * Length;
@@ -854,12 +913,28 @@ void FScope::InternalEmitCode(FEmitContext& Context, FCodeList& List, ENextScope
 	List.Num++;
 }
 
-void PrepareScopeValues(FEmitContext& Context, const FScope* InScope)
+bool PrepareScopeValues(FEmitContext& Context, FScope* InScope)
 {
-	if (InScope && InScope->Statement)
+	if (InScope && InScope->State == EScopeState::Uninitialized)
 	{
-		InScope->Statement->PrepareValues(Context);
+		if (!InScope->ParentScope || PrepareScopeValues(Context, InScope->ParentScope))
+		{
+			if (InScope->OwnerStatement)
+			{
+				InScope->OwnerStatement->PrepareValues(Context);
+			}
+			else
+			{
+				InScope->State = EScopeState::Live;
+			}
+		}
+		else
+		{
+			InScope->State = EScopeState::Dead;
+		}
 	}
+
+	return InScope && InScope->State != EScopeState::Dead;
 }
 
 FTree* FTree::Create(FMemStackBase& Allocator)
@@ -907,6 +982,8 @@ void FScope::MarkLive()
 
 void FScope::MarkLiveRecursive()
 {
+	return MarkLive();
+
 	FScope* Scope = this;
 	while (Scope && Scope->State == EScopeState::Uninitialized)
 	{
@@ -987,36 +1064,32 @@ void FStructType::WriteHLSL(FStringBuilderBase& OutString) const
 	OutString.Append(TEXT("\n"));
 }
 
-bool FTree::EmitHLSL(FEmitContext& Context,
-	FStringBuilderBase& OutDeclarations,
-	FStringBuilderBase& OutCode) const
+void FTree::ResetNodes()
 {
-	if (!ResultStatement)
-	{
-		return false;
-	}
-
-	PrepareScopeValues(Context, RootScope);
-	if (Context.Errors.Num() > 0)
+	FNode* Node = Nodes;
+	while (Node)
 	{
-		return false;
-	}
-
-	ResultStatement->bEmitHLSL = true;
-	ResultStatement->EmitHLSL(Context);
-	if (!RootScope->IsLive())
-	{
-		return false;
+		FNode* Next = Node->NextNode;
+		Node->Reset();
+		Node = Next;
 	}
+}
 
-	if (Context.Errors.Num() > 0)
+void FTree::EmitDeclarationsCode(FStringBuilderBase& OutCode) const
+{
+	const FStructType* StructType = StructTypes;
+	while (StructType)
 	{
-		return false;
+		StructType->WriteHLSL(OutCode);
+		StructType = StructType->NextType;
 	}
+}
 
-	if (RootScope->Statement)
+bool FTree::EmitHLSL(FEmitContext& Context, FStringBuilderBase& OutCode) const
+{
+	if (RootScope->ContainedStatement)
 	{
-		RootScope->Statement->EmitHLSL(Context);
+		RootScope->ContainedStatement->EmitHLSL(Context);
 		if (Context.Errors.Num() > 0)
 		{
 			return false;
@@ -1024,16 +1097,7 @@ bool FTree::EmitHLSL(FEmitContext& Context,
 	}
 
 	Context.Finalize();
-
 	RootScope->WriteHLSL(1, OutCode);
-
-	const FStructType* StructType = StructTypes;
-	while (StructType)
-	{
-		StructType->WriteHLSL(OutDeclarations);
-		StructType = StructType->NextType;
-	}
-
 	return Context.Errors.Num() == 0;
 }
 
@@ -1045,10 +1109,10 @@ void FTree::RegisterExpression(FScope& Scope, FExpression* Expression)
 
 void FTree::RegisterStatement(FScope& Scope, FStatement* Statement)
 {
-	check(!Scope.Statement)
+	check(!Scope.ContainedStatement)
 	check(!Statement->ParentScope);
 	Statement->ParentScope = &Scope;
-	Scope.Statement = Statement;
+	Scope.ContainedStatement = Statement;
 }
 
 FScope* FTree::NewScope(FScope& Scope)
@@ -1060,6 +1124,16 @@ FScope* FTree::NewScope(FScope& Scope)
 	return NewScope;
 }
 
+FScope* FTree::NewOwnedScope(FStatement& Owner)
+{
+	FScope* NewScope = NewNode<FScope>();
+	NewScope->OwnerStatement = &Owner;
+	NewScope->ParentScope = Owner.ParentScope;
+	NewScope->NestedLevel = NewScope->ParentScope->NestedLevel + 1;
+	NewScope->NumPreviousScopes = 0;
+	return NewScope;
+}
+
 FTextureParameterDeclaration* FTree::NewTextureParameterDeclaration(const FName& Name, const FTextureDescription& DefaultValue)
 {
 	FTextureParameterDeclaration* Declaration = NewNode<FTextureParameterDeclaration>(Name, DefaultValue);
@@ -1120,11 +1194,5 @@ const FStructType* FTree::NewStructType(const FStructTypeInitializer& Initialize
 	return StructType;
 }
 
-void FTree::SetResult(FStatement& InResult)
-{
-	check(!ResultStatement);
-	ResultStatement = &InResult;
-}
-
 } // namespace HLSLTree
 } // namespace UE
\ No newline at end of file
diff --git a/Engine/Source/Runtime/Engine/Private/HLSLTree/HLSLTreeCommon.cpp b/Engine/Source/Runtime/Engine/Private/HLSLTree/HLSLTreeCommon.cpp
--- a/Engine/Source/Runtime/Engine/Private/HLSLTree/HLSLTreeCommon.cpp
+++ b/Engine/Source/Runtime/Engine/Private/HLSLTree/HLSLTreeCommon.cpp
@@ -335,7 +335,10 @@ void FExpressionGetStructField::PrepareValue(FEmitContext& Context, const FReque
 		return Context.Errors.AddErrorf(this, TEXT("Invalid field %s"), FieldName);
 	}
 
-	const FPrepareValueResult StructResult = PrepareExpressionValue(Context, StructExpression, RequestedType.MakeFieldAccess(StructType, FieldRef));
+	FRequestedType RequestedStructType(StructType, false);
+	RequestedStructType.SetField(FieldRef, RequestedType);
+
+	const FPrepareValueResult StructResult = PrepareExpressionValue(Context, StructExpression, RequestedStructType);
 	if (StructResult && StructResult.Type != FType(StructType))
 	{
 		return Context.Errors.AddErrorf(this, TEXT("Expected type %s"), StructType->Name);
@@ -361,7 +364,10 @@ void FExpressionSetStructField::PrepareValue(FEmitContext& Context, const FReque
 		return Context.Errors.AddErrorf(this, TEXT("Invalid field %s"), FieldName);
 	}
 
-	const FPrepareValueResult StructResult = PrepareExpressionValue(Context, StructExpression, RequestedType.CopyWithFieldRemoved(FieldRef));
+	FRequestedType RequestedStructType(RequestedType);
+	RequestedStructType.ClearFieldRequested(FieldRef);
+
+	const FPrepareValueResult StructResult = PrepareExpressionValue(Context, StructExpression, RequestedStructType);
 	if (StructResult && StructResult.Type != FType(StructType))
 	{
 		return Context.Errors.AddErrorf(this, TEXT("Expected type %s"), StructType->Name);
@@ -623,20 +629,21 @@ void FExpressionReflectionVector::EmitValueShader(FEmitContext& Context, FShader
 	OutShader.Code.Append(TEXT("Parameters.ReflectionVector"));
 }
 
+void FStatementBreak::PrepareValues(FEmitContext& Context) const
+{
+}
+
 void FStatementBreak::EmitHLSL(FEmitContext& Context) const
 {
-	ParentScope->MarkLive();
 	ParentScope->EmitStatementf(Context, TEXT("break;"));
 }
 
 void FStatementReturn::PrepareValues(FEmitContext& Context) const
 {
-	PrepareExpressionValue(Context, Expression, Type);
 }
 
 void FStatementReturn::EmitHLSL(FEmitContext& Context) const
 {
-	ParentScope->MarkLiveRecursive();
 	ParentScope->EmitStatementf(Context, TEXT("return %s;"), Expression->GetValueShader(Context));
 }
 
@@ -648,21 +655,20 @@ void FStatementIf::PrepareValues(FEmitContext& Context) const
 		const bool bCondition = ConditionResult.ConstantValue.AsBoolScalar();
 		if (bCondition)
 		{
+			MarkScopeLive(ThenScope);
 			MarkScopeDead(ElseScope);
-			PrepareScopeValues(Context, ThenScope);
 		}
 		else
 		{
 			MarkScopeDead(ThenScope);
-			PrepareScopeValues(Context, ElseScope);
+			MarkScopeLive(ElseScope);
 		}
 	}
-	else
+	else if(ConditionResult)
 	{
-		PrepareScopeValues(Context, ThenScope);
-		PrepareScopeValues(Context, ElseScope);
+		MarkScopeLive(ThenScope);
+		MarkScopeLive(ElseScope);
 	}
-	PrepareScopeValues(Context, NextScope);
 }
 
 void FStatementIf::EmitHLSL(FEmitContext& Context) const
@@ -697,8 +703,7 @@ void FStatementIf::EmitHLSL(FEmitContext& Context) const
 
 void FStatementLoop::PrepareValues(FEmitContext& Context) const
 {
-	PrepareScopeValues(Context, LoopScope);
-	PrepareScopeValues(Context, NextScope);
+	MarkScopeLive(LoopScope);
 }
 
 void FStatementLoop::EmitHLSL(FEmitContext& Context) const
diff --git a/Engine/Source/Runtime/Engine/Private/Materials/MaterialExpressionHLSL.cpp b/Engine/Source/Runtime/Engine/Private/Materials/MaterialExpressionHLSL.cpp
--- a/Engine/Source/Runtime/Engine/Private/Materials/MaterialExpressionHLSL.cpp
+++ b/Engine/Source/Runtime/Engine/Private/Materials/MaterialExpressionHLSL.cpp
@@ -4,7 +4,6 @@
 #if WITH_EDITOR
 
 #include "MaterialHLSLGenerator.h"
-#include "MaterialHLSLTree.h"
 #include "HLSLTree/HLSLTree.h"
 #include "HLSLTree/HLSLTreeCommon.h"
 
@@ -392,8 +391,8 @@ EMaterialGenerateHLSLStatus UMaterialExpressionIfThenElse::GenerateHLSLStatement
 	UE::HLSLTree::FStatementIf* IfStatement = Generator.GetTree().NewStatement<UE::HLSLTree::FStatementIf>(Scope);
 	IfStatement->ConditionExpression = ConditionExpression;
 	IfStatement->NextScope = Generator.NewJoinedScope(Scope);
-	IfStatement->ThenScope = Then.NewScopeWithStatements(Generator, Scope);
-	IfStatement->ElseScope = Else.NewScopeWithStatements(Generator, Scope);
+	IfStatement->ThenScope = Then.NewOwnedScopeWithStatements(Generator, *IfStatement);
+	IfStatement->ElseScope = Else.NewOwnedScopeWithStatements(Generator, *IfStatement);
 
 	return EMaterialGenerateHLSLStatus::Success;
 }
@@ -411,11 +410,11 @@ EMaterialGenerateHLSLStatus UMaterialExpressionWhileLoop::GenerateHLSLStatements
 	}
 
 	UE::HLSLTree::FStatementLoop* LoopStatement = Generator.GetTree().NewStatement<UE::HLSLTree::FStatementLoop>(Scope);
-	LoopStatement->LoopScope = Generator.NewScope(Scope);
+	LoopStatement->LoopScope = Generator.NewOwnedScope(*LoopStatement);
 
 	UE::HLSLTree::FStatementIf* IfStatement = Generator.GetTree().NewStatement<UE::HLSLTree::FStatementIf>(*LoopStatement->LoopScope);
-	IfStatement->ThenScope = Generator.NewScope(*LoopStatement->LoopScope);
-	IfStatement->ElseScope = Generator.NewScope(*LoopStatement->LoopScope);
+	IfStatement->ThenScope = Generator.NewOwnedScope(*IfStatement);
+	IfStatement->ElseScope = Generator.NewOwnedScope(*IfStatement);
 	LoopStatement->NextScope = Generator.NewScope(Scope, EMaterialNewScopeFlag::NoPreviousScope);
 	LoopStatement->LoopScope->AddPreviousScope(*IfStatement->ThenScope);
 	LoopStatement->NextScope->AddPreviousScope(*IfStatement->ElseScope);
@@ -482,12 +481,12 @@ EMaterialGenerateHLSLStatus UMaterialExpressionForLoop::GenerateHLSLStatements(F
 	Generator.GenerateAssignLocal(Scope, ExpressionData->LocalName, StartExpression);
 
 	UE::HLSLTree::FStatementLoop* LoopStatement = Generator.GetTree().NewStatement<UE::HLSLTree::FStatementLoop>(Scope);
-	LoopStatement->LoopScope = Generator.NewScope(Scope);
+	LoopStatement->LoopScope = Generator.NewOwnedScope(*LoopStatement);
 	ExpressionData->LoopScope = LoopStatement->LoopScope;
 
 	UE::HLSLTree::FStatementIf* IfStatement = Generator.GetTree().NewStatement<UE::HLSLTree::FStatementIf>(*LoopStatement->LoopScope);
-	IfStatement->ThenScope = Generator.NewScope(*LoopStatement->LoopScope);
-	IfStatement->ElseScope = Generator.NewScope(*LoopStatement->LoopScope);
+	IfStatement->ThenScope = Generator.NewOwnedScope(*IfStatement);
+	IfStatement->ElseScope = Generator.NewOwnedScope(*IfStatement);
 	LoopStatement->NextScope = Generator.NewScope(Scope, EMaterialNewScopeFlag::NoPreviousScope);
 	LoopStatement->LoopScope->AddPreviousScope(*IfStatement->ThenScope);
 	LoopStatement->NextScope->AddPreviousScope(*IfStatement->ElseScope);
diff --git a/Engine/Source/Runtime/Engine/Private/Materials/MaterialHLSLEmitter.cpp b/Engine/Source/Runtime/Engine/Private/Materials/MaterialHLSLEmitter.cpp
--- a/Engine/Source/Runtime/Engine/Private/Materials/MaterialHLSLEmitter.cpp
+++ b/Engine/Source/Runtime/Engine/Private/Materials/MaterialHLSLEmitter.cpp
@@ -11,6 +11,7 @@
 #include "Materials/MaterialExpressionFunctionOutput.h"
 #include "Materials/MaterialExpressionVolumetricAdvancedMaterialOutput.h"
 #include "Materials/Material.h"
+#include "MaterialHLSLGenerator.h"
 #include "ShaderCore.h"
 #include "HLSLTree/HLSLTree.h"
 #include "Containers/LazyPrintf.h"
@@ -45,7 +46,8 @@ static FString GenerateMaterialTemplateHLSL(EShaderPlatform ShaderPlatform,
 	const FMaterial& Material,
 	const UE::HLSLTree::FEmitContext& EmitContext,
 	const TCHAR* DeclarationsCode,
-	const TCHAR* PixelShaderCode)
+	const TCHAR* PixelShaderCodePhase0,
+	const TCHAR* PixelShaderCodePhase1)
 {
 	FString MaterialTemplateSource;
 	LoadShaderSourceFileChecked(TEXT("/Engine/Private/MaterialTemplate.ush"), ShaderPlatform, MaterialTemplateSource);
@@ -297,9 +299,27 @@ static FString GenerateMaterialTemplateHLSL(EShaderPlatform ShaderPlatform,
 		//EvaluateVertexCode += TranslatedAttributesCodeChunks[SF_Vertex];
 
 		LazyPrintf.PushParam(*EvaluateVertexCode);
-		LazyPrintf.PushParam(PixelShaderCode);
 
-		FString EvaluateMaterialAttributesCode = TEXT("    FMaterialAttributes MaterialAttributes = EvaluatePixelMaterialAttributes(Parameters);" LINE_TERMINATOR);
+		FString EvaluateMaterialDeclaration;
+		FString EvaluateMaterialAttributesPhase0 = TEXT("    FMaterialAttributes MaterialAttributesPhase0 = EvaluatePixelMaterialAttributesPhase0(Parameters);" LINE_TERMINATOR);
+		FString EvaluateMaterialAttributesPhase1;
+
+		EvaluateMaterialDeclaration += TEXT("FMaterialAttributes EvaluatePixelMaterialAttributesPhase0(in out FMaterialPixelParameters Parameters)" LINE_TERMINATOR);
+		EvaluateMaterialDeclaration += TEXT("{" LINE_TERMINATOR);
+		EvaluateMaterialDeclaration += PixelShaderCodePhase0;
+		EvaluateMaterialDeclaration += TEXT("}" LINE_TERMINATOR);
+
+		if (PixelShaderCodePhase1)
+		{
+			EvaluateMaterialDeclaration += TEXT("FMaterialAttributes EvaluatePixelMaterialAttributesPhase1(in out FMaterialPixelParameters Parameters)" LINE_TERMINATOR);
+			EvaluateMaterialDeclaration += TEXT("{" LINE_TERMINATOR);
+			EvaluateMaterialDeclaration += PixelShaderCodePhase1;
+			EvaluateMaterialDeclaration += TEXT("}" LINE_TERMINATOR);
+
+			EvaluateMaterialAttributesPhase1 = TEXT("    FMaterialAttributes MaterialAttributesPhase1 = EvaluatePixelMaterialAttributesPhase1(Parameters);" LINE_TERMINATOR);
+		}
+
+		LazyPrintf.PushParam(*EvaluateMaterialDeclaration);
 
 		for (int32 PropertyIndex = 0; PropertyIndex < MP_MAX; ++PropertyIndex)
 		{
@@ -314,6 +334,11 @@ static FString GenerateMaterialTemplateHLSL(EShaderPlatform ShaderPlatform,
 			// Special case MP_SubsurfaceColor as the actual property is a combination of the color and the profile but we don't want to expose the profile
 			const FString PropertyName = FMaterialAttributeDefinitionMap::GetAttributeName(Property);
 
+			// 'Normal' is always evaluated in Phase0
+			// If there is a Phase1, it will contain all other attributes
+			const int32 EvaluatePhase = (Property == MP_Normal || !PixelShaderCodePhase1) ? 0 : 1;
+			FString& EvaluateMaterialAttributesCode = (EvaluatePhase == 0) ? EvaluateMaterialAttributesPhase0 : EvaluateMaterialAttributesPhase1;
+
 			if (PropertyIndex == MP_FrontMaterial)
 			{
 				// TODO - Strata
@@ -322,20 +347,20 @@ static FString GenerateMaterialTemplateHLSL(EShaderPlatform ShaderPlatform,
 			else if (PropertyIndex == MP_SubsurfaceColor)
 			{
 				// TODO - properly handle subsurface profile
-				EvaluateMaterialAttributesCode += FString::Printf("    PixelMaterialInputs.Subsurface = float4(MaterialAttributes.%s, 0.0f);" LINE_TERMINATOR, *PropertyName);
+				EvaluateMaterialAttributesCode += FString::Printf("    PixelMaterialInputs.Subsurface = float4(MaterialAttributesPhase%d.%s, 0.0f);" LINE_TERMINATOR, EvaluatePhase, *PropertyName);
 			}
 			else
 			{
-				EvaluateMaterialAttributesCode += FString::Printf("    PixelMaterialInputs.%s = MaterialAttributes.%s;" LINE_TERMINATOR, *PropertyName, *PropertyName);
+				EvaluateMaterialAttributesCode += FString::Printf("    PixelMaterialInputs.%s = MaterialAttributesPhase%d.%s;" LINE_TERMINATOR, *PropertyName, EvaluatePhase, *PropertyName);
 			}
 		}
 
 		// TODO - deriv
 		for (int32 Iter = 0; Iter < CompiledPDV_MAX; Iter++)
 		{
-			LazyPrintf.PushParam(*EvaluateMaterialAttributesCode);
-			LazyPrintf.PushParam(TEXT(""));
 			LazyPrintf.PushParam(TEXT(""));
+			LazyPrintf.PushParam(*EvaluateMaterialAttributesPhase0);
+			LazyPrintf.PushParam(*EvaluateMaterialAttributesPhase1);
 		}
 	}
 
@@ -662,12 +687,13 @@ public:
 };
 
 bool MaterialEmitHLSL(const FMaterialCompileTargetParameters& InCompilerTarget,
-	const FMaterial& InMaterial,
 	const FStaticParameterSet& InStaticParameters,
-	const UE::HLSLTree::FTree& InTree,
+	FMaterial& InOutMaterial,
 	FMaterialCompilationOutput& OutCompilationOutput,
 	TRefCountPtr<FSharedShaderCompilerEnvironment>& OutMaterialEnvironment)
 {
+	using namespace UE::HLSLTree;
+
 	FMemory::Memzero(SharedPixelProperties);
 	SharedPixelProperties[MP_Normal] = true;
 	SharedPixelProperties[MP_Tangent] = true;
@@ -687,26 +713,108 @@ bool MaterialEmitHLSL(const FMaterialCompileTargetParameters& InCompilerTarget,
 	SharedPixelProperties[MP_FrontMaterial] = true;
 
 	FMemStackBase Allocator;
-	FMemMark MemMark(Allocator);
 
-	FStringBuilderMemstack Declarations(Allocator, 64 * 1024);
-	FStringBuilderMemstack Code(Allocator, 1024 * 1024);
+	UMaterialInterface* MaterialInterface = InOutMaterial.GetMaterialInterface();
+	UMaterial* BaseMaterial = MaterialInterface->GetMaterial();
+
+	FTree* HLSLTree = FTree::Create(Allocator);
+	FMaterialHLSLGenerator Generator(BaseMaterial, InCompilerTarget, *HLSLTree);
+	FScope& RootScope = HLSLTree->GetRootScope();
+
+	bool bGenerateResult = false;
+	if (BaseMaterial->IsCompiledWithExecutionFlow())
+	{
+		UMaterialExpression* BaseExpression = BaseMaterial->ExpressionExecBegin;
+		bGenerateResult = Generator.GenerateStatements(RootScope, BaseExpression);
+	}
+	else
+	{
+		bGenerateResult = Generator.GenerateResult(RootScope);
+	}
+
+	if (!Generator.Finalize())
+	{
+		bGenerateResult = false;
+	}
+
+	Generator.AcquireErrors(InOutMaterial);
+	if (!bGenerateResult)
+	{
+		return false;
+	}
 
-	UE::HLSLTree::FEmitContext EmitContext(Allocator);
-	EmitContext.Material = &InMaterial;
+	FEmitContext EmitContext(Allocator);
+	EmitContext.Material = &InOutMaterial;
 	EmitContext.StaticParameters = &InStaticParameters;
 	EmitContext.MaterialCompilationOutput = &OutCompilationOutput;
-	InTree.EmitHLSL(EmitContext, Declarations, Code);
 
-	//InOutMaterial.CompileErrors = MoveTemp(CompileErrors);
-	//InOutMaterial.ErrorExpressions = MoveTemp(ErrorExpressions);
+	const FStructFieldRef NormalFieldRef = Generator.GetMaterialAttributesType()->FindFieldByName(*FMaterialAttributeDefinitionMap::GetAttributeName(MP_Normal));
+	check(NormalFieldRef);
+
+	if (!PrepareScopeValues(EmitContext, Generator.GetResultStatement()->ParentScope))
+	{
+		return false;
+	}
+
+	// Prepare all fields *except* normal
+	FRequestedType RequestedMaterialAttributesType(Generator.GetMaterialAttributesType());
+	RequestedMaterialAttributesType.ClearFieldRequested(NormalFieldRef);
+	if (!PrepareExpressionValue(EmitContext, Generator.GetResultExpression(), RequestedMaterialAttributesType))
+	{
+		return false;
+	}
+
+	const bool bReadMaterialNormal = EmitContext.bReadMaterialNormal;
+	FStringBuilderMemstack CodePhase0(Allocator, 1024 * 1024);
+	FStringBuilderMemstack CodePhase1(Allocator, 1024 * 1024);
+	if (!bReadMaterialNormal)
+	{
+		// No access to material normal, can execute everything in phase0
+		RequestedMaterialAttributesType.SetFieldRequested(NormalFieldRef);
+		if (!PrepareExpressionValue(EmitContext, Generator.GetResultExpression(), RequestedMaterialAttributesType))
+		{
+			return false;
+		}
+		HLSLTree->EmitHLSL(EmitContext, CodePhase0);
+	}
+	else
+	{
+		// Execute everything *except* normal in phase1
+		HLSLTree->EmitHLSL(EmitContext, CodePhase1);
+
+		// Reset state
+		HLSLTree->ResetNodes();
+		if (!PrepareScopeValues(EmitContext, Generator.GetResultStatement()->ParentScope))
+		{
+			return false;
+		}
+
+		// Prepare code for just the normal
+		FRequestedType RequestedMaterialNormal(Generator.GetMaterialAttributesType(), false);
+		RequestedMaterialNormal.SetFieldRequested(NormalFieldRef);
+		if (!PrepareExpressionValue(EmitContext, Generator.GetResultExpression(), RequestedMaterialNormal))
+		{
+			return false;
+		}
+
+		// Execute the normal in phase0
+		HLSLTree->EmitHLSL(EmitContext, CodePhase0);
+	}
+
+	FStringBuilderMemstack Declarations(Allocator, 64 * 1024);
+	HLSLTree->EmitDeclarationsCode(Declarations);
 
-	const FString MaterialTemplateSource = GenerateMaterialTemplateHLSL(InCompilerTarget.ShaderPlatform, InMaterial, EmitContext, Declarations.ToString(), Code.ToString());
+	FString MaterialTemplateSource = GenerateMaterialTemplateHLSL(InCompilerTarget.ShaderPlatform,
+		InOutMaterial,
+		EmitContext,
+		Declarations.ToString(),
+		CodePhase0.ToString(),
+		bReadMaterialNormal ? CodePhase1.ToString() : nullptr);
 
 	OutMaterialEnvironment = new FSharedShaderCompilerEnvironment();
 	OutMaterialEnvironment->TargetPlatform = InCompilerTarget.TargetPlatform;
-	GetMaterialEnvironment(InCompilerTarget.ShaderPlatform, InMaterial, OutCompilationOutput, *OutMaterialEnvironment);
-	OutMaterialEnvironment->IncludeVirtualPathToContentsMap.Add(TEXT("/Engine/Generated/Material.ush"), MaterialTemplateSource);
+	GetMaterialEnvironment(InCompilerTarget.ShaderPlatform, InOutMaterial, OutCompilationOutput, *OutMaterialEnvironment);
+	OutMaterialEnvironment->IncludeVirtualPathToContentsMap.Add(TEXT("/Engine/Generated/Material.ush"), MoveTemp(MaterialTemplateSource));
 
 	return true;
 }
diff --git a/Engine/Source/Runtime/Engine/Private/Materials/MaterialHLSLEmitter.h b/Engine/Source/Runtime/Engine/Private/Materials/MaterialHLSLEmitter.h
--- a/Engine/Source/Runtime/Engine/Private/Materials/MaterialHLSLEmitter.h
+++ b/Engine/Source/Runtime/Engine/Private/Materials/MaterialHLSLEmitter.h
@@ -12,18 +12,9 @@ struct FSharedShaderCompilerEnvironment;
 struct FMaterialCompileTargetParameters;
 struct FStaticParameterSet;
 
-namespace UE
-{
-namespace HLSLTree
-{
-class FTree;
-}
-}
-
 bool MaterialEmitHLSL(const FMaterialCompileTargetParameters& InCompilerTarget,
-	const FMaterial& InOutMaterial,
 	const FStaticParameterSet& InStaticParameters,
-	const UE::HLSLTree::FTree& InTree,
+	FMaterial& InOutMaterial,
 	FMaterialCompilationOutput& OutCompilationOutput,
 	TRefCountPtr<FSharedShaderCompilerEnvironment>& OutMaterialEnvironment);
 
diff --git a/Engine/Source/Runtime/Engine/Private/Materials/MaterialHLSLGenerator.cpp b/Engine/Source/Runtime/Engine/Private/Materials/MaterialHLSLGenerator.cpp
--- a/Engine/Source/Runtime/Engine/Private/Materials/MaterialHLSLGenerator.cpp
+++ b/Engine/Source/Runtime/Engine/Private/Materials/MaterialHLSLGenerator.cpp
@@ -10,7 +10,6 @@
 #include "Materials/MaterialExpressionFunctionOutput.h"
 #include "Materials/MaterialExpressionVolumetricAdvancedMaterialOutput.h"
 #include "Materials/Material.h"
-#include "MaterialHLSLTree.h"
 #include "ShaderCore.h"
 #include "HLSLTree/HLSLTree.h"
 #include "HLSLTree/HLSLTreeCommon.h"
@@ -85,10 +84,10 @@ FMaterialHLSLGenerator::FMaterialHLSLGenerator(UMaterial* InTargetMaterial, cons
 	MaterialAttributesDefaultValue.Type = MaterialAttributesType;
 }
 
-void FMaterialHLSLGenerator::AcquireErrors(TArray<FString>& OutCompileErrors, TArray<UMaterialExpression*>& OutErrorExpressions)
+void FMaterialHLSLGenerator::AcquireErrors(FMaterial& OutMaterial)
 {
-	OutCompileErrors = MoveTemp(CompileErrors);
-	OutErrorExpressions = MoveTemp(ErrorExpressions);
+	OutMaterial.CompileErrors = MoveTemp(CompileErrors);
+	OutMaterial.ErrorExpressions = MoveTemp(ErrorExpressions);
 }
 
 bool FMaterialHLSLGenerator::Finalize()
@@ -101,6 +100,12 @@ bool FMaterialHLSLGenerator::Finalize()
 		return false;
 	}
 
+	if (!ResultExpression || !ResultStatement)
+	{
+		Error(TEXT("Failed to initialize result"));
+		return false;
+	}
+
 	for (const auto& It : StatementMap)
 	{
 		const UMaterialExpression* Expression = It.Key;
@@ -242,6 +247,9 @@ bool FMaterialHLSLGenerator::GenerateResult(UE::HLSLTree::FScope& Scope)
 	}
 	else
 	{
+		check(!ResultStatement);
+		check(!ResultExpression);
+
 		if (TargetMaterial)
 		{
 			UE::HLSLTree::FExpression* AttributesExpression = nullptr;
@@ -267,9 +275,9 @@ bool FMaterialHLSLGenerator::GenerateResult(UE::HLSLTree::FScope& Scope)
 			if (AttributesExpression)
 			{
 				UE::HLSLTree::FStatementReturn* ReturnStatement = HLSLTree->NewStatement<UE::HLSLTree::FStatementReturn>(Scope);
-				ReturnStatement->Type = MaterialAttributesType;
 				ReturnStatement->Expression = AttributesExpression;
-				HLSLTree->SetResult(*ReturnStatement);
+				ResultExpression = AttributesExpression;
+				ResultStatement = ReturnStatement;
 				bResult = true;
 			}
 		}
@@ -293,6 +301,13 @@ UE::HLSLTree::FScope* FMaterialHLSLGenerator::NewScope(UE::HLSLTree::FScope& Sco
 	return NewScope;
 }
 
+UE::HLSLTree::FScope* FMaterialHLSLGenerator::NewOwnedScope(UE::HLSLTree::FStatement& Owner)
+{
+	UE::HLSLTree::FScope* NewScope = HLSLTree->NewOwnedScope(Owner);
+	NewScope->AddPreviousScope(*Owner.ParentScope);
+	return NewScope;
+}
+
 UE::HLSLTree::FScope* FMaterialHLSLGenerator::NewJoinedScope(UE::HLSLTree::FScope& Scope)
 {
 	UE::HLSLTree::FScope* NewScope = HLSLTree->NewScope(Scope);
diff --git a/Engine/Source/Runtime/Engine/Private/Materials/MaterialHLSLTree.cpp b/Engine/Source/Runtime/Engine/Private/Materials/MaterialHLSLTree.cpp
deleted file mode 100644
diff --git a/Engine/Source/Runtime/Engine/Private/Materials/MaterialShared.cpp b/Engine/Source/Runtime/Engine/Private/Materials/MaterialShared.cpp
--- a/Engine/Source/Runtime/Engine/Private/Materials/MaterialShared.cpp
+++ b/Engine/Source/Runtime/Engine/Private/Materials/MaterialShared.cpp
@@ -49,7 +49,6 @@
 #if WITH_EDITOR
 #include "Rendering/StaticLightingSystemInterface.h"
 #include "MaterialHLSLGenerator.h"
-#include "MaterialHLSLTree.h"
 #include "MaterialHLSLEmitter.h"
 #include "HLSLTree/HLSLTreeCommon.h"
 #include "Misc/ScopeLock.h"
@@ -360,9 +359,17 @@ bool FExpressionExecOutput::GenerateHLSLStatements(FMaterialHLSLGenerator& Gener
 	return bResult;
 }
 
-UE::HLSLTree::FScope* FExpressionExecOutput::NewScopeWithStatements(FMaterialHLSLGenerator& Generator, UE::HLSLTree::FScope& Scope) const
+UE::HLSLTree::FScope* FExpressionExecOutput::NewOwnedScopeWithStatements(FMaterialHLSLGenerator& Generator, UE::HLSLTree::FStatement& Owner) const
 {
-	return NewScopeWithStatements(Generator, Scope, EMaterialNewScopeFlag::None);
+	UE::HLSLTree::FScope* Result = nullptr;
+	if (Expression)
+	{
+		Expression->ValidateState();
+		Result = Generator.NewOwnedScope(Owner); // Create a new scope for the statements
+		Generator.GenerateStatements(*Result, Expression);
+	}
+
+	return Result;
 }
 
 UE::HLSLTree::FScope* FExpressionExecOutput::NewScopeWithStatements(FMaterialHLSLGenerator& Generator, UE::HLSLTree::FScope& Scope, EMaterialNewScopeFlag Flags) const
@@ -2368,11 +2375,7 @@ bool FMaterial::Translate_New(const FMaterialShaderMapId& ShaderMapId,
 {
 #if WITH_EDITOR
 	const FMaterialCompileTargetParameters TargetParams(InPlatform, ShaderMapId.FeatureLevel, InTargetPlatform);
-
-	FMaterialHLSLTree Tree;
-	Tree.InitializeForMaterial(TargetParams, *this);
-
-	return MaterialEmitHLSL(TargetParams, *this, InStaticParameters, Tree.GetTree(), OutCompilationOutput, OutMaterialEnvironment);
+	return MaterialEmitHLSL(TargetParams, InStaticParameters, *this, OutCompilationOutput, OutMaterialEnvironment);
 #else
 	checkNoEntry();
 	return false;
diff --git a/Engine/Source/Runtime/Engine/Private/Shader/Preshader.cpp b/Engine/Source/Runtime/Engine/Private/Shader/Preshader.cpp
--- a/Engine/Source/Runtime/Engine/Private/Shader/Preshader.cpp
+++ b/Engine/Source/Runtime/Engine/Private/Shader/Preshader.cpp
@@ -9,12 +9,17 @@
 
 IMPLEMENT_TYPE_LAYOUT(UE::Shader::FPreshaderData);
 
-void UE::Shader::FPreshaderData::WriteData(const void* Value, uint32 Size)
+namespace UE
+{
+namespace Shader
+{
+
+void FPreshaderData::WriteData(const void* Value, uint32 Size)
 {
 	Data.Append((uint8*)Value, Size);
 }
 
-void UE::Shader::FPreshaderData::WriteName(const FScriptName& Name)
+void FPreshaderData::WriteName(const FScriptName& Name)
 {
 	int32 Index = Names.Find(Name);
 	if (Index == INDEX_NONE)
@@ -27,7 +32,7 @@ void UE::Shader::FPreshaderData::WriteName(const FScriptName& Name)
 	Write((uint16)Index);
 }
 
-void UE::Shader::FPreshaderData::Append(const FPreshaderData& InPreshader)
+void FPreshaderData::Append(const FPreshaderData& InPreshader)
 {
 	const uint32 BaseOffset = Data.Num();
 	Data.Append(InPreshader.Data);
@@ -51,14 +56,14 @@ void UE::Shader::FPreshaderData::Append(const FPreshaderData& InPreshader)
 	}
 }
 
-UE::Shader::FPreshaderDataContext::FPreshaderDataContext(const FPreshaderData& InData)
+FPreshaderDataContext::FPreshaderDataContext(const FPreshaderData& InData)
 	: Ptr(InData.Data.GetData())
 	, EndPtr(Ptr + InData.Data.Num())
 	, Names(InData.Names.GetData())
 	, NumNames(InData.Names.Num())
 {}
 
-UE::Shader::FPreshaderDataContext::FPreshaderDataContext(const FPreshaderDataContext& InContext, uint32 InOffset, uint32 InSize)
+FPreshaderDataContext::FPreshaderDataContext(const FPreshaderDataContext& InContext, uint32 InOffset, uint32 InSize)
 	: Ptr(InContext.Ptr + InOffset)
 	, EndPtr(Ptr + InSize)
 	, Names(InContext.Names)
@@ -66,7 +71,7 @@ UE::Shader::FPreshaderDataContext::FPreshaderDataContext(const FPreshaderDataCon
 {}
 
 template<typename T>
-static inline T ReadPreshaderValue(UE::Shader::FPreshaderDataContext& RESTRICT Data)
+static inline T ReadPreshaderValue(FPreshaderDataContext& RESTRICT Data)
 {
 	T Result;
 	FMemory::Memcpy(&Result, Data.Ptr, sizeof(T));
@@ -76,14 +81,14 @@ static inline T ReadPreshaderValue(UE::Shader::FPreshaderDataContext& RESTRICT D
 }
 
 template<>
-inline uint8 ReadPreshaderValue<uint8>(UE::Shader::FPreshaderDataContext& RESTRICT Data)
+inline uint8 ReadPreshaderValue<uint8>(FPreshaderDataContext& RESTRICT Data)
 {
 	checkSlow(Data.Ptr < Data.EndPtr);
 	return *Data.Ptr++;
 }
 
 template<>
-FScriptName ReadPreshaderValue<FScriptName>(UE::Shader::FPreshaderDataContext& RESTRICT Data)
+FScriptName ReadPreshaderValue<FScriptName>(FPreshaderDataContext& RESTRICT Data)
 {
 	const int32 Index = ReadPreshaderValue<uint16>(Data);
 	check(Index >= 0 && Index < Data.NumNames);
@@ -91,10 +96,10 @@ FScriptName ReadPreshaderValue<FScriptName>(UE::Shader::FPreshaderDataContext& R
 }
 
 template<>
-FName ReadPreshaderValue<FName>(UE::Shader::FPreshaderDataContext& RESTRICT Data) = delete;
+FName ReadPreshaderValue<FName>(FPreshaderDataContext& RESTRICT Data) = delete;
 
 template<>
-FHashedMaterialParameterInfo ReadPreshaderValue<FHashedMaterialParameterInfo>(UE::Shader::FPreshaderDataContext& RESTRICT Data)
+FHashedMaterialParameterInfo ReadPreshaderValue<FHashedMaterialParameterInfo>(FPreshaderDataContext& RESTRICT Data)
 {
 	const FScriptName Name = ReadPreshaderValue<FScriptName>(Data);
 	const int32 Index = ReadPreshaderValue<int32>(Data);
@@ -103,17 +108,17 @@ FHashedMaterialParameterInfo ReadPreshaderValue<FHashedMaterialParameterInfo>(UE
 }
 
 template<>
-UE::Shader::FValue ReadPreshaderValue<UE::Shader::FValue>(UE::Shader::FPreshaderDataContext& RESTRICT Data)
+FValue ReadPreshaderValue<FValue>(FPreshaderDataContext& RESTRICT Data)
 {
-	const UE::Shader::EValueType Type = (UE::Shader::EValueType)ReadPreshaderValue<uint8>(Data);
+	const EValueType Type = (EValueType)ReadPreshaderValue<uint8>(Data);
 	uint32 Size = 0u;
-	UE::Shader::FValue Result = UE::Shader::FValue::FromMemoryImage(Type, Data.Ptr, &Size);
+	FValue Result = FValue::FromMemoryImage(Type, Data.Ptr, &Size);
 	Data.Ptr += Size;
 	checkSlow(Data.Ptr <= Data.EndPtr);
 	return Result;
 }
 
-static void EvaluateParameter(const FUniformExpressionSet& UniformExpressionSet, uint32 ParameterIndex, const FMaterialRenderContext& Context, UE::Shader::FValue& OutValue)
+static void EvaluateParameter(const FUniformExpressionSet& UniformExpressionSet, uint32 ParameterIndex, const FMaterialRenderContext& Context, FValue& OutValue)
 {
 	const FMaterialNumericParameterInfo& Parameter = UniformExpressionSet.GetNumericParameter(ParameterIndex);
 	bool bFoundParameter = false;
@@ -145,30 +150,30 @@ static void EvaluateParameter(const FUniformExpressionSet& UniformExpressionSet,
 }
 
 template<typename Operation>
-static inline void EvaluateUnaryOp(UE::Shader::FPreshaderStack& Stack, const Operation& Op)
+static inline void EvaluateUnaryOp(FPreshaderStack& Stack, const Operation& Op)
 {
-	const UE::Shader::FValue Value = Stack.Pop(false);
+	const FValue Value = Stack.Pop(false);
 	Stack.Add(Op(Value));
 }
 
 template<typename Operation>
-static inline void EvaluateBinaryOp(UE::Shader::FPreshaderStack& Stack, const Operation& Op)
+static inline void EvaluateBinaryOp(FPreshaderStack& Stack, const Operation& Op)
 {
-	const UE::Shader::FValue Value1 = Stack.Pop(false);
-	const UE::Shader::FValue Value0 = Stack.Pop(false);
+	const FValue Value1 = Stack.Pop(false);
+	const FValue Value0 = Stack.Pop(false);
 	Stack.Add(Op(Value0, Value1));
 }
 
 template<typename Operation>
-static inline void EvaluateTernaryOp(UE::Shader::FPreshaderStack& Stack, const Operation& Op)
+static inline void EvaluateTernaryOp(FPreshaderStack& Stack, const Operation& Op)
 {
-	const UE::Shader::FValue Value2 = Stack.Pop(false);
-	const UE::Shader::FValue Value1 = Stack.Pop(false);
-	const UE::Shader::FValue Value0 = Stack.Pop(false);
+	const FValue Value2 = Stack.Pop(false);
+	const FValue Value1 = Stack.Pop(false);
+	const FValue Value0 = Stack.Pop(false);
 	Stack.Add(Op(Value0, Value1, Value2));
 }
 
-static void EvaluateComponentSwizzle(UE::Shader::FPreshaderStack& Stack, UE::Shader::FPreshaderDataContext& RESTRICT Data)
+static void EvaluateComponentSwizzle(FPreshaderStack& Stack, FPreshaderDataContext& RESTRICT Data)
 {
 	const uint8 NumElements = ReadPreshaderValue<uint8>(Data);
 	const uint8 IndexR = ReadPreshaderValue<uint8>(Data);
@@ -176,8 +181,8 @@ static void EvaluateComponentSwizzle(UE::Shader::FPreshaderStack& Stack, UE::Sha
 	const uint8 IndexB = ReadPreshaderValue<uint8>(Data);
 	const uint8 IndexA = ReadPreshaderValue<uint8>(Data);
 
-	UE::Shader::FValue Value = Stack.Pop(false);
-	UE::Shader::FValue Result;
+	FValue Value = Stack.Pop(false);
+	FValue Result;
 	Result.ComponentType = Value.ComponentType;
 	Result.NumComponents = NumElements;
 
@@ -202,7 +207,7 @@ static void EvaluateComponentSwizzle(UE::Shader::FPreshaderStack& Stack, UE::Sha
 	Stack.Add(Result);
 }
 
-static const UTexture* GetTextureParameter(const FMaterialRenderContext& Context, UE::Shader::FPreshaderDataContext& RESTRICT Data)
+static const UTexture* GetTextureParameter(const FMaterialRenderContext& Context, FPreshaderDataContext& RESTRICT Data)
 {
 	const FHashedMaterialParameterInfo ParameterInfo = ReadPreshaderValue<FHashedMaterialParameterInfo>(Data);
 	const int32 TextureIndex = ReadPreshaderValue<int32>(Data);
@@ -212,7 +217,7 @@ static const UTexture* GetTextureParameter(const FMaterialRenderContext& Context
 	return Texture;
 }
 
-static void EvaluateTextureSize(const FMaterialRenderContext& Context, UE::Shader::FPreshaderStack& Stack, UE::Shader::FPreshaderDataContext& RESTRICT Data)
+static void EvaluateTextureSize(const FMaterialRenderContext& Context, FPreshaderStack& Stack, FPreshaderDataContext& RESTRICT Data)
 {
 	const UTexture* Texture = GetTextureParameter(Context, Data);
 	if (Texture && Texture->GetResource())
@@ -220,15 +225,15 @@ static void EvaluateTextureSize(const FMaterialRenderContext& Context, UE::Shade
 		const uint32 SizeX = Texture->GetResource()->GetSizeX();
 		const uint32 SizeY = Texture->GetResource()->GetSizeY();
 		const uint32 SizeZ = Texture->GetResource()->GetSizeZ();
-		Stack.Add(UE::Shader::FValue((float)SizeX, (float)SizeY, (float)SizeZ));
+		Stack.Add(FValue((float)SizeX, (float)SizeY, (float)SizeZ));
 	}
 	else
 	{
-		Stack.Add(UE::Shader::FValue(0.0f, 0.0f, 0.0f));
+		Stack.Add(FValue(0.0f, 0.0f, 0.0f));
 	}
 }
 
-static void EvaluateTexelSize(const FMaterialRenderContext& Context, UE::Shader::FPreshaderStack& Stack, UE::Shader::FPreshaderDataContext& RESTRICT Data)
+static void EvaluateTexelSize(const FMaterialRenderContext& Context, FPreshaderStack& Stack, FPreshaderDataContext& RESTRICT Data)
 {
 	const UTexture* Texture = GetTextureParameter(Context, Data);
 	if (Texture && Texture->GetResource())
@@ -236,15 +241,15 @@ static void EvaluateTexelSize(const FMaterialRenderContext& Context, UE::Shader:
 		const uint32 SizeX = Texture->GetResource()->GetSizeX();
 		const uint32 SizeY = Texture->GetResource()->GetSizeY();
 		const uint32 SizeZ = Texture->GetResource()->GetSizeZ();
-		Stack.Add(UE::Shader::FValue(1.0f / (float)SizeX, 1.0f / (float)SizeY, (SizeZ > 0 ? 1.0f / (float)SizeZ : 0.0f)));
+		Stack.Add(FValue(1.0f / (float)SizeX, 1.0f / (float)SizeY, (SizeZ > 0 ? 1.0f / (float)SizeZ : 0.0f)));
 	}
 	else
 	{
-		Stack.Add(UE::Shader::FValue(0.0f, 0.0f, 0.0f));
+		Stack.Add(FValue(0.0f, 0.0f, 0.0f));
 	}
 }
 
-static FGuid GetExternalTextureGuid(const FMaterialRenderContext& Context, UE::Shader::FPreshaderDataContext& RESTRICT Data)
+static FGuid GetExternalTextureGuid(const FMaterialRenderContext& Context, FPreshaderDataContext& RESTRICT Data)
 {
 	const FScriptName ParameterName = ReadPreshaderValue<FScriptName>(Data);
 	const FGuid ExternalTextureGuid = ReadPreshaderValue<FGuid>(Data);
@@ -252,7 +257,7 @@ static FGuid GetExternalTextureGuid(const FMaterialRenderContext& Context, UE::S
 	return Context.GetExternalTextureGuid(ExternalTextureGuid, ScriptNameToName(ParameterName), TextureIndex);
 }
 
-static void EvaluateExternalTextureCoordinateScaleRotation(const FMaterialRenderContext& Context, UE::Shader::FPreshaderStack& Stack, UE::Shader::FPreshaderDataContext& RESTRICT Data)
+static void EvaluateExternalTextureCoordinateScaleRotation(const FMaterialRenderContext& Context, FPreshaderStack& Stack, FPreshaderDataContext& RESTRICT Data)
 {
 	const FGuid GuidToLookup = GetExternalTextureGuid(Context, Data);
 	FLinearColor Result(1.f, 0.f, 0.f, 1.f);
@@ -263,7 +268,7 @@ static void EvaluateExternalTextureCoordinateScaleRotation(const FMaterialRender
 	Stack.Add(Result);
 }
 
-static void EvaluateExternalTextureCoordinateOffset(const FMaterialRenderContext& Context, UE::Shader::FPreshaderStack& Stack, UE::Shader::FPreshaderDataContext& RESTRICT Data)
+static void EvaluateExternalTextureCoordinateOffset(const FMaterialRenderContext& Context, FPreshaderStack& Stack, FPreshaderDataContext& RESTRICT Data)
 {
 	const FGuid GuidToLookup = GetExternalTextureGuid(Context, Data);
 	FLinearColor Result(0.f, 0.f, 0.f, 0.f);
@@ -274,7 +279,7 @@ static void EvaluateExternalTextureCoordinateOffset(const FMaterialRenderContext
 	Stack.Add(Result);
 }
 
-static void EvaluateRuntimeVirtualTextureUniform(const FMaterialRenderContext& Context, UE::Shader::FPreshaderStack& Stack, UE::Shader::FPreshaderDataContext& RESTRICT Data)
+static void EvaluateRuntimeVirtualTextureUniform(const FMaterialRenderContext& Context, FPreshaderStack& Stack, FPreshaderDataContext& RESTRICT Data)
 {
 	const FHashedMaterialParameterInfo ParameterInfo = ReadPreshaderValue<FHashedMaterialParameterInfo>(Data);
 	const int32 TextureIndex = ReadPreshaderValue<int32>(Data);
@@ -287,15 +292,15 @@ static void EvaluateRuntimeVirtualTextureUniform(const FMaterialRenderContext& C
 	}
 	if (Texture != nullptr && VectorIndex != INDEX_NONE)
 	{
-		Stack.Add(UE::Shader::FValue(Texture->GetUniformParameter(VectorIndex)));
+		Stack.Add(FValue(Texture->GetUniformParameter(VectorIndex)));
 	}
 	else
 	{
-		Stack.Add(UE::Shader::FValue(0.f, 0.f, 0.f, 0.f));
+		Stack.Add(FValue(0.f, 0.f, 0.f, 0.f));
 	}
 }
 
-void UE::Shader::EvaluatePreshader(const FUniformExpressionSet* UniformExpressionSet, const FMaterialRenderContext& Context, FPreshaderStack& Stack, FPreshaderDataContext& RESTRICT Data, FValue& OutValue)
+void EvaluatePreshader(const FUniformExpressionSet* UniformExpressionSet, const FMaterialRenderContext& Context, FPreshaderStack& Stack, FPreshaderDataContext& RESTRICT Data, FValue& OutValue)
 {
 	uint8 const* const DataEnd = Data.EndPtr;
 
@@ -306,49 +311,49 @@ void UE::Shader::EvaluatePreshader(const FUniformExpressionSet* UniformExpressio
 		switch (Opcode)
 		{
 		case EPreshaderOpcode::ConstantZero:
-			Stack.Add(UE::Shader::FValue(0.0f));
+			Stack.Add(FValue(0.0f));
 			break;
 		case EPreshaderOpcode::Constant:
-			Stack.Add(ReadPreshaderValue<UE::Shader::FValue>(Data));
+			Stack.Add(ReadPreshaderValue<FValue>(Data));
 			break;
 		case EPreshaderOpcode::Parameter:
 			check(UniformExpressionSet);
 			EvaluateParameter(*UniformExpressionSet, ReadPreshaderValue<uint16>(Data), Context, Stack.AddDefaulted_GetRef());
 			break;
-		case EPreshaderOpcode::Add: EvaluateBinaryOp(Stack, UE::Shader::Add); break;
-		case EPreshaderOpcode::Sub: EvaluateBinaryOp(Stack, UE::Shader::Sub); break;
-		case EPreshaderOpcode::Mul: EvaluateBinaryOp(Stack, UE::Shader::Mul); break;
-		case EPreshaderOpcode::Div: EvaluateBinaryOp(Stack, UE::Shader::Div); break;
-		case EPreshaderOpcode::Fmod: EvaluateBinaryOp(Stack, UE::Shader::Fmod); break;
-		case EPreshaderOpcode::Min: EvaluateBinaryOp(Stack, UE::Shader::Min); break;
-		case EPreshaderOpcode::Max: EvaluateBinaryOp(Stack, UE::Shader::Max); break;
-		case EPreshaderOpcode::Clamp: EvaluateTernaryOp(Stack, UE::Shader::Clamp); break;
-		case EPreshaderOpcode::Dot: EvaluateBinaryOp(Stack, UE::Shader::Dot); break;
-		case EPreshaderOpcode::Cross: EvaluateBinaryOp(Stack, UE::Shader::Cross); break;
-		case EPreshaderOpcode::Sqrt: EvaluateUnaryOp(Stack, UE::Shader::Sqrt); break;
-		case EPreshaderOpcode::Rcp: EvaluateUnaryOp(Stack, UE::Shader::Rcp); break;
-		case EPreshaderOpcode::Length: EvaluateUnaryOp(Stack, [](const UE::Shader::FValue& Value) { return Sqrt(Dot(Value, Value)); }); break;
-		case EPreshaderOpcode::Normalize: EvaluateUnaryOp(Stack, [](const UE::Shader::FValue& Value) { return Div(Value, Sqrt(Dot(Value, Value))); }); break;
-		case EPreshaderOpcode::Sin: EvaluateUnaryOp(Stack, UE::Shader::Sin); break;
-		case EPreshaderOpcode::Cos: EvaluateUnaryOp(Stack, UE::Shader::Cos); break;
-		case EPreshaderOpcode::Tan: EvaluateUnaryOp(Stack, UE::Shader::Tan); break;
-		case EPreshaderOpcode::Asin: EvaluateUnaryOp(Stack, UE::Shader::Asin); break;;
-		case EPreshaderOpcode::Acos: EvaluateUnaryOp(Stack, UE::Shader::Acos); break;
-		case EPreshaderOpcode::Atan: EvaluateUnaryOp(Stack, UE::Shader::Atan); break;
-		case EPreshaderOpcode::Atan2: EvaluateBinaryOp(Stack, UE::Shader::Atan2); break;
-		case EPreshaderOpcode::Abs: EvaluateUnaryOp(Stack, UE::Shader::Abs); break;
-		case EPreshaderOpcode::Saturate: EvaluateUnaryOp(Stack, UE::Shader::Saturate); break;
-		case EPreshaderOpcode::Floor: EvaluateUnaryOp(Stack, UE::Shader::Floor); break;
-		case EPreshaderOpcode::Ceil: EvaluateUnaryOp(Stack, UE::Shader::Ceil); break;
-		case EPreshaderOpcode::Round: EvaluateUnaryOp(Stack, UE::Shader::Round); break;
-		case EPreshaderOpcode::Trunc: EvaluateUnaryOp(Stack, UE::Shader::Trunc); break;
-		case EPreshaderOpcode::Sign: EvaluateUnaryOp(Stack, UE::Shader::Sign); break;
-		case EPreshaderOpcode::Frac: EvaluateUnaryOp(Stack, UE::Shader::Frac); break;
-		case EPreshaderOpcode::Fractional: EvaluateUnaryOp(Stack, UE::Shader::Fractional); break;
-		case EPreshaderOpcode::Log2: EvaluateUnaryOp(Stack, UE::Shader::Log2); break;
-		case EPreshaderOpcode::Log10: EvaluateUnaryOp(Stack, UE::Shader::Log10); break;
+		case EPreshaderOpcode::Add: EvaluateBinaryOp(Stack, Add); break;
+		case EPreshaderOpcode::Sub: EvaluateBinaryOp(Stack, Sub); break;
+		case EPreshaderOpcode::Mul: EvaluateBinaryOp(Stack, Mul); break;
+		case EPreshaderOpcode::Div: EvaluateBinaryOp(Stack, Div); break;
+		case EPreshaderOpcode::Fmod: EvaluateBinaryOp(Stack, Fmod); break;
+		case EPreshaderOpcode::Min: EvaluateBinaryOp(Stack, Min); break;
+		case EPreshaderOpcode::Max: EvaluateBinaryOp(Stack, Max); break;
+		case EPreshaderOpcode::Clamp: EvaluateTernaryOp(Stack, Clamp); break;
+		case EPreshaderOpcode::Dot: EvaluateBinaryOp(Stack, Dot); break;
+		case EPreshaderOpcode::Cross: EvaluateBinaryOp(Stack, Cross); break;
+		case EPreshaderOpcode::Sqrt: EvaluateUnaryOp(Stack, Sqrt); break;
+		case EPreshaderOpcode::Rcp: EvaluateUnaryOp(Stack, Rcp); break;
+		case EPreshaderOpcode::Length: EvaluateUnaryOp(Stack, [](const FValue& Value) { return Sqrt(Dot(Value, Value)); }); break;
+		case EPreshaderOpcode::Normalize: EvaluateUnaryOp(Stack, [](const FValue& Value) { return Div(Value, Sqrt(Dot(Value, Value))); }); break;
+		case EPreshaderOpcode::Sin: EvaluateUnaryOp(Stack, Sin); break;
+		case EPreshaderOpcode::Cos: EvaluateUnaryOp(Stack, Cos); break;
+		case EPreshaderOpcode::Tan: EvaluateUnaryOp(Stack, Tan); break;
+		case EPreshaderOpcode::Asin: EvaluateUnaryOp(Stack, Asin); break;;
+		case EPreshaderOpcode::Acos: EvaluateUnaryOp(Stack, Acos); break;
+		case EPreshaderOpcode::Atan: EvaluateUnaryOp(Stack, Atan); break;
+		case EPreshaderOpcode::Atan2: EvaluateBinaryOp(Stack, Atan2); break;
+		case EPreshaderOpcode::Abs: EvaluateUnaryOp(Stack, Abs); break;
+		case EPreshaderOpcode::Saturate: EvaluateUnaryOp(Stack, Saturate); break;
+		case EPreshaderOpcode::Floor: EvaluateUnaryOp(Stack, Floor); break;
+		case EPreshaderOpcode::Ceil: EvaluateUnaryOp(Stack, Ceil); break;
+		case EPreshaderOpcode::Round: EvaluateUnaryOp(Stack, Round); break;
+		case EPreshaderOpcode::Trunc: EvaluateUnaryOp(Stack, Trunc); break;
+		case EPreshaderOpcode::Sign: EvaluateUnaryOp(Stack, Sign); break;
+		case EPreshaderOpcode::Frac: EvaluateUnaryOp(Stack, Frac); break;
+		case EPreshaderOpcode::Fractional: EvaluateUnaryOp(Stack, Fractional); break;
+		case EPreshaderOpcode::Log2: EvaluateUnaryOp(Stack, Log2); break;
+		case EPreshaderOpcode::Log10: EvaluateUnaryOp(Stack, Log10); break;
 		case EPreshaderOpcode::ComponentSwizzle: EvaluateComponentSwizzle(Stack, Data); break;
-		case EPreshaderOpcode::AppendVector: EvaluateBinaryOp(Stack, UE::Shader::Append); break;
+		case EPreshaderOpcode::AppendVector: EvaluateBinaryOp(Stack, Append); break;
 		case EPreshaderOpcode::TextureSize: EvaluateTextureSize(Context, Stack, Data); break;
 		case EPreshaderOpcode::TexelSize: EvaluateTexelSize(Context, Stack, Data); break;
 		case EPreshaderOpcode::ExternalTextureCoordinateScaleRotation: EvaluateExternalTextureCoordinateScaleRotation(Context, Stack, Data); break;
@@ -368,9 +373,26 @@ void UE::Shader::EvaluatePreshader(const FUniformExpressionSet* UniformExpressio
 	}
 }
 
-void UE::Shader::FPreshaderData::Evaluate(FUniformExpressionSet* UniformExpressionSet, const struct FMaterialRenderContext& Context, FValue& OutValue)
+void FPreshaderData::Evaluate(FUniformExpressionSet* UniformExpressionSet, const struct FMaterialRenderContext& Context, FValue& OutValue)
 {
 	FPreshaderStack Stack;
 	FPreshaderDataContext PreshaderContext(*this);
 	EvaluatePreshader(UniformExpressionSet, Context, Stack, PreshaderContext, OutValue);
 }
+
+void FPreshaderData::AppendHash(FSHA1& OutHasher) const
+{
+	OutHasher.Update((uint8*)Names.GetData(), Names.Num() * Names.GetTypeSize());
+	OutHasher.Update((uint8*)NameOffsets.GetData(), NameOffsets.Num() * NameOffsets.GetTypeSize());
+	OutHasher.Update((uint8*)Data.GetData(), Data.Num() * Data.GetTypeSize());
+}
+
+FSHAHash FPreshaderData::GetHash() const
+{
+	FSHA1 Hasher;
+	AppendHash(Hasher);
+	return Hasher.Finalize();
+}
+
+} // namespace Shader
+} // namespace UE
diff --git a/Engine/Source/Runtime/Engine/Public/HLSLTree/HLSLTree.h b/Engine/Source/Runtime/Engine/Public/HLSLTree/HLSLTree.h
--- a/Engine/Source/Runtime/Engine/Public/HLSLTree/HLSLTree.h
+++ b/Engine/Source/Runtime/Engine/Public/HLSLTree/HLSLTree.h
@@ -136,7 +136,7 @@ public:
 
 	const TCHAR* CastShaderValue(const FNode* Node, const TCHAR* Code, const FType& SourceType, const FType& DestType, ECastFlags Flags);
 
-	void AddPreshader(Shader::EValueType Type, const Shader::FPreshaderData& Preshader, FStringBuilderBase& OutCode);
+	const TCHAR* AcquirePreshader(Shader::EValueType Type, const Shader::FPreshaderData& Preshader);
 
 	FMemStackBase* Allocator = nullptr;
 	FErrors Errors;
@@ -146,6 +146,7 @@ public:
 	const FStaticParameterSet* StaticParameters = nullptr;
 	FMaterialCompilationOutput* MaterialCompilationOutput = nullptr;
 	TMap<Shader::FValue, uint32> DefaultUniformValues;
+	TMap<FSHAHash, const TCHAR*> Preshaders;
 	uint32 UniformPreshaderOffset = 0u;
 	bool bReadMaterialNormal = false;
 
@@ -184,6 +185,8 @@ public:
 
 	virtual ENodeVisitResult Visit(FNodeVisitor& Visitor) = 0;
 
+	virtual void Reset() {}
+
 	FNode* NextNode = nullptr;
 };
 
@@ -250,9 +253,10 @@ struct FStructTypeInitializer
 class FStatement : public FNode
 {
 public:
-	static constexpr bool MarkScopeLive = false;
-	static constexpr bool MarkScopeLiveRecursive = false;
+	//static constexpr bool MarkScopeLive = false;
+	//static constexpr bool MarkScopeLiveRecursive = false;
 
+	virtual void Reset() override;
 	virtual ENodeVisitResult Visit(FNodeVisitor& Visitor) override;
 
 	virtual void PrepareValues(FEmitContext& Context) const = 0;
@@ -282,23 +286,25 @@ public:
 		RequestedComponents.Reset();
 	}
 
-	/** Returns a type that represents access of FieldRef from a struct on given type, using 'this' */
-	FRequestedType MakeFieldAccess(const FStructType* InStructType, const FStructFieldRef& FieldRef) const
+	/** Marks the given field as requested (or not) */
+	void SetFieldRequested(const FStructFieldRef& FieldRef, bool bRequested = true)
 	{
-		FRequestedType Result(InStructType, false);
-		Result.RequestedComponents.SetRangeFromRange(FieldRef.ComponentIndex, FieldRef.ComponentNum, RequestedComponents, 0);
-		return Result;
+		RequestedComponents.SetRange(FieldRef.ComponentIndex, FieldRef.ComponentNum, bRequested);
 	}
 
-	/** Returns copy of 'this', with request for given field removed */
-	FRequestedType CopyWithFieldRemoved(const FStructFieldRef& FieldRef) const
+	void ClearFieldRequested(const FStructFieldRef& FieldRef)
 	{
-		FRequestedType Result(*this);
-		Result.RequestedComponents.SetRange(FieldRef.ComponentIndex, FieldRef.ComponentNum, false);
-		return Result;
+		SetFieldRequested(FieldRef, false);
 	}
 
-	/** Returns a type that represents access of the given field from 'this' */
+	/** Marks the given field as requested, based on the input request type (which should match the field type) */
+	void SetField(const FStructFieldRef& FieldRef, const FRequestedType& InRequest)
+	{
+		ensure(InRequest.GetNumComponents() == FieldRef.ComponentNum);
+		RequestedComponents.SetRangeFromRange(FieldRef.ComponentIndex, FieldRef.ComponentNum, InRequest.RequestedComponents, 0);
+	}
+
+	/** Returns the requested type of the given field */
 	FRequestedType GetField(const FStructFieldRef& FieldRef) const
 	{
 		FRequestedType Result(FieldRef.Type);
@@ -345,6 +351,7 @@ public:
 	const FType& GetType() const { return PrepareValueResult.Type; }
 	EExpressionEvaluationType GetEvaluationType() const { return PrepareValueResult.EvaluationType; }
 
+	virtual void Reset() override;
 	virtual ENodeVisitResult Visit(FNodeVisitor& Visitor) override;
 
 	friend FPrepareValueResult PrepareExpressionValue(FEmitContext& Context, FExpression* InExpression, const FRequestedType& RequestedType);
@@ -420,6 +427,8 @@ enum class EScopeState : uint8
 class FScope final : public FNode
 {
 public:
+	virtual void Reset() override;
+
 	static FScope* FindSharedParent(FScope* Lhs, FScope* Rhs);
 
 	inline FScope* GetParentScope() const { return ParentScope; }
@@ -440,7 +449,7 @@ public:
 
 	void UseExpression(FExpression* Expression);
 
-	friend void PrepareScopeValues(FEmitContext& Context, const FScope* InScope);
+	friend bool PrepareScopeValues(FEmitContext& Context, FScope* InScope);
 
 	template<typename FormatType, typename... Types>
 	void EmitDeclarationf(FEmitContext& Context, const FormatType& Format, Types... Args)
@@ -509,8 +518,9 @@ private:
 		InternalEmitCode(Context, List, ScopeFormat, NextScope, String.ToString(), String.Len());
 	}
 
+	FStatement* OwnerStatement = nullptr;
 	FScope* ParentScope = nullptr;
-	FStatement* Statement = nullptr;
+	FStatement* ContainedStatement = nullptr;
 	FScope* PreviousScope[MaxNumPreviousScopes];
 	TMap<FSHAHash, const TCHAR*> ExpressionCodeMap;
 	FCodeList Declarations;
@@ -525,6 +535,14 @@ inline bool IsScopeLive(const FScope* InScope)
 	return (bool)InScope && InScope->IsLive();
 }
 
+inline void MarkScopeLive(FScope* InScope)
+{
+	if (InScope)
+	{
+		InScope->MarkLive();
+	}
+}
+
 inline void MarkScopeDead(FScope* InScope)
 {
 	if (InScope)
@@ -544,9 +562,11 @@ public:
 
 	FMemStackBase& GetAllocator() { return *Allocator; }
 
-	bool EmitHLSL(FEmitContext& Context,
-		FStringBuilderBase& OutDeclarations,
-		FStringBuilderBase& OutCode) const;
+	void ResetNodes();
+
+	void EmitDeclarationsCode(FStringBuilderBase& OutCode) const;
+
+	bool EmitHLSL(FEmitContext& Context, FStringBuilderBase& OutCode) const;
 
 	FScope& GetRootScope() const { return *RootScope; }
 
@@ -563,25 +583,25 @@ public:
 	{
 		T* Statement = NewNode<T>(Forward<ArgTypes>(Args)...);
 		RegisterStatement(Scope, Statement);
-		if constexpr (T::MarkScopeLiveRecursive)
+		/*if constexpr (T::MarkScopeLiveRecursive)
 		{
 			Scope.MarkLiveRecursive();
 		}
 		else if constexpr (T::MarkScopeLive)
 		{
 			Scope.MarkLive();
-		}
+		}*/
 
 		return Statement;
 	}
 
 	FScope* NewScope(FScope& Scope);
+	FScope* NewOwnedScope(FStatement& Owner);
+
 	FTextureParameterDeclaration* NewTextureParameterDeclaration(const FName& Name, const FTextureDescription& DefaultValue);
 
 	const FStructType* NewStructType(const FStructTypeInitializer& Initializer);
 
-	void SetResult(FStatement& InResult);
-
 private:
 	template<typename T, typename... ArgTypes>
 	inline T* NewNode(ArgTypes&&... Args)
@@ -599,7 +619,6 @@ private:
 	FNode* Nodes = nullptr;
 	FStructType* StructTypes = nullptr;
 	FScope* RootScope = nullptr;
-	FStatement* ResultStatement = nullptr;
 };
 
 //Shader::EValueType RequestExpressionType(FEmitContext& Context, FExpression* InExpression, int8 InRequestedNumComponents); // friend of FExpression
diff --git a/Engine/Source/Runtime/Engine/Public/HLSLTree/HLSLTreeCommon.h b/Engine/Source/Runtime/Engine/Public/HLSLTree/HLSLTreeCommon.h
--- a/Engine/Source/Runtime/Engine/Public/HLSLTree/HLSLTreeCommon.h
+++ b/Engine/Source/Runtime/Engine/Public/HLSLTree/HLSLTreeCommon.h
@@ -309,9 +309,7 @@ public:
 class FStatementReturn : public FStatement
 {
 public:
-	static constexpr bool MarkScopeLiveRecursive = true;
-
-	FType Type;
+	//static constexpr bool MarkScopeLiveRecursive = true;
 	FExpression* Expression;
 
 	virtual ENodeVisitResult Visit(FNodeVisitor& Visitor) override
@@ -331,9 +329,9 @@ public:
 class FStatementBreak : public FStatement
 {
 public:
-	static constexpr bool MarkScopeLive = true;
+	//static constexpr bool MarkScopeLive = true;
 
-	virtual void PrepareValues(FEmitContext& Context) const override {}
+	virtual void PrepareValues(FEmitContext& Context) const override;
 	virtual void EmitHLSL(FEmitContext& Context) const override;
 };
 
diff --git a/Engine/Source/Runtime/Engine/Public/MaterialHLSLGenerator.h b/Engine/Source/Runtime/Engine/Public/MaterialHLSLGenerator.h
--- a/Engine/Source/Runtime/Engine/Public/MaterialHLSLGenerator.h
+++ b/Engine/Source/Runtime/Engine/Public/MaterialHLSLGenerator.h
@@ -64,7 +64,7 @@ public:
 	bool Finalize();
 
 	/** Retrieve the compile errors from the generator */
-	void AcquireErrors(TArray<FString>& OutCompileErrors, TArray<UMaterialExpression*>& OutErrorExpressions);
+	void AcquireErrors(FMaterial& OutMaterial);
 
 	EMaterialGenerateHLSLStatus Error(const FString& Message);
 
@@ -76,9 +76,13 @@ public:
 
 	UE::HLSLTree::FTree& GetTree() const { return *HLSLTree; }
 
+	UE::HLSLTree::FExpression* GetResultExpression() { return ResultExpression; }
+	UE::HLSLTree::FStatement* GetResultStatement() { return ResultStatement; }
+
 	bool GenerateResult(UE::HLSLTree::FScope& Scope);
 
 	UE::HLSLTree::FScope* NewScope(UE::HLSLTree::FScope& Scope, EMaterialNewScopeFlag Flags = EMaterialNewScopeFlag::None);
+	UE::HLSLTree::FScope* NewOwnedScope(UE::HLSLTree::FStatement& Owner);
 
 	UE::HLSLTree::FScope* NewJoinedScope(UE::HLSLTree::FScope& Scope);
 
@@ -215,6 +219,9 @@ private:
 	UE::HLSLTree::FConstantValue MaterialAttributesDefaultValue;
 
 	UE::HLSLTree::FTree* HLSLTree;
+	UE::HLSLTree::FExpression* ResultExpression = nullptr;
+	UE::HLSLTree::FStatement* ResultStatement = nullptr;
+
 	TArray<FExpressionKey> ExpressionStack;
 	TArray<FFunctionCallEntry*> FunctionCallStack;
 	TArray<UE::HLSLTree::FScope*> JoinedScopeStack;
diff --git a/Engine/Source/Runtime/Engine/Public/MaterialHLSLTree.h b/Engine/Source/Runtime/Engine/Public/MaterialHLSLTree.h
deleted file mode 100644
diff --git a/Engine/Source/Runtime/Engine/Public/MaterialShared.h b/Engine/Source/Runtime/Engine/Public/MaterialShared.h
--- a/Engine/Source/Runtime/Engine/Public/MaterialShared.h
+++ b/Engine/Source/Runtime/Engine/Public/MaterialShared.h
@@ -64,6 +64,7 @@ class FMeshMaterialShaderMapLayout;
 struct FMaterialShaderTypes;
 struct FMaterialShaders;
 struct FMaterialCachedExpressionData;
+class FMaterialHLSLGenerator;
 
 namespace UE
 {
@@ -2286,7 +2287,7 @@ private:
 	friend class FMaterialShaderMap;
 	friend class FShaderCompilingManager;
 	friend class FHLSLMaterialTranslator;
-	friend class FMaterialHLSLTree;
+	friend class FMaterialHLSLGenerator;
 };
 
 /**
diff --git a/Engine/Source/Runtime/Engine/Public/Shader/Preshader.h b/Engine/Source/Runtime/Engine/Public/Shader/Preshader.h
--- a/Engine/Source/Runtime/Engine/Public/Shader/Preshader.h
+++ b/Engine/Source/Runtime/Engine/Public/Shader/Preshader.h
@@ -67,7 +67,7 @@ class FPreshaderData
 public:
 	friend inline bool operator==(const FPreshaderData& Lhs, const FPreshaderData& Rhs)
 	{
-		return Lhs.Data == Rhs.Data;
+		return Lhs.Names == Rhs.Names && Lhs.Data == Rhs.Data;
 	}
 
 	friend inline bool operator!=(const FPreshaderData& Lhs, const FPreshaderData& Rhs)
@@ -75,6 +75,9 @@ public:
 		return !operator==(Lhs, Rhs);
 	}
 
+	FSHAHash GetHash() const;
+	void AppendHash(FSHA1& OutHasher) const;
+
 	void Evaluate(FUniformExpressionSet* UniformExpressionSet, const struct FMaterialRenderContext& Context, FValue& OutValue);
 
 	void Append(const FPreshaderData& InPreshader);
