Bulletproofing and logging around not finding a sound class

diff --git a/Engine/Source/Runtime/Engine/Private/SoundWave.cpp b/Engine/Source/Runtime/Engine/Private/SoundWave.cpp
--- a/Engine/Source/Runtime/Engine/Private/SoundWave.cpp
+++ b/Engine/Source/Runtime/Engine/Private/SoundWave.cpp
@@ -2248,49 +2248,58 @@ void USoundWave::Parse(FAudioDevice* AudioDevice, const UPTRINT NodeWaveInstance
 
 	// Properties from the sound class
 	WaveInstance->SoundClass = ParseParams.SoundClass;
+	bool bFoundSoundClass = false;
 	if (ParseParams.SoundClass)
 	{
 		const FSoundClassProperties* SoundClassProperties = AudioDevice->GetSoundClassCurrentProperties(ParseParams.SoundClass);
-		check(SoundClassProperties);
+		if (!SoundClassProperties)
+		{
+			UE_LOG(LogAudio, Warning, TEXT("Unable to find sound class properties for sound class %s"), *ParseParams.SoundClass->GetName());
+		}
+		else
+		{
+			bFoundSoundClass = true;
 
-		// Use values from "parsed/ propagated" sound class properties
-		float VolumeMultiplier = WaveInstance->GetVolumeMultiplier();
-		WaveInstance->SetVolumeMultiplier(VolumeMultiplier * SoundClassProperties->Volume);
-		WaveInstance->SetPitch(WaveInstance->Pitch * SoundClassProperties->Pitch);
+			// Use values from "parsed/ propagated" sound class properties
+			float VolumeMultiplier = WaveInstance->GetVolumeMultiplier();
+			WaveInstance->SetVolumeMultiplier(VolumeMultiplier * SoundClassProperties->Volume);
+			WaveInstance->SetPitch(WaveInstance->Pitch * SoundClassProperties->Pitch);
 
-		WaveInstance->SoundClassFilterFrequency = SoundClassProperties->LowPassFilterFrequency;
-		WaveInstance->VoiceCenterChannelVolume = SoundClassProperties->VoiceCenterChannelVolume;
-		WaveInstance->RadioFilterVolume = SoundClassProperties->RadioFilterVolume * ParseParams.VolumeMultiplier;
-		WaveInstance->RadioFilterVolumeThreshold = SoundClassProperties->RadioFilterVolumeThreshold * ParseParams.VolumeMultiplier;
-		WaveInstance->LFEBleed = SoundClassProperties->LFEBleed;
+			WaveInstance->SoundClassFilterFrequency = SoundClassProperties->LowPassFilterFrequency;
+			WaveInstance->VoiceCenterChannelVolume = SoundClassProperties->VoiceCenterChannelVolume;
+			WaveInstance->RadioFilterVolume = SoundClassProperties->RadioFilterVolume * ParseParams.VolumeMultiplier;
+			WaveInstance->RadioFilterVolumeThreshold = SoundClassProperties->RadioFilterVolumeThreshold * ParseParams.VolumeMultiplier;
+			WaveInstance->LFEBleed = SoundClassProperties->LFEBleed;
 
-		WaveInstance->bIsUISound = ActiveSound.bIsUISound || SoundClassProperties->bIsUISound;
-		WaveInstance->bIsMusic = ActiveSound.bIsMusic || SoundClassProperties->bIsMusic;
-		WaveInstance->bCenterChannelOnly = ActiveSound.bCenterChannelOnly || SoundClassProperties->bCenterChannelOnly;
-		WaveInstance->bReverb = ActiveSound.bReverb || SoundClassProperties->bReverb;
-		WaveInstance->OutputTarget = SoundClassProperties->OutputTarget;
+			WaveInstance->bIsUISound = ActiveSound.bIsUISound || SoundClassProperties->bIsUISound;
+			WaveInstance->bIsMusic = ActiveSound.bIsMusic || SoundClassProperties->bIsMusic;
+			WaveInstance->bCenterChannelOnly = ActiveSound.bCenterChannelOnly || SoundClassProperties->bCenterChannelOnly;
+			WaveInstance->bReverb = ActiveSound.bReverb || SoundClassProperties->bReverb;
+			WaveInstance->OutputTarget = SoundClassProperties->OutputTarget;
 
-		if (SoundClassProperties->bApplyEffects)
-		{
-			UAudioSettings* Settings = GetMutableDefault<UAudioSettings>();
-			WaveInstance->SoundSubmix = Cast<USoundSubmix>(FSoftObjectPtr(Settings->EQSubmix).Get());
-		}
-		else if (SoundClassProperties->DefaultSubmix)
-		{
-			WaveInstance->SoundSubmix = SoundClassProperties->DefaultSubmix;
-		}
+			if (SoundClassProperties->bApplyEffects)
+			{
+				UAudioSettings* Settings = GetMutableDefault<UAudioSettings>();
+				WaveInstance->SoundSubmix = Cast<USoundSubmix>(FSoftObjectPtr(Settings->EQSubmix).Get());
+			}
+			else if (SoundClassProperties->DefaultSubmix)
+			{
+				WaveInstance->SoundSubmix = SoundClassProperties->DefaultSubmix;
+			}
 
-		if (SoundClassProperties->bApplyAmbientVolumes)
-		{
-			VolumeMultiplier = WaveInstance->GetVolumeMultiplier();
-			WaveInstance->SetVolumeMultiplier(VolumeMultiplier * ParseParams.InteriorVolumeMultiplier);
-			WaveInstance->RadioFilterVolume *= ParseParams.InteriorVolumeMultiplier;
-			WaveInstance->RadioFilterVolumeThreshold *= ParseParams.InteriorVolumeMultiplier;
-		}
+			if (SoundClassProperties->bApplyAmbientVolumes)
+			{
+				VolumeMultiplier = WaveInstance->GetVolumeMultiplier();
+				WaveInstance->SetVolumeMultiplier(VolumeMultiplier * ParseParams.InteriorVolumeMultiplier);
+				WaveInstance->RadioFilterVolume *= ParseParams.InteriorVolumeMultiplier;
+				WaveInstance->RadioFilterVolumeThreshold *= ParseParams.InteriorVolumeMultiplier;
+			}
 
-		bAlwaysPlay = ActiveSound.bAlwaysPlay || SoundClassProperties->bAlwaysPlay;
+			bAlwaysPlay = ActiveSound.bAlwaysPlay || SoundClassProperties->bAlwaysPlay;
+		}
 	}
-	else
+
+	if (!bFoundSoundClass)
 	{
 		WaveInstance->VoiceCenterChannelVolume = 0.f;
 		WaveInstance->RadioFilterVolume = 0.f;
