Rename/cleanup some HLSLTree stuff

diff --git a/Engine/Source/Runtime/Engine/Private/HLSLTree/HLSLTree.cpp b/Engine/Source/Runtime/Engine/Private/HLSLTree/HLSLTree.cpp
--- a/Engine/Source/Runtime/Engine/Private/HLSLTree/HLSLTree.cpp
+++ b/Engine/Source/Runtime/Engine/Private/HLSLTree/HLSLTree.cpp
@@ -435,7 +435,7 @@ void FExpressionLocalPHI::PrepareValue(FEmitContext& Context, const FRequestedTy
 	for (int32 i = 0; i < NumValues; ++i)
 	{
 		// Ignore values in dead scopes
-		if (PrepareScopeValues(Context, Scopes[i]))
+		if (PrepareScope(Context, Scopes[i]))
 		{
 			FExpression* ScopeExpression = Values[i];
 			if (!ForwardExpression)
@@ -464,7 +464,7 @@ void FExpressionLocalPHI::PrepareValue(FEmitContext& Context, const FRequestedTy
 	{
 		for (int32 i = 0; i < NumValues; ++i)
 		{
-			if (!TypePerValue[i] && PrepareScopeValues(Context, Scopes[i]))
+			if (!TypePerValue[i] && PrepareScope(Context, Scopes[i]))
 			{
 				const FPrepareValueResult ValueResult = PrepareExpressionValue(Context, Values[i], RequestedType);
 				if (ValueResult.Type)
@@ -588,7 +588,7 @@ ENodeVisitResult FScope::Visit(FNodeVisitor& Visitor)
 
 void FStatement::Reset()
 {
-	bEmitHLSL = false;
+	bEmitShader = false;
 }
 
 void FExpression::Reset()
@@ -882,12 +882,12 @@ void FScope::UseExpression(FExpression* Expression)
 	Visitor.VisitNode(Expression);
 }
 
-void FScope::InternalEmitCode(FEmitContext& Context, FCodeList& List, ENextScopeFormat ScopeFormat, FScope* NextScope, const TCHAR* String, int32 Length)
+void FScope::InternalEmitCode(FEmitContext& Context, FCodeList& List, ENextScopeFormat ScopeFormat, FScope* Scope, const TCHAR* String, int32 Length)
 {
-	if (NextScope && NextScope->ContainedStatement && !NextScope->ContainedStatement->bEmitHLSL)
+	if (Scope && Scope->ContainedStatement && !Scope->ContainedStatement->bEmitShader)
 	{
-		NextScope->ContainedStatement->bEmitHLSL = true;
-		NextScope->ContainedStatement->EmitHLSL(Context);
+		Scope->ContainedStatement->bEmitShader = true;
+		Scope->ContainedStatement->EmitShader(Context);
 	}
 
 	const int32 SizeofString = sizeof(TCHAR) * Length;
@@ -896,7 +896,7 @@ void FScope::InternalEmitCode(FEmitContext& Context, FCodeList& List, ENextScope
 	FMemory::Memcpy(CodeEntry->String, String, SizeofString);
 	CodeEntry->String[Length] = 0;
 	CodeEntry->Length = Length;
-	CodeEntry->NextScope = NextScope;
+	CodeEntry->Scope = Scope;
 	CodeEntry->ScopeFormat = ScopeFormat;
 	CodeEntry->Next = nullptr;
 
@@ -913,15 +913,15 @@ void FScope::InternalEmitCode(FEmitContext& Context, FCodeList& List, ENextScope
 	List.Num++;
 }
 
-bool PrepareScopeValues(FEmitContext& Context, FScope* InScope)
+bool PrepareScope(FEmitContext& Context, FScope* InScope)
 {
 	if (InScope && InScope->State == EScopeState::Uninitialized)
 	{
-		if (!InScope->ParentScope || PrepareScopeValues(Context, InScope->ParentScope))
+		if (!InScope->ParentScope || PrepareScope(Context, InScope->ParentScope))
 		{
 			if (InScope->OwnerStatement)
 			{
-				InScope->OwnerStatement->PrepareValues(Context);
+				InScope->OwnerStatement->Prepare(Context);
 			}
 			else
 			{
@@ -1004,7 +1004,7 @@ void FScope::WriteHLSL(int32 Indent, FStringBuilderBase& OutString) const
 		const FCodeEntry* CodeDeclaration = Declarations.First;
 		while (CodeDeclaration)
 		{
-			check(!CodeDeclaration->NextScope);
+			check(!CodeDeclaration->Scope);
 			WriteIndent(Indent, OutString);
 			OutString.Append(CodeDeclaration->String, CodeDeclaration->Length);
 			OutString.Append(TEXT('\n'));
@@ -1022,7 +1022,7 @@ void FScope::WriteHLSL(int32 Indent, FStringBuilderBase& OutString) const
 				OutString.Append(CodeStatement->String, CodeStatement->Length);
 				OutString.Append(TEXT('\n'));
 			}
-			if (CodeStatement->NextScope)
+			if (CodeStatement->Scope)
 			{
 				int32 NextIndent = Indent;
 				bool bNeedToCloseScope = false;
@@ -1034,7 +1034,7 @@ void FScope::WriteHLSL(int32 Indent, FStringBuilderBase& OutString) const
 					bNeedToCloseScope = true;
 				}
 
-				CodeStatement->NextScope->WriteHLSL(NextIndent, OutString);
+				CodeStatement->Scope->WriteHLSL(NextIndent, OutString);
 				if (bNeedToCloseScope)
 				{
 					WriteIndent(Indent, OutString);
@@ -1085,11 +1085,12 @@ void FTree::EmitDeclarationsCode(FStringBuilderBase& OutCode) const
 	}
 }
 
-bool FTree::EmitHLSL(FEmitContext& Context, FStringBuilderBase& OutCode) const
+bool FTree::EmitShader(FEmitContext& Context, FStringBuilderBase& OutCode) const
 {
 	if (RootScope->ContainedStatement)
 	{
-		RootScope->ContainedStatement->EmitHLSL(Context);
+		RootScope->ContainedStatement->bEmitShader = true;
+		RootScope->ContainedStatement->EmitShader(Context);
 		if (Context.Errors.Num() > 0)
 		{
 			return false;
diff --git a/Engine/Source/Runtime/Engine/Private/HLSLTree/HLSLTreeCommon.cpp b/Engine/Source/Runtime/Engine/Private/HLSLTree/HLSLTreeCommon.cpp
--- a/Engine/Source/Runtime/Engine/Private/HLSLTree/HLSLTreeCommon.cpp
+++ b/Engine/Source/Runtime/Engine/Private/HLSLTree/HLSLTreeCommon.cpp
@@ -629,25 +629,25 @@ void FExpressionReflectionVector::EmitValueShader(FEmitContext& Context, FShader
 	OutShader.Code.Append(TEXT("Parameters.ReflectionVector"));
 }
 
-void FStatementBreak::PrepareValues(FEmitContext& Context) const
+void FStatementBreak::Prepare(FEmitContext& Context) const
 {
 }
 
-void FStatementBreak::EmitHLSL(FEmitContext& Context) const
+void FStatementBreak::EmitShader(FEmitContext& Context) const
 {
 	ParentScope->EmitStatementf(Context, TEXT("break;"));
 }
 
-void FStatementReturn::PrepareValues(FEmitContext& Context) const
+void FStatementReturn::Prepare(FEmitContext& Context) const
 {
 }
 
-void FStatementReturn::EmitHLSL(FEmitContext& Context) const
+void FStatementReturn::EmitShader(FEmitContext& Context) const
 {
 	ParentScope->EmitStatementf(Context, TEXT("return %s;"), Expression->GetValueShader(Context));
 }
 
-void FStatementIf::PrepareValues(FEmitContext& Context) const
+void FStatementIf::Prepare(FEmitContext& Context) const
 {
 	const FPrepareValueResult ConditionResult = PrepareExpressionValue(Context, ConditionExpression, 1);
 	if (ConditionResult.EvaluationType == EExpressionEvaluationType::Constant)
@@ -671,18 +671,18 @@ void FStatementIf::PrepareValues(FEmitContext& Context) const
 	}
 }
 
-void FStatementIf::EmitHLSL(FEmitContext& Context) const
+void FStatementIf::EmitShader(FEmitContext& Context) const
 {
 	if (ConditionExpression->GetEvaluationType() == EExpressionEvaluationType::Constant)
 	{
 		const bool bCondition = ConditionExpression->GetValueConstant(Context).AsBoolScalar();
 		if (bCondition && IsScopeLive(ThenScope))
 		{
-			ParentScope->EmitNextScope(Context, ThenScope);
+			ParentScope->EmitScope(Context, ThenScope);
 		}
 		else if(!bCondition && IsScopeLive(ElseScope))
 		{
-			ParentScope->EmitNextScope(Context, ElseScope);
+			ParentScope->EmitScope(Context, ElseScope);
 		}
 	}
 	else
@@ -698,21 +698,21 @@ void FStatementIf::EmitHLSL(FEmitContext& Context) const
 		}
 	}
 
-	ParentScope->EmitNextScope(Context, NextScope);
+	ParentScope->EmitScope(Context, NextScope);
 }
 
-void FStatementLoop::PrepareValues(FEmitContext& Context) const
+void FStatementLoop::Prepare(FEmitContext& Context) const
 {
 	MarkScopeLive(LoopScope);
 }
 
-void FStatementLoop::EmitHLSL(FEmitContext& Context) const
+void FStatementLoop::EmitShader(FEmitContext& Context) const
 {
 	if (IsScopeLive(LoopScope))
 	{
 		ParentScope->EmitNestedScopef(Context, LoopScope, TEXT("while (true)"));
 	}
-	ParentScope->EmitNextScope(Context, NextScope);
+	ParentScope->EmitScope(Context, NextScope);
 }
 
 } // namespace HLSLTree
diff --git a/Engine/Source/Runtime/Engine/Private/Materials/MaterialHLSLEmitter.cpp b/Engine/Source/Runtime/Engine/Private/Materials/MaterialHLSLEmitter.cpp
--- a/Engine/Source/Runtime/Engine/Private/Materials/MaterialHLSLEmitter.cpp
+++ b/Engine/Source/Runtime/Engine/Private/Materials/MaterialHLSLEmitter.cpp
@@ -751,7 +751,7 @@ bool MaterialEmitHLSL(const FMaterialCompileTargetParameters& InCompilerTarget,
 	const FStructFieldRef NormalFieldRef = Generator.GetMaterialAttributesType()->FindFieldByName(*FMaterialAttributeDefinitionMap::GetAttributeName(MP_Normal));
 	check(NormalFieldRef);
 
-	if (!PrepareScopeValues(EmitContext, Generator.GetResultStatement()->ParentScope))
+	if (!PrepareScope(EmitContext, Generator.GetResultStatement()->ParentScope))
 	{
 		return false;
 	}
@@ -775,16 +775,16 @@ bool MaterialEmitHLSL(const FMaterialCompileTargetParameters& InCompilerTarget,
 		{
 			return false;
 		}
-		HLSLTree->EmitHLSL(EmitContext, CodePhase0);
+		HLSLTree->EmitShader(EmitContext, CodePhase0);
 	}
 	else
 	{
 		// Execute everything *except* normal in phase1
-		HLSLTree->EmitHLSL(EmitContext, CodePhase1);
+		HLSLTree->EmitShader(EmitContext, CodePhase1);
 
 		// Reset state
 		HLSLTree->ResetNodes();
-		if (!PrepareScopeValues(EmitContext, Generator.GetResultStatement()->ParentScope))
+		if (!PrepareScope(EmitContext, Generator.GetResultStatement()->ParentScope))
 		{
 			return false;
 		}
@@ -798,7 +798,7 @@ bool MaterialEmitHLSL(const FMaterialCompileTargetParameters& InCompilerTarget,
 		}
 
 		// Execute the normal in phase0
-		HLSLTree->EmitHLSL(EmitContext, CodePhase0);
+		HLSLTree->EmitShader(EmitContext, CodePhase0);
 	}
 
 	FStringBuilderMemstack Declarations(Allocator, 64 * 1024);
diff --git a/Engine/Source/Runtime/Engine/Public/HLSLTree/HLSLTree.h b/Engine/Source/Runtime/Engine/Public/HLSLTree/HLSLTree.h
--- a/Engine/Source/Runtime/Engine/Public/HLSLTree/HLSLTree.h
+++ b/Engine/Source/Runtime/Engine/Public/HLSLTree/HLSLTree.h
@@ -259,11 +259,11 @@ public:
 	virtual void Reset() override;
 	virtual ENodeVisitResult Visit(FNodeVisitor& Visitor) override;
 
-	virtual void PrepareValues(FEmitContext& Context) const = 0;
-	virtual void EmitHLSL(FEmitContext& Context) const = 0;
+	virtual void Prepare(FEmitContext& Context) const = 0;
+	virtual void EmitShader(FEmitContext& Context) const = 0;
 
 	FScope* ParentScope = nullptr;
-	bool bEmitHLSL = false;
+	bool bEmitShader = false;
 };
 
 class FRequestedType
@@ -449,7 +449,7 @@ public:
 
 	void UseExpression(FExpression* Expression);
 
-	friend bool PrepareScopeValues(FEmitContext& Context, FScope* InScope);
+	friend bool PrepareScope(FEmitContext& Context, FScope* InScope);
 
 	template<typename FormatType, typename... Types>
 	void EmitDeclarationf(FEmitContext& Context, const FormatType& Format, Types... Args)
@@ -463,7 +463,7 @@ public:
 		InternalEmitCodef(Context, Statements, ENextScopeFormat::None, nullptr, Format, Forward<Types>(Args)...);
 	}
 
-	void EmitNextScope(FEmitContext& Context, FScope* NestedScope)
+	void EmitScope(FEmitContext& Context, FScope* NestedScope)
 	{
 		InternalEmitCode(Context, Statements, ENextScopeFormat::Unscoped, NestedScope, nullptr, 0);
 	}
@@ -495,7 +495,7 @@ private:
 	struct FCodeEntry
 	{
 		FCodeEntry* Next;
-		FScope* NextScope;
+		FScope* Scope;
 		int32 Length;
 		ENextScopeFormat ScopeFormat;
 		TCHAR String[1];
@@ -508,14 +508,14 @@ private:
 		int32 Num = 0;
 	};
 
-	void InternalEmitCode(FEmitContext& Context, FCodeList& List, ENextScopeFormat ScopeFormat, FScope* NextScope, const TCHAR* String, int32 Length);
+	void InternalEmitCode(FEmitContext& Context, FCodeList& List, ENextScopeFormat ScopeFormat, FScope* Scope, const TCHAR* String, int32 Length);
 
 	template<typename FormatType, typename... Types>
-	void InternalEmitCodef(FEmitContext& Context, FCodeList& List, ENextScopeFormat ScopeFormat, FScope* NextScope, const FormatType& Format, Types... Args)
+	void InternalEmitCodef(FEmitContext& Context, FCodeList& List, ENextScopeFormat ScopeFormat, FScope* Scope, const FormatType& Format, Types... Args)
 	{
 		TStringBuilder<2048> String;
 		String.Appendf(Format, Forward<Types>(Args)...);
-		InternalEmitCode(Context, List, ScopeFormat, NextScope, String.ToString(), String.Len());
+		InternalEmitCode(Context, List, ScopeFormat, Scope, String.ToString(), String.Len());
 	}
 
 	FStatement* OwnerStatement = nullptr;
@@ -566,7 +566,7 @@ public:
 
 	void EmitDeclarationsCode(FStringBuilderBase& OutCode) const;
 
-	bool EmitHLSL(FEmitContext& Context, FStringBuilderBase& OutCode) const;
+	bool EmitShader(FEmitContext& Context, FStringBuilderBase& OutCode) const;
 
 	FScope& GetRootScope() const { return *RootScope; }
 
@@ -583,15 +583,6 @@ public:
 	{
 		T* Statement = NewNode<T>(Forward<ArgTypes>(Args)...);
 		RegisterStatement(Scope, Statement);
-		/*if constexpr (T::MarkScopeLiveRecursive)
-		{
-			Scope.MarkLiveRecursive();
-		}
-		else if constexpr (T::MarkScopeLive)
-		{
-			Scope.MarkLive();
-		}*/
-
 		return Statement;
 	}
 
diff --git a/Engine/Source/Runtime/Engine/Public/HLSLTree/HLSLTreeCommon.h b/Engine/Source/Runtime/Engine/Public/HLSLTree/HLSLTreeCommon.h
--- a/Engine/Source/Runtime/Engine/Public/HLSLTree/HLSLTreeCommon.h
+++ b/Engine/Source/Runtime/Engine/Public/HLSLTree/HLSLTreeCommon.h
@@ -322,8 +322,8 @@ public:
 		return Result;
 	}
 
-	virtual void PrepareValues(FEmitContext& Context) const override;
-	virtual void EmitHLSL(FEmitContext& Context) const override;
+	virtual void Prepare(FEmitContext& Context) const override;
+	virtual void EmitShader(FEmitContext& Context) const override;
 };
 
 class FStatementBreak : public FStatement
@@ -331,8 +331,8 @@ class FStatementBreak : public FStatement
 public:
 	//static constexpr bool MarkScopeLive = true;
 
-	virtual void PrepareValues(FEmitContext& Context) const override;
-	virtual void EmitHLSL(FEmitContext& Context) const override;
+	virtual void Prepare(FEmitContext& Context) const override;
+	virtual void EmitShader(FEmitContext& Context) const override;
 };
 
 class FStatementIf : public FStatement
@@ -356,8 +356,8 @@ public:
 		return Result;
 	}
 
-	virtual void PrepareValues(FEmitContext& Context) const override;
-	virtual void EmitHLSL(FEmitContext& Context) const override;
+	virtual void Prepare(FEmitContext& Context) const override;
+	virtual void EmitShader(FEmitContext& Context) const override;
 };
 
 class FStatementLoop : public FStatement
@@ -377,8 +377,8 @@ public:
 		return Result;
 	}
 
-	virtual void PrepareValues(FEmitContext& Context) const override;
-	virtual void EmitHLSL(FEmitContext& Context) const override;
+	virtual void Prepare(FEmitContext& Context) const override;
+	virtual void EmitShader(FEmitContext& Context) const override;
 };
 
 }
