Reduce FString temporaries in AssetManager
  - 2.680M to 1.8M allocs from process start to end of skydiving section in a replayrun

diff --git a/Engine/Source/Runtime/Engine/Private/AssetManager.cpp b/Engine/Source/Runtime/Engine/Private/AssetManager.cpp
--- a/Engine/Source/Runtime/Engine/Private/AssetManager.cpp
+++ b/Engine/Source/Runtime/Engine/Private/AssetManager.cpp
@@ -1621,35 +1621,36 @@ TSharedPtr<FStreamableHandle> UAssetManager::ChangeBundleStateForPrimaryAssets(c
 				}
 			}
 
+			if (PathsToLoad.Num() == 0)
+			{
+				// New state has no assets to load. Set the CurrentState's bundles and clear the handle
+				NameData->CurrentState.BundleNames = NewBundleState;
+				NameData->CurrentState.Handle.Reset();
+				continue;
+			}
+
 			TSharedPtr<FStreamableHandle> NewHandle;
 
-			FString DebugName = PrimaryAssetId.ToString();
+			TStringBuilder<1024> DebugName;
+			PrimaryAssetId.AppendString(DebugName);
 
 			if (NewBundleState.Num() > 0)
 			{
-				DebugName += TEXT(" (");
+				DebugName << TEXT(" (");
 
-				for (int32 i = 0; i < NewBundleState.Num(); i++)
+				for (int32 Index = 0; Index < NewBundleState.Num(); Index++)
 				{
-					if (i != 0)
+					if (Index > 0)
 					{
-						DebugName += TEXT(", ");
+						DebugName << TEXT(", ");
 					}
-					DebugName += NewBundleState[i].ToString();
+					NewBundleState[Index].AppendString(DebugName);
 				}
 
-				DebugName += TEXT(")");
-			}
-
-			if (PathsToLoad.Num() == 0)
-			{
-				// New state has no assets to load. Set the CurrentState's bundles and clear the handle
-				NameData->CurrentState.BundleNames = NewBundleState;
-				NameData->CurrentState.Handle.Reset();
-				continue;
+				DebugName << TEXT(")");
 			}
 
-			NewHandle = LoadAssetList(PathsToLoad.Array(), FStreamableDelegate(), Priority, DebugName);
+			NewHandle = LoadAssetList(PathsToLoad.Array(), FStreamableDelegate(), Priority, FString(DebugName.Len(), DebugName.ToString()));
 
 			if (!NewHandle.IsValid())
 			{
