Remove EditorHash usage when cooking world partition.

diff --git a/Engine/Source/Runtime/Engine/Private/WorldPartition/WorldPartition.cpp b/Engine/Source/Runtime/Engine/Private/WorldPartition/WorldPartition.cpp
--- a/Engine/Source/Runtime/Engine/Private/WorldPartition/WorldPartition.cpp
+++ b/Engine/Source/Runtime/Engine/Private/WorldPartition/WorldPartition.cpp
@@ -382,8 +382,18 @@ void UWorldPartition::Initialize(UWorld* InWorld, const FTransform& InTransform)
 	AWorldPartitionReplay::Initialize(World);
 	
 #if WITH_EDITOR
-	bool bEditorOnly = !World->IsGameWorld();
-	if (bEditorOnly)
+	const bool bIsGame = IsRunningGame();
+	const bool bIsEditor = !World->IsGameWorld();
+	const bool bIsCooking = IsRunningCookCommandlet();
+
+	UE_LOG(LogWorldPartition, Log, TEXT("UWorldPartition::Initialize(IsEditor=%d, IsGame=%d, IsCooking=%d)"), bIsEditor ? 1 : 0, bIsGame ? 1 : 0, bIsCooking ? 1 : 0);
+
+	if (bIsGame || bIsCooking)
+	{
+		// Don't rely on the editor hash for cooking or -game
+		EditorHash = nullptr;
+	}
+	else if (bIsEditor)
 	{
 		CreateOrRepairWorldPartition(World->GetWorldSettings());
 
@@ -398,7 +408,7 @@ void UWorldPartition::Initialize(UWorld* InWorld, const FTransform& InTransform)
 	// Did we travel into a WP map in PIE (null StreamingPolicy means GenerateStreaming wasn't called)
 	const bool bPIEWorldTravel = World->WorldType == EWorldType::PIE && !StreamingPolicy;
 
-	if (bEditorOnly || IsRunningGame() || bPIEWorldTravel)
+	if (bIsEditor || bIsGame || bPIEWorldTravel)
 	{
 		UPackage* LevelPackage = OuterWorld->PersistentLevel->GetOutermost();
 		const FName PackageName = LevelPackage->GetLoadedPath().GetPackageFName();
@@ -431,7 +441,7 @@ void UWorldPartition::Initialize(UWorld* InWorld, const FTransform& InTransform)
 					ActorDesc->TransformInstance(SourceWorldPath, RemappedWorldPath);
 				}
 
-				if (bEditorOnly)
+				if (bIsEditor && !bIsCooking)
 				{
 					HashActorDesc(ActorDesc);
 				}
@@ -450,7 +460,7 @@ void UWorldPartition::Initialize(UWorld* InWorld, const FTransform& InTransform)
 		}
 	}
 
-	if (bEditorOnly)
+	if (bIsEditor && !bIsCooking)
 	{
 		// Load the always loaded cell, don't call LoadCells to avoid creating a transaction
 		UpdateLoadingEditorCell(EditorHash->GetAlwaysLoadedCell(), true, false);
@@ -484,9 +494,9 @@ void UWorldPartition::Initialize(UWorld* InWorld, const FTransform& InTransform)
 	InitState = EWorldPartitionInitState::Initialized;
 
 #if WITH_EDITOR
-	if (!bEditorOnly)
+	if (!bIsEditor)
 	{
-		if (IsRunningGame() || bPIEWorldTravel)
+		if (bIsGame || bPIEWorldTravel)
 		{
 			OnBeginPlay();
 		}
diff --git a/Engine/Source/Runtime/Engine/Private/WorldPartition/WorldPartitionEditorSpatialHash.cpp b/Engine/Source/Runtime/Engine/Private/WorldPartition/WorldPartitionEditorSpatialHash.cpp
--- a/Engine/Source/Runtime/Engine/Private/WorldPartition/WorldPartitionEditorSpatialHash.cpp
+++ b/Engine/Source/Runtime/Engine/Private/WorldPartition/WorldPartitionEditorSpatialHash.cpp
@@ -129,13 +129,6 @@ bool UWorldPartitionEditorSpatialHash::IsActorAlwaysLoaded(const FWorldPartition
 		return true;
 	}
 
-	// When cooking, we only want actors marked as AlwaysLoaded to be loaded in the PersistentLevel.
-	// @todo_ow: https://jira.it.epicgames.com/browse/UETOOL-3601
-	if (IsRunningCookCommandlet())
-	{
-		return false;
-	}
-
 	// If an actor covers more that 4 levels in the octree (which means 32K cells), treat it as always loaded
 	const FBox ActorBounds = GetActorBounds(InActorHandle);
 	const int32 ActorLevel = GetLevelForBox(ActorBounds);
