Remove unecessary Gbuffer output during base pass when strata is not enabled.

diff --git a/Engine/Shaders/Shared/StrataDefinitions.h b/Engine/Shaders/Shared/StrataDefinitions.h
--- a/Engine/Shaders/Shared/StrataDefinitions.h
+++ b/Engine/Shaders/Shared/StrataDefinitions.h
@@ -8,7 +8,7 @@
 #pragma once
 
 // Change this to force recompilation of all strata dependent shaders (use https://www.random.org/cgi-bin/randbyte?nbytes=4&format=h)
-#define STRATA_SHADER_VERSION 0xc9624d0e 
+#define STRATA_SHADER_VERSION 0x04499eb2  
 
 
 
diff --git a/Engine/Source/Runtime/Engine/Private/ShaderCompiler/ShaderGenerationUtil.cpp b/Engine/Source/Runtime/Engine/Private/ShaderCompiler/ShaderGenerationUtil.cpp
--- a/Engine/Source/Runtime/Engine/Private/ShaderCompiler/ShaderGenerationUtil.cpp
+++ b/Engine/Source/Runtime/Engine/Private/ShaderCompiler/ShaderGenerationUtil.cpp
@@ -251,6 +251,12 @@ void FShaderCompileUtilities::ApplyFetchEnvironment(FShaderCompilerDefines& SrcD
 	FETCH_COMPILE_BOOL(PLATFORM_SUPPORTS_DEVELOPMENT_SHADERS);
 }
 
+// Strata::IsEnabled is only accessible in the Renderer module
+static bool IsStrataEnabled()
+{
+	static IConsoleVariable* CVar = IConsoleManager::Get().FindConsoleVariable(TEXT("r.Strata"));
+	return CVar->GetInt() > 0;
+}
 
 // if we change the logic, increment this number to force a DDC key change
 static const int32 GBufferGeneratorVersion = 4;
@@ -1097,6 +1103,17 @@ static FString CreateGBufferDecodeFunctionDirect(const FGBufferInfo& BufferInfo)
 
 	FullStr += TEXT("\tFGBufferData Ret = (FGBufferData)0;\n");
 	
+	// Default initialization in case no gbuffer data are generated when Strata is enabled to 
+	// prevent shader compiler error (division by zero, variable not-initialized) with passes 
+	// not converted to Strata and still using Gbuffer data
+	const bool bStrata = IsStrataEnabled();
+	if (bStrata)
+	{
+		FullStr += TEXT("\tRet.WorldNormal = float3(0,0,1);\n");
+		FullStr += TEXT("\tRet.Depth = 0.f;\n");
+		FullStr += TEXT("\tRet.ShadingModelID = 1;\n"); 
+	}
+
 	FString Swizzles[4];
 	Swizzles[0] = TEXT("x");
 	Swizzles[1] = TEXT("y");
@@ -1686,6 +1703,9 @@ static void DetermineUsedMaterialSlots(
 	bool bHasStaticLighting = Dst.GBUFFER_HAS_PRECSHADOWFACTOR || Dst.WRITES_PRECSHADOWFACTOR_TO_GBUFFER;
 	bool bIsStrataMaterial = Mat.PROJECT_STRATA && Mat.MATERIAL_IS_STRATA;
 
+	// Strata doesn't use gbuffer, and thus doesn't need CustomData
+	const bool bUseCustomData = !bIsStrataMaterial;
+
 	// we have to use if statements, not switch or if/else statements because we can have multiple shader model ids.
 	if (Mat.MATERIAL_SHADINGMODEL_UNLIT)
 	{
@@ -1700,49 +1720,49 @@ static void DetermineUsedMaterialSlots(
 	if (Mat.MATERIAL_SHADINGMODEL_SUBSURFACE)
 	{
 		SetStandardGBufferSlots(Slots, bWriteEmissive, bHasTangent, bHasVelocity, bHasStaticLighting, bIsStrataMaterial);
-		Slots[GBS_CustomData] = true;
+		Slots[GBS_CustomData] = bUseCustomData;
 	}
 
 	if (Mat.MATERIAL_SHADINGMODEL_PREINTEGRATED_SKIN)
 	{
 		SetStandardGBufferSlots(Slots, bWriteEmissive, bHasTangent, bHasVelocity, bHasStaticLighting, bIsStrataMaterial);
-		Slots[GBS_CustomData] = true;
+		Slots[GBS_CustomData] = bUseCustomData;
 	}
 
 	if (Mat.MATERIAL_SHADINGMODEL_SUBSURFACE_PROFILE)
 	{
 		SetStandardGBufferSlots(Slots, bWriteEmissive, bHasTangent, bHasVelocity, bHasStaticLighting, bIsStrataMaterial);
-		Slots[GBS_CustomData] = true;
+		Slots[GBS_CustomData] = bUseCustomData;
 	}
 
 	if (Mat.MATERIAL_SHADINGMODEL_CLEAR_COAT)
 	{
 		SetStandardGBufferSlots(Slots, bWriteEmissive, bHasTangent, bHasVelocity, bHasStaticLighting, bIsStrataMaterial);
-		Slots[GBS_CustomData] = true;
+		Slots[GBS_CustomData] = bUseCustomData;
 	}
 
 	if (Mat.MATERIAL_SHADINGMODEL_TWOSIDED_FOLIAGE)
 	{
 		SetStandardGBufferSlots(Slots, bWriteEmissive, bHasTangent, bHasVelocity, bHasStaticLighting, bIsStrataMaterial);
-		Slots[GBS_CustomData] = true;
+		Slots[GBS_CustomData] = bUseCustomData;
 	}
 
 	if (Mat.MATERIAL_SHADINGMODEL_HAIR)
 	{
 		SetStandardGBufferSlots(Slots, bWriteEmissive, bHasTangent, bHasVelocity, bHasStaticLighting, bIsStrataMaterial);
-		Slots[GBS_CustomData] = true;
+		Slots[GBS_CustomData] = bUseCustomData;
 	}
 
 	if (Mat.MATERIAL_SHADINGMODEL_CLOTH)
 	{
 		SetStandardGBufferSlots(Slots, bWriteEmissive, bHasTangent, bHasVelocity, bHasStaticLighting, bIsStrataMaterial);
-		Slots[GBS_CustomData] = true;
+		Slots[GBS_CustomData] = bUseCustomData;
 	}
 
 	if (Mat.MATERIAL_SHADINGMODEL_EYE)
 	{
 		SetStandardGBufferSlots(Slots, bWriteEmissive, bHasTangent, bHasVelocity, bHasStaticLighting, bIsStrataMaterial);
-		Slots[GBS_CustomData] = true;
+		Slots[GBS_CustomData] = bUseCustomData;
 	}
 
 	if (Mat.MATERIAL_SHADINGMODEL_SINGLELAYERWATER)
diff --git a/Engine/Source/Runtime/RenderCore/Private/GBufferInfo.cpp b/Engine/Source/Runtime/RenderCore/Private/GBufferInfo.cpp
--- a/Engine/Source/Runtime/RenderCore/Private/GBufferInfo.cpp
+++ b/Engine/Source/Runtime/RenderCore/Private/GBufferInfo.cpp
@@ -18,8 +18,12 @@ static const int32 EGBufferFormat_HighPrecisionNormals = 3;
 static const int32 EGBufferFormat_Force16BitsPerChannel = 5;
 
 
-
-
+// Strata::IsEnabled is only accessible in the Renderer module
+static bool IsStrataEnabled()
+{
+	static IConsoleVariable* CVar = IConsoleManager::Get().FindConsoleVariable(TEXT("r.Strata"));
+	return CVar->GetInt() > 0;
+}
 
 static bool IsGBufferPackingEqual(const FGBufferPacking& Lhs, const FGBufferPacking& Rhs)
 {
@@ -233,6 +237,60 @@ FGBufferInfo RENDERCORE_API FetchLegacyGBufferInfo(const FGBufferParams& Params)
 	int32 TargetGBufferF = -1;
 	int32 TargetVelocity = -1;
 
+	// Strata ouputs material data through UAV. Only SceneColor, PrecalcShadow & Velocity data are still emitted through RenderTargets
+	const bool bStrata = IsStrataEnabled();
+	if (bStrata)
+	{
+		TargetGBufferA = -1;
+		TargetGBufferB = -1;
+		TargetGBufferC = -1;
+
+		Info.NumTargets = 1;
+		if (Params.bHasPrecShadowFactor)
+		{
+			TargetGBufferE = Info.NumTargets++;
+
+		}
+		if (Params.bHasVelocity)
+		{
+			TargetVelocity = Info.NumTargets++;
+		}
+
+		// this value isn't correct, becuase it doesn't resepect the scene color format cvar, but it's ignored anyways
+		// so it's ok for now
+		Info.Targets[TargetLighting].Init(GBT_Unorm_11_11_10, TEXT("Lighting"), false, true, true, true);
+		Info.Slots[GBS_SceneColor] = FGBufferItem(GBS_SceneColor, GBC_Raw_Float_11_11_10, GBCH_Both);
+		Info.Slots[GBS_SceneColor].Packing[0] = FGBufferPacking(TargetLighting, 0, 0);
+		Info.Slots[GBS_SceneColor].Packing[1] = FGBufferPacking(TargetLighting, 1, 1);
+		Info.Slots[GBS_SceneColor].Packing[2] = FGBufferPacking(TargetLighting, 2, 2);
+
+		if (Params.bHasVelocity)
+		{
+			Info.Targets[TargetVelocity].Init(Params.bUsesVelocityDepth ? GBT_Float_16_16_16_16 : GBT_Float_16_16, TEXT("Velocity"), false, true, true, false); // Velocity
+			Info.Slots[GBS_Velocity] = FGBufferItem(GBS_Velocity, Params.bUsesVelocityDepth ? GBC_Raw_Float_16_16_16_16 : GBC_Raw_Float_16_16, GBCH_Both);
+			Info.Slots[GBS_Velocity].Packing[0] = FGBufferPacking(TargetVelocity, 0, 0);
+			Info.Slots[GBS_Velocity].Packing[1] = FGBufferPacking(TargetVelocity, 1, 1);
+
+			if (Params.bUsesVelocityDepth)
+			{
+				Info.Slots[GBS_Velocity].Packing[2] = FGBufferPacking(TargetVelocity, 2, 2);
+				Info.Slots[GBS_Velocity].Packing[3] = FGBufferPacking(TargetVelocity, 3, 3);
+			}
+		}
+
+		if (Params.bHasPrecShadowFactor)
+		{
+			Info.Targets[TargetGBufferE].Init(GBT_Unorm_8_8_8_8, TEXT("GBufferE"), false, true, true, false); // Precalc
+			Info.Slots[GBS_PrecomputedShadowFactor] = FGBufferItem(GBS_PrecomputedShadowFactor, GBC_Raw_Unorm_8_8_8_8, GBCH_Both);
+			Info.Slots[GBS_PrecomputedShadowFactor].Packing[0] = FGBufferPacking(TargetGBufferE, 0, 0);
+			Info.Slots[GBS_PrecomputedShadowFactor].Packing[1] = FGBufferPacking(TargetGBufferE, 1, 1);
+			Info.Slots[GBS_PrecomputedShadowFactor].Packing[2] = FGBufferPacking(TargetGBufferE, 2, 2);
+			Info.Slots[GBS_PrecomputedShadowFactor].Packing[3] = FGBufferPacking(TargetGBufferE, 3, 3);
+		}
+
+		return Info;
+	}
+
 	if (Params.bHasVelocity == 0 && Params.bHasTangent == 0)
 	{
 		Info.NumTargets = Params.bHasPrecShadowFactor ? 6 : 5;
