Allow notify instances (and native notify types) to optionally disable firing in animation editors

This allows notify authors that do 'dangerous' operations to limit their effects in editor

diff --git a/Engine/Source/Runtime/Engine/Classes/Animation/AnimNotifies/AnimNotify.h b/Engine/Source/Runtime/Engine/Classes/Animation/AnimNotifies/AnimNotify.h
--- a/Engine/Source/Runtime/Engine/Classes/Animation/AnimNotifies/AnimNotify.h
+++ b/Engine/Source/Runtime/Engine/Classes/Animation/AnimNotifies/AnimNotify.h
@@ -57,12 +57,19 @@ class ENGINE_API UAnimNotify : public UObject
 	/** Color of Notify in editor */
 	UPROPERTY(EditAnywhere, BlueprintReadWrite, Category=AnimNotify)
 	FColor NotifyColor;
+
+	/** Whether this notify instance should fire in animation editors */
+	UPROPERTY(EditAnywhere, BlueprintReadWrite, AdvancedDisplay, Category=AnimNotify)
+	bool bShouldFireInEditor;
 #endif // WITH_EDITORONLY_DATA
 
 #if WITH_EDITOR
 	virtual void OnAnimNotifyCreatedInEditor(FAnimNotifyEvent& ContainingAnimNotifyEvent) {};
 	virtual bool CanBePlaced(UAnimSequenceBase* Animation) const { return true; }
 	virtual void ValidateAssociatedAssets() {}
+
+	/** Override this to prevent firing this notify type in animation editors */
+	virtual bool ShouldFireInEditor() { return bShouldFireInEditor; }
 #endif
 
 	UE_DEPRECATED(5.0, "Please use the other Notify function instead")
diff --git a/Engine/Source/Runtime/Engine/Classes/Animation/AnimNotifies/AnimNotifyState.h b/Engine/Source/Runtime/Engine/Classes/Animation/AnimNotifies/AnimNotifyState.h
--- a/Engine/Source/Runtime/Engine/Classes/Animation/AnimNotifies/AnimNotifyState.h
+++ b/Engine/Source/Runtime/Engine/Classes/Animation/AnimNotifies/AnimNotifyState.h
@@ -39,13 +39,19 @@ class ENGINE_API UAnimNotifyState : public UObject
 	/** Color of Notify in editor */
 	UPROPERTY(EditAnywhere, BlueprintReadWrite, Category=AnimNotify)
 	FColor NotifyColor;
-
+	
+	/** Whether this notify state instance should fire in animation editors */
+	UPROPERTY(EditAnywhere, BlueprintReadWrite, AdvancedDisplay, Category=AnimNotify)
+	bool bShouldFireInEditor;
 #endif // WITH_EDITORONLY_DATA
 
 #if WITH_EDITOR
 	virtual void OnAnimNotifyCreatedInEditor(FAnimNotifyEvent& ContainingAnimNotifyEvent) {};
 	virtual bool CanBePlaced(UAnimSequenceBase* Animation) const { return true; }
 	virtual void ValidateAssociatedAssets() {}
+
+	/** Override this to prevent firing this notify state type in animation editors */
+	virtual bool ShouldFireInEditor() { return bShouldFireInEditor; }
 #endif
 
 	UE_DEPRECATED(5.0, "This function is deprecated. Use the other NotifyBegin instead.")
diff --git a/Engine/Source/Runtime/Engine/Private/Animation/AnimInstance.cpp b/Engine/Source/Runtime/Engine/Private/Animation/AnimInstance.cpp
--- a/Engine/Source/Runtime/Engine/Private/Animation/AnimInstance.cpp
+++ b/Engine/Source/Runtime/Engine/Private/Animation/AnimInstance.cpp
@@ -32,7 +32,9 @@
 #include "UObject/UE5MainStreamObjectVersion.h"
 #include "Animation/AnimSubsystem.h"
 #include "Animation/AnimSubsystem_Tag.h"
-
+#if WITH_EDITOR
+#include "Animation/DebugSkelMeshComponent.h"
+#endif
 /** Anim stats */
 
 DEFINE_STAT(STAT_CalcSkelMeshBounds);
@@ -292,8 +294,14 @@ void UAnimInstance::UninitializeAnimation()
 				const FAnimNotifyEventReference& EventReference = ActiveAnimNotifyEventReference[Index];
 				if (ShouldTriggerAnimNotifyState(AnimNotifyEvent.NotifyStateClass))
 				{
-					TRACE_ANIM_NOTIFY(this, AnimNotifyEvent, End);
-					AnimNotifyEvent.NotifyStateClass->NotifyEnd(SkelMeshComp, Cast<UAnimSequenceBase>(AnimNotifyEvent.NotifyStateClass->GetOuter()), EventReference);
+#if WITH_EDITOR
+					// Prevent firing notifies in animation editors if requested 
+					if(!SkelMeshComp->IsA<UDebugSkelMeshComponent>() || AnimNotifyEvent.NotifyStateClass->ShouldFireInEditor())
+#endif
+					{
+						TRACE_ANIM_NOTIFY(this, AnimNotifyEvent, End);
+						AnimNotifyEvent.NotifyStateClass->NotifyEnd(SkelMeshComp, Cast<UAnimSequenceBase>(AnimNotifyEvent.NotifyStateClass->GetOuter()), EventReference);
+					}
 				}
 			}
 		}
@@ -1432,8 +1440,14 @@ void UAnimInstance::TriggerAnimNotifies(float DeltaSeconds)
 		const FAnimNotifyEventReference& EventReference = ActiveAnimNotifyEventReference[Index];
 		if (AnimNotifyEvent.NotifyStateClass && ShouldTriggerAnimNotifyState(AnimNotifyEvent.NotifyStateClass))
 		{
-			TRACE_ANIM_NOTIFY(this, AnimNotifyEvent, End);
-			AnimNotifyEvent.NotifyStateClass->NotifyEnd(SkelMeshComp, Cast<UAnimSequenceBase>(AnimNotifyEvent.NotifyStateClass->GetOuter()), EventReference);
+#if WITH_EDITOR
+			// Prevent firing notifies in animation editors if requested 
+			if(!SkelMeshComp->IsA<UDebugSkelMeshComponent>() || AnimNotifyEvent.NotifyStateClass->ShouldFireInEditor())
+#endif
+			{
+				TRACE_ANIM_NOTIFY(this, AnimNotifyEvent, End);
+				AnimNotifyEvent.NotifyStateClass->NotifyEnd(SkelMeshComp, Cast<UAnimSequenceBase>(AnimNotifyEvent.NotifyStateClass->GetOuter()), EventReference);
+			}
 		}
 		// The NotifyEnd callback above may have triggered actor destruction and the tear down
 		// of this instance via UninitializeAnimation which empties ActiveAnimNotifyState.
@@ -1452,8 +1466,14 @@ void UAnimInstance::TriggerAnimNotifies(float DeltaSeconds)
 		const FAnimNotifyEventReference * AnimNotifyEventReference = NotifyStateBeginEventReference[Index];
 		if (ShouldTriggerAnimNotifyState(AnimNotifyEvent->NotifyStateClass))
 		{
-			TRACE_ANIM_NOTIFY(this, *AnimNotifyEvent, Begin);
-			AnimNotifyEvent->NotifyStateClass->NotifyBegin(SkelMeshComp, Cast<UAnimSequenceBase>(AnimNotifyEvent->NotifyStateClass->GetOuter()), AnimNotifyEvent->GetDuration(), *AnimNotifyEventReference);
+#if WITH_EDITOR
+			// Prevent firing notifies in animation editors if requested 
+			if(!SkelMeshComp->IsA<UDebugSkelMeshComponent>() || AnimNotifyEvent->NotifyStateClass->ShouldFireInEditor())
+#endif
+			{
+				TRACE_ANIM_NOTIFY(this, *AnimNotifyEvent, Begin);
+				AnimNotifyEvent->NotifyStateClass->NotifyBegin(SkelMeshComp, Cast<UAnimSequenceBase>(AnimNotifyEvent->NotifyStateClass->GetOuter()), AnimNotifyEvent->GetDuration(), *AnimNotifyEventReference);
+			}
 		}
 	}
 
@@ -1467,8 +1487,14 @@ void UAnimInstance::TriggerAnimNotifies(float DeltaSeconds)
 		const FAnimNotifyEventReference& EventReference = ActiveAnimNotifyEventReference[Index];
 		if (ShouldTriggerAnimNotifyState(AnimNotifyEvent.NotifyStateClass))
 		{
-			TRACE_ANIM_NOTIFY(this, AnimNotifyEvent, Tick);
-			AnimNotifyEvent.NotifyStateClass->NotifyTick(SkelMeshComp, Cast<UAnimSequenceBase>(AnimNotifyEvent.NotifyStateClass->GetOuter()), DeltaSeconds, EventReference);
+#if WITH_EDITOR
+			// Prevent firing notifies in animation editors if requested 
+			if(!SkelMeshComp->IsA<UDebugSkelMeshComponent>() || AnimNotifyEvent.NotifyStateClass->ShouldFireInEditor())
+#endif
+			{
+				TRACE_ANIM_NOTIFY(this, AnimNotifyEvent, Tick);
+				AnimNotifyEvent.NotifyStateClass->NotifyTick(SkelMeshComp, Cast<UAnimSequenceBase>(AnimNotifyEvent.NotifyStateClass->GetOuter()), DeltaSeconds, EventReference);
+			}
 		}
 	}
 }
@@ -1492,10 +1518,16 @@ void UAnimInstance::TriggerSingleAnimNotify(FAnimNotifyEventReference& EventRefe
 		}
 
 		if (AnimNotifyEvent->Notify != nullptr)
-		{	
-			// Implemented notify: just call Notify. UAnimNotify will forward this to the event which will do the work.
-			TRACE_ANIM_NOTIFY(this, *AnimNotifyEvent, Event);
-			AnimNotifyEvent->Notify->Notify(GetSkelMeshComponent(), Cast<UAnimSequenceBase>(AnimNotifyEvent->Notify->GetOuter()), EventReference);
+		{
+#if WITH_EDITOR
+			// Prevent firing notifies in animation editors if requested 
+			if(!GetSkelMeshComponent()->IsA<UDebugSkelMeshComponent>() || AnimNotifyEvent->Notify->ShouldFireInEditor())
+#endif
+			{
+				// Implemented notify: just call Notify. UAnimNotify will forward this to the event which will do the work.
+				TRACE_ANIM_NOTIFY(this, *AnimNotifyEvent, Event);
+				AnimNotifyEvent->Notify->Notify(GetSkelMeshComponent(), Cast<UAnimSequenceBase>(AnimNotifyEvent->Notify->GetOuter()), EventReference);
+			}
 		}
 		else if (AnimNotifyEvent->NotifyName != NAME_None)
 		{
@@ -1566,8 +1598,14 @@ void UAnimInstance::EndNotifyStates()
 		const FAnimNotifyEventReference& EventReference = ActiveAnimNotifyEventReference[Index];
 		if (UAnimNotifyState* NotifyState = Event.NotifyStateClass)
 		{
-			TRACE_ANIM_NOTIFY(this, Event, End);
-			NotifyState->NotifyEnd(SkelMeshComp, Cast<UAnimSequenceBase>(NotifyState->GetOuter()), EventReference);
+#if WITH_EDITOR
+			// Prevent firing notifies in animation editors if requested 
+			if(!SkelMeshComp->IsA<UDebugSkelMeshComponent>() || NotifyState->ShouldFireInEditor())
+#endif
+			{
+				TRACE_ANIM_NOTIFY(this, Event, End);
+				NotifyState->NotifyEnd(SkelMeshComp, Cast<UAnimSequenceBase>(NotifyState->GetOuter()), EventReference);
+			}
 		}
 	}
 	ActiveAnimNotifyState.Reset();
@@ -1796,8 +1834,14 @@ void UAnimInstance::TriggerMontageEndedEvent(const FQueuedMontageEndedEvent& Mon
 			{
 				if (ShouldTriggerAnimNotifyState(AnimNotifyEvent.NotifyStateClass))
 				{
-					TRACE_ANIM_NOTIFY(this, AnimNotifyEvent, End);
-					AnimNotifyEvent.NotifyStateClass->NotifyEnd(SkelMeshComp, NotifyMontage, EventReference);
+#if WITH_EDITOR
+					// Prevent firing notifies in animation editors if requested 
+					if(!SkelMeshComp->IsA<UDebugSkelMeshComponent>() || AnimNotifyEvent.NotifyStateClass->ShouldFireInEditor())
+#endif
+					{
+						TRACE_ANIM_NOTIFY(this, AnimNotifyEvent, End);
+						AnimNotifyEvent.NotifyStateClass->NotifyEnd(SkelMeshComp, NotifyMontage, EventReference);
+					}
 				}
 
 				if (ActiveAnimNotifyState.IsValidIndex(Index))
diff --git a/Engine/Source/Runtime/Engine/Private/Animation/AnimNotify.cpp b/Engine/Source/Runtime/Engine/Private/Animation/AnimNotify.cpp
--- a/Engine/Source/Runtime/Engine/Private/Animation/AnimNotify.cpp
+++ b/Engine/Source/Runtime/Engine/Private/Animation/AnimNotify.cpp
@@ -17,6 +17,7 @@ UAnimNotify::UAnimNotify(const FObjectInitializer& ObjectInitializer)
 
 #if WITH_EDITORONLY_DATA
 	NotifyColor = FColor(255, 200, 200, 255);
+	bShouldFireInEditor = true;
 #endif // WITH_EDITORONLY_DATA
 
 	bIsNativeBranchingPoint = false;
diff --git a/Engine/Source/Runtime/Engine/Private/Animation/AnimNotifyState.cpp b/Engine/Source/Runtime/Engine/Private/Animation/AnimNotifyState.cpp
--- a/Engine/Source/Runtime/Engine/Private/Animation/AnimNotifyState.cpp
+++ b/Engine/Source/Runtime/Engine/Private/Animation/AnimNotifyState.cpp
@@ -15,6 +15,7 @@ UAnimNotifyState::UAnimNotifyState(const FObjectInitializer& ObjectInitializer)
 
 #if WITH_EDITORONLY_DATA
 	NotifyColor = FColor(200, 200, 255, 255);
+	bShouldFireInEditor = true;
 #endif // WITH_EDITORONLY_DATA
 
 	bIsNativeBranchingPoint = false;
