Fix for: Single Frame blendspace doesn't sample root motion

diff --git a/Engine/Plugins/Experimental/Animation/AnimationWarping/Source/Runtime/Private/AnimationWarpingRuntimeModule.cpp b/Engine/Plugins/Experimental/Animation/AnimationWarping/Source/Runtime/Private/AnimationWarpingRuntimeModule.cpp
--- a/Engine/Plugins/Experimental/Animation/AnimationWarping/Source/Runtime/Private/AnimationWarpingRuntimeModule.cpp
+++ b/Engine/Plugins/Experimental/Animation/AnimationWarping/Source/Runtime/Private/AnimationWarpingRuntimeModule.cpp
@@ -41,7 +41,7 @@ void FModule::ShutdownModule()
 
 void FModule::SampleRootMotion(const FDeltaTimeRecord& SampleRange, const UAnimSequence& Sequence, bool bLoopingSequence, UE::Anim::FStackAttributeContainer& OutAttributes) const
 {
-	const FTransform RootMotionTransform = Sequence.ExtractRootMotion(SampleRange.Previous, SampleRange.Delta, bLoopingSequence);
+	const FTransform RootMotionTransform = Sequence.ExtractRootMotion(SampleRange.GetPrevious(), SampleRange.Delta, bLoopingSequence);
 	FTransformAnimationAttribute* RootMotionAttribute = OutAttributes.FindOrAdd<FTransformAnimationAttribute>(RootMotionAttributeId);
 	RootMotionAttribute->Value = RootMotionTransform;
 }
diff --git a/Engine/Source/Runtime/AnimGraphRuntime/Private/AnimNodes/AnimNode_BlendSpacePlayer.cpp b/Engine/Source/Runtime/AnimGraphRuntime/Private/AnimNodes/AnimNode_BlendSpacePlayer.cpp
--- a/Engine/Source/Runtime/AnimGraphRuntime/Private/AnimNodes/AnimNode_BlendSpacePlayer.cpp
+++ b/Engine/Source/Runtime/AnimGraphRuntime/Private/AnimNodes/AnimNode_BlendSpacePlayer.cpp
@@ -85,7 +85,9 @@ void FAnimNode_BlendSpacePlayer::UpdateInternal(const FAnimationUpdateContext& C
 		// Create a tick record and push into the closest scope
 		UE::Anim::FAnimSyncGroupScope& SyncScope = Context.GetMessageChecked<UE::Anim::FAnimSyncGroupScope>();
 
-		FAnimTickRecord TickRecord(CurrentBlendSpace, Position, BlendSampleDataCache, BlendFilter, GetLoop(), GetPlayRate(), Context.GetFinalBlendWeight(), /*inout*/ InternalTimeAccumulator, MarkerTickRecord);
+		FAnimTickRecord TickRecord(
+			CurrentBlendSpace, Position, BlendSampleDataCache, BlendFilter, GetLoop(), GetPlayRate(), ShouldTeleportToTime(), 
+			IsEvaluator(), Context.GetFinalBlendWeight(), /*inout*/ InternalTimeAccumulator, MarkerTickRecord);
 		TickRecord.RootMotionWeightModifier = Context.GetRootMotionWeightModifier();
 		TickRecord.DeltaTimeRecord = &DeltaTimeRecord;
 
diff --git a/Engine/Source/Runtime/AnimGraphRuntime/Private/AnimNodes/AnimNode_RandomPlayer.cpp b/Engine/Source/Runtime/AnimGraphRuntime/Private/AnimNodes/AnimNode_RandomPlayer.cpp
--- a/Engine/Source/Runtime/AnimGraphRuntime/Private/AnimNodes/AnimNode_RandomPlayer.cpp
+++ b/Engine/Source/Runtime/AnimGraphRuntime/Private/AnimNodes/AnimNode_RandomPlayer.cpp
@@ -129,7 +129,7 @@ void FAnimNode_RandomPlayer::Update_AnyThread(const FAnimationUpdateContext& Con
 
 	// If we looped around, adjust the previous play time to always be before the current playtime,
 	// since we can assume modulo. This makes the crossing check for the start time a lot simpler.
-	float AdjustedPreviousPlayTime = CurrentData->DeltaTimeRecord.Previous;
+	float AdjustedPreviousPlayTime = CurrentData->DeltaTimeRecord.GetPrevious();
 	if (CurrentData->CurrentPlayTime < AdjustedPreviousPlayTime)
 	{
 		AdjustedPreviousPlayTime -= CurrentSequence->GetPlayLength();
@@ -219,8 +219,8 @@ void FAnimNode_RandomPlayer::Update_AnyThread(const FAnimationUpdateContext& Con
 	}
 
 	// Cache time to detect loops
-	CurrentData->DeltaTimeRecord.Previous = CurrentData->CurrentPlayTime;
-	NextData->DeltaTimeRecord.Previous = NextData->CurrentPlayTime;
+	CurrentData->DeltaTimeRecord.SetPrevious(CurrentData->CurrentPlayTime);
+	NextData->DeltaTimeRecord.SetPrevious(NextData->CurrentPlayTime);
 
 	if (bAdvanceToNextEntry)
 	{
@@ -377,8 +377,7 @@ void FAnimNode_RandomPlayer::InitPlayData(FRandomAnimPlayData& Data, int32 Valid
 
 	Data.PlayStartTime = 0.0f;
 	Data.CurrentPlayTime = 0.0f;
-	Data.DeltaTimeRecord.Previous = 0.0f;
-	Data.DeltaTimeRecord.Delta = 0.0f;
+	Data.DeltaTimeRecord = FDeltaTimeRecord();
 	Data.MarkerTickRecord.Reset();
 }
 
diff --git a/Engine/Source/Runtime/AnimGraphRuntime/Private/AnimNodes/AnimNode_SequenceEvaluator.cpp b/Engine/Source/Runtime/AnimGraphRuntime/Private/AnimNodes/AnimNode_SequenceEvaluator.cpp
--- a/Engine/Source/Runtime/AnimGraphRuntime/Private/AnimNodes/AnimNode_SequenceEvaluator.cpp
+++ b/Engine/Source/Runtime/AnimGraphRuntime/Private/AnimNodes/AnimNode_SequenceEvaluator.cpp
@@ -41,7 +41,7 @@ void FAnimNode_SequenceEvaluatorBase::UpdateAssetPlayer(const FAnimationUpdateCo
 
 		if ((!GetTeleportToExplicitTime() || (GetGroupName() != NAME_None) || (GetGroupMethod() == EAnimSyncMethod::Graph)) && (Context.AnimInstanceProxy->IsSkeletonCompatible(CurrentSequence->GetSkeleton())))
 		{
-			DeltaTimeRecord.Previous = InternalTimeAccumulator;
+			DeltaTimeRecord.SetPrevious(InternalTimeAccumulator);
 
 			if (bReinitialized)
 			{
diff --git a/Engine/Source/Runtime/AnimGraphRuntime/Public/AnimNodes/AnimNode_BlendSpaceEvaluator.h b/Engine/Source/Runtime/AnimGraphRuntime/Public/AnimNodes/AnimNode_BlendSpaceEvaluator.h
--- a/Engine/Source/Runtime/AnimGraphRuntime/Public/AnimNodes/AnimNode_BlendSpaceEvaluator.h
+++ b/Engine/Source/Runtime/AnimGraphRuntime/Public/AnimNodes/AnimNode_BlendSpaceEvaluator.h
@@ -20,6 +20,11 @@ public:
 	UPROPERTY(EditAnywhere, BlueprintReadWrite, Category=Settings, meta=(PinShownByDefault))
 	float NormalizedTime;
 
+	/** If true, teleport to normalized time, does NOT advance time (does not trigger notifies, does not extract Root Motion, etc.)
+	If false, will advance time (will trigger notifies, extract root motion if applicable, etc). */
+	UPROPERTY(EditAnywhere, Category = Settings, meta = (PinHiddenByDefault))
+	bool bTeleportToNormalizedTime = true;
+
 public:	
 	FAnimNode_BlendSpaceEvaluator();
 
@@ -30,6 +35,8 @@ public:
 
 	// FAnimNode_BlendSpacePlayer interface
 	virtual float GetPlayRate() const override;
+	virtual bool ShouldTeleportToTime() const override { return bTeleportToNormalizedTime; }
+	virtual bool IsEvaluator() const override { return true; }
 	// End of FAnimNode_BlendSpacePlayer
 
 };
diff --git a/Engine/Source/Runtime/AnimGraphRuntime/Public/AnimNodes/AnimNode_BlendSpacePlayer.h b/Engine/Source/Runtime/AnimGraphRuntime/Public/AnimNodes/AnimNode_BlendSpacePlayer.h
--- a/Engine/Source/Runtime/AnimGraphRuntime/Public/AnimNodes/AnimNode_BlendSpacePlayer.h
+++ b/Engine/Source/Runtime/AnimGraphRuntime/Public/AnimNodes/AnimNode_BlendSpacePlayer.h
@@ -121,6 +121,12 @@ public:
 	// Get the play rate multiplier. Can be negative, which will cause the animation to play in reverse.
 	virtual float GetPlayRate() const;
 
+	// An evaluator will be setting the play rate to zero and setting the time explicitly. ShouldTeleportToTime indicates whether we should jump to that time, or move to it playing out root motion and events etc.
+	virtual bool ShouldTeleportToTime() const { return false; }
+
+	// Indicates if we are an evaluator - i.e. will be setting the time explicitly rather than letting it play out
+	virtual bool IsEvaluator() const { return false; }
+
 	// Should the animation loop back to the start when it reaches the end?
 	bool GetLoop() const;
 
diff --git a/Engine/Source/Runtime/Engine/Classes/Animation/AnimMontage.h b/Engine/Source/Runtime/Engine/Classes/Animation/AnimMontage.h
--- a/Engine/Source/Runtime/Engine/Classes/Animation/AnimMontage.h
+++ b/Engine/Source/Runtime/Engine/Classes/Animation/AnimMontage.h
@@ -516,7 +516,7 @@ public:
 	float GetPosition() const { return Position; };
 	float GetPlayRate() const { return PlayRate; }
 	float GetDeltaMoved() const { return DeltaTimeRecord.Delta; }
-	float GetPreviousPosition() const { return DeltaTimeRecord.Previous;  }
+	float GetPreviousPosition() const { return DeltaTimeRecord.GetPrevious();  }
 	float GetBlendStartAlpha() const { return BlendStartAlpha; }
 	const FAnimMontageInstance* GetMontageSyncLeader() const { return MontageSyncLeader; } 
 	const UBlendProfile* GetActiveBlendProfile() const { return ActiveBlendProfile; }
diff --git a/Engine/Source/Runtime/Engine/Classes/Animation/AnimationAsset.h b/Engine/Source/Runtime/Engine/Classes/Animation/AnimationAsset.h
--- a/Engine/Source/Runtime/Engine/Classes/Animation/AnimationAsset.h
+++ b/Engine/Source/Runtime/Engine/Classes/Animation/AnimationAsset.h
@@ -84,8 +84,21 @@ struct FMarkerTickRecord
  */
 struct FDeltaTimeRecord
 {
-	float Previous = 0.f;
+public:
+	void Set(float InPrevious, float InDelta)
+	{
+		Previous = InPrevious;
+		Delta = InDelta;
+		bPreviousIsValid = true;
+	}
+	void SetPrevious(float InPrevious) { Previous = InPrevious; bPreviousIsValid = true; }
+	float GetPrevious() const { return Previous; }
+	bool IsPreviousValid() const { return bPreviousIsValid; }
+
 	float Delta = 0.f;
+private:
+	float Previous = 0.f;
+	bool  bPreviousIsValid = false; // Will be set to true when Previous has been set
 };
 
 /** Transform definition */
@@ -346,11 +359,13 @@ struct FAnimTickRecord
 	{
 		struct
 		{
-			float  BlendSpacePositionX;
-			float  BlendSpacePositionY;
-			int32  TriangulationIndex;
 			FBlendFilter* BlendFilter;
 			TArray<FBlendSampleData>* BlendSampleDataCache;
+			int32  TriangulationIndex;
+			float  BlendSpacePositionX;
+			float  BlendSpacePositionY;
+			bool   bTeleportToTime;
+			bool   bIsEvaluator;
 		} BlendSpace;
 
 		struct
@@ -391,7 +406,9 @@ public:
 	ENGINE_API FAnimTickRecord(UAnimSequenceBase* InSequence, bool bInLooping, float InPlayRate, float InFinalBlendWeight, float& InCurrentTime, FMarkerTickRecord& InMarkerTickRecord);
 
 	// Create a tick record for a blendspace
-	ENGINE_API FAnimTickRecord(UBlendSpace* InBlendSpace, const FVector& InBlendInput, TArray<FBlendSampleData>& InBlendSampleDataCache, FBlendFilter& InBlendFilter, bool bInLooping, float InPlayRate, float InFinalBlendWeight, float& InCurrentTime, FMarkerTickRecord& InMarkerTickRecord);
+	ENGINE_API FAnimTickRecord(
+		UBlendSpace* InBlendSpace, const FVector& InBlendInput, TArray<FBlendSampleData>& InBlendSampleDataCache, FBlendFilter& InBlendFilter, bool bInLooping, 
+		float InPlayRate, bool bShouldTeleportToTime, bool bIsEvaluator, float InFinalBlendWeight, float& InCurrentTime, FMarkerTickRecord& InMarkerTickRecord);
 
 	// Create a tick record for a montage
 	UE_DEPRECATED(5.0, "Please use the montage FAnimTickRecord constructor which removes InPreviousPosition and InMoveDelta")
diff --git a/Engine/Source/Runtime/Engine/Private/Animation/AnimCompositeBase.cpp b/Engine/Source/Runtime/Engine/Private/Animation/AnimCompositeBase.cpp
--- a/Engine/Source/Runtime/Engine/Private/Animation/AnimCompositeBase.cpp
+++ b/Engine/Source/Runtime/Engine/Private/Animation/AnimCompositeBase.cpp
@@ -573,7 +573,8 @@ void FAnimTrack::GetAnimationPose(FAnimationPoseData& OutAnimationPoseData, cons
 			FAnimExtractContext SequenceExtractionContext(ExtractionContext);
 			if (const UAnimSequenceBase* const AnimRef = AnimSegment->GetAnimationData(ClampedTime, SequenceExtractionContext.CurrentTime))
 			{
-				SequenceExtractionContext.DeltaTimeRecord.Previous = SequenceExtractionContext.CurrentTime - SequenceExtractionContext.DeltaTimeRecord.Delta;
+				SequenceExtractionContext.DeltaTimeRecord.SetPrevious(
+					SequenceExtractionContext.CurrentTime - SequenceExtractionContext.DeltaTimeRecord.Delta);
 				SequenceExtractionContext.bExtractRootMotion &= AnimRef->HasRootMotion();
 				SequenceExtractionContext.bLooping = AnimSegment->LoopingCount > 1;
 				AnimRef->GetAnimationPose(OutAnimationPoseData, SequenceExtractionContext);
diff --git a/Engine/Source/Runtime/Engine/Private/Animation/AnimInstanceProxy.cpp b/Engine/Source/Runtime/Engine/Private/Animation/AnimInstanceProxy.cpp
--- a/Engine/Source/Runtime/Engine/Private/Animation/AnimInstanceProxy.cpp
+++ b/Engine/Source/Runtime/Engine/Private/Animation/AnimInstanceProxy.cpp
@@ -667,11 +667,16 @@ void FAnimInstanceProxy::MakeSequenceTickRecord(FAnimTickRecord& TickRecord, cla
 	TickRecord.bLooping = bLooping;
 }
 
-void FAnimInstanceProxy::MakeBlendSpaceTickRecord(FAnimTickRecord& TickRecord, class UBlendSpace* BlendSpace, const FVector& BlendInput, TArray<FBlendSampleData>& BlendSampleDataCache, FBlendFilter& BlendFilter, bool bLooping, float PlayRate, float FinalBlendWeight, float& CurrentTime, FMarkerTickRecord& MarkerTickRecord) const
+void FAnimInstanceProxy::MakeBlendSpaceTickRecord(
+	FAnimTickRecord& TickRecord, class UBlendSpace* BlendSpace, const FVector& BlendInput, TArray<FBlendSampleData>& BlendSampleDataCache, FBlendFilter& BlendFilter, 
+	bool bLooping, float PlayRate, float FinalBlendWeight, float& CurrentTime, FMarkerTickRecord& MarkerTickRecord) const
 {
 	TickRecord.SourceAsset = BlendSpace;
 	TickRecord.BlendSpace.BlendSpacePositionX = BlendInput.X;
 	TickRecord.BlendSpace.BlendSpacePositionY = BlendInput.Y;
+	// This way of making a tick record is deprecated, so just set to defaults here rather than changing the API
+	TickRecord.BlendSpace.bTeleportToTime = false; 
+	TickRecord.BlendSpace.bIsEvaluator = false;
 	TickRecord.BlendSpace.BlendSampleDataCache = &BlendSampleDataCache;
 	TickRecord.BlendSpace.BlendFilter = &BlendFilter;
 	TickRecord.TimeAccumulator = &CurrentTime;
diff --git a/Engine/Source/Runtime/Engine/Private/Animation/AnimMontage.cpp b/Engine/Source/Runtime/Engine/Private/Animation/AnimMontage.cpp
--- a/Engine/Source/Runtime/Engine/Private/Animation/AnimMontage.cpp
+++ b/Engine/Source/Runtime/Engine/Private/Animation/AnimMontage.cpp
@@ -1158,7 +1158,7 @@ void UAnimMontage::TickAssetPlayer(FAnimTickRecord& Instance, struct FAnimNotify
 	{
 		check(Instance.DeltaTimeRecord);
 		const float CurrentTime = Instance.Montage.CurrentPosition;
-		const float PreviousTime = Instance.DeltaTimeRecord->Previous;
+		const float PreviousTime = Instance.DeltaTimeRecord->GetPrevious();
 		const float MoveDelta = Instance.DeltaTimeRecord->Delta;
 
 		Context.SetLeaderDelta(MoveDelta);
@@ -2308,8 +2308,7 @@ void FAnimMontageInstance::Advance(float DeltaTime, struct FRootMotionMovementPa
 		{
 			const bool bExtractRootMotion = (OutRootMotionParams != nullptr) && Montage->HasRootMotion();
 			
-			DeltaTimeRecord.Delta = 0.f;
-			DeltaTimeRecord.Previous = Position;
+			DeltaTimeRecord.Set(0.f, Position);
 
 			bDidUseMarkerSyncThisTick = CanUseMarkerSync();
 			if (bDidUseMarkerSyncThisTick)
diff --git a/Engine/Source/Runtime/Engine/Private/Animation/AnimSequenceBase.cpp b/Engine/Source/Runtime/Engine/Private/Animation/AnimSequenceBase.cpp
--- a/Engine/Source/Runtime/Engine/Private/Animation/AnimSequenceBase.cpp
+++ b/Engine/Source/Runtime/Engine/Private/Animation/AnimSequenceBase.cpp
@@ -514,8 +514,7 @@ void UAnimSequenceBase::TickAssetPlayer(FAnimTickRecord& Instance, struct FAnimN
 
 	// Capture the final adjusted delta time and previous frame time as an asset player record
 	check(Instance.DeltaTimeRecord);
-	Instance.DeltaTimeRecord->Previous = PreviousTime;
-	Instance.DeltaTimeRecord->Delta = DeltaTime;
+	Instance.DeltaTimeRecord->Set(PreviousTime, DeltaTime);
 
 	HandleAssetPlayerTickedInternal(Context, PreviousTime, DeltaTime, Instance, NotifyQueue);
 }
diff --git a/Engine/Source/Runtime/Engine/Private/Animation/AnimSingleNodeInstanceProxy.cpp b/Engine/Source/Runtime/Engine/Private/Animation/AnimSingleNodeInstanceProxy.cpp
--- a/Engine/Source/Runtime/Engine/Private/Animation/AnimSingleNodeInstanceProxy.cpp
+++ b/Engine/Source/Runtime/Engine/Private/Animation/AnimSingleNodeInstanceProxy.cpp
@@ -575,7 +575,9 @@ void FAnimNode_SingleNode::Update_AnyThread(const FAnimationUpdateContext& Conte
 
 		if (UBlendSpace* BlendSpace = Cast<UBlendSpace>(Proxy->CurrentAsset))
 		{
-			FAnimTickRecord TickRecord(BlendSpace, Proxy->BlendSpacePosition, Proxy->BlendSampleData, Proxy->BlendFilter, Proxy->bLooping, NewPlayRate, 1.f, /*inout*/ Proxy->CurrentTime, Proxy->MarkerTickRecord);
+			FAnimTickRecord TickRecord(
+				BlendSpace, Proxy->BlendSpacePosition, Proxy->BlendSampleData, Proxy->BlendFilter, Proxy->bLooping, 
+				NewPlayRate, false, false, 1.f, /*inout*/ Proxy->CurrentTime, Proxy->MarkerTickRecord);
 			TickRecord.DeltaTimeRecord = &(Proxy->DeltaTimeRecord);
 			
 			SyncScope.AddTickRecord(TickRecord);
diff --git a/Engine/Source/Runtime/Engine/Private/Animation/BlendSpace.cpp b/Engine/Source/Runtime/Engine/Private/Animation/BlendSpace.cpp
--- a/Engine/Source/Runtime/Engine/Private/Animation/BlendSpace.cpp
+++ b/Engine/Source/Runtime/Engine/Private/Animation/BlendSpace.cpp
@@ -294,11 +294,19 @@ void UBlendSpace::TickAssetPlayer(FAnimTickRecord& Instance, struct FAnimNotifyQ
 	const float DeltaTime = Context.GetDeltaTime();
 
 	check(Instance.DeltaTimeRecord);
-	// The blend space delta time record isn't currently used directly
+
+	// Our current time will normally be the same as the recorded previous time, if we were simply playing animations. 
+	// If it's not, then that would be because the time has been explicitly set due to be run as an evaluator, in 
+	// which case we might want to use that delta in addition to any time delta due to the play rate.
+	float ExtraNormalizedDeltaTime = 
+		Instance.DeltaTimeRecord->IsPreviousValid() ?
+		FMath::Wrap(*(Instance.TimeAccumulator) - Instance.DeltaTimeRecord->GetPrevious(), -0.5f, 0.5f) :
+		0.0f;
+
+	// The blend space delta time record isn't normally used directly.
 	// Instead, the blend sample time record's are used to drive pose evaluation
 	// As a consequence, we currently follow the convention that TimeAccumulator and Previous are normalized
-	Instance.DeltaTimeRecord->Previous = *(Instance.TimeAccumulator);
-	Instance.DeltaTimeRecord->Delta = Instance.PlayRateMultiplier * DeltaTime; // Referenced as MoveDelta within comments and logging
+	Instance.DeltaTimeRecord->Set(*(Instance.TimeAccumulator), Instance.PlayRateMultiplier * DeltaTime); 
 
 	// this happens even if MoveDelta == 0.f. This still should happen if it is being interpolated
 	// since we allow setting position of blendspace, we can't ignore MoveDelta == 0.f
@@ -370,7 +378,8 @@ void UBlendSpace::TickAssetPlayer(FAnimTickRecord& Instance, struct FAnimNotifyQ
 				}
 
 				Instance.DeltaTimeRecord->Delta *= FilterMultiplier;
-				UE_LOG(LogAnimation, Log, TEXT("BlendSpace(%s) - FilteredBlendInput(%s) : FilteredBlendInput(%s), FilterMultiplier(%0.2f)"), *GetName(), *BlendSpacePosition.ToString(), *FilteredBlendInput.ToString(), FilterMultiplier);
+				UE_LOG(LogAnimation, Log, TEXT("BlendSpace(%s) - FilteredBlendInput(%s) : FilteredBlendInput(%s), FilterMultiplier(%0.2f)"), 
+					*GetName(), *BlendSpacePosition.ToString(), *FilteredBlendInput.ToString(), FilterMultiplier);
 			}
 
 			bool bCanDoMarkerSync = (SampleIndexWithMarkers != INDEX_NONE) && (Context.IsSingleAnimationContext() || (Instance.bCanUseMarkerSync && Context.CanUseMarkerPosition()));
@@ -465,7 +474,9 @@ void UBlendSpace::TickAssetPlayer(FAnimTickRecord& Instance, struct FAnimNotifyQ
 					float CurrentTime = NormalizedCurrentTime * NewAnimLength;
 					FAnimationRuntime::AdvanceTime(Instance.bLooping, Instance.DeltaTimeRecord->Delta, /*inout*/ CurrentTime, NewAnimLength);
 					NormalizedCurrentTime = NewAnimLength ? (CurrentTime / NewAnimLength) : 0.0f;
-					UE_LOG(LogAnimMarkerSync, Log, TEXT("Leader (%s) (bCanDoMarkerSync == false)  - PreviousTime (%0.2f), CurrentTime (%0.2f), MoveDelta (%0.2f) "), *GetName(), NormalizedPreviousTime, NormalizedCurrentTime, Instance.DeltaTimeRecord->Delta);
+					UE_LOG(LogAnimMarkerSync, Log, 
+						TEXT("Leader (%s) (bCanDoMarkerSync == false)  - PreviousTime (%0.2f), CurrentTime (%0.2f), MoveDelta (%0.2f) "), 
+						*GetName(), NormalizedPreviousTime, NormalizedCurrentTime, Instance.DeltaTimeRecord->Delta);
 				}
 
 				Context.SetAnimationPositionRatio(NormalizedCurrentTime);
@@ -499,15 +510,25 @@ void UBlendSpace::TickAssetPlayer(FAnimTickRecord& Instance, struct FAnimNotifyQ
 				{
 					NormalizedPreviousTime = Context.GetPreviousAnimationPositionRatio();
 					NormalizedCurrentTime = Context.GetAnimationPositionRatio();
-					UE_LOG(LogAnimMarkerSync, Log, TEXT("Follower (%s) (bCanDoMarkerSync == false) - PreviousTime (%0.2f), CurrentTime (%0.2f), MoveDelta (%0.2f) "), *GetName(), NormalizedPreviousTime, NormalizedCurrentTime, Instance.DeltaTimeRecord->Delta);
+					UE_LOG(LogAnimMarkerSync, Log, 
+						TEXT("Follower (%s) (bCanDoMarkerSync == false) - PreviousTime (%0.2f), CurrentTime (%0.2f), MoveDelta (%0.2f) "), 
+						*GetName(), NormalizedPreviousTime, NormalizedCurrentTime, Instance.DeltaTimeRecord->Delta);
 				}
 			}
 
 			// generate notifies and sets time
 			{
 				FAnimNotifyContext NotifyContext(Instance);
-				const float ClampedNormalizedPreviousTime = FMath::Clamp<float>(NormalizedPreviousTime, 0.f, 1.f);
+				float ClampedNormalizedPreviousTime = FMath::Clamp<float>(NormalizedPreviousTime, 0.f, 1.f);
 				const float ClampedNormalizedCurrentTime = FMath::Clamp<float>(NormalizedCurrentTime, 0.f, 1.f);
+
+				if (Instance.BlendSpace.bIsEvaluator && !Instance.BlendSpace.bTeleportToTime)
+				{
+					// When running under an evaluator the time is being set explicitly and we want to add on the deltas.
+					// Note that this could take ClampedNormalizedPreviousTime out of the 0-1 range
+					ClampedNormalizedPreviousTime -= ExtraNormalizedDeltaTime;
+				}
+
 				const bool bGenerateNotifies = (NormalizedCurrentTime != NormalizedPreviousTime) && NotifyTriggerMode != ENotifyTriggerMode::None;
 
 				// Get the index of the highest weight, assuming that the first is the highest until we find otherwise
@@ -564,8 +585,7 @@ void UBlendSpace::TickAssetPlayer(FAnimTickRecord& Instance, struct FAnimNotifyQ
 							}
 
 							// Capture the final adjusted delta time and previous frame time as an asset player record
-							SampleEntry.DeltaTimeRecord.Previous = PrevSampleDataTime;
-							SampleEntry.DeltaTimeRecord.Delta = DeltaTimePosition;
+							SampleEntry.DeltaTimeRecord.Set(PrevSampleDataTime, DeltaTimePosition);
 
 							UE_LOG(LogAnimation, Verbose, TEXT("%d. Blending animation(%s) with %f weight at time %0.2f"), I + 1, *Sample.Animation->GetName(), SampleEntry.GetClampedWeight(), CurrentSampleDataTime);
 						}
diff --git a/Engine/Source/Runtime/Engine/Private/AnimationAsset.cpp b/Engine/Source/Runtime/Engine/Private/AnimationAsset.cpp
--- a/Engine/Source/Runtime/Engine/Private/AnimationAsset.cpp
+++ b/Engine/Source/Runtime/Engine/Private/AnimationAsset.cpp
@@ -191,13 +191,17 @@ FAnimTickRecord::FAnimTickRecord(UAnimSequenceBase* InSequence, bool bInLooping,
 	bLooping = bInLooping;
 }
 
-FAnimTickRecord::FAnimTickRecord(UBlendSpace* InBlendSpace, const FVector& InBlendInput, TArray<FBlendSampleData>& InBlendSampleDataCache, FBlendFilter& InBlendFilter, bool bInLooping, float InPlayRate, float InFinalBlendWeight, float& InCurrentTime, FMarkerTickRecord& InMarkerTickRecord)
+FAnimTickRecord::FAnimTickRecord(
+	UBlendSpace* InBlendSpace, const FVector& InBlendInput, TArray<FBlendSampleData>& InBlendSampleDataCache, FBlendFilter& InBlendFilter, bool bInLooping, 
+	float InPlayRate, bool bTeleportToTime, bool bIsEvaluator, float InFinalBlendWeight, float& InCurrentTime, FMarkerTickRecord& InMarkerTickRecord)
 {
 	SourceAsset = InBlendSpace;
 	BlendSpace.BlendSpacePositionX = InBlendInput.X;
 	BlendSpace.BlendSpacePositionY = InBlendInput.Y;
 	BlendSpace.BlendSampleDataCache = &InBlendSampleDataCache;
 	BlendSpace.BlendFilter = &InBlendFilter;
+	BlendSpace.bTeleportToTime = bTeleportToTime;
+	BlendSpace.bIsEvaluator = bIsEvaluator;
 	TimeAccumulator = &InCurrentTime;
 	MarkerTickRecord = &InMarkerTickRecord;
 	PlayRateMultiplier = InPlayRate;
