Deprecate project settings for virtual texture compression.
Remaining use case is in RVT which needs to be migrated to use oodle RDO before we can strip out compression code.

diff --git a/Engine/Source/Runtime/Engine/Classes/Engine/RendererSettings.h b/Engine/Source/Runtime/Engine/Classes/Engine/RendererSettings.h
--- a/Engine/Source/Runtime/Engine/Classes/Engine/RendererSettings.h
+++ b/Engine/Source/Runtime/Engine/Classes/Engine/RendererSettings.h
@@ -270,7 +270,7 @@ class ENGINE_API URendererSettings : public UDeveloperSettings
 	UPROPERTY(config, EditAnywhere, Category = Mobile, meta = (
 		EditCondition = "bVirtualTextures",
 		ConsoleVariable = "r.Mobile.VirtualTextures", DisplayName = "Enable virtual texture support on Mobile",
-		ToolTip = "Whether to support Virtual Textures on mobile. Requires general Virtual Texturing option enabled as well. Changing this setting requires restarting the editor.",
+		ToolTip = "Whether to support virtual textures on mobile. Requires general virtual texturing option enabled as well. Changing this setting requires restarting the editor.",
 		ConfigRestartRequired = true))
 	uint32 bMobileVirtualTextures : 1;
 
@@ -351,8 +351,8 @@ class ENGINE_API URendererSettings : public UDeveloperSettings
 
 	UPROPERTY(config, EditAnywhere, Category = VirtualTextures, meta = (
 		EditCondition = "bVirtualTextures",
-		ConsoleVariable = "r.VT.AnisotropicFiltering", DisplayName = "Enable VT anisotropic filtering",
-		ToolTip = "When enabled, VTs will use potentially use anisotropic filtering. This adds a cost to all shaders using VT. Changing this setting requires restarting the editor.",
+		ConsoleVariable = "r.VT.AnisotropicFiltering", DisplayName = "Enable virtual texture anisotropic filtering",
+		ToolTip = "When enabled, virtual textures will use anisotropic filtering. This adds a cost to all shaders using virtual textures. Changing this setting requires restarting the editor.",
 		ConfigRestartRequired = true))
 	uint32 bVirtualTextureAnisotropicFiltering : 1;
 
@@ -376,20 +376,6 @@ class ENGINE_API URendererSettings : public UDeveloperSettings
 		ToolTip = "Lower factor will increase virtual texture feedback resolution which increases CPU/GPU overhead, but may decrease streaming latency, especially if materials use many virtual textures."))
 	uint32 VirtualTextureFeedbackFactor;
 
-	UPROPERTY(config, EditAnywhere, Category = VirtualTextures, meta = (
-		EditCondition = "bVirtualTextures",
-		ConsoleVariable = "r.VT.EnableCompressZlib", DisplayName = "Enable Zlib compression",
-		ToolTip = "Use zlib to compress virtual textures. Changing this setting requires restarting the editor.",
-		ConfigRestartRequired = true))
-	uint32 bVirtualTextureEnableCompressZlib : 1;
-
-	UPROPERTY(config, EditAnywhere, Category = VirtualTextures, meta = (
-		EditCondition = "bVirtualTextures",
-		ConsoleVariable = "r.VT.EnableCompressCrunch", DisplayName = "Enable Crunch compression",
-		ToolTip = "Use Crunch library to compress virtual textures for supported formats, this is a lossy compression format that gives much better ratio than zlib. Changing this setting requires restarting the editor.",
-		ConfigRestartRequired = true))
-	uint32 bVirtualTextureEnableCompressCrunch : 1;
-
 	UPROPERTY(config, EditAnywhere, Category = WorkingColorSpace, meta = (
 		DisplayName = "Working Color Space",
 		ToolTip = "Choose from list of provided working color spaces, or custom to provide user-defined space.",
diff --git a/Engine/Source/Runtime/Engine/Private/VT/VirtualTextureBuildSettings.cpp b/Engine/Source/Runtime/Engine/Private/VT/VirtualTextureBuildSettings.cpp
--- a/Engine/Source/Runtime/Engine/Private/VT/VirtualTextureBuildSettings.cpp
+++ b/Engine/Source/Runtime/Engine/Private/VT/VirtualTextureBuildSettings.cpp
@@ -4,18 +4,6 @@
 
 #include "HAL/IConsoleManager.h"
 
-static TAutoConsoleVariable<int32> CVarVTCompressCrunch(
-	TEXT("r.VT.EnableCompressCrunch"),
-	0,
-	TEXT("Enable Crunch compression for virtual textures, for supported formats")
-);
-
-static TAutoConsoleVariable<int32> CVarVTCompressZlib(
-	TEXT("r.VT.EnableCompressZlib"),
-	0,
-	TEXT("Enables Zlib compression for virtual textures, instead of iostore compression")
-);
-
 static TAutoConsoleVariable<int32> CVarVTTileSize(
 	TEXT("r.VT.TileSize"),
 	128,
@@ -32,6 +20,6 @@ void FVirtualTextureBuildSettings::Init()
 {
 	TileSize = CVarVTTileSize.GetValueOnAnyThread();
 	TileBorderSize = CVarVTTileBorderSize.GetValueOnAnyThread();
-	bEnableCompressCrunch = CVarVTCompressCrunch.GetValueOnAnyThread() != 0;
-	bEnableCompressZlib = CVarVTCompressZlib.GetValueOnAnyThread() != 0;
+	bEnableCompressCrunch = false;
+	bEnableCompressZlib = false;
 }
