Refactor FGPUBaseSkinVertexFactory, TGPUSkinVertexFactory, TGPUSkinMorphVertexFactory and TGPUSkinAPEXClothVertexFactory to utilise better polymorphism.

*) Remove templates from functions that create and update skinned/morph/cloth vertex factories.
*) Move FDataType Data out from VFs and change to FGPUSkinDataType/ FGPUSkinMorphDataType/FGPUSkinAPEXClothDataType.
*) Refactor to use virtual function interfaces to access FGPUSkinDataType and its derived structs.
*) Make FGPUSkinDataType a unique pointer of FGPUBaseSkinVertexFactory.
*) Remove the data objects from morph & cloth VFs as they duplicate the same data as the data object in the base skinned VF.
*) Store alias pointes of  FGPUSkinMorphDataType and  FGPUSkinAPEXClothDataType to the base skinned VF FGPUSkinDataType for easy access.

diff --git a/Engine/Source/Runtime/Engine/Private/GPUSkinVertexFactory.cpp b/Engine/Source/Runtime/Engine/Private/GPUSkinVertexFactory.cpp
--- a/Engine/Source/Runtime/Engine/Private/GPUSkinVertexFactory.cpp
+++ b/Engine/Source/Runtime/Engine/Private/GPUSkinVertexFactory.cpp
@@ -440,6 +440,44 @@ bool FGPUBaseSkinVertexFactory::GetUnlimitedBoneInfluences()
 	return (GCVarUnlimitedBoneInfluences!=0);
 }
 
+void FGPUBaseSkinVertexFactory::SetData(const FGPUSkinDataType* InData)
+{
+	check(InData);
+
+	if (!Data)
+	{
+		Data = MakeUnique<FGPUSkinDataType>();
+	}
+
+	*Data = *InData;
+	UpdateRHI();
+}
+
+void FGPUBaseSkinVertexFactory::CopyDataTypeForPassthroughFactory(FGPUSkinPassthroughVertexFactory* PassthroughVertexFactory)
+{
+	FGPUSkinPassthroughVertexFactory::FDataType DestDataType;
+	check(Data.IsValid());
+
+	DestDataType.PositionComponent = Data->PositionComponent;
+	DestDataType.TangentBasisComponents[0] = Data->TangentBasisComponents[0];
+	DestDataType.TangentBasisComponents[1] = Data->TangentBasisComponents[1];
+	DestDataType.TextureCoordinates = Data->TextureCoordinates;
+	DestDataType.ColorComponent = Data->ColorComponent;
+	DestDataType.PreSkinPositionComponent = Data->PositionComponent;
+	DestDataType.PositionComponentSRV = Data->PositionComponentSRV;
+	DestDataType.PreSkinPositionComponentSRV = Data->PositionComponentSRV;
+	DestDataType.TangentsSRV = Data->TangentsSRV;
+	DestDataType.ColorComponentsSRV = Data->ColorComponentsSRV;
+	DestDataType.ColorIndexMask = Data->ColorIndexMask;
+	DestDataType.TextureCoordinatesSRV = Data->TextureCoordinatesSRV;
+	DestDataType.LightMapCoordinateIndex = Data->LightMapCoordinateIndex;
+	DestDataType.NumTexCoords = Data->NumTexCoords;
+	DestDataType.LODLightmapDataIndex = Data->LODLightmapDataIndex;
+
+	PassthroughVertexFactory->SetData(DestDataType);
+}
+
+
 /*-----------------------------------------------------------------------------
 TGPUSkinVertexFactory
 -----------------------------------------------------------------------------*/
@@ -475,74 +513,54 @@ void TGPUSkinVertexFactory<BoneInfluenceType>::ModifyCompilationEnvironment(cons
 	OutEnvironment.SetDefine(TEXT("VF_SUPPORTS_PRIMITIVE_SCENE_DATA"), bSupportsPrimitiveIdStream && bUseGPUScene);
 }
 
-template<GPUSkinBoneInfluenceType BoneInfluenceType>
-void TGPUSkinVertexFactory<BoneInfluenceType>::CopyDataTypeForPassthroughFactory(FGPUSkinPassthroughVertexFactory* PassthroughVertexFactory)
-{
-	FGPUSkinPassthroughVertexFactory::FDataType DestDataType;
-	DestDataType.PositionComponent = Data.PositionComponent;
-	DestDataType.TangentBasisComponents[0] = Data.TangentBasisComponents[0];
-	DestDataType.TangentBasisComponents[1] = Data.TangentBasisComponents[1];
-	DestDataType.TextureCoordinates = Data.TextureCoordinates;
-	DestDataType.ColorComponent = Data.ColorComponent;
-	DestDataType.PreSkinPositionComponent = Data.PositionComponent;
-	DestDataType.PositionComponentSRV = Data.PositionComponentSRV;
-	DestDataType.PreSkinPositionComponentSRV = Data.PositionComponentSRV;
-	DestDataType.TangentsSRV = Data.TangentsSRV;
-	DestDataType.ColorComponentsSRV = Data.ColorComponentsSRV;
-	DestDataType.ColorIndexMask = Data.ColorIndexMask;
-	DestDataType.TextureCoordinatesSRV = Data.TextureCoordinatesSRV;
-	DestDataType.LightMapCoordinateIndex = Data.LightMapCoordinateIndex;
-	DestDataType.NumTexCoords = Data.NumTexCoords;
-	DestDataType.LODLightmapDataIndex = Data.LODLightmapDataIndex;
-	PassthroughVertexFactory->SetData(DestDataType);
-}
-
 /**
 * Add the vertex declaration elements for the streams.
 * @param InData - Type with stream components.
 * @param OutElements - Vertex declaration list to modify.
 */
 template <GPUSkinBoneInfluenceType BoneInfluenceType>
-void TGPUSkinVertexFactory<BoneInfluenceType>::AddVertexElements(FDataType& InData, FVertexDeclarationElementList& OutElements)
+void TGPUSkinVertexFactory<BoneInfluenceType>::AddVertexElements(FVertexDeclarationElementList& OutElements)
 {
+	check(Data.IsValid());
+
 	// Position
-	OutElements.Add(AccessStreamComponent(InData.PositionComponent, 0));
+	OutElements.Add(AccessStreamComponent(Data->PositionComponent, 0));
 
 	// Tangent basis vector
-	OutElements.Add(AccessStreamComponent(InData.TangentBasisComponents[0], 1));
-	OutElements.Add(AccessStreamComponent(InData.TangentBasisComponents[1], 2));
+	OutElements.Add(AccessStreamComponent(Data->TangentBasisComponents[0], 1));
+	OutElements.Add(AccessStreamComponent(Data->TangentBasisComponents[1], 2));
 
 	// Texture coordinates
-	if (InData.TextureCoordinates.Num())
+	if (Data->TextureCoordinates.Num())
 	{
 		const uint8 BaseTexCoordAttribute = 5;
-		for (int32 CoordinateIndex = 0; CoordinateIndex < InData.TextureCoordinates.Num(); ++CoordinateIndex)
+		for (int32 CoordinateIndex = 0; CoordinateIndex < Data->TextureCoordinates.Num(); ++CoordinateIndex)
 		{
 			OutElements.Add(AccessStreamComponent(
-				InData.TextureCoordinates[CoordinateIndex],
+				Data->TextureCoordinates[CoordinateIndex],
 				BaseTexCoordAttribute + CoordinateIndex
 			));
 		}
 
-		for (int32 CoordinateIndex = InData.TextureCoordinates.Num(); CoordinateIndex < MAX_TEXCOORDS; ++CoordinateIndex)
+		for (int32 CoordinateIndex = Data->TextureCoordinates.Num(); CoordinateIndex < MAX_TEXCOORDS; ++CoordinateIndex)
 		{
 			OutElements.Add(AccessStreamComponent(
-				InData.TextureCoordinates[InData.TextureCoordinates.Num() - 1],
+				Data->TextureCoordinates[Data->TextureCoordinates.Num() - 1],
 				BaseTexCoordAttribute + CoordinateIndex
 			));
 		}
 	}
 
-	if (Data.ColorComponentsSRV == nullptr)
+	if (Data->ColorComponentsSRV == nullptr)
 	{
-		Data.ColorComponentsSRV = GNullColorVertexBuffer.VertexBufferSRV;
-		Data.ColorIndexMask = 0;
+		Data->ColorComponentsSRV = GNullColorVertexBuffer.VertexBufferSRV;
+		Data->ColorIndexMask = 0;
 	}
 
 	// Vertex color - account for the possibility that the mesh has no vertex colors
-	if (InData.ColorComponent.VertexBuffer)
+	if (Data->ColorComponent.VertexBuffer)
 	{
-		OutElements.Add(AccessStreamComponent(InData.ColorComponent, 13));
+		OutElements.Add(AccessStreamComponent(Data->ColorComponent, 13));
 	}
 	else
 	{
@@ -555,26 +573,26 @@ void TGPUSkinVertexFactory<BoneInfluenceType>::AddVertexElements(FDataType& InDa
 	if (BoneInfluenceType == UnlimitedBoneInfluence)
 	{
 		// Blend offset count
-		OutElements.Add(AccessStreamComponent(InData.BlendOffsetCount, 3));
+		OutElements.Add(AccessStreamComponent(Data->BlendOffsetCount, 3));
 	}
 	else
 	{
 		// Bone indices
-		OutElements.Add(AccessStreamComponent(InData.BoneIndices, 3));
+		OutElements.Add(AccessStreamComponent(Data->BoneIndices, 3));
 
 		// Bone weights
-		OutElements.Add(AccessStreamComponent(InData.BoneWeights, 4));
+		OutElements.Add(AccessStreamComponent(Data->BoneWeights, 4));
 
 		// Extra bone indices & weights
 		if (GetNumBoneInfluences() > MAX_INFLUENCES_PER_STREAM)
 		{
-			OutElements.Add(AccessStreamComponent(InData.ExtraBoneIndices, 14));
-			OutElements.Add(AccessStreamComponent(InData.ExtraBoneWeights, 15));
+			OutElements.Add(AccessStreamComponent(Data->ExtraBoneIndices, 14));
+			OutElements.Add(AccessStreamComponent(Data->ExtraBoneWeights, 15));
 		}
 		else
 		{
-			OutElements.Add(AccessStreamComponent(InData.BoneIndices, 14));
-			OutElements.Add(AccessStreamComponent(InData.BoneWeights, 15));
+			OutElements.Add(AccessStreamComponent(Data->BoneIndices, 14));
+			OutElements.Add(AccessStreamComponent(Data->BoneWeights, 15));
 		}
 	}
 
@@ -591,7 +609,7 @@ void TGPUSkinVertexFactory<BoneInfluenceType>::InitRHI()
 {
 	// list of declaration items
 	FVertexDeclarationElementList Elements;
-	AddVertexElements(Data,Elements);	
+	AddVertexElements(Elements);	
 
 	// create the actual device decls
 	InitDeclaration(Elements);
@@ -957,22 +975,39 @@ bool TGPUSkinMorphVertexFactory<BoneInfluenceType>::ShouldCompilePermutation(con
 		&& Super::ShouldCompilePermutation(Parameters);
 }
 
+template <GPUSkinBoneInfluenceType BoneInfluenceType>
+void TGPUSkinMorphVertexFactory<BoneInfluenceType>::SetData(const FGPUSkinDataType* InData)
+{
+	const FGPUSkinMorphDataType* InMorphData = (const FGPUSkinMorphDataType*)(InData);
+	check(InMorphData);
+
+	if (!this->Data)
+	{
+		MorphDataPtr = new FGPUSkinMorphDataType();
+		this->Data = TUniquePtr<FGPUSkinDataType>(MorphDataPtr);
+	}
+
+	*MorphDataPtr = *InMorphData;
+	FGPUBaseSkinVertexFactory::UpdateRHI();
+}
+
 /**
 * Add the decl elements for the streams
 * @param InData - type with stream components
 * @param OutElements - vertex decl list to modify
 */
 template <GPUSkinBoneInfluenceType BoneInfluenceType>
-void TGPUSkinMorphVertexFactory<BoneInfluenceType>::AddVertexElements(FDataType& InData, FVertexDeclarationElementList& OutElements)
+void TGPUSkinMorphVertexFactory<BoneInfluenceType>::AddVertexElements(FVertexDeclarationElementList& OutElements)
 {
 	// add the base gpu skin elements
-	TGPUSkinVertexFactory<BoneInfluenceType>::AddVertexElements(InData,OutElements);
+	Super::AddVertexElements(OutElements);
 	// add the morph delta elements
-	FVertexElement DeltaPositionElement = FVertexFactory::AccessStreamComponent(InData.DeltaPositionComponent, 9);
+	check(MorphDataPtr);
+	FVertexElement DeltaPositionElement = FVertexFactory::AccessStreamComponent(MorphDataPtr->DeltaPositionComponent, 9);
 	// Cache delta stream index (position & tangentZ share the same stream)
 	MorphDeltaStreamIndex = DeltaPositionElement.StreamIndex;
 	OutElements.Add(DeltaPositionElement);
-	OutElements.Add(FVertexFactory::AccessStreamComponent(InData.DeltaTangentZComponent, 10));
+	OutElements.Add(FVertexFactory::AccessStreamComponent(MorphDataPtr->DeltaTangentZComponent, 10));
 }
 
 /**
@@ -984,7 +1019,7 @@ void TGPUSkinMorphVertexFactory<BoneInfluenceType>::InitRHI()
 {
 	// list of declaration items
 	FVertexDeclarationElementList Elements;	
-	AddVertexElements(MorphData,Elements);
+	AddVertexElements(Elements);
 
 	// create the actual device decls
 	FVertexFactory::InitDeclaration(Elements);
@@ -1005,7 +1040,9 @@ void TGPUSkinMorphVertexFactory<BoneInfluenceType>::UpdateMorphVertexStream(cons
 template <GPUSkinBoneInfluenceType BoneInfluenceType>
 const FMorphVertexBuffer* TGPUSkinMorphVertexFactory<BoneInfluenceType>::GetMorphVertexBuffer(bool bPrevious, uint32 FrameNumber) const
 {
-	return MorphData.MorphVertexBufferPool ? &MorphData.MorphVertexBufferPool->GetMorphVertexBufferForReading(bPrevious, FrameNumber) : nullptr;
+	FGPUSkinMorphDataType* MorphData = (FGPUSkinMorphDataType*)(this->Data.Get());
+	check(MorphData);
+	return MorphData->MorphVertexBufferPool ? &MorphData->MorphVertexBufferPool->GetMorphVertexBufferForReading(bPrevious, FrameNumber) : nullptr;
 }
 
 IMPLEMENT_GPUSKINNING_VERTEX_FACTORY_PARAMETER_TYPE(TGPUSkinMorphVertexFactory, SF_Vertex, FGPUSkinMorphVertexFactoryShaderParameters);
@@ -1210,6 +1247,22 @@ bool TGPUSkinAPEXClothVertexFactory<BoneInfluenceType>::ShouldCompilePermutation
 		&& Super::ShouldCompilePermutation(Parameters);
 }
 
+template <GPUSkinBoneInfluenceType BoneInfluenceType>
+void TGPUSkinAPEXClothVertexFactory<BoneInfluenceType>::SetData(const FGPUSkinDataType* InData)
+{
+	const FGPUSkinAPEXClothDataType* InClothData = (const FGPUSkinAPEXClothDataType*)(InData);
+	check(InClothData);
+
+	if (!this->Data)
+	{
+		ClothDataPtr = new FGPUSkinAPEXClothDataType();
+		this->Data = TUniquePtr<FGPUSkinDataType>(ClothDataPtr);
+	}
+
+	*ClothDataPtr = *InClothData;
+	FGPUBaseSkinVertexFactory::UpdateRHI();
+}
+
 /**
 * Creates declarations for each of the vertex stream components and
 * initializes the device resource
@@ -1219,7 +1272,7 @@ void TGPUSkinAPEXClothVertexFactory<BoneInfluenceType>::InitRHI()
 {
 	// list of declaration items
 	FVertexDeclarationElementList Elements;	
-	TGPUSkinVertexFactory<BoneInfluenceType>::AddVertexElements(MeshMappingData, Elements);
+	Super::AddVertexElements(Elements);
 
 	// create the actual device decls
 	FVertexFactory::InitDeclaration(Elements);
diff --git a/Engine/Source/Runtime/Engine/Private/SkeletalRenderGPUSkin.cpp b/Engine/Source/Runtime/Engine/Private/SkeletalRenderGPUSkin.cpp
--- a/Engine/Source/Runtime/Engine/Private/SkeletalRenderGPUSkin.cpp
+++ b/Engine/Source/Runtime/Engine/Private/SkeletalRenderGPUSkin.cpp
@@ -1313,12 +1313,8 @@ void FSkeletalMeshObjectGPUSkin::RefreshClothingTransforms(const FMatrix& InNewL
  * @param VertexBuffers - vertex buffers which contains the data and also stride info
  * @param bUseInstancedVertexWeights - use instanced influence weights instead of default weights
  */
-template<class VertexFactoryType>
-void InitGPUSkinVertexFactoryComponents(typename VertexFactoryType::FDataType* VertexFactoryData, 
-										const FSkeletalMeshObjectGPUSkin::FVertexFactoryBuffers& VertexBuffers, VertexFactoryType* VertexFactory)
+void InitGPUSkinVertexFactoryComponents(FGPUSkinDataType* VertexFactoryData, const FSkeletalMeshObjectGPUSkin::FVertexFactoryBuffers& VertexBuffers, FGPUBaseSkinVertexFactory* VertexFactory)
 {
-	typedef TGPUSkinVertexBase BaseVertexType;
-
 	//position
 	VertexBuffers.StaticVertexBuffers->PositionVertexBuffer.BindPositionVertexBuffer(VertexFactory, *VertexFactoryData);
 
@@ -1380,8 +1376,7 @@ void InitGPUSkinVertexFactoryComponents(typename VertexFactoryType::FDataType* V
  * @param VertexBuffers - vertex buffers which contains the data and also stride info
  * @param bUseInstancedVertexWeights - use instanced influence weights instead of default weights
  */
-template<class VertexFactoryType>
-void InitMorphVertexFactoryComponents(typename VertexFactoryType::FDataType* VertexFactoryData, 
+void InitMorphVertexFactoryComponents(FGPUSkinMorphDataType* VertexFactoryData,
 										const FSkeletalMeshObjectGPUSkin::FVertexFactoryBuffers& VertexBuffers)
 {
 	VertexFactoryData->MorphVertexBufferPool = VertexBuffers.MorphVertexBufferPool;
@@ -1402,8 +1397,7 @@ void InitMorphVertexFactoryComponents(typename VertexFactoryType::FDataType* Ver
  * @param VertexBuffers - vertex buffers which contains the data and also stride info
  * @param bUseInstancedVertexWeights - use instanced influence weights instead of default weights
  */
-template<class VertexFactoryType>
-void InitAPEXClothVertexFactoryComponents(typename VertexFactoryType::FDataType* VertexFactoryData, 
+void InitAPEXClothVertexFactoryComponents(FGPUSkinAPEXClothDataType* VertexFactoryData,
 										const FSkeletalMeshObjectGPUSkin::FVertexFactoryBuffers& VertexBuffers)
 {
 	VertexFactoryData->ClothBuffer = VertexBuffers.APEXClothVertexBuffer->GetSRV();
@@ -1413,43 +1407,50 @@ void InitAPEXClothVertexFactoryComponents(typename VertexFactoryType::FDataType*
 /** 
  * Handles transferring data between game/render threads when initializing vertex factory components 
  */
-template <class VertexFactoryType>
-class TDynamicUpdateVertexFactoryData
+class FDynamicUpdateVertexFactoryData
 {
 public:
-	TDynamicUpdateVertexFactoryData(
-		VertexFactoryType* InVertexFactory,
+	FDynamicUpdateVertexFactoryData(
+		FGPUBaseSkinVertexFactory* InVertexFactory,
 		const FSkeletalMeshObjectGPUSkin::FVertexFactoryBuffers& InVertexBuffers)
 		:	VertexFactory(InVertexFactory)
 		,	VertexBuffers(InVertexBuffers)
 	{}
-	VertexFactoryType* VertexFactory;
+
+	FGPUBaseSkinVertexFactory* VertexFactory;
 	const FSkeletalMeshObjectGPUSkin::FVertexFactoryBuffers VertexBuffers;
-	
 };
 
 /**
  * Creates a vertex factory entry for the given type and initialize it on the render thread
  */
-template <class VertexFactoryTypeBase, class VertexFactoryType>
-static VertexFactoryType* CreateVertexFactory(TArray<TUniquePtr<VertexFactoryTypeBase>>& VertexFactories,
+static FGPUBaseSkinVertexFactory* CreateVertexFactory(TArray<TUniquePtr<FGPUBaseSkinVertexFactory>>& VertexFactories,
 						 const FSkeletalMeshObjectGPUSkin::FVertexFactoryBuffers& InVertexBuffers,
 						 ERHIFeatureLevel::Type FeatureLevel
 						 )
 {
-	VertexFactoryType* VertexFactory = new VertexFactoryType(FeatureLevel, InVertexBuffers.NumVertices);
-	VertexFactories.Add(TUniquePtr<VertexFactoryTypeBase>(VertexFactory));
+	FGPUBaseSkinVertexFactory* VertexFactory = nullptr;
+	GPUSkinBoneInfluenceType BoneInfluenceType = InVertexBuffers.SkinWeightVertexBuffer->GetBoneInfluenceType();
+	if (BoneInfluenceType == GPUSkinBoneInfluenceType::DefaultBoneInfluence)
+	{
+		VertexFactory = new TGPUSkinVertexFactory<GPUSkinBoneInfluenceType::DefaultBoneInfluence>(FeatureLevel, InVertexBuffers.NumVertices);
+	}
+	else
+	{
+		VertexFactory = new TGPUSkinVertexFactory<GPUSkinBoneInfluenceType::UnlimitedBoneInfluence>(FeatureLevel, InVertexBuffers.NumVertices);
+	}
+	VertexFactories.Add(TUniquePtr<FGPUBaseSkinVertexFactory>(VertexFactory));
 
 	// Setup the update data for enqueue
-	TDynamicUpdateVertexFactoryData<VertexFactoryType> VertexUpdateData(VertexFactory,InVertexBuffers);
+	FDynamicUpdateVertexFactoryData VertexUpdateData(VertexFactory, InVertexBuffers);
 
 	// update vertex factory components and sync it
 	ENQUEUE_RENDER_COMMAND(InitGPUSkinVertexFactory)(
 		[VertexUpdateData](FRHICommandList& CmdList)
 		{
-			typename VertexFactoryType::FDataType Data;
-			InitGPUSkinVertexFactoryComponents<VertexFactoryType>(&Data, VertexUpdateData.VertexBuffers, VertexUpdateData.VertexFactory);
-			VertexUpdateData.VertexFactory->SetData(Data);
+			FGPUSkinDataType Data;
+			InitGPUSkinVertexFactoryComponents(&Data, VertexUpdateData.VertexBuffers, VertexUpdateData.VertexFactory);
+			VertexUpdateData.VertexFactory->SetData(&Data);
 			VertexUpdateData.VertexFactory->InitResource();
 		}
 	);
@@ -1468,34 +1469,32 @@ void FGPUSkinPassthroughVertexFactory::SetData(const FDataType& InData)
 	}
 }
 
-template <class VertexFactoryTypeBase, class VertexFactoryType>
-void UpdateVertexFactory(TArray<TUniquePtr<VertexFactoryTypeBase>>& VertexFactories,
+void UpdateVertexFactory(TArray<TUniquePtr<FGPUBaseSkinVertexFactory>>& VertexFactories,
 	const FSkeletalMeshObjectGPUSkin::FVertexFactoryBuffers& InVertexBuffers)
 {
-	for (TUniquePtr<VertexFactoryTypeBase>& FactoryPtr : VertexFactories)
+	for (TUniquePtr<FGPUBaseSkinVertexFactory>& FactoryPtr : VertexFactories)
 	{
-		VertexFactoryType* VertexFactory = (VertexFactoryType*)FactoryPtr.Get();
+		FGPUBaseSkinVertexFactory* VertexFactory = FactoryPtr.Get();
 
 		if (VertexFactory != nullptr)
 		{
 			// Setup the update data for enqueue
-			TDynamicUpdateVertexFactoryData<VertexFactoryType> VertexUpdateData(VertexFactory, InVertexBuffers);
+			FDynamicUpdateVertexFactoryData VertexUpdateData(VertexFactory, InVertexBuffers);
 
 			// update vertex factory components and sync it
 			ENQUEUE_RENDER_COMMAND(UpdateGPUSkinVertexFactory)(
 				[VertexUpdateData](FRHICommandList& CmdList)
 			{
-				typename VertexFactoryType::FDataType Data;
-				InitGPUSkinVertexFactoryComponents<VertexFactoryType>(&Data, VertexUpdateData.VertexBuffers, VertexUpdateData.VertexFactory);
-				VertexUpdateData.VertexFactory->SetData(Data);
+				FGPUSkinDataType Data;
+				InitGPUSkinVertexFactoryComponents(&Data, VertexUpdateData.VertexBuffers, VertexUpdateData.VertexFactory);
+				VertexUpdateData.VertexFactory->SetData(&Data);
 			});
 		}
 	}
 }
 
-template<typename VertexFactoryType>
 static void CreatePassthroughVertexFactory(ERHIFeatureLevel::Type InFeatureLevel, TArray<TUniquePtr<FGPUSkinPassthroughVertexFactory>>& PassthroughVertexFactories,
-	VertexFactoryType* SourceVertexFactory)
+	FGPUBaseSkinVertexFactory* SourceVertexFactory)
 {
 	FGPUSkinPassthroughVertexFactory* NewPassthroughVertexFactory = new FGPUSkinPassthroughVertexFactory(InFeatureLevel);
 	PassthroughVertexFactories.Add(TUniquePtr<FGPUSkinPassthroughVertexFactory>(NewPassthroughVertexFactory));
@@ -1513,55 +1512,59 @@ static void CreatePassthroughVertexFactory(ERHIFeatureLevel::Type InFeatureLevel
 /**
  * Creates a vertex factory entry for the given type and initialize it on the render thread
  */
-template <class VertexFactoryTypeBase, class VertexFactoryType>
-static VertexFactoryType* CreateVertexFactoryMorph(TArray<TUniquePtr<VertexFactoryTypeBase>>& VertexFactories,
+static void CreateVertexFactoryMorph(TArray<TUniquePtr<FGPUBaseSkinVertexFactory>>& VertexFactories,
 						 const FSkeletalMeshObjectGPUSkin::FVertexFactoryBuffers& InVertexBuffers,
 						 ERHIFeatureLevel::Type FeatureLevel
 						 )
-
 {
-	VertexFactoryType* VertexFactory = new VertexFactoryType(FeatureLevel, InVertexBuffers.NumVertices);
-	VertexFactories.Add(TUniquePtr<VertexFactoryTypeBase>(VertexFactory));
-						
+	FGPUBaseSkinVertexFactory* VertexFactory = nullptr;
+	GPUSkinBoneInfluenceType BoneInfluenceType = InVertexBuffers.SkinWeightVertexBuffer->GetBoneInfluenceType();
+	if (BoneInfluenceType == GPUSkinBoneInfluenceType::DefaultBoneInfluence)
+	{
+		VertexFactory = new TGPUSkinMorphVertexFactory<GPUSkinBoneInfluenceType::DefaultBoneInfluence>(FeatureLevel, InVertexBuffers.NumVertices);
+	}
+	else
+	{
+		VertexFactory = new TGPUSkinMorphVertexFactory<GPUSkinBoneInfluenceType::UnlimitedBoneInfluence>(FeatureLevel, InVertexBuffers.NumVertices);
+	}
+	VertexFactories.Add(TUniquePtr<FGPUBaseSkinVertexFactory>(VertexFactory));
+
 	// Setup the update data for enqueue
-	TDynamicUpdateVertexFactoryData<VertexFactoryType> VertexUpdateData(VertexFactory, InVertexBuffers);
+	FDynamicUpdateVertexFactoryData VertexUpdateData(VertexFactory, InVertexBuffers);
 
 	// update vertex factory components and sync it
 	ENQUEUE_RENDER_COMMAND(InitGPUSkinVertexFactoryMorph)(
 		[VertexUpdateData](FRHICommandList& RHICmdList)
 		{
-			typename VertexFactoryType::FDataType Data;
-			InitGPUSkinVertexFactoryComponents<VertexFactoryType>(&Data, VertexUpdateData.VertexBuffers, VertexUpdateData.VertexFactory);
-			InitMorphVertexFactoryComponents<VertexFactoryType>(&Data, VertexUpdateData.VertexBuffers);
-			VertexUpdateData.VertexFactory->SetData(Data);
+			FGPUSkinMorphDataType Data;
+			InitGPUSkinVertexFactoryComponents(&Data, VertexUpdateData.VertexBuffers, VertexUpdateData.VertexFactory);
+			InitMorphVertexFactoryComponents(&Data, VertexUpdateData.VertexBuffers);
+			VertexUpdateData.VertexFactory->SetData(&Data);
 			VertexUpdateData.VertexFactory->InitResource();
 		}
 	);
-
-	return VertexFactory;
 }
 
-template <class VertexFactoryTypeBase, class VertexFactoryType>
-static void UpdateVertexFactoryMorph(TArray<TUniquePtr<VertexFactoryTypeBase>>& VertexFactories,
+static void UpdateVertexFactoryMorph(TArray<TUniquePtr<FGPUBaseSkinVertexFactory>>& VertexFactories,
 	const FSkeletalMeshObjectGPUSkin::FVertexFactoryBuffers& InVertexBuffers)
 {
-	for (TUniquePtr<VertexFactoryTypeBase>& FactoryPtr : VertexFactories)
+	for (TUniquePtr<FGPUBaseSkinVertexFactory>& FactoryPtr : VertexFactories)
 	{
-		VertexFactoryType* VertexFactory = (VertexFactoryType*)FactoryPtr.Get();
+		FGPUBaseSkinVertexFactory* VertexFactory = FactoryPtr.Get();
 
 		if (VertexFactory != nullptr)
 		{
 			// Setup the update data for enqueue
-			TDynamicUpdateVertexFactoryData<VertexFactoryType> VertexUpdateData(VertexFactory, InVertexBuffers);
+			FDynamicUpdateVertexFactoryData VertexUpdateData(VertexFactory, InVertexBuffers);
 
 			// update vertex factory components and sync it
 			ENQUEUE_RENDER_COMMAND(InitGPUSkinVertexFactoryMorph)(
 				[VertexUpdateData](FRHICommandList& RHICmdList)
 			{
-				typename VertexFactoryType::FDataType Data;
-				InitGPUSkinVertexFactoryComponents<VertexFactoryType>(&Data, VertexUpdateData.VertexBuffers, VertexUpdateData.VertexFactory);
-				InitMorphVertexFactoryComponents<VertexFactoryType>(&Data, VertexUpdateData.VertexBuffers);
-				VertexUpdateData.VertexFactory->SetData(Data);
+				FGPUSkinMorphDataType Data;
+				InitGPUSkinVertexFactoryComponents(&Data, VertexUpdateData.VertexBuffers, VertexUpdateData.VertexFactory);
+				InitMorphVertexFactoryComponents(&Data, VertexUpdateData.VertexBuffers);
+				VertexUpdateData.VertexFactory->SetData(&Data);
 			});
 		}
 	}
@@ -1573,53 +1576,75 @@ static void UpdateVertexFactoryMorph(TArray<TUniquePtr<VertexFactoryTypeBase>>&
 /**
  * Creates a vertex factory entry for the given type and initialize it on the render thread
  */
-template <class VertexFactoryTypeBase, class VertexFactoryType>
-static void CreateVertexFactoryCloth(TArray<TUniquePtr<VertexFactoryTypeBase>>& VertexFactories,
+static void CreateVertexFactoryCloth(TArray<TUniquePtr<FGPUBaseSkinAPEXClothVertexFactory>>& VertexFactories,
 						 const FSkeletalMeshObjectGPUSkin::FVertexFactoryBuffers& InVertexBuffers,
-						 ERHIFeatureLevel::Type FeatureLevel
+						 ERHIFeatureLevel::Type FeatureLevel,
+						 bool bUseMultipleInfluences
 						 )
-
 {
-	VertexFactoryType* VertexFactory = new VertexFactoryType(FeatureLevel, InVertexBuffers.NumVertices);
-	VertexFactories.Add(TUniquePtr<VertexFactoryTypeBase>(VertexFactory));
-						
+	FGPUBaseSkinAPEXClothVertexFactory* ClothVertexFactory = nullptr;
+	GPUSkinBoneInfluenceType BoneInfluenceType = InVertexBuffers.SkinWeightVertexBuffer->GetBoneInfluenceType();
+	if (BoneInfluenceType == GPUSkinBoneInfluenceType::DefaultBoneInfluence)
+	{
+		if (bUseMultipleInfluences)
+		{
+			ClothVertexFactory = new TMultipleInfluenceClothVertexFactory<GPUSkinBoneInfluenceType::DefaultBoneInfluence>(FeatureLevel, InVertexBuffers.NumVertices);
+		}
+		else
+		{
+			ClothVertexFactory = new TGPUSkinAPEXClothVertexFactory<GPUSkinBoneInfluenceType::DefaultBoneInfluence>(FeatureLevel, InVertexBuffers.NumVertices);
+		}
+	}
+	else
+	{
+		if (bUseMultipleInfluences)
+		{
+			ClothVertexFactory = new TMultipleInfluenceClothVertexFactory<GPUSkinBoneInfluenceType::UnlimitedBoneInfluence>(FeatureLevel, InVertexBuffers.NumVertices);
+		}
+		else
+		{
+			ClothVertexFactory = new TGPUSkinAPEXClothVertexFactory<GPUSkinBoneInfluenceType::UnlimitedBoneInfluence>(FeatureLevel, InVertexBuffers.NumVertices);
+		}
+	}
+	VertexFactories.Add(TUniquePtr<FGPUBaseSkinAPEXClothVertexFactory>(ClothVertexFactory));
+
 	// Setup the update data for enqueue
-	TDynamicUpdateVertexFactoryData<VertexFactoryType> VertexUpdateData(VertexFactory, InVertexBuffers);
+	FGPUBaseSkinVertexFactory* VertexFactory = ClothVertexFactory->GetVertexFactory();
+	FDynamicUpdateVertexFactoryData VertexUpdateData(VertexFactory, InVertexBuffers);
 
 	// update vertex factory components and sync it
 	ENQUEUE_RENDER_COMMAND(InitGPUSkinAPEXClothVertexFactory)(
 		[VertexUpdateData](FRHICommandList& RHICmdList)
 		{
-			typename VertexFactoryType::FDataType Data;
-			InitGPUSkinVertexFactoryComponents<VertexFactoryType>(&Data, VertexUpdateData.VertexBuffers, VertexUpdateData.VertexFactory);
-			InitAPEXClothVertexFactoryComponents<VertexFactoryType>(&Data, VertexUpdateData.VertexBuffers);
-			VertexUpdateData.VertexFactory->SetData(Data);
+			FGPUSkinAPEXClothDataType Data;
+			InitGPUSkinVertexFactoryComponents(&Data, VertexUpdateData.VertexBuffers, VertexUpdateData.VertexFactory);
+			InitAPEXClothVertexFactoryComponents(&Data, VertexUpdateData.VertexBuffers);
+			VertexUpdateData.VertexFactory->SetData(&Data);
 			VertexUpdateData.VertexFactory->InitResource();
 		}
 	);
 }
 
-template <class VertexFactoryTypeBase, class VertexFactoryType>
-static void UpdateVertexFactoryCloth(TArray<TUniquePtr<VertexFactoryTypeBase>>& VertexFactories,
+static void UpdateVertexFactoryCloth(TArray<TUniquePtr<FGPUBaseSkinAPEXClothVertexFactory>>& VertexFactories,
 	const FSkeletalMeshObjectGPUSkin::FVertexFactoryBuffers& InVertexBuffers)
 {
-	for (TUniquePtr<VertexFactoryTypeBase>& FactoryPtr : VertexFactories)
+	for (TUniquePtr<FGPUBaseSkinAPEXClothVertexFactory>& FactoryPtr : VertexFactories)
 	{
-		VertexFactoryType* VertexFactory = (VertexFactoryType*)FactoryPtr.Get();
+		FGPUBaseSkinVertexFactory* VertexFactory = FactoryPtr->GetVertexFactory();
 
 		if (VertexFactory != nullptr)
 		{
 			// Setup the update data for enqueue
-			TDynamicUpdateVertexFactoryData<VertexFactoryType> VertexUpdateData(VertexFactory, InVertexBuffers);
+			FDynamicUpdateVertexFactoryData VertexUpdateData(VertexFactory, InVertexBuffers);
 
 			// update vertex factory components and sync it
 			ENQUEUE_RENDER_COMMAND(InitGPUSkinAPEXClothVertexFactory)(
 				[VertexUpdateData](FRHICommandList& RHICmdList)
 			{
-				typename VertexFactoryType::FDataType Data;
-				InitGPUSkinVertexFactoryComponents<VertexFactoryType>(&Data, VertexUpdateData.VertexBuffers, VertexUpdateData.VertexFactory);
-				InitAPEXClothVertexFactoryComponents<VertexFactoryType>(&Data, VertexUpdateData.VertexBuffers);
-				VertexUpdateData.VertexFactory->SetData(Data);
+				FGPUSkinAPEXClothDataType Data;
+				InitGPUSkinVertexFactoryComponents(&Data, VertexUpdateData.VertexBuffers, VertexUpdateData.VertexFactory);
+				InitAPEXClothVertexFactoryComponents(&Data, VertexUpdateData.VertexBuffers);
+				VertexUpdateData.VertexFactory->SetData(&Data);
 			});
 		}
 	}
@@ -1659,19 +1684,8 @@ void FSkeletalMeshObjectGPUSkin::FVertexFactoryData::InitVertexFactories(
 	{
 		for(int32 FactoryIdx = 0; FactoryIdx < Sections.Num(); ++FactoryIdx)
 		{
-			GPUSkinBoneInfluenceType BoneInfluenceType = VertexBuffers.SkinWeightVertexBuffer->GetBoneInfluenceType();
-			if (BoneInfluenceType == GPUSkinBoneInfluenceType::DefaultBoneInfluence)
-			{
-				TGPUSkinVertexFactory<GPUSkinBoneInfluenceType::DefaultBoneInfluence>* VertexFactory = 
-					CreateVertexFactory< FGPUBaseSkinVertexFactory, TGPUSkinVertexFactory<GPUSkinBoneInfluenceType::DefaultBoneInfluence> >(VertexFactories, VertexBuffers, InFeatureLevel);
-				CreatePassthroughVertexFactory<TGPUSkinVertexFactory<GPUSkinBoneInfluenceType::DefaultBoneInfluence>>(InFeatureLevel, PassthroughVertexFactories, VertexFactory);
-			}
-			else
-			{
-				TGPUSkinVertexFactory<GPUSkinBoneInfluenceType::UnlimitedBoneInfluence>* VertexFactory = 
-					CreateVertexFactory< FGPUBaseSkinVertexFactory, TGPUSkinVertexFactory<GPUSkinBoneInfluenceType::UnlimitedBoneInfluence> >(VertexFactories, VertexBuffers, InFeatureLevel);
-				CreatePassthroughVertexFactory<TGPUSkinVertexFactory<GPUSkinBoneInfluenceType::UnlimitedBoneInfluence>>(InFeatureLevel, PassthroughVertexFactories, VertexFactory);
-			}
+			FGPUBaseSkinVertexFactory* VertexFactory = CreateVertexFactory(VertexFactories, VertexBuffers, InFeatureLevel);
+			CreatePassthroughVertexFactory(InFeatureLevel, PassthroughVertexFactories, VertexFactory);
 		}
 	}
 }
@@ -1703,15 +1717,7 @@ void FSkeletalMeshObjectGPUSkin::FVertexFactoryData::InitMorphVertexFactories(
 	MorphVertexFactories.Empty(Sections.Num());
 	for( int32 FactoryIdx=0; FactoryIdx < Sections.Num(); FactoryIdx++ )
 	{
-		GPUSkinBoneInfluenceType BoneInfluenceType = VertexBuffers.SkinWeightVertexBuffer->GetBoneInfluenceType();
-		if (BoneInfluenceType == GPUSkinBoneInfluenceType::DefaultBoneInfluence)
-		{
-			CreateVertexFactoryMorph<FGPUBaseSkinVertexFactory, TGPUSkinMorphVertexFactory< GPUSkinBoneInfluenceType::DefaultBoneInfluence > >(MorphVertexFactories,VertexBuffers,InFeatureLevel);
-		}
-		else
-		{
-			CreateVertexFactoryMorph<FGPUBaseSkinVertexFactory, TGPUSkinMorphVertexFactory<GPUSkinBoneInfluenceType::UnlimitedBoneInfluence> >(MorphVertexFactories, VertexBuffers, InFeatureLevel);
-		}
+		CreateVertexFactoryMorph(MorphVertexFactories, VertexBuffers, InFeatureLevel);
 	}
 }
 
@@ -1745,36 +1751,7 @@ void FSkeletalMeshObjectGPUSkin::FVertexFactoryData::InitAPEXClothVertexFactorie
 			const uint32 NumClothWeights = Section.ClothMappingDataLODs.Num() ? Section.ClothMappingDataLODs[ClothLODBias].Num(): 0;
 			const uint32 NumPositionVertices = Section.NumVertices;
 			const bool bUseMultipleInfluences = (NumClothWeights > NumPositionVertices);
-
-			GPUSkinBoneInfluenceType BoneInfluenceType = VertexBuffers.SkinWeightVertexBuffer->GetBoneInfluenceType();
-			if (BoneInfluenceType == GPUSkinBoneInfluenceType::DefaultBoneInfluence)
-			{
-				if (bUseMultipleInfluences)
-				{
-					CreateVertexFactoryCloth<FGPUBaseSkinAPEXClothVertexFactory, TMultipleInfluenceClothVertexFactory<GPUSkinBoneInfluenceType::DefaultBoneInfluence> >(ClothVertexFactories, VertexBuffers, InFeatureLevel);
-				}
-				else
-				{
-					CreateVertexFactoryCloth
-						<FGPUBaseSkinAPEXClothVertexFactory, 
-						TGPUSkinAPEXClothVertexFactory<GPUSkinBoneInfluenceType::DefaultBoneInfluence>>
-						(ClothVertexFactories, VertexBuffers, InFeatureLevel);
-				}
-			}
-			else
-			{
-				if (bUseMultipleInfluences)
-				{
-					CreateVertexFactoryCloth<FGPUBaseSkinAPEXClothVertexFactory, TMultipleInfluenceClothVertexFactory<GPUSkinBoneInfluenceType::UnlimitedBoneInfluence> >(ClothVertexFactories, VertexBuffers, InFeatureLevel);
-				}
-				else
-				{
-					CreateVertexFactoryCloth
-						<FGPUBaseSkinAPEXClothVertexFactory, 
-						TGPUSkinAPEXClothVertexFactory<GPUSkinBoneInfluenceType::UnlimitedBoneInfluence>>
-						(ClothVertexFactories, VertexBuffers, InFeatureLevel);
-				}
-			}
+			CreateVertexFactoryCloth(ClothVertexFactories, VertexBuffers, InFeatureLevel, bUseMultipleInfluences);
 		}
 		else
 		{
@@ -1801,47 +1778,9 @@ void FSkeletalMeshObjectGPUSkin::FVertexFactoryData::ReleaseAPEXClothVertexFacto
 
 void FSkeletalMeshObjectGPUSkin::FVertexFactoryData::UpdateVertexFactoryData(const FVertexFactoryBuffers& VertexBuffers)
 {
-	bool bUseMultipleInfluences = (VertexBuffers.APEXClothVertexBuffer->GetNumVertices() > VertexBuffers.StaticVertexBuffers->PositionVertexBuffer.GetNumVertices());
-
-	GPUSkinBoneInfluenceType BoneInfluenceType = VertexBuffers.SkinWeightVertexBuffer->GetBoneInfluenceType();
-	if (BoneInfluenceType == GPUSkinBoneInfluenceType::DefaultBoneInfluence)
-	{
-		UpdateVertexFactory<FGPUBaseSkinVertexFactory, TGPUSkinVertexFactory<GPUSkinBoneInfluenceType::DefaultBoneInfluence>>(VertexFactories, VertexBuffers);
-		if (bUseMultipleInfluences)
-		{ 
-			UpdateVertexFactoryCloth
-				<FGPUBaseSkinAPEXClothVertexFactory, 
-				TMultipleInfluenceClothVertexFactory<GPUSkinBoneInfluenceType::DefaultBoneInfluence>>
-				(ClothVertexFactories, VertexBuffers);
-		}
-		else
-		{
-			UpdateVertexFactoryCloth
-				<FGPUBaseSkinAPEXClothVertexFactory, 
-				TGPUSkinAPEXClothVertexFactory<GPUSkinBoneInfluenceType::DefaultBoneInfluence>>
-				(ClothVertexFactories, VertexBuffers);
-		}
-		UpdateVertexFactoryMorph<FGPUBaseSkinVertexFactory, TGPUSkinMorphVertexFactory<GPUSkinBoneInfluenceType::DefaultBoneInfluence>>(MorphVertexFactories, VertexBuffers);
-	}
-	else
-	{
-		UpdateVertexFactory<FGPUBaseSkinVertexFactory, TGPUSkinVertexFactory<GPUSkinBoneInfluenceType::UnlimitedBoneInfluence>>(VertexFactories, VertexBuffers);
-		if (bUseMultipleInfluences)
-		{
-			UpdateVertexFactoryCloth
-				<FGPUBaseSkinAPEXClothVertexFactory, 
-				TMultipleInfluenceClothVertexFactory<GPUSkinBoneInfluenceType::UnlimitedBoneInfluence>>
-				(ClothVertexFactories, VertexBuffers);
-		}
-		else
-		{
-			UpdateVertexFactoryCloth
-				<FGPUBaseSkinAPEXClothVertexFactory, 
-				TGPUSkinAPEXClothVertexFactory<GPUSkinBoneInfluenceType::UnlimitedBoneInfluence>>
-				(ClothVertexFactories, VertexBuffers);
-		}
-		UpdateVertexFactoryMorph<FGPUBaseSkinVertexFactory, TGPUSkinMorphVertexFactory<GPUSkinBoneInfluenceType::UnlimitedBoneInfluence>>(MorphVertexFactories, VertexBuffers);
-	}
+	UpdateVertexFactory(VertexFactories, VertexBuffers);
+	UpdateVertexFactoryCloth(ClothVertexFactories, VertexBuffers);
+	UpdateVertexFactoryMorph(MorphVertexFactories, VertexBuffers);
 }
 
 void FSkeletalMeshObjectGPUSkin::FSkeletalMeshObjectLOD::InitResources(const FSkelMeshObjectLODInfo& MeshLODInfo, FSkelMeshComponentLODInfo* CompLODInfo, ERHIFeatureLevel::Type InFeatureLevel)
diff --git a/Engine/Source/Runtime/Engine/Public/GPUSkinVertexFactory.h b/Engine/Source/Runtime/Engine/Public/GPUSkinVertexFactory.h
--- a/Engine/Source/Runtime/Engine/Public/GPUSkinVertexFactory.h
+++ b/Engine/Source/Runtime/Engine/Public/GPUSkinVertexFactory.h
@@ -157,6 +157,31 @@ enum GPUSkinBoneInfluenceType
 	UnlimitedBoneInfluence	// unlimited bones per vertex
 };
 
+/** Stream component data bound to GPU skinned vertex factory */
+struct FGPUSkinDataType : public FStaticMeshDataType
+{
+	/** The stream to read the bone indices from */
+	FVertexStreamComponent BoneIndices;
+
+	/** The stream to read the extra bone indices from */
+	FVertexStreamComponent ExtraBoneIndices;
+
+	/** The stream to read the bone weights from */
+	FVertexStreamComponent BoneWeights;
+
+	/** The stream to read the extra bone weights from */
+	FVertexStreamComponent ExtraBoneWeights;
+
+	/** The stream to read the blend stream offset and num of influences from */
+	FVertexStreamComponent BlendOffsetCount;
+
+	/** Number of bone influences */
+	uint32 NumBoneInfluences = 0;
+
+	/** If the bone indices are 16 or 8-bit format */
+	bool bUse16BitBoneIndex = 0;
+};
+
 /** Vertex factory with vertex stream components for GPU skinned vertices */
 class FGPUBaseSkinVertexFactory : public FVertexFactory
 {
@@ -304,9 +329,12 @@ public:
 		return ShaderData;
 	}
 
-	virtual GPUSkinBoneInfluenceType GetBoneInfluenceType() const { return DefaultBoneInfluence; }
-	virtual uint32 GetNumBoneInfluences() const { return 0; }
-	virtual bool Use16BitBoneIndex() const { return false; }
+	/**
+	* An implementation of the interface used by TSynchronizedResource to
+	* update the resource with new data from the game thread.
+	* @param	InData - new stream component data
+	*/
+	virtual void SetData(const FGPUSkinDataType* InData);
 
 	uint32 GetNumVertices() const
 	{
@@ -328,26 +356,39 @@ public:
 	virtual void UpdateMorphVertexStream(const class FMorphVertexBuffer* MorphVertexBuffer) {}
 	virtual const class FMorphVertexBuffer* GetMorphVertexBuffer(bool bPrevious, uint32 FrameNumber) const { return nullptr; }
 
-	virtual const FShaderResourceViewRHIRef GetPositionsSRV() const = 0;
-	virtual const FShaderResourceViewRHIRef GetTangentsSRV() const = 0;
-	virtual const FShaderResourceViewRHIRef GetTextureCoordinatesSRV() const = 0;
-	virtual const FShaderResourceViewRHIRef GetColorComponentsSRV() const = 0;
-	virtual uint32 GetNumTexCoords() const = 0;
-	virtual const uint32 GetColorIndexMask() const = 0;
+	virtual GPUSkinBoneInfluenceType GetBoneInfluenceType() const				{ return DefaultBoneInfluence; }
+	virtual uint32 GetNumBoneInfluences() const									{ return Data.IsValid() ? Data->NumBoneInfluences : 0; }
+	virtual bool Use16BitBoneIndex() const										{ return Data.IsValid() ? Data->bUse16BitBoneIndex : false; }
+	virtual const FShaderResourceViewRHIRef GetPositionsSRV() const				{ return Data.IsValid() ? Data->PositionComponentSRV : nullptr; }
+	virtual const FShaderResourceViewRHIRef GetTangentsSRV() const				{ return Data.IsValid() ? Data->TangentsSRV : nullptr; }
+	virtual const FShaderResourceViewRHIRef GetTextureCoordinatesSRV() const	{ return Data.IsValid() ? Data->TextureCoordinatesSRV : nullptr; }
+	virtual const FShaderResourceViewRHIRef GetColorComponentsSRV() const		{ return Data.IsValid() ? Data->ColorComponentsSRV : nullptr; }
+	virtual uint32 GetNumTexCoords() const										{ return Data.IsValid() ? Data->NumTexCoords : 0; }
+	virtual const uint32 GetColorIndexMask() const								{ return Data.IsValid() ? Data->ColorIndexMask : 0; }
 
 	inline const FVertexStreamComponent& GetTangentStreamComponent(int Index)
 	{
-		check(TangentStreamComponents[Index].VertexBuffer != nullptr);
-		return TangentStreamComponents[Index];
+		check(Data.IsValid() && Data->TangentBasisComponents[Index].VertexBuffer != nullptr);
+		return Data->TangentBasisComponents[Index];
 	}
 
+	void CopyDataTypeForPassthroughFactory(class FGPUSkinPassthroughVertexFactory* PassthroughVertexFactory);
+
 protected:
+	/**
+	* Add the decl elements for the streams
+	* @param InData - type with stream components
+	* @param OutElements - vertex decl list to modify
+	*/
+	virtual void AddVertexElements(FVertexDeclarationElementList& OutElements) = 0;
+
 	/** dynamic data need for setting the shader */ 
 	FShaderDataType ShaderData;
 	/** Pool of buffers for bone matrices. */
 	static TGlobalResource<FBoneBufferPool> BoneBufferPool;
 
-	FVertexStreamComponent TangentStreamComponents[2];
+	/** stream component data bound to this vertex factory */
+	TUniquePtr<FGPUSkinDataType> Data;
 
 private:
 	uint32 NumVertices;
@@ -360,30 +401,6 @@ class ENGINE_API TGPUSkinVertexFactory : public FGPUBaseSkinVertexFactory
 	DECLARE_VERTEX_FACTORY_TYPE(TGPUSkinVertexFactory<BoneInfluenceType>);
 
 public:
-	struct FDataType : public FStaticMeshDataType
-	{
-		/** The stream to read the bone indices from */
-		FVertexStreamComponent BoneIndices;
-
-		/** The stream to read the extra bone indices from */
-		FVertexStreamComponent ExtraBoneIndices;
-
-		/** The stream to read the bone weights from */
-		FVertexStreamComponent BoneWeights;
-
-		/** The stream to read the extra bone weights from */
-		FVertexStreamComponent ExtraBoneWeights;
-
-		/** The stream to read the blend stream offset and num of influences from */
-		FVertexStreamComponent BlendOffsetCount;
-
-		/** Number of bone influences */
-		uint32 NumBoneInfluences = 0;
-
-		/** If the bone indices are 16 or 8-bit format */
-		bool bUse16BitBoneIndex = 0;
-	};
-
 	/**
 	 * Constructor presizing bone matrices array to used amount.
 	 *
@@ -396,82 +413,21 @@ public:
 	virtual GPUSkinBoneInfluenceType GetBoneInfluenceType() const override
 	{ return BoneInfluenceType; }
 
-	virtual uint32 GetNumBoneInfluences() const override
-	{
-		return Data.NumBoneInfluences;
-	}
-
-	virtual bool Use16BitBoneIndex() const override
-	{
-		return Data.bUse16BitBoneIndex;
-	}
-
 	static void ModifyCompilationEnvironment(const FVertexFactoryShaderPermutationParameters& Parameters, FShaderCompilerEnvironment& OutEnvironment);
 	static bool ShouldCompilePermutation(const FVertexFactoryShaderPermutationParameters& Parameters);
 	
-	/**
-	* An implementation of the interface used by TSynchronizedResource to 
-	* update the resource with new data from the game thread.
-	* @param	InData - new stream component data
-	*/
-	void SetData(const FDataType& InData)
-	{
-		Data = InData;
-		this->TangentStreamComponents[0] = InData.TangentBasisComponents[0];
-		this->TangentStreamComponents[1] = InData.TangentBasisComponents[1];
-		FGPUBaseSkinVertexFactory::UpdateRHI();
-	}
-
 	// FRenderResource interface.
 	virtual void InitRHI() override;
 	virtual void InitDynamicRHI() override;
 	virtual void ReleaseDynamicRHI() override;
 
-	void CopyDataTypeForPassthroughFactory(class FGPUSkinPassthroughVertexFactory* PassthroughVertexFactory);
-
-	const FShaderResourceViewRHIRef GetPositionsSRV() const override
-	{
-		return Data.PositionComponentSRV;
-	}
-
-	const FShaderResourceViewRHIRef GetTangentsSRV() const override
-	{
-		return Data.TangentsSRV;
-	}
-
-	const FShaderResourceViewRHIRef GetTextureCoordinatesSRV() const override
-	{
-		return Data.TextureCoordinatesSRV;
-	}
-
-	uint32 GetNumTexCoords() const override
-	{
-		return Data.NumTexCoords;
-	}
-
-	const FShaderResourceViewRHIRef GetColorComponentsSRV() const override
-	{
-		return Data.ColorComponentsSRV;
-	}
-
-	const uint32 GetColorIndexMask() const override
-	{
-		return Data.ColorIndexMask;
-	}
-
 protected:
 	/**
 	* Add the decl elements for the streams
 	* @param InData - type with stream components
 	* @param OutElements - vertex decl list to modify
 	*/
-	void AddVertexElements(FDataType& InData, FVertexDeclarationElementList& OutElements);
-
-	inline const FDataType& GetData() const { return Data; }
-
-private:
-	/** stream component data bound to this vertex factory */
-	FDataType Data;  
+	virtual void AddVertexElements(FVertexDeclarationElementList& OutElements) override;
 };
 
 /** 
@@ -545,6 +501,17 @@ protected:
 	void InternalUpdateVertexDeclaration(FGPUBaseSkinVertexFactory* SourceVertexFactory, struct FRWBuffer* PositionRWBuffer, class FRHIShaderResourceView* PreSkinPositionSRV, struct FRWBuffer* TangentRWBuffer);
 };
 
+/** Stream component data bound to morph vertex factory */
+struct FGPUSkinMorphDataType : public FGPUSkinDataType
+{
+	/** stream which has the position deltas to add to the vertex position */
+	FVertexStreamComponent DeltaPositionComponent;
+	/** stream which has the TangentZ deltas to add to the vertex normals */
+	FVertexStreamComponent DeltaTangentZComponent;
+	/** Morph vertex buffer pool double buffering delta data  */
+	class FMorphVertexBufferPool* MorphVertexBufferPool = nullptr;
+};
+
 /** Vertex factory with vertex stream components for GPU-skinned and morph target streams */
 template<GPUSkinBoneInfluenceType BoneInfluenceType>
 class TGPUSkinMorphVertexFactory : public TGPUSkinVertexFactory<BoneInfluenceType>
@@ -553,17 +520,6 @@ class TGPUSkinMorphVertexFactory : public TGPUSkinVertexFactory<BoneInfluenceTyp
 
 	typedef TGPUSkinVertexFactory<BoneInfluenceType> Super;
 public:
-
-	struct FDataType : TGPUSkinVertexFactory<BoneInfluenceType>::FDataType
-	{
-		/** stream which has the position deltas to add to the vertex position */
-		FVertexStreamComponent DeltaPositionComponent;
-		/** stream which has the TangentZ deltas to add to the vertex normals */
-		FVertexStreamComponent DeltaTangentZComponent;
-		/** Morph vertex buffer pool double buffering delta data  */
-		class FMorphVertexBufferPool* MorphVertexBufferPool = nullptr;
-	};
-
 	/**
 	 * Constructor presizing bone matrices array to used amount.
 	 *
@@ -581,13 +537,7 @@ public:
 	* update the resource with new data from the game thread.
 	* @param	InData - new stream component data
 	*/
-	void SetData(const FDataType& InData)
-	{
-		MorphData = InData;
-		this->TangentStreamComponents[0] = InData.TangentBasisComponents[0];
-		this->TangentStreamComponents[1] = InData.TangentBasisComponents[1];
-		FGPUBaseSkinVertexFactory::UpdateRHI();
-	}
+	virtual void SetData(const FGPUSkinDataType* InData) override;
 
 	/** FGPUBaseSkinVertexFactory overrides */
 	virtual void UpdateMorphVertexStream(const class FMorphVertexBuffer* MorphVertexBuffer) override;
@@ -601,58 +551,18 @@ public:
 	*/
 	virtual void InitRHI() override;
 
-
-	virtual uint32 GetNumBoneInfluences() const override
-	{
-		return MorphData.NumBoneInfluences;
-	}
-
-	virtual bool Use16BitBoneIndex() const override
-	{
-		return MorphData.bUse16BitBoneIndex;
-	}
-
-	const FShaderResourceViewRHIRef GetPositionsSRV() const override
-	{
-		return MorphData.PositionComponentSRV;
-	}
-
-	const FShaderResourceViewRHIRef GetTangentsSRV() const override
-	{
-		return MorphData.TangentsSRV;
-	}
-
-	const FShaderResourceViewRHIRef GetTextureCoordinatesSRV() const override
-	{
-		return MorphData.TextureCoordinatesSRV;
-	}
-
-	uint32 GetNumTexCoords() const override
-	{
-		return MorphData.NumTexCoords;
-	}
-
-	const FShaderResourceViewRHIRef GetColorComponentsSRV() const override
-	{
-		return MorphData.ColorComponentsSRV;
-	}
-
-	const uint32 GetColorIndexMask() const override
-	{
-		return MorphData.ColorIndexMask;
-	}
-
 protected:
 	/**
 	* Add the decl elements for the streams
 	* @param InData - type with stream components
 	* @param OutElements - vertex decl list to modify
 	*/
-	void AddVertexElements(FDataType& InData, FVertexDeclarationElementList& OutElements);
+	virtual void AddVertexElements(FVertexDeclarationElementList& OutElements) override;
+
+	/** Alias pointer to TUniquePtr<FGPUSkinDataType> Data of FGPUBaseSkinVertexFactory. Note memory isn't managed through this pointer. */
+	FGPUSkinMorphDataType* MorphDataPtr = nullptr;
 
 private:
-	/** stream component data bound to this vertex factory */
-	FDataType MorphData;
 	int32 MorphDeltaStreamIndex = -1;
 };
 
@@ -856,6 +766,14 @@ protected:
 	static TGlobalResource<FClothBufferPool> ClothSimulDataBufferPool;
 };
 
+/** Stream component data bound to Apex cloth vertex factory */
+struct FGPUSkinAPEXClothDataType : public FGPUSkinDataType
+{
+	FShaderResourceViewRHIRef ClothBuffer;
+	// Packed Map: u32 Key, u32 Value
+	TArray<FClothBufferIndexMapping> ClothIndexMapping;
+};
+
 /** Vertex factory with vertex stream components for GPU-skinned and morph target streams */
 template<GPUSkinBoneInfluenceType BoneInfluenceType>
 class TGPUSkinAPEXClothVertexFactory : public FGPUBaseSkinAPEXClothVertexFactory, public TGPUSkinVertexFactory<BoneInfluenceType>
@@ -865,31 +783,26 @@ class TGPUSkinAPEXClothVertexFactory : public FGPUBaseSkinAPEXClothVertexFactory
 	typedef TGPUSkinVertexFactory<BoneInfluenceType> Super;
 
 public:
-
-	struct FDataType : TGPUSkinVertexFactory<BoneInfluenceType>::FDataType
-	{
-		FShaderResourceViewRHIRef ClothBuffer;
-		// Packed Map: u32 Key, u32 Value
-		TArray<FClothBufferIndexMapping> ClothIndexMapping;
-	};
-
 	inline FShaderResourceViewRHIRef GetClothBuffer()
 	{
-		return MeshMappingData.ClothBuffer;
+		return ClothDataPtr ? ClothDataPtr->ClothBuffer : nullptr;
 	}
 
 	inline const FShaderResourceViewRHIRef GetClothBuffer() const
 	{
-		return MeshMappingData.ClothBuffer;
+		return ClothDataPtr ? ClothDataPtr->ClothBuffer : nullptr;
 	}
 
 	inline uint32 GetClothIndexOffset(uint32 VertexIndex, uint32 LODBias = 0) const
 	{
-		for (const FClothBufferIndexMapping& Mapping : MeshMappingData.ClothIndexMapping)
+		if (ClothDataPtr)
 		{
-			if (Mapping.BaseVertexIndex == VertexIndex)
+			for (const FClothBufferIndexMapping& Mapping : ClothDataPtr->ClothIndexMapping)
 			{
-				return Mapping.MappingOffset + Mapping.LODBiasStride * LODBias;
+				if (Mapping.BaseVertexIndex == VertexIndex)
+				{
+					return Mapping.MappingOffset + Mapping.LODBiasStride * LODBias;
+				}
 			}
 		}
 
@@ -914,14 +827,7 @@ public:
 	* update the resource with new data from the game thread.
 	* @param	InData - new stream component data
 	*/
-	void SetData(const FDataType& InData)
-	{
-        Super::SetData(InData);
-		MeshMappingData = InData;
-		this->TangentStreamComponents[0] = InData.TangentBasisComponents[0];
-		this->TangentStreamComponents[1] = InData.TangentBasisComponents[1];
-		FGPUBaseSkinVertexFactory::UpdateRHI();
-	}
+	virtual void SetData(const FGPUSkinDataType* InData) override;
 
 	virtual FGPUBaseSkinVertexFactory* GetVertexFactory() override
 	{
@@ -933,46 +839,6 @@ public:
 		return this;
 	}
 
-	virtual uint32 GetNumBoneInfluences() const override
-	{
-		return MeshMappingData.NumBoneInfluences;
-	}
-
-	virtual bool Use16BitBoneIndex() const override
-	{
-		return MeshMappingData.bUse16BitBoneIndex;
-	}
-
-	const FShaderResourceViewRHIRef GetPositionsSRV() const override
-	{
-		return MeshMappingData.PositionComponentSRV;
-	}
-
-	const FShaderResourceViewRHIRef GetTangentsSRV() const override
-	{
-		return MeshMappingData.TangentsSRV;
-	}
-
-	const FShaderResourceViewRHIRef GetTextureCoordinatesSRV() const override
-	{
-		return MeshMappingData.TextureCoordinatesSRV;
-	}
-
-	uint32 GetNumTexCoords() const override
-	{
-		return MeshMappingData.NumTexCoords;
-	}
-
-	const FShaderResourceViewRHIRef GetColorComponentsSRV() const override
-	{
-		return MeshMappingData.ColorComponentsSRV;
-	}
-
-	const uint32 GetColorIndexMask() const override
-	{
-		return MeshMappingData.ColorIndexMask;
-	}
-
 	// FRenderResource interface.
 
 	/**
@@ -982,9 +848,9 @@ public:
 	virtual void InitRHI() override;
 	virtual void ReleaseDynamicRHI() override;
 
-private:
-	/** stream component data bound to this vertex factory */
-	FDataType MeshMappingData; 
+protected:
+	/** Alias pointer to TUniquePtr<FGPUSkinDataType> Data of FGPUBaseSkinVertexFactory. Note memory isn't managed through this pointer. */
+	FGPUSkinAPEXClothDataType* ClothDataPtr = nullptr;
 };
 
 
