Check for INDEX_NONE return from material compiler

diff --git a/Engine/Source/Runtime/Engine/Private/Materials/MaterialShared.cpp b/Engine/Source/Runtime/Engine/Private/Materials/MaterialShared.cpp
--- a/Engine/Source/Runtime/Engine/Private/Materials/MaterialShared.cpp
+++ b/Engine/Source/Runtime/Engine/Private/Materials/MaterialShared.cpp
@@ -4142,12 +4142,15 @@ int32 UMaterialInterface::CompileProperty(FMaterialCompiler* Compiler, EMaterial
 		Result = FMaterialAttributeDefinitionMap::CompileDefaultExpression(Compiler, Property);
 	}
 
-	// Cast is always required to go between float and LWC
-	const EMaterialValueType ResultType = Compiler->GetParameterType(Result);
-	const EMaterialValueType PropertyType = FMaterialAttributeDefinitionMap::GetValueType(Property);
-	if ((ForceCastFlags & MFCF_ForceCast) || IsLWCType(ResultType) != IsLWCType(PropertyType))
+	if (Result != INDEX_NONE)
 	{
-		Result = Compiler->ForceCast(Result, PropertyType, ForceCastFlags);
+		// Cast is always required to go between float and LWC
+		const EMaterialValueType ResultType = Compiler->GetParameterType(Result);
+		const EMaterialValueType PropertyType = FMaterialAttributeDefinitionMap::GetValueType(Property);
+		if ((ForceCastFlags & MFCF_ForceCast) || IsLWCType(ResultType) != IsLWCType(PropertyType))
+		{
+			Result = Compiler->ForceCast(Result, PropertyType, ForceCastFlags);
+		}
 	}
 
 	return Result;
