Remove FString temporaries
  - Gets rid of 800K FString temporaries from process start to end of skydiving section

diff --git a/Engine/Source/Runtime/Engine/Private/Materials/MaterialShared.cpp b/Engine/Source/Runtime/Engine/Private/Materials/MaterialShared.cpp
--- a/Engine/Source/Runtime/Engine/Private/Materials/MaterialShared.cpp
+++ b/Engine/Source/Runtime/Engine/Private/Materials/MaterialShared.cpp
@@ -1608,8 +1608,8 @@ bool FMaterialResource::ShouldApplyFogging() const { return Material->bUseTransl
 bool FMaterialResource::ShouldApplyCloudFogging() const { return Material->bApplyCloudFogging; }
 bool FMaterialResource::IsSky() const { return Material->bIsSky; }
 bool FMaterialResource::ComputeFogPerPixel() const {return Material->bComputeFogPerPixel;}
-FString FMaterialResource::GetFriendlyName() const { return *GetNameSafe(Material); } //avoid using the material instance name here, we want materials that share a shadermap to also share a friendly name.
-FString FMaterialResource::GetAssetName() const { return MaterialInstance ? *GetNameSafe(MaterialInstance) : *GetNameSafe(Material); }
+FString FMaterialResource::GetFriendlyName() const { return GetNameSafe(Material); } //avoid using the material instance name here, we want materials that share a shadermap to also share a friendly name.
+FString FMaterialResource::GetAssetName() const { return MaterialInstance ? GetNameSafe(MaterialInstance) : GetNameSafe(Material); }
 
 uint32 FMaterialResource::GetMaterialDecalResponse() const
 {
