Fix gamepad state not maintained if using legacy interface for player input

diff --git a/Engine/Source/Runtime/Engine/Private/PlayerController.cpp b/Engine/Source/Runtime/Engine/Private/PlayerController.cpp
--- a/Engine/Source/Runtime/Engine/Private/PlayerController.cpp
+++ b/Engine/Source/Runtime/Engine/Private/PlayerController.cpp
@@ -2283,6 +2283,7 @@ bool APlayerController::InputKey(FKey Key, EInputEvent EventType, float AmountDe
 	Params.Key = Key;
 	Params.Event = EventType;
 	Params.Delta.X = AmountDepressed;
+	Params.bIsGamepadOverride = bGamepad;
 	
 	return InputKey(Params);
 }
@@ -2377,6 +2378,7 @@ bool APlayerController::InputAxis(FKey Key, float Delta, float DeltaTime, int32
 	Params.Delta = FVector(static_cast<double>(Delta), 0.0, 0.0);
 	Params.NumSamples = NumSamples;
 	Params.DeltaTime = DeltaTime;
+	Params.bIsGamepadOverride = bGamepad;
 
 	return InputKey(Params);
 }
