Restore the property with the deprecation warning for 5.0.

diff --git a/Engine/Source/Runtime/Engine/Classes/Engine/Texture.h b/Engine/Source/Runtime/Engine/Classes/Engine/Texture.h
--- a/Engine/Source/Runtime/Engine/Classes/Engine/Texture.h
+++ b/Engine/Source/Runtime/Engine/Classes/Engine/Texture.h
@@ -23,6 +23,10 @@
 #include "DerivedDataCacheKeyProxy.h"
 #endif
 
+#ifndef WITH_TEXTURE_RESOURCE_DEPRECATIONS
+#define WITH_TEXTURE_RESOURCE_DEPRECATIONS 1
+#endif
+
 #if WITH_EDITOR
 #include "Templates/DontCopy.h"
 #endif
@@ -1213,6 +1217,11 @@ private:
 	class FTextureResource* PrivateResourceRenderThread;
 
 public:
+#if WITH_TEXTURE_RESOURCE_DEPRECATIONS
+	UE_DEPRECATED(5.00, "Use GetResource() / SetResource() accessors instead. This value cannot be relied upon anymore.")
+	class FTextureResource* Resource {nullptr};
+#endif
+
 	PRAGMA_DISABLE_DEPRECATION_WARNINGS
 	ENGINE_API virtual ~UTexture() {};
 	PRAGMA_ENABLE_DEPRECATION_WARNINGS
diff --git a/Engine/Source/Runtime/Engine/Classes/Engine/Texture2D.h b/Engine/Source/Runtime/Engine/Classes/Engine/Texture2D.h
--- a/Engine/Source/Runtime/Engine/Classes/Engine/Texture2D.h
+++ b/Engine/Source/Runtime/Engine/Classes/Engine/Texture2D.h
@@ -14,6 +14,10 @@
 class FTexture2DResourceMem;
 class FTexture2DResource;
 
+#ifndef WITH_TEXTURE_PLATFORMDATA_DEPRECATIONS
+#define WITH_TEXTURE_PLATFORMDATA_DEPRECATIONS 1
+#endif
+
 UCLASS(hidecategories=Object, MinimalAPI, BlueprintType)
 class UTexture2D : public UTexture
 {
@@ -71,6 +75,11 @@ private:
 	FTexturePlatformData* PrivatePlatformData;
 
 public:
+#if WITH_TEXTURE_PLATFORMDATA_DEPRECATIONS
+	UE_DEPRECATED(5.00, "Use GetPlatformData() / SetPlatformData() accessors instead. This value cannot be relied upon anymore.")
+	FTexturePlatformData* PlatformData {nullptr};
+#endif
+
 	/** Set the derived data for this texture on this platform. */
 	ENGINE_API void SetPlatformData(FTexturePlatformData* PlatformData);
 	/** Get the derived data for this texture on this platform. */
diff --git a/Engine/Source/Runtime/Engine/Classes/Engine/TextureCube.h b/Engine/Source/Runtime/Engine/Classes/Engine/TextureCube.h
--- a/Engine/Source/Runtime/Engine/Classes/Engine/TextureCube.h
+++ b/Engine/Source/Runtime/Engine/Classes/Engine/TextureCube.h
@@ -10,6 +10,10 @@
 
 class FTextureResource;
 
+#ifndef WITH_TEXTURE_PLATFORMDATA_DEPRECATIONS
+#define WITH_TEXTURE_PLATFORMDATA_DEPRECATIONS 1
+#endif
+
 UCLASS(hidecategories=Object, MinimalAPI)
 class UTextureCube : public UTexture
 {
@@ -19,6 +23,11 @@ class UTextureCube : public UTexture
 	FTexturePlatformData* PrivatePlatformData;
 
 public:
+#if WITH_TEXTURE_PLATFORMDATA_DEPRECATIONS
+	UE_DEPRECATED(5.00, "Use GetPlatformData() / SetPlatformData() accessors instead. This value cannot be relied upon anymore.")
+	FTexturePlatformData* PlatformData {nullptr};
+#endif
+
 	/** Set the derived data for this texture on this platform. */
 	ENGINE_API void SetPlatformData(FTexturePlatformData* PlatformData);
 	/** Get the derived data for this texture on this platform. */
