Fix for race condition crash in stream cache

diff --git a/Engine/Source/Runtime/Engine/Private/AudioStreamingCache.cpp b/Engine/Source/Runtime/Engine/Private/AudioStreamingCache.cpp
--- a/Engine/Source/Runtime/Engine/Private/AudioStreamingCache.cpp
+++ b/Engine/Source/Runtime/Engine/Private/AudioStreamingCache.cpp
@@ -1806,6 +1806,7 @@ bool FAudioChunkCache::IsKeyValid(const FChunkKey& InKey)
 
 uint64 FAudioChunkCache::GetCacheLookupIDForChunk(const FChunkKey& InChunkKey) const
 {
+	FScopeLock Lock(&CacheMutationCriticalSection);
 	const uint64* ID = CacheLookupIdMap.Find(InChunkKey);
 
 	if (ID)
@@ -1820,8 +1821,8 @@ uint64 FAudioChunkCache::GetCacheLookupIDForChunk(const FChunkKey& InChunkKey) c
 
 void FAudioChunkCache::SetCacheLookupIDForChunk(const FChunkKey& InChunkKey, uint64 InCacheLookupID)
 {
-	uint64& ID = CacheLookupIdMap.FindOrAdd(InChunkKey);
-	ID = InCacheLookupID;
+	FScopeLock Lock(&CacheMutationCriticalSection);
+	CacheLookupIdMap.FindOrAdd(InChunkKey, InCacheLookupID);
 }
 
 
diff --git a/Engine/Source/Runtime/Engine/Public/AudioStreamingCache.h b/Engine/Source/Runtime/Engine/Public/AudioStreamingCache.h
--- a/Engine/Source/Runtime/Engine/Public/AudioStreamingCache.h
+++ b/Engine/Source/Runtime/Engine/Public/AudioStreamingCache.h
@@ -388,7 +388,7 @@ private:
 
 	// Critical section: only used when we are modifying element positions in the cache. This only happens in TouchElement, EvictLeastRecentChunk, and TrimMemory.
 	// Individual cache elements should be thread safe to access.
-	FCriticalSection CacheMutationCriticalSection;
+	mutable FCriticalSection CacheMutationCriticalSection;
 
 	// Map that USoundWaves, FSoundWaveProxys, FChunkKeys, and FAudioChunkHandles can use to
 	// quickly lookup where their chunks are currently stored in the cache
