World Partition: always create the world data layers actor.

diff --git a/Engine/Source/Editor/UnrealEd/Private/Commandlets/WorldPartitionConvertCommandlet.cpp b/Engine/Source/Editor/UnrealEd/Private/Commandlets/WorldPartitionConvertCommandlet.cpp
--- a/Engine/Source/Editor/UnrealEd/Private/Commandlets/WorldPartitionConvertCommandlet.cpp
+++ b/Engine/Source/Editor/UnrealEd/Private/Commandlets/WorldPartitionConvertCommandlet.cpp
@@ -861,10 +861,7 @@ int32 UWorldPartitionConvertCommandlet::Main(const FString& Params)
 
 	UPackage* MainPackage = MainLevel->GetPackage();
 	AWorldDataLayers* MainWorldDataLayers = MainWorld->GetWorldDataLayers();
-	if (!MainWorldDataLayers)
-	{
-		MainWorldDataLayers = AWorldDataLayers::Create(MainWorld);
-	}
+	check(MainWorldDataLayers);
 
 	OnWorldLoaded(MainWorld);
 
diff --git a/Engine/Source/Runtime/Engine/Private/WorldPartition/WorldPartition.cpp b/Engine/Source/Runtime/Engine/Private/WorldPartition/WorldPartition.cpp
--- a/Engine/Source/Runtime/Engine/Private/WorldPartition/WorldPartition.cpp
+++ b/Engine/Source/Runtime/Engine/Private/WorldPartition/WorldPartition.cpp
@@ -664,6 +664,7 @@ void UWorldPartition::OnWorldMatchStarting()
 #if WITH_EDITOR
 UWorldPartition* UWorldPartition::CreateOrRepairWorldPartition(AWorldSettings* WorldSettings, TSubclassOf<UWorldPartitionEditorHash> EditorHashClass, TSubclassOf<UWorldPartitionRuntimeHash> RuntimeHashClass)
 {
+	UWorld* World = WorldSettings->GetWorld();
 	UWorldPartition* WorldPartition = WorldSettings->GetWorldPartition();
 
 	if (!WorldPartition)
@@ -701,7 +702,13 @@ UWorldPartition* UWorldPartition::CreateOrRepairWorldPartition(AWorldSettings* W
 		WorldPartition->RuntimeHash->SetDefaultValues();
 	}
 
-	WorldPartition->GetWorld()->PersistentLevel->bIsPartitioned = true;
+	AWorldDataLayers* WorldDataLayers = World->GetWorldDataLayers();
+	if (!WorldDataLayers)
+	{
+		WorldDataLayers = AWorldDataLayers::Create(World);
+	}
+
+	World->PersistentLevel->bIsPartitioned = true;
 
 	return WorldPartition;
 }
