Move definition in common file to fix include analyzer warnings

diff --git a/Engine/Source/Runtime/Engine/Classes/Engine/Texture.h b/Engine/Source/Runtime/Engine/Classes/Engine/Texture.h
--- a/Engine/Source/Runtime/Engine/Classes/Engine/Texture.h
+++ b/Engine/Source/Runtime/Engine/Classes/Engine/Texture.h
@@ -27,6 +27,10 @@
 #define WITH_TEXTURE_RESOURCE_DEPRECATIONS 1
 #endif
 
+#ifndef WITH_TEXTURE_PLATFORMDATA_DEPRECATIONS
+#define WITH_TEXTURE_PLATFORMDATA_DEPRECATIONS 1
+#endif
+
 #if WITH_EDITOR
 #include "Templates/DontCopy.h"
 #endif
diff --git a/Engine/Source/Runtime/Engine/Classes/Engine/Texture2D.h b/Engine/Source/Runtime/Engine/Classes/Engine/Texture2D.h
--- a/Engine/Source/Runtime/Engine/Classes/Engine/Texture2D.h
+++ b/Engine/Source/Runtime/Engine/Classes/Engine/Texture2D.h
@@ -14,10 +14,6 @@
 class FTexture2DResourceMem;
 class FTexture2DResource;
 
-#ifndef WITH_TEXTURE_PLATFORMDATA_DEPRECATIONS
-#define WITH_TEXTURE_PLATFORMDATA_DEPRECATIONS 1
-#endif
-
 UCLASS(hidecategories=Object, MinimalAPI, BlueprintType)
 class UTexture2D : public UTexture
 {
diff --git a/Engine/Source/Runtime/Engine/Classes/Engine/TextureCube.h b/Engine/Source/Runtime/Engine/Classes/Engine/TextureCube.h
--- a/Engine/Source/Runtime/Engine/Classes/Engine/TextureCube.h
+++ b/Engine/Source/Runtime/Engine/Classes/Engine/TextureCube.h
@@ -10,10 +10,6 @@
 
 class FTextureResource;
 
-#ifndef WITH_TEXTURE_PLATFORMDATA_DEPRECATIONS
-#define WITH_TEXTURE_PLATFORMDATA_DEPRECATIONS 1
-#endif
-
 UCLASS(hidecategories=Object, MinimalAPI)
 class UTextureCube : public UTexture
 {
