Fix missing StereoRenderingDevice nullptr check in CL 18291679

diff --git a/Engine/Source/Runtime/Engine/Private/LocalPlayer.cpp b/Engine/Source/Runtime/Engine/Private/LocalPlayer.cpp
--- a/Engine/Source/Runtime/Engine/Private/LocalPlayer.cpp
+++ b/Engine/Source/Runtime/Engine/Private/LocalPlayer.cpp
@@ -773,6 +773,11 @@ bool ULocalPlayer::CalcSceneViewInitOptions(
 		ViewInitOptions.bInCameraCut = PlayerController->PlayerCameraManager->bGameCameraCutThisFrame;
 	}
 
+	if (GEngine->StereoRenderingDevice.IsValid())
+	{
+		ViewInitOptions.StereoPass = GEngine->StereoRenderingDevice->GetViewPassForIndex(StereoViewIndex != INDEX_NONE, StereoViewIndex);
+	}
+
 	check(PlayerController && PlayerController->GetWorld());
 
 	const uint32 ViewIndex = StereoViewIndex != INDEX_NONE ? StereoViewIndex : 0;
@@ -790,7 +795,6 @@ bool ULocalPlayer::CalcSceneViewInitOptions(
 	ViewInitOptions.ViewElementDrawer = ViewDrawer;
 	ViewInitOptions.BackgroundColor = FLinearColor::Black;
 	ViewInitOptions.LODDistanceFactor = PlayerController->LocalPlayerCachedLODDistanceFactor;
-	ViewInitOptions.StereoPass = GEngine->StereoRenderingDevice->GetViewPassForIndex(StereoViewIndex != INDEX_NONE, StereoViewIndex);
 	ViewInitOptions.StereoViewIndex = StereoViewIndex;
 	ViewInitOptions.WorldToMetersScale = PlayerController->GetWorldSettings()->WorldToMeters;
 	ViewInitOptions.CursorPos = Viewport->HasMouseCapture() ? FIntPoint(-1, -1) : FIntPoint(Viewport->GetMouseX(), Viewport->GetMouseY());
