Add group index support to hair attribute node, and Strands/Cards vertex factory.

diff --git a/Engine/Plugins/Runtime/HairStrands/Source/HairStrandsCore/Private/HairCardsVertexFactory.cpp b/Engine/Plugins/Runtime/HairStrands/Source/HairStrandsCore/Private/HairCardsVertexFactory.cpp
--- a/Engine/Plugins/Runtime/HairStrands/Source/HairStrandsCore/Private/HairCardsVertexFactory.cpp
+++ b/Engine/Plugins/Runtime/HairStrands/Source/HairStrandsCore/Private/HairCardsVertexFactory.cpp
@@ -36,6 +36,7 @@ FHairCardsUniformBuffer CreateHairCardsVFUniformBuffer(
 		const FHairGroupInstance::FCards::FLOD& LOD = Instance->Cards.LODs[LODIndex];
 
 		// Cards atlas UV are inverted so fetching needs to be inverted on the y-axis
+		UniformParameters.GroupIndex = Instance->Debug.GroupIndex;
 		UniformParameters.bInvertUV = LOD.RestResource->bInvertUV;
 		UniformParameters.bUseTextureRootUV = 0;
 		UniformParameters.MaxVertexCount = LOD.RestResource->GetVertexCount();
@@ -88,6 +89,7 @@ FHairCardsUniformBuffer CreateHairCardsVFUniformBuffer(
 		}
 
 		// Meshes UV are not inverted so no need to invert the y-axis
+		UniformParameters.GroupIndex = Instance->Debug.GroupIndex;
 		UniformParameters.bInvertUV = 0;
 		UniformParameters.bUseTextureRootUV = 1;
 		UniformParameters.MaxVertexCount = LOD.RestResource->GetVertexCount();
diff --git a/Engine/Plugins/Runtime/HairStrands/Source/HairStrandsCore/Private/HairStrandsVertexFactory.cpp b/Engine/Plugins/Runtime/HairStrands/Source/HairStrandsCore/Private/HairStrandsVertexFactory.cpp
--- a/Engine/Plugins/Runtime/HairStrands/Source/HairStrandsCore/Private/HairStrandsVertexFactory.cpp
+++ b/Engine/Plugins/Runtime/HairStrands/Source/HairStrandsCore/Private/HairStrandsVertexFactory.cpp
@@ -69,6 +69,7 @@ FHairStrandsVertexFactoryUniformShaderParameters FHairGroupInstance::GetHairStan
 	const FHairGroupPublicData::FVertexFactoryInput VFInput = ComputeHairStrandsVertexInputData(this);
 
 	FHairStrandsVertexFactoryUniformShaderParameters Out = {};
+	Out.GroupIndex					= Debug.GroupIndex;
 	Out.Radius 						= VFInput.Strands.HairRadius;
 	Out.RootScale 					= VFInput.Strands.HairRootScale;
 	Out.TipScale 					= VFInput.Strands.HairTipScale;
diff --git a/Engine/Plugins/Runtime/HairStrands/Source/HairStrandsCore/Public/GroomInstance.h b/Engine/Plugins/Runtime/HairStrands/Source/HairStrandsCore/Public/GroomInstance.h
--- a/Engine/Plugins/Runtime/HairStrands/Source/HairStrandsCore/Public/GroomInstance.h
+++ b/Engine/Plugins/Runtime/HairStrands/Source/HairStrandsCore/Public/GroomInstance.h
@@ -22,6 +22,7 @@ BEGIN_GLOBAL_SHADER_PARAMETER_STRUCT(FHairCardsVertexFactoryUniformShaderParamet
 	SHADER_PARAMETER(uint32, bInvertUV)
 	SHADER_PARAMETER(uint32, bUseTextureRootUV)
 	SHADER_PARAMETER(uint32, MaxVertexCount)
+	SHADER_PARAMETER(float, GroupIndex)
 	SHADER_PARAMETER_SRV(Buffer<float4>, PositionBuffer)
 	SHADER_PARAMETER_SRV(Buffer<float4>, PreviousPositionBuffer)
 	SHADER_PARAMETER_SRV(Buffer<float4>, NormalsBuffer)
@@ -54,6 +55,7 @@ BEGIN_GLOBAL_SHADER_PARAMETER_STRUCT(FHairStrandsVertexFactoryUniformShaderParam
 	SHADER_PARAMETER(uint32, StableRasterization)
 	SHADER_PARAMETER(uint32, ScatterSceneLighing)
 	SHADER_PARAMETER(uint32, RaytracingProceduralSplits)
+	SHADER_PARAMETER(float, GroupIndex)
 
 	SHADER_PARAMETER_SRV(Buffer<float4>, PositionOffsetBuffer)
 	SHADER_PARAMETER_SRV(Buffer<float4>, PreviousPositionOffsetBuffer)
diff --git a/Engine/Shaders/Private/HairStrands/HairCardsAttributeCommon.ush b/Engine/Shaders/Private/HairStrands/HairCardsAttributeCommon.ush
--- a/Engine/Shaders/Private/HairStrands/HairCardsAttributeCommon.ush
+++ b/Engine/Shaders/Private/HairStrands/HairCardsAttributeCommon.ush
@@ -108,6 +108,10 @@ float GetHairStrandsRoughness(float2 HairPrimitiveUV)
 	return GetHairStrandsAttributes(HairPrimitiveUV).HairRoughness;
 }
 
+float GetHairStrandsGroupIndex(float2 InAtlasUV)
+{
+	return HairCardsVF.GroupIndex; 
+}
 ////////////////////////////////////////////////////////////////////////////////
 
 #endif // HAIR_CARD_MESH_FACTORY
\ No newline at end of file
diff --git a/Engine/Shaders/Private/HairStrands/HairStrandsAttributeCommon.ush b/Engine/Shaders/Private/HairStrands/HairStrandsAttributeCommon.ush
--- a/Engine/Shaders/Private/HairStrands/HairStrandsAttributeCommon.ush
+++ b/Engine/Shaders/Private/HairStrands/HairStrandsAttributeCommon.ush
@@ -109,6 +109,10 @@ float GetHairStrandsDepth(uint HairPrimitiveId, float2 InUV, float InDeviceZ)
 	return 0.f;
 }
 
+float GetHairStrandsGroupIndex(float2 InAtlasUV)
+{
+	return HairStrandsVF.GroupIndex;
+}
 ////////////////////////////////////////////////////////////////////////////////
 
 #endif // HAIR_STRAND_MESH_FACTORY
\ No newline at end of file
diff --git a/Engine/Shaders/Private/MaterialTemplate.ush b/Engine/Shaders/Private/MaterialTemplate.ush
--- a/Engine/Shaders/Private/MaterialTemplate.ush
+++ b/Engine/Shaders/Private/MaterialTemplate.ush
@@ -1397,6 +1397,17 @@ float4 MaterialExpressionGetHairAuxilaryData(FMaterialPixelParameters Parameters
 #endif
 }
 
+float MaterialExpressionGetHairGroupIndex(FMaterialPixelParameters Parameters)
+{
+#if HAIR_STRAND_MESH_FACTORY
+	return GetHairStrandsGroupIndex(Parameters.HairPrimitiveId);
+#elif HAIR_CARD_MESH_FACTORY
+	return GetHairStrandsGroupIndex(Parameters.HairPrimitiveUV);
+#else
+	return 0;
+#endif
+
+}
 float3 MaterialExpressionGetHairColorFromMelanin(float Melanin, float Redness, float3 DyeColor)
 {
 	return GetHairColorFromMelanin(Melanin, Redness, DyeColor);
diff --git a/Engine/Source/Runtime/Engine/Private/Materials/HLSLMaterialTranslator.cpp b/Engine/Source/Runtime/Engine/Private/Materials/HLSLMaterialTranslator.cpp
--- a/Engine/Source/Runtime/Engine/Private/Materials/HLSLMaterialTranslator.cpp
+++ b/Engine/Source/Runtime/Engine/Private/Materials/HLSLMaterialTranslator.cpp
@@ -8909,6 +8909,11 @@ int32 FHLSLMaterialTranslator::GetHairAtlasUVs()
 	return AddCodeChunk(MCT_Float2, TEXT("MaterialExpressionGetAtlasUVs(Parameters)"));
 }
 
+int32 FHLSLMaterialTranslator::GetHairGroupIndex()
+{
+	return AddCodeChunk(MCT_Float1, TEXT("MaterialExpressionGetHairGroupIndex(Parameters)"));
+}
+
 int32 FHLSLMaterialTranslator::GetHairColorFromMelanin(int32 Melanin, int32 Redness, int32 DyeColor)
 {
 	if (Melanin == INDEX_NONE)
diff --git a/Engine/Source/Runtime/Engine/Private/Materials/HLSLMaterialTranslator.h b/Engine/Source/Runtime/Engine/Private/Materials/HLSLMaterialTranslator.h
--- a/Engine/Source/Runtime/Engine/Private/Materials/HLSLMaterialTranslator.h
+++ b/Engine/Source/Runtime/Engine/Private/Materials/HLSLMaterialTranslator.h
@@ -946,6 +946,7 @@ protected:
 	virtual int32 GetHairCoverage() override;
 	virtual int32 GetHairAuxilaryData() override;
 	virtual int32 GetHairAtlasUVs() override;
+	virtual int32 GetHairGroupIndex() override;
 	virtual int32 GetHairColorFromMelanin(int32 Melanin, int32 Redness, int32 DyeColor) override;
 	virtual int32 DistanceToNearestSurface(int32 PositionArg) override;
 	virtual int32 DistanceFieldGradient(int32 PositionArg) override;
diff --git a/Engine/Source/Runtime/Engine/Private/Materials/MaterialExpressions.cpp b/Engine/Source/Runtime/Engine/Private/Materials/MaterialExpressions.cpp
--- a/Engine/Source/Runtime/Engine/Private/Materials/MaterialExpressions.cpp
+++ b/Engine/Source/Runtime/Engine/Private/Materials/MaterialExpressions.cpp
@@ -19302,6 +19302,7 @@ UMaterialExpressionHairAttributes::UMaterialExpressionHairAttributes(const FObje
 	Outputs.Add(FExpressionOutput(TEXT("Coverage"), 1, 1, 0, 0, 0));
 	Outputs.Add(FExpressionOutput(TEXT("AuxilaryData"), 1, 1, 1, 1, 1));
 	Outputs.Add(FExpressionOutput(TEXT("AtlasUVs"), 1, 1, 1, 0, 0));
+	Outputs.Add(FExpressionOutput(TEXT("Group Index"), 1, 1, 0, 0, 0));
 #endif
 }
 
@@ -19352,6 +19353,10 @@ int32 UMaterialExpressionHairAttributes::Compile(class FMaterialCompiler* Compil
 	{
 		return Compiler->GetHairAtlasUVs();
 	}
+	else if (OutputIndex == 13)
+	{
+		return Compiler->GetHairGroupIndex();
+	}
 
 	return Compiler->Errorf(TEXT("Invalid input parameter"));
 }
diff --git a/Engine/Source/Runtime/Engine/Public/MaterialCompiler.h b/Engine/Source/Runtime/Engine/Public/MaterialCompiler.h
--- a/Engine/Source/Runtime/Engine/Public/MaterialCompiler.h
+++ b/Engine/Source/Runtime/Engine/Public/MaterialCompiler.h
@@ -420,6 +420,7 @@ public:
 	virtual int32 GetHairCoverage() = 0;
 	virtual int32 GetHairAuxilaryData() = 0;
 	virtual int32 GetHairAtlasUVs() = 0;
+	virtual int32 GetHairGroupIndex() = 0;
 	virtual int32 GetHairColorFromMelanin(int32 Melanin, int32 Redness, int32 DyeColor) = 0;
 	virtual int32 CustomPrimitiveData(int32 OutputIndex, EMaterialValueType Type) = 0;
 	virtual int32 ShadingModel(EMaterialShadingModel InSelectedShadingModel) = 0;
@@ -777,6 +778,7 @@ public:
 	virtual int32 GetHairCoverage() override { return Compiler->GetHairCoverage(); }
 	virtual int32 GetHairAuxilaryData() override { return Compiler->GetHairAuxilaryData(); }
 	virtual int32 GetHairAtlasUVs() override { return Compiler->GetHairAtlasUVs(); }
+	virtual int32 GetHairGroupIndex() override { return Compiler->GetHairGroupIndex(); }
 	virtual int32 GetHairColorFromMelanin(int32 Melanin, int32 Redness, int32 DyeColor) override { return Compiler->GetHairColorFromMelanin(Melanin, Redness, DyeColor); }
 
 	virtual int32 RotateScaleOffsetTexCoords(int32 TexCoordCodeIndex, int32 RotationScale, int32 Offset) override
