Fix crash when skeletal mesh has Capsule Direct Shadow & Render Static on. Cause is that FSkeletalMeshObjectStatic contains an empty ReferenceToLocalMatrices array.

diff --git a/Engine/Source/Runtime/Engine/Private/SkeletalMesh.cpp b/Engine/Source/Runtime/Engine/Private/SkeletalMesh.cpp
--- a/Engine/Source/Runtime/Engine/Private/SkeletalMesh.cpp
+++ b/Engine/Source/Runtime/Engine/Private/SkeletalMesh.cpp
@@ -6648,7 +6648,11 @@ void FSkeletalMeshSceneProxy::GetShadowShapes(TArray<FCapsuleShape3f>& CapsuleSh
 
 	for(const TPair<int32, FCapsuleShape>& CapsuleData : ShadowCapsuleData)
 	{
-		FMatrix ReferenceToWorld = ReferenceToLocalMatrices[CapsuleData.Key] * ProxyLocalToWorld;
+		FMatrix ReferenceToWorld = ProxyLocalToWorld;
+		if (ReferenceToLocalMatrices.IsValidIndex(CapsuleData.Key))
+		{
+			ReferenceToWorld = ReferenceToLocalMatrices[CapsuleData.Key] * ProxyLocalToWorld;
+		}
 		const float MaxScale = ReferenceToWorld.GetScaleVector().GetMax();
 
 		FCapsuleShape3f& NewCapsule = CapsuleShapes[CapsuleIndex++];
