Fixed CIS Issue 116415 (static analysis)

diff --git a/Engine/Source/Runtime/Engine/Private/VisualLogger/VisualLogger.cpp b/Engine/Source/Runtime/Engine/Private/VisualLogger/VisualLogger.cpp
--- a/Engine/Source/Runtime/Engine/Private/VisualLogger/VisualLogger.cpp
+++ b/Engine/Source/Runtime/Engine/Private/VisualLogger/VisualLogger.cpp
@@ -217,6 +217,7 @@ FVisualLogEntry* FVisualLogger::GetEntryToWrite(const UObject* Object, const flo
 
 	if (bInitializeEntry)
 	{
+		checkf(CurrentEntry != nullptr, TEXT("bInitializeEntry can only be true when CurrentEntry is valid."));
 		CurrentEntry->Reset();
 		CurrentEntry->TimeStamp = TimeStamp;
 		CurrentEntry->bIsInitialized = true;
