Use FName::LexicalLess to sort names instead of default FName::operator<, which is deprecated.

diff --git a/Engine/Source/Runtime/Engine/Private/WorldPartition/WorldPartitionActorDesc.cpp b/Engine/Source/Runtime/Engine/Private/WorldPartition/WorldPartitionActorDesc.cpp
--- a/Engine/Source/Runtime/Engine/Private/WorldPartition/WorldPartitionActorDesc.cpp
+++ b/Engine/Source/Runtime/Engine/Private/WorldPartition/WorldPartitionActorDesc.cpp
@@ -139,8 +139,8 @@ bool FWorldPartitionActorDesc::Equals(const FWorldPartitionActorDesc* Other) con
 	{
 		TArray<FName> SortedDataLayers(DataLayers);
 		TArray<FName> SortedDataLayersOther(Other->DataLayers);
-		SortedDataLayers.Sort();
-		SortedDataLayersOther.Sort();
+		SortedDataLayers.Sort([](const FName& LHS, const FName& RHS) { return LHS.LexicalLess(RHS); });
+		SortedDataLayersOther.Sort([](const FName& LHS, const FName& RHS) { return LHS.LexicalLess(RHS); });
 
 		TArray<FGuid> SortedReferences(References);
 		TArray<FGuid> SortedReferencesOther(Other->References);
