LevelInstance:
-Remove multi-edit
-Cleanup Edit/Commit

diff --git a/Engine/Source/Editor/LevelInstanceEditor/Private/LevelInstanceEditorModule.cpp b/Engine/Source/Editor/LevelInstanceEditor/Private/LevelInstanceEditorModule.cpp
--- a/Engine/Source/Editor/LevelInstanceEditor/Private/LevelInstanceEditorModule.cpp
+++ b/Engine/Source/Editor/LevelInstanceEditor/Private/LevelInstanceEditorModule.cpp
@@ -76,65 +76,7 @@ namespace LevelInstanceMenuUtils
 			Section.AddMenuEntry(NAME_None, EntryLabel, EntryDesc, FSlateIcon(), LevelInstanceEditAction);
 		}
 	}
-
-	void CreateCommitSubMenu(UToolMenu* Menu, TArray<ALevelInstance*> LevelInstanceEdits, ALevelInstance* ContextLevelInstance, bool bDiscard)
-	{
-		FText OtherSectionLabel = LOCTEXT("LevelInstanceOtherCommitSection", "Other(s)");
-		FToolMenuSection* Section = &Menu->AddSection("LevelInstanceContextCommitSection", ContextLevelInstance != nullptr? LOCTEXT("LevelInstanceContextCommitSection", "Context") : OtherSectionLabel);
-		for (ALevelInstance* LevelInstanceActor : LevelInstanceEdits)
-		{
-			FText EntryDesc = LOCTEXT("LevelInstanceEditSubMenuEntry", "");
-			const bool bCanCommit = LevelInstanceActor->CanCommit(&EntryDesc);
-
-			FToolUIAction LevelInstanceEditAction;
-			LevelInstanceEditAction.ExecuteAction.BindLambda([bDiscard, LevelInstanceActor](const FToolMenuContext&)
-			{
-				if (bDiscard)
-				{
-					LevelInstanceActor->Discard();
-				}
-				else
-				{
-					LevelInstanceActor->Commit();
-				}
-			});
-			LevelInstanceEditAction.CanExecuteAction.BindLambda([bCanCommit](const FToolMenuContext&)
-			{
-				return bCanCommit;
-			});
-
-			const FText EntryLabel = FText::Format(LOCTEXT("LevelInstanceName", "{0}:{1}"), FText::FromString(LevelInstanceActor->GetActorLabel()), FText::FromString(LevelInstanceActor->GetWorldAssetPackage()));
-			Section->AddMenuEntry(NAME_None, EntryLabel, EntryDesc, FSlateIcon(), LevelInstanceEditAction);
-
-			if (ContextLevelInstance == LevelInstanceActor && LevelInstanceEdits.Num() > 1)
-			{
-				Section = &Menu->AddSection(FName("LevelInstanceOtherCommitSection"), OtherSectionLabel);
-			}
-		}
-	}
-
-	void CreateSetCurrentSubMenu(UToolMenu* Menu, TArray<ALevelInstance*> LevelInstanceEdits, ALevelInstance* ContextLevelInstance)
-	{
-		FText OtherSectionLabel = LOCTEXT("LevelInstanceOtherSetCurrentSection", "Other(s)");
-		FToolMenuSection* Section = &Menu->AddSection(FName("LevelInstanceContextSetCurrentSection"), ContextLevelInstance != nullptr? LOCTEXT("LevelInstanceContextSetCurrentSection", "Context") : OtherSectionLabel);
-		for (ALevelInstance* LevelInstanceActor : LevelInstanceEdits)
-		{
-			FToolUIAction LevelInstanceSetCurrentAction;
-			LevelInstanceSetCurrentAction.ExecuteAction.BindLambda([LevelInstanceActor](const FToolMenuContext&)
-			{
-				LevelInstanceActor->SetCurrent();
-			});
 		
-			const FText EntryLabel = FText::Format(LOCTEXT("LevelInstanceName", "{0}:{1}"), FText::FromString(LevelInstanceActor->GetActorLabel()), FText::FromString(LevelInstanceActor->GetWorldAssetPackage()));
-			Section->AddMenuEntry(NAME_None, EntryLabel, TAttribute<FText>(), FSlateIcon(), LevelInstanceSetCurrentAction);
-
-			if (LevelInstanceActor == ContextLevelInstance && LevelInstanceEdits.Num() > 1)
-			{
-				Section = &Menu->AddSection(FName("LevelInstanceOtherSetCurrentSection"), OtherSectionLabel);
-			}
-		}
-	}
-
 	void MoveSelectionToLevelInstance(ALevelInstance* DestinationLevelInstance)
 	{
 		if (ULevelInstanceSubsystem* LevelInstanceSubsystem = DestinationLevelInstance->GetLevelInstanceSubsystem())
@@ -152,43 +94,7 @@ namespace LevelInstanceMenuUtils
 			LevelInstanceSubsystem->MoveActorsTo(DestinationLevelInstance, ActorsToMove);
 		}
 	}
-
-	void CreateMoveSelectionToSubMenu(UToolMenu* Menu, TArray<ALevelInstance*> LevelInstanceEdits)
-	{
-		FToolMenuSection* Section = &Menu->AddSection(NAME_None);
-		if (ULevelInstanceSubsystem* LevelInstanceSubsystem = GEditor->GetEditorWorldContext().World()->GetSubsystem<ULevelInstanceSubsystem>())
-		{
-			for (ALevelInstance* LevelInstanceActor : LevelInstanceEdits)
-			{
-				FToolUIAction LevelInstanceMoveSelectionAction;
-
-				LevelInstanceMoveSelectionAction.CanExecuteAction.BindLambda([LevelInstanceActor, LevelInstanceSubsystem](const FToolMenuContext& MenuContext)
-					{
-						for (FSelectionIterator It(GEditor->GetSelectedActorIterator()); It; ++It)
-						{
-							if (AActor* Actor = Cast<AActor>(*It))
-							{
-								if (Actor->GetLevel() == LevelInstanceSubsystem->GetLevelInstanceLevel(LevelInstanceActor))
-								{
-									return false;
-								}
-							}
-						}
-					
-						return GEditor->GetSelectedActorCount() > 0;
-					});
-
-				LevelInstanceMoveSelectionAction.ExecuteAction.BindLambda([LevelInstanceActor](const FToolMenuContext&)
-					{
-						MoveSelectionToLevelInstance(LevelInstanceActor);
-					});
-
-				const FText EntryLabel = FText::Format(LOCTEXT("LevelInstanceName", "{0}:{1}"), FText::FromString(LevelInstanceActor->GetActorLabel()), FText::FromString(LevelInstanceActor->GetWorldAssetPackage()));
-				Section->AddMenuEntry(NAME_None, EntryLabel, TAttribute<FText>(), FSlateIcon(), LevelInstanceMoveSelectionAction);
-			}
-		}
-	}
-
+		
 	void CreateEditMenu(UToolMenu* Menu, AActor* ContextActor)
 	{
 		if (ULevelInstanceSubsystem* LevelInstanceSubsystem = ContextActor->GetWorld()->GetSubsystem<ULevelInstanceSubsystem>())
@@ -212,103 +118,86 @@ namespace LevelInstanceMenuUtils
 			}
 		}
 	}
-	
+		
 	void CreateCommitMenu(UToolMenu* Menu, AActor* ContextActor)
 	{
-		ALevelInstance* ContextLevelInstance = nullptr;
-		TArray<ALevelInstance*> LevelInstanceEdits;
+		ALevelInstance* LevelInstanceEdit = nullptr;
 		if (ContextActor)
 		{
 			if (ULevelInstanceSubsystem* LevelInstanceSubsystem = ContextActor->GetWorld()->GetSubsystem<ULevelInstanceSubsystem>())
 			{
-				LevelInstanceSubsystem->ForEachLevelInstanceAncestorsAndSelf(ContextActor, [&ContextLevelInstance,&LevelInstanceEdits](ALevelInstance* LevelInstanceActor)
-				{
-					if (LevelInstanceActor->IsEditing())
-					{
-						ContextLevelInstance = LevelInstanceActor;
-						LevelInstanceEdits.Add(ContextLevelInstance);
-						return false;
-					}
-					return true;
-				});
+				LevelInstanceEdit = LevelInstanceSubsystem->GetEditingLevelInstance();
 			}
 		}
 
-		if (ULevelInstanceSubsystem* LevelInstanceSubsystem = GEditor->GetEditorWorldContext().World()->GetSubsystem<ULevelInstanceSubsystem>())
+		if (!LevelInstanceEdit)
 		{
-			LevelInstanceSubsystem->ForEachLevelInstanceEdit([&LevelInstanceEdits, ContextLevelInstance](ALevelInstance* LevelInstanceActor)
+			if (ULevelInstanceSubsystem* LevelInstanceSubsystem = GEditor->GetEditorWorldContext().World()->GetSubsystem<ULevelInstanceSubsystem>())
 			{
-				if (ContextLevelInstance != LevelInstanceActor)
-				{
-					LevelInstanceEdits.Add(LevelInstanceActor);
-				}
-				return true;
-			});
+				LevelInstanceEdit = LevelInstanceSubsystem->GetEditingLevelInstance();
+			}
 		}
 
-		if (LevelInstanceEdits.Num() > 0)
+		if (LevelInstanceEdit)
 		{
-			FToolMenuSection& Section = CreateLevelInstanceSection(Menu);
-			Section.AddSubMenu(
-				"CommitLevelInstances",
-				LOCTEXT("CommitLevelInstances", "Commit"),
-				TAttribute<FText>(),
-				FNewToolMenuDelegate::CreateStatic(&CreateCommitSubMenu, LevelInstanceEdits, ContextLevelInstance, /*bDiscard=*/ false)
-			);
-			Section.AddSubMenu(
-				"DiscardLevelInstances",
-				LOCTEXT("DiscardLevelInstances", "Discard"),
-				TAttribute<FText>(),
-				FNewToolMenuDelegate::CreateStatic(&CreateCommitSubMenu, LevelInstanceEdits, ContextLevelInstance, /*bDiscard=*/ true)
-			);
-		}
-	}
+			FText Tooltip;
+			const bool bCanCommit = LevelInstanceEdit->CanCommit(&Tooltip);
 
-	void CreateSetCurrentMenu(UToolMenu* Menu, AActor* ContextActor)
-	{
-		ALevelInstance* ContextLevelInstance = nullptr;
-		TArray<ALevelInstance*> LevelInstanceEdits;
-		if (ContextActor)
-		{
-			if (ULevelInstanceSubsystem* LevelInstanceSubsystem = ContextActor->GetWorld()->GetSubsystem<ULevelInstanceSubsystem>())
+			auto CreateMenuEntry = [Menu, LevelInstanceEdit, Tooltip](bool bDiscard, bool bCanCommit)
 			{
-				LevelInstanceSubsystem->ForEachLevelInstanceAncestorsAndSelf(ContextActor, [&ContextLevelInstance,&LevelInstanceEdits](ALevelInstance* LevelInstanceActor)
+				FToolUIAction MenuAction;
+
+				MenuAction.ExecuteAction.BindLambda([bDiscard, LevelInstanceEdit](const FToolMenuContext&)
 				{
-					if (LevelInstanceActor->IsEditing())
+					if (bDiscard)
 					{
-						if (!LevelInstanceActor->IsCurrent())
-						{
-							ContextLevelInstance = LevelInstanceActor;
-							LevelInstanceEdits.Add(ContextLevelInstance);
-						}
-						return false;
+						LevelInstanceEdit->Discard();
+					}
+					else
+					{
+						LevelInstanceEdit->Commit();
 					}
-					return true;
 				});
-			}
+				MenuAction.CanExecuteAction.BindLambda([bCanCommit, LevelInstanceEdit](const FToolMenuContext&)
+				{
+					return bCanCommit;
+				});
+
+				FToolMenuSection& Section = CreateLevelInstanceSection(Menu);
+				Section.AddMenuEntry(NAME_None, bDiscard ? LOCTEXT("LevelInstanceDiscardLabel", "Discard") : LOCTEXT("LevelInstanceCommitLabel", "Commit"), Tooltip, FSlateIcon(), MenuAction);
+			};
+			
+			CreateMenuEntry(/*bDiscard=*/ false, bCanCommit);
+			CreateMenuEntry(/*bDiscard=*/ true, bCanCommit);
 		}
+	}
 
-		if (ULevelInstanceSubsystem* LevelInstanceSubsystem = GEditor->GetEditorWorldContext().World()->GetSubsystem<ULevelInstanceSubsystem>())
+	void CreateSetCurrentMenu(UToolMenu* Menu, AActor* ContextActor)
+	{
+		ALevelInstance* LevelInstanceEdit = nullptr;
+		if (ULevelInstanceSubsystem* LevelInstanceSubsystem = ContextActor->GetWorld()->GetSubsystem<ULevelInstanceSubsystem>())
 		{
-			LevelInstanceSubsystem->ForEachLevelInstanceEdit([&LevelInstanceEdits, ContextLevelInstance](ALevelInstance* LevelInstanceActor)
+			if (ContextActor)
 			{
-				if (ContextLevelInstance != LevelInstanceActor && !LevelInstanceActor->IsCurrent())
-				{
-					LevelInstanceEdits.Add(LevelInstanceActor);
-				}
-				return true;
-			});
-		}
+				LevelInstanceEdit = LevelInstanceSubsystem->GetEditingLevelInstance();
+			}
 
-		if (LevelInstanceEdits.Num() > 0)
-		{
-			FToolMenuSection& Section = CreateLevelInstanceSection(Menu);
-			Section.AddSubMenu(
-				"SetCurrentLevelInstances",
-				LOCTEXT("SetCurrentLevelInstances", "Set Current"),
-				TAttribute<FText>(),
-				FNewToolMenuDelegate::CreateStatic(&CreateSetCurrentSubMenu, MoveTemp(LevelInstanceEdits), ContextLevelInstance)
-			);
+			if (!LevelInstanceEdit)
+			{
+				LevelInstanceEdit = LevelInstanceSubsystem->GetEditingLevelInstance();
+			}
+
+			if (LevelInstanceEdit)
+			{
+				FToolUIAction LevelInstanceSetCurrentAction;
+				LevelInstanceSetCurrentAction.ExecuteAction.BindLambda([LevelInstanceEdit](const FToolMenuContext&)
+					{
+						LevelInstanceEdit->SetCurrent();
+					});
+
+				FToolMenuSection& Section = CreateLevelInstanceSection(Menu);
+				Section.AddMenuEntry(NAME_None, LOCTEXT("LevelInstanceSetCurrent", "Set Current Level"), TAttribute<FText>(), FSlateIcon(), LevelInstanceSetCurrentAction);
+			}
 		}
 	}
 
@@ -316,26 +205,41 @@ namespace LevelInstanceMenuUtils
 	{
 		if (GEditor->GetSelectedActorCount() > 0)
 		{
-
-			TArray<ALevelInstance*> LevelInstanceEdits;
-			if (ULevelInstanceSubsystem* LevelInstanceSubsystem = GEditor->GetEditorWorldContext().World()->GetSubsystem<ULevelInstanceSubsystem>())
+			ALevelInstance* LevelInstanceEdit = nullptr;
+			ULevelInstanceSubsystem* LevelInstanceSubsystem = GEditor->GetEditorWorldContext().World()->GetSubsystem<ULevelInstanceSubsystem>();
+			if (LevelInstanceSubsystem)
 			{
-				LevelInstanceSubsystem->ForEachLevelInstanceEdit([&LevelInstanceEdits](ALevelInstance* LevelInstanceActor)
-				{
-					LevelInstanceEdits.Add(LevelInstanceActor);
-					return true;
-				});
+				LevelInstanceEdit = LevelInstanceSubsystem->GetEditingLevelInstance();
 			}
-
-			if (LevelInstanceEdits.Num() > 0)
+			
+			if (LevelInstanceEdit)
 			{
+				FToolUIAction LevelInstanceMoveSelectionAction;
+
+				LevelInstanceMoveSelectionAction.CanExecuteAction.BindLambda([LevelInstanceEdit, LevelInstanceSubsystem](const FToolMenuContext& MenuContext)
+					{
+						for (FSelectionIterator It(GEditor->GetSelectedActorIterator()); It; ++It)
+						{
+							if (AActor* Actor = Cast<AActor>(*It))
+							{
+								if (Actor->GetLevel() == LevelInstanceSubsystem->GetLevelInstanceLevel(LevelInstanceEdit))
+								{
+									return false;
+								}
+							}
+						}
+
+						return GEditor->GetSelectedActorCount() > 0;
+					});
+
+				LevelInstanceMoveSelectionAction.ExecuteAction.BindLambda([LevelInstanceEdit](const FToolMenuContext&)
+					{
+						MoveSelectionToLevelInstance(LevelInstanceEdit);
+					});
+
 				FToolMenuSection& Section = CreateLevelInstanceSection(Menu);
-				Section.AddSubMenu(
-					"MoveSelectionToLevelInstances",
-					LOCTEXT("MoveSelectionToLevelInstances", "Move Selection to"),
-					TAttribute<FText>(),
-					FNewToolMenuDelegate::CreateStatic(&CreateMoveSelectionToSubMenu, MoveTemp(LevelInstanceEdits))
-				);
+				Section.AddMenuEntry(NAME_None, LOCTEXT("LevelInstanceMoveSelectionTo", "Move Selection to"), TAttribute<FText>(), FSlateIcon(), LevelInstanceMoveSelectionAction);
+
 			}
 		}
 	}
diff --git a/Engine/Source/Runtime/Engine/Private/LevelInstance/LevelInstanceActor.cpp b/Engine/Source/Runtime/Engine/Private/LevelInstance/LevelInstanceActor.cpp
--- a/Engine/Source/Runtime/Engine/Private/LevelInstance/LevelInstanceActor.cpp
+++ b/Engine/Source/Runtime/Engine/Private/LevelInstance/LevelInstanceActor.cpp
@@ -23,7 +23,6 @@ ALevelInstance::ALevelInstance(const FObjectInitializer& ObjectInitializer)
 	: Super(ObjectInitializer)
 #if WITH_EDITOR
 	, bGuardLoadUnload(false)
-	, bEditLockLocation(false)
 #endif
 {
 	RootComponent = CreateDefaultSubobject<ULevelInstanceComponent>(TEXT("Root"));
@@ -225,16 +224,6 @@ AActor* ALevelInstance::FindEditorInstanceActor() const
 	return FoundActor;
 }
 
-void ALevelInstance::OnEdit()
-{
-	bEditLockLocation = true;
-}
-
-void ALevelInstance::OnCommit(bool bChanged, bool bPromptForSave)
-{
-	bEditLockLocation = false;
-}
-
 ALevelInstance::FOnLevelInstanceActorPostLoad ALevelInstance::OnLevelInstanceActorPostLoad;
 
 void ALevelInstance::PostLoad()
@@ -554,7 +543,7 @@ bool ALevelInstance::CanDeleteSelectedActor(FText& OutReason) const
 		return false;
 	}
 
-	if (HasEditingChildren())
+	if (HasChildEdit())
 	{
 		OutReason = LOCTEXT("HasEditingChildLevel", "Can't delete LevelInstance because it has editing child LevelInstances!");
 		return false;
@@ -680,7 +669,7 @@ FBox ALevelInstance::GetStreamingBounds() const
 
 bool ALevelInstance::IsLockLocation() const
 {
-	return Super::IsLockLocation() || bEditLockLocation;
+	return Super::IsLockLocation() || IsEditing() || HasChildEdit();
 }
 
 FBox ALevelInstance::GetComponentsBoundingBox(bool bNonColliding, bool bIncludeFromChildActors) const
@@ -705,9 +694,12 @@ FBox ALevelInstance::GetComponentsBoundingBox(bool bNonColliding, bool bIncludeF
 
 bool ALevelInstance::CanEdit(FText* OutReason) const
 {
-	if (ULevelInstanceSubsystem* LevelInstanceSubsystem = GetLevelInstanceSubsystem())
+	if (HasValidLevelInstanceID())
 	{
-		return LevelInstanceSubsystem->CanEditLevelInstance(this, OutReason);
+		if (ULevelInstanceSubsystem* LevelInstanceSubsystem = GetLevelInstanceSubsystem())
+		{
+			return LevelInstanceSubsystem->CanEditLevelInstance(this, OutReason);
+		}
 	}
 
 	return false;
@@ -715,9 +707,12 @@ bool ALevelInstance::CanEdit(FText* OutReason) const
 
 bool ALevelInstance::CanCommit(FText* OutReason) const
 {
-	if (ULevelInstanceSubsystem* LevelInstanceSubsystem = GetLevelInstanceSubsystem())
+	if (HasValidLevelInstanceID())
 	{
-		return LevelInstanceSubsystem->CanCommitLevelInstance(this, OutReason);
+		if (ULevelInstanceSubsystem* LevelInstanceSubsystem = GetLevelInstanceSubsystem())
+		{
+			return LevelInstanceSubsystem->CanCommitLevelInstance(this, OutReason);
+		}
 	}
 
 	return false;
@@ -725,19 +720,24 @@ bool ALevelInstance::CanCommit(FText* OutReason) const
 
 bool ALevelInstance::IsEditing() const
 {
-	if (ULevelInstanceSubsystem* LevelInstanceSubsystem = GetLevelInstanceSubsystem())
+	if (HasValidLevelInstanceID())
 	{
-		return LevelInstanceSubsystem->IsEditingLevelInstance(this);
+		if (ULevelInstanceSubsystem* LevelInstanceSubsystem = GetLevelInstanceSubsystem())
+		{
+			return LevelInstanceSubsystem->IsEditingLevelInstance(this);
+		}
 	}
-
 	return false;
 }
 
-bool ALevelInstance::HasEditingChildren() const
+bool ALevelInstance::HasChildEdit() const
 {
-	if (ULevelInstanceSubsystem* LevelInstanceSubsystem = GetLevelInstanceSubsystem())
+	if (HasValidLevelInstanceID())
 	{
-		return LevelInstanceSubsystem->HasEditingChildrenLevelInstances(this);
+		if (ULevelInstanceSubsystem* LevelInstanceSubsystem = GetLevelInstanceSubsystem())
+		{
+			return LevelInstanceSubsystem->HasChildEdit(this);
+		}
 	}
 
 	return false;
diff --git a/Engine/Source/Runtime/Engine/Private/LevelInstance/LevelInstanceSubsystem.cpp b/Engine/Source/Runtime/Engine/Private/LevelInstance/LevelInstanceSubsystem.cpp
--- a/Engine/Source/Runtime/Engine/Private/LevelInstance/LevelInstanceSubsystem.cpp
+++ b/Engine/Source/Runtime/Engine/Private/LevelInstance/LevelInstanceSubsystem.cpp
@@ -422,28 +422,18 @@ void ULevelInstanceSubsystem::Tick()
 	if (!GetWorld()->IsGameWorld())
 	{
 		UpdateStreamingState();
-
-		// Begin editing the pending LevelInstance when loads complete
-		if (PendingLevelInstanceToEdit.IsValid() && !LevelInstancesToLoadOrUpdate.Num())
-		{
-			if (ALevelInstance** LevelInstanceActor = RegisteredLevelInstances.Find(PendingLevelInstanceToEdit))
-			{
-				EditLevelInstance(*LevelInstanceActor);
-			}
-		}
 	}
 }
 
 void ULevelInstanceSubsystem::OnPackageChanged(UPackage* Package)
 {
-	for (auto& Pair : LevelInstanceEdits)
+	if (LevelInstanceEdit)
 	{
-		FLevelInstanceEdit& LevelInstanceEdit = Pair.Value;
-		if (UWorld* EditWorld = LevelInstanceEdit.GetEditWorld())
+		if (UWorld* EditWorld = LevelInstanceEdit->GetEditWorld())
 		{
 			if (EditWorld->GetPackage() == Package)
 			{
-				LevelInstanceEdit.bCommittedChanges = true;
+				LevelInstanceEdit->bCommittedChanges = true;
 			}
 			else
 			{
@@ -451,7 +441,7 @@ void ULevelInstanceSubsystem::OnPackageChanged(UPackage* Package)
 				Packages.Append(EditWorld->GetPackage()->GetExternalPackages());
 				if (Packages.Contains(Package))
 				{
-					LevelInstanceEdit.bCommittedChanges = true;
+					LevelInstanceEdit->bCommittedChanges = true;
 				}
 			}
 		}
@@ -473,7 +463,7 @@ void ULevelInstanceSubsystem::OnPackageDeleted(UPackage* Package)
 
 bool ULevelInstanceSubsystem::CanPackLevelInstances() const
 {
-	return !LevelInstanceEdits.Num();
+	return !LevelInstanceEdit;
 }
 
 void ULevelInstanceSubsystem::PackLevelInstances()
@@ -518,12 +508,15 @@ void ULevelInstanceSubsystem::PackLevelInstances()
 
 	TArray<APackedLevelInstance*> PackedLevelInstancesToUpdate;
 	TArray<UBlueprint*> BlueprintsToUpdate;
-	for (TObjectIterator<UWorld> It(RF_ClassDefaultObject | RF_ArchetypeObject, true, EInternalObjectFlags::PendingKill); It; ++It)
+	for (TObjectIterator<UWorld> It(RF_ClassDefaultObject | RF_ArchetypeObject, true); It; ++It)
 	{
 		UWorld* CurrentWorld = *It;
-		for (TActorIterator<APackedLevelInstance> LevelInstanceIt(GetWorld()); LevelInstanceIt; ++LevelInstanceIt)
+		if (IsValid(CurrentWorld) && CurrentWorld->GetSubsystem<ULevelInstanceSubsystem>() != nullptr)
 		{
-			GatherDepencenciesRecursive(*LevelInstanceIt, BlueprintsToUpdate, PackedLevelInstancesToUpdate);
+			for (TActorIterator<APackedLevelInstance> LevelInstanceIt(CurrentWorld); LevelInstanceIt; ++LevelInstanceIt)
+			{
+				GatherDepencenciesRecursive(*LevelInstanceIt, BlueprintsToUpdate, PackedLevelInstancesToUpdate);
+			}
 		}
 	}
 
@@ -569,9 +562,9 @@ bool ULevelInstanceSubsystem::GetLevelInstanceBounds(const ALevelInstance* Level
 		OutBounds = LevelInstance.LevelStreaming->GetBounds();
 		return true;
 	}
-	else if (const FLevelInstanceEdit* LevelInstanceEdit = GetLevelInstanceEdit(LevelInstanceActor))
+	else if (const FLevelInstanceEdit* CurrentEdit = GetLevelInstanceEdit(LevelInstanceActor))
 	{
-		OutBounds = LevelInstanceEdit->LevelStreaming->GetBounds();
+		OutBounds = CurrentEdit->LevelStreaming->GetBounds();
 		return true;
 	}
 	else if(LevelInstanceActor->IsLevelInstancePathValid())
@@ -690,17 +683,6 @@ bool ULevelInstanceSubsystem::ForEachLevelInstanceChildImpl(ALevelInstance* Leve
 	return bContinue;
 }
 
-void ULevelInstanceSubsystem::ForEachLevelInstanceEdit(TFunctionRef<bool(ALevelInstance*)> Operation) const
-{
-	for (const auto& Pair : LevelInstanceEdits)
-	{
-		if (!Operation(Pair.Value.LevelStreaming->GetLevelInstanceActor()))
-		{
-			return;
-		}
-	}
-}
-
 bool ULevelInstanceSubsystem::HasDirtyChildrenLevelInstances(const ALevelInstance* LevelInstanceActor) const
 {
 	bool bDirtyChildren = false;
@@ -716,21 +698,6 @@ bool ULevelInstanceSubsystem::HasDirtyChildrenLevelInstances(const ALevelInstanc
 	return bDirtyChildren;
 }
 
-bool ULevelInstanceSubsystem::HasEditingChildrenLevelInstances(const ALevelInstance* LevelInstanceActor) const
-{
-	bool bEditingChildren = false;
-	ForEachLevelInstanceChild(LevelInstanceActor, true, [this, &bEditingChildren](const ALevelInstance* ChildLevelInstanceActor)
-		{
-			if (IsEditingLevelInstance(ChildLevelInstanceActor))
-			{
-				bEditingChildren = true;
-				return false;
-			}
-			return true;
-		});
-	return bEditingChildren;
-}
-
 void ULevelInstanceSubsystem::SetIsHiddenEdLayer(ALevelInstance* LevelInstanceActor, bool bIsHiddenEdLayer)
 {
 	if (ULevel* LevelInstanceLevel = GetLevelInstanceLevel(LevelInstanceActor))
@@ -966,10 +933,24 @@ ALevelInstance* ULevelInstanceSubsystem::CreateLevelInstanceFrom(const TArray<AA
 
 	NewLevelInstanceActor->OnEdit();
 
-	FLevelInstanceEdit& LevelInstanceEdit = LevelInstanceEdits.Add(FName(*NewLevelInstanceActor->GetWorldAssetPackage()));
-	LevelInstanceEdit.LevelStreaming = LevelStreaming;
-	LevelStreaming->LevelInstanceID = NewLevelInstanceActor->GetLevelInstanceID();
-		
+	// Notify parents of edit
+	TArray<FLevelInstanceID> AncestorIDs;
+	ForEachLevelInstanceAncestors(NewLevelInstanceActor, [&AncestorIDs](ALevelInstance* InAncestor)
+	{
+		AncestorIDs.Add(InAncestor->GetLevelInstanceID());
+		return true;
+	});
+
+	for (const FLevelInstanceID& AncestorID : AncestorIDs)
+	{
+		OnEditChild(AncestorID);
+	}
+	
+	// New level instance
+	TUniquePtr<FLevelInstanceEdit> TempLevelInstanceEdit = MakeUnique<FLevelInstanceEdit>(LevelStreaming, NewLevelInstanceActor->GetLevelInstanceID());
+	// Force mark it as changed
+	TempLevelInstanceEdit->bCommittedChanges = true;
+
 	GetWorld()->SetCurrentLevel(LoadedLevel);
 
 	// Don't force saving of unsaved/temp packages onto the user.
@@ -977,8 +958,11 @@ ALevelInstance* ULevelInstanceSubsystem::CreateLevelInstanceFrom(const TArray<AA
 	{
 		DirtyPackages.Add(NewLevelInstanceActor->GetPackage()->GetFName());
 	}
-	// Use LevelStreaming->GetLevelInstanceActor() because OnWorldAssetSaved could've reinstanced the LevelInstanceActor
-	return CommitLevelInstance(LevelStreaming->GetLevelInstanceActor(), /*bDiscardEdits*/false, CreationParams.bPromptForSave, &DirtyPackages);
+			
+	ALevelInstance* CommittedLevelInstance = CommitLevelInstanceInternal(TempLevelInstanceEdit, /*bDiscardEdits*/false, CreationParams.bPromptForSave, &DirtyPackages);
+	check(!TempLevelInstanceEdit);
+
+	return CommittedLevelInstance;
 }
 
 bool ULevelInstanceSubsystem::BreakLevelInstance(ALevelInstance* LevelInstanceActor, uint32 Levels /* = 1 */, TArray<AActor*>* OutMovedActors /* = nullptr */)
@@ -1147,7 +1131,7 @@ ULevel* ULevelInstanceSubsystem::GetLevelInstanceLevel(const ALevelInstance* Lev
 {
 	if (LevelInstanceActor->HasValidLevelInstanceID())
 	{
-		if (const FLevelInstanceEdit* LevelInstanceEdit = GetLevelInstanceEdit(LevelInstanceActor))
+		if (const FLevelInstanceEdit* CurrentEdit = GetLevelInstanceEdit(LevelInstanceActor))
 		{
 			return LevelInstanceEdit->LevelStreaming->GetLoadedLevel();
 		}
@@ -1187,13 +1171,24 @@ bool ULevelInstanceSubsystem::LevelInstanceHasLevelScriptBlueprint(const ALevelI
 	return false;
 }
 
+void ULevelInstanceSubsystem::UpdateEngineShowFlags()
+{
+	for (FLevelEditorViewportClient* LevelVC : GEditor->GetLevelViewportClients())
+	{
+		if (LevelVC && LevelVC->GetWorld() == GetWorld())
+		{
+			LevelVC->EngineShowFlags.EditingLevelInstance = !!LevelInstanceEdit;
+		}
+	}
+}
+
 void ULevelInstanceSubsystem::RemoveLevelsFromWorld(const TArray<ULevel*>& InLevels, bool bResetTrans)
 {
 	if (LevelsToRemoveScope && LevelsToRemoveScope->IsValid())
-{
-		for (ULevel* Level : InLevels)
 	{
-		LevelsToRemoveScope->Levels.AddUnique(Level);
+		for (ULevel* Level : InLevels)
+		{
+			LevelsToRemoveScope->Levels.AddUnique(Level);
 		}
 		LevelsToRemoveScope->bResetTrans |= bResetTrans;
 	}
@@ -1240,7 +1235,7 @@ void ULevelInstanceSubsystem::RemoveLevelsFromWorld(const TArray<ULevel*>& InLev
 		// No need to clear the whole editor selection since actor of this level will be removed from the selection by: UEditorEngine::OnLevelRemovedFromWorld
 		EditorLevelUtils::RemoveLevelsFromWorld(InLevels, /*bClearSelection*/false, bResetTrans);
 	}
-	}
+}
 
 ULevelInstanceSubsystem::FLevelsToRemoveScope::FLevelsToRemoveScope(ULevelInstanceSubsystem* InOwner)
 	: Owner(InOwner)
@@ -1304,11 +1299,7 @@ bool ULevelInstanceSubsystem::CanMoveActorToLevel(const AActor* Actor, FText* Ou
 
 void ULevelInstanceSubsystem::DiscardEdits()
 {
-	for (const auto& Pair : LevelInstanceEdits)
-	{
-		ULevelStreamingLevelInstanceEditor::Unload(Pair.Value.LevelStreaming);
-	}
-	LevelInstanceEdits.Empty();
+	LevelInstanceEdit.Reset();
 }
 
 void ULevelInstanceSubsystem::OnActorDeleted(AActor* Actor)
@@ -1349,7 +1340,21 @@ void ULevelInstanceSubsystem::OnActorDeleted(AActor* Actor)
 		{
 			CommitLevelInstance(LevelInstanceActor);
 		}
-		CommitChildrenLevelInstances(LevelInstanceActor);
+		else
+		{
+			// We are ending editing. Discard Non dirty child edits
+			ForEachLevelInstanceChild(LevelInstanceActor, /*bRecursive=*/true, [this](const ALevelInstance* ChildLevelInstanceActor)
+			{
+				if (const FLevelInstanceEdit* ChildLevelInstanceEdit = GetLevelInstanceEdit(ChildLevelInstanceActor))
+				{
+					check(!IsLevelInstanceEditDirty(ChildLevelInstanceEdit));
+					LevelInstanceEdit.Reset();
+					return false;
+				}
+				return true;
+			});
+		}
+
 		LevelInstancesToLoadOrUpdate.Remove(LevelInstanceActor);
 				
 		UnloadLevelInstance(LevelInstanceActor->GetLevelInstanceID());
@@ -1383,6 +1388,17 @@ bool ULevelInstanceSubsystem::ShouldIgnoreDirtyPackage(UPackage* DirtyPackage, c
 	return bIgnore;
 }
 
+ULevelInstanceSubsystem::FLevelInstanceEdit::FLevelInstanceEdit(ULevelStreamingLevelInstanceEditor* InLevelStreaming, FLevelInstanceID InLevelInstanceID)
+	: LevelStreaming(InLevelStreaming)
+{
+	LevelStreaming->LevelInstanceID = InLevelInstanceID;
+}
+
+ULevelInstanceSubsystem::FLevelInstanceEdit::~FLevelInstanceEdit()
+{
+	ULevelStreamingLevelInstanceEditor::Unload(LevelStreaming);
+}
+
 UWorld* ULevelInstanceSubsystem::FLevelInstanceEdit::GetEditWorld() const
 {
 	if (LevelStreaming && LevelStreaming->GetLoadedLevel())
@@ -1393,31 +1409,35 @@ UWorld* ULevelInstanceSubsystem::FLevelInstanceEdit::GetEditWorld() const
 	return nullptr;
 }
 
+FLevelInstanceID ULevelInstanceSubsystem::FLevelInstanceEdit::GetLevelInstanceID() const
+{
+	return LevelStreaming ? LevelStreaming->GetLevelInstanceID() : FLevelInstanceID();
+}
+
 const ULevelInstanceSubsystem::FLevelInstanceEdit* ULevelInstanceSubsystem::GetLevelInstanceEdit(const ALevelInstance* LevelInstanceActor) const
 {
-	const FLevelInstanceEdit* LevelInstanceEdit = LevelInstanceEdits.Find(FName(*LevelInstanceActor->GetWorldAssetPackage()));
-	if (!LevelInstanceEdit || LevelInstanceEdit->LevelStreaming->GetLevelInstanceActor() != LevelInstanceActor)
+	if (LevelInstanceEdit && LevelInstanceEdit->GetLevelInstanceID() == LevelInstanceActor->GetLevelInstanceID())
 	{
-		return nullptr;
+		return LevelInstanceEdit.Get();
 	}
 
-	return LevelInstanceEdit;
+	return nullptr;
 }
 
 bool ULevelInstanceSubsystem::IsEditingLevelInstanceDirty(const ALevelInstance* LevelInstanceActor) const
 {
-	const FLevelInstanceEdit* LevelInstanceEdit = GetLevelInstanceEdit(LevelInstanceActor);
-	if (!LevelInstanceEdit)
+	const FLevelInstanceEdit* CurrentEdit = GetLevelInstanceEdit(LevelInstanceActor);
+	if (!CurrentEdit)
 	{
 		return false;
 	}
 
-	return IsLevelInstanceEditDirty(LevelInstanceEdit);
+	return IsLevelInstanceEditDirty(CurrentEdit);
 }
 
-bool ULevelInstanceSubsystem::IsLevelInstanceEditDirty(const FLevelInstanceEdit* LevelInstanceEdit) const
+bool ULevelInstanceSubsystem::IsLevelInstanceEditDirty(const FLevelInstanceEdit* InLevelInstanceEdit) const
 {
-	const UWorld* EditingWorld = LevelInstanceEdit->GetEditWorld();
+	const UWorld* EditingWorld = InLevelInstanceEdit->GetEditWorld();
 	check(EditingWorld);
 
 	TArray<UPackage*> OutDirtyPackages;
@@ -1429,6 +1449,16 @@ bool ULevelInstanceSubsystem::IsLevelInstanceEditDirty(const FLevelInstanceEdit*
 	return OutDirtyPackages.Num() > 0;
 }
 
+ALevelInstance* ULevelInstanceSubsystem::GetEditingLevelInstance() const
+{
+	if (LevelInstanceEdit)
+	{
+		return GetLevelInstance(LevelInstanceEdit->GetLevelInstanceID());
+	}
+
+	return nullptr;
+}
+
 bool ULevelInstanceSubsystem::CanEditLevelInstance(const ALevelInstance* LevelInstanceActor, FText* OutReason) const
 {
 	// Only allow Editing in Editor World
@@ -1448,94 +1478,27 @@ bool ULevelInstanceSubsystem::CanEditLevelInstance(const ALevelInstance* LevelIn
 			return false;
 		}
 	}
-
-	if (const FLevelInstanceEdit* LevelInstanceEdit = LevelInstanceEdits.Find(FName(*LevelInstanceActor->GetWorldAssetPackage())))
+	
+	if (LevelInstanceEdit)
 	{
-		if (OutReason)
+		if (LevelInstanceEdit->GetLevelInstanceID() == LevelInstanceActor->GetLevelInstanceID())
 		{
-			if (LevelInstanceEdit->LevelStreaming->GetLevelInstanceActor() == LevelInstanceActor)
+			if (OutReason)
 			{
 				*OutReason = LOCTEXT("CanEditLevelInstanceAlreadyBeingEdited", "Level Instance already being edited");
 			}
-			else
-			{
-
-				*OutReason = LOCTEXT("CanEditLevelInstanceAlreadyEditing", "Another Level Instance pointing to the same level is being edited");
-			}
 		}
-		return false;
-	}
 
-	// Do not allow multiple LevelInstances of the same hierarchy to be edited... (checking ancestors)
-	bool bAncestorBeingEdited = false;
-	if (LevelInstanceEdits.Num() > 0)
-	{
-		ForEachLevelInstanceAncestors(LevelInstanceActor, [this, &bAncestorBeingEdited, OutReason](const ALevelInstance* AncestorLevelInstance)
+		if (IsLevelInstanceEditDirty(LevelInstanceEdit.Get()))
 		{
-			if (const FLevelInstanceEdit* LevelInstanceEdit = LevelInstanceEdits.Find(FName(*AncestorLevelInstance->GetWorldAssetPackage())))
-			{
-				// Allow children to be edited if ancestor is clean
-				if (IsLevelInstanceEditDirty(LevelInstanceEdit))
-				{
-					if (OutReason)
-					{
-						*OutReason = LOCTEXT("CanEditLevelInstanceAncestorBeingEdited", "Ancestor Level Instance is being edited and dirty. Commit changes first.");
-					}
-					bAncestorBeingEdited = true;
-					return false;
-				}
-			}
-			return true;
-		});
-	}
-
-	if (bAncestorBeingEdited)
-	{
-		return false;
-	}
-		
-	// Do not allow multiple LevelInstances of the same hierarchy to be edited... (checking children)
-	bool bChildBeingEdited = false;
-	if (LevelInstanceEdits.Num() > 0)
-	{
-		for (const auto& Pair : LevelInstanceEdits)
-		{
-			check(Pair.Value.LevelStreaming);
-			ALevelInstance* LevelInstanceEditActor = Pair.Value.LevelStreaming->GetLevelInstanceActor();
-			check(LevelInstanceEditActor);
-			ForEachLevelInstanceAncestors(LevelInstanceEditActor, [this, &bChildBeingEdited, OutReason, LevelInstanceActor](const ALevelInstance* AncestorLevelInstance)
-			{
-				if (AncestorLevelInstance->GetWorldAsset() == LevelInstanceActor->GetWorldAsset())
-				{
-					if (const FLevelInstanceEdit* AncestorLevelInstanceEdit = LevelInstanceEdits.Find(FName(*AncestorLevelInstance->GetWorldAssetPackage())))
-					{
-						if (IsLevelInstanceEditDirty(AncestorLevelInstanceEdit))
-						{
-							if (OutReason)
-							{
-								*OutReason = LOCTEXT("CanEditLevelInstanceOtherChildren", "Children Level Instance already being edited and dirty. Commit changes first.");
-							}
-							bChildBeingEdited = true;
-							return false;
-						}
-					}
-				}
-
-				return true;
-			});
-
-			if (bChildBeingEdited)
+			if (OutReason)
 			{
-				break;
+				*OutReason = LOCTEXT("CanEditLevelInstanceDirtyEdit", "Current Level Instance has unsaved changes and needs to be committed first.");
 			}
+			return false;
 		}
 	}
 	
-	if (bChildBeingEdited)
-	{
-		return false;
-	}
-
 	if (!LevelInstanceActor->IsLevelInstancePathValid())
 	{
 		if (OutReason)
@@ -1577,23 +1540,17 @@ bool ULevelInstanceSubsystem::CanCommitLevelInstance(const ALevelInstance* Level
 		return false;
 	}
 
-	if (HasDirtyChildrenLevelInstances(LevelInstanceActor))
-	{
-		if (OutReason)
-		{
-			*OutReason = LOCTEXT("CanCommitLevelInstanceUnsavedChild", "Level Instance has an unsaved child level instance");
-		}
-		return false;
-	}
-
 	return true;
 }
 
 void ULevelInstanceSubsystem::EditLevelInstance(ALevelInstance* LevelInstanceActor, TWeakObjectPtr<AActor> ContextActorPtr)
+{
+	EditLevelInstanceInternal(LevelInstanceActor, ContextActorPtr, false);
+}
+
+void ULevelInstanceSubsystem::EditLevelInstanceInternal(ALevelInstance* LevelInstanceActor, TWeakObjectPtr<AActor> ContextActorPtr, bool bRecursive)
 {
 	check(CanEditLevelInstance(LevelInstanceActor));
-	const bool bWasPendingEdit = PendingLevelInstanceToEdit.IsValid();
-	PendingLevelInstanceToEdit = FLevelInstanceID();
 		
 	FScopedSlowTask SlowTask(0, LOCTEXT("BeginEditLevelInstance", "Loading Level Instance for edit..."), !GetWorld()->IsGameWorld());
 	SlowTask.MakeDialog();
@@ -1617,48 +1574,39 @@ void ULevelInstanceSubsystem::EditLevelInstance(ALevelInstance* LevelInstanceAct
 	}
 
 	GEditor->SelectNone(false, true);
-
-	// Check if there is an open (but clean) ancestor and child and unload it before opening the LevelInstance for editing
-	if (LevelInstanceEdits.Num() > 0)
+	
+	// Avoid calling OnEditChild twice  on ancestors when EditLevelInstance calls itself
+	if (!bRecursive)
 	{
-		ALevelInstance* LevelInstanceToCommit = nullptr;
-
-		auto GetLevelInstanceToCommit = [this, &LevelInstanceToCommit](ALevelInstance* LevelInstance)
+		TArray<FLevelInstanceID> AncestorIDs;
+		ForEachLevelInstanceAncestors(LevelInstanceActor, [&AncestorIDs](ALevelInstance* InAncestor)
 		{
-			if (const FLevelInstanceEdit* LevelInstanceEdit = LevelInstanceEdits.Find(FName(*LevelInstance->GetWorldAssetPackage())))
-			{
-				check(!IsLevelInstanceEditDirty(LevelInstanceEdit));
-				check(LevelInstanceToCommit == nullptr);
-				LevelInstanceToCommit = LevelInstanceEdit->LevelStreaming->GetLevelInstanceActor();
-				check(LevelInstanceToCommit != nullptr);
-				return false;
-			}
-			return true;
-		};
-
-		ForEachLevelInstanceAncestors(LevelInstanceActor, [&LevelInstanceToCommit, &GetLevelInstanceToCommit](ALevelInstance* AncestorLevelInstance)
-		{
-			if (!LevelInstanceToCommit)
-			{
-				GetLevelInstanceToCommit(AncestorLevelInstance);
-			}
-			AncestorLevelInstance->OnEditChild();
+			AncestorIDs.Add(InAncestor->GetLevelInstanceID());
 			return true;
 		});
 
-		if (!LevelInstanceToCommit)
+		for (const FLevelInstanceID& AncestorID : AncestorIDs)
 		{
-			ForEachLevelInstanceChild(LevelInstanceActor, true, GetLevelInstanceToCommit);
+			OnEditChild(AncestorID);
 		}
-		
-		if (LevelInstanceToCommit)
-		{
-			PendingLevelInstanceToEdit = LevelInstanceActor->GetLevelInstanceID();
-			CommitLevelInstance(LevelInstanceToCommit);
+	}
 
-			// Stop here. The LevelInstance will be open for editing after an async reload.
-			return;
-		}
+	// Check if there is an open (but clean) ancestor unload it before opening the LevelInstance for editing
+	if (LevelInstanceEdit)
+	{	
+		// Only support one level of recursion to commit current edit
+		check(!bRecursive);
+		FLevelInstanceID PendingEditId = LevelInstanceActor->GetLevelInstanceID();
+		ALevelInstance* LevelInstanceToCommit = GetLevelInstance(LevelInstanceEdit->GetLevelInstanceID());
+		CommitLevelInstance(LevelInstanceToCommit);
+
+		ALevelInstance* LevelInstanceToEdit = GetLevelInstance(PendingEditId);
+		check(LevelInstanceToEdit);
+
+		EditLevelInstanceInternal(LevelInstanceToEdit, nullptr, /*bRecursive=*/true);
+
+		// Stop here. The LevelInstance will be open for editing after an async reload.
+		return;
 	}
 
 	// Cleanup async requests in case
@@ -1675,8 +1623,7 @@ void ULevelInstanceSubsystem::EditLevelInstance(ALevelInstance* LevelInstanceAct
 		return;
 	}
 
-	FLevelInstanceEdit& LevelInstanceEdit = LevelInstanceEdits.Add(FName(*LevelInstanceActor->GetWorldAssetPackage()));
-	LevelInstanceEdit.LevelStreaming = LevelStreaming;
+	LevelInstanceEdit = MakeUnique<FLevelInstanceEdit>(LevelStreaming, LevelInstanceActor->GetLevelInstanceID());
 		
 	// Try and select something meaningful
 	AActor* ActorToSelect = nullptr;
@@ -1695,16 +1642,6 @@ void ULevelInstanceSubsystem::EditLevelInstance(ALevelInstance* LevelInstanceAct
 	// Notify
 	LevelInstanceActor->OnEdit();
 
-	// Pending edit notification was already sent
-	if (!bWasPendingEdit)
-	{
-		ForEachLevelInstanceAncestors(LevelInstanceActor, [](ALevelInstance* AncestorLevelInstance)
-		{
-			AncestorLevelInstance->OnEditChild();
-			return true;
-		});
-	}
-
 	GEditor->SelectActor(ActorToSelect, true, true);
 
 	for (const auto& Actor : LevelStreaming->LoadedLevel->Actors)
@@ -1716,47 +1653,27 @@ void ULevelInstanceSubsystem::EditLevelInstance(ALevelInstance* LevelInstanceAct
 		}
 	}
 
-	const bool bHasLevelInstanceEdits = LevelInstanceEdits.Num() > 0;
-	if (bHasLevelInstanceEdits)
-	{
-		for (FLevelEditorViewportClient* LevelVC : GEditor->GetLevelViewportClients())
-		{
-			if (LevelVC && LevelVC->GetWorld() == GetWorld())
-			{
-				LevelVC->EngineShowFlags.EditingLevelInstance = bHasLevelInstanceEdits;
-			}
-		}
-	}
-
+	UpdateEngineShowFlags();
+	
 	// Edit can't be undone
 	GEditor->ResetTransaction(LOCTEXT("LevelInstanceEditResetTrans", "Edit Level Instance"));
 }
 
-void ULevelInstanceSubsystem::CommitChildrenLevelInstances(ALevelInstance* LevelInstanceActor)
-{
-	// We are ending editing. Discard Non dirty child edits
-	ForEachLevelInstanceChild(LevelInstanceActor, /*bRecursive=*/true, [this](const ALevelInstance* ChildLevelInstanceActor)
-	{
-		if (const FLevelInstanceEdit* ChildLevelInstanceEdit = GetLevelInstanceEdit(ChildLevelInstanceActor))
-		{
-			check(!IsLevelInstanceEditDirty(ChildLevelInstanceEdit));
-			ULevelStreamingLevelInstanceEditor::Unload(ChildLevelInstanceEdit->LevelStreaming);
-			LevelInstanceEdits.Remove(FName(*ChildLevelInstanceActor->GetWorldAssetPackage()));
-		}
-		return true;
-	});
-}
-
 ALevelInstance* ULevelInstanceSubsystem::CommitLevelInstance(ALevelInstance* LevelInstanceActor, bool bDiscardEdits, bool bPromptForSave, TSet<FName>* DirtyPackages)
 {
+	check(LevelInstanceEdit.Get() == GetLevelInstanceEdit(LevelInstanceActor));
 	check(CanCommitLevelInstance(LevelInstanceActor));
+	return CommitLevelInstanceInternal(LevelInstanceEdit, bDiscardEdits, bPromptForSave, DirtyPackages);
+}
 
-	const FLevelInstanceEdit* LevelInstanceEdit = GetLevelInstanceEdit(LevelInstanceActor);
-	check(LevelInstanceEdit);
-	UWorld* EditingWorld = LevelInstanceEdit->GetEditWorld();
+ALevelInstance* ULevelInstanceSubsystem::CommitLevelInstanceInternal(TUniquePtr<FLevelInstanceEdit>& InLevelInstanceEdit, bool bDiscardEdits, bool bPromptForSave, TSet<FName>* DirtyPackages)
+{
+	ALevelInstance* LevelInstanceActor = GetLevelInstance(InLevelInstanceEdit->GetLevelInstanceID());
+	check(InLevelInstanceEdit);
+	UWorld* EditingWorld = InLevelInstanceEdit->GetEditWorld();
 	check(EditingWorld);
 			
-	if (IsLevelInstanceEditDirty(LevelInstanceEdit) && !bDiscardEdits)
+	if (IsLevelInstanceEditDirty(InLevelInstanceEdit.Get()) && !bDiscardEdits)
 	{
 		const bool bPromptUserToSave = bPromptForSave;
 		const bool bSaveMapPackages = true;
@@ -1783,40 +1700,18 @@ ALevelInstance* ULevelInstanceSubsystem::CommitLevelInstance(ALevelInstance* Lev
 	SlowTask.MakeDialog();
 
 	GEditor->SelectNone(false, true);
-		
-	// End edit non dirty child edits
-	CommitChildrenLevelInstances(LevelInstanceActor);
-		
-	// Try to find proper LevelInstance to select
-	TWeakObjectPtr<ALevelInstance> ActorToSelect = LevelInstanceActor;
-	
-	ForEachLevelInstanceAncestors(LevelInstanceActor, [&ActorToSelect](ALevelInstance* AncestorLevelInstance)
-	{
-		// if we find a parent editing LevelInstance this is what we want to select
-		if (AncestorLevelInstance->IsEditing())
-		{
-			ActorToSelect = AncestorLevelInstance;
-			return false;
-		}
-		else // if not we go up the ancestor LevelInstances to the highest level
-		{
-			ActorToSelect = AncestorLevelInstance;
-		}
-		return true;
-	});
 
 	const FString EditPackage = LevelInstanceActor->GetWorldAssetPackage();
 
 	// Remove from streaming level...
-	const bool bChangesCommitted = LevelInstanceEdit->bCommittedChanges;
-	ULevelStreamingLevelInstanceEditor::Unload(LevelInstanceEdit->LevelStreaming);
-	LevelInstanceEdits.Remove(FName(*EditPackage));
+	const bool bChangesCommitted = InLevelInstanceEdit->bCommittedChanges;
+	InLevelInstanceEdit.Reset();
 
 	if (bChangesCommitted)
 	{
 		// Sync the AssetData so that the updated instances have the latest Actor Registry Data
 		IAssetRegistry& AssetRegistry = FModuleManager::LoadModuleChecked<FAssetRegistryModule>(TEXT("AssetRegistry")).Get();
-		AssetRegistry.ScanPathsSynchronous({ LevelInstanceActor->GetWorldAssetPackage() }, true);
+		AssetRegistry.ScanPathsSynchronous({ EditPackage }, true);
 	}
 	
 	// Backup ID on Commit in case Actor gets recreated
@@ -1828,51 +1723,63 @@ ALevelInstance* ULevelInstanceSubsystem::CommitLevelInstance(ALevelInstance* Lev
 	// Update pointer since BP Compilation might have invalidated LevelInstanceActor
 	LevelInstanceActor = GetLevelInstance(LevelInstanceID);
 
-	// Notify Ancestors
-	ForEachLevelInstanceAncestors(LevelInstanceActor, [bChangesCommitted](ALevelInstance* AncestorLevelInstance)
-	{
-		AncestorLevelInstance->OnCommitChild(bChangesCommitted);
-		return true;
-	});
-				
-	// Propagate to other instances
-	for (TObjectIterator<UWorld> It(RF_ClassDefaultObject|RF_ArchetypeObject, true, EInternalObjectFlags::PendingKill); It; ++It)
+	TArray<FLevelInstanceID> LevelInstancesToUpdate;
+	// Gather list to update
+	for (TObjectIterator<UWorld> It(RF_ClassDefaultObject | RF_ArchetypeObject, true); It; ++It)
 	{
 		UWorld* CurrentWorld = *It;
-		if (CurrentWorld->GetSubsystem<ULevelInstanceSubsystem>() != nullptr)
+		if (IsValid(CurrentWorld) && CurrentWorld->GetSubsystem<ULevelInstanceSubsystem>() != nullptr)
 		{
 			for (TActorIterator<ALevelInstance> LevelInstanceIt(CurrentWorld); LevelInstanceIt; ++LevelInstanceIt)
 			{
 				ALevelInstance* CurrentLevelInstanceActor = *LevelInstanceIt;
 				if (CurrentLevelInstanceActor->GetWorldAssetPackage() == EditPackage && (LevelInstanceActor == CurrentLevelInstanceActor || bChangesCommitted))
 				{
-					CurrentLevelInstanceActor->UpdateLevelInstance();
+					LevelInstancesToUpdate.Add(CurrentLevelInstanceActor->GetLevelInstanceID());
 				}
 			}
 		}
 	}
 
-	GEngine->BroadcastLevelActorListChanged();
-
-	if (ALevelInstance* Actor = ActorToSelect.Get())
+	// Do update
+	for (const FLevelInstanceID& LevelInstanceToUpdateID : LevelInstancesToUpdate)
 	{
-		GEditor->SelectActor(Actor, true, true);
+		if (ALevelInstance* LevelInstance = GetLevelInstance(LevelInstanceToUpdateID))
+		{
+			LevelInstance->UpdateLevelInstance();
+		}
 	}
 
-	const bool bHasLevelInstanceEdits = LevelInstanceEdits.Num() > 0;
+	LevelInstanceActor = GetLevelInstance(LevelInstanceID);
+	
+	// Notify Ancestors
+	FLevelInstanceID LevelInstanceToSelectID = LevelInstanceID;
+	TArray<FLevelInstanceID> AncestorIDs;
+	ForEachLevelInstanceAncestors(LevelInstanceActor, [&LevelInstanceToSelectID, &AncestorIDs](ALevelInstance* AncestorLevelInstance)
+	{
+		LevelInstanceToSelectID = AncestorLevelInstance->GetLevelInstanceID();
+		AncestorIDs.Add(AncestorLevelInstance->GetLevelInstanceID());
+		return true;
+	});
 
-	if (!bHasLevelInstanceEdits)
+	for (const FLevelInstanceID& AncestorID : AncestorIDs)
 	{
-		for (FLevelEditorViewportClient* LevelVC : GEditor->GetLevelViewportClients())
-		{
-			if (LevelVC && LevelVC->GetWorld() == GetWorld())
-			{
-				LevelVC->EngineShowFlags.EditingLevelInstance = bHasLevelInstanceEdits;
-			}
-		}
+		OnCommitChild(AncestorID, bChangesCommitted);
+	}
+		
+	if (ALevelInstance* Actor = GetLevelInstance(LevelInstanceToSelectID))
+	{
+		GEditor->SelectActor(Actor, true, true);
 	}
+				
+	// Wait for Level Instances to be loaded
+	BlockOnLoading();
+
+	GEngine->BroadcastLevelActorListChanged();
+
+	UpdateEngineShowFlags();
 
-	return LevelInstanceActor;
+	return GetLevelInstance(LevelInstanceID);
 }
 
 void ULevelInstanceSubsystem::SaveLevelInstanceAs(ALevelInstance* LevelInstanceActor)
@@ -1914,11 +1821,8 @@ ALevelInstance* ULevelInstanceSubsystem::GetParentLevelInstance(const AActor* Ac
 	return GetOwningLevelInstance(OwningLevel);
 }
 
-void ULevelInstanceSubsystem::BlockLoadLevelInstance(ALevelInstance* LevelInstanceActor)
+void ULevelInstanceSubsystem::BlockOnLoading()
 {
-	check(!LevelInstanceActor->IsEditing());
-	RequestLoadLevelInstance(LevelInstanceActor, true);
-
 	// Make sure blocking loads can happen and are not part of transaction
 	TGuardValue<ITransaction*> TransactionGuard(GUndo, nullptr);
 
@@ -1929,18 +1833,51 @@ void ULevelInstanceSubsystem::BlockLoadLevelInstance(ALevelInstance* LevelInstan
 	}
 }
 
+void ULevelInstanceSubsystem::BlockLoadLevelInstance(ALevelInstance* LevelInstanceActor)
+{
+	check(!LevelInstanceActor->IsEditing());
+	RequestLoadLevelInstance(LevelInstanceActor, true);
+
+	BlockOnLoading();
+}
+
 void ULevelInstanceSubsystem::BlockUnloadLevelInstance(ALevelInstance* LevelInstanceActor)
 {
 	check(!LevelInstanceActor->IsEditing());
 	RequestUnloadLevelInstance(LevelInstanceActor);
 
-	// Make sure blocking unloads can happen and are not part of transaction
-	TGuardValue<ITransaction*> TransactionGuard(GUndo, nullptr);
+	BlockOnLoading();
+}
 
-	// Blocking until LevelInstance is loaded and all its child LevelInstances
-	while (LevelInstancesToUnload.Num())
+bool ULevelInstanceSubsystem::HasChildEdit(const ALevelInstance* LevelInstanceActor) const
+{
+	const int32* ChildEditCountPtr = ChildEdits.Find(LevelInstanceActor->GetLevelInstanceID());
+	return ChildEditCountPtr && *ChildEditCountPtr;
+}
+
+void ULevelInstanceSubsystem::OnCommitChild(FLevelInstanceID LevelInstanceID, bool bChildChanged)
+{
+	int32& ChildEditCount = ChildEdits.FindChecked(LevelInstanceID);
+	check(ChildEditCount > 0);
+	ChildEditCount--;
+
+	if (ALevelInstance* LevelInstance = GetLevelInstance(LevelInstanceID))
 	{
-		UpdateStreamingState();
+		LevelInstance->OnCommitChild(bChildChanged);
+	}
+}
+
+void ULevelInstanceSubsystem::OnEditChild(FLevelInstanceID LevelInstanceID)
+{
+	int32& ChildEditCount = ChildEdits.FindOrAdd(LevelInstanceID, 0);
+	// Child edit count can reach 2 maximum in the Context of creating a LevelInstance inside an already editing child level instance
+	// through CreateLevelInstanceFrom
+	check(ChildEditCount < 2);
+	ChildEditCount++;
+
+	if (ALevelInstance* LevelInstance = GetLevelInstance(LevelInstanceID))
+	{
+		LevelInstance->OnEditChild();
 	}
 }
 
diff --git a/Engine/Source/Runtime/Engine/Private/LevelInstance/Packed/PackedLevelInstanceActor.cpp b/Engine/Source/Runtime/Engine/Private/LevelInstance/Packed/PackedLevelInstanceActor.cpp
--- a/Engine/Source/Runtime/Engine/Private/LevelInstance/Packed/PackedLevelInstanceActor.cpp
+++ b/Engine/Source/Runtime/Engine/Private/LevelInstance/Packed/PackedLevelInstanceActor.cpp
@@ -11,8 +11,6 @@ APackedLevelInstance::APackedLevelInstance(const FObjectInitializer& ObjectIniti
 	: Super(ObjectInitializer)
 #if WITH_EDITORONLY_DATA
 	, ISMComponentClass(UInstancedStaticMeshComponent::StaticClass())
-	, bEditing(false)
-	, ChildEditing(0)
 	, bChildChanged(false)
 #endif
 {
@@ -39,7 +37,7 @@ void APackedLevelInstance::Serialize(FArchive& Ar)
 bool APackedLevelInstance::SupportsLoading() const
 {
 #if WITH_EDITOR
-	return ChildEditing > 0 || IsLoaded();
+	return HasChildEdit() || IsLoaded();
 #else
 	return false;
 #endif
@@ -94,6 +92,8 @@ FName APackedLevelInstance::GetPackedComponentTag()
 
 void APackedLevelInstance::UpdateLevelInstance()
 {
+	Super::UpdateLevelInstance();
+
 	if (!Cast<UBlueprint>(GetClass()->ClassGeneratedBy))
 	{
 		if (IsLevelInstancePathValid())
@@ -110,18 +110,20 @@ void APackedLevelInstance::UpdateLevelInstance()
 
 void APackedLevelInstance::OnEditChild()
 {
+	Super::OnEditChild();
+
 	check(GetLevelInstanceSubsystem()->GetLevelInstanceLevel(this) != nullptr);
-	ChildEditing++;
 	MarkComponentsRenderStateDirty();
 }
 
 void APackedLevelInstance::OnCommitChild(bool bChanged)
 {
-	check(GetLevelInstanceSubsystem()->GetLevelInstanceLevel(this) != nullptr);
-	check(ChildEditing > 0);
-	ChildEditing--;
+	Super::OnCommitChild(bChanged);
+
+	ULevelInstanceSubsystem* LevelInstanceSubsystem = GetLevelInstanceSubsystem();
+	check(GetLevelInstanceSubsystem()->GetLevelInstanceLevel(this));
 	bChildChanged |= bChanged;
-	if (!ChildEditing)
+	if (!HasChildEdit())
 	{
 		UnloadLevelInstance();
 
@@ -149,17 +151,12 @@ void APackedLevelInstance::OnCommitChild(bool bChanged)
 void APackedLevelInstance::OnEdit()
 {
 	Super::OnEdit();
-	check(!bEditing);
-	
-	bEditing = true;
 	MarkComponentsRenderStateDirty();
 }
 
 void APackedLevelInstance::OnCommit(bool bChanged, bool bPromptForSave)
 {
 	Super::OnCommit(bChanged, bPromptForSave);
-	check(bEditing);
-	bEditing = false;
 
 	if (bChanged)
 	{
@@ -178,7 +175,7 @@ void APackedLevelInstance::OnCommit(bool bChanged, bool bPromptForSave)
 
 bool APackedLevelInstance::IsHiddenEd() const
 {
-	return Super::IsHiddenEd() || bEditing || (ChildEditing > 0);
+	return Super::IsHiddenEd() || IsEditing() || HasChildEdit();
 }
 
 bool APackedLevelInstance::IsHLODRelevant() const
diff --git a/Engine/Source/Runtime/Engine/Public/LevelInstance/LevelInstanceActor.h b/Engine/Source/Runtime/Engine/Public/LevelInstance/LevelInstanceActor.h
--- a/Engine/Source/Runtime/Engine/Public/LevelInstance/LevelInstanceActor.h
+++ b/Engine/Source/Runtime/Engine/Public/LevelInstance/LevelInstanceActor.h
@@ -90,7 +90,7 @@ public:
 	bool CanEdit(FText* OutReason = nullptr) const;
 	bool CanCommit(FText* OutReason = nullptr) const;
 	bool IsEditing() const;
-	bool HasEditingChildren() const;
+	bool HasChildEdit() const;
 	void Edit(AActor* ContextActor = nullptr);
 	void Commit();
 	void Discard();
@@ -107,9 +107,9 @@ public:
 	AActor* FindEditorInstanceActor() const;
 
 	virtual void OnWorldAssetChanged() { UpdateLevelInstance(); }
-	virtual void OnEdit();
+	virtual void OnEdit() {}
 	virtual void OnEditChild() {}
-	virtual void OnCommit(bool bChanged, bool bPromptForSave);
+	virtual void OnCommit(bool bChanged, bool bPromptForSave) {}
 	virtual void OnCommitChild(bool bChanged) {}
 	virtual void UpdateLevelInstance();
 		
@@ -127,7 +127,6 @@ private:
 	FLevelInstanceID CachedLevelInstanceID;
 	bool bCachedIsTemporarilyHiddenInEditor;
 	bool bGuardLoadUnload;
-	bool bEditLockLocation;
 #else
 	// This Guid is used to compute the LevelInstanceID. Because in non-editor build we don't have an ActorGuid, we store it at cook time.
 	FGuid LevelInstanceActorGuid;
diff --git a/Engine/Source/Runtime/Engine/Public/LevelInstance/LevelInstanceEditorLevelStreaming.h b/Engine/Source/Runtime/Engine/Public/LevelInstance/LevelInstanceEditorLevelStreaming.h
--- a/Engine/Source/Runtime/Engine/Public/LevelInstance/LevelInstanceEditorLevelStreaming.h
+++ b/Engine/Source/Runtime/Engine/Public/LevelInstance/LevelInstanceEditorLevelStreaming.h
@@ -19,6 +19,7 @@ public:
 #if WITH_EDITOR
 	virtual bool ShowInLevelCollection() const override { return false; }
 	ALevelInstance* GetLevelInstanceActor() const;
+	const FLevelInstanceID& GetLevelInstanceID() const { return LevelInstanceID; }
 	FBox GetBounds() const;
 
 protected:
diff --git a/Engine/Source/Runtime/Engine/Public/LevelInstance/LevelInstanceSubsystem.h b/Engine/Source/Runtime/Engine/Public/LevelInstance/LevelInstanceSubsystem.h
--- a/Engine/Source/Runtime/Engine/Public/LevelInstance/LevelInstanceSubsystem.h
+++ b/Engine/Source/Runtime/Engine/Public/LevelInstance/LevelInstanceSubsystem.h
@@ -53,6 +53,7 @@ public:
 	void PackLevelInstances();
 	bool CanPackLevelInstances() const;
 
+	ALevelInstance* GetEditingLevelInstance() const;
 	bool CanEditLevelInstance(const ALevelInstance* LevelInstanceActor, FText* OutReason = nullptr) const;
 	bool CanCommitLevelInstance(const ALevelInstance* LevelInstanceActor, FText* OutReason = nullptr) const;
 	void EditLevelInstance(ALevelInstance* LevelInstanceActor, TWeakObjectPtr<AActor> ContextActorPtr = nullptr);
@@ -69,10 +70,8 @@ public:
 	void ForEachLevelInstanceAncestors(const AActor* Actor, TFunctionRef<bool(const ALevelInstance*)> Operation) const;
 	void ForEachLevelInstanceChild(const ALevelInstance* LevelInstanceActor, bool bRecursive, TFunctionRef<bool(const ALevelInstance*)> Operation) const;
 	void ForEachLevelInstanceChild(ALevelInstance* LevelInstanceActor, bool bRecursive, TFunctionRef<bool(ALevelInstance*)> Operation) const;
-	void ForEachLevelInstanceEdit(TFunctionRef<bool(ALevelInstance*)> Operation) const;
 	bool HasDirtyChildrenLevelInstances(const ALevelInstance* LevelInstanceActor) const;
-	bool HasEditingChildrenLevelInstances(const ALevelInstance* LevelInstanceActor) const;
-
+	
 	void SetIsHiddenEdLayer(ALevelInstance* LevelInstanceActor, bool bIsHiddenEdLayer);
 	void SetIsTemporarilyHiddenInEditor(ALevelInstance* LevelInstanceActor, bool bIsHidden);
 
@@ -93,19 +92,24 @@ public:
 	ALevelInstance* GetParentLevelInstance(const AActor* Actor) const;
 	void BlockLoadLevelInstance(ALevelInstance* LevelInstanceActor);
 	void BlockUnloadLevelInstance(ALevelInstance* LevelInstanceActor);
+		
+	bool HasChildEdit(const ALevelInstance* LevelInstanceActor) const;
 #endif
 
 private:
+	void BlockOnLoading();
 	void LoadLevelInstance(ALevelInstance* LevelInstanceActor);
 	void UnloadLevelInstance(const FLevelInstanceID& LevelInstanceID);
 	void ForEachActorInLevel(ULevel* Level, TFunctionRef<bool(AActor * LevelActor)> Operation) const;
 	void ForEachLevelInstanceAncestors(AActor* Actor, TFunctionRef<bool(ALevelInstance*)> Operation) const;
 	ALevelInstance* GetOwningLevelInstance(const ULevel* Level) const;
 #if WITH_EDITOR
+	void OnEditChild(FLevelInstanceID LevelInstanceID);
+	void OnCommitChild(FLevelInstanceID LevelInstanceID, bool bChildChanged);
+
 	bool ForEachLevelInstanceChildImpl(const ALevelInstance* LevelInstanceActor, bool bRecursive, TFunctionRef<bool(const ALevelInstance*)> Operation) const;
 	bool ForEachLevelInstanceChildImpl(ALevelInstance* LevelInstanceActor, bool bRecursive, TFunctionRef<bool(ALevelInstance*)> Operation) const;
 
-	void CommitChildrenLevelInstances(ALevelInstance* LevelInstanceActor);
 	void BreakLevelInstance_Impl(ALevelInstance* LevelInstanceActor, uint32 Levels, TArray<AActor*>& OutMovedActors);
 
 	static bool ShouldIgnoreDirtyPackage(UPackage* DirtyPackage, const UWorld* EditingWorld);
@@ -113,12 +117,18 @@ private:
 
 	struct FLevelInstanceEdit
 	{
-		ULevelStreamingLevelInstanceEditor* LevelStreaming = nullptr;
+		ULevelStreamingLevelInstanceEditor* LevelStreaming;
 		bool bCommittedChanges = false;
 
+		FLevelInstanceEdit(ULevelStreamingLevelInstanceEditor* InLevelStreaming, FLevelInstanceID InLevelInstanceID);
+		~FLevelInstanceEdit();
+
 		UWorld* GetEditWorld() const;
+		FLevelInstanceID GetLevelInstanceID() const;
 	};
 
+	void EditLevelInstanceInternal(ALevelInstance* LevelInstanceActor, TWeakObjectPtr<AActor> ContextActorPtr, bool bRecursive);
+	ALevelInstance* CommitLevelInstanceInternal(TUniquePtr<FLevelInstanceEdit>& InLevelInstanceEdit, bool bDiscardEdits, bool bPromptForSave, TSet<FName>* DirtyPackages);
 	const FLevelInstanceEdit* GetLevelInstanceEdit(const ALevelInstance* LevelInstanceActor) const;
 	bool IsLevelInstanceEditDirty(const FLevelInstanceEdit* LevelInstanceEdit) const;
 	
@@ -137,6 +147,7 @@ private:
 	friend ULevelStreamingLevelInstance;
 	friend ULevelStreamingLevelInstanceEditor;
 	void RemoveLevelsFromWorld(const TArray<ULevel*>& Levels, bool bResetTrans = true);
+	void UpdateEngineShowFlags();
 #endif
 
 
@@ -144,7 +155,7 @@ private:
 	// Optional scope to accelerate level unload by batching them
 	TUniquePtr<FLevelsToRemoveScope> LevelsToRemoveScope;
 	
-	TMap<FName, FLevelInstanceEdit> LevelInstanceEdits;
+	TUniquePtr<FLevelInstanceEdit> LevelInstanceEdit;
 #endif
 
 	struct FLevelInstance
@@ -158,5 +169,5 @@ private:
 	TMap<FLevelInstanceID, ALevelInstance*> RegisteredLevelInstances;
 	TMap<FObjectKey, FFolder::FRootObject> UnregisteringLevelInstanceLevels;
 	TMap<FFolder::FRootObject, FObjectKey > UnregisteringLevelInstances;
-	FLevelInstanceID PendingLevelInstanceToEdit;
+	TMap<FLevelInstanceID, int32> ChildEdits;
 };
\ No newline at end of file
diff --git a/Engine/Source/Runtime/Engine/Public/LevelInstance/Packed/PackedLevelInstanceActor.h b/Engine/Source/Runtime/Engine/Public/LevelInstance/Packed/PackedLevelInstanceActor.h
--- a/Engine/Source/Runtime/Engine/Public/LevelInstance/Packed/PackedLevelInstanceActor.h
+++ b/Engine/Source/Runtime/Engine/Public/LevelInstance/Packed/PackedLevelInstanceActor.h
@@ -83,13 +83,6 @@ public:
 	TArray<TSoftObjectPtr<UBlueprint>> PackedBPDependencies;
 
 private:
-	UPROPERTY(NonTransactional)
-	bool bEditing;
-
-	UPROPERTY(NonTransactional)
-	int32 ChildEditing;
-
-	UPROPERTY(NonTransactional)
 	bool bChildChanged;
 
 	UPROPERTY()
