WorldPartition: Level Instance created in World Partition world are always OFPA

diff --git a/Engine/Source/Editor/LevelInstanceEditor/Private/LevelInstanceEditorModule.cpp b/Engine/Source/Editor/LevelInstanceEditor/Private/LevelInstanceEditorModule.cpp
--- a/Engine/Source/Editor/LevelInstanceEditor/Private/LevelInstanceEditorModule.cpp
+++ b/Engine/Source/Editor/LevelInstanceEditor/Private/LevelInstanceEditorModule.cpp
@@ -272,6 +272,9 @@ namespace LevelInstanceMenuUtils
 			.ParentWindow(NewLevelInstanceWindow)
 			.PivotActors(ActorsToMove);
 
+		const bool bForceExternalActors = LevelInstanceSubsystem->GetWorld()->IsPartitionedWorld();
+		FNewLevelInstanceParams& DialogParams = NewLevelInstanceDialog->GetCreationParams();
+		DialogParams.SetForceExternalActors(bForceExternalActors);
 		NewLevelInstanceWindow->SetContent(NewLevelInstanceDialog);
 
 		FSlateApplication::Get().AddModalWindow(NewLevelInstanceWindow.ToSharedRef(), MainFrameModule.GetParentWindow());
diff --git a/Engine/Source/Editor/LevelInstanceEditor/Private/SNewLevelInstanceDialog.h b/Engine/Source/Editor/LevelInstanceEditor/Private/SNewLevelInstanceDialog.h
--- a/Engine/Source/Editor/LevelInstanceEditor/Private/SNewLevelInstanceDialog.h
+++ b/Engine/Source/Editor/LevelInstanceEditor/Private/SNewLevelInstanceDialog.h
@@ -42,7 +42,7 @@ public:
 
 	bool ClickedOk() const { return bClickedOk; }
 	
-	const FNewLevelInstanceParams& GetCreationParams() const { return CreationParams; }
+	FNewLevelInstanceParams& GetCreationParams() { return CreationParams; }
 
 private:
 	FReply OnOkClicked();
diff --git a/Engine/Source/Runtime/Engine/Private/LevelInstance/LevelInstanceSubsystem.cpp b/Engine/Source/Runtime/Engine/Private/LevelInstance/LevelInstanceSubsystem.cpp
--- a/Engine/Source/Runtime/Engine/Private/LevelInstance/LevelInstanceSubsystem.cpp
+++ b/Engine/Source/Runtime/Engine/Private/LevelInstance/LevelInstanceSubsystem.cpp
@@ -861,7 +861,7 @@ ALevelInstance* ULevelInstanceSubsystem::CreateLevelInstanceFrom(const TArray<AA
 	}
 	
 	// Convert to OFPA
-	if (CreationParams.bExternalActors)
+	if (CreationParams.UseExternalActors())
 	{
 		LoadedLevel->ConvertAllActorsToPackaging(true);
 		LoadedLevel->bUseExternalActors = true;
diff --git a/Engine/Source/Runtime/Engine/Public/LevelInstance/LevelInstanceTypes.h b/Engine/Source/Runtime/Engine/Public/LevelInstance/LevelInstanceTypes.h
--- a/Engine/Source/Runtime/Engine/Public/LevelInstance/LevelInstanceTypes.h
+++ b/Engine/Source/Runtime/Engine/Public/LevelInstance/LevelInstanceTypes.h
@@ -62,9 +62,6 @@ struct FNewLevelInstanceParams
 	UPROPERTY(EditAnywhere, Category = Default)
 	ELevelInstanceCreationType Type = ELevelInstanceCreationType::LevelInstance;
 
-	UPROPERTY(EditAnywhere, Category = Default)
-	bool bExternalActors = true;
-
 	UPROPERTY(EditAnywhere, Category = Pivot)
 	ELevelInstancePivotType PivotType = ELevelInstancePivotType::CenterMinZ;
 
@@ -79,4 +76,16 @@ struct FNewLevelInstanceParams
 
 	UPROPERTY()
 	bool bPromptForSave = false;
+
+private:
+	UPROPERTY(EditAnywhere, Category = Default, meta = (EditCondition = "!bForceExternalActors", EditConditionHides, HideEditConditionToggle))
+	bool bExternalActors = true;
+	
+	UPROPERTY()
+	bool bForceExternalActors = false;
+
+public:
+	void SetForceExternalActors(bool bInForceExternalActors) { bForceExternalActors = bInForceExternalActors; }
+	void SetExternalActors(bool bInExternalActors) { bExternalActors = bInExternalActors; }
+	bool UseExternalActors() const { return bForceExternalActors || bExternalActors; }
 };
\ No newline at end of file
