Fix precomputed shadow factor render target being alway bound even when the feature is disabled.

This issue was introduced in CL 18294368 for changing how precomputed shadow factor are managed with Strata.

diff --git a/Engine/Source/Runtime/Engine/Private/ShaderCompiler/ShaderGenerationUtil.cpp b/Engine/Source/Runtime/Engine/Private/ShaderCompiler/ShaderGenerationUtil.cpp
--- a/Engine/Source/Runtime/Engine/Private/ShaderCompiler/ShaderGenerationUtil.cpp
+++ b/Engine/Source/Runtime/Engine/Private/ShaderCompiler/ShaderGenerationUtil.cpp
@@ -1681,7 +1681,7 @@ static void SetStandardGBufferSlots(bool Slots[], bool bWriteEmissive, bool bHas
 {
 	Slots[GBS_SceneColor] = bWriteEmissive;
 	Slots[GBS_Velocity] = bHasVelocity;
-	Slots[GBS_PrecomputedShadowFactor] = bIsStrataMaterial ? false : true;
+	Slots[GBS_PrecomputedShadowFactor] = bIsStrataMaterial ? false : bHasStaticLighting;
 
 	Slots[GBS_WorldNormal] =			bIsStrataMaterial ? false : true;
 	Slots[GBS_PerObjectGBufferData] =	bIsStrataMaterial ? false : true;
