Restore backward compatibility for Resource and PlatformData fields

diff --git a/Engine/Source/Runtime/Engine/Classes/Engine/Texture.h b/Engine/Source/Runtime/Engine/Classes/Engine/Texture.h
--- a/Engine/Source/Runtime/Engine/Classes/Engine/Texture.h
+++ b/Engine/Source/Runtime/Engine/Classes/Engine/Texture.h
@@ -18,6 +18,7 @@
 #include "TextureResource.h"
 #include "Engine/StreamableRenderAsset.h"
 #include "PerPlatformProperties.h"
+#include "Misc/FieldAccessor.h"
 #if WITH_EDITORONLY_DATA
 #include "Misc/TVariant.h"
 #include "DerivedDataCacheKeyProxy.h"
@@ -1222,8 +1223,8 @@ private:
 
 public:
 #if WITH_TEXTURE_RESOURCE_DEPRECATIONS
-	UE_DEPRECATED(5.00, "Use GetResource() / SetResource() accessors instead. This value cannot be relied upon anymore.")
-	class FTextureResource* Resource {nullptr};
+	UE_DEPRECATED(5.00, "Use GetResource() / SetResource() accessors instead.")
+	TFieldPtrAccessor<FTextureResource> Resource;
 #endif
 
 	PRAGMA_DISABLE_DEPRECATION_WARNINGS
diff --git a/Engine/Source/Runtime/Engine/Classes/Engine/Texture2D.h b/Engine/Source/Runtime/Engine/Classes/Engine/Texture2D.h
--- a/Engine/Source/Runtime/Engine/Classes/Engine/Texture2D.h
+++ b/Engine/Source/Runtime/Engine/Classes/Engine/Texture2D.h
@@ -8,6 +8,7 @@
 #include "UObject/ScriptMacros.h"
 #include "Engine/Texture.h"
 #include "TextureResource.h"
+#include "Misc/FieldAccessor.h"
 #include "Serialization/BulkData2.h"
 #include "Texture2D.generated.h"
 
@@ -72,8 +73,8 @@ private:
 
 public:
 #if WITH_TEXTURE_PLATFORMDATA_DEPRECATIONS
-	UE_DEPRECATED(5.00, "Use GetPlatformData() / SetPlatformData() accessors instead. This value cannot be relied upon anymore.")
-	FTexturePlatformData* PlatformData {nullptr};
+	UE_DEPRECATED(5.00, "Use GetPlatformData() / SetPlatformData() accessors instead.")
+	TFieldPtrAccessor<FTexturePlatformData> PlatformData;
 #endif
 
 	/** Set the derived data for this texture on this platform. */
diff --git a/Engine/Source/Runtime/Engine/Classes/Engine/TextureCube.h b/Engine/Source/Runtime/Engine/Classes/Engine/TextureCube.h
--- a/Engine/Source/Runtime/Engine/Classes/Engine/TextureCube.h
+++ b/Engine/Source/Runtime/Engine/Classes/Engine/TextureCube.h
@@ -6,6 +6,7 @@
 #include "CoreMinimal.h"
 #include "UObject/ObjectMacros.h"
 #include "Engine/Texture.h"
+#include "Misc/FieldAccessor.h"
 #include "TextureCube.generated.h"
 
 class FTextureResource;
@@ -20,8 +21,8 @@ class UTextureCube : public UTexture
 
 public:
 #if WITH_TEXTURE_PLATFORMDATA_DEPRECATIONS
-	UE_DEPRECATED(5.00, "Use GetPlatformData() / SetPlatformData() accessors instead. This value cannot be relied upon anymore.")
-	FTexturePlatformData* PlatformData {nullptr};
+	UE_DEPRECATED(5.00, "Use GetPlatformData() / SetPlatformData() accessors instead.")
+	TFieldPtrAccessor<FTexturePlatformData> PlatformData;
 #endif
 
 	/** Set the derived data for this texture on this platform. */
@@ -36,7 +37,9 @@ public:
 #endif
 
 	/** Destructor */
+	PRAGMA_DISABLE_DEPRECATION_WARNINGS
 	ENGINE_API virtual ~UTextureCube() {};
+	PRAGMA_ENABLE_DEPRECATION_WARNINGS
 
 	//~ Begin UObject Interface.
 	ENGINE_API virtual void Serialize(FArchive& Ar) override;
diff --git a/Engine/Source/Runtime/Engine/Private/Texture.cpp b/Engine/Source/Runtime/Engine/Private/Texture.cpp
--- a/Engine/Source/Runtime/Engine/Private/Texture.cpp
+++ b/Engine/Source/Runtime/Engine/Private/Texture.cpp
@@ -93,6 +93,11 @@ UTexture::UTexture(const FObjectInitializer& ObjectInitializer)
 	: Super(ObjectInitializer)
 	, PrivateResource(nullptr)
 	, PrivateResourceRenderThread(nullptr)
+#if WITH_TEXTURE_RESOURCE_DEPRECATIONS
+	, Resource(
+		[this]()-> FTextureResource* { return GetResource(); },
+		[this](FTextureResource* InTextureResource) { SetResource(InTextureResource); })
+#endif
 {
 	SRGB = true;
 	Filter = TF_Default;
diff --git a/Engine/Source/Runtime/Engine/Private/Texture2D.cpp b/Engine/Source/Runtime/Engine/Private/Texture2D.cpp
--- a/Engine/Source/Runtime/Engine/Private/Texture2D.cpp
+++ b/Engine/Source/Runtime/Engine/Private/Texture2D.cpp
@@ -64,6 +64,11 @@
 UTexture2D::UTexture2D(const FObjectInitializer& ObjectInitializer)
 	: Super(ObjectInitializer)
 	, PrivatePlatformData(nullptr)
+#if WITH_TEXTURE_PLATFORMDATA_DEPRECATIONS
+	, PlatformData(
+		[this]()-> FTexturePlatformData* { return GetPlatformData(); },
+		[this](FTexturePlatformData* InPlatformData) { SetPlatformData(InPlatformData); })
+#endif
 {
 	PendingUpdate = nullptr;
 	SRGB = true;
diff --git a/Engine/Source/Runtime/Engine/Private/TextureCube.cpp b/Engine/Source/Runtime/Engine/Private/TextureCube.cpp
--- a/Engine/Source/Runtime/Engine/Private/TextureCube.cpp
+++ b/Engine/Source/Runtime/Engine/Private/TextureCube.cpp
@@ -73,6 +73,11 @@ static UTextureCube* GetDefaultTextureCube(const UTextureCube* Texture)
 UTextureCube::UTextureCube(const FObjectInitializer& ObjectInitializer)
 	: Super(ObjectInitializer)
 	, PrivatePlatformData(nullptr)
+#if WITH_TEXTURE_PLATFORMDATA_DEPRECATIONS
+	, PlatformData(
+		[this]()-> FTexturePlatformData* { return GetPlatformData(); },
+		[this](FTexturePlatformData* InPlatformData) { SetPlatformData(InPlatformData); })
+#endif
 {
 	SRGB = true;
 }
