Added option to allow setting DataLayer runtime state recursively from parent to children.

diff --git a/Engine/Source/Runtime/Engine/Private/WorldPartition/DataLayer/DataLayerSubsystem.cpp b/Engine/Source/Runtime/Engine/Private/WorldPartition/DataLayer/DataLayerSubsystem.cpp
--- a/Engine/Source/Runtime/Engine/Private/WorldPartition/DataLayer/DataLayerSubsystem.cpp
+++ b/Engine/Source/Runtime/Engine/Private/WorldPartition/DataLayer/DataLayerSubsystem.cpp
@@ -94,13 +94,13 @@ UDataLayer* UDataLayerSubsystem::GetDataLayerFromName(FName InDataLayerName) con
 	return WorldDataLayers ? const_cast<UDataLayer*>(WorldDataLayers->GetDataLayerFromName(InDataLayerName)) : nullptr;
 }
 
-void UDataLayerSubsystem::SetDataLayerRuntimeState(const UDataLayer* InDataLayer, EDataLayerRuntimeState InState)
+void UDataLayerSubsystem::SetDataLayerRuntimeState(const UDataLayer* InDataLayer, EDataLayerRuntimeState InState, bool bInIsRecursive)
 {
 	if (InDataLayer)
 	{
 		if (AWorldDataLayers* WorldDataLayers = GetWorld()->GetWorldDataLayers())
 		{
-			WorldDataLayers->SetDataLayerRuntimeState(FActorDataLayer(InDataLayer->GetFName()), InState);
+			WorldDataLayers->SetDataLayerRuntimeState(FActorDataLayer(InDataLayer->GetFName()), InState, bInIsRecursive);
 		}
 	}
 	else
@@ -109,11 +109,11 @@ void UDataLayerSubsystem::SetDataLayerRuntimeState(const UDataLayer* InDataLayer
 	}
 }
 
-void UDataLayerSubsystem::SetDataLayerRuntimeStateByName(const FName& InDataLayerName, EDataLayerRuntimeState InState)
+void UDataLayerSubsystem::SetDataLayerRuntimeStateByName(const FName& InDataLayerName, EDataLayerRuntimeState InState, bool bInIsRecursive)
 {
 	if (UDataLayer* DataLayer = GetDataLayerFromName(InDataLayerName))
 	{
-		SetDataLayerRuntimeState(DataLayer, InState);
+		SetDataLayerRuntimeState(DataLayer, InState, bInIsRecursive);
 	}
 	else
 	{
@@ -121,11 +121,11 @@ void UDataLayerSubsystem::SetDataLayerRuntimeStateByName(const FName& InDataLaye
 	}
 }
 
-void UDataLayerSubsystem::SetDataLayerRuntimeState(const FActorDataLayer& InDataLayer, EDataLayerRuntimeState InState)
+void UDataLayerSubsystem::SetDataLayerRuntimeState(const FActorDataLayer& InDataLayer, EDataLayerRuntimeState InState, bool bInIsRecursive)
 {
 	if (UDataLayer* DataLayer = GetDataLayerFromName(InDataLayer.Name))
 	{
-		SetDataLayerRuntimeState(DataLayer, InState);
+		SetDataLayerRuntimeState(DataLayer, InState, bInIsRecursive);
 	}
 	else
 	{
@@ -133,11 +133,11 @@ void UDataLayerSubsystem::SetDataLayerRuntimeState(const FActorDataLayer& InData
 	}
 }
 
-void UDataLayerSubsystem::SetDataLayerRuntimeStateByLabel(const FName& InDataLayerLabel, EDataLayerRuntimeState InState)
+void UDataLayerSubsystem::SetDataLayerRuntimeStateByLabel(const FName& InDataLayerLabel, EDataLayerRuntimeState InState, bool bInIsRecursive)
 {
 	if (UDataLayer* DataLayer = GetDataLayerFromLabel(InDataLayerLabel))
 	{
-		SetDataLayerRuntimeState(DataLayer, InState);
+		SetDataLayerRuntimeState(DataLayer, InState, bInIsRecursive);
 	}
 	else
 	{
diff --git a/Engine/Source/Runtime/Engine/Private/WorldPartition/DataLayer/WorldDataLayers.cpp b/Engine/Source/Runtime/Engine/Private/WorldPartition/DataLayer/WorldDataLayers.cpp
--- a/Engine/Source/Runtime/Engine/Private/WorldPartition/DataLayer/WorldDataLayers.cpp
+++ b/Engine/Source/Runtime/Engine/Private/WorldPartition/DataLayer/WorldDataLayers.cpp
@@ -128,7 +128,7 @@ void AWorldDataLayers::ResetDataLayerRuntimeStates()
 	RepEffectiveLoadedDataLayerNames.Reset();
 }
 
-void AWorldDataLayers::SetDataLayerRuntimeState(FActorDataLayer InDataLayer, EDataLayerRuntimeState InState)
+void AWorldDataLayers::SetDataLayerRuntimeState(FActorDataLayer InDataLayer, EDataLayerRuntimeState InState, bool bInIsRecursive)
 {
 	if (ensure(GetLocalRole() == ROLE_Authority))
 	{
@@ -166,6 +166,15 @@ void AWorldDataLayers::SetDataLayerRuntimeState(FActorDataLayer InDataLayer, EDa
 
 			ResolveEffectiveRuntimeState(DataLayer);
 		}
+
+		if (bInIsRecursive)
+		{
+			DataLayer->ForEachChild([this, InState, bInIsRecursive](const UDataLayer* Child)
+			{
+				SetDataLayerRuntimeState(Child->GetFName(), InState, bInIsRecursive);
+				return true;
+			});
+		}
 	}
 }
 
diff --git a/Engine/Source/Runtime/Engine/Public/WorldPartition/DataLayer/DataLayerSubsystem.h b/Engine/Source/Runtime/Engine/Public/WorldPartition/DataLayer/DataLayerSubsystem.h
--- a/Engine/Source/Runtime/Engine/Public/WorldPartition/DataLayer/DataLayerSubsystem.h
+++ b/Engine/Source/Runtime/Engine/Public/WorldPartition/DataLayer/DataLayerSubsystem.h
@@ -45,11 +45,11 @@ public:
 
 	/** Set the Data Layer state using its name. */
 	UFUNCTION(BlueprintCallable, Category = DataLayers, BlueprintAuthorityOnly)
-	void SetDataLayerRuntimeState(const FActorDataLayer& InDataLayer, EDataLayerRuntimeState InState);
+	void SetDataLayerRuntimeState(const FActorDataLayer& InDataLayer, EDataLayerRuntimeState InState, bool bInIsRecursive = false);
 
 	/** Set the Data Layer state using its label. */
 	UFUNCTION(BlueprintCallable, Category = DataLayers, BlueprintAuthorityOnly)
-	void SetDataLayerRuntimeStateByLabel(const FName& InDataLayerLabel, EDataLayerRuntimeState InState);
+	void SetDataLayerRuntimeStateByLabel(const FName& InDataLayerLabel, EDataLayerRuntimeState InState, bool bInIsRecursive = false);
 		
 	/** Get the Data Layer state using its name. */
 	UFUNCTION(BlueprintCallable, Category = DataLayers)
@@ -73,8 +73,8 @@ public:
 
 	//~ End Blueprint callable functions
 
-	void SetDataLayerRuntimeState(const UDataLayer* InDataLayer, EDataLayerRuntimeState InState);
-	void SetDataLayerRuntimeStateByName(const FName& InDataLayerName, EDataLayerRuntimeState InState);
+	void SetDataLayerRuntimeState(const UDataLayer* InDataLayer, EDataLayerRuntimeState InState, bool bInIsRecursive = false);
+	void SetDataLayerRuntimeStateByName(const FName& InDataLayerName, EDataLayerRuntimeState InState, bool bInIsRecursive = false);
 	EDataLayerRuntimeState GetDataLayerRuntimeState(const UDataLayer* InDataLayer) const;
 	EDataLayerRuntimeState GetDataLayerRuntimeStateByName(const FName& InDataLayerName) const;
 	EDataLayerRuntimeState GetDataLayerEffectiveRuntimeState(const UDataLayer* InDataLayer) const;
diff --git a/Engine/Source/Runtime/Engine/Public/WorldPartition/DataLayer/WorldDataLayers.h b/Engine/Source/Runtime/Engine/Public/WorldPartition/DataLayer/WorldDataLayers.h
--- a/Engine/Source/Runtime/Engine/Public/WorldPartition/DataLayer/WorldDataLayers.h
+++ b/Engine/Source/Runtime/Engine/Public/WorldPartition/DataLayer/WorldDataLayers.h
@@ -57,7 +57,7 @@ public:
 	void ForEachDataLayer(TFunctionRef<bool(UDataLayer*)> Func) const;
 
 	// DataLayer Runtime State
-	void SetDataLayerRuntimeState(FActorDataLayer InDataLayer, EDataLayerRuntimeState InState);
+	void SetDataLayerRuntimeState(FActorDataLayer InDataLayer, EDataLayerRuntimeState InState, bool bIsRecursive = false);
 	EDataLayerRuntimeState GetDataLayerRuntimeStateByName(FName InDataLayerName) const;
 	EDataLayerRuntimeState GetDataLayerEffectiveRuntimeStateByName(FName InDataLAyerName) const;
 	const TSet<FName>& GetEffectiveActiveDataLayerNames() const { return EffectiveActiveDataLayerNames; }
