When converting a world composition map to world partition, all actors are now going in a single runtime grid. Removed the creation of 1 spatial hash runtime grid per world composition layer.

diff --git a/Engine/Source/Editor/UnrealEd/Classes/Commandlets/WorldPartitionConvertCommandlet.h b/Engine/Source/Editor/UnrealEd/Classes/Commandlets/WorldPartitionConvertCommandlet.h
--- a/Engine/Source/Editor/UnrealEd/Classes/Commandlets/WorldPartitionConvertCommandlet.h
+++ b/Engine/Source/Editor/UnrealEd/Classes/Commandlets/WorldPartitionConvertCommandlet.h
@@ -16,7 +16,6 @@ UNREALED_API DECLARE_LOG_CATEGORY_EXTERN(LogWorldPartitionConvertCommandlet, Log
 
 class ULevelStreaming;
 class UWorldPartition;
-class UWorldComposition;
 class ULevelStreaming;
 
 USTRUCT()
@@ -58,7 +57,7 @@ protected:
 	virtual void OnWorldLoaded(UWorld* World);
 	virtual void ReadAdditionalTokensAndSwitches(const TArray<FString>& Tokens, const TArray<FString>& Switches) {}
 
-	UWorldPartition* CreateWorldPartition(class AWorldSettings* MainWorldSettings, UWorldComposition* WorldComposition) const;
+	UWorldPartition* CreateWorldPartition(class AWorldSettings* MainWorldSettings) const;
 	UWorld* LoadWorld(const FString& LevelToLoad);
 	ULevel* InitWorld(UWorld* World);
 
diff --git a/Engine/Source/Editor/UnrealEd/Private/Commandlets/WorldPartitionConvertCommandlet.cpp b/Engine/Source/Editor/UnrealEd/Private/Commandlets/WorldPartitionConvertCommandlet.cpp
--- a/Engine/Source/Editor/UnrealEd/Private/Commandlets/WorldPartitionConvertCommandlet.cpp
+++ b/Engine/Source/Editor/UnrealEd/Private/Commandlets/WorldPartitionConvertCommandlet.cpp
@@ -213,7 +213,7 @@ ULevel* UWorldPartitionConvertCommandlet::InitWorld(UWorld* World)
 	return World->PersistentLevel;
 }
 
-UWorldPartition* UWorldPartitionConvertCommandlet::CreateWorldPartition(AWorldSettings* MainWorldSettings, UWorldComposition* WorldComposition) const
+UWorldPartition* UWorldPartitionConvertCommandlet::CreateWorldPartition(AWorldSettings* MainWorldSettings) const
 {
 	TRACE_CPUPROFILER_EVENT_SCOPE(UWorldPartitionConvertCommandlet::CreateWorldPartition);
 
@@ -229,8 +229,6 @@ UWorldPartition* UWorldPartitionConvertCommandlet::CreateWorldPartition(AWorldSe
 
 	WorldPartition->EditorHash->Initialize();
 	
-	WorldPartition->RuntimeHash->ImportFromWorldComposition(WorldComposition);
-
 	return WorldPartition;
 }
 
@@ -307,22 +305,7 @@ bool UWorldPartitionConvertCommandlet::PrepareStreamingLevelForConversion(const
 			}
 		}
 	}
-	else
-	{
-		UE_LOG(LogWorldPartitionConvertCommandlet, Log, TEXT("Converting dynamic streaming level %s"), *StreamingLevel->GetWorldAssetPackageName());
-
-		const UWorldPartitionRuntimeHash* RuntimeHash = WorldPartition->RuntimeHash;
-		check(RuntimeHash);
-
-		for (AActor* Actor : SubLevel->Actors)
-		{
-			if (Actor)
-			{
-				Actor->SetRuntimeGrid(RuntimeHash->GetActorRuntimeGrid(Actor));
-			}
-		}
-	}
-
+	
 	return true;
 }
 
@@ -839,7 +822,7 @@ int32 UWorldPartitionConvertCommandlet::Main(const FString& Params)
 	}
 
 	// Setup the world partition object
-	UWorldPartition* WorldPartition = CreateWorldPartition(MainWorldSettings, MainWorld->WorldComposition);
+	UWorldPartition* WorldPartition = CreateWorldPartition(MainWorldSettings);
 	if (!WorldPartition)
 	{
 		return 1;
diff --git a/Engine/Source/Runtime/Engine/Private/WorldPartition/WorldPartitionRuntimeSpatialHash.cpp b/Engine/Source/Runtime/Engine/Private/WorldPartition/WorldPartitionRuntimeSpatialHash.cpp
--- a/Engine/Source/Runtime/Engine/Private/WorldPartition/WorldPartitionRuntimeSpatialHash.cpp
+++ b/Engine/Source/Runtime/Engine/Private/WorldPartition/WorldPartitionRuntimeSpatialHash.cpp
@@ -641,18 +641,6 @@ void UWorldPartitionRuntimeSpatialHash::DrawPreview() const
 	GridPreviewer.Draw(GetWorld(), Grids, bPreviewGrids);
 }
 
-FName UWorldPartitionRuntimeSpatialHash::GetActorRuntimeGrid(const AActor* Actor) const
-{
-	if (ULevel* Level = Actor ? Actor->GetLevel() : nullptr)
-	{
-		if (const FName* ActorRuntimeGrid = WorldCompositionStreamingLevelToRuntimeGrid.Find(FLevelUtils::FindStreamingLevel(Level)))
-		{
-			return *ActorRuntimeGrid;
-		}
-	}
-	return Super::GetActorRuntimeGrid(Actor);
-}
-
 void UWorldPartitionRuntimeSpatialHash::SetDefaultValues()
 {
 	check(!Grids.Num());
@@ -664,42 +652,6 @@ void UWorldPartitionRuntimeSpatialHash::SetDefaultValues()
 	MainGrid.DebugColor = FLinearColor::Gray;
 }
 
-void UWorldPartitionRuntimeSpatialHash::ImportFromWorldComposition(UWorldComposition* WorldComposition)
-{
-	check(IsRunningCommandlet());
-
-	if (WorldComposition)
-	{
-		const TArray<FWorldTileLayer> WorldCompositionTileLayers = WorldComposition->GetDistanceDependentLayers();
-		for (const FWorldTileLayer& Layer : WorldCompositionTileLayers)
-		{
-			FName GridName = FName(Layer.Name);
-			FSpatialHashRuntimeGrid* Grid = Algo::FindByPredicate(Grids, [GridName](const FSpatialHashRuntimeGrid& Grid) { return Grid.GridName == GridName; });
-			if (!Grid)
-			{
-				Grid = &Grids.AddDefaulted_GetRef();
-				Grid->GridName = GridName;
-				Grid->CellSize = 12800;
-				Grid->DebugColor = FLinearColor::MakeRandomColor();
-			}
-			// World Composition Layer Streaming Distance always wins over existing value (config file)
-			Grid->LoadingRange = Layer.StreamingDistance;
-		}
-
-		const UWorldComposition::FTilesList& Tiles = WorldComposition->GetTilesList();
-		for (int32 TileIdx = 0; TileIdx < Tiles.Num(); TileIdx++)
-		{
-			const FWorldCompositionTile& Tile = Tiles[TileIdx];
-			ULevelStreaming* StreamingLevel = WorldComposition->TilesStreaming[TileIdx];
-			if (StreamingLevel && WorldComposition->IsDistanceDependentLevel(Tile.PackageName))
-			{
-				// Map WorldComposition tiles streaming level to Runtime Grid
-				WorldCompositionStreamingLevelToRuntimeGrid.Add(StreamingLevel, FName(Tile.Info.Layer.Name));
-			}
-		}
-	}
-}
-
 bool UWorldPartitionRuntimeSpatialHash::GenerateStreaming(UWorldPartitionStreamingPolicy* StreamingPolicy, const FActorClusterContext& ActorClusterContext, TArray<FString>* OutPackagesToGenerate)
 {
 	TRACE_CPUPROFILER_EVENT_SCOPE(UWorldPartitionRuntimeSpatialHash::GenerateStreaming);
diff --git a/Engine/Source/Runtime/Engine/Public/WorldPartition/WorldPartitionRuntimeHash.h b/Engine/Source/Runtime/Engine/Public/WorldPartition/WorldPartitionRuntimeHash.h
--- a/Engine/Source/Runtime/Engine/Public/WorldPartition/WorldPartitionRuntimeHash.h
+++ b/Engine/Source/Runtime/Engine/Public/WorldPartition/WorldPartitionRuntimeHash.h
@@ -32,14 +32,12 @@ class ENGINE_API UWorldPartitionRuntimeHash : public UObject
 
 #if WITH_EDITOR
 	virtual void SetDefaultValues() {}
-	virtual void ImportFromWorldComposition(class UWorldComposition* WorldComposition) {}
 	virtual bool PopulateGeneratedPackageForCook(UPackage* InPackage, const FString& InPackageRelativePath) { return false; }
 	virtual bool FinalizeGeneratorPackageForCook(const TArray<ICookPackageSplitter::FGeneratedPackageForPreSave>& InGeneratedPackages) { return false; }
 	virtual bool GenerateStreaming(class UWorldPartitionStreamingPolicy* StreamingPolicy, const FActorClusterContext& ActorClusterContext, TArray<FString>* OutPackagesToGenerate) { return false; }
 	virtual void FlushStreaming() {}
 	virtual bool GenerateHLOD(ISourceControlHelper* SourceControlHelper, FActorClusterContext& ActorClusterContext, bool bCreateActorsOnly) { return false; }
 	virtual bool GenerateNavigationData(const FBox& LoadedBounds) { return false; }
-	virtual FName GetActorRuntimeGrid(const AActor* Actor) const { return NAME_None; }
 	virtual void DrawPreview() const {}
 
 	void CheckForErrors() const;
diff --git a/Engine/Source/Runtime/Engine/Public/WorldPartition/WorldPartitionRuntimeSpatialHash.h b/Engine/Source/Runtime/Engine/Public/WorldPartition/WorldPartitionRuntimeSpatialHash.h
--- a/Engine/Source/Runtime/Engine/Public/WorldPartition/WorldPartitionRuntimeSpatialHash.h
+++ b/Engine/Source/Runtime/Engine/Public/WorldPartition/WorldPartitionRuntimeSpatialHash.h
@@ -228,13 +228,11 @@ public:
 
 #if WITH_EDITOR
 	virtual void SetDefaultValues() override;
-	virtual void ImportFromWorldComposition(class UWorldComposition* WorldComposition) override;
 	virtual bool PopulateGeneratedPackageForCook(UPackage* InPackage, const FString& InPackageRelativePath) override;
 	virtual bool FinalizeGeneratorPackageForCook(const TArray<ICookPackageSplitter::FGeneratedPackageForPreSave>& InGeneratedPackages) override;
 	virtual void FlushStreaming() override;
 	virtual bool GenerateHLOD(ISourceControlHelper* SourceControlHelper, FActorClusterContext& ActorClusterContext, bool bCreateActorsOnly) override;
 	virtual bool GenerateNavigationData(const FBox& LoadedBounds) override;
-	virtual FName GetActorRuntimeGrid(const AActor* Actor) const override;
 	virtual void DrawPreview() const override;
 
 	FName GetCellName(FName InGridName, const FIntVector& InCellGlobalCoord, const FDataLayersID& InDataLayerID) const;
@@ -258,9 +256,6 @@ protected:
 	virtual bool GenerateStreaming(class UWorldPartitionStreamingPolicy* StreamingPolicy, const FActorClusterContext& ActorClusterContext, TArray<FString>* OutPackagesToGenerate = nullptr) override;
 #endif
 
-	// Used to convert from a world using World Composition
-	TMap<class ULevelStreaming*, FName> WorldCompositionStreamingLevelToRuntimeGrid;
-
 private:
 	/** Console command used to change loading range for a given streaming grid */
 	static class FAutoConsoleCommand OverrideLoadingRangeCommand;
