World Partition: cleaned code that adds actor clusters into cells.

diff --git a/Engine/Source/Runtime/Engine/Private/WorldPartition/RuntimeSpatialHash/RuntimeSpatialHashGridHelper.cpp b/Engine/Source/Runtime/Engine/Private/WorldPartition/RuntimeSpatialHash/RuntimeSpatialHashGridHelper.cpp
--- a/Engine/Source/Runtime/Engine/Private/WorldPartition/RuntimeSpatialHash/RuntimeSpatialHashGridHelper.cpp
+++ b/Engine/Source/Runtime/Engine/Private/WorldPartition/RuntimeSpatialHash/RuntimeSpatialHashGridHelper.cpp
@@ -163,82 +163,76 @@ FSquare2DGridHelper GetPartitionedActors(const UWorldPartition* WorldPartition,
 	for (const FActorClusterInstance* ClusterInstance : GridActors)
 	{
 		const FActorCluster* ActorCluster = ClusterInstance->Cluster;
+		const FBox2D ClusterBounds(FVector2D(ClusterInstance->Bounds.Min), FVector2D(ClusterInstance->Bounds.Max));
 		check(ActorCluster && ActorCluster->Actors.Num() > 0);
-		EActorGridPlacement GridPlacement = ActorCluster->GridPlacement;
+		EActorGridPlacement GridPlacement = ActorCluster->GridPlacement;		
+		FSquare2DGridHelper::FGridLevel::FGridCell* GridCell = nullptr;
 
 		switch (GridPlacement)
 		{
 		case EActorGridPlacement::Location:
 		{
 			FIntVector2 CellCoords;
-			check(ActorCluster->Actors.Num() == 1);
-			const FGuid& ActorGuid = *ActorCluster->Actors.CreateConstIterator();
-			FActorInstance ActorInstance(ActorGuid, ClusterInstance->ContainerInstance);
-			if (ensure(PartitionedActors.GetLowestLevel().GetCellCoords(FVector2D(ActorInstance.GetOrigin()), CellCoords)))
+			if (PartitionedActors.GetLowestLevel().GetCellCoords(ClusterBounds.GetCenter(), CellCoords))
 			{
-				PartitionedActors.GetLowestLevel().GetCell(CellCoords).AddActor(MoveTemp(ActorInstance), ClusterInstance->DataLayers);
-			}
-			else
-			{
-				GridPlacement = EActorGridPlacement::AlwaysLoaded;
-				PartitionedActors.GetAlwaysLoadedCell().AddActor(MoveTemp(ActorInstance), ClusterInstance->DataLayers);
+				GridCell = &PartitionedActors.GetLowestLevel().GetCell(CellCoords);
 			}
 			break;
 		}
 		case EActorGridPlacement::Bounds:
 		{
 			// Find grid level cell that encompasses the actor cluster and put actors in it.
-			const FBox2D ClusterBounds(FVector2D(ClusterInstance->Bounds.Min), FVector2D(ClusterInstance->Bounds.Max));
 			const FVector2D ClusterSize = ClusterBounds.GetSize();
 			const float MinRequiredCellExtent = FMath::Max(ClusterSize.X, ClusterSize.Y);
 			const int32 FirstPotentialGridLevel = FMath::Max(FMath::CeilToFloat(FMath::Log2(MinRequiredCellExtent / (float)PartitionedActors.CellSize)), 0);
 
-			bool bFoundCell = false;
 			for (int32 GridLevelIndex = FirstPotentialGridLevel; GridLevelIndex < PartitionedActors.Levels.Num(); GridLevelIndex++)
 			{
 				FSquare2DGridHelper::FGridLevel& GridLevel = PartitionedActors.Levels[GridLevelIndex];
 
 				if (GridLevel.GetNumIntersectingCells(ClusterInstance->Bounds) == 1)
 				{
-					GridLevel.ForEachIntersectingCells(ClusterInstance->Bounds, [&GridLevel, ActorCluster, ClusterInstance, &bFoundCell](const FIntVector2& Coords)
+					GridLevel.ForEachIntersectingCells(ClusterInstance->Bounds, [&GridLevel, ActorCluster, ClusterInstance, &GridCell](const FIntVector2& Coords)
 					{
-						check(!bFoundCell);
-						GridLevel.GetCell(Coords).AddActors(ActorCluster->Actors, ClusterInstance->ContainerInstance, ClusterInstance->DataLayers);						
-						bFoundCell = true;
+						check(!GridCell);
+						GridCell = &GridLevel.GetCell(Coords);
 					});
 
 					break;
 				}
 			}
-
-			if (!bFoundCell)
-			{
-				GridPlacement = EActorGridPlacement::AlwaysLoaded;
-				PartitionedActors.GetAlwaysLoadedCell().AddActors(ActorCluster->Actors, ClusterInstance->ContainerInstance, ClusterInstance->DataLayers);
-
-				UE_LOG(LogWorldPartition, Warning, TEXT("Cluster of %d actors was promoted always loaded, BV of [%d x %d] (meters)"), 
-					ActorCluster->Actors.Num(),
-					(int)(0.01f * ClusterInstance->Bounds.GetSize().X),
-					(int)(0.01f * ClusterInstance->Bounds.GetSize().Y));
-				for (const FGuid& ActorGuid : ActorCluster->Actors)
-				{
-					const FWorldPartitionActorDescView& ActorDescView = ClusterInstance->ContainerInstance->GetActorDescView(ActorGuid);
-					UE_LOG(LogWorldPartition, Warning, TEXT("   - Actor: %s (%s)"), *ActorDescView.GetActorPath().ToString(), *ActorGuid.ToString(EGuidFormats::UniqueObjectGuid));
-					UE_LOG(LogWorldPartition, Warning, TEXT("         Package: %s"), *ActorDescView.GetActorPackage().ToString());
-					UE_LOG(LogWorldPartition, Warning, TEXT("         Container (%s): %s"), *ClusterInstance->ContainerInstance->ID.ToString(), *ClusterInstance->ContainerInstance->Container->GetContainerPackage().ToString())
-				}
-			}
 			break;
 		}
 		case EActorGridPlacement::AlwaysLoaded:
 		{
-			PartitionedActors.GetAlwaysLoadedCell().AddActors(ActorCluster->Actors, ClusterInstance->ContainerInstance, ClusterInstance->DataLayers);
+			GridCell = &PartitionedActors.GetAlwaysLoadedCell();
 			break;
 		}
 		default:
 			check(0);
 		}
 
+		if (!GridCell)
+		{
+			UE_LOG(LogWorldPartition, Warning, TEXT("Cluster of %d actors was promoted always loaded, BV of [%d x %d] (meters)"), 
+				ActorCluster->Actors.Num(),
+				(int)(0.01f * ClusterInstance->Bounds.GetSize().X),
+				(int)(0.01f * ClusterInstance->Bounds.GetSize().Y));
+			for (const FGuid& ActorGuid : ActorCluster->Actors)
+			{
+				const FWorldPartitionActorDescView& ActorDescView = ClusterInstance->ContainerInstance->GetActorDescView(ActorGuid);
+				UE_LOG(LogWorldPartition, Warning, TEXT("   - Actor: %s (%s)"), *ActorDescView.GetActorPath().ToString(), *ActorGuid.ToString(EGuidFormats::UniqueObjectGuid));
+				UE_LOG(LogWorldPartition, Warning, TEXT("         Package: %s"), *ActorDescView.GetActorPackage().ToString());
+				UE_LOG(LogWorldPartition, Warning, TEXT("         Container (%s): %s"), *ClusterInstance->ContainerInstance->ID.ToString(), *ClusterInstance->ContainerInstance->Container->GetContainerPackage().ToString())
+			}
+
+			GridPlacement = EActorGridPlacement::AlwaysLoaded;
+			GridCell = &PartitionedActors.GetAlwaysLoadedCell();
+		}
+
+		check(GridCell);
+		GridCell->AddActors(ActorCluster->Actors, ClusterInstance->ContainerInstance, ClusterInstance->DataLayers);
+
 		if (UE_LOG_ACTIVE(LogWorldPartition, Verbose))
 		{
 			if (ActorCluster->Actors.Num() > 1)
diff --git a/Engine/Source/Runtime/Engine/Private/WorldPartition/RuntimeSpatialHash/RuntimeSpatialHashGridHelper.h b/Engine/Source/Runtime/Engine/Private/WorldPartition/RuntimeSpatialHash/RuntimeSpatialHashGridHelper.h
--- a/Engine/Source/Runtime/Engine/Private/WorldPartition/RuntimeSpatialHash/RuntimeSpatialHashGridHelper.h
+++ b/Engine/Source/Runtime/Engine/Private/WorldPartition/RuntimeSpatialHash/RuntimeSpatialHashGridHelper.h
@@ -282,22 +282,19 @@ struct FSquare2DGridHelper
 				: Coords(InCoords)
 			{}
 
-			void AddActor(FActorInstance ActorInstance, const TArray<const UDataLayer*>& InDataLayers)
-			{
-				FDataLayersID DataLayersID = FDataLayersID(InDataLayers);
-				FGridCellDataChunk* ActorDataChunk = Algo::FindByPredicate(DataChunks, [&](FGridCellDataChunk& InDataChunk) { return InDataChunk.GetDataLayersID() == DataLayersID; });
-				if (!ActorDataChunk)
-				{
-					ActorDataChunk = &DataChunks.Emplace_GetRef(InDataLayers);
-				}
-				ActorDataChunk->AddActor(MoveTemp(ActorInstance));
-			}
-
 			void AddActors(const TSet<FGuid>& InActors, const FActorContainerInstance* ContainerInstance, const TArray<const UDataLayer*>& InDataLayers)
 			{
 				for (const FGuid& Actor : InActors)
 				{
-					AddActor(FActorInstance(Actor, ContainerInstance), InDataLayers);
+					FActorInstance ActorInstance(Actor, ContainerInstance);
+
+					FDataLayersID DataLayersID = FDataLayersID(InDataLayers);
+					FGridCellDataChunk* ActorDataChunk = Algo::FindByPredicate(DataChunks, [&](FGridCellDataChunk& InDataChunk) { return InDataChunk.GetDataLayersID() == DataLayersID; });
+					if (!ActorDataChunk)
+					{
+						ActorDataChunk = &DataChunks.Emplace_GetRef(InDataLayers);
+					}
+					ActorDataChunk->AddActor(MoveTemp(ActorInstance));
 				}
 			}
 
