UWorld: Add geometry to an empty default brush during InitWorld()

diff --git a/Engine/Source/Runtime/Engine/Private/World.cpp b/Engine/Source/Runtime/Engine/Private/World.cpp
--- a/Engine/Source/Runtime/Engine/Private/World.cpp
+++ b/Engine/Source/Runtime/Engine/Private/World.cpp
@@ -1710,18 +1710,63 @@ void UWorld::InitWorld(const InitializationValues IVS)
 	ABrush* DefaultBrush = PersistentLevel->Actors.Num()<2 ? NULL : Cast<ABrush>(PersistentLevel->Actors[1]);
 	if (GIsEditor)
 	{
-		if (!DefaultBrush || !DefaultBrush->IsStaticBrush() || DefaultBrush->BrushType!=Brush_Default || !DefaultBrush->GetBrushComponent() || !DefaultBrush->Brush)
+		if (!DefaultBrush || !DefaultBrush->IsStaticBrush() || DefaultBrush->BrushType != Brush_Default || !DefaultBrush->GetBrushComponent() ||
+			!DefaultBrush->Brush || !DefaultBrush->Brush->Polys || DefaultBrush->Brush->Polys->Element.IsEmpty())
 		{
 			// Spawn the default brush.
 			DefaultBrush = SpawnBrush();
 			check(DefaultBrush->GetBrushComponent());
 			DefaultBrush->Brush = NewObject<UModel>(DefaultBrush->GetOuter(), TEXT("Brush"));
-			DefaultBrush->Brush->Initialize(DefaultBrush, 1);
+			DefaultBrush->Brush->Initialize(DefaultBrush, true);
 			DefaultBrush->GetBrushComponent()->Brush = DefaultBrush->Brush;
 			DefaultBrush->SetNotForClientOrServer();
 			DefaultBrush->Brush->SetFlags( RF_Transactional );
 			DefaultBrush->Brush->Polys->SetFlags( RF_Transactional );
 
+			// Create cube geometry.
+			// We effectively replicate what is done in UEditorEngine::InitBuilderBrush() since we can't just use this function or the underlying
+			// function UCubeBuilder::BuildCube() here directly.
+			{
+				constexpr float HalfSize = 128.0f;
+				const FVector3f Vertices[8] = {
+					{-HalfSize, -HalfSize, -HalfSize},
+					{-HalfSize, -HalfSize,  HalfSize},
+					{-HalfSize,  HalfSize, -HalfSize},
+					{-HalfSize,  HalfSize,  HalfSize},
+					{ HalfSize, -HalfSize, -HalfSize},
+					{ HalfSize, -HalfSize,  HalfSize},
+					{ HalfSize,  HalfSize, -HalfSize},
+					{ HalfSize,  HalfSize,  HalfSize}
+				};
+
+				constexpr int32 Indices[24] = {
+					0, 1, 3, 2,
+					2, 3, 7, 6,
+					6, 7, 5, 4,
+					4, 5, 1, 0,
+					3, 1, 5, 7,
+					0, 2, 6, 4
+				};
+
+				for (int32 i = 0; i < 6; ++i)
+				{
+					FPoly Poly;
+					Poly.Init();
+					Poly.Base = Vertices[0];
+
+					for (int32 j = 0; j < 4; ++j)
+					{
+						new(Poly.Vertices) FVector3f(Vertices[Indices[i * 4 + j]]);
+					}
+					if (Poly.Finalize(DefaultBrush, 1) == 0)
+					{
+						new(DefaultBrush->Brush->Polys->Element)FPoly(Poly);
+					}
+				}
+
+				DefaultBrush->Brush->BuildBound();
+			}
+
 			// The default brush is legacy but has to exist for some old bsp operations.  However it should not be interacted with in the editor. 
 			DefaultBrush->SetIsTemporarilyHiddenInEditor(true);
 
