UWorld: Move functionality for repairing default brush during InitWorld() into new function RepairDefaultBrush()

diff --git a/Engine/Source/Runtime/Engine/Classes/Engine/World.h b/Engine/Source/Runtime/Engine/Classes/Engine/World.h
--- a/Engine/Source/Runtime/Engine/Classes/Engine/World.h
+++ b/Engine/Source/Runtime/Engine/Classes/Engine/World.h
@@ -1476,6 +1476,11 @@ private:
 	void RepairChaosActors();
 #endif
 
+#if WITH_EDITOR
+	/** Utility function to make sure there is a valid default builder brush */
+	void RepairDefaultBrush();
+#endif
+
 	/** Gameplay timers. */
 	class FTimerManager* TimerManager;
 
diff --git a/Engine/Source/Runtime/Engine/Private/World.cpp b/Engine/Source/Runtime/Engine/Private/World.cpp
--- a/Engine/Source/Runtime/Engine/Private/World.cpp
+++ b/Engine/Source/Runtime/Engine/Private/World.cpp
@@ -1577,6 +1577,115 @@
 	check(GetWorldSettings());
 }
 
+#if WITH_EDITOR
+void UWorld::RepairDefaultBrush()
+{
+	// See whether we're missing the default brush. It was possible in earlier builds to accidentally delete the default
+	// brush of sublevels so we simply spawn a new one if we encounter it missing.
+	ABrush* DefaultBrush = PersistentLevel->Actors.Num()<2 ? NULL : Cast<ABrush>(PersistentLevel->Actors[1]);
+	if (GIsEditor)
+	{
+		if (!DefaultBrush || !DefaultBrush->IsStaticBrush() || DefaultBrush->BrushType != Brush_Default || !DefaultBrush->GetBrushComponent() ||
+			!DefaultBrush->Brush || !DefaultBrush->Brush->Polys || DefaultBrush->Brush->Polys->Element.IsEmpty())
+		{
+			// Spawn the default brush.
+			DefaultBrush = SpawnBrush();
+			check(DefaultBrush->GetBrushComponent());
+			DefaultBrush->Brush = NewObject<UModel>(DefaultBrush->GetOuter(), TEXT("Brush"));
+			DefaultBrush->Brush->Initialize(DefaultBrush, true);
+			DefaultBrush->GetBrushComponent()->Brush = DefaultBrush->Brush;
+			DefaultBrush->SetNotForClientOrServer();
+			DefaultBrush->Brush->SetFlags( RF_Transactional );
+			DefaultBrush->Brush->Polys->SetFlags( RF_Transactional );
+
+			// Create cube geometry.
+			// We effectively replicate what is done in UEditorEngine::InitBuilderBrush() since we can't just use this function or the underlying
+			// function UCubeBuilder::BuildCube() here directly.
+			{
+				constexpr float HalfSize = 128.0f;
+				static const FVector3f Vertices[8] = {
+					{-HalfSize, -HalfSize, -HalfSize},
+					{-HalfSize, -HalfSize,  HalfSize},
+					{-HalfSize,  HalfSize, -HalfSize},
+					{-HalfSize,  HalfSize,  HalfSize},
+					{ HalfSize, -HalfSize, -HalfSize},
+					{ HalfSize, -HalfSize,  HalfSize},
+					{ HalfSize,  HalfSize, -HalfSize},
+					{ HalfSize,  HalfSize,  HalfSize}
+				};
+
+				constexpr int32 Indices[24] = {
+					0, 1, 3, 2,
+					2, 3, 7, 6,
+					6, 7, 5, 4,
+					4, 5, 1, 0,
+					3, 1, 5, 7,
+					0, 2, 6, 4
+				};
+
+				for (int32 i = 0; i < 6; ++i)
+				{
+					FPoly Poly;
+					Poly.Init();
+					Poly.Base = Vertices[0];
+
+					for (int32 j = 0; j < 4; ++j)
+					{
+						new(Poly.Vertices) FVector3f(Vertices[Indices[i * 4 + j]]);
+					}
+					if (Poly.Finalize(DefaultBrush, 1) == 0)
+					{
+						new(DefaultBrush->Brush->Polys->Element)FPoly(Poly);
+					}
+				}
+
+				DefaultBrush->Brush->BuildBound();
+			}
+
+			// The default brush is legacy but has to exist for some old bsp operations.  However it should not be interacted with in the editor. 
+			DefaultBrush->SetIsTemporarilyHiddenInEditor(true);
+
+			// Find the index in the array the default brush has been spawned at. Not necessarily
+			// the last index as the code might spawn the default physics volume afterwards.
+			const int32 DefaultBrushActorIndex = PersistentLevel->Actors.Find( DefaultBrush );
+
+			// The default brush needs to reside at index 1.
+			Exchange(PersistentLevel->Actors[1],PersistentLevel->Actors[DefaultBrushActorIndex]);
+
+			// Re-sort actor list as we just shuffled things around.
+			PersistentLevel->SortActorList();
+		}
+		else
+		{
+			// Ensure that the Brush and BrushComponent both point to the same model
+			DefaultBrush->GetBrushComponent()->Brush = DefaultBrush->Brush;
+		}
+
+		// Reset the lightmass settings on the default brush; they can't be edited by the user but could have
+		// been tainted if the map was created during a window where the memory was uninitialized
+		if (DefaultBrush->Brush != NULL)
+		{
+			UModel* Model = DefaultBrush->Brush;
+
+			const FLightmassPrimitiveSettings DefaultSettings;
+
+			for (int32 i = 0; i < Model->LightmassSettings.Num(); ++i)
+			{
+				Model->LightmassSettings[i] = DefaultSettings;
+			}
+
+			if (Model->Polys != NULL) 
+			{
+				for (int32 i = 0; i < Model->Polys->Element.Num(); ++i)
+				{
+					Model->Polys->Element[i].LightmassSettings = DefaultSettings;
+				}
+			}
+		}
+	}
+}
+#endif
+
 void UWorld::InitWorld(const InitializationValues IVS)
 {
 	if (!ensure(!bIsWorldInitialized))
@@ -1705,109 +1814,7 @@
 	bDoDelayedUpdateCullDistanceVolumes = false;
 
 #if WITH_EDITOR
-	// See whether we're missing the default brush. It was possible in earlier builds to accidentally delete the default
-	// brush of sublevels so we simply spawn a new one if we encounter it missing.
-	ABrush* DefaultBrush = PersistentLevel->Actors.Num()<2 ? NULL : Cast<ABrush>(PersistentLevel->Actors[1]);
-	if (GIsEditor)
-	{
-		if (!DefaultBrush || !DefaultBrush->IsStaticBrush() || DefaultBrush->BrushType != Brush_Default || !DefaultBrush->GetBrushComponent() ||
-			!DefaultBrush->Brush || !DefaultBrush->Brush->Polys || DefaultBrush->Brush->Polys->Element.IsEmpty())
-		{
-			// Spawn the default brush.
-			DefaultBrush = SpawnBrush();
-			check(DefaultBrush->GetBrushComponent());
-			DefaultBrush->Brush = NewObject<UModel>(DefaultBrush->GetOuter(), TEXT("Brush"));
-			DefaultBrush->Brush->Initialize(DefaultBrush, true);
-			DefaultBrush->GetBrushComponent()->Brush = DefaultBrush->Brush;
-			DefaultBrush->SetNotForClientOrServer();
-			DefaultBrush->Brush->SetFlags( RF_Transactional );
-			DefaultBrush->Brush->Polys->SetFlags( RF_Transactional );
+	RepairDefaultBrush();
 
-			// Create cube geometry.
-			// We effectively replicate what is done in UEditorEngine::InitBuilderBrush() since we can't just use this function or the underlying
-			// function UCubeBuilder::BuildCube() here directly.
-			{
-				constexpr float HalfSize = 128.0f;
-				const FVector3f Vertices[8] = {
-					{-HalfSize, -HalfSize, -HalfSize},
-					{-HalfSize, -HalfSize,  HalfSize},
-					{-HalfSize,  HalfSize, -HalfSize},
-					{-HalfSize,  HalfSize,  HalfSize},
-					{ HalfSize, -HalfSize, -HalfSize},
-					{ HalfSize, -HalfSize,  HalfSize},
-					{ HalfSize,  HalfSize, -HalfSize},
-					{ HalfSize,  HalfSize,  HalfSize}
-				};
-
-				constexpr int32 Indices[24] = {
-					0, 1, 3, 2,
-					2, 3, 7, 6,
-					6, 7, 5, 4,
-					4, 5, 1, 0,
-					3, 1, 5, 7,
-					0, 2, 6, 4
-				};
-
-				for (int32 i = 0; i < 6; ++i)
-				{
-					FPoly Poly;
-					Poly.Init();
-					Poly.Base = Vertices[0];
-
-					for (int32 j = 0; j < 4; ++j)
-					{
-						new(Poly.Vertices) FVector3f(Vertices[Indices[i * 4 + j]]);
-					}
-					if (Poly.Finalize(DefaultBrush, 1) == 0)
-					{
-						new(DefaultBrush->Brush->Polys->Element)FPoly(Poly);
-					}
-				}
-
-				DefaultBrush->Brush->BuildBound();
-			}
-
-			// The default brush is legacy but has to exist for some old bsp operations.  However it should not be interacted with in the editor. 
-			DefaultBrush->SetIsTemporarilyHiddenInEditor(true);
-
-			// Find the index in the array the default brush has been spawned at. Not necessarily
-			// the last index as the code might spawn the default physics volume afterwards.
-			const int32 DefaultBrushActorIndex = PersistentLevel->Actors.Find( DefaultBrush );
-
-			// The default brush needs to reside at index 1.
-			Exchange(PersistentLevel->Actors[1],PersistentLevel->Actors[DefaultBrushActorIndex]);
-
-			// Re-sort actor list as we just shuffled things around.
-			PersistentLevel->SortActorList();
-		}
-		else
-		{
-			// Ensure that the Brush and BrushComponent both point to the same model
-			DefaultBrush->GetBrushComponent()->Brush = DefaultBrush->Brush;
-		}
-
-		// Reset the lightmass settings on the default brush; they can't be edited by the user but could have
-		// been tainted if the map was created during a window where the memory was uninitialized
-		if (DefaultBrush->Brush != NULL)
-		{
-			UModel* Model = DefaultBrush->Brush;
-
-			const FLightmassPrimitiveSettings DefaultSettings;
-
-			for (int32 i = 0; i < Model->LightmassSettings.Num(); ++i)
-			{
-				Model->LightmassSettings[i] = DefaultSettings;
-			}
-
-			if (Model->Polys != NULL) 
-			{
-				for (int32 i = 0; i < Model->Polys->Element.Num(); ++i)
-				{
-					Model->Polys->Element[i].LightmassSettings = DefaultSettings;
-				}
-			}
-		}
-	}
-
 	// invalidate lighting if VT is enabled but no valid VT data is present or VT is disabled and no valid non-VT data is present.
 	for (auto Level : Levels) //Note: PersistentLevel is part of this array
