Fixed assert in UInheritableComponentHandler::CreateOverridenComponentTemplate

The NewObject call could assert when trying to restore a UnnecessaryComponents during a compile-on-load, as the code for discarding instances invalid to be re-used by NewObject wasn't considering the pending load flags.

diff --git a/Engine/Source/Runtime/Engine/Private/InheritableComponentHandler.cpp b/Engine/Source/Runtime/Engine/Private/InheritableComponentHandler.cpp
--- a/Engine/Source/Runtime/Engine/Private/InheritableComponentHandler.cpp
+++ b/Engine/Source/Runtime/Engine/Private/InheritableComponentHandler.cpp
@@ -120,7 +120,7 @@ UActorComponent* UInheritableComponentHandler::CreateOverridenComponentTemplate(
 	// If we find an existing object with our name that the object recycling system won't allow for we need to deal with it 
 	// or else the NewObject call below will fatally assert
 	UObject* ExistingObj = FindObjectFast<UObject>(GetOuter(), NewComponentTemplateName);
-	if (ExistingObj && !ExistingObj->GetClass()->IsChildOf(BestArchetype->GetClass()))
+	if (ExistingObj && (!ExistingObj->GetClass()->IsChildOf(BestArchetype->GetClass()) || ExistingObj->HasAnyFlags(RF_NeedLoad | RF_NeedPostLoad)))
 	{
 		// If this isn't an unnecessary component there is something else we need to investigate
 		// but if it is, just consign it to oblivion as its purpose is no longer required with the allocation
