World Partition Streaming Generation
- Actor descriptor views can now be set to have invalid data layers.
- Moved data layers and grid placement validations of referenced objects from clustering code to streaming generation validation pass.

diff --git a/Engine/Source/Runtime/Engine/Private/WorldPartition/WorldPartitionActorCluster.cpp b/Engine/Source/Runtime/Engine/Private/WorldPartition/WorldPartitionActorCluster.cpp
--- a/Engine/Source/Runtime/Engine/Private/WorldPartition/WorldPartitionActorCluster.cpp
+++ b/Engine/Source/Runtime/Engine/Private/WorldPartition/WorldPartitionActorCluster.cpp
@@ -45,89 +45,21 @@ FActorCluster::FActorCluster(UWorld* InWorld, const FWorldPartitionActorDescView
 
 void FActorCluster::Add(const FActorCluster& InActorCluster, const TMap<FGuid, FWorldPartitionActorDescView>& InActorDescViewMap)
 {
-	// Merge RuntimeGrid
-	if (RuntimeGrid != InActorCluster.RuntimeGrid)
-	{
-		RuntimeGrid = NAME_None;
-	}
+	check(RuntimeGrid == InActorCluster.RuntimeGrid);
+	check(DataLayersID == InActorCluster.DataLayersID);
+	check((GridPlacement == EActorGridPlacement::AlwaysLoaded) == (InActorCluster.GridPlacement == EActorGridPlacement::AlwaysLoaded));
 
-	// Merge Bounds
+	// Merge bounds
 	Bounds += InActorCluster.Bounds;
 
-	// Merge GridPlacement
+	// Merge Actors
+	Actors.Append(InActorCluster.Actors);
+
+	// Set grid placement
 	if (GridPlacement != EActorGridPlacement::AlwaysLoaded)
 	{
-		GridPlacement = (InActorCluster.GridPlacement == EActorGridPlacement::AlwaysLoaded) ? EActorGridPlacement::AlwaysLoaded : EActorGridPlacement::Bounds;
+		GridPlacement = EActorGridPlacement::Bounds;
 	}
-
-	if (DataLayersID != InActorCluster.DataLayersID)
-	{
-		auto LogActorGuid = [&InActorDescViewMap](const FGuid& ActorGuid)
-		{
-			const FWorldPartitionActorDescView* ActorDescView = InActorDescViewMap.Find(ActorGuid);
-			UE_LOG(LogWorldPartition, Verbose, TEXT("   - Actor: %s (%s)"), ActorDescView ? *ActorDescView->GetActorPath().ToString() : TEXT("None"), *ActorGuid.ToString(EGuidFormats::UniqueObjectGuid));
-		};
-
-		auto LogDataLayers = [](const TSet<const UDataLayer*>& InDataLayers)
-		{
-			TArray<FString> DataLayerLabels;
-			Algo::Transform(InDataLayers, DataLayerLabels, [](const UDataLayer* DataLayer) { return DataLayer->GetDataLayerLabel().ToString(); });
-			UE_LOG(LogWorldPartition, Verbose, TEXT("   - DataLayers: %s"), *FString::Join(DataLayerLabels, TEXT(", ")));
-		};
-
-		if (DataLayers.Num() && InActorCluster.DataLayers.Num())
-		{
-			UE_SUPPRESS(LogWorldPartition, Verbose,
-			{
-				UE_LOG(LogWorldPartition, Verbose, TEXT("Merging Data Layers for clustered actors with different sets of Data Layers."));
-				UE_LOG(LogWorldPartition, Verbose, TEXT("1st cluster :"));
-				LogDataLayers(DataLayers);
-				for (const FGuid& ActorGuid : Actors)
-				{
-					LogActorGuid(ActorGuid);
-				}
-				UE_LOG(LogWorldPartition, Verbose, TEXT("2nd cluster :"));
-				LogDataLayers(InActorCluster.DataLayers);
-				for (const FGuid& ActorGuid : InActorCluster.Actors)
-				{
-					LogActorGuid(ActorGuid);
-				}
-			});
-
-			// Merge Data Layers
-			for (const UDataLayer* DataLayer : InActorCluster.DataLayers)
-			{
-				check(DataLayer->IsRuntime());
-				DataLayers.Add(DataLayer);
-			}
-		}
-		else
-		{
-			// In the case where one of the actor cluster has no Data Layer and the other does, 
-			// the merged actor cluster removes all Data Layers
-			UE_SUPPRESS(LogWorldPartition, Verbose,
-			{
-				UE_LOG(LogWorldPartition, Verbose, TEXT("Removing Data Layers for clustered actors because they are referenced by or are referencing other actors with no Data Layer."));
-				UE_LOG(LogWorldPartition, Verbose, TEXT("Clustered actors with Data Layers :"));
-				const FActorCluster * ActorClusterWithDataLayer = DataLayers.Num() ? this : &InActorCluster;
-				for (const FGuid& ActorGuid : ActorClusterWithDataLayer->Actors)
-				{
-					LogActorGuid(ActorGuid);
-				}
-				UE_LOG(LogWorldPartition, Verbose, TEXT("Clustered actors without Data Layer :"));
-				const FActorCluster * ActorClusterWithoutDataLayer = DataLayers.Num() ? &InActorCluster : this;
-				for (const FGuid& ActorGuid : ActorClusterWithoutDataLayer->Actors)
-				{
-					LogActorGuid(ActorGuid);
-				}
-			});
-			DataLayers.Empty();
-		}
-		DataLayersID = FDataLayersID(DataLayers.Array());
-	}
-
-	// Merge Actors
-	Actors.Append(InActorCluster.Actors);
 }
 
 FActorClusterInstance::FActorClusterInstance(const FActorCluster* InCluster, const FActorContainerInstance* InContainerInstance)
@@ -138,26 +70,16 @@ FActorClusterInstance::FActorClusterInstance(const FActorCluster* InCluster, con
 	
 	TSet<const UDataLayer*> DataLayerSet;
 	DataLayerSet.Reserve(Cluster->DataLayers.Num() + ContainerInstance->DataLayers.Num());
-	// If the ContainerInstance is a WorldPartition we want the Cluster DataLayers to be propagated to the ClusterInstance
-	if (InContainerInstance->Container->IsA<UWorldPartition>())
-	{
-		DataLayerSet.Append(Cluster->DataLayers);
-	}
-	// We also want to propagate the ContainerInstance DataLayers to the ClusterInstance always
-	if (ContainerInstance->DataLayers.Num())
-	{
-		DataLayerSet.Append(ContainerInstance->DataLayers);	
-	}
-	DataLayers.Reserve(DataLayerSet.Num());
-	DataLayers.Append(DataLayerSet.Array());
+	DataLayerSet.Append(Cluster->DataLayers);
+	DataLayerSet.Append(ContainerInstance->DataLayers);
+
+	DataLayers = DataLayerSet.Array();
 }
 
 FActorClusterContext::FActorClusterContext(TArray<FActorContainerInstance>&& InContainerInstances, FFilterActorDescViewFunc InFilterActorDescViewFunc)
 	: FilterActorDescViewFunc(InFilterActorDescViewFunc)
 	, ContainerInstances(MoveTemp(InContainerInstances))
 {
-	TRACE_CPUPROFILER_EVENT_SCOPE(CreateActorClusters);
-
 	for (const FActorContainerInstance& ContainerInstance : ContainerInstances)
 	{
 		const TArray<FActorCluster>& NewClusters = CreateActorClusters(ContainerInstance);
@@ -195,11 +117,6 @@ FActorInstance::FActorInstance(const FGuid& InActor, const FActorContainerInstan
 	check(ContainerInstance);
 }
 
-FVector FActorInstance::GetOrigin() const
-{
-	return ContainerInstance->Transform.TransformPosition(GetActorDescView().GetOrigin());
-}
-
 const FWorldPartitionActorDescView& FActorInstance::GetActorDescView() const
 {
 	return ContainerInstance->GetActorDescView(Actor);
diff --git a/Engine/Source/Runtime/Engine/Private/WorldPartition/WorldPartitionActorDescView.cpp b/Engine/Source/Runtime/Engine/Private/WorldPartition/WorldPartitionActorDescView.cpp
--- a/Engine/Source/Runtime/Engine/Private/WorldPartition/WorldPartitionActorDescView.cpp
+++ b/Engine/Source/Runtime/Engine/Private/WorldPartition/WorldPartitionActorDescView.cpp
@@ -9,11 +9,13 @@
 FWorldPartitionActorDescView::FWorldPartitionActorDescView()
 	: ActorDesc(nullptr)
 	, GridPlacement(EActorGridPlacement::None)
+	, bInvalidDataLayers(false)
 {}
 
 FWorldPartitionActorDescView::FWorldPartitionActorDescView(const FWorldPartitionActorDesc* InActorDesc)
 	: ActorDesc(InActorDesc)
 	, GridPlacement(InActorDesc->GetGridPlacement())
+	, bInvalidDataLayers(false)
 {}
 
 const FGuid& FWorldPartitionActorDescView::GetGuid() const
@@ -68,7 +70,8 @@ UHLODLayer* FWorldPartitionActorDescView::GetHLODLayer() const
 
 const TArray<FName>& FWorldPartitionActorDescView::GetDataLayers() const
 {
-	return ActorDesc->GetDataLayers();
+	static TArray<FName> EmptyDataLayers;
+	return bInvalidDataLayers ? EmptyDataLayers : ActorDesc->GetDataLayers();
 }
 
 FName FWorldPartitionActorDescView::GetActorPackage() const
@@ -114,4 +117,13 @@ void FWorldPartitionActorDescView::SetGridPlacement(EActorGridPlacement InGridPl
 		UE_LOG(LogWorldPartition, Verbose, TEXT("Actor '%s' grid placement changed to %s"), *GetActorLabel().ToString(), *StaticEnum<EActorGridPlacement>()->GetNameStringByValue((int64)InGridPlacement));
 	}
 }
+
+void FWorldPartitionActorDescView::SetInvalidDataLayers()
+{
+	if (!bInvalidDataLayers)
+	{
+		bInvalidDataLayers = true;
+		UE_LOG(LogWorldPartition, Verbose, TEXT("Actor '%s' data layers invalidated"), *GetActorLabel().ToString());
+	}
+}
 #endif
\ No newline at end of file
diff --git a/Engine/Source/Runtime/Engine/Private/WorldPartition/WorldPartitionRuntimeSpatialHash.cpp b/Engine/Source/Runtime/Engine/Private/WorldPartition/WorldPartitionRuntimeSpatialHash.cpp
--- a/Engine/Source/Runtime/Engine/Private/WorldPartition/WorldPartitionRuntimeSpatialHash.cpp
+++ b/Engine/Source/Runtime/Engine/Private/WorldPartition/WorldPartitionRuntimeSpatialHash.cpp
@@ -925,7 +925,7 @@ bool UWorldPartitionRuntimeSpatialHash::CreateStreamingGrid(const FSpatialHashRu
 							StreamingCell->ActorContainer->Actors.Add(Actor->GetFName(), Actor);
 						}
 					}
-					UE_LOG(LogWorldPartition, Verbose, TEXT("  Actor : %s (%s) (Container %s) Origin(%s)"), *(ActorDescView.GetActorPath().ToString()), *ActorDescView.GetGuid().ToString(EGuidFormats::UniqueObjectGuid), *ActorInstance.ContainerInstance->ID.ToString(), *FVector2D(ActorInstance.GetOrigin()).ToString());
+					UE_LOG(LogWorldPartition, Verbose, TEXT("  Actor : %s (%s) (Container %s)"), *(ActorDescView.GetActorPath().ToString()), *ActorDescView.GetGuid().ToString(EGuidFormats::UniqueObjectGuid), *ActorInstance.ContainerInstance->ID.ToString());
 				}
 
 				if (IsRunningCookCommandlet())
diff --git a/Engine/Source/Runtime/Engine/Private/WorldPartition/WorldPartitionStreamingGeneration.cpp b/Engine/Source/Runtime/Engine/Private/WorldPartition/WorldPartitionStreamingGeneration.cpp
--- a/Engine/Source/Runtime/Engine/Private/WorldPartition/WorldPartitionStreamingGeneration.cpp
+++ b/Engine/Source/Runtime/Engine/Private/WorldPartition/WorldPartitionStreamingGeneration.cpp
@@ -83,7 +83,7 @@ class FWorldPartitionStreamingGenerator
 		TMap<FGuid, FWorldPartitionActorDescView> ActorDescViewMap;
 		
 		// Gather actor descriptor views for this container
-		CreateActorDescViewMap(InContainer, ActorDescViewMap, InContainerID);		
+		CreateActorDescViewMap(InContainer, ActorDescViewMap, InContainerID);
 
 		// Parse actor descriptors
 		for (auto It = ActorDescViewMap.CreateIterator(); It; ++It)
@@ -177,6 +177,45 @@ class FWorldPartitionStreamingGenerator
 
 			// Give the associated runtime hash to adjust the grid placement based on its internal settings, etc.
 			RuntimeHash->UpdateActorDescViewMap(ContainerDescriptor.ActorDescViewMap);
+
+			// Perform various adjustements based on validations and errors
+			for (auto ActorDescIt = ContainerDescriptor.ActorDescViewMap.CreateIterator(); ActorDescIt; ++ActorDescIt)
+			{
+				FWorldPartitionActorDescView& ActorDescView = ActorDescIt.Value();
+
+				// Validate references
+				for (const FGuid& ReferenceGuid : ActorDescView.GetReferences())
+				{
+					if (FWorldPartitionActorDescView* ReferenceActorDescView = ContainerDescriptor.ActorDescViewMap.Find(ReferenceGuid))
+					{
+						// Validate grid placement
+						if ((ActorDescView.GetGridPlacement() != ReferenceActorDescView->GetGridPlacement()) &&
+							((ActorDescView.GetGridPlacement() == EActorGridPlacement::AlwaysLoaded) || (ReferenceActorDescView->GetGridPlacement() == EActorGridPlacement::AlwaysLoaded)))
+						{
+							ActorDescView.SetGridPlacement(EActorGridPlacement::AlwaysLoaded);
+							ReferenceActorDescView->SetGridPlacement(EActorGridPlacement::AlwaysLoaded);
+						}
+
+						// Validate data layers
+						if (ActorDescView.GetDataLayers().Num() != ReferenceActorDescView->GetDataLayers().Num())
+						{
+							ActorDescView.SetInvalidDataLayers();
+							ReferenceActorDescView->SetInvalidDataLayers();
+						}
+						else
+						{
+							const TSet<FName> ActorDescDataLayers(ActorDescView.GetDataLayers());
+							const TSet<FName> ReferenceActorDescDataLayers(ReferenceActorDescView->GetDataLayers());
+
+							if (!ActorDescDataLayers.Includes(ReferenceActorDescDataLayers))
+							{
+								ActorDescView.SetInvalidDataLayers();
+								ReferenceActorDescView->SetInvalidDataLayers();
+							}
+						}
+					}
+				}
+			}
 		}
 	}
 
diff --git a/Engine/Source/Runtime/Engine/Public/WorldPartition/WorldPartitionActorCluster.h b/Engine/Source/Runtime/Engine/Public/WorldPartition/WorldPartitionActorCluster.h
--- a/Engine/Source/Runtime/Engine/Public/WorldPartition/WorldPartitionActorCluster.h
+++ b/Engine/Source/Runtime/Engine/Public/WorldPartition/WorldPartitionActorCluster.h
@@ -90,7 +90,7 @@ struct FActorContainerInstance
 {
 	FActorContainerInstance(const FActorContainerID& InID, const FTransform& InTransform, const FBox& InBounds, const TSet<FName>& InDataLayers, EContainerClusterMode InClusterMode, const UActorDescContainer* InContainer, TMap<FGuid, FWorldPartitionActorDescView> InActorDescViewMap);
 	
-	FActorContainerID				ID;
+	FActorContainerID			ID;
 	FTransform					Transform;
 	FBox						Bounds;
 	EContainerClusterMode		ClusterMode;
@@ -125,8 +125,6 @@ struct ENGINE_API FActorInstance
 	FActorInstance();
 	FActorInstance(const FGuid& InActor, const FActorContainerInstance* InContainerInstance);
 
-	FVector GetOrigin() const;
-
 	friend uint32 GetTypeHash(const FActorInstance& InActorInstance)
 	{
 		FHashBuilder HashBuilder;
diff --git a/Engine/Source/Runtime/Engine/Public/WorldPartition/WorldPartitionActorDescView.h b/Engine/Source/Runtime/Engine/Public/WorldPartition/WorldPartitionActorDescView.h
--- a/Engine/Source/Runtime/Engine/Public/WorldPartition/WorldPartitionActorDescView.h
+++ b/Engine/Source/Runtime/Engine/Public/WorldPartition/WorldPartitionActorDescView.h
@@ -45,6 +45,8 @@ public:
 
 	void SetGridPlacement(EActorGridPlacement InGridPlacement);
 
+	void SetInvalidDataLayers();
+
 	bool operator==(const FWorldPartitionActorDescView& Other) const
 	{
 		return GetGuid() == Other.GetGuid();
@@ -58,5 +60,6 @@ public:
 protected:
 	const FWorldPartitionActorDesc* ActorDesc;
 	EActorGridPlacement GridPlacement;
+	bool bInvalidDataLayers;
 };
 #endif
\ No newline at end of file
