Saves 4.6MB in FString only used when warning occurs
  - Replace FString properties for asset path by the Package FName instead to save on storage space.
  - Remove FString properties, the properties has been added during 5.0 development, no need to deprecate.

diff --git a/Engine/Source/Runtime/Engine/Classes/Sound/SoundWave.h b/Engine/Source/Runtime/Engine/Classes/Sound/SoundWave.h
--- a/Engine/Source/Runtime/Engine/Classes/Sound/SoundWave.h
+++ b/Engine/Source/Runtime/Engine/Classes/Sound/SoundWave.h
@@ -1118,8 +1118,7 @@ public:
 
 	// UObject Interface cached data getters
 	const FName& GetFName() const { return NameCached; }
-	const FString& GetFullName() const { return FullNameCached; }
-	const FString& GetPathName() const { return PathNameCached; }
+	const FName& GetPackageName() const { return PackageNameCached; }
 	const FName& GetRuntimeFormat() const { return RuntimeFormat; }
 	const FObjectKey& GetFObjectKey() const { return SoundWaveKeyCached; };
 
@@ -1170,8 +1169,7 @@ public:
 
 private:
 	FName NameCached;
-	FString FullNameCached;
-	FString PathNameCached;
+	FName PackageNameCached;
 	FObjectKey SoundWaveKeyCached;
 	FName RuntimeFormat{ "FSoundWaveProxy_InvalidFormat" };
 
diff --git a/Engine/Source/Runtime/Engine/Private/SoundWave.cpp b/Engine/Source/Runtime/Engine/Private/SoundWave.cpp
--- a/Engine/Source/Runtime/Engine/Private/SoundWave.cpp
+++ b/Engine/Source/Runtime/Engine/Private/SoundWave.cpp
@@ -3355,8 +3355,6 @@ ESoundWaveLoadingBehavior USoundWave::GetLoadingBehavior(bool bCheckSoundClasses
 FSoundWaveProxy::FSoundWaveProxy(USoundWave* InWave)
 	: SoundWaveDataPtr(InWave->SoundWaveDataPtr)
 	, NameCached(InWave->GetFName())
-	, FullNameCached(InWave->GetFullName())
-	, PathNameCached(InWave->GetPathName())
 	, SoundWaveKeyCached(FObjectKey(InWave))
 	, SampleRate(InWave->GetSampleRateForCurrentPlatform())
 	, NumChannels(InWave->NumChannels)
@@ -3369,6 +3367,11 @@ FSoundWaveProxy::FSoundWaveProxy(USoundWave* InWave)
 {
 	LLM_SCOPE(ELLMTag::AudioSoundWaveProxies);
 
+	if (UPackage* Package = InWave->GetPackage())
+	{
+		PackageNameCached = Package->GetFName();
+	}
+
 	// this should have been allocated by the USoundWave and should always be valid
 	check(SoundWaveDataPtr);
 
@@ -3533,12 +3536,12 @@ bool FSoundWaveProxy::GetChunkData(int32 ChunkIndex, uint8** OutChunkData, bool
 	{
 #if WITH_EDITORONLY_DATA
 		// Unable to load chunks from the cache. Rebuild the sound and attempt to precache it.
-		UE_LOG(LogAudio, Display, TEXT("GetChunkData failed, rebuilding %s"), *GetPathName());
+		UE_LOG(LogAudio, Display, TEXT("GetChunkData failed, rebuilding %s"), *GetPackageName().ToString());
 
 		//		ForceRebuildPlatformData();
 		if (GetChunkFromDDC(ChunkIndex, OutChunkData, bMakeSureChunkIsLoaded) == 0)
 		{
-			UE_LOG(LogAudio, Warning, TEXT("Failed to build sound %s."), *GetPathName());
+			UE_LOG(LogAudio, Warning, TEXT("Failed to build sound %s."), *GetPackageName().ToString());
 		}
 		else
 		{
@@ -3547,7 +3550,7 @@ bool FSoundWaveProxy::GetChunkData(int32 ChunkIndex, uint8** OutChunkData, bool
 		}
 #else
 		// Failed to find the SoundWave chunk in the cooked package.
-		UE_LOG(LogAudio, Warning, TEXT("GetChunkData failed while streaming. Ensure the following file is cooked: %s"), *GetPathName());
+		UE_LOG(LogAudio, Warning, TEXT("GetChunkData failed while streaming. Ensure the following package is cooked: %s"), *GetPackageName().ToString());
 #endif // #if WITH_EDITORONLY_DATA
 		return false;
 	}
