Reduce cook time by avoiding costly object iterator when unnecessary
  - Saves 15 minutes on a 2h20 cook

diff --git a/Engine/Source/Runtime/Engine/Private/Components/ActorComponent.cpp b/Engine/Source/Runtime/Engine/Private/Components/ActorComponent.cpp
--- a/Engine/Source/Runtime/Engine/Private/Components/ActorComponent.cpp
+++ b/Engine/Source/Runtime/Engine/Private/Components/ActorComponent.cpp
@@ -225,40 +225,52 @@ void FGlobalComponentReregisterContext::UpdateAllPrimitiveSceneInfos()
 
 FGlobalComponentRecreateRenderStateContext::FGlobalComponentRecreateRenderStateContext()
 {
-	// wait until resources are released
-	FlushRenderingCommands();
-
-	// recreate render state for all components.
-	for (UActorComponent* Component : TObjectRange<UActorComponent>())
+	if (FApp::CanEverRender())
 	{
-		if (Component->IsRegistered() && Component->IsRenderStateCreated())
+		TRACE_CPUPROFILER_EVENT_SCOPE(FGlobalComponentRecreateRenderStateContext::FGlobalComponentRecreateRenderStateContext);
+
+		// wait until resources are released
+		FlushRenderingCommands();
+
+		// recreate render state for all components.
+		for (UActorComponent* Component : TObjectRange<UActorComponent>())
 		{
-			ComponentContexts.Emplace(Component, &ScenesToUpdateAllPrimitiveSceneInfos);
+			if (Component->IsRegistered() && Component->IsRenderStateCreated())
+			{
+				ComponentContexts.Emplace(Component, &ScenesToUpdateAllPrimitiveSceneInfos);
+			}
 		}
-	}
 
-	UpdateAllPrimitiveSceneInfos();
+		UpdateAllPrimitiveSceneInfos();
+	}
 }
 
 FGlobalComponentRecreateRenderStateContext::FGlobalComponentRecreateRenderStateContext(const TArray<UActorComponent*>& InComponents)
 {
-	// wait until resources are released
-	FlushRenderingCommands();
-
-	// recreate render state for provided components.
-	for (UActorComponent* Component : InComponents)
+	if (FApp::CanEverRender())
 	{
-		if (Component->IsRegistered() && Component->IsRenderStateCreated())
+		TRACE_CPUPROFILER_EVENT_SCOPE(FGlobalComponentRecreateRenderStateContext::FGlobalComponentRecreateRenderStateContext);
+
+		// wait until resources are released
+		FlushRenderingCommands();
+
+		// recreate render state for provided components.
+		for (UActorComponent* Component : InComponents)
 		{
-			ComponentContexts.Emplace(Component, &ScenesToUpdateAllPrimitiveSceneInfos);
+			if (Component->IsRegistered() && Component->IsRenderStateCreated())
+			{
+				ComponentContexts.Emplace(Component, &ScenesToUpdateAllPrimitiveSceneInfos);
+			}
 		}
-	}
 
-	UpdateAllPrimitiveSceneInfos();
+		UpdateAllPrimitiveSceneInfos();
+	}
 }
 
 FGlobalComponentRecreateRenderStateContext::~FGlobalComponentRecreateRenderStateContext()
 {
+	TRACE_CPUPROFILER_EVENT_SCOPE(FGlobalComponentRecreateRenderStateContext::~FGlobalComponentRecreateRenderStateContext);
+
 	ComponentContexts.Empty();
 
 	UpdateAllPrimitiveSceneInfos();
@@ -1046,6 +1058,7 @@ TStructOnScope<FActorComponentInstanceData> UActorComponent::GetComponentInstanc
 
 void FActorComponentTickFunction::ExecuteTick(float DeltaTime, enum ELevelTick TickType, ENamedThreads::Type CurrentThread, const FGraphEventRef& MyCompletionGraphEvent)
 {
+	TRACE_CPUPROFILER_EVENT_SCOPE(FActorComponentTickFunction::ExecuteTick);
 	ExecuteTickHelper(Target, Target->bTickInEditor, DeltaTime, TickType, [this, TickType](float DilatedTime)
 	{
 		Target->TickComponent(DilatedTime, TickType, this);
