Add BackReferences map in WorldPartitionEdtiorSpatialHash so that when referenced (attached or otherwise) actors located in another cell are deleted, we can do proper clean up of the references in all the cells.

diff --git a/Engine/Source/Runtime/Engine/Private/WorldPartition/WorldPartitionEditorCell.cpp b/Engine/Source/Runtime/Engine/Private/WorldPartition/WorldPartitionEditorCell.cpp
--- a/Engine/Source/Runtime/Engine/Private/WorldPartition/WorldPartitionEditorCell.cpp
+++ b/Engine/Source/Runtime/Engine/Private/WorldPartition/WorldPartitionEditorCell.cpp
@@ -68,6 +68,7 @@ void UWorldPartitionEditorCell::AddActor(const FGuid& Source, const FWorldPartit
 			if (ReferenceHandle.IsValid())
 			{
 				AddActor(ActorHandle->GetGuid(), ReferenceHandle);
+				WorldPartition->EditorHash->AddBackReference(ReferenceHandle, this, Source);
 			}
 		}
 	}
@@ -103,6 +104,7 @@ void UWorldPartitionEditorCell::RemoveActor(const FGuid& Source, const FWorldPar
 			if (ReferenceHandle.IsValid())
 			{
 				RemoveActor(ActorHandle->GetGuid(), ReferenceHandle);
+				WorldPartition->EditorHash->RemoveBackReference(ReferenceHandle, this, Source);
 			}
 		}
 	}
diff --git a/Engine/Source/Runtime/Engine/Private/WorldPartition/WorldPartitionEditorSpatialHash.cpp b/Engine/Source/Runtime/Engine/Private/WorldPartition/WorldPartitionEditorSpatialHash.cpp
--- a/Engine/Source/Runtime/Engine/Private/WorldPartition/WorldPartitionEditorSpatialHash.cpp
+++ b/Engine/Source/Runtime/Engine/Private/WorldPartition/WorldPartitionEditorSpatialHash.cpp
@@ -314,6 +314,23 @@ void UWorldPartitionEditorSpatialHash::UnhashActor(FWorldPartitionHandle& InActo
 			}
 		});
 	}
+
+	// Remove from all cells where it's back referenced
+	// This is important because the FWorldPartitionHandle will soon become invalid and Actors that were still 
+	// referencing this one will not be able to construct valid FWorldPartitionHandle to the actor we're 
+	// removing and it'll become impossible to clean the cells of their dangling actors
+	TArray< TTuple<UWorldPartitionEditorCell*, FGuid>> BackRefs;
+	BackReferences.MultiFind(InActorHandle, BackRefs);
+
+	for (auto& BackRef : BackRefs)
+	{
+		BackRef.Get<0>()->RemoveActor(BackRef.Get<1>(), InActorHandle);
+	}
+
+	if (BackRefs.Num() != 0)
+	{
+		BackReferences.Remove(InActorHandle);
+	}
 }
 
 void UWorldPartitionEditorSpatialHash::OnCellLoaded(const UWorldPartitionEditorCell* Cell)
@@ -571,6 +588,18 @@ void UWorldPartitionEditorSpatialHash::PostLoad()
 	}
 }
 
+
+void UWorldPartitionEditorSpatialHash::AddBackReference(FWorldPartitionHandle& ReferenceHandle,  UWorldPartitionEditorCell* Cell, const FGuid& Source)
+{
+	BackReferences.Add(ReferenceHandle) = MakeTuple(Cell, Source);
+}
+
+void UWorldPartitionEditorSpatialHash::RemoveBackReference(FWorldPartitionHandle& ReferenceHandle,  UWorldPartitionEditorCell* Cell, const FGuid& Source)
+{
+	BackReferences.RemoveSingle(ReferenceHandle, MakeTuple(Cell, Source));
+}
+
+
 #endif
 
 #undef LOCTEXT_NAMESPACE
diff --git a/Engine/Source/Runtime/Engine/Public/WorldPartition/WorldPartitionEditorHash.h b/Engine/Source/Runtime/Engine/Public/WorldPartition/WorldPartitionEditorHash.h
--- a/Engine/Source/Runtime/Engine/Public/WorldPartition/WorldPartitionEditorHash.h
+++ b/Engine/Source/Runtime/Engine/Public/WorldPartition/WorldPartitionEditorHash.h
@@ -37,6 +37,10 @@ public:
 	virtual uint32 GetWantedEditorCellSize() const PURE_VIRTUAL(UWorldPartitionEditorHash::GetWantedEditorCellSize, return 0;);
 	virtual void SetEditorWantedCellSize(uint32 InCellSize) PURE_VIRTUAL(UWorldPartitionEditorHash::SetEditorWantedCellSize, );
 
+	virtual void AddBackReference(FWorldPartitionHandle& ReferenceHandle,  UWorldPartitionEditorCell* Cell, const FGuid& Source) PURE_VIRTUAL(UWorldPartitionEditorHash::AddBackReference, ;);
+	virtual void RemoveBackReference(FWorldPartitionHandle& ReferenceHandle,  UWorldPartitionEditorCell* Cell, const FGuid& Source) PURE_VIRTUAL(UWorldPartitionEditorHash::RemoveBackReference, ;);
+
+
 	// Helpers
 	inline int32 GetIntersectingActors(const FBox& Box, TArray<FWorldPartitionActorDesc*>& OutActors)
 	{
diff --git a/Engine/Source/Runtime/Engine/Public/WorldPartition/WorldPartitionEditorSpatialHash.h b/Engine/Source/Runtime/Engine/Public/WorldPartition/WorldPartitionEditorSpatialHash.h
--- a/Engine/Source/Runtime/Engine/Public/WorldPartition/WorldPartitionEditorSpatialHash.h
+++ b/Engine/Source/Runtime/Engine/Public/WorldPartition/WorldPartitionEditorSpatialHash.h
@@ -230,6 +230,10 @@ public:
 
 	virtual uint32 GetWantedEditorCellSize() const override;
 	virtual void SetEditorWantedCellSize(uint32 InCellSize) override;
+
+	virtual void AddBackReference(FWorldPartitionHandle& ReferenceHandle,  UWorldPartitionEditorCell* Cell, const FGuid& Source) override;
+	virtual void RemoveBackReference(FWorldPartitionHandle& ReferenceHandle,  UWorldPartitionEditorCell* Cell, const FGuid& Source) override;
+		
 	// UWorldPartitionEditorHash interface end
 #endif
 
@@ -250,6 +254,8 @@ private:
 
 	TMap<FCellCoord, FCellNode> HashNodes;
 	TMap<FCellCoord, UWorldPartitionEditorCell*> HashCells;
+	
+	TMultiMap<FWorldPartitionHandle, TTuple<UWorldPartitionEditorCell*, FGuid>>	BackReferences;
 
 	UPROPERTY(Transient)
 	TSet<TObjectPtr<UWorldPartitionEditorCell>> Cells;
