World Partition Streaming Generation:
- Added IStreamingGenerationErrorHandler interface to handle world partition streaming generation errors.
  - Implemented FStreamingGenerationLogErrorHandler error handler that logs errors to the console.
  - Implemented FStreamingGenerationMapCheckErrorHandler error handle that logs errors to the map check window.
- Moved all map checks code from runtime hash to streaming generation code using new interface.

diff --git a/Engine/Source/Editor/UnrealEd/Private/EditorServer.cpp b/Engine/Source/Editor/UnrealEd/Private/EditorServer.cpp
--- a/Engine/Source/Editor/UnrealEd/Private/EditorServer.cpp
+++ b/Engine/Source/Editor/UnrealEd/Private/EditorServer.cpp
@@ -120,7 +120,9 @@
 #include "InstancedFoliageActor.h"
 #include "IMovieSceneCapture.h"
 #include "MovieSceneCaptureModule.h"
+
 #include "WorldPartition/WorldPartition.h"
+#include "WorldPartition/ErrorHandling/WorldPartitionStreamingGenerationMapCheckErrorHandler.h"
 
 #include "Kismet2/KismetEditorUtilities.h"
 #include "PropertyEditorModule.h"
@@ -4031,7 +4033,8 @@ bool UEditorEngine::Map_Check( UWorld* InWorld, const TCHAR* Str, FOutputDevice&
 
 	if (const UWorldPartition* WorldPartition = InWorld->GetWorldPartition())
 	{
-		WorldPartition->CheckForErrors();
+		FStreamingGenerationMapCheckErrorHandler MapCheckErrorHandler;
+		WorldPartition->CheckForErrors(&MapCheckErrorHandler);
 	}
 
 	GWarn->StatusUpdate( 0, ProgressDenominator, CheckMapLocText );
diff --git a/Engine/Source/Runtime/Engine/Private/WorldPartition/ErrorHandling/WorldPartitionStreamingGenerationErrorHandler.cpp b/Engine/Source/Runtime/Engine/Private/WorldPartition/ErrorHandling/WorldPartitionStreamingGenerationErrorHandler.cpp
new file mode 100644
--- /dev/null
+++ b/Engine/Source/Runtime/Engine/Private/WorldPartition/ErrorHandling/WorldPartitionStreamingGenerationErrorHandler.cpp
@@ -0,0 +1,26 @@
+// Copyright Epic Games, Inc. All Rights Reserved.
+
+#if WITH_EDITOR
+#include "WorldPartition/ErrorHandling/WorldPartitionStreamingGenerationErrorHandler.h"
+
+FString IStreamingGenerationErrorHandler::GetActorLabel(const FWorldPartitionActorDescView& ActorDescView)
+{
+	const FName ActorLabel = ActorDescView.GetActorLabel();
+	if (!ActorLabel.IsNone())
+	{
+		return ActorLabel.ToString();
+	}
+
+	const FString ActorPath = ActorDescView.GetActorPath().ToString();
+
+	FString SubObjectName;
+	FString SubObjectContext;
+	if (FString(ActorPath).Split(TEXT("."), &SubObjectContext, &SubObjectName))
+	{
+		return SubObjectName;
+	}
+
+	return ActorPath;
+};
+
+#endif
diff --git a/Engine/Source/Runtime/Engine/Private/WorldPartition/ErrorHandling/WorldPartitionStreamingGenerationLogErrorHandler.cpp b/Engine/Source/Runtime/Engine/Private/WorldPartition/ErrorHandling/WorldPartitionStreamingGenerationLogErrorHandler.cpp
new file mode 100644
--- /dev/null
+++ b/Engine/Source/Runtime/Engine/Private/WorldPartition/ErrorHandling/WorldPartitionStreamingGenerationLogErrorHandler.cpp
@@ -0,0 +1,35 @@
+// Copyright Epic Games, Inc. All Rights Reserved.
+
+#if WITH_EDITOR
+#include "WorldPartition/ErrorHandling/WorldPartitionStreamingGenerationLogErrorHandler.h"
+#include "WorldPartition/WorldPartitionLog.h"
+
+void FStreamingGenerationLogErrorHandler::OnInvalidReference(const FWorldPartitionActorDescView& ActorDescView, const FGuid& ReferenceGuid)
+{
+	UE_LOG(LogWorldPartition, Log, TEXT("Actor %s have missing reference to %s"), *GetActorLabel(ActorDescView), *ReferenceGuid.ToString());
+}
+
+void FStreamingGenerationLogErrorHandler::OnInvalidReferenceGridPlacement(const FWorldPartitionActorDescView& ActorDescView, const FWorldPartitionActorDescView& ReferenceActorDescView)
+{
+	const FString StreamedActor(TEXT("Streamed actor"));
+	const FString AlwaysLoadedActor(TEXT("Always loaded actor"));
+	const bool bIsActorDescAlwaysLoaded = ActorDescView.GetGridPlacement() == EActorGridPlacement::AlwaysLoaded;
+	const bool bIsActorDescRefAlwaysLoaded = ReferenceActorDescView.GetGridPlacement() == EActorGridPlacement::AlwaysLoaded;
+	UE_LOG(LogWorldPartition, Log, TEXT("%s %s reference %s %s"), bIsActorDescAlwaysLoaded ? *AlwaysLoadedActor : *StreamedActor, *GetActorLabel(ActorDescView), bIsActorDescRefAlwaysLoaded ? *AlwaysLoadedActor : *StreamedActor, *GetActorLabel(ReferenceActorDescView));
+}
+
+void FStreamingGenerationLogErrorHandler::OnInvalidReferenceDataLayers(const FWorldPartitionActorDescView& ActorDescView, const FWorldPartitionActorDescView& ReferenceActorDescView)
+{
+	UE_LOG(LogWorldPartition, Log, TEXT("Actor %s references an actor in a different set of data layers %s"), *GetActorLabel(ActorDescView), *GetActorLabel(ReferenceActorDescView));
+}
+
+void FStreamingGenerationLogErrorHandler::OnInvalidReferenceLevelScriptStreamed(const FWorldPartitionActorDescView& ActorDescView)
+{
+	UE_LOG(LogWorldPartition, Log, TEXT("Level Script Blueprint references streamed actor %s"), *GetActorLabel(ActorDescView));
+}
+
+void FStreamingGenerationLogErrorHandler::OnInvalidReferenceLevelScriptDataLayers(const FWorldPartitionActorDescView& ActorDescView)
+{
+	UE_LOG(LogWorldPartition, Log, TEXT("Level Script Blueprint references streamed actor %s with a non empty set of data layers"), *GetActorLabel(ActorDescView));
+}
+#endif
diff --git a/Engine/Source/Runtime/Engine/Private/WorldPartition/ErrorHandling/WorldPartitionStreamingGenerationMapCheckErrorHandler.cpp b/Engine/Source/Runtime/Engine/Private/WorldPartition/ErrorHandling/WorldPartitionStreamingGenerationMapCheckErrorHandler.cpp
new file mode 100644
--- /dev/null
+++ b/Engine/Source/Runtime/Engine/Private/WorldPartition/ErrorHandling/WorldPartitionStreamingGenerationMapCheckErrorHandler.cpp
@@ -0,0 +1,67 @@
+// Copyright Epic Games, Inc. All Rights Reserved.
+
+#if WITH_EDITOR
+#include "WorldPartition/ErrorHandling/WorldPartitionStreamingGenerationMapCheckErrorHandler.h"
+#include "Logging/MessageLog.h"
+#include "Misc/UObjectToken.h"
+#include "Misc/MapErrors.h"
+
+#define LOCTEXT_NAMESPACE "WorldPartition"
+
+void FStreamingGenerationMapCheckErrorHandler::OnInvalidReference(const FWorldPartitionActorDescView& ActorDescView, const FGuid& ReferenceGuid)
+{
+	TSharedRef<FTokenizedMessage> Error = FMessageLog("MapCheck").Warning()
+		->AddToken(FTextToken::Create(LOCTEXT("MapCheck_WorldPartition_Actor", "Actor")))
+		->AddToken(FAssetNameToken::Create(GetActorLabel(ActorDescView)))
+		->AddToken(FTextToken::Create(LOCTEXT("MapCheck_WorldPartition_HaveMissingRefsTo", "have missing references to")))
+		->AddToken(FTextToken::Create(FText::FromString(ReferenceGuid.ToString())))
+		->AddToken(FMapErrorToken::Create(FName(TEXT("WorldPartition_MissingActorReference_CheckForErrors"))));
+}
+
+void FStreamingGenerationMapCheckErrorHandler::OnInvalidReferenceGridPlacement(const FWorldPartitionActorDescView& ActorDescView, const FWorldPartitionActorDescView& ReferenceActorDescView)
+{
+	const bool bIsActorDescAlwaysLoaded = ActorDescView.GetGridPlacement() == EActorGridPlacement::AlwaysLoaded;
+	const bool bIsActorDescRefAlwaysLoaded = ReferenceActorDescView.GetGridPlacement() == EActorGridPlacement::AlwaysLoaded;
+
+	const FText StreamedActor(LOCTEXT("MapCheck_WorldPartition_StreamedActor", "Streamed actor"));
+	const FText AlwaysLoadedActor(LOCTEXT("MapCheck_WorldPartition_AlwaysLoadedActor", "Always loaded actor"));
+
+	TSharedRef<FTokenizedMessage> Error = FMessageLog("MapCheck").Error()
+		->AddToken(FTextToken::Create(bIsActorDescAlwaysLoaded ? AlwaysLoadedActor : StreamedActor))
+		->AddToken(FAssetNameToken::Create(GetActorLabel(ActorDescView)))
+		->AddToken(FTextToken::Create(LOCTEXT("MapCheck_WorldPartition_References", "references")))
+		->AddToken(FTextToken::Create(bIsActorDescRefAlwaysLoaded ? AlwaysLoadedActor : StreamedActor)								)
+		->AddToken(FAssetNameToken::Create(GetActorLabel(ReferenceActorDescView)))
+		->AddToken(FMapErrorToken::Create(FName(TEXT("WorldPartition_StreamedActorReferenceAlwaysLoadedActor_CheckForErrors"))));
+}
+
+void FStreamingGenerationMapCheckErrorHandler::OnInvalidReferenceDataLayers(const FWorldPartitionActorDescView& ActorDescView, const FWorldPartitionActorDescView& ReferenceActorDescView)
+{
+	TSharedRef<FTokenizedMessage> Error = FMessageLog("MapCheck").Error()
+		->AddToken(FTextToken::Create(LOCTEXT("MapCheck_WorldPartition_Actor", "Actor")))
+		->AddToken(FAssetNameToken::Create(GetActorLabel(ActorDescView)))
+		->AddToken(FTextToken::Create(LOCTEXT("MapCheck_WorldPartition_ReferenceActorInOtherDataLayers", "references an actor in a different set of data layers")))
+		->AddToken(FAssetNameToken::Create(GetActorLabel(ReferenceActorDescView)))
+		->AddToken(FMapErrorToken::Create(FName(TEXT("WorldPartition_ActorReferenceActorInAnotherDataLayer_CheckForErrors"))));
+}
+
+void FStreamingGenerationMapCheckErrorHandler::OnInvalidReferenceLevelScriptStreamed(const FWorldPartitionActorDescView& ActorDescView)
+{
+	TSharedRef<FTokenizedMessage> Error = FMessageLog("MapCheck").Error()
+		->AddToken(FTextToken::Create(LOCTEXT("MapCheck_WorldPartition_LevelScriptBlueprintStreamedActorReference", "Level Script Blueprint references streamed actor")))
+		->AddToken(FAssetNameToken::Create(GetActorLabel(ActorDescView)))
+		->AddToken(FMapErrorToken::Create(FName(TEXT("WorldPartition_LevelScriptBlueprintRefefenceStreamed_CheckForErrors"))));
+}
+
+void FStreamingGenerationMapCheckErrorHandler::OnInvalidReferenceLevelScriptDataLayers(const FWorldPartitionActorDescView& ActorDescView)
+{
+	TSharedRef<FTokenizedMessage> Error = FMessageLog("MapCheck").Error()
+		->AddToken(FTextToken::Create(LOCTEXT("MapCheck_WorldPartition_LevelScriptBlueprintActorReference", "Level Script Blueprint references actor")))
+		->AddToken(FAssetNameToken::Create(GetActorLabel(ActorDescView)))
+		->AddToken(FTextToken::Create(LOCTEXT("MapCheck_WorldPartition_LevelScriptBlueprintDataLayerReference", "with a non empty set of data layers")))
+		->AddToken(FMapErrorToken::Create(FName(TEXT("WorldPartition_LevelScriptBlueprintRefefenceDataLayer_CheckForErrors"))));
+}
+
+#undef LOCTEXT_NAMESPACE
+
+#endif
diff --git a/Engine/Source/Runtime/Engine/Private/WorldPartition/WorldPartition.cpp b/Engine/Source/Runtime/Engine/Private/WorldPartition/WorldPartition.cpp
--- a/Engine/Source/Runtime/Engine/Private/WorldPartition/WorldPartition.cpp
+++ b/Engine/Source/Runtime/Engine/Private/WorldPartition/WorldPartition.cpp
@@ -1314,14 +1314,6 @@ void UWorldPartition::DumpActorDescs(const FString& Path)
 	}
 }
 
-void UWorldPartition::CheckForErrors() const
-{
-	if (RuntimeHash)
-	{
-		RuntimeHash->CheckForErrors();
-	}
-}
-
 uint32 UWorldPartition::GetWantedEditorCellSize() const
 {
 	return EditorHash->GetWantedEditorCellSize();
diff --git a/Engine/Source/Runtime/Engine/Private/WorldPartition/WorldPartitionRuntimeHash.cpp b/Engine/Source/Runtime/Engine/Private/WorldPartition/WorldPartitionRuntimeHash.cpp
--- a/Engine/Source/Runtime/Engine/Private/WorldPartition/WorldPartitionRuntimeHash.cpp
+++ b/Engine/Source/Runtime/Engine/Private/WorldPartition/WorldPartitionRuntimeHash.cpp
@@ -10,9 +10,6 @@
 #if WITH_EDITOR
 #include "WorldPartition/WorldPartitionHandle.h"
 #include "WorldPartition/DataLayer/WorldDataLayers.h"
-#include "Logging/MessageLog.h"
-#include "Misc/UObjectToken.h"
-#include "Misc/MapErrors.h"
 #endif
 
 #define LOCTEXT_NAMESPACE "WorldPartition"
@@ -58,124 +55,6 @@ void UWorldPartitionRuntimeHash::ForceExternalActorLevelReference(bool bForceExt
 		}
 	}
 }
-
-void UWorldPartitionRuntimeHash::CheckForErrors() const
-{
-	TMap<FGuid, FWorldPartitionActorViewProxy> ActorDescList;
-	for (UActorDescContainer::TConstIterator<> ActorDescIt(GetOuterUWorldPartition()); ActorDescIt; ++ActorDescIt)
-	{
-		ActorDescList.Emplace(ActorDescIt->GetGuid(), *ActorDescIt);
-	}
-
-	CheckForErrorsInternal(ActorDescList);
-}
-
-void UWorldPartitionRuntimeHash::CheckForErrorsInternal(const TMap<FGuid, FWorldPartitionActorViewProxy>& ActorDescList) const
-{
-	AWorldDataLayers* WorldDataLayers = GetWorld()->GetWorldDataLayers();
-	auto GetActorLabel = [](const FWorldPartitionActorDescView& ActorDescView) -> FString
-	{
-		const FName ActorLabel = ActorDescView.GetActorLabel();
-		if (!ActorLabel.IsNone())
-		{
-			return ActorLabel.ToString();
-		}
-
-		const FString ActorPath = ActorDescView.GetActorPath().ToString();
-
-		FString SubObjectName;
-		FString SubObjectContext;
-		if (FString(ActorPath).Split(TEXT("."), &SubObjectContext, &SubObjectName))
-		{
-			return SubObjectName;
-		}
-
-		return ActorPath;
-	};
-
-	for (auto& ActorDescListPair: ActorDescList)
-	{
-		const FWorldPartitionActorViewProxy& ActorDescView = ActorDescListPair.Value;
-
-		if (!ActorDescView.GetActorIsEditorOnly())
-		{
-			for (const FGuid ActorDescRefGuid : ActorDescView.GetReferences())
-			{
-				if (const FWorldPartitionActorViewProxy* ActorDescRefView = ActorDescList.Find(ActorDescRefGuid))
-				{
-					const bool bIsActorDescAlwaysLoaded = ActorDescView.GetGridPlacement() == EActorGridPlacement::AlwaysLoaded;
-					const bool bIsActorDescRefAlwaysLoaded = ActorDescRefView->GetGridPlacement() == EActorGridPlacement::AlwaysLoaded;
-
-					if (bIsActorDescAlwaysLoaded != bIsActorDescRefAlwaysLoaded)
-					{
-						const FText StreamedActor(LOCTEXT("MapCheck_WorldPartition_StreamedActor", "Streamed actor"));
-						const FText AlwaysLoadedActor(LOCTEXT("MapCheck_WorldPartition_AlwaysLoadedActor", "Always loaded actor"));
-
-						TSharedRef<FTokenizedMessage> Error = FMessageLog("MapCheck").Warning()
-							->AddToken(FTextToken::Create(bIsActorDescAlwaysLoaded ? AlwaysLoadedActor : StreamedActor))
-							->AddToken(FAssetNameToken::Create(GetActorLabel(ActorDescView)))
-							->AddToken(FTextToken::Create(LOCTEXT("MapCheck_WorldPartition_References", "references")))
-							->AddToken(FTextToken::Create(bIsActorDescRefAlwaysLoaded ? AlwaysLoadedActor : StreamedActor)								)
-							->AddToken(FAssetNameToken::Create(GetActorLabel(*ActorDescRefView)))
-							->AddToken(FMapErrorToken::Create(FName(TEXT("WorldPartition_StreamedActorReferenceAlwaysLoadedActor_CheckForErrors"))));
-					}
-
-					if (ensure(WorldDataLayers != nullptr))
-					{
-						TArray<FName> ActorDescLayerNames = ActorDescView.GetDataLayers();
-						ActorDescLayerNames.Sort([](const FName& A, const FName& B) { return A.FastLess(B); });
-						TArray<const UDataLayer*> ActorDescLayers = WorldDataLayers->GetDataLayerObjects(ActorDescLayerNames);
-
-						TArray<FName> ActorDescRefLayerNames = ActorDescRefView->GetDataLayers();
-						ActorDescRefLayerNames.Sort([](const FName& A, const FName& B) { return A.FastLess(B); });
-						TArray<const UDataLayer*> ActorDescRefLayers = WorldDataLayers->GetDataLayerObjects(ActorDescRefLayerNames);
-
-						if (ActorDescLayers != ActorDescRefLayers)
-						{
-							TSharedRef<FTokenizedMessage> Error = FMessageLog("MapCheck").Error()
-								->AddToken(FTextToken::Create(LOCTEXT("MapCheck_WorldPartition_Actor", "Actor")))
-								->AddToken(FAssetNameToken::Create(GetActorLabel(ActorDescView)))
-								->AddToken(FTextToken::Create(LOCTEXT("MapCheck_WorldPartition_ReferenceActorInOtherDataLayers", "references an actor in a different set of data layers")))
-								->AddToken(FAssetNameToken::Create(GetActorLabel(*ActorDescRefView)))
-								->AddToken(FMapErrorToken::Create(FName(TEXT("WorldPartition_ActorReferenceActorInAnotherDataLayer_CheckForErrors"))));
-						}
-					}
-				}
-				else
-				{
-					TSharedRef<FTokenizedMessage> Error = FMessageLog("MapCheck").Warning()
-						->AddToken(FTextToken::Create(LOCTEXT("MapCheck_WorldPartition_Actor", "Actor")))
-						->AddToken(FAssetNameToken::Create(GetActorLabel(ActorDescView)))
-						->AddToken(FTextToken::Create(LOCTEXT("MapCheck_WorldPartition_HaveMissingRefsTo", "have missing references to")))
-						->AddToken(FTextToken::Create(FText::FromString(ActorDescRefGuid.ToString())))
-						->AddToken(FMapErrorToken::Create(FName(TEXT("WorldPartition_MissingActorReference_CheckForErrors"))));
-				}
-			}
-		}
-	}
-
-	// Check Level Script Blueprint
-	if (ULevelScriptBlueprint* LevelScriptBlueprint = GetOuterUWorldPartition()->GetWorld()->PersistentLevel->GetLevelScriptBlueprint(true))
-	{
-		TArray<AActor*> LevelScriptExternalActorReferences = ActorsReferencesUtils::GetExternalActorReferences(LevelScriptBlueprint);
-
-		for (AActor* Actor : LevelScriptExternalActorReferences)
-		{
-			if (const FWorldPartitionActorDescView* ActorDescView = ActorDescList.Find(Actor->GetActorGuid()))
-			{
-				TArray<FName> ActorDescLayers = ActorDescView->GetDataLayers();
-				if (ActorDescLayers.Num())
-				{
-					TSharedRef<FTokenizedMessage> Error = FMessageLog("MapCheck").Error()
-						->AddToken(FTextToken::Create(LOCTEXT("MapCheck_WorldPartition_LevelScriptBlueprintActorReference", "Level Script Blueprint references actor")))
-						->AddToken(FAssetNameToken::Create(GetActorLabel(*ActorDescView)))
-						->AddToken(FTextToken::Create(LOCTEXT("MapCheck_WorldPartition_LevelScriptBlueprintDataLayerReference", "with a non empty set of data layers")))
-						->AddToken(FMapErrorToken::Create(FName(TEXT("WorldPartition_LevelScriptBlueprintRefefenceDataLayer_CheckForErrors"))));
-				}
-			}
-		}
-	}
-}
 #endif
 
 void UWorldPartitionRuntimeHash::SortStreamingCellsByImportance(const TSet<const UWorldPartitionRuntimeCell*>& InCells, const TArray<FWorldPartitionStreamingSource>& InSources, TArray<const UWorldPartitionRuntimeCell*, TInlineAllocator<256>>& OutSortedCells) const
diff --git a/Engine/Source/Runtime/Engine/Private/WorldPartition/WorldPartitionStreamingGeneration.cpp b/Engine/Source/Runtime/Engine/Private/WorldPartition/WorldPartitionStreamingGeneration.cpp
--- a/Engine/Source/Runtime/Engine/Private/WorldPartition/WorldPartitionStreamingGeneration.cpp
+++ b/Engine/Source/Runtime/Engine/Private/WorldPartition/WorldPartitionStreamingGeneration.cpp
@@ -9,13 +9,12 @@
 #include "WorldPartition/WorldPartitionRuntimeHash.h"
 #include "WorldPartition/WorldPartitionStreamingPolicy.h"
 #include "WorldPartition/WorldPartitionActorCluster.h"
+#include "WorldPartition/ErrorHandling/WorldPartitionStreamingGenerationLogErrorHandler.h"
+#include "WorldPartition/ErrorHandling/WorldPartitionStreamingGenerationMapCheckErrorHandler.h"
 #include "WorldPartition/HLOD/HLODActor.h"
-#endif
 
 #define LOCTEXT_NAMESPACE "WorldPartition"
 
-#if WITH_EDITOR
-
 /*
 	Preparation Phase
 		Actor Descriptor Views Creation
@@ -36,7 +35,7 @@ class FWorldPartitionStreamingGenerator
 {
 	void CreateActorDescViewMap(const UActorDescContainer* InContainer, TMap<FGuid, FWorldPartitionActorDescView>& OutActorDescViewMap, const FActorContainerID& InContainerID)
 	{
-		const bool bIncludeUnsavedActors = (bIsPIE && InContainerID.IsMainContainer());
+		const bool bHandleUnsavedActors = bIncludeUnsavedActors && InContainerID.IsMainContainer();
 
 		// Consider all actors of a /Temp/ container package as Unsaved because loading them from disk will fail (Outer world name mismatch)
 		const bool bIsTempContainerPackage = FPackageName::IsTempPackage(InContainer->GetPackage()->GetName());
@@ -48,7 +47,7 @@ class FWorldPartitionStreamingGenerator
 			{
 				AActor* Actor = ActorDescIt->GetActor();
 
-				if (bIncludeUnsavedActors && IsValid(Actor) && (bIsTempContainerPackage || Actor->GetPackage()->IsDirty()))
+				if (bHandleUnsavedActors && IsValid(Actor) && (bIsTempContainerPackage || Actor->GetPackage()->IsDirty()))
 				{
 					// Dirty, unsaved actor for PIE
 					FWorldPartitionActorDesc* ActorDesc = RuntimeHash->ModifiedActorDescListForPIE.AddActor(Actor);
@@ -64,7 +63,7 @@ class FWorldPartitionStreamingGenerator
 		}
 
 		// Append new unsaved actors for the persistent level
-		if (bIncludeUnsavedActors)
+		if (bHandleUnsavedActors)
 		{
 			for (AActor* Actor : InContainer->GetWorld()->PersistentLevel->Actors)
 			{
@@ -167,15 +166,32 @@ class FWorldPartitionStreamingGenerator
 
 					for (AActor* Actor : LevelScriptExternalActorReferences)
 					{
-						if (FWorldPartitionActorDescView* ActorDescView = ContainerDescriptor.ActorDescViewMap.Find(Actor->GetActorGuid()))
+						FWorldPartitionActorDescView& ActorDescView = ContainerDescriptor.ActorDescViewMap.FindChecked(Actor->GetActorGuid());
+
+						if (ActorDescView.GetGridPlacement() != EActorGridPlacement::AlwaysLoaded)
+						{						
+							ActorDescView.SetGridPlacement(EActorGridPlacement::AlwaysLoaded);
+
+							if (ErrorHandler)
+							{
+								ErrorHandler->OnInvalidReferenceLevelScriptStreamed(ActorDescView);
+							}
+						}
+
+						if (ActorDescView.GetDataLayers().Num())
 						{
-							ActorDescView->SetGridPlacement(EActorGridPlacement::AlwaysLoaded);
+							ActorDescView.SetInvalidDataLayers();
+
+							if (ErrorHandler)
+							{
+								ErrorHandler->OnInvalidReferenceLevelScriptDataLayers(ActorDescView);
+							}
 						}
 					}
 				}
 			}
 
-			// Give the associated runtime hash to adjust the grid placement based on its internal settings, etc.
+			// Give the associated runtime hash the possibility to adjust actor descriptor views based on its internal settings, etc.
 			RuntimeHash->UpdateActorDescViewMap(ContainerDescriptor.ActorDescViewMap);
 
 			// Perform various adjustements based on validations and errors
@@ -189,31 +205,49 @@ class FWorldPartitionStreamingGenerator
 					if (FWorldPartitionActorDescView* ReferenceActorDescView = ContainerDescriptor.ActorDescViewMap.Find(ReferenceGuid))
 					{
 						// Validate grid placement
-						if ((ActorDescView.GetGridPlacement() != ReferenceActorDescView->GetGridPlacement()) &&
-							((ActorDescView.GetGridPlacement() == EActorGridPlacement::AlwaysLoaded) || (ReferenceActorDescView->GetGridPlacement() == EActorGridPlacement::AlwaysLoaded)))
+						const bool bIsActorDescAlwaysLoaded = ActorDescView.GetGridPlacement() == EActorGridPlacement::AlwaysLoaded;
+						const bool bIsActorDescRefAlwaysLoaded = ReferenceActorDescView->GetGridPlacement() == EActorGridPlacement::AlwaysLoaded;
+
+						if (bIsActorDescAlwaysLoaded != bIsActorDescRefAlwaysLoaded)
 						{
 							ActorDescView.SetGridPlacement(EActorGridPlacement::AlwaysLoaded);
 							ReferenceActorDescView->SetGridPlacement(EActorGridPlacement::AlwaysLoaded);
+
+							if (ErrorHandler)
+							{
+								ErrorHandler->OnInvalidReferenceGridPlacement(ActorDescView, *ReferenceActorDescView);
+							}
 						}
 
 						// Validate data layers
-						if (ActorDescView.GetDataLayers().Num() != ReferenceActorDescView->GetDataLayers().Num())
+						auto IsReferenceDataLayersValid = [](const FWorldPartitionActorDescView& ActorDescView, const FWorldPartitionActorDescView& ReferenceActorDescView)
+						{
+							if (ActorDescView.GetDataLayers().Num() == ReferenceActorDescView.GetDataLayers().Num())
+							{
+								const TSet<FName> ActorDescDataLayers(ActorDescView.GetDataLayers());
+								const TSet<FName> ReferenceActorDescDataLayers(ReferenceActorDescView.GetDataLayers());
+
+								return ActorDescDataLayers.Includes(ReferenceActorDescDataLayers);
+							}
+
+							return false;
+						};
+
+						if (!IsReferenceDataLayersValid(ActorDescView, *ReferenceActorDescView))
 						{
 							ActorDescView.SetInvalidDataLayers();
 							ReferenceActorDescView->SetInvalidDataLayers();
-						}
-						else
-						{
-							const TSet<FName> ActorDescDataLayers(ActorDescView.GetDataLayers());
-							const TSet<FName> ReferenceActorDescDataLayers(ReferenceActorDescView->GetDataLayers());
 
-							if (!ActorDescDataLayers.Includes(ReferenceActorDescDataLayers))
+							if (ErrorHandler)
 							{
-								ActorDescView.SetInvalidDataLayers();
-								ReferenceActorDescView->SetInvalidDataLayers();
+								ErrorHandler->OnInvalidReferenceDataLayers(ActorDescView, *ReferenceActorDescView);
 							}
 						}
 					}
+					else if (ErrorHandler)
+					{
+						ErrorHandler->OnInvalidReference(ActorDescView, ReferenceGuid);
+					}
 				}
 			}
 		}
@@ -252,12 +286,13 @@ class FWorldPartitionStreamingGenerator
 	}
 
 public:
-	FWorldPartitionStreamingGenerator(bool bInIsPIE, UWorldPartitionRuntimeHash* InRuntimeHash)
-	: bIsPIE(bInIsPIE)
+	FWorldPartitionStreamingGenerator(bool bInIncludeUnsavedActors, UWorldPartitionRuntimeHash* InRuntimeHash, IStreamingGenerationErrorHandler* InErrorHandler = nullptr)
+	: bIncludeUnsavedActors(bInIncludeUnsavedActors)
 	, RuntimeHash(InRuntimeHash)
+	, ErrorHandler(InErrorHandler)
 	{}
 
-	void PreparationPhase(UWorldPartition* WorldPartition)
+	void PreparationPhase(const UWorldPartition* WorldPartition)
 	{
 		// Preparation Phase :: Actor Descriptor Views Creation
 		CreateActorDescriptorViews(WorldPartition);
@@ -286,8 +321,9 @@ public:
 	}
 
 private:
-	bool bIsPIE;
+	bool bIncludeUnsavedActors;
 	UWorldPartitionRuntimeHash* RuntimeHash;
+	IStreamingGenerationErrorHandler* ErrorHandler;
 
 	struct FContainerDescriptor
 	{
@@ -316,9 +352,19 @@ private:
 
 bool UWorldPartition::GenerateStreaming(TArray<FString>* OutPackagesToGenerate)
 {
+	FStreamingGenerationLogErrorHandler LogErrorHandler;
+	FStreamingGenerationMapCheckErrorHandler MapCheckErrorHandler;	
+	IStreamingGenerationErrorHandler* ErrorHandler = &LogErrorHandler;
+
+	if (bIsPIE)
+	{
+		// In PIE, we always want to populate the map check dialog
+		ErrorHandler = &MapCheckErrorHandler;
+	}
+
 	FActorClusterContext ActorClusterContext;
 	{
-		FWorldPartitionStreamingGenerator StreamingGenerator(bIsPIE, RuntimeHash);
+		FWorldPartitionStreamingGenerator StreamingGenerator(bIsPIE, RuntimeHash, ErrorHandler);
 
 		// Preparation Phase
 		StreamingGenerator.PreparationPhase(this);
@@ -343,7 +389,8 @@ bool UWorldPartition::GenerateStreaming(TArray<FString>* OutPackagesToGenerate)
 
 void UWorldPartition::GenerateHLOD(ISourceControlHelper* SourceControlHelper, bool bCreateActorsOnly)
 {
-	FWorldPartitionStreamingGenerator StreamingGenerator(bIsPIE, RuntimeHash);
+	FStreamingGenerationLogErrorHandler LogErrorHandler;
+	FWorldPartitionStreamingGenerator StreamingGenerator(bIsPIE, RuntimeHash, &LogErrorHandler);
 	StreamingGenerator.PreparationPhase(this);
 
 	// Preparation Phase :: Actor Clusters Creation
@@ -355,6 +402,14 @@ void UWorldPartition::GenerateHLOD(ISourceControlHelper* SourceControlHelper, bo
 	RuntimeHash->GenerateHLOD(SourceControlHelper, ActorClusterContext, bCreateActorsOnly);
 }
 
+void UWorldPartition::CheckForErrors(IStreamingGenerationErrorHandler* ErrorHandler) const
+{
+	FActorClusterContext ActorClusterContext;
+	{		
+		FWorldPartitionStreamingGenerator StreamingGenerator(true, RuntimeHash, ErrorHandler);
+		StreamingGenerator.PreparationPhase(this);
+	}
+}
 #endif // WITH_EDITOR
 
 #undef LOCTEXT_NAMESPACE
\ No newline at end of file
diff --git a/Engine/Source/Runtime/Engine/Public/WorldPartition/ErrorHandling/WorldPartitionStreamingGenerationErrorHandler.h b/Engine/Source/Runtime/Engine/Public/WorldPartition/ErrorHandling/WorldPartitionStreamingGenerationErrorHandler.h
new file mode 100644
--- /dev/null
+++ b/Engine/Source/Runtime/Engine/Public/WorldPartition/ErrorHandling/WorldPartitionStreamingGenerationErrorHandler.h
@@ -0,0 +1,41 @@
+// Copyright Epic Games, Inc. All Rights Reserved.
+#pragma once
+
+#if WITH_EDITOR
+#include "CoreMinimal.h"
+#include "WorldPartition/WorldPartitionActorDescView.h"
+
+class ENGINE_API IStreamingGenerationErrorHandler
+{
+protected:
+	FString GetActorLabel(const FWorldPartitionActorDescView& ActorDescView);
+
+public:
+	virtual ~IStreamingGenerationErrorHandler() {}
+
+	/** 
+	 * Called when an actor references an invalid actor.
+	 */
+	virtual void OnInvalidReference(const FWorldPartitionActorDescView& ActorDescView, const FGuid& ReferenceGuid) = 0;
+
+	/** 
+	 * Called when an actor references an actor using a different grid placement.
+	 */
+	virtual void OnInvalidReferenceGridPlacement(const FWorldPartitionActorDescView& ActorDescView, const FWorldPartitionActorDescView& ReferenceActorDescView) = 0;
+
+	/** 
+	 * Called when an actor references an actor using a different set of data layers.
+	 */
+	virtual void OnInvalidReferenceDataLayers(const FWorldPartitionActorDescView& ActorDescView, const FWorldPartitionActorDescView& ReferenceActorDescView) = 0;
+
+	/** 
+	 * Called when the level script references a streamed actor.
+	 */
+	virtual void OnInvalidReferenceLevelScriptStreamed(const FWorldPartitionActorDescView& ActorDescView) = 0;
+
+	/** 
+	 * Called when an actor descriptor references an actor using data layers.
+	 */
+	virtual void OnInvalidReferenceLevelScriptDataLayers(const FWorldPartitionActorDescView& ActorDescView) = 0;
+};
+#endif
diff --git a/Engine/Source/Runtime/Engine/Public/WorldPartition/ErrorHandling/WorldPartitionStreamingGenerationLogErrorHandler.h b/Engine/Source/Runtime/Engine/Public/WorldPartition/ErrorHandling/WorldPartitionStreamingGenerationLogErrorHandler.h
new file mode 100644
--- /dev/null
+++ b/Engine/Source/Runtime/Engine/Public/WorldPartition/ErrorHandling/WorldPartitionStreamingGenerationLogErrorHandler.h
@@ -0,0 +1,16 @@
+// Copyright Epic Games, Inc. All Rights Reserved.
+#pragma once
+
+#if WITH_EDITOR
+#include "WorldPartition/ErrorHandling/WorldPartitionStreamingGenerationErrorHandler.h"
+
+class ENGINE_API FStreamingGenerationLogErrorHandler : public IStreamingGenerationErrorHandler
+{
+public:
+	virtual void OnInvalidReference(const FWorldPartitionActorDescView& ActorDescView, const FGuid& ReferenceGuid) override;
+	virtual void OnInvalidReferenceGridPlacement(const FWorldPartitionActorDescView& ActorDescView, const FWorldPartitionActorDescView& ReferenceActorDescView) override;
+	virtual void OnInvalidReferenceDataLayers(const FWorldPartitionActorDescView& ActorDescView, const FWorldPartitionActorDescView& ReferenceActorDescView) override;
+	virtual void OnInvalidReferenceLevelScriptStreamed(const FWorldPartitionActorDescView& ActorDescView) override;
+	virtual void OnInvalidReferenceLevelScriptDataLayers(const FWorldPartitionActorDescView& ActorDescView) override;
+};
+#endif
diff --git a/Engine/Source/Runtime/Engine/Public/WorldPartition/ErrorHandling/WorldPartitionStreamingGenerationMapCheckErrorHandler.h b/Engine/Source/Runtime/Engine/Public/WorldPartition/ErrorHandling/WorldPartitionStreamingGenerationMapCheckErrorHandler.h
new file mode 100644
--- /dev/null
+++ b/Engine/Source/Runtime/Engine/Public/WorldPartition/ErrorHandling/WorldPartitionStreamingGenerationMapCheckErrorHandler.h
@@ -0,0 +1,16 @@
+// Copyright Epic Games, Inc. All Rights Reserved.
+#pragma once
+
+#if WITH_EDITOR
+#include "WorldPartition/ErrorHandling/WorldPartitionStreamingGenerationErrorHandler.h"
+
+class ENGINE_API FStreamingGenerationMapCheckErrorHandler : public IStreamingGenerationErrorHandler
+{
+public:
+	virtual void OnInvalidReference(const FWorldPartitionActorDescView& ActorDescView, const FGuid& ReferenceGuid) override;
+	virtual void OnInvalidReferenceGridPlacement(const FWorldPartitionActorDescView& ActorDescView, const FWorldPartitionActorDescView& ReferenceActorDescView) override;
+	virtual void OnInvalidReferenceDataLayers(const FWorldPartitionActorDescView& ActorDescView, const FWorldPartitionActorDescView& ReferenceActorDescView) override;
+	virtual void OnInvalidReferenceLevelScriptStreamed(const FWorldPartitionActorDescView& ActorDescView) override;
+	virtual void OnInvalidReferenceLevelScriptDataLayers(const FWorldPartitionActorDescView& ActorDescView) override;
+};
+#endif
diff --git a/Engine/Source/Runtime/Engine/Public/WorldPartition/WorldPartition.h b/Engine/Source/Runtime/Engine/Public/WorldPartition/WorldPartition.h
--- a/Engine/Source/Runtime/Engine/Public/WorldPartition/WorldPartition.h
+++ b/Engine/Source/Runtime/Engine/Public/WorldPartition/WorldPartition.h
@@ -26,6 +26,7 @@ class UWorldPartitionEditorHash;
 class UWorldPartitionRuntimeCell;
 class UWorldPartitionRuntimeHash;
 class UWorldPartitionStreamingPolicy;
+class IStreamingGenerationErrorHandler;
 class FHLODActorDesc;
 class UCanvas;
 
@@ -141,7 +142,7 @@ public:
 	void DrawRuntimeHashPreview();
 	void DumpActorDescs(const FString& Path);
 
-	void CheckForErrors() const;
+	void CheckForErrors(IStreamingGenerationErrorHandler* ErrorHandler) const;
 
 	uint32 GetWantedEditorCellSize() const;
 	void SetEditorWantedCellSize(uint32 InCellSize);
diff --git a/Engine/Source/Runtime/Engine/Public/WorldPartition/WorldPartitionRuntimeHash.h b/Engine/Source/Runtime/Engine/Public/WorldPartition/WorldPartitionRuntimeHash.h
--- a/Engine/Source/Runtime/Engine/Public/WorldPartition/WorldPartitionRuntimeHash.h
+++ b/Engine/Source/Runtime/Engine/Public/WorldPartition/WorldPartitionRuntimeHash.h
@@ -40,8 +40,6 @@ class ENGINE_API UWorldPartitionRuntimeHash : public UObject
 	virtual bool GenerateNavigationData(const FBox& LoadedBounds) { return false; }
 	virtual void DrawPreview() const {}
 
-	void CheckForErrors() const;
-
 	// PIE/Game methods
 	void OnBeginPlay();
 	void OnEndPlay();
@@ -74,8 +72,6 @@ protected:
 	virtual EWorldPartitionStreamingPerformance GetStreamingPerformanceForCell(const UWorldPartitionRuntimeCell* Cell) const { return EWorldPartitionStreamingPerformance::Good; }
 
 #if WITH_EDITOR
-	virtual void CheckForErrorsInternal(const TMap<FGuid, FWorldPartitionActorViewProxy>& ActorDescList) const;
-
 public:
 	virtual void UpdateActorDescViewMap(TMap<FGuid, FWorldPartitionActorDescView>& ActorDescViewMap) const {}
 #endif
