LevelInstance: Editor Mode

diff --git a/Engine/Source/Editor/LevelInstanceEditor/Private/LevelInstanceEditorMode.cpp b/Engine/Source/Editor/LevelInstanceEditor/Private/LevelInstanceEditorMode.cpp
new file mode 100644
--- /dev/null
+++ b/Engine/Source/Editor/LevelInstanceEditor/Private/LevelInstanceEditorMode.cpp
@@ -0,0 +1,130 @@
+// Copyright Epic Games, Inc. All Rights Reserved.
+
+#include "LevelInstanceEditorMode.h"
+#include "LevelInstanceEditorModeToolkit.h"
+#include "LevelInstanceEditorModeCommands.h"
+#include "Editor.h"
+#include "Engine/World.h"
+#include "LevelInstance/LevelInstanceSubsystem.h"
+#include "LevelInstance/LevelInstanceActor.h"
+#include "LevelEditorViewport.h"
+#include "LevelEditorActions.h"
+#include "EditorModeManager.h"
+
+#define LOCTEXT_NAMESPACE "LevelInstanceEditorMode"
+
+FEditorModeID ULevelInstanceEditorMode::EM_LevelInstanceEditorModeId("EditMode.LevelInstance");
+
+ULevelInstanceEditorMode::ULevelInstanceEditorMode()
+	: UEdMode()
+{
+	Info = FEditorModeInfo(ULevelInstanceEditorMode::EM_LevelInstanceEditorModeId,
+		LOCTEXT("LevelInstanceEditorModeName", "LevelInstanceEditorMode"),
+		FSlateIcon(),
+		false);
+
+	bContextRestriction = true;
+}
+
+ULevelInstanceEditorMode::~ULevelInstanceEditorMode()
+{
+}
+
+void ULevelInstanceEditorMode::UpdateEngineShowFlags()
+{
+	for (FLevelEditorViewportClient* LevelVC : GEditor->GetLevelViewportClients())
+	{
+		if (LevelVC && LevelVC->GetWorld())
+		{
+			if(ULevelInstanceSubsystem* LevelInstanceSubsystem = LevelVC->GetWorld()->GetSubsystem<ULevelInstanceSubsystem>())
+			{
+				LevelVC->EngineShowFlags.EditingLevelInstance = !!LevelInstanceSubsystem->GetEditingLevelInstance();
+			}
+		}
+	}
+}
+
+void ULevelInstanceEditorMode::Enter()
+{
+	UEdMode::Enter();
+
+	UpdateEngineShowFlags();
+}
+
+void ULevelInstanceEditorMode::Exit()
+{
+	UEdMode::Exit();
+		
+	UpdateEngineShowFlags();
+}
+
+void ULevelInstanceEditorMode::CreateToolkit()
+{
+	Toolkit = MakeShared<FLevelInstanceEditorModeToolkit>();
+}
+
+void ULevelInstanceEditorMode::BindCommands()
+{
+	UEdMode::BindCommands();
+	const TSharedRef<FUICommandList>& CommandList = Toolkit->GetToolkitCommands();
+	const FLevelInstanceEditorModeCommands& Commands = FLevelInstanceEditorModeCommands::Get();
+
+	CommandList->MapAction(
+		Commands.ExitMode,
+		FExecuteAction::CreateUObject(this, &ULevelInstanceEditorMode::ExitModeCommand));
+
+	CommandList->MapAction(
+		Commands.ToggleContextRestriction,
+		FExecuteAction::CreateUObject(this, &ULevelInstanceEditorMode::ToggleContextRestrictionCommand),
+		FCanExecuteAction(),
+		FIsActionChecked::CreateUObject(this, &ULevelInstanceEditorMode::IsContextRestrictionCommandEnabled));
+}
+
+bool ULevelInstanceEditorMode::IsSelectionDisallowed(AActor* InActor, bool bInSelection) const
+{
+	const bool bRestrict = bContextRestriction && bInSelection;
+
+	if (bRestrict)
+	{
+		if (UWorld* World = InActor->GetWorld())
+		{
+			if (ALevelInstance* LevelInstanceActor = Cast<ALevelInstance>(InActor))
+			{
+				if (LevelInstanceActor->IsEditing())
+				{
+					return false;
+				}
+			}
+
+			if (ULevelInstanceSubsystem* LevelInstanceSubsystem = World->GetSubsystem<ULevelInstanceSubsystem>())
+			{
+				ALevelInstance* EditingLevelInstance = LevelInstanceSubsystem->GetEditingLevelInstance();
+				ALevelInstance* LevelInstance = LevelInstanceSubsystem->GetParentLevelInstance(InActor);
+
+				return EditingLevelInstance != LevelInstance;
+			}
+		}
+	}
+
+	return bRestrict;
+}
+
+void ULevelInstanceEditorMode::ExitModeCommand()
+{	
+	if (FEditorModeTools* Manager = GetModeManager())
+	{
+		Manager->DeactivateMode(EM_LevelInstanceEditorModeId);
+	}
+}
+
+void ULevelInstanceEditorMode::ToggleContextRestrictionCommand()
+{
+	bContextRestriction = !bContextRestriction;
+}
+
+bool ULevelInstanceEditorMode::IsContextRestrictionCommandEnabled() const
+{
+	return bContextRestriction;
+}
+
+#undef LOCTEXT_NAMESPACE
diff --git a/Engine/Source/Editor/LevelInstanceEditor/Private/LevelInstanceEditorMode.h b/Engine/Source/Editor/LevelInstanceEditor/Private/LevelInstanceEditorMode.h
new file mode 100644
--- /dev/null
+++ b/Engine/Source/Editor/LevelInstanceEditor/Private/LevelInstanceEditorMode.h
@@ -0,0 +1,44 @@
+// Copyright Epic Games, Inc. All Rights Reserved.
+#pragma once
+
+#include "CoreMinimal.h"
+#include "Tools/UEdMode.h"
+#include "Framework/Commands/UICommandList.h"
+#include "LevelInstanceEditorMode.generated.h"
+
+UCLASS()
+class ULevelInstanceEditorMode : public UEdMode
+{
+	GENERATED_BODY()
+public:
+	static FEditorModeID EM_LevelInstanceEditorModeId;
+
+	/** Constructor */
+	ULevelInstanceEditorMode();
+
+	/** Destructor */
+	virtual ~ULevelInstanceEditorMode();
+
+	// Begin UEdMode
+	virtual void Enter() override;
+	virtual void Exit() override;
+	virtual void CreateToolkit() override;
+	
+	virtual bool IsCompatibleWith(FEditorModeID OtherModeID) const override { return true; }
+
+	/** Check to see if an actor can be selected in this mode - no side effects */
+	virtual bool IsSelectionDisallowed(AActor* InActor, bool bInSelection) const override;
+	
+	// End UEdMode
+
+private:
+	void UpdateEngineShowFlags();
+	virtual void BindCommands() override;
+
+	void ExitModeCommand();
+	void ToggleContextRestrictionCommand();
+	bool IsContextRestrictionCommandEnabled() const;
+
+	bool bContextRestriction;
+};
+
diff --git a/Engine/Source/Editor/LevelInstanceEditor/Private/LevelInstanceEditorModeCommands.cpp b/Engine/Source/Editor/LevelInstanceEditor/Private/LevelInstanceEditorModeCommands.cpp
new file mode 100644
--- /dev/null
+++ b/Engine/Source/Editor/LevelInstanceEditor/Private/LevelInstanceEditorModeCommands.cpp
@@ -0,0 +1,13 @@
+// Copyright Epic Games, Inc. All Rights Reserved.
+
+#include "LevelInstanceEditorModeCommands.h"
+
+#define LOCTEXT_NAMESPACE "LevelInstanceEditorModeCommands"
+
+void FLevelInstanceEditorModeCommands::RegisterCommands()
+{
+	UI_COMMAND(ExitMode, "Exit Mode", "Exits mode asking to save pending changes", EUserInterfaceActionType::Button, FInputChord(EKeys::Escape));
+	UI_COMMAND(ToggleContextRestriction, "Toggle Context Restriction", "Toggles edit restrictions like selection to current level edit", EUserInterfaceActionType::ToggleButton, FInputChord(EKeys::L, EModifierKey::Shift));
+}
+
+#undef LOCTEXT_NAMESPACE
\ No newline at end of file
diff --git a/Engine/Source/Editor/LevelInstanceEditor/Private/LevelInstanceEditorModeCommands.h b/Engine/Source/Editor/LevelInstanceEditor/Private/LevelInstanceEditorModeCommands.h
new file mode 100644
--- /dev/null
+++ b/Engine/Source/Editor/LevelInstanceEditor/Private/LevelInstanceEditorModeCommands.h
@@ -0,0 +1,33 @@
+// Copyright Epic Games, Inc. All Rights Reserved.
+
+#pragma once
+
+#include "CoreMinimal.h"
+#include "Framework/Commands/Commands.h"
+#include "EditorStyleSet.h"
+
+class FLevelInstanceEditorModeCommands : public TCommands<FLevelInstanceEditorModeCommands>
+{
+
+public:
+	FLevelInstanceEditorModeCommands() : TCommands<FLevelInstanceEditorModeCommands>
+		(
+			"LevelInstanceEditorMode",
+			NSLOCTEXT("Contexts", "LevelInstanceEditor", "Level Instance Editor"),
+			NAME_None,
+			FEditorStyle::GetStyleSetName() // Icon Style Set
+		)
+	{}
+
+
+	/** Exit Editor Mode */
+	TSharedPtr<FUICommandInfo> ExitMode;
+
+	/** Toggle Restrict Selection */
+	TSharedPtr<FUICommandInfo> ToggleContextRestriction;
+		
+	/**
+	 * Initialize commands
+	 */
+	virtual void RegisterCommands() override;
+};
\ No newline at end of file
diff --git a/Engine/Source/Editor/LevelInstanceEditor/Private/LevelInstanceEditorModeToolkit.cpp b/Engine/Source/Editor/LevelInstanceEditor/Private/LevelInstanceEditorModeToolkit.cpp
new file mode 100644
--- /dev/null
+++ b/Engine/Source/Editor/LevelInstanceEditor/Private/LevelInstanceEditorModeToolkit.cpp
@@ -0,0 +1,25 @@
+// Copyright Epic Games, Inc. All Rights Reserved.
+#include "LevelInstanceEditorModeToolkit.h"
+
+#define LOCTEXT_NAMESPACE "LevelInstanceEditorModeToolkit"
+
+FLevelInstanceEditorModeToolkit::FLevelInstanceEditorModeToolkit()
+{
+}
+
+FName FLevelInstanceEditorModeToolkit::GetToolkitFName() const
+{
+	return FName("LevelInstanceEditorModeToolkit");
+}
+
+FText FLevelInstanceEditorModeToolkit::GetBaseToolkitName() const
+{
+	return LOCTEXT("ToolkitDisplayName", "Level Instance Editor Mode");
+}
+
+void FLevelInstanceEditorModeToolkit::SetModeUILayer(const TSharedPtr<FAssetEditorModeUILayer> InLayer)
+{
+	// Do nothing (prevents Ed Mode Tab spawn)
+}
+
+#undef LOCTEXT_NAMESPACE
\ No newline at end of file
diff --git a/Engine/Source/Editor/LevelInstanceEditor/Private/LevelInstanceEditorModeToolkit.h b/Engine/Source/Editor/LevelInstanceEditor/Private/LevelInstanceEditorModeToolkit.h
new file mode 100644
--- /dev/null
+++ b/Engine/Source/Editor/LevelInstanceEditor/Private/LevelInstanceEditorModeToolkit.h
@@ -0,0 +1,17 @@
+// Copyright Epic Games, Inc. All Rights Reserved.
+#pragma once
+
+#include "CoreMinimal.h"
+#include "Toolkits/BaseToolkit.h"
+
+class FLevelInstanceEditorModeToolkit : public FModeToolkit
+{
+public:
+	FLevelInstanceEditorModeToolkit();
+
+	// IToolkit interface 
+	virtual FName GetToolkitFName() const override;
+	virtual FText GetBaseToolkitName() const override;
+
+	virtual void SetModeUILayer(const TSharedPtr<FAssetEditorModeUILayer> InLayer) override;
+};
\ No newline at end of file
diff --git a/Engine/Source/Editor/LevelInstanceEditor/Private/LevelInstanceEditorModule.cpp b/Engine/Source/Editor/LevelInstanceEditor/Private/LevelInstanceEditorModule.cpp
--- a/Engine/Source/Editor/LevelInstanceEditor/Private/LevelInstanceEditorModule.cpp
+++ b/Engine/Source/Editor/LevelInstanceEditor/Private/LevelInstanceEditorModule.cpp
@@ -9,6 +9,10 @@
 #include "LevelInstanceEditorSettings.h"
 #include "ToolMenus.h"
 #include "Editor.h"
+#include "EditorModeManager.h"
+#include "EditorModeRegistry.h"
+#include "LevelInstanceEditorMode.h"
+#include "LevelInstanceEditorModeCommands.h"
 #include "LevelEditorMenuContext.h"
 #include "ContentBrowserMenuContexts.h"
 #include "ContentBrowserModule.h"
@@ -616,15 +620,19 @@ void FLevelInstanceEditorModule::StartupModule()
 	
 	EditorLevelUtils::CanMoveActorToLevelDelegate.AddRaw(this, &FLevelInstanceEditorModule::CanMoveActorToLevel);
 
-	FLevelEditorModule& LevelEditorModule = FModuleManager::LoadModuleChecked<FLevelEditorModule>("LevelEditor");
-	LevelEditorModule.OnMapChanged().AddRaw(this, &FLevelInstanceEditorModule::OnMapChanged);
-
 	FMessageLogModule& MessageLogModule = FModuleManager::LoadModuleChecked<FMessageLogModule>("MessageLog");
 	FMessageLogInitializationOptions InitOptions;
 	InitOptions.bShowFilters = true;
 	InitOptions.bShowPages = false;
 	InitOptions.bAllowClear = true;
 	MessageLogModule.RegisterLogListing("LevelInstance", LOCTEXT("LevelInstanceLog", "Level Instance Log"), InitOptions);
+		
+	FLevelInstanceEditorModeCommands::Register();
+
+	if (!IsRunningCommandlet())
+	{
+		GLevelEditorModeTools().OnEditorModeIDChanged().AddRaw(this, &FLevelInstanceEditorModule::OnEditorModeIDChanged);
+	}
 }
 
 void FLevelInstanceEditorModule::ShutdownModule()
@@ -636,30 +644,40 @@ void FLevelInstanceEditorModule::ShutdownModule()
 
 	EditorLevelUtils::CanMoveActorToLevelDelegate.RemoveAll(this);
 
-	if (FLevelEditorModule* LevelEditorModule = FModuleManager::GetModulePtr<FLevelEditorModule>("LevelEditor"))
+	if (!IsRunningCommandlet() && GLevelEditorModeToolsIsValid())
 	{
-		LevelEditorModule->OnMapChanged().RemoveAll(this);
+		GLevelEditorModeTools().OnEditorModeIDChanged().RemoveAll(this);
 	}
 }
 
-void FLevelInstanceEditorModule::OnLevelActorDeleted(AActor* Actor)
+void FLevelInstanceEditorModule::OnEditorModeIDChanged(const FEditorModeID& InModeID, bool bIsEnteringMode)
 {
-	if (ULevelInstanceSubsystem* LevelInstanceSubsystem = Actor->GetWorld()->GetSubsystem<ULevelInstanceSubsystem>())
+	if (InModeID == ULevelInstanceEditorMode::EM_LevelInstanceEditorModeId && !bIsEnteringMode)
 	{
-		LevelInstanceSubsystem->OnActorDeleted(Actor);
+		ExitEditorModeEvent.Broadcast();
 	}
 }
 
-void FLevelInstanceEditorModule::OnMapChanged(UWorld* World, EMapChangeType MapChangeType)
+void FLevelInstanceEditorModule::ActivateEditorMode()
 {
-	// On Map Changed, Users will be asked to save unchanged maps. Once we hit the teardown we need to force
-	// LevelInstance Edits to be cancelled. If they are still dirty it means the user decided not to save changes.
-	if (World && MapChangeType == EMapChangeType::TearDownWorld)
+	if (!GLevelEditorModeTools().IsModeActive(ULevelInstanceEditorMode::EM_LevelInstanceEditorModeId))
 	{
-		if (ULevelInstanceSubsystem* LevelInstanceSubsystem = World->GetSubsystem<ULevelInstanceSubsystem>())
-		{
-			LevelInstanceSubsystem->DiscardEdits();
-		}
+		GLevelEditorModeTools().ActivateMode(ULevelInstanceEditorMode::EM_LevelInstanceEditorModeId);
+	}
+}
+void FLevelInstanceEditorModule::DeactivateEditorMode()
+{
+	if (GLevelEditorModeTools().IsModeActive(ULevelInstanceEditorMode::EM_LevelInstanceEditorModeId))
+	{
+		GLevelEditorModeTools().DeactivateMode(ULevelInstanceEditorMode::EM_LevelInstanceEditorModeId);
+	}
+}
+
+void FLevelInstanceEditorModule::OnLevelActorDeleted(AActor* Actor)
+{
+	if (ULevelInstanceSubsystem* LevelInstanceSubsystem = Actor->GetWorld()->GetSubsystem<ULevelInstanceSubsystem>())
+	{
+		LevelInstanceSubsystem->OnActorDeleted(Actor);
 	}
 }
 
diff --git a/Engine/Source/Editor/LevelInstanceEditor/Private/LevelInstanceEditorModule.h b/Engine/Source/Editor/LevelInstanceEditor/Private/LevelInstanceEditorModule.h
--- a/Engine/Source/Editor/LevelInstanceEditor/Private/LevelInstanceEditorModule.h
+++ b/Engine/Source/Editor/LevelInstanceEditor/Private/LevelInstanceEditorModule.h
@@ -3,6 +3,7 @@
 
 #include "CoreMinimal.h"
 #include "LevelInstance/ILevelInstanceEditorModule.h"
+#include "Tools/Modes.h"
 
 class AActor;
 class ULevel;
@@ -26,10 +27,18 @@ public:
 	 */
 	virtual void ShutdownModule();
 
+	virtual void ActivateEditorMode() override;
+	virtual void DeactivateEditorMode() override;
+		
+	DECLARE_DERIVED_EVENT(FLevelInstanceEditorModule, ILevelInstanceEditorModule::FExitEditorModeEvent, FExitEditorModeEvent);
+	virtual FExitEditorModeEvent& OnExitEditorMode() override { return ExitEditorModeEvent; }
+
 private:
+	void OnEditorModeIDChanged(const FEditorModeID& InModeID, bool bIsEnteringMode);
 	void OnLevelActorDeleted(AActor* Actor);
-	void OnMapChanged(UWorld* World, EMapChangeType MapChangeType);
 	void CanMoveActorToLevel(const AActor* ActorToMove, const ULevel* DestLevel, bool& bOutCanMove);
 
 	void ExtendContextMenu();
+
+	FExitEditorModeEvent ExitEditorModeEvent;
 };
diff --git a/Engine/Source/Runtime/Engine/Private/LevelInstance/LevelInstanceSubsystem.cpp b/Engine/Source/Runtime/Engine/Private/LevelInstance/LevelInstanceSubsystem.cpp
--- a/Engine/Source/Runtime/Engine/Private/LevelInstance/LevelInstanceSubsystem.cpp
+++ b/Engine/Source/Runtime/Engine/Private/LevelInstance/LevelInstanceSubsystem.cpp
@@ -31,7 +31,6 @@
 #include "Engine/Selection.h"
 #include "Engine/LevelBounds.h"
 #include "LevelInstance/LevelInstanceEditorInstanceActor.h"
-#include "LevelEditorViewport.h"
 #include "Modules/ModuleManager.h"
 #include "Engine/Blueprint.h"
 #include "LevelInstance/Packed/PackedLevelInstanceActor.h"
@@ -41,6 +40,7 @@
 #include "UObject/ObjectSaveContext.h"
 #include "UObject/UObjectGlobals.h"
 #include "EditorActorFolders.h"
+#include "Misc/MessageDialog.h"
 #endif
 
 #include "HAL/IConsoleManager.h"
@@ -67,7 +67,8 @@ void ULevelInstanceSubsystem::Initialize(FSubsystemCollectionBase& Collection)
 #if WITH_EDITOR
 	if (GEditor)
 	{
-		FModuleManager::LoadModuleChecked<ILevelInstanceEditorModule>("LevelInstanceEditor");
+		ILevelInstanceEditorModule& EditorModule = FModuleManager::LoadModuleChecked<ILevelInstanceEditorModule>("LevelInstanceEditor");
+		EditorModule.OnExitEditorMode().AddUObject(this, &ULevelInstanceSubsystem::OnExitEditorMode);
 
 		if (!GetWorld()->IsGameWorld())
 		{
@@ -83,6 +84,9 @@ void ULevelInstanceSubsystem::Deinitialize()
 	Super::Deinitialize();
 
 #if WITH_EDITOR
+	ILevelInstanceEditorModule& EditorModule = FModuleManager::LoadModuleChecked<ILevelInstanceEditorModule>("LevelInstanceEditor");
+	EditorModule.OnExitEditorMode().RemoveAll(this);
+
 	if (!GetWorld()->IsGameWorld())
 	{
 		FCoreUObjectDelegates::OnObjectPreSave.RemoveAll(this);
@@ -461,6 +465,22 @@ void ULevelInstanceSubsystem::OnPackageDeleted(UPackage* Package)
 	OnPackageChanged(Package);
 }
 
+void ULevelInstanceSubsystem::OnExitEditorMode()
+{
+	if (LevelInstanceEdit && !bCreatingLevelInstance)
+	{
+		bool bDiscard = true;
+		bool bIsDirty = IsLevelInstanceEditDirty(LevelInstanceEdit.Get());
+		if (bIsDirty)
+		{
+			FText Title = LOCTEXT("CommitOrDiscardChangesTitle", "Save changes?");
+			bDiscard = (FMessageDialog::Open(EAppMsgType::YesNo, LOCTEXT("CommitOrDiscardChangesMsg", "Unsaved Level changes will get discarded. Do you want to save them now?"), &Title) == EAppReturnType::No);
+		}
+
+		CommitLevelInstanceInternal(LevelInstanceEdit, bDiscard, /*bPromptForSave=*/false);
+	}
+}
+
 bool ULevelInstanceSubsystem::CanPackLevelInstances() const
 {
 	return !LevelInstanceEdit;
@@ -779,6 +799,7 @@ bool ULevelInstanceSubsystem::MoveActorsTo(ALevelInstance* LevelInstanceActor, c
 
 ALevelInstance* ULevelInstanceSubsystem::CreateLevelInstanceFrom(const TArray<AActor*>& ActorsToMove, const FNewLevelInstanceParams& CreationParams)
 {
+	TGuardValue<bool> CreateLevelInstanceGuard(bCreatingLevelInstance, true);
 	ULevel* CurrentLevel = GetWorld()->GetCurrentLevel();
 		
 	if (ActorsToMove.Num() == 0)
@@ -962,6 +983,13 @@ ALevelInstance* ULevelInstanceSubsystem::CreateLevelInstanceFrom(const TArray<AA
 	ALevelInstance* CommittedLevelInstance = CommitLevelInstanceInternal(TempLevelInstanceEdit, /*bDiscardEdits*/false, CreationParams.bPromptForSave, &DirtyPackages);
 	check(!TempLevelInstanceEdit);
 
+	// EditorLevelUtils::CreateNewStreamingLevelForWorld deactivates all modes. Re-activate if needed
+	if (LevelInstanceEdit)
+	{
+		ILevelInstanceEditorModule& EditorModule = FModuleManager::GetModuleChecked<ILevelInstanceEditorModule>("LevelInstanceEditor");
+		EditorModule.ActivateEditorMode();
+	}
+
 	return CommittedLevelInstance;
 }
 
@@ -1171,17 +1199,6 @@ bool ULevelInstanceSubsystem::LevelInstanceHasLevelScriptBlueprint(const ALevelI
 	return false;
 }
 
-void ULevelInstanceSubsystem::UpdateEngineShowFlags()
-{
-	for (FLevelEditorViewportClient* LevelVC : GEditor->GetLevelViewportClients())
-	{
-		if (LevelVC && LevelVC->GetWorld() == GetWorld())
-		{
-			LevelVC->EngineShowFlags.EditingLevelInstance = !!LevelInstanceEdit;
-		}
-	}
-}
-
 void ULevelInstanceSubsystem::RemoveLevelsFromWorld(const TArray<ULevel*>& InLevels, bool bResetTrans)
 {
 	if (LevelsToRemoveScope && LevelsToRemoveScope->IsValid())
@@ -1545,10 +1562,14 @@ bool ULevelInstanceSubsystem::CanCommitLevelInstance(const ALevelInstance* Level
 
 void ULevelInstanceSubsystem::EditLevelInstance(ALevelInstance* LevelInstanceActor, TWeakObjectPtr<AActor> ContextActorPtr)
 {
-	EditLevelInstanceInternal(LevelInstanceActor, ContextActorPtr, false);
+	if (EditLevelInstanceInternal(LevelInstanceActor, ContextActorPtr, false))
+	{
+		ILevelInstanceEditorModule& EditorModule = FModuleManager::GetModuleChecked<ILevelInstanceEditorModule>("LevelInstanceEditor");
+		EditorModule.ActivateEditorMode();
+	}
 }
 
-void ULevelInstanceSubsystem::EditLevelInstanceInternal(ALevelInstance* LevelInstanceActor, TWeakObjectPtr<AActor> ContextActorPtr, bool bRecursive)
+bool ULevelInstanceSubsystem::EditLevelInstanceInternal(ALevelInstance* LevelInstanceActor, TWeakObjectPtr<AActor> ContextActorPtr, bool bRecursive)
 {
 	check(CanEditLevelInstance(LevelInstanceActor));
 		
@@ -1597,16 +1618,14 @@ void ULevelInstanceSubsystem::EditLevelInstanceInternal(ALevelInstance* LevelIns
 		// Only support one level of recursion to commit current edit
 		check(!bRecursive);
 		FLevelInstanceID PendingEditId = LevelInstanceActor->GetLevelInstanceID();
-		ALevelInstance* LevelInstanceToCommit = GetLevelInstance(LevelInstanceEdit->GetLevelInstanceID());
-		CommitLevelInstance(LevelInstanceToCommit);
+		
+		check(!IsLevelInstanceEditDirty(LevelInstanceEdit.Get()));
+		CommitLevelInstanceInternal(LevelInstanceEdit);
 
 		ALevelInstance* LevelInstanceToEdit = GetLevelInstance(PendingEditId);
 		check(LevelInstanceToEdit);
 
-		EditLevelInstanceInternal(LevelInstanceToEdit, nullptr, /*bRecursive=*/true);
-
-		// Stop here. The LevelInstance will be open for editing after an async reload.
-		return;
+		return EditLevelInstanceInternal(LevelInstanceToEdit, nullptr, /*bRecursive=*/true);
 	}
 
 	// Cleanup async requests in case
@@ -1620,7 +1639,7 @@ void ULevelInstanceSubsystem::EditLevelInstanceInternal(ALevelInstance* LevelIns
 	if (!LevelStreaming)
 	{
 		LevelInstanceActor->LoadLevelInstance();
-		return;
+		return false;
 	}
 
 	LevelInstanceEdit = MakeUnique<FLevelInstanceEdit>(LevelStreaming, LevelInstanceActor->GetLevelInstanceID());
@@ -1652,18 +1671,23 @@ void ULevelInstanceSubsystem::EditLevelInstanceInternal(ALevelInstance* LevelIns
 			Actor->PushLevelInstanceEditingStateToProxies(bEditing);
 		}
 	}
-
-	UpdateEngineShowFlags();
 	
 	// Edit can't be undone
 	GEditor->ResetTransaction(LOCTEXT("LevelInstanceEditResetTrans", "Edit Level Instance"));
+
+	return true;
 }
 
 ALevelInstance* ULevelInstanceSubsystem::CommitLevelInstance(ALevelInstance* LevelInstanceActor, bool bDiscardEdits, bool bPromptForSave, TSet<FName>* DirtyPackages)
 {
 	check(LevelInstanceEdit.Get() == GetLevelInstanceEdit(LevelInstanceActor));
 	check(CanCommitLevelInstance(LevelInstanceActor));
-	return CommitLevelInstanceInternal(LevelInstanceEdit, bDiscardEdits, bPromptForSave, DirtyPackages);
+	ALevelInstance* CommittedLevelInstance = CommitLevelInstanceInternal(LevelInstanceEdit, bDiscardEdits, bPromptForSave, DirtyPackages);
+
+	ILevelInstanceEditorModule& EditorModule = FModuleManager::GetModuleChecked<ILevelInstanceEditorModule>("LevelInstanceEditor");
+	EditorModule.DeactivateEditorMode();
+
+	return CommittedLevelInstance;
 }
 
 ALevelInstance* ULevelInstanceSubsystem::CommitLevelInstanceInternal(TUniquePtr<FLevelInstanceEdit>& InLevelInstanceEdit, bool bDiscardEdits, bool bPromptForSave, TSet<FName>* DirtyPackages)
@@ -1777,8 +1801,6 @@ ALevelInstance* ULevelInstanceSubsystem::CommitLevelInstanceInternal(TUniquePtr<
 
 	GEngine->BroadcastLevelActorListChanged();
 
-	UpdateEngineShowFlags();
-
 	return GetLevelInstance(LevelInstanceID);
 }
 
diff --git a/Engine/Source/Runtime/Engine/Public/LevelInstance/ILevelInstanceEditorModule.h b/Engine/Source/Runtime/Engine/Public/LevelInstance/ILevelInstanceEditorModule.h
--- a/Engine/Source/Runtime/Engine/Public/LevelInstance/ILevelInstanceEditorModule.h
+++ b/Engine/Source/Runtime/Engine/Public/LevelInstance/ILevelInstanceEditorModule.h
@@ -11,4 +11,11 @@ class ILevelInstanceEditorModule : public IModuleInterface
 {
 public:
 	virtual ~ILevelInstanceEditorModule() {}
+
+	virtual void ActivateEditorMode() = 0;
+	virtual void DeactivateEditorMode() = 0;
+
+	/** Broadcasts before exiting mode */
+	DECLARE_EVENT(ILevelInstanceEditorModule, FExitEditorModeEvent);
+	virtual FExitEditorModeEvent& OnExitEditorMode() = 0;
 };
\ No newline at end of file
diff --git a/Engine/Source/Runtime/Engine/Public/LevelInstance/LevelInstanceSubsystem.h b/Engine/Source/Runtime/Engine/Public/LevelInstance/LevelInstanceSubsystem.h
--- a/Engine/Source/Runtime/Engine/Public/LevelInstance/LevelInstanceSubsystem.h
+++ b/Engine/Source/Runtime/Engine/Public/LevelInstance/LevelInstanceSubsystem.h
@@ -50,6 +50,7 @@ public:
 	void Tick();
 	void OnObjectPreSave(UObject* Object, FObjectPreSaveContext SaveContext);
 	void OnPackageDeleted(UPackage* Package);
+	void OnExitEditorMode();
 	void PackLevelInstances();
 	bool CanPackLevelInstances() const;
 
@@ -127,8 +128,9 @@ private:
 		FLevelInstanceID GetLevelInstanceID() const;
 	};
 
-	void EditLevelInstanceInternal(ALevelInstance* LevelInstanceActor, TWeakObjectPtr<AActor> ContextActorPtr, bool bRecursive);
-	ALevelInstance* CommitLevelInstanceInternal(TUniquePtr<FLevelInstanceEdit>& InLevelInstanceEdit, bool bDiscardEdits, bool bPromptForSave, TSet<FName>* DirtyPackages);
+	bool EditLevelInstanceInternal(ALevelInstance* LevelInstanceActor, TWeakObjectPtr<AActor> ContextActorPtr, bool bRecursive);
+	ALevelInstance* CommitLevelInstanceInternal(TUniquePtr<FLevelInstanceEdit>& InLevelInstanceEdit, bool bDiscardEdits = false, bool bPromptForSave = true, TSet<FName>* DirtyPackages = nullptr);
+	
 	const FLevelInstanceEdit* GetLevelInstanceEdit(const ALevelInstance* LevelInstanceActor) const;
 	bool IsLevelInstanceEditDirty(const FLevelInstanceEdit* LevelInstanceEdit) const;
 	
@@ -147,15 +149,6 @@ private:
 	friend ULevelStreamingLevelInstance;
 	friend ULevelStreamingLevelInstanceEditor;
 	void RemoveLevelsFromWorld(const TArray<ULevel*>& Levels, bool bResetTrans = true);
-	void UpdateEngineShowFlags();
-#endif
-
-
-#if WITH_EDITORONLY_DATA
-	// Optional scope to accelerate level unload by batching them
-	TUniquePtr<FLevelsToRemoveScope> LevelsToRemoveScope;
-	
-	TUniquePtr<FLevelInstanceEdit> LevelInstanceEdit;
 #endif
 
 	struct FLevelInstance
@@ -167,7 +160,16 @@ private:
 	TSet<FLevelInstanceID> LevelInstancesToUnload;
 	TMap<FLevelInstanceID, FLevelInstance> LevelInstances;
 	TMap<FLevelInstanceID, ALevelInstance*> RegisteredLevelInstances;
+
+#if WITH_EDITORONLY_DATA
+	// Optional scope to accelerate level unload by batching them
+	TUniquePtr<FLevelsToRemoveScope> LevelsToRemoveScope;
+
+	TUniquePtr<FLevelInstanceEdit> LevelInstanceEdit;
+	bool bCreatingLevelInstance = false;
+
 	TMap<FObjectKey, FFolder::FRootObject> UnregisteringLevelInstanceLevels;
 	TMap<FFolder::FRootObject, FObjectKey > UnregisteringLevelInstances;
 	TMap<FLevelInstanceID, int32> ChildEdits;
+#endif
 };
\ No newline at end of file
