PR #8652: Updated incorrect PrintString parameter comment in KismetSystemLibrary.h (Contributed by [REDACTED_PERSON])

diff --git a/Engine/Source/Runtime/Engine/Classes/Kismet/KismetSystemLibrary.h b/Engine/Source/Runtime/Engine/Classes/Kismet/KismetSystemLibrary.h
--- a/Engine/Source/Runtime/Engine/Classes/Kismet/KismetSystemLibrary.h
+++ b/Engine/Source/Runtime/Engine/Classes/Kismet/KismetSystemLibrary.h
@@ -410,7 +410,7 @@ class ENGINE_API UKismetSystemLibrary : public UBlueprintFunctionLibrary
 	 * @param	bPrintToScreen	Whether or not to print the output to the screen
 	 * @param	bPrintToLog		Whether or not to print the output to the log
 	 * @param	bPrintToConsole	Whether or not to print the output to the console
-	 * @param	TextColor		Whether or not to print the output to the console
+	 * @param	TextColor		The color of the text to display
 	 * @param	Duration		The display duration (if Print to Screen is True). Using negative number will result in loading the duration time from the config.
 	 */
 	UFUNCTION(BlueprintCallable, meta=(WorldContext="WorldContextObject", CallableWithoutWorldContext, Keywords = "log print", AdvancedDisplay = "2", DevelopmentOnly), Category="Development")
@@ -424,7 +424,7 @@ class ENGINE_API UKismetSystemLibrary : public UBlueprintFunctionLibrary
 	 * @param	bPrintToScreen	Whether or not to print the output to the screen
 	 * @param	bPrintToLog		Whether or not to print the output to the log
 	 * @param	bPrintToConsole	Whether or not to print the output to the console
-	 * @param	TextColor		Whether or not to print the output to the console
+	 * @param	TextColor		The color of the text to display
 	 * @param	Duration		The display duration (if Print to Screen is True). Using negative number will result in loading the duration time from the config.
 	 */
 	UFUNCTION(BlueprintCallable, meta=(WorldContext="WorldContextObject", CallableWithoutWorldContext, Keywords = "log", AdvancedDisplay = "2", DevelopmentOnly), Category="Development")
