PR #8671: APlayerController will never return a null view target just like AController (Contributed by [REDACTED_PERSON])

diff --git a/Engine/Source/Runtime/Engine/Private/PlayerController.cpp b/Engine/Source/Runtime/Engine/Private/PlayerController.cpp
--- a/Engine/Source/Runtime/Engine/Private/PlayerController.cpp
+++ b/Engine/Source/Runtime/Engine/Private/PlayerController.cpp
@@ -431,7 +431,9 @@ void APlayerController::CleanUpAudioComponents()
 
 AActor* APlayerController::GetViewTarget() const
 {
-	return PlayerCameraManager ? PlayerCameraManager->GetViewTarget() : NULL;
+	AActor* CameraManagerViewTarget = PlayerCameraManager ? PlayerCameraManager->GetViewTarget() : NULL;
+
+	return CameraManagerViewTarget ? CameraManagerViewTarget : const_cast<APlayerController*>(this);
 }
 
 void APlayerController::SetViewTarget(class AActor* NewViewTarget, struct FViewTargetTransitionParams TransitionParams)
