World Partition: fix unsaved, deleted actors not taken into account when going into PIE.

diff --git a/Engine/Source/Runtime/Engine/Private/WorldPartition/WorldPartitionStreamingGeneration.cpp b/Engine/Source/Runtime/Engine/Private/WorldPartition/WorldPartitionStreamingGeneration.cpp
--- a/Engine/Source/Runtime/Engine/Private/WorldPartition/WorldPartitionStreamingGeneration.cpp
+++ b/Engine/Source/Runtime/Engine/Private/WorldPartition/WorldPartitionStreamingGeneration.cpp
@@ -45,20 +45,28 @@ class FWorldPartitionStreamingGenerator
 		{
 			if (!ActorDescIt->GetActorIsEditorOnly())
 			{
-				AActor* Actor = ActorDescIt->GetActor();
-
-				if (bHandleUnsavedActors && IsValid(Actor) && (bIsTempContainerPackage || Actor->GetPackage()->IsDirty()))
-				{
-					// Dirty, unsaved actor for PIE
-					FWorldPartitionActorDesc* ActorDesc = RuntimeHash->ModifiedActorDescListForPIE.AddActor(Actor);
-					ActorDesc->OnRegister(Actor->GetWorld());
-					OutActorDescViewMap.Emplace(ActorDescIt->GetGuid(), ActorDesc);
-				}
-				else
+				// Handle unsaved actors
+				if (AActor* Actor = ActorDescIt->GetActor())
 				{
-					// Non-dirty actor
-					OutActorDescViewMap.Emplace(ActorDescIt->GetGuid(), *ActorDescIt);
+					// Deleted actors
+					if (!IsValid(Actor))
+					{
+						continue;
+					}
+
+					// Dirty actors
+					if (bHandleUnsavedActors && (bIsTempContainerPackage || Actor->GetPackage()->IsDirty()))
+					{
+						// Dirty, unsaved actor for PIE
+						FWorldPartitionActorDesc* ActorDesc = RuntimeHash->ModifiedActorDescListForPIE.AddActor(Actor);
+						ActorDesc->OnRegister(Actor->GetWorld());
+						OutActorDescViewMap.Emplace(ActorDescIt->GetGuid(), ActorDesc);
+						continue;
+					}
 				}
+
+				// Non-dirty actor
+				OutActorDescViewMap.Emplace(ActorDescIt->GetGuid(), *ActorDescIt);
 			}
 		}
 
