No longer track scopes where FExpressions are used
Instead, code generated by expressions is deduplicated, and automatically introduced into the tightest scope possible for all uses
Remove HLSLTree visitors, no longer needed
Minor HLSLTree API cleanup

diff --git a/Engine/Source/Runtime/Engine/Private/HLSLTree/HLSLTree.cpp b/Engine/Source/Runtime/Engine/Private/HLSLTree/HLSLTree.cpp
--- a/Engine/Source/Runtime/Engine/Private/HLSLTree/HLSLTree.cpp
+++ b/Engine/Source/Runtime/Engine/Private/HLSLTree/HLSLTree.cpp
@@ -33,6 +33,18 @@ static const TCHAR* AllocateStringf(FMemStackBase& Allocator, const FormatType&
 	return AllocateString(Allocator, String);
 }
 
+
+template<typename T, typename AllocatorType>
+static TArrayView<T> AllocateArrayView(FMemStackBase& Allocator, const TArray<T, AllocatorType>& Array)
+{
+	T* Data = new(Allocator) T[Array.Num()];
+	for (int32 i = 0; i < Array.Num(); ++i)
+	{
+		Data[i] = Array[i];
+	}
+	return MakeArrayView(Data, Array.Num());
+}
+
 static FSHAHash HashString(const TCHAR* String, int32 Length)
 {
 	FSHAHash Hash;
@@ -296,8 +308,41 @@ const TCHAR* FEmitContext::AcquirePreshader(Shader::EValueType Type, const Shade
 	return Code;
 }
 
+namespace Private
+{
+void EmitShaderValue(FEmitContext& Context, FEmitShaderValue* ShaderValue)
+{
+	if (ShaderValue->Scope)
+	{
+		// Emit dependencies first
+		for (FEmitShaderValue* Dependency : ShaderValue->Dependencies)
+		{
+			EmitShaderValue(Context, Dependency);
+		}
+		// Don't need a declaration for inline values
+		if (!ShaderValue->IsInline())
+		{
+			ShaderValue->Scope->EmitDeclarationf(Context, TEXT("const %s %s = %s;"),
+				ShaderValue->Expression->GetType().GetName(),
+				ShaderValue->Reference,
+				ShaderValue->Value);
+		}
+		ShaderValue->Scope = nullptr; // Don't emit again
+	}
+}
+}
+
 void FEmitContext::Finalize()
 {
+	check(ScopeStack.Num() == 0);
+	check(ShaderValueStack.Num() == 0);
+
+	for (const auto& It : ShaderValueMap)
+	{
+		Private::EmitShaderValue(*this, It.Value);
+	}
+	ShaderValueMap.Reset();
+
 	MaterialCompilationOutput->UniformExpressionSet.UniformPreshaderBufferSize = (UniformPreshaderOffset + 3u) / 4u;
 }
 
@@ -378,7 +423,6 @@ FConstantValue::FConstantValue(const Shader::FValue& InValue) : Type(InValue.Get
 void FScope::Reset()
 {
 	State = EScopeState::Uninitialized;
-	ExpressionCodeMap.Reset();
 	Declarations = FCodeList();
 	Statements = FCodeList();
 }
@@ -406,8 +450,9 @@ FScope* FScope::FindSharedParent(FScope* Lhs, FScope* Rhs)
 	return Scope0;
 }
 
-void FExpressionLocalPHI::PrepareValue(FEmitContext& Context, const FRequestedType& RequestedType)
+void FExpressionLocalPHI::PrepareValue(FEmitContext& Context, const FRequestedType& RequestedType, FPrepareValueResult& OutResult)
 {
+	check(NumValues <= MaxNumPreviousScopes);
 	FExpression* ForwardExpression = Values[0];
 	bool bForwardExpressionValid = true;
 
@@ -427,7 +472,7 @@ void FExpressionLocalPHI::PrepareValue(FEmitContext& Context, const FRequestedTy
 	if (bForwardExpressionValid)
 	{
 		check(ForwardExpression);
-		return SetForwardValue(Context, ForwardExpression, RequestedType);
+		return OutResult.SetForwardValue(Context, ForwardExpression, RequestedType);
 	}
 
 	// 2) PHI has different values in previous scopes, but possible some previous scopes may become dead due to constant folding
@@ -453,7 +498,7 @@ void FExpressionLocalPHI::PrepareValue(FEmitContext& Context, const FRequestedTy
 	if (bForwardExpressionValid)
 	{
 		check(ForwardExpression);
-		return SetForwardValue(Context, ForwardExpression, RequestedType);
+		return OutResult.SetForwardValue(Context, ForwardExpression, RequestedType);
 	}
 
 	FType TypePerValue[MaxNumPreviousScopes];
@@ -485,7 +530,7 @@ void FExpressionLocalPHI::PrepareValue(FEmitContext& Context, const FRequestedTy
 	UpdateValueTypes();
 
 	// Assuming we have at least one value with a valid type, we use that to initialize our type
-	SetType(Context, EExpressionEvaluationType::Shader, CurrentType);
+	OutResult.SetType(Context, EExpressionEvaluationType::Shader, CurrentType);
 
 	if (NumValidTypes < NumValues)
 	{
@@ -509,7 +554,7 @@ void FExpressionLocalPHI::EmitValueShader(FEmitContext& Context, FShaderValue& O
 void FExpressionLocalPHI::EmitShaderDependencies(FEmitContext& Context, const FShaderValue& Shader) const
 {
 	// Find the outermost scope to declare our local variable
-	FScope* DeclarationScope = ParentScope;
+	FScope* DeclarationScope = Context.ScopeStack.Last();
 	for (int32 i = 0; i < NumValues; ++i)
 	{
 		DeclarationScope = FScope::FindSharedParent(DeclarationScope, Scopes[i]);
@@ -527,7 +572,10 @@ void FExpressionLocalPHI::EmitShaderDependencies(FEmitContext& Context, const FS
 		FScope* ValueScope = Scopes[i];
 		check(IsScopeLive(ValueScope));
 
+		Context.ScopeStack.Add(ValueScope);
 		const TCHAR* ValueCode = Values[i]->GetValueShader(Context, LocalType);
+		Context.ScopeStack.Pop();
+
 		if (ValueScope == DeclarationScope)
 		{
 			ValueScope->EmitDeclarationf(Context, TEXT("%s %s = %s;"),
@@ -553,39 +601,6 @@ void FExpressionLocalPHI::EmitShaderDependencies(FEmitContext& Context, const FS
 	}
 }
 
-void FNodeVisitor::VisitNode(FNode* Node)
-{
-	if (Node)
-	{
-		Node->Visit(*this);
-	}
-}
-
-ENodeVisitResult FStatement::Visit(FNodeVisitor& Visitor)
-{
-	return Visitor.OnStatement(*this);
-}
-
-ENodeVisitResult FExpression::Visit(FNodeVisitor& Visitor)
-{
-	return Visitor.OnExpression(*this);
-}
-
-ENodeVisitResult FTextureParameterDeclaration::Visit(FNodeVisitor& Visitor)
-{
-	return Visitor.OnTextureParameterDeclaration(*this);
-}
-
-ENodeVisitResult FScope::Visit(FNodeVisitor& Visitor)
-{
-	const ENodeVisitResult Result = Visitor.OnScope(*this);
-	if (ShouldVisitDependentNodes(Result))
-	{
-		Visitor.VisitNode(ContainedStatement);
-	}
-	return Result;
-}
-
 void FStatement::Reset()
 {
 	bEmitShader = false;
@@ -593,18 +608,17 @@ void FStatement::Reset()
 
 void FExpression::Reset()
 {
-	LocalVariableName = nullptr;
-	Code = nullptr;
-	ForwardedValue = nullptr;
+	ShaderValue = nullptr;
 	CurrentRequestedType.Reset();
 	PrepareValueResult = FPrepareValueResult();
 }
 
-FPrepareValueResult PrepareExpressionValue(FEmitContext& Context, FExpression* InExpression, const FRequestedType& RequestedType)
+const FPrepareValueResult& PrepareExpressionValue(FEmitContext& Context, FExpression* InExpression, const FRequestedType& RequestedType)
 {
+	static FPrepareValueResult EmptyResult;
 	if (!InExpression)
 	{
-		return FPrepareValueResult();
+		return EmptyResult;
 	}
 
 	if (InExpression->bReentryFlag)
@@ -624,7 +638,7 @@ FPrepareValueResult PrepareExpressionValue(FEmitContext& Context, FExpression* I
 	else if(InExpression->CurrentRequestedType.StructType != RequestedType.StructType)
 	{
 		Context.Errors.AddError(InExpression, TEXT("Type mismatch"));
-		return FPrepareValueResult();
+		return EmptyResult;
 	}
 	else
 	{
@@ -640,22 +654,23 @@ FPrepareValueResult PrepareExpressionValue(FEmitContext& Context, FExpression* I
 		}
 	}
 
+	FPrepareValueResult& Result = InExpression->PrepareValueResult;
 	if (bNeedToUpdateType)
 	{
 		check(!InExpression->CurrentRequestedType.IsVoid());
 
 		InExpression->bReentryFlag = true;
-		InExpression->PrepareValue(Context, InExpression->CurrentRequestedType);
+		InExpression->PrepareValue(Context, InExpression->CurrentRequestedType, Result);
 		InExpression->bReentryFlag = false;
 
-		if (InExpression->PrepareValueResult.EvaluationType == EExpressionEvaluationType::None)
+		if (Result.EvaluationType == EExpressionEvaluationType::None)
 		{
 			// If we failed to assign a valid type, reset the requested type as well
 			// This ensures we'll try to compute a type again the next time we're called
 			InExpression->CurrentRequestedType.Reset();
 		}
-		else if (InExpression->PrepareValueResult.EvaluationType == EExpressionEvaluationType::Constant &&
-			InExpression->PrepareValueResult.Type != FType(InExpression->PrepareValueResult.ConstantValue.GetType()))
+		else if (Result.EvaluationType == EExpressionEvaluationType::Constant &&
+			Result.Type != FType(Result.ConstantValue.GetType()))
 		{
 			Shader::FPreshaderData ConstantPreshader;
 			InExpression->bReentryFlag = true;
@@ -666,25 +681,36 @@ FPrepareValueResult PrepareExpressionValue(FEmitContext& Context, FExpression* I
 			FMaterialRenderContext RenderContext(nullptr, *Context.Material, nullptr);
 			Shader::FValue Value;
 			ConstantPreshader.Evaluate(nullptr, RenderContext, Value);
-			InExpression->PrepareValueResult.ConstantValue = Shader::Cast(Value, InExpression->PrepareValueResult.Type);
+			Result.ConstantValue = Shader::Cast(Value, Result.Type);
 		}
 	}
 
-	return InExpression->PrepareValueResult;
+	return Result;
+}
+
+void FPrepareValueResult::SetConstant(FEmitContext& Context, const Shader::FValue& InValue)
+{
+	EvaluationType = EExpressionEvaluationType::Constant;
+	Type = InValue.GetType();
+	ConstantValue = InValue;
+	ForwardValue = nullptr;
 }
 
-void FExpression::SetType(FEmitContext& Context, EExpressionEvaluationType InEvaluationType, const FType& InType)
+void FPrepareValueResult::SetType(FEmitContext& Context, EExpressionEvaluationType InEvaluationType, const FType& InType)
 {
-	PrepareValueResult.EvaluationType = InEvaluationType;
-	PrepareValueResult.Type = InType;
-	ForwardedValue = nullptr;
+	EvaluationType = InEvaluationType;
+	Type = InType;
+	ForwardValue = nullptr;
 }
 
-void FExpression::SetForwardValue(FEmitContext& Context, FExpression* InForwardValue, const FRequestedType& RequestedType)
+void FPrepareValueResult::SetForwardValue(FEmitContext& Context, FExpression* InForwardValue, const FRequestedType& RequestedType)
 {
 	check(InForwardValue);
-	PrepareValueResult = PrepareExpressionValue(Context, InForwardValue, RequestedType);
-	ForwardedValue = InForwardValue;
+	const FPrepareValueResult& ForwardResult = PrepareExpressionValue(Context, InForwardValue, RequestedType);
+	EvaluationType = ForwardResult.EvaluationType;
+	Type = ForwardResult.Type;
+	ConstantValue = ForwardResult.ConstantValue;
+	ForwardValue = InForwardValue;
 }
 
 void FExpression::EmitValueShader(FEmitContext& Context, FShaderValue& OutShader) const
@@ -702,68 +728,86 @@ void FExpression::EmitValuePreshader(FEmitContext& Context, Shader::FPreshaderDa
 	check(false);
 }
 
+namespace Private
+{
+void MoveToScope(FEmitShaderValue* ShaderValue, FScope* Scope)
+{
+	if (ShaderValue->Scope != Scope)
+	{
+		FScope* NewScope = FScope::FindSharedParent(ShaderValue->Scope, Scope);
+		check(NewScope);
+
+		ShaderValue->Scope = NewScope;
+		for (FEmitShaderValue* Dependency : ShaderValue->Dependencies)
+		{
+			MoveToScope(Dependency, NewScope);
+		}
+	}
+}
+}
+
 const TCHAR* FExpression::GetValueShader(FEmitContext& Context)
 {
 	check(PrepareValueResult.EvaluationType != EExpressionEvaluationType::None);
-	if (ForwardedValue)
+	if (PrepareValueResult.ForwardValue)
 	{
-		return ForwardedValue->GetValueShader(Context);
+		return PrepareValueResult.ForwardValue->GetValueShader(Context);
 	}
 
-	if (!Code)
+	FScope* CurrentScope = Context.ScopeStack.Last();
+	check(IsScopeLive(CurrentScope));
+
+	if (!ShaderValue)
 	{
 		check(!bReentryFlag);
 
 		TStringBuilder<1024> FormattedCode;
 		if (PrepareValueResult.EvaluationType == EExpressionEvaluationType::Shader)
 		{
-			FShaderValue ShaderValue(FormattedCode);
+			FShaderValue ShaderResult(FormattedCode);
+			FEmitShaderValueDependencies Dependencies;
+			{
+				FEmitShaderValueContext& ValueContext = Context.ShaderValueStack.AddDefaulted_GetRef();
 
-			bReentryFlag = true;
-			EmitValueShader(Context, ShaderValue);
-			bReentryFlag = false;
+				bReentryFlag = true;
+				EmitValueShader(Context, ShaderResult);
+				bReentryFlag = false;
+
+				Dependencies = MoveTemp(ValueContext.Dependencies);
+				Context.ShaderValueStack.Pop(false);
+			}
 
-			if (ShaderValue.bInline)
+			if (ShaderResult.bInline)
 			{
-				Code = AllocateString(*Context.Allocator, FormattedCode);
+				ShaderValue = new(*Context.Allocator) FEmitShaderValue(this, CurrentScope);
+				ShaderValue->Reference = AllocateString(*Context.Allocator, FormattedCode);
+				ShaderValue->Dependencies = AllocateArrayView(*Context.Allocator, Dependencies);
 			}
 			else
 			{
-				// Check to see if we've already generated code for an equivalent expression in either this scope, or any outer visible scope
-				const FSHAHash Hash = HashString(FormattedCode);
-				const TCHAR* Declaration = nullptr;
-
-				FScope* CheckScope = ParentScope;
-				while (CheckScope)
+				// Check to see if we've already generated code for an equivalent expression
+				const FSHAHash ShaderHash = HashString(FormattedCode);
+				FEmitShaderValue** const PrevShaderValue = Context.ShaderValueMap.Find(ShaderHash);
+				if (PrevShaderValue)
 				{
-					const TCHAR** FoundDeclaration = CheckScope->ExpressionCodeMap.Find(Hash);
-					if (FoundDeclaration)
-					{
-						// Re-use results from previous matching expression
-						check(IsScopeLive(CheckScope));
-						Declaration = *FoundDeclaration;
-						break;
-					}
-					CheckScope = CheckScope->ParentScope;
+					ShaderValue = *PrevShaderValue;
 				}
-
-				if (!Declaration)
+				else
 				{
-					check(IsScopeLive(ParentScope));
-					Declaration = Context.AcquireLocalDeclarationCode();
-					ParentScope->ExpressionCodeMap.Add(Hash, Declaration);
-					ParentScope->EmitDeclarationf(Context, TEXT("const %s %s = %s;"),
-						PrepareValueResult.Type.GetName(),
-						Declaration,
-						FormattedCode.ToString());
+					ShaderValue = new(*Context.Allocator) FEmitShaderValue(this, CurrentScope);
+					ShaderValue->Hash = ShaderHash;
+					ShaderValue->Reference = Context.AcquireLocalDeclarationCode();
+					ShaderValue->Value = AllocateString(*Context.Allocator, FormattedCode);
+					ShaderValue->Dependencies = AllocateArrayView(*Context.Allocator, Dependencies);
+					Context.ShaderValueMap.Add(ShaderHash, ShaderValue);
 				}
-				Code = Declaration;
 			}
-			check(Code);
-			if (ShaderValue.bHasDependencies)
+
+			check(ShaderValue);
+			if (ShaderResult.bHasDependencies)
 			{
 				bReentryFlag = true;
-				EmitShaderDependencies(Context, ShaderValue);
+				EmitShaderDependencies(Context, ShaderResult);
 				bReentryFlag = false;
 			}
 		}
@@ -775,18 +819,27 @@ const TCHAR* FExpression::GetValueShader(FEmitContext& Context)
 			EmitValuePreshader(Context, Preshader);
 			bReentryFlag = false;
 
-			Code = Context.AcquirePreshader(PrepareValueResult.Type, Preshader);
+			ShaderValue = new(*Context.Allocator) FEmitShaderValue(this, CurrentScope);
+			ShaderValue->Reference = Context.AcquirePreshader(PrepareValueResult.Type, Preshader);
 		}
 		else
 		{
 			check(PrepareValueResult.EvaluationType == EExpressionEvaluationType::Constant);
 			PrepareValueResult.ConstantValue.ToString(Shader::EValueStringFormat::HLSL, FormattedCode);
-			Code = AllocateString(*Context.Allocator, FormattedCode);
+
+			ShaderValue = new(*Context.Allocator) FEmitShaderValue(this, CurrentScope);
+			ShaderValue->Reference = AllocateString(*Context.Allocator, FormattedCode);
 		}
 	}
 
-	check(Code);
-	return Code;
+	check(ShaderValue);
+	Private::MoveToScope(ShaderValue, CurrentScope);
+
+	if (Context.ShaderValueStack.Num() > 0)
+	{
+		Context.ShaderValueStack.Last().Dependencies.Add(ShaderValue);
+	}
+	return ShaderValue->Reference;
 }
 
 const TCHAR* FExpression::GetValueShader(FEmitContext& Context, const FType& InType)
@@ -810,9 +863,9 @@ void FExpression::GetValuePreshader(FEmitContext& Context, Shader::FPreshaderDat
 	switch (PrepareValueResult.EvaluationType)
 	{
 	case EExpressionEvaluationType::Preshader:
-		if (ForwardedValue)
+		if (PrepareValueResult.ForwardValue)
 		{
-			ForwardedValue->EmitValuePreshader(Context, OutPreshader);
+			PrepareValueResult.ForwardValue->EmitValuePreshader(Context, OutPreshader);
 		}
 		else
 		{
@@ -854,40 +907,14 @@ void FScope::AddPreviousScope(FScope& Scope)
 	PreviousScope[NumPreviousScopes++] = &Scope;
 }
 
-class FNodeVisitor_MoveToScope : public FNodeVisitor
-{
-public:
-	explicit FNodeVisitor_MoveToScope(FScope* InScope) : Scope(InScope) {}
-
-	virtual ENodeVisitResult OnScope(FScope& InScope) override
-	{
-		InScope.ParentScope = Scope;
-		// Don't want to recurse into any child scopes
-		return ENodeVisitResult::SkipDependentNodes;
-	}
-
-	virtual ENodeVisitResult OnExpression(FExpression& InExpression) override
-	{
-		InExpression.ParentScope = FScope::FindSharedParent(Scope, InExpression.ParentScope);
-		return ENodeVisitResult::VisitDependentNodes;
-	}
-
-	FScope* Scope;
-};
-
-void FScope::UseExpression(FExpression* Expression)
-{
-	// Need to move all dependent expressions/declarations into this scope
-	FNodeVisitor_MoveToScope Visitor(this);
-	Visitor.VisitNode(Expression);
-}
-
 void FScope::InternalEmitCode(FEmitContext& Context, FCodeList& List, ENextScopeFormat ScopeFormat, FScope* Scope, const TCHAR* String, int32 Length)
 {
 	if (Scope && Scope->ContainedStatement && !Scope->ContainedStatement->bEmitShader)
 	{
 		Scope->ContainedStatement->bEmitShader = true;
+		Context.ScopeStack.Add(Scope);
 		Scope->ContainedStatement->EmitShader(Context);
+		Context.ScopeStack.Pop();
 	}
 
 	const int32 SizeofString = sizeof(TCHAR) * Length;
@@ -1090,7 +1117,10 @@ bool FTree::EmitShader(FEmitContext& Context, FStringBuilderBase& OutCode) const
 	if (RootScope->ContainedStatement)
 	{
 		RootScope->ContainedStatement->bEmitShader = true;
+		Context.ScopeStack.Add(RootScope);
 		RootScope->ContainedStatement->EmitShader(Context);
+		Context.ScopeStack.Pop(false);
+
 		if (Context.Errors.Num() > 0)
 		{
 			return false;
@@ -1102,10 +1132,8 @@ bool FTree::EmitShader(FEmitContext& Context, FStringBuilderBase& OutCode) const
 	return Context.Errors.Num() == 0;
 }
 
-void FTree::RegisterExpression(FScope& Scope, FExpression* Expression)
+void FTree::RegisterExpression(FExpression* Expression)
 {
-	check(!Expression->ParentScope);
-	Expression->ParentScope = &Scope;
 }
 
 void FTree::RegisterStatement(FScope& Scope, FStatement* Statement)
diff --git a/Engine/Source/Runtime/Engine/Private/Materials/MaterialExpressionHLSL.cpp b/Engine/Source/Runtime/Engine/Private/Materials/MaterialExpressionHLSL.cpp
--- a/Engine/Source/Runtime/Engine/Private/Materials/MaterialExpressionHLSL.cpp
+++ b/Engine/Source/Runtime/Engine/Private/Materials/MaterialExpressionHLSL.cpp
@@ -99,7 +99,7 @@ EMaterialGenerateHLSLStatus UMaterialExpressionStaticSwitch::GenerateHLSLExpress
 	UE::HLSLTree::FExpression* TrueExpression = A.AcquireHLSLExpression(Generator, Scope);
 	UE::HLSLTree::FExpression* FalseExpression = B.AcquireHLSLExpression(Generator, Scope);
 
-	OutExpression = Generator.GetTree().NewExpression<UE::HLSLTree::FExpressionSelect>(Scope, ConditionExpression, TrueExpression, FalseExpression);
+	OutExpression = Generator.GetTree().NewExpression<UE::HLSLTree::FExpressionSelect>(ConditionExpression, TrueExpression, FalseExpression);
 	return EMaterialGenerateHLSLStatus::Success;
 }
 
@@ -115,26 +115,26 @@ EMaterialGenerateHLSLStatus UMaterialExpressionGetLocal::GenerateHLSLExpression(
 
 EMaterialGenerateHLSLStatus UMaterialExpressionVectorParameter::GenerateHLSLExpression(FMaterialHLSLGenerator& Generator, UE::HLSLTree::FScope& Scope, int32 OutputIndex, UE::HLSLTree::FExpression*& OutExpression)
 {
-	OutExpression = Generator.GetTree().NewExpression<UE::HLSLTree::FExpressionMaterialParameter>(Scope, EMaterialParameterType::Vector, ParameterName, DefaultValue);
+	OutExpression = Generator.GetTree().NewExpression<UE::HLSLTree::FExpressionMaterialParameter>(EMaterialParameterType::Vector, ParameterName, DefaultValue);
 	return EMaterialGenerateHLSLStatus::Success;
 }
 
 EMaterialGenerateHLSLStatus UMaterialExpressionScalarParameter::GenerateHLSLExpression(FMaterialHLSLGenerator& Generator, UE::HLSLTree::FScope& Scope, int32 OutputIndex, UE::HLSLTree::FExpression*& OutExpression)
 {
-	OutExpression = Generator.GetTree().NewExpression<UE::HLSLTree::FExpressionMaterialParameter>(Scope, EMaterialParameterType::Scalar, ParameterName, DefaultValue);
+	OutExpression = Generator.GetTree().NewExpression<UE::HLSLTree::FExpressionMaterialParameter>(EMaterialParameterType::Scalar, ParameterName, DefaultValue);
 	return EMaterialGenerateHLSLStatus::Success;
 }
 
 EMaterialGenerateHLSLStatus UMaterialExpressionStaticBoolParameter::GenerateHLSLExpression(FMaterialHLSLGenerator& Generator, UE::HLSLTree::FScope& Scope, int32 OutputIndex, UE::HLSLTree::FExpression*& OutExpression)
 {
-	OutExpression = Generator.GetTree().NewExpression<UE::HLSLTree::FExpressionMaterialParameter>(Scope, EMaterialParameterType::StaticSwitch, ParameterName, (bool)DefaultValue);
+	OutExpression = Generator.GetTree().NewExpression<UE::HLSLTree::FExpressionMaterialParameter>(EMaterialParameterType::StaticSwitch, ParameterName, (bool)DefaultValue);
 	return EMaterialGenerateHLSLStatus::Success;
 }
 
 EMaterialGenerateHLSLStatus UMaterialExpressionTextureCoordinate::GenerateHLSLExpression(FMaterialHLSLGenerator& Generator, UE::HLSLTree::FScope& Scope, int32 OutputIndex, UE::HLSLTree::FExpression*& OutExpression)
 {
 	// TODO - tiling, mirroring
-	OutExpression = Generator.NewTexCoord(Scope, CoordinateIndex);
+	OutExpression = Generator.NewTexCoord(CoordinateIndex);
 	return EMaterialGenerateHLSLStatus::Success;
 }
 
@@ -159,8 +159,8 @@ EMaterialGenerateHLSLStatus UMaterialExpressionTextureSample::GenerateHLSLExpres
 		return Generator.Error(TEXT("Missing input texture"));
 	}
 
-	UE::HLSLTree::FExpression* TexCoordExpression = Coordinates.GetTracedInput().Expression ? Coordinates.AcquireHLSLExpression(Generator, Scope) : Generator.NewTexCoord(Scope, ConstCoordinate);
-	UE::HLSLTree::FExpressionTextureSample* ExpressionTextureSample = Generator.GetTree().NewExpression<UE::HLSLTree::FExpressionTextureSample>(Scope, TextureDeclaration, TexCoordExpression);
+	UE::HLSLTree::FExpression* TexCoordExpression = Coordinates.GetTracedInput().Expression ? Coordinates.AcquireHLSLExpression(Generator, Scope) : Generator.NewTexCoord(ConstCoordinate);
+	UE::HLSLTree::FExpressionTextureSample* ExpressionTextureSample = Generator.GetTree().NewExpression<UE::HLSLTree::FExpressionTextureSample>(TextureDeclaration, TexCoordExpression);
 	ExpressionTextureSample->SamplerSource = SamplerSource;
 	ExpressionTextureSample->MipValueMode = MipValueMode;
 	OutExpression = ExpressionTextureSample;
@@ -208,7 +208,7 @@ EMaterialGenerateHLSLStatus UMaterialExpressionBinaryOp::GenerateHLSLExpression(
 		return EMaterialGenerateHLSLStatus::Error;
 	}
 
-	OutExpression = Generator.GetTree().NewExpression<UE::HLSLTree::FExpressionBinaryOp>(Scope, GetBinaryOp(), Lhs, Rhs);
+	OutExpression = Generator.GetTree().NewExpression<UE::HLSLTree::FExpressionBinaryOp>(GetBinaryOp(), Lhs, Rhs);
 	return EMaterialGenerateHLSLStatus::Success;
 }
 
@@ -220,7 +220,7 @@ EMaterialGenerateHLSLStatus UMaterialExpressionAdd::GenerateHLSLExpression(FMate
 	{
 		return EMaterialGenerateHLSLStatus::Error;
 	}
-	OutExpression = Generator.GetTree().NewExpression<UE::HLSLTree::FExpressionBinaryOp>(Scope, UE::HLSLTree::EBinaryOp::Add, Lhs, Rhs);
+	OutExpression = Generator.GetTree().NewExpression<UE::HLSLTree::FExpressionBinaryOp>(UE::HLSLTree::EBinaryOp::Add, Lhs, Rhs);
 	return EMaterialGenerateHLSLStatus::Success;
 }
 
@@ -232,7 +232,7 @@ EMaterialGenerateHLSLStatus UMaterialExpressionMultiply::GenerateHLSLExpression(
 	{
 		return EMaterialGenerateHLSLStatus::Error;
 	}
-	OutExpression = Generator.GetTree().NewExpression<UE::HLSLTree::FExpressionBinaryOp>(Scope, UE::HLSLTree::EBinaryOp::Mul, Lhs, Rhs);
+	OutExpression = Generator.GetTree().NewExpression<UE::HLSLTree::FExpressionBinaryOp>(UE::HLSLTree::EBinaryOp::Mul, Lhs, Rhs);
 	return EMaterialGenerateHLSLStatus::Success;
 }
 
@@ -244,7 +244,7 @@ EMaterialGenerateHLSLStatus UMaterialExpressionDivide::GenerateHLSLExpression(FM
 	{
 		return EMaterialGenerateHLSLStatus::Error;
 	}
-	OutExpression = Generator.GetTree().NewExpression<UE::HLSLTree::FExpressionBinaryOp>(Scope, UE::HLSLTree::EBinaryOp::Div, Lhs, Rhs);
+	OutExpression = Generator.GetTree().NewExpression<UE::HLSLTree::FExpressionBinaryOp>(UE::HLSLTree::EBinaryOp::Div, Lhs, Rhs);
 	return EMaterialGenerateHLSLStatus::Success;
 }
 
@@ -256,7 +256,7 @@ EMaterialGenerateHLSLStatus UMaterialExpressionAppendVector::GenerateHLSLExpress
 	{
 		return EMaterialGenerateHLSLStatus::Error;
 	}
-	OutExpression = Generator.GetTree().NewExpression<UE::HLSLTree::FExpressionAppend>(Scope, Lhs, Rhs);
+	OutExpression = Generator.GetTree().NewExpression<UE::HLSLTree::FExpressionAppend>(Lhs, Rhs);
 	return EMaterialGenerateHLSLStatus::Success;
 }
 
@@ -281,7 +281,7 @@ EMaterialGenerateHLSLStatus UMaterialExpressionGetMaterialAttributes::GenerateHL
 	const FGuid& AttributeID = AttributeGetTypes[AttributeIndex];
 	const FString& AttributeName = FMaterialAttributeDefinitionMap::GetAttributeName(AttributeID);
 
-	UE::HLSLTree::FExpressionGetStructField* SetAttributeExpression = Generator.GetTree().NewExpression<UE::HLSLTree::FExpressionGetStructField>(Scope);
+	UE::HLSLTree::FExpressionGetStructField* SetAttributeExpression = Generator.GetTree().NewExpression<UE::HLSLTree::FExpressionGetStructField>();
 	SetAttributeExpression->StructType = Generator.GetMaterialAttributesType();
 	SetAttributeExpression->StructExpression = AttributesExpression;
 	SetAttributeExpression->FieldName = *AttributeName;
@@ -300,7 +300,7 @@ EMaterialGenerateHLSLStatus UMaterialExpressionSetMaterialAttributes::GenerateHL
 	}
 	else
 	{
-		AttributesExpression = Generator.GetTree().NewExpression<UE::HLSLTree::FExpressionConstant>(Scope, Generator.GetMaterialAttributesDefaultValue());
+		AttributesExpression = Generator.GetTree().NewExpression<UE::HLSLTree::FExpressionConstant>(Generator.GetMaterialAttributesDefaultValue());
 	}
 
 	for (int32 PinIndex = 0; PinIndex < AttributeSetTypes.Num(); ++PinIndex)
@@ -317,7 +317,7 @@ EMaterialGenerateHLSLStatus UMaterialExpressionSetMaterialAttributes::GenerateHL
 				UE::HLSLTree::FExpression* ValueExpression = AttributeInput.AcquireHLSLExpression(Generator, Scope);
 				if (ValueExpression)
 				{
-					UE::HLSLTree::FExpressionSetStructField* SetAttributeExpression = Generator.GetTree().NewExpression<UE::HLSLTree::FExpressionSetStructField>(Scope);
+					UE::HLSLTree::FExpressionSetStructField* SetAttributeExpression = Generator.GetTree().NewExpression<UE::HLSLTree::FExpressionSetStructField>();
 					SetAttributeExpression->StructType = Generator.GetMaterialAttributesType();
 					SetAttributeExpression->FieldName = *AttributeName;
 					SetAttributeExpression->StructExpression = AttributesExpression;
@@ -336,7 +336,7 @@ EMaterialGenerateHLSLStatus UMaterialExpressionReflectionVectorWS::GenerateHLSLE
 {
 	check(!CustomWorldNormal.GetTracedInput().Expression); // TODO
 
-	OutExpression = Generator.GetTree().NewExpression<UE::HLSLTree::FExpressionReflectionVector>(Scope);
+	OutExpression = Generator.GetTree().NewExpression<UE::HLSLTree::FExpressionReflectionVector>();
 	return EMaterialGenerateHLSLStatus::Success;
 }
 
@@ -495,10 +495,10 @@ EMaterialGenerateHLSLStatus UMaterialExpressionForLoop::GenerateHLSLStatements(F
 
 	UE::HLSLTree::FExpression* LocalExpression = Generator.AcquireLocalValue(*LoopStatement->LoopScope, ExpressionData->LocalName);
 
-	IfStatement->ConditionExpression = Generator.GetTree().NewExpression<UE::HLSLTree::FExpressionBinaryOp>(*LoopStatement->LoopScope, UE::HLSLTree::EBinaryOp::Less, LocalExpression, EndExpression);
+	IfStatement->ConditionExpression = Generator.GetTree().NewExpression<UE::HLSLTree::FExpressionBinaryOp>(UE::HLSLTree::EBinaryOp::Less, LocalExpression, EndExpression);
 	LoopBody.GenerateHLSLStatements(Generator, *IfStatement->ThenScope);
 
-	UE::HLSLTree::FExpression* NewLocalExpression = Generator.GetTree().NewExpression<UE::HLSLTree::FExpressionBinaryOp>(*IfStatement->ThenScope, UE::HLSLTree::EBinaryOp::Add, Generator.AcquireLocalValue(*IfStatement->ThenScope, ExpressionData->LocalName), StepExpression);
+	UE::HLSLTree::FExpression* NewLocalExpression = Generator.GetTree().NewExpression<UE::HLSLTree::FExpressionBinaryOp>(UE::HLSLTree::EBinaryOp::Add, Generator.AcquireLocalValue(*IfStatement->ThenScope, ExpressionData->LocalName), StepExpression);
 	Generator.GenerateAssignLocal(*IfStatement->ThenScope, ExpressionData->LocalName, NewLocalExpression);
 
 	Completed.GenerateHLSLStatements(Generator, *LoopStatement->NextScope);
diff --git a/Engine/Source/Runtime/Engine/Private/Materials/MaterialHLSLGenerator.cpp b/Engine/Source/Runtime/Engine/Private/Materials/MaterialHLSLGenerator.cpp
--- a/Engine/Source/Runtime/Engine/Private/Materials/MaterialHLSLGenerator.cpp
+++ b/Engine/Source/Runtime/Engine/Private/Materials/MaterialHLSLGenerator.cpp
@@ -264,7 +264,7 @@ bool FMaterialHLSLGenerator::GenerateResult(UE::HLSLTree::FScope& Scope)
 			}
 			else
 			{
-				AttributesExpression = HLSLTree->NewExpression<UE::HLSLTree::FExpressionConstant>(Scope, MaterialAttributesDefaultValue);
+				AttributesExpression = HLSLTree->NewExpression<UE::HLSLTree::FExpressionConstant>(MaterialAttributesDefaultValue);
 				for (int32 PropertyIndex = 0; PropertyIndex < MP_MAX; ++PropertyIndex)
 				{
 					const EMaterialProperty Property = (EMaterialProperty)PropertyIndex;
@@ -317,19 +317,19 @@ UE::HLSLTree::FScope* FMaterialHLSLGenerator::NewJoinedScope(UE::HLSLTree::FScop
 
 UE::HLSLTree::FExpressionConstant* FMaterialHLSLGenerator::NewConstant(const UE::Shader::FValue& Value)
 {
-	UE::HLSLTree::FExpressionConstant* Expression = HLSLTree->NewExpression<UE::HLSLTree::FExpressionConstant>(HLSLTree->GetRootScope(), Value);
+	UE::HLSLTree::FExpressionConstant* Expression = HLSLTree->NewExpression<UE::HLSLTree::FExpressionConstant>(Value);
 	return Expression;
 }
 
-UE::HLSLTree::FExpressionExternalInput* FMaterialHLSLGenerator::NewTexCoord(UE::HLSLTree::FScope& Scope, int32 Index)
+UE::HLSLTree::FExpressionExternalInput* FMaterialHLSLGenerator::NewTexCoord(int32 Index)
 {
-	UE::HLSLTree::FExpressionExternalInput* Expression = HLSLTree->NewExpression<UE::HLSLTree::FExpressionExternalInput>(Scope, UE::HLSLTree::MakeInputTexCoord(Index));
+	UE::HLSLTree::FExpressionExternalInput* Expression = HLSLTree->NewExpression<UE::HLSLTree::FExpressionExternalInput>(UE::HLSLTree::MakeInputTexCoord(Index));
 	return Expression;
 }
 
-UE::HLSLTree::FExpressionSwizzle* FMaterialHLSLGenerator::NewSwizzle(UE::HLSLTree::FScope& Scope, const UE::HLSLTree::FSwizzleParameters& Params, UE::HLSLTree::FExpression* Input)
+UE::HLSLTree::FExpressionSwizzle* FMaterialHLSLGenerator::NewSwizzle(const UE::HLSLTree::FSwizzleParameters& Params, UE::HLSLTree::FExpression* Input)
 {
-	UE::HLSLTree::FExpressionSwizzle* Expression = HLSLTree->NewExpression<UE::HLSLTree::FExpressionSwizzle>(Scope, Params, Input);
+	UE::HLSLTree::FExpressionSwizzle* Expression = HLSLTree->NewExpression<UE::HLSLTree::FExpressionSwizzle>(Params, Input);
 	return Expression;
 }
 
@@ -386,7 +386,7 @@ UE::HLSLTree::FExpression* FMaterialHLSLGenerator::InternalAcquireLocalValue(UE:
 	const TArrayView<UE::HLSLTree::FScope*> PreviousScopes = Scope.GetPreviousScopes();
 	if (PreviousScopes.Num() > 1)
 	{
-		UE::HLSLTree::FExpressionLocalPHI* Expression = HLSLTree->NewExpression<UE::HLSLTree::FExpressionLocalPHI>(Scope);
+		UE::HLSLTree::FExpressionLocalPHI* Expression = HLSLTree->NewExpression<UE::HLSLTree::FExpressionLocalPHI>();
 		Expression->LocalName = LocalName;
 		Expression->NumValues = PreviousScopes.Num();
 		for (int32 i = 0; i < PreviousScopes.Num(); ++i)
@@ -420,6 +420,8 @@ UE::HLSLTree::FExpression* FMaterialHLSLGenerator::AcquireExpression(UE::HLSLTre
 	UE::HLSLTree::FExpression* Expression = nullptr;
 	if (!PrevExpression)
 	{
+		// TODO - need to rethink this caching, won't work to cache expressions that depend on current value of local variables
+		// May just remove caching at this level....need to rework function inputs in this case
 		ExpressionStack.Add(Key);
 		const EMaterialGenerateHLSLStatus Status = MaterialExpression->GenerateHLSLExpression(*this, Scope, OutputIndex, Expression);
 		verify(ExpressionStack.Pop() == Key);
@@ -428,7 +430,6 @@ UE::HLSLTree::FExpression* FMaterialHLSLGenerator::AcquireExpression(UE::HLSLTre
 	else
 	{
 		Expression = *PrevExpression;
-		Scope.UseExpression(Expression);
 	}
 	return Expression;
 }
diff --git a/Engine/Source/Runtime/Engine/Private/Materials/MaterialShared.cpp b/Engine/Source/Runtime/Engine/Private/Materials/MaterialShared.cpp
--- a/Engine/Source/Runtime/Engine/Private/Materials/MaterialShared.cpp
+++ b/Engine/Source/Runtime/Engine/Private/Materials/MaterialShared.cpp
@@ -281,7 +281,7 @@ UE::HLSLTree::FExpression* FExpressionInput::AcquireHLSLExpression(FMaterialHLSL
 		if (Mask && Result)
 		{
 			const UE::HLSLTree::FSwizzleParameters SwizzleParams = UE::HLSLTree::MakeSwizzleMask(!!MaskR, !!MaskG, !!MaskB, !!MaskA);
-			Result = Generator.NewSwizzle(Scope, SwizzleParams, Result);
+			Result = Generator.NewSwizzle(SwizzleParams, Result);
 		}
 	}
 	
diff --git a/Engine/Source/Runtime/Engine/Public/HLSLTree/HLSLTree.h b/Engine/Source/Runtime/Engine/Public/HLSLTree/HLSLTree.h
--- a/Engine/Source/Runtime/Engine/Public/HLSLTree/HLSLTree.h
+++ b/Engine/Source/Runtime/Engine/Public/HLSLTree/HLSLTree.h
@@ -33,6 +33,7 @@ namespace HLSLTree
 
 class FNode;
 class FScope;
+struct FEmitShaderValue;
 
 static constexpr int32 MaxNumPreviousScopes = 2;
 
@@ -122,6 +123,13 @@ enum class ECastFlags : uint32
 };
 ENUM_CLASS_FLAGS(ECastFlags);
 
+using FEmitShaderValueDependencies = TArray<FEmitShaderValue*, TInlineAllocator<8>>;
+
+struct FEmitShaderValueContext
+{
+	FEmitShaderValueDependencies Dependencies;
+};
+
 /** Tracks shared state while emitting HLSL code */
 class FEmitContext
 {
@@ -139,6 +147,7 @@ public:
 	const TCHAR* AcquirePreshader(Shader::EValueType Type, const Shader::FPreshaderData& Preshader);
 
 	FMemStackBase* Allocator = nullptr;
+	TMap<FSHAHash, FEmitShaderValue*> ShaderValueMap;
 	FErrors Errors;
 
 	// TODO - remove preshader material dependency
@@ -147,6 +156,8 @@ public:
 	FMaterialCompilationOutput* MaterialCompilationOutput = nullptr;
 	TMap<Shader::FValue, uint32> DefaultUniformValues;
 	TMap<FSHAHash, const TCHAR*> Preshaders;
+	TArray<FScope*, TInlineAllocator<16>> ScopeStack;
+	TArray<FEmitShaderValueContext, TInlineAllocator<16>> ShaderValueStack;
 	uint32 UniformPreshaderOffset = 0u;
 	bool bReadMaterialNormal = false;
 
@@ -155,38 +166,12 @@ public:
 	int32 NumTexCoords = 0;
 };
 
-enum class ENodeVisitResult
-{
-	VisitDependentNodes,
-	SkipDependentNodes,
-};
-
-/** Should be overriden to inspect the nodes of an HLSLTree */
-class FNodeVisitor
-{
-public:
-	virtual ~FNodeVisitor() {}
-
-	void VisitNode(FNode* Node);
-
-	virtual ENodeVisitResult OnScope(FScope& Scope) { return ENodeVisitResult::VisitDependentNodes; }
-	virtual ENodeVisitResult OnStatement(FStatement& Statement) { return ENodeVisitResult::VisitDependentNodes; }
-	virtual ENodeVisitResult OnExpression(FExpression& Expression) { return ENodeVisitResult::VisitDependentNodes; }
-	virtual ENodeVisitResult OnTextureParameterDeclaration(FTextureParameterDeclaration& Declaration) { return ENodeVisitResult::VisitDependentNodes; }
-};
-
 /** Root class of the HLSL AST */
 class FNode
 {
 public:
 	virtual ~FNode() {}
-
-	static bool ShouldVisitDependentNodes(ENodeVisitResult Result) { return Result == ENodeVisitResult::VisitDependentNodes; }
-
-	virtual ENodeVisitResult Visit(FNodeVisitor& Visitor) = 0;
-
 	virtual void Reset() {}
-
 	FNode* NextNode = nullptr;
 };
 
@@ -223,8 +208,6 @@ public:
 	 */
 	TArrayView<const Shader::EValueComponentType> ComponentTypes;
 
-	virtual ENodeVisitResult Visit(FNodeVisitor& Visitor) override { return ENodeVisitResult::VisitDependentNodes; }
-
 	FStructFieldRef FindFieldByName(const TCHAR* InName) const;
 
 	void WriteHLSL(FStringBuilderBase& OutWriter) const;
@@ -257,8 +240,6 @@ public:
 	//static constexpr bool MarkScopeLiveRecursive = false;
 
 	virtual void Reset() override;
-	virtual ENodeVisitResult Visit(FNodeVisitor& Visitor) override;
-
 	virtual void Prepare(FEmitContext& Context) const = 0;
 	virtual void EmitShader(FEmitContext& Context) const = 0;
 
@@ -330,18 +311,35 @@ struct FShaderValue
 
 struct FPrepareValueResult
 {
+	FExpression* ForwardValue = nullptr;
 	FType Type;
 	EExpressionEvaluationType EvaluationType = EExpressionEvaluationType::None;
 	Shader::FValue ConstantValue;
 
-	inline operator bool() const { return EvaluationType != EExpressionEvaluationType::None; }
-	inline bool operator!() const { return EvaluationType == EExpressionEvaluationType::None; }
+	void SetConstant(FEmitContext& Context, const Shader::FValue& InValue);
+	void SetType(FEmitContext& Context, EExpressionEvaluationType InEvaluationType, const FType& InType);
+	void SetForwardValue(FEmitContext& Context, FExpression* InValue, const FRequestedType& RequestedType);
+
+	inline bool IsValid() const { return EvaluationType != EExpressionEvaluationType::None; }
+};
+
+struct FEmitShaderValue
+{
+	FEmitShaderValue(FExpression* InExpression, FScope* InScope) : Expression(InExpression), Scope(InScope) {}
+
+	inline bool IsInline() const { return Value == nullptr; }
+
+	FExpression* Expression = nullptr;
+	FScope* Scope = nullptr;
+	const TCHAR* Reference = nullptr;
+	const TCHAR* Value = nullptr;
+	TArrayView<FEmitShaderValue*> Dependencies;
+	FSHAHash Hash;
 };
 
 /**
  * Represents an HLSL expression.  This is a piece of code that evaluates to a value, but has no side effects.
  * Unlike statements, expressions are not expected to execute in any particular order.  They may be cached (or not) in generated code, without the underlying implementation needing to care.
- * Expressions track the outer-most scope in which they're accessed. The generated HLSL code will ensure they are defined in that scope.
  * Examples include constant literals, variable accessors, and various types of math operations
  * This is an abstract base class, with derived classes representing various types of expression
  */
@@ -352,33 +350,27 @@ public:
 	EExpressionEvaluationType GetEvaluationType() const { return PrepareValueResult.EvaluationType; }
 
 	virtual void Reset() override;
-	virtual ENodeVisitResult Visit(FNodeVisitor& Visitor) override;
 
-	friend FPrepareValueResult PrepareExpressionValue(FEmitContext& Context, FExpression* InExpression, const FRequestedType& RequestedType);
+	friend const FPrepareValueResult& PrepareExpressionValue(FEmitContext& Context, FExpression* InExpression, const FRequestedType& RequestedType);
 
 	const TCHAR* GetValueShader(FEmitContext& Context);
 	const TCHAR* GetValueShader(FEmitContext& Context, const FType& InType);
 	void GetValuePreshader(FEmitContext& Context, Shader::FPreshaderData& OutPreshader);
 	Shader::FValue GetValueConstant(FEmitContext& Context);
 
-	FScope* ParentScope = nullptr;
-
 protected:
-	virtual void PrepareValue(FEmitContext& Context, const FRequestedType& RequestedType) = 0;
+	virtual void PrepareValue(FEmitContext& Context, const FRequestedType& RequestedType, FPrepareValueResult& OutResult) = 0;
 	virtual void EmitValueShader(FEmitContext& Context, FShaderValue& OutShader) const;
 	virtual void EmitShaderDependencies(FEmitContext& Context, const FShaderValue& Shader) const;
 	virtual void EmitValuePreshader(FEmitContext& Context, Shader::FPreshaderData& OutPreshader) const;
 
-	void SetType(FEmitContext& Context, EExpressionEvaluationType InEvaluationType, const FType& InType);
-	void SetForwardValue(FEmitContext& Context, FExpression* InForwardValue, const FRequestedType& RequestedType);
-
 private:
-	const TCHAR* LocalVariableName = nullptr;
-	const TCHAR* Code = nullptr;
+	FEmitShaderValue* ShaderValue = nullptr;
 	FRequestedType CurrentRequestedType;
 	FPrepareValueResult PrepareValueResult;
-	FExpression* ForwardedValue = nullptr;
 	bool bReentryFlag = false;
+
+	friend class FEmitContext;
 };
 
 
@@ -390,7 +382,7 @@ private:
 class FExpressionLocalPHI final : public FExpression
 {
 public:
-	virtual void PrepareValue(FEmitContext& Context, const FRequestedType& RequestedType) override;
+	virtual void PrepareValue(FEmitContext& Context, const FRequestedType& RequestedType, FPrepareValueResult& OutResult) override;
 	virtual void EmitValueShader(FEmitContext& Context, FShaderValue& OutShader) const override;
 	virtual void EmitShaderDependencies(FEmitContext& Context, const FShaderValue& Shader) const override;
 
@@ -408,8 +400,6 @@ class FTextureParameterDeclaration final : public FNode
 public:
 	FTextureParameterDeclaration(const FName& InName, const FTextureDescription& InDescription) : Name(InName), Description(InDescription) {}
 
-	virtual ENodeVisitResult Visit(FNodeVisitor& Visitor) override;
-
 	FName Name;
 	FTextureDescription Description;
 };
@@ -443,12 +433,8 @@ public:
 
 	bool HasParentScope(const FScope& ParentScope) const;
 
-	virtual ENodeVisitResult Visit(FNodeVisitor& Visitor) override;
-
 	void AddPreviousScope(FScope& Scope);
 
-	void UseExpression(FExpression* Expression);
-
 	friend bool PrepareScope(FEmitContext& Context, FScope* InScope);
 
 	template<typename FormatType, typename... Types>
@@ -483,7 +469,6 @@ public:
 private:
 	friend class FTree;
 	friend class FExpression;
-	friend class FNodeVisitor_MoveToScope;
 
 	enum class ENextScopeFormat : uint8
 	{
@@ -522,7 +507,6 @@ private:
 	FScope* ParentScope = nullptr;
 	FStatement* ContainedStatement = nullptr;
 	FScope* PreviousScope[MaxNumPreviousScopes];
-	TMap<FSHAHash, const TCHAR*> ExpressionCodeMap;
 	FCodeList Declarations;
 	FCodeList Statements;
 	int32 NumPreviousScopes = 0;
@@ -571,10 +555,10 @@ public:
 	FScope& GetRootScope() const { return *RootScope; }
 
 	template<typename T, typename... ArgTypes>
-	inline T* NewExpression(FScope& Scope, ArgTypes&&... Args)
+	inline T* NewExpression(ArgTypes&&... Args)
 	{
 		T* Expression = NewNode<T>(Forward<ArgTypes>(Args)...);
-		RegisterExpression(Scope, Expression);
+		RegisterExpression(Expression);
 		return Expression;
 	}
 
@@ -603,11 +587,12 @@ private:
 		return Node;
 	}
 
-	void RegisterExpression(FScope& Scope, FExpression* Expression);
+	void RegisterExpression(FExpression* Expression);
 	void RegisterStatement(FScope& Scope, FStatement* Statement);
 
 	FMemStackBase* Allocator = nullptr;
 	FNode* Nodes = nullptr;
+	FExpression* ExpressionsToDeclare = nullptr;
 	FStructType* StructTypes = nullptr;
 	FScope* RootScope = nullptr;
 };
diff --git a/Engine/Source/Runtime/Engine/Public/HLSLTree/HLSLTreeCommon.h b/Engine/Source/Runtime/Engine/Public/HLSLTree/HLSLTreeCommon.h
--- a/Engine/Source/Runtime/Engine/Public/HLSLTree/HLSLTreeCommon.h
+++ b/Engine/Source/Runtime/Engine/Public/HLSLTree/HLSLTreeCommon.h
@@ -47,7 +47,7 @@ public:
 
 	FConstantValue Value;
 
-	virtual void PrepareValue(FEmitContext& Context, const FRequestedType& RequestedType) override;
+	virtual void PrepareValue(FEmitContext& Context, const FRequestedType& RequestedType, FPrepareValueResult& OutResult) override;
 	virtual void EmitValueShader(FEmitContext& Context, FShaderValue& OutShader) const override;
 	virtual void EmitValuePreshader(FEmitContext& Context, Shader::FPreshaderData& OutPreshader) const override;
 };
@@ -63,7 +63,7 @@ public:
 	Shader::FValue DefaultValue;
 	EMaterialParameterType ParameterType;
 
-	virtual void PrepareValue(FEmitContext& Context, const FRequestedType& RequestedType) override;
+	virtual void PrepareValue(FEmitContext& Context, const FRequestedType& RequestedType, FPrepareValueResult& OutResult) override;
 	virtual void EmitValuePreshader(FEmitContext& Context, Shader::FPreshaderData& OutPreshader) const override;
 };
 
@@ -95,7 +95,7 @@ public:
 
 	EExternalInputType InputType;
 
-	virtual void PrepareValue(FEmitContext& Context, const FRequestedType& RequestedType) override { SetType(Context, EExpressionEvaluationType::Shader, GetInputExpressionType(InputType)); }
+	virtual void PrepareValue(FEmitContext& Context, const FRequestedType& RequestedType, FPrepareValueResult& OutResult) override { OutResult.SetType(Context, EExpressionEvaluationType::Shader, GetInputExpressionType(InputType)); }
 	virtual void EmitValueShader(FEmitContext& Context, FShaderValue& OutShader) const override;
 };
 
@@ -114,18 +114,7 @@ public:
 	ESamplerSourceMode SamplerSource;
 	ETextureMipValueMode MipValueMode;
 
-	virtual ENodeVisitResult Visit(FNodeVisitor& Visitor) override
-	{
-		const ENodeVisitResult Result = FExpression::Visit(Visitor);
-		if (ShouldVisitDependentNodes(Result))
-		{
-			Visitor.VisitNode(Declaration);
-			Visitor.VisitNode(TexCoordExpression);
-		}
-		return Result;
-	}
-
-	virtual void PrepareValue(FEmitContext& Context, const FRequestedType& RequestedType) override;
+	virtual void PrepareValue(FEmitContext& Context, const FRequestedType& RequestedType, FPrepareValueResult& OutResult) override;
 	virtual void EmitValueShader(FEmitContext& Context, FShaderValue& OutShader) const override;
 };
 
@@ -138,17 +127,7 @@ public:
 	const TCHAR* FieldName;
 	FExpression* StructExpression;
 
-	virtual ENodeVisitResult Visit(FNodeVisitor& Visitor) override
-	{
-		const ENodeVisitResult Result = FExpression::Visit(Visitor);
-		if (ShouldVisitDependentNodes(Result))
-		{
-			Visitor.VisitNode(StructExpression);
-		}
-		return Result;
-	}
-
-	virtual void PrepareValue(FEmitContext& Context, const FRequestedType& RequestedType) override;
+	virtual void PrepareValue(FEmitContext& Context, const FRequestedType& RequestedType, FPrepareValueResult& OutResult) override;
 	virtual void EmitValueShader(FEmitContext& Context, FShaderValue& OutShader) const override;
 };
 
@@ -161,19 +140,8 @@ public:
 	const TCHAR* FieldName;
 	FExpression* StructExpression;
 	FExpression* FieldExpression;
-	
-	virtual ENodeVisitResult Visit(FNodeVisitor& Visitor) override
-	{
-		const ENodeVisitResult Result = FExpression::Visit(Visitor);
-		if (ShouldVisitDependentNodes(Result))
-		{
-			Visitor.VisitNode(StructExpression);
-			Visitor.VisitNode(FieldExpression);
-		}
-		return Result;
-	}
-
-	virtual void PrepareValue(FEmitContext& Context, const FRequestedType& RequestedType) override;
+
+	virtual void PrepareValue(FEmitContext& Context, const FRequestedType& RequestedType, FPrepareValueResult& OutResult) override;
 	virtual void EmitValueShader(FEmitContext& Context, FShaderValue& OutShader) const override;
 };
 
@@ -190,19 +158,7 @@ public:
 	FExpression* TrueExpression;
 	FExpression* FalseExpression;
 
-	virtual ENodeVisitResult Visit(FNodeVisitor& Visitor) override
-	{
-		const ENodeVisitResult Result = FExpression::Visit(Visitor);
-		if (ShouldVisitDependentNodes(Result))
-		{
-			Visitor.VisitNode(ConditionExpression);
-			Visitor.VisitNode(TrueExpression);
-			Visitor.VisitNode(FalseExpression);
-		}
-		return Result;
-	}
-
-	virtual void PrepareValue(FEmitContext& Context, const FRequestedType& RequestedType) override;
+	virtual void PrepareValue(FEmitContext& Context, const FRequestedType& RequestedType, FPrepareValueResult& OutResult) override;
 	virtual void EmitValueShader(FEmitContext& Context, FShaderValue& OutShader) const override;
 	virtual void EmitValuePreshader(FEmitContext& Context, Shader::FPreshaderData& OutPreshader) const override;
 };
@@ -220,18 +176,7 @@ public:
 	FExpression* Lhs;
 	FExpression* Rhs;
 
-	virtual ENodeVisitResult Visit(FNodeVisitor& Visitor) override
-	{
-		const ENodeVisitResult Result = FExpression::Visit(Visitor);
-		if (ShouldVisitDependentNodes(Result))
-		{
-			Visitor.VisitNode(Lhs);
-			Visitor.VisitNode(Rhs);
-		}
-		return Result;
-	}
-
-	virtual void PrepareValue(FEmitContext& Context, const FRequestedType& RequestedType) override;
+	virtual void PrepareValue(FEmitContext& Context, const FRequestedType& RequestedType, FPrepareValueResult& OutResult) override;
 	virtual void EmitValueShader(FEmitContext& Context, FShaderValue& OutShader) const override;
 	virtual void EmitValuePreshader(FEmitContext& Context, Shader::FPreshaderData& OutPreshader) const override;
 };
@@ -257,17 +202,7 @@ public:
 	FSwizzleParameters Parameters;
 	FExpression* Input;
 
-	virtual ENodeVisitResult Visit(FNodeVisitor& Visitor) override
-	{
-		const ENodeVisitResult Result = FExpression::Visit(Visitor);
-		if (ShouldVisitDependentNodes(Result))
-		{
-			Visitor.VisitNode(Input);
-		}
-		return Result;
-	}
-
-	virtual void PrepareValue(FEmitContext& Context, const FRequestedType& RequestedType) override;
+	virtual void PrepareValue(FEmitContext& Context, const FRequestedType& RequestedType, FPrepareValueResult& OutResult) override;
 	virtual void EmitValueShader(FEmitContext& Context, FShaderValue& OutShader) const override;
 	virtual void EmitValuePreshader(FEmitContext& Context, Shader::FPreshaderData& OutPreshader) const override;
 };
@@ -283,18 +218,7 @@ public:
 	FExpression* Lhs;
 	FExpression* Rhs;
 
-	virtual ENodeVisitResult Visit(FNodeVisitor& Visitor) override
-	{
-		const ENodeVisitResult Result = FExpression::Visit(Visitor);
-		if (ShouldVisitDependentNodes(Result))
-		{
-			Visitor.VisitNode(Lhs);
-			Visitor.VisitNode(Rhs);
-		}
-		return Result;
-	}
-
-	virtual void PrepareValue(FEmitContext& Context, const FRequestedType& RequestedType) override;
+	virtual void PrepareValue(FEmitContext& Context, const FRequestedType& RequestedType, FPrepareValueResult& OutResult) override;
 	virtual void EmitValueShader(FEmitContext& Context, FShaderValue& OutShader) const override;
 	virtual void EmitValuePreshader(FEmitContext& Context, Shader::FPreshaderData& OutPreshader) const override;
 };
@@ -302,7 +226,7 @@ public:
 class FExpressionReflectionVector : public FExpression
 {
 public:
-	virtual void PrepareValue(FEmitContext& Context, const FRequestedType& RequestedType) override;
+	virtual void PrepareValue(FEmitContext& Context, const FRequestedType& RequestedType, FPrepareValueResult& OutResult) override;
 	virtual void EmitValueShader(FEmitContext& Context, FShaderValue& OutShader) const override;
 };
 
@@ -312,16 +236,6 @@ public:
 	//static constexpr bool MarkScopeLiveRecursive = true;
 	FExpression* Expression;
 
-	virtual ENodeVisitResult Visit(FNodeVisitor& Visitor) override
-	{
-		const ENodeVisitResult Result = FStatement::Visit(Visitor);
-		if (ShouldVisitDependentNodes(Result))
-		{
-			Visitor.VisitNode(Expression);
-		}
-		return Result;
-	}
-
 	virtual void Prepare(FEmitContext& Context) const override;
 	virtual void EmitShader(FEmitContext& Context) const override;
 };
@@ -343,19 +257,6 @@ public:
 	FScope* ElseScope;
 	FScope* NextScope;
 
-	virtual ENodeVisitResult Visit(FNodeVisitor& Visitor) override
-	{
-		const ENodeVisitResult Result = FStatement::Visit(Visitor);
-		if (ShouldVisitDependentNodes(Result))
-		{
-			Visitor.VisitNode(ConditionExpression);
-			Visitor.VisitNode(ThenScope);
-			Visitor.VisitNode(ElseScope);
-			Visitor.VisitNode(NextScope);
-		}
-		return Result;
-	}
-
 	virtual void Prepare(FEmitContext& Context) const override;
 	virtual void EmitShader(FEmitContext& Context) const override;
 };
@@ -366,17 +267,6 @@ public:
 	FScope* LoopScope;
 	FScope* NextScope;
 
-	virtual ENodeVisitResult Visit(FNodeVisitor& Visitor) override
-	{
-		const ENodeVisitResult Result = FStatement::Visit(Visitor);
-		if (ShouldVisitDependentNodes(Result))
-		{
-			Visitor.VisitNode(LoopScope);
-			Visitor.VisitNode(NextScope);
-		}
-		return Result;
-	}
-
 	virtual void Prepare(FEmitContext& Context) const override;
 	virtual void EmitShader(FEmitContext& Context) const override;
 };
diff --git a/Engine/Source/Runtime/Engine/Public/MaterialHLSLGenerator.h b/Engine/Source/Runtime/Engine/Public/MaterialHLSLGenerator.h
--- a/Engine/Source/Runtime/Engine/Public/MaterialHLSLGenerator.h
+++ b/Engine/Source/Runtime/Engine/Public/MaterialHLSLGenerator.h
@@ -87,8 +87,8 @@ public:
 	UE::HLSLTree::FScope* NewJoinedScope(UE::HLSLTree::FScope& Scope);
 
 	UE::HLSLTree::FExpressionConstant* NewConstant(const UE::Shader::FValue& Value);
-	UE::HLSLTree::FExpressionExternalInput* NewTexCoord(UE::HLSLTree::FScope& Scope, int32 Index);
-	UE::HLSLTree::FExpressionSwizzle* NewSwizzle(UE::HLSLTree::FScope& Scope, const UE::HLSLTree::FSwizzleParameters& Params, UE::HLSLTree::FExpression* Input);
+	UE::HLSLTree::FExpressionExternalInput* NewTexCoord(int32 Index);
+	UE::HLSLTree::FExpressionSwizzle* NewSwizzle(const UE::HLSLTree::FSwizzleParameters& Params, UE::HLSLTree::FExpression* Input);
 
 	/** Returns a declaration to access the given texture, with no parameter */
 	UE::HLSLTree::FTextureParameterDeclaration* AcquireTextureDeclaration(const UE::HLSLTree::FTextureDescription& Value);
