Missed files

diff --git a/Engine/Source/Runtime/Engine/Private/HLSLTree/HLSLTreeCommon.cpp b/Engine/Source/Runtime/Engine/Private/HLSLTree/HLSLTreeCommon.cpp
--- a/Engine/Source/Runtime/Engine/Private/HLSLTree/HLSLTreeCommon.cpp
+++ b/Engine/Source/Runtime/Engine/Private/HLSLTree/HLSLTreeCommon.cpp
@@ -85,7 +85,7 @@ FSwizzleParameters MakeSwizzleMask(bool bInR, bool bInG, bool bInB, bool bInA)
 	return FSwizzleParameters(ComponentIndex[0], ComponentIndex[1], ComponentIndex[2], ComponentIndex[3]);
 }
 
-void FExpressionConstant::PrepareValue(FEmitContext& Context, const FRequestedType& RequestedType)
+void FExpressionConstant::PrepareValue(FEmitContext& Context, const FRequestedType& RequestedType, FPrepareValueResult& OutResult)
 {
 	if (Value.Type.IsStruct())
 	{
@@ -95,11 +95,11 @@ void FExpressionConstant::PrepareValue(FEmitContext& Context, const FRequestedTy
 		}
 
 		// TODO - Support for preshader structs?
-		return SetType(Context, EExpressionEvaluationType::Shader, Value.Type);
+		OutResult.SetType(Context, EExpressionEvaluationType::Shader, Value.Type);
 	}
 	else
 	{
-		return SetType(Context, EExpressionEvaluationType::Constant, Shader::MakeValueTypeWithRequestedNumComponents(Value.Type, RequestedType.GetRequestedNumComponents()));
+		OutResult.SetType(Context, EExpressionEvaluationType::Constant, Shader::MakeValueTypeWithRequestedNumComponents(Value.Type, RequestedType.GetRequestedNumComponents()));
 	}
 }
 
@@ -128,11 +128,11 @@ void FExpressionConstant::EmitValuePreshader(FEmitContext& Context, Shader::FPre
 	OutPreshader.WriteOpcode(Shader::EPreshaderOpcode::Constant).Write(BaseValue);
 }
 
-void FExpressionMaterialParameter::PrepareValue(FEmitContext& Context, const FRequestedType& RequestedType)
+void FExpressionMaterialParameter::PrepareValue(FEmitContext& Context, const FRequestedType& RequestedType, FPrepareValueResult& OutResult)
 {
 	const Shader::EValueType CurrentType = Shader::MakeValueTypeWithRequestedNumComponents(GetShaderValueType(ParameterType), RequestedType.GetRequestedNumComponents());
 	const EExpressionEvaluationType EvaluationType = IsStaticMaterialParameter(ParameterType) ? EExpressionEvaluationType::Constant : EExpressionEvaluationType::Preshader;
-	return SetType(Context, EvaluationType, CurrentType);
+	OutResult.SetType(Context, EvaluationType, CurrentType);
 }
 
 void FExpressionMaterialParameter::EmitValuePreshader(FEmitContext& Context, Shader::FPreshaderData& OutPreshader) const
@@ -187,10 +187,10 @@ void FExpressionExternalInput::EmitValueShader(FEmitContext& Context, FShaderVal
 	}
 }
 
-void FExpressionTextureSample::PrepareValue(FEmitContext& Context, const FRequestedType& RequestedType)
+void FExpressionTextureSample::PrepareValue(FEmitContext& Context, const FRequestedType& RequestedType, FPrepareValueResult& OutResult)
 {
 	PrepareExpressionValue(Context, TexCoordExpression, 2);
-	return SetType(Context, EExpressionEvaluationType::Shader, Shader::EValueType::Float4);
+	OutResult.SetType(Context, EExpressionEvaluationType::Shader, Shader::EValueType::Float4);
 }
 
 void FExpressionTextureSample::EmitValueShader(FEmitContext& Context, FShaderValue& OutShader) const
@@ -327,7 +327,7 @@ void FExpressionTextureSample::EmitValueShader(FEmitContext& Context, FShaderVal
 		TexCoordExpression->GetValueShader(Context, Shader::EValueType::Float2));
 }
 
-void FExpressionGetStructField::PrepareValue(FEmitContext& Context, const FRequestedType& RequestedType)
+void FExpressionGetStructField::PrepareValue(FEmitContext& Context, const FRequestedType& RequestedType, FPrepareValueResult& OutResult)
 {
 	const FStructFieldRef FieldRef = StructType->FindFieldByName(FieldName);
 	if (!FieldRef)
@@ -338,14 +338,14 @@ void FExpressionGetStructField::PrepareValue(FEmitContext& Context, const FReque
 	FRequestedType RequestedStructType(StructType, false);
 	RequestedStructType.SetField(FieldRef, RequestedType);
 
-	const FPrepareValueResult StructResult = PrepareExpressionValue(Context, StructExpression, RequestedStructType);
-	if (StructResult && StructResult.Type != FType(StructType))
+	const FPrepareValueResult& StructResult = PrepareExpressionValue(Context, StructExpression, RequestedStructType);
+	if (StructResult.IsValid() && StructResult.Type != FType(StructType))
 	{
 		return Context.Errors.AddErrorf(this, TEXT("Expected type %s"), StructType->Name);
 	}
 
 	// TODO - preshader structs
-	return SetType(Context, EExpressionEvaluationType::Shader, FieldRef.Type);
+	OutResult.SetType(Context, EExpressionEvaluationType::Shader, FieldRef.Type);
 }
 
 void FExpressionGetStructField::EmitValueShader(FEmitContext& Context, FShaderValue& OutShader) const
@@ -356,7 +356,7 @@ void FExpressionGetStructField::EmitValueShader(FEmitContext& Context, FShaderVa
 		FieldName);
 }
 
-void FExpressionSetStructField::PrepareValue(FEmitContext& Context, const FRequestedType& RequestedType)
+void FExpressionSetStructField::PrepareValue(FEmitContext& Context, const FRequestedType& RequestedType, FPrepareValueResult& OutResult)
 {
 	const FStructFieldRef FieldRef = StructType->FindFieldByName(FieldName);
 	if (!FieldRef)
@@ -367,20 +367,22 @@ void FExpressionSetStructField::PrepareValue(FEmitContext& Context, const FReque
 	FRequestedType RequestedStructType(RequestedType);
 	RequestedStructType.ClearFieldRequested(FieldRef);
 
-	const FPrepareValueResult StructResult = PrepareExpressionValue(Context, StructExpression, RequestedStructType);
-	if (StructResult && StructResult.Type != FType(StructType))
+	const FPrepareValueResult& StructResult = PrepareExpressionValue(Context, StructExpression, RequestedStructType);
+	if (StructResult.IsValid() && StructResult.Type != FType(StructType))
 	{
 		return Context.Errors.AddErrorf(this, TEXT("Expected type %s"), StructType->Name);
 	}
 
-	const FPrepareValueResult FieldResult = PrepareExpressionValue(Context, FieldExpression, RequestedType.GetField(FieldRef));
-	if (!FieldResult)
+	const FPrepareValueResult& FieldResult = PrepareExpressionValue(Context, FieldExpression, RequestedType.GetField(FieldRef));
+	if (!FieldResult.IsValid())
 	{
-		return SetForwardValue(Context, StructExpression, RequestedType);
+		OutResult.SetForwardValue(Context, StructExpression, RequestedType);
+	}
+	else
+	{
+		// TODO - preshader structs
+		OutResult.SetType(Context, EExpressionEvaluationType::Shader, StructType);
 	}
-
-	// TODO - preshader structs
-	return SetType(Context, EExpressionEvaluationType::Shader, StructType);
 }
 
 void FExpressionSetStructField::EmitValueShader(FEmitContext& Context, FShaderValue& OutShader) const
@@ -407,18 +409,18 @@ void FExpressionSetStructField::EmitValueShader(FEmitContext& Context, FShaderVa
 	}
 }
 
-void FExpressionSelect::PrepareValue(FEmitContext& Context, const FRequestedType& RequestedType)
+void FExpressionSelect::PrepareValue(FEmitContext& Context, const FRequestedType& RequestedType, FPrepareValueResult& OutResult)
 {
-	const FPrepareValueResult ConditionResult = PrepareExpressionValue(Context, ConditionExpression, 1);
+	const FPrepareValueResult& ConditionResult = PrepareExpressionValue(Context, ConditionExpression, 1);
 	if (ConditionResult.EvaluationType == EExpressionEvaluationType::Constant)
 	{
 		const bool bCondition = ConditionResult.ConstantValue.AsBoolScalar();
-		return SetForwardValue(Context, bCondition ? TrueExpression : FalseExpression, RequestedType);
+		return OutResult.SetForwardValue(Context, bCondition ? TrueExpression : FalseExpression, RequestedType);
 	}
 
-	const FPrepareValueResult LhsResult = PrepareExpressionValue(Context, FalseExpression, RequestedType);
-	const FPrepareValueResult RhsResult = PrepareExpressionValue(Context, TrueExpression, RequestedType);
-	if (!ConditionResult || !LhsResult || !RhsResult)
+	const FPrepareValueResult& LhsResult = PrepareExpressionValue(Context, FalseExpression, RequestedType);
+	const FPrepareValueResult& RhsResult = PrepareExpressionValue(Context, TrueExpression, RequestedType);
+	if (!ConditionResult.IsValid() || !LhsResult.IsValid() || !RhsResult.IsValid())
 	{
 		return;
 	}
@@ -432,12 +434,12 @@ void FExpressionSelect::PrepareValue(FEmitContext& Context, const FRequestedType
 
 	if (LhsResult.Type.IsStruct())
 	{
-		return SetType(Context, EExpressionEvaluationType::Shader, LhsResult.Type);
+		OutResult.SetType(Context, EExpressionEvaluationType::Shader, LhsResult.Type);
 	}
 	else
 	{
 		const int8 NumComponents = FMath::Max(LhsTypeDesc.NumComponents, RhsTypeDesc.NumComponents);
-		return SetType(Context, EExpressionEvaluationType::Shader, Shader::MakeValueType(LhsTypeDesc.ComponentType, NumComponents));
+		OutResult.SetType(Context, EExpressionEvaluationType::Shader, Shader::MakeValueType(LhsTypeDesc.ComponentType, NumComponents));
 	}
 }
 
@@ -454,11 +456,11 @@ void FExpressionSelect::EmitValuePreshader(FEmitContext& Context, Shader::FPresh
 	check(false); // TODO
 }
 
-void FExpressionBinaryOp::PrepareValue(FEmitContext& Context, const FRequestedType& RequestedType)
+void FExpressionBinaryOp::PrepareValue(FEmitContext& Context, const FRequestedType& RequestedType, FPrepareValueResult& OutResult)
 {
-	const FPrepareValueResult LhsResult = PrepareExpressionValue(Context, Lhs, RequestedType);
-	const FPrepareValueResult RhsResult = PrepareExpressionValue(Context, Rhs, RequestedType);
-	if (!LhsResult || !RhsResult)
+	const FPrepareValueResult& LhsResult = PrepareExpressionValue(Context, Lhs, RequestedType);
+	const FPrepareValueResult& RhsResult = PrepareExpressionValue(Context, Rhs, RequestedType);
+	if (!LhsResult.IsValid() || !RhsResult.IsValid())
 	{
 		return;
 	}
@@ -483,7 +485,7 @@ void FExpressionBinaryOp::PrepareValue(FEmitContext& Context, const FRequestedTy
 	const EExpressionEvaluationType EvaluationType = OpDesc.PreshaderOpcode != Shader::EPreshaderOpcode::Nop
 		? CombineEvaluationTypes(LhsResult.EvaluationType, RhsResult.EvaluationType)
 		: EExpressionEvaluationType::Shader;
-	return SetType(Context, EvaluationType, ResultType);
+	OutResult.SetType(Context, EvaluationType, ResultType);
 }
 
 void FExpressionBinaryOp::EmitValueShader(FEmitContext& Context, FShaderValue& OutShader) const
@@ -511,7 +513,7 @@ void FExpressionBinaryOp::EmitValuePreshader(FEmitContext& Context, Shader::FPre
 	OutPreshader.WriteOpcode(OpDesc.PreshaderOpcode);
 }
 
-void FExpressionSwizzle::PrepareValue(FEmitContext& Context, const FRequestedType& RequestedType)
+void FExpressionSwizzle::PrepareValue(FEmitContext& Context, const FRequestedType& RequestedType, FPrepareValueResult& OutResult)
 {
 	FRequestedType RequestedInputType(4, false);
 	int32 MaxComponentIndex = 0;
@@ -525,8 +527,8 @@ void FExpressionSwizzle::PrepareValue(FEmitContext& Context, const FRequestedTyp
 		}
 	}
 
-	const FPrepareValueResult InputResult = PrepareExpressionValue(Context, Input, RequestedInputType);
-	return SetType(Context, InputResult.EvaluationType, Shader::MakeValueType(InputResult.Type, MaxComponentIndex + 1));
+	const FPrepareValueResult& InputResult = PrepareExpressionValue(Context, Input, RequestedInputType);
+	OutResult.SetType(Context, InputResult.EvaluationType, Shader::MakeValueType(InputResult.Type, MaxComponentIndex + 1));
 }
 
 void FExpressionSwizzle::EmitValueShader(FEmitContext& Context, FShaderValue& OutShader) const
@@ -574,11 +576,11 @@ void FExpressionSwizzle::EmitValuePreshader(FEmitContext& Context, Shader::FPres
 		.Write((uint8)Parameters.ComponentIndex[3]);
 }
 
-void FExpressionAppend::PrepareValue(FEmitContext& Context, const FRequestedType& RequestedType)
+void FExpressionAppend::PrepareValue(FEmitContext& Context, const FRequestedType& RequestedType, FPrepareValueResult& OutResult)
 {
-	const FPrepareValueResult LhsResult = PrepareExpressionValue(Context, Lhs, 4);
-	const FPrepareValueResult RhsResult = PrepareExpressionValue(Context, Rhs, 4);
-	if (!LhsResult || !RhsResult)
+	const FPrepareValueResult& LhsResult = PrepareExpressionValue(Context, Lhs, 4);
+	const FPrepareValueResult& RhsResult = PrepareExpressionValue(Context, Rhs, 4);
+	if (!LhsResult.IsValid() || !RhsResult.IsValid())
 	{
 		return;
 	}
@@ -597,7 +599,7 @@ void FExpressionAppend::PrepareValue(FEmitContext& Context, const FRequestedType
 
 	const EExpressionEvaluationType EvaluationType = CombineEvaluationTypes(LhsResult.EvaluationType, RhsResult.EvaluationType);
 	const int32 NumComponents = FMath::Min<int32>(LhsTypeDesc.NumComponents + RhsTypeDesc.NumComponents, RequestedType.GetRequestedNumComponents());
-	return SetType(Context, EvaluationType, Shader::MakeValueType(LhsTypeDesc.ComponentType, NumComponents));
+	OutResult.SetType(Context, EvaluationType, Shader::MakeValueType(LhsTypeDesc.ComponentType, NumComponents));
 }
 
 void FExpressionAppend::EmitValueShader(FEmitContext& Context, FShaderValue& OutShader) const
@@ -617,10 +619,10 @@ void FExpressionAppend::EmitValuePreshader(FEmitContext& Context, Shader::FPresh
 	OutPreshader.WriteOpcode(Shader::EPreshaderOpcode::AppendVector);
 }
 
-void FExpressionReflectionVector::PrepareValue(FEmitContext& Context, const FRequestedType& RequestedType)
+void FExpressionReflectionVector::PrepareValue(FEmitContext& Context, const FRequestedType& RequestedType, FPrepareValueResult& OutResult)
 {
 	Context.bReadMaterialNormal = true;
-	return SetType(Context, EExpressionEvaluationType::Shader, Shader::EValueType::Float3);
+	OutResult.SetType(Context, EExpressionEvaluationType::Shader, Shader::EValueType::Float3);
 }
 
 void FExpressionReflectionVector::EmitValueShader(FEmitContext& Context, FShaderValue& OutShader) const
@@ -649,7 +651,7 @@ void FStatementReturn::EmitShader(FEmitContext& Context) const
 
 void FStatementIf::Prepare(FEmitContext& Context) const
 {
-	const FPrepareValueResult ConditionResult = PrepareExpressionValue(Context, ConditionExpression, 1);
+	const FPrepareValueResult& ConditionResult = PrepareExpressionValue(Context, ConditionExpression, 1);
 	if (ConditionResult.EvaluationType == EExpressionEvaluationType::Constant)
 	{
 		const bool bCondition = ConditionResult.ConstantValue.AsBoolScalar();
@@ -664,7 +666,7 @@ void FStatementIf::Prepare(FEmitContext& Context) const
 			MarkScopeLive(ElseScope);
 		}
 	}
-	else if(ConditionResult)
+	else if(ConditionResult.IsValid())
 	{
 		MarkScopeLive(ThenScope);
 		MarkScopeLive(ElseScope);
diff --git a/Engine/Source/Runtime/Engine/Private/Materials/MaterialHLSLEmitter.cpp b/Engine/Source/Runtime/Engine/Private/Materials/MaterialHLSLEmitter.cpp
--- a/Engine/Source/Runtime/Engine/Private/Materials/MaterialHLSLEmitter.cpp
+++ b/Engine/Source/Runtime/Engine/Private/Materials/MaterialHLSLEmitter.cpp
@@ -759,10 +759,7 @@ bool MaterialEmitHLSL(const FMaterialCompileTargetParameters& InCompilerTarget,
 	// Prepare all fields *except* normal
 	FRequestedType RequestedMaterialAttributesType(Generator.GetMaterialAttributesType());
 	RequestedMaterialAttributesType.ClearFieldRequested(NormalFieldRef);
-	if (!PrepareExpressionValue(EmitContext, Generator.GetResultExpression(), RequestedMaterialAttributesType))
-	{
-		return false;
-	}
+	PrepareExpressionValue(EmitContext, Generator.GetResultExpression(), RequestedMaterialAttributesType);
 
 	const bool bReadMaterialNormal = EmitContext.bReadMaterialNormal;
 	FStringBuilderMemstack CodePhase0(Allocator, 1024 * 1024);
@@ -771,10 +768,7 @@ bool MaterialEmitHLSL(const FMaterialCompileTargetParameters& InCompilerTarget,
 	{
 		// No access to material normal, can execute everything in phase0
 		RequestedMaterialAttributesType.SetFieldRequested(NormalFieldRef);
-		if (!PrepareExpressionValue(EmitContext, Generator.GetResultExpression(), RequestedMaterialAttributesType))
-		{
-			return false;
-		}
+		PrepareExpressionValue(EmitContext, Generator.GetResultExpression(), RequestedMaterialAttributesType);
 		HLSLTree->EmitShader(EmitContext, CodePhase0);
 	}
 	else
@@ -792,10 +786,7 @@ bool MaterialEmitHLSL(const FMaterialCompileTargetParameters& InCompilerTarget,
 		// Prepare code for just the normal
 		FRequestedType RequestedMaterialNormal(Generator.GetMaterialAttributesType(), false);
 		RequestedMaterialNormal.SetFieldRequested(NormalFieldRef);
-		if (!PrepareExpressionValue(EmitContext, Generator.GetResultExpression(), RequestedMaterialNormal))
-		{
-			return false;
-		}
+		PrepareExpressionValue(EmitContext, Generator.GetResultExpression(), RequestedMaterialNormal);
 
 		// Execute the normal in phase0
 		HLSLTree->EmitShader(EmitContext, CodePhase0);
