Avoid triggering sound loading and compilation during GC
  - Saves 1 minutes of GC during cook

diff --git a/Engine/Source/Runtime/Engine/Private/SoundWave.cpp b/Engine/Source/Runtime/Engine/Private/SoundWave.cpp
--- a/Engine/Source/Runtime/Engine/Private/SoundWave.cpp
+++ b/Engine/Source/Runtime/Engine/Private/SoundWave.cpp
@@ -2097,11 +2097,10 @@ bool USoundWave::IsReadyForFinishDestroy()
 	bool bIsStreamingInProgress = false;
 	if (CurrentLoadingBehavior != ESoundWaveLoadingBehavior::ForceInline)
 	{
-		if (!InternalProxy.IsValid())
+		if (InternalProxy.IsValid())
 		{
-			InternalProxy = CreateSoundWaveProxy();
+			bIsStreamingInProgress = IStreamingManager::Get().GetAudioStreamingManager().IsStreamingInProgress(InternalProxy);
 		}
-		bIsStreamingInProgress = IStreamingManager::Get().GetAudioStreamingManager().IsStreamingInProgress(InternalProxy);
 	}
 
 	check(GetPrecacheState() != ESoundWavePrecacheState::InProgress);
