World Partition: don't use LoadPackageAsync when cooking to load actor packages to use the same codepath as non partitioned worlds.
This make sure we keep the same behavior for missing imports (warnings).

diff --git a/Engine/Source/Runtime/Engine/Private/WorldPartition/WorldPartitionLevelHelper.cpp b/Engine/Source/Runtime/Engine/Private/WorldPartition/WorldPartitionLevelHelper.cpp
--- a/Engine/Source/Runtime/Engine/Private/WorldPartition/WorldPartitionLevelHelper.cpp
+++ b/Engine/Source/Runtime/Engine/Private/WorldPartition/WorldPartitionLevelHelper.cpp
@@ -214,7 +214,7 @@ ULevel* FWorldPartitionLevelHelper::CreateEmptyLevelForRuntimeCell(const UWorld*
 	return NewLevel;
 }
 
-bool FWorldPartitionLevelHelper::LoadActors(ULevel* InDestLevel, TArrayView<FWorldPartitionRuntimeCellObjectMapping> InActorPackages, FWorldPartitionPackageCache& InPackageCache, TFunction<void(bool)> InCompletionCallback, bool bLoadForPlay, bool bLoadAsync /*=true*/, FLinkerInstancingContext* InOutInstancingContext /*=nullptr*/)
+bool FWorldPartitionLevelHelper::LoadActors(ULevel* InDestLevel, TArrayView<FWorldPartitionRuntimeCellObjectMapping> InActorPackages, FWorldPartitionPackageCache& InPackageCache, TFunction<void(bool)> InCompletionCallback, bool bLoadForPlay, FLinkerInstancingContext* InOutInstancingContext /*=nullptr*/)
 {
 	UPackage* DestPackage = InDestLevel ? InDestLevel->GetPackage() : nullptr;
 	FString ShortLevelPackageName = DestPackage? FPackageName::GetShortName(DestPackage->GetFName()) : FString();
@@ -297,17 +297,17 @@ bool FWorldPartitionLevelHelper::LoadActors(ULevel* InDestLevel, TArrayView<FWor
 			}
 		});
 
-
-		FPackagePath PackagePath = FPackagePath::FromPackageNameChecked(*ActorPackageName);
 		if (bLoadForPlay)
 		{
+			FPackagePath PackagePath = FPackagePath::FromPackageNameChecked(*ActorPackageName);
 			check(DestPackage);
 			const EPackageFlags PackageFlags = DestPackage->HasAnyPackageFlags(PKG_PlayInEditor) ? PKG_PlayInEditor : PKG_None;
 			::LoadPackageAsync(PackagePath, FName(*ActorPackageInstanceNames[ChildIndex]), CompletionCallback, PackageFlags, DestPackage->PIEInstanceID, 0, InOutInstancingContext);
 		}
 		else
 		{
-			::LoadPackageAsync(PackagePath, NAME_None, CompletionCallback, PKG_None, INDEX_NONE, 0, InOutInstancingContext);
+			UPackage* Package = LoadPackage(nullptr, *ActorPackageName, LOAD_None, nullptr, InOutInstancingContext);
+			CompletionCallback.Execute(*ActorPackageName, Package, Package ? EAsyncLoadingResult::Succeeded : EAsyncLoadingResult::Failed);
 		}
 	}
 
@@ -408,14 +408,8 @@ bool FWorldPartitionLevelHelper::LoadActors(ULevel* InDestLevel, TArrayView<FWor
 		// PostFix the Package so that we don't load over the original package
 		FName PackageDuplicateSeedName(*FString::Printf(TEXT("%s_DUP"), *PackageToLoadFrom));
 		
-		// Use Package Cache here because Level Instances can load the same levels across cells we want to avoid Loading the same package multiple times...
-		InPackageCache.LoadWorldPackageAsync(PackageDuplicateSeedName, *PackageToLoadFrom, CompletionCallback);
-	}
-
-	if (!bLoadAsync)
-	{
-		// Finish all async loading.
-		FlushAsyncLoading();
+		// Use package cache here because level instances can load the same levels across cells we want to avoid Loading the same package multiple times...
+		InPackageCache.LoadPackage(PackageDuplicateSeedName, *PackageToLoadFrom, CompletionCallback, /*bLoadAsync*/bLoadForPlay, /*bInWorldPackage*/true);
 	}
 
 	return (LoadProgress->NumPendingLoadRequests == 0);
diff --git a/Engine/Source/Runtime/Engine/Private/WorldPartition/WorldPartitionLevelStreamingDynamic.cpp b/Engine/Source/Runtime/Engine/Private/WorldPartition/WorldPartitionLevelStreamingDynamic.cpp
--- a/Engine/Source/Runtime/Engine/Private/WorldPartition/WorldPartitionLevelStreamingDynamic.cpp
+++ b/Engine/Source/Runtime/Engine/Private/WorldPartition/WorldPartitionLevelStreamingDynamic.cpp
@@ -297,7 +297,7 @@ bool UWorldPartitionLevelStreamingDynamic::IssueLoadRequests()
 			}
 
 			FinalizeLoading();
-		},/*bLoadForPlay=*/true, /*bLoadAsync=*/true, &InstancingContext);
+		},/*bLoadForPlay=*/true, &InstancingContext);
 	}
 	else
 	{
diff --git a/Engine/Source/Runtime/Engine/Private/WorldPartition/WorldPartitionPackageCache.cpp b/Engine/Source/Runtime/Engine/Private/WorldPartition/WorldPartitionPackageCache.cpp
--- a/Engine/Source/Runtime/Engine/Private/WorldPartition/WorldPartitionPackageCache.cpp
+++ b/Engine/Source/Runtime/Engine/Private/WorldPartition/WorldPartitionPackageCache.cpp
@@ -72,17 +72,7 @@ void FWorldPartitionPackageCache::UnloadPackage(UPackage* InPackage)
 	}
 }
 
-void FWorldPartitionPackageCache::LoadWorldPackageAsync(FName InPackageName, const TCHAR* InPackageToLoadFrom /*= nullptr*/, FLoadPackageAsyncDelegate InCompletionDelegate /*= FLoadPackageAsyncDelegate()*/, EPackageFlags InPackageFlags /*= PKG_None*/, int32 InPIEInstanceID /*= INDEX_NONE*/, int32 InPackagePriority /*= 0*/, const FLinkerInstancingContext* InInstancingContext /*=nullptr*/)
-{
-	LoadPackageAsyncInternal(InPackageName, InPackageToLoadFrom, InCompletionDelegate, InPackageFlags, InPIEInstanceID, InPackagePriority, InInstancingContext, /* bInWorldPackage=*/true);
-}
-
-void FWorldPartitionPackageCache::LoadPackageAsync(FName InPackageName, const TCHAR* InPackageToLoadFrom /*= nullptr*/, FLoadPackageAsyncDelegate InCompletionDelegate /*= FLoadPackageAsyncDelegate()*/, EPackageFlags InPackageFlags /*= PKG_None*/, int32 InPIEInstanceID /*= INDEX_NONE*/, int32 InPackagePriority /*= 0*/, const FLinkerInstancingContext* InInstancingContext /*=nullptr*/)
-{
-	LoadPackageAsyncInternal(InPackageName, InPackageToLoadFrom, InCompletionDelegate, InPackageFlags, InPIEInstanceID, InPackagePriority, InInstancingContext, /* bInWorldPackage=*/false);
-}
-
-void FWorldPartitionPackageCache::LoadPackageAsyncInternal(FName InPackageName, const TCHAR* InPackageToLoadFrom, FLoadPackageAsyncDelegate InCompletionDelegate, EPackageFlags InPackageFlags, int32 InPIEInstanceID, int32 InPackagePriority, const FLinkerInstancingContext* InInstancingContext, bool bInWorldPackage)
+void FWorldPartitionPackageCache::LoadPackage(FName InPackageName, const TCHAR* InPackageToLoadFrom, FLoadPackageAsyncDelegate InCompletionDelegate, bool bLoadAsync, bool bInWorldPackage)
 {
 	if (UPackage* CachedPackage = FindPackage(InPackageName))
 	{
@@ -119,7 +109,17 @@ void FWorldPartitionPackageCache::LoadPackageAsyncInternal(FName InPackageName,
 	{
 		UWorld::WorldTypePreLoadMap.FindOrAdd(InPackageName) = EWorldType::Editor;
 	}
-	::LoadPackageAsync(FPackagePath::FromPackageNameChecked(InPackageToLoadFrom), InPackageName, CompletionCallback, InPackageFlags, InPIEInstanceID, InPackagePriority, InInstancingContext);
+
+	if (bLoadAsync)
+	{
+		::LoadPackageAsync(FPackagePath::FromPackageNameChecked(InPackageToLoadFrom), InPackageName, CompletionCallback);
+	}
+	else
+	{
+		UPackage* Package = CreatePackage(*InPackageName.ToString());
+		Package = ::LoadPackage(Package, InPackageToLoadFrom, LOAD_None);
+		CompletionCallback.Execute(InPackageName, Package, Package ? EAsyncLoadingResult::Succeeded : EAsyncLoadingResult::Failed);
+	}
 }
 
 UPackage* FWorldPartitionPackageCache::FindPackage(FName InPackageName)
diff --git a/Engine/Source/Runtime/Engine/Private/WorldPartition/WorldPartitionRuntimeLevelStreamingCell.cpp b/Engine/Source/Runtime/Engine/Private/WorldPartition/WorldPartitionRuntimeLevelStreamingCell.cpp
--- a/Engine/Source/Runtime/Engine/Private/WorldPartition/WorldPartitionRuntimeLevelStreamingCell.cpp
+++ b/Engine/Source/Runtime/Engine/Private/WorldPartition/WorldPartitionRuntimeLevelStreamingCell.cpp
@@ -140,7 +140,7 @@ UWorldPartitionLevelStreamingDynamic* UWorldPartitionRuntimeLevelStreamingCell::
 void UWorldPartitionRuntimeLevelStreamingCell::LoadActorsForCook()
 {
 	FWorldPartitionPackageCache PackageCache;
-	verify(FWorldPartitionLevelHelper::LoadActors(nullptr, Packages, PackageCache, [](bool){}, /*bLoadForPlay=*/false, /*bLoadAsync=*/false));
+	verify(FWorldPartitionLevelHelper::LoadActors(nullptr, Packages, PackageCache, [](bool){}, /*bLoadForPlay=*/false));
 }
 
 void UWorldPartitionRuntimeLevelStreamingCell::MoveAlwaysLoadedContentToPersistentLevel()
diff --git a/Engine/Source/Runtime/Engine/Public/WorldPartition/WorldPartitionLevelHelper.h b/Engine/Source/Runtime/Engine/Public/WorldPartition/WorldPartitionLevelHelper.h
--- a/Engine/Source/Runtime/Engine/Public/WorldPartition/WorldPartitionLevelHelper.h
+++ b/Engine/Source/Runtime/Engine/Public/WorldPartition/WorldPartitionLevelHelper.h
@@ -25,7 +25,7 @@ public:
 	static void MoveExternalActorsToLevel(const TArray<FWorldPartitionRuntimeCellObjectMapping>& InChildPackages, ULevel* InLevel);
 	static void RemapLevelSoftObjectPaths(ULevel* InLevel, UWorldPartition* InWorldPartition);
 	
-	static bool LoadActors(ULevel* InDestLevel, TArrayView<FWorldPartitionRuntimeCellObjectMapping> InActorPackages, FWorldPartitionPackageCache& InPackageCache, TFunction<void(bool)> InCompletionCallback, bool bLoadForPlay, bool bLoadAsync = true, FLinkerInstancingContext* InOutInstancingContext = nullptr);
+	static bool LoadActors(ULevel* InDestLevel, TArrayView<FWorldPartitionRuntimeCellObjectMapping> InActorPackages, FWorldPartitionPackageCache& InPackageCache, TFunction<void(bool)> InCompletionCallback, bool bLoadForPlay, FLinkerInstancingContext* InOutInstancingContext = nullptr);
 	
 	static FString AddActorContainerIDToActorPath(const FActorContainerID& InContainerID, const FString& InActorPath);
 	static FString AddActorContainerIDToSubPathString(const FActorContainerID& InContainerID, const FString& InSubPathString);
diff --git a/Engine/Source/Runtime/Engine/Public/WorldPartition/WorldPartitionPackageCache.h b/Engine/Source/Runtime/Engine/Public/WorldPartition/WorldPartitionPackageCache.h
--- a/Engine/Source/Runtime/Engine/Public/WorldPartition/WorldPartitionPackageCache.h
+++ b/Engine/Source/Runtime/Engine/Public/WorldPartition/WorldPartitionPackageCache.h
@@ -18,15 +18,14 @@ public:
 	FWorldPartitionPackageCache();
 	~FWorldPartitionPackageCache();
 
-	void LoadWorldPackageAsync(FName InPackageName, const TCHAR* InPackageToLoadFrom = nullptr, FLoadPackageAsyncDelegate InCompletionDelegate = FLoadPackageAsyncDelegate(), EPackageFlags InPackageFlags = PKG_None, int32 InPIEInstanceID = INDEX_NONE, int32 InPackagePriority = 0, const FLinkerInstancingContext* InInstancingContext = nullptr);
-	void LoadPackageAsync(FName InPackageName, const TCHAR* InPackageToLoadFrom = nullptr, FLoadPackageAsyncDelegate InCompletionDelegate = FLoadPackageAsyncDelegate(), EPackageFlags InPackageFlags = PKG_None, int32 InPIEInstanceID = INDEX_NONE, int32 InPackagePriority = 0, const FLinkerInstancingContext* InInstancingContext = nullptr);
+	void LoadPackage(FName InPackageName, const TCHAR* InPackageToLoadFrom, FLoadPackageAsyncDelegate InCompletionDelegate, bool bLoadAsync, bool bInWorldPackage);
 	UPackage* FindPackage(FName InPackageName);
 	UPackage* DuplicateWorldPackage(UPackage* InPackage, FName InPackageName);
 	void UnloadPackages();
 	void TrashPackage(UPackage* InPackage);
 	
 private:
-	void LoadPackageAsyncInternal(FName InPackageName, const TCHAR* InPackageToLoadFrom, FLoadPackageAsyncDelegate InCompletionDelegate, EPackageFlags InPackageFlags, int32 InPIEInstanceID, int32 InPackagePriority, const FLinkerInstancingContext* InInstancingContext, bool bInWorldPackage);
+
 	void UnloadPackage(UPackage* InPackage);
 
 	TMap<FName, TWeakObjectPtr<UPackage>> CachedPackages;
