Non-Unity fixes

diff --git a/Engine/Plugins/Experimental/ComputeFramework/Source/ComputeFrameworkEditor/Private/ComputeFrameworkEditorModule.cpp b/Engine/Plugins/Experimental/ComputeFramework/Source/ComputeFrameworkEditor/Private/ComputeFrameworkEditorModule.cpp
--- a/Engine/Plugins/Experimental/ComputeFramework/Source/ComputeFrameworkEditor/Private/ComputeFrameworkEditorModule.cpp
+++ b/Engine/Plugins/Experimental/ComputeFramework/Source/ComputeFrameworkEditor/Private/ComputeFrameworkEditorModule.cpp
@@ -1,6 +1,7 @@
 // Copyright Epic Games, Inc. All Rights Reserved.
 
 #include "ComputeFramework/ComputeFrameworkEditorModule.h"
+#include "ComputeFramework/ComputeFrameworkCompilationTick.h"
 
 #include "ComputeFramework/ComputeKernelAssetActions.h"
 #include "ComputeFramework/ComputeKernelFromTextAssetActions.h"
diff --git a/Engine/Source/Editor/UnrealEd/Private/WorldPartition/WorldPartitionResaveActorsBuilder.cpp b/Engine/Source/Editor/UnrealEd/Private/WorldPartition/WorldPartitionResaveActorsBuilder.cpp
--- a/Engine/Source/Editor/UnrealEd/Private/WorldPartition/WorldPartitionResaveActorsBuilder.cpp
+++ b/Engine/Source/Editor/UnrealEd/Private/WorldPartition/WorldPartitionResaveActorsBuilder.cpp
@@ -11,6 +11,8 @@
 #include "Commandlets/Commandlet.h"
 #include "UObject/MetaData.h"
 #include "UObject/SavePackage.h"
+#include "Kismet2/BlueprintEditorUtils.h"
+
 
 DEFINE_LOG_CATEGORY_STATIC(LogWorldPartitionResaveActorsBuilder, All, All);
 
diff --git a/Engine/Source/Runtime/Engine/Private/WorldPartition/ErrorHandling/WorldPartitionStreamingGenerationLogErrorHandler.cpp b/Engine/Source/Runtime/Engine/Private/WorldPartition/ErrorHandling/WorldPartitionStreamingGenerationLogErrorHandler.cpp
--- a/Engine/Source/Runtime/Engine/Private/WorldPartition/ErrorHandling/WorldPartitionStreamingGenerationLogErrorHandler.cpp
+++ b/Engine/Source/Runtime/Engine/Private/WorldPartition/ErrorHandling/WorldPartitionStreamingGenerationLogErrorHandler.cpp
@@ -3,6 +3,7 @@
 #if WITH_EDITOR
 #include "WorldPartition/ErrorHandling/WorldPartitionStreamingGenerationLogErrorHandler.h"
 #include "WorldPartition/WorldPartitionLog.h"
+#include "GameFramework/Actor.h"
 
 void FStreamingGenerationLogErrorHandler::OnInvalidReference(const FWorldPartitionActorDescView& ActorDescView, const FGuid& ReferenceGuid)
 {
