World Partition: don't pollute the runtime hash object when checking map for errors.

diff --git a/Engine/Source/Runtime/Engine/Private/WorldPartition/WorldPartitionStreamingGeneration.cpp b/Engine/Source/Runtime/Engine/Private/WorldPartition/WorldPartitionStreamingGeneration.cpp
--- a/Engine/Source/Runtime/Engine/Private/WorldPartition/WorldPartitionStreamingGeneration.cpp
+++ b/Engine/Source/Runtime/Engine/Private/WorldPartition/WorldPartitionStreamingGeneration.cpp
@@ -35,7 +35,7 @@ class FWorldPartitionStreamingGenerator
 {
 	void CreateActorDescViewMap(const UActorDescContainer* InContainer, TMap<FGuid, FWorldPartitionActorDescView>& OutActorDescViewMap, const FActorContainerID& InContainerID)
 	{
-		const bool bHandleUnsavedActors = bIncludeUnsavedActors && InContainerID.IsMainContainer();
+		const bool bHandleUnsavedActors = ModifiedActorsDescList && InContainerID.IsMainContainer();
 
 		// Consider all actors of a /Temp/ container package as Unsaved because loading them from disk will fail (Outer world name mismatch)
 		const bool bIsTempContainerPackage = FPackageName::IsTempPackage(InContainer->GetPackage()->GetName());
@@ -58,7 +58,7 @@ class FWorldPartitionStreamingGenerator
 					if (bHandleUnsavedActors && (bIsTempContainerPackage || Actor->GetPackage()->IsDirty()))
 					{
 						// Dirty, unsaved actor for PIE
-						FWorldPartitionActorDesc* ActorDesc = RuntimeHash->ModifiedActorDescListForPIE.AddActor(Actor);
+						FWorldPartitionActorDesc* ActorDesc = ModifiedActorsDescList->AddActor(Actor);
 						ActorDesc->OnRegister(Actor->GetWorld());
 						OutActorDescViewMap.Emplace(ActorDescIt->GetGuid(), ActorDesc);
 						continue;
@@ -77,7 +77,7 @@ class FWorldPartitionStreamingGenerator
 			{
 				if (IsValid(Actor) && Actor->IsPackageExternal() && Actor->IsMainPackageActor() && !Actor->IsEditorOnly() && !InContainer->GetActorDesc(Actor->GetActorGuid()))
 				{
-					FWorldPartitionActorDesc* ActorDesc = RuntimeHash->ModifiedActorDescListForPIE.AddActor(Actor);
+					FWorldPartitionActorDesc* ActorDesc = ModifiedActorsDescList->AddActor(Actor);
 					ActorDesc->OnRegister(Actor->GetWorld());
 					OutActorDescViewMap.Emplace(ActorDesc->GetGuid(), ActorDesc);
 				}
@@ -294,16 +294,16 @@ class FWorldPartitionStreamingGenerator
 	}
 
 public:
-	FWorldPartitionStreamingGenerator(bool bInIncludeUnsavedActors, UWorldPartitionRuntimeHash* InRuntimeHash, IStreamingGenerationErrorHandler* InErrorHandler = nullptr)
-	: bIncludeUnsavedActors(bInIncludeUnsavedActors)
-	, RuntimeHash(InRuntimeHash)
+	FWorldPartitionStreamingGenerator(UWorldPartitionRuntimeHash* InRuntimeHash, FActorDescList* InModifiedActorsDescList, IStreamingGenerationErrorHandler* InErrorHandler = nullptr)
+	: RuntimeHash(InRuntimeHash)
+	, ModifiedActorsDescList(InModifiedActorsDescList)
 	, ErrorHandler(InErrorHandler)
 	{}
 
-	void PreparationPhase(const UWorldPartition* WorldPartition)
+	void PreparationPhase(const UActorDescContainer* Container)
 	{
 		// Preparation Phase :: Actor Descriptor Views Creation
-		CreateActorDescriptorViews(WorldPartition);
+		CreateActorDescriptorViews(Container);
 
 		// Preparation Phase :: Actor Descriptor Views Validation
 		ValidateActorDescriptorViews();
@@ -329,8 +329,8 @@ public:
 	}
 
 private:
-	bool bIncludeUnsavedActors;
 	UWorldPartitionRuntimeHash* RuntimeHash;
+	FActorDescList* ModifiedActorsDescList;
 	IStreamingGenerationErrorHandler* ErrorHandler;
 
 	struct FContainerDescriptor
@@ -360,19 +360,23 @@ private:
 
 bool UWorldPartition::GenerateStreaming(TArray<FString>* OutPackagesToGenerate)
 {
+	FActorDescList* ModifiedActorsDescList = nullptr;
+
 	FStreamingGenerationLogErrorHandler LogErrorHandler;
 	FStreamingGenerationMapCheckErrorHandler MapCheckErrorHandler;	
-	IStreamingGenerationErrorHandler* ErrorHandler = &LogErrorHandler;
+	IStreamingGenerationErrorHandler* ErrorHandler = &LogErrorHandler;	
 
 	if (bIsPIE)
 	{
+		ModifiedActorsDescList = &RuntimeHash->ModifiedActorDescListForPIE;
+		
 		// In PIE, we always want to populate the map check dialog
-		ErrorHandler = &MapCheckErrorHandler;
+		ErrorHandler = &MapCheckErrorHandler;		
 	}
 
 	FActorClusterContext ActorClusterContext;
 	{
-		FWorldPartitionStreamingGenerator StreamingGenerator(bIsPIE, RuntimeHash, ErrorHandler);
+		FWorldPartitionStreamingGenerator StreamingGenerator(RuntimeHash, ModifiedActorsDescList, ErrorHandler);
 
 		// Preparation Phase
 		StreamingGenerator.PreparationPhase(this);
@@ -398,7 +402,7 @@ bool UWorldPartition::GenerateStreaming(TArray<FString>* OutPackagesToGenerate)
 void UWorldPartition::GenerateHLOD(ISourceControlHelper* SourceControlHelper, bool bCreateActorsOnly)
 {
 	FStreamingGenerationLogErrorHandler LogErrorHandler;
-	FWorldPartitionStreamingGenerator StreamingGenerator(bIsPIE, RuntimeHash, &LogErrorHandler);
+	FWorldPartitionStreamingGenerator StreamingGenerator(RuntimeHash, nullptr, &LogErrorHandler);
 	StreamingGenerator.PreparationPhase(this);
 
 	// Preparation Phase :: Actor Clusters Creation
@@ -414,7 +418,8 @@ void UWorldPartition::CheckForErrors(IStreamingGenerationErrorHandler* ErrorHand
 {
 	FActorClusterContext ActorClusterContext;
 	{		
-		FWorldPartitionStreamingGenerator StreamingGenerator(true, RuntimeHash, ErrorHandler);
+		FActorDescList ModifiedActorDescList;
+		FWorldPartitionStreamingGenerator StreamingGenerator(RuntimeHash, &ModifiedActorDescList, ErrorHandler);
 		StreamingGenerator.PreparationPhase(this);
 	}
 }
