Fix crash on shutdown: Use GetModulePtr instead of LoadModuleChecked

diff --git a/Engine/Source/Runtime/Engine/Private/LevelInstance/LevelInstanceSubsystem.cpp b/Engine/Source/Runtime/Engine/Private/LevelInstance/LevelInstanceSubsystem.cpp
--- a/Engine/Source/Runtime/Engine/Private/LevelInstance/LevelInstanceSubsystem.cpp
+++ b/Engine/Source/Runtime/Engine/Private/LevelInstance/LevelInstanceSubsystem.cpp
@@ -84,8 +84,10 @@ void ULevelInstanceSubsystem::Deinitialize()
 	Super::Deinitialize();
 
 #if WITH_EDITOR
-	ILevelInstanceEditorModule& EditorModule = FModuleManager::LoadModuleChecked<ILevelInstanceEditorModule>("LevelInstanceEditor");
-	EditorModule.OnExitEditorMode().RemoveAll(this);
+	if (ILevelInstanceEditorModule* EditorModule = FModuleManager::GetModulePtr<ILevelInstanceEditorModule>("LevelInstanceEditor"))
+	{
+		EditorModule->OnExitEditorMode().RemoveAll(this);
+	}
 
 	if (!GetWorld()->IsGameWorld())
 	{
