Fix for crash in UAnimMontage::TickAssetPlayer
#rnx

diff --git a/Engine/Source/Runtime/Engine/Private/Animation/AnimMontage.cpp b/Engine/Source/Runtime/Engine/Private/Animation/AnimMontage.cpp
--- a/Engine/Source/Runtime/Engine/Private/Animation/AnimMontage.cpp
+++ b/Engine/Source/Runtime/Engine/Private/Animation/AnimMontage.cpp
@@ -2308,7 +2308,7 @@ void FAnimMontageInstance::Advance(float DeltaTime, struct FRootMotionMovementPa
 		{
 			const bool bExtractRootMotion = (OutRootMotionParams != nullptr) && Montage->HasRootMotion();
 			
-			DeltaTimeRecord.Set(0.f, Position);
+			DeltaTimeRecord.Set(Position, 0.f);
 
 			bDidUseMarkerSyncThisTick = CanUseMarkerSync();
 			if (bDidUseMarkerSyncThisTick)
