Fix for UAnimInstance::GetInstanceTransitionTimeElapsed returning wrong values during multiple transitions
#rnx

diff --git a/Engine/Source/Runtime/Engine/Private/Animation/AnimInstance.cpp b/Engine/Source/Runtime/Engine/Private/Animation/AnimInstance.cpp
--- a/Engine/Source/Runtime/Engine/Private/Animation/AnimInstance.cpp
+++ b/Engine/Source/Runtime/Engine/Private/Animation/AnimInstance.cpp
@@ -3484,8 +3484,7 @@ float UAnimInstance::GetInstanceTransitionCrossfadeDuration(int32 MachineIndex,
 
 float UAnimInstance::GetInstanceTransitionTimeElapsed(int32 MachineIndex, int32 TransitionIndex)
 {
-	// Just an alias for readability in the anim graph
-	return GetInstanceCurrentStateElapsedTime(MachineIndex);
+	return GetProxyOnAnyThread<FAnimInstanceProxy>().GetInstanceTransitionTimeElapsed(MachineIndex, TransitionIndex);
 }
 
 float UAnimInstance::GetInstanceTransitionTimeElapsedFraction(int32 MachineIndex, int32 TransitionIndex)
diff --git a/Engine/Source/Runtime/Engine/Private/Animation/AnimInstanceProxy.cpp b/Engine/Source/Runtime/Engine/Private/Animation/AnimInstanceProxy.cpp
--- a/Engine/Source/Runtime/Engine/Private/Animation/AnimInstanceProxy.cpp
+++ b/Engine/Source/Runtime/Engine/Private/Animation/AnimInstanceProxy.cpp
@@ -2174,7 +2174,6 @@ float FAnimInstanceProxy::GetInstanceTransitionCrossfadeDuration(int32 MachineIn
 
 float FAnimInstanceProxy::GetInstanceTransitionTimeElapsed(int32 MachineIndex, int32 TransitionIndex) const
 {
-	// Just an alias for readability in the anim graph
 	if(const FAnimNode_StateMachine* MachineInstance = GetStateMachineInstance(MachineIndex))
 	{
 		if(MachineInstance->IsValidTransitionIndex(TransitionIndex))
