Fix merge

diff --git a/Engine/Source/Runtime/Engine/Private/Materials/HLSLMaterialDerivativeAutogen.cpp b/Engine/Source/Runtime/Engine/Private/Materials/HLSLMaterialDerivativeAutogen.cpp
--- a/Engine/Source/Runtime/Engine/Private/Materials/HLSLMaterialDerivativeAutogen.cpp
+++ b/Engine/Source/Runtime/Engine/Private/Materials/HLSLMaterialDerivativeAutogen.cpp
@@ -468,7 +468,7 @@ int32 FMaterialDerivativeAutogen::GenerateExpressionFunc1(FHLSLMaterialTranslato
 	const FOperationType1 OperationType = GetFunc1ReturnType(SrcDerivInfo.DerivType, Op);
 	if (OperationType.ReturnType == EDerivativeType::None)
 	{
-		return Translator.Errorf(TEXT("Invalid input type: %ss"), GetFloatVectorName(SrcDerivInfo.DerivType));
+		return Translator.Errorf(TEXT("Invalid input type: %ss"), GetFloatVectorName(SrcDerivInfo.TypeIndex));
 	}
 	const bool bIsLWC = IsLWCType(OperationType.IntermediateType);
 
@@ -801,7 +801,7 @@ int32 FMaterialDerivativeAutogen::GenerateExpressionFunc2(FHLSLMaterialTranslato
 	const FOperationType2 OperationType = GetFunc2ReturnType(LhsDerivInfo.DerivType, RhsDerivInfo.DerivType, Op);
 	if (OperationType.ReturnType == EDerivativeType::None)
 	{
-		return Translator.Errorf(TEXT("Invalid input types: %s, %s"), GetFloatVectorName(LhsDerivInfo.DerivType), GetFloatVectorName(RhsDerivInfo.DerivType));
+		return Translator.Errorf(TEXT("Invalid input types: %s, %s"), GetFloatVectorName(LhsDerivInfo.TypeIndex), GetFloatVectorName(RhsDerivInfo.TypeIndex));
 	}
 	const bool bIsLWC = IsLWCType(OperationType.LhsIntermediateType) || IsLWCType(OperationType.RhsIntermediateType);
 
