Backout recently added SoftObjectPath remapping support on LevelStreaming as it causes perf regressions (remains supported in World Partition grid level instances which doesn't happen at runtime)

diff --git a/Engine/Source/Runtime/Engine/Private/LevelStreaming.cpp b/Engine/Source/Runtime/Engine/Private/LevelStreaming.cpp
--- a/Engine/Source/Runtime/Engine/Private/LevelStreaming.cpp
+++ b/Engine/Source/Runtime/Engine/Private/LevelStreaming.cpp
@@ -976,29 +976,13 @@ void ULevelStreaming::PrepareLoadedLevel(ULevel* InLevel, UPackage* InLevelPacka
 		else
 		{
 			check(PendingUnloadLevel == nullptr);
-				
-			// FSoftObjectMap fixup support for Instanced Levels
-			FString SourceWorldPath, RemappedWorldPath;
-			const bool bRemap = InLevel->GetTypedOuter<UWorld>()->GetSoftObjectPathMapping(SourceWorldPath, RemappedWorldPath);
+
 #if WITH_EDITOR
 			if (InPIEInstanceID != INDEX_NONE)
 			{
-				InLevel->FixupForPIE(InPIEInstanceID, [=](int32, FSoftObjectPath& Value)
-				{
-					if (bRemap && !Value.IsNull() && Value.GetAssetPathString().Equals(SourceWorldPath, ESearchCase::IgnoreCase))
-					{
-						Value.SetAssetPathName(FName(*RemappedWorldPath));
-					}
-				});
+				InLevel->FixupForPIE(InPIEInstanceID);
 			}
-			else 
 #endif
-            if(bRemap)
-			{
-				FSoftObjectPathFixupArchive FixupArchive(SourceWorldPath, RemappedWorldPath);
-				FixupArchive.Fixup(InLevel);
-			}
-
 			SetLoadedLevel(InLevel);
 			// Broadcast level loaded event to blueprints
 			OnLevelLoaded.Broadcast();
