Include all the shader compilation errors in the same log message as the "Failed to compile Material" message.  This way the shader compiler errors show up as errors/warnings and will show up in build messages.

diff --git a/Engine/Source/Runtime/Engine/Private/Materials/Material.cpp b/Engine/Source/Runtime/Engine/Private/Materials/Material.cpp
--- a/Engine/Source/Runtime/Engine/Private/Materials/Material.cpp
+++ b/Engine/Source/Runtime/Engine/Private/Materials/Material.cpp
@@ -2234,24 +2234,35 @@ void UMaterial::CacheShadersForResources(EShaderPlatform ShaderPlatform, const T
 
 		if (!bSuccess)
 		{
-			if (IsDefaultMaterial())
+			const bool bIsDefaultMaterial = IsDefaultMaterial();
+			FString ErrorString;
+			if (bIsDefaultMaterial)
 			{
-				UE_ASSET_LOG(LogMaterial, Fatal, this,
-					TEXT("Failed to compile Default Material for platform %s!"),
+				ErrorString += FString::Printf(TEXT("Failed to compile Default Material for platform %s!\n"),
+					*LegacyShaderPlatformToShaderFormat(ShaderPlatform).ToString());
+			}
+			else
+			{
+				ErrorString += FString::Printf(TEXT("Failed to compile Material for platform %s, Default Material will be used in game.\n"),
 					*LegacyShaderPlatformToShaderFormat(ShaderPlatform).ToString());
 			}
-
-			UE_ASSET_LOG(LogMaterial, Warning, this, TEXT("Failed to compile Material for platform %s, Default Material will be used in game."), 
-				*LegacyShaderPlatformToShaderFormat(ShaderPlatform).ToString());
 
 #if WITH_EDITOR
 			const TArray<FString>& CompileErrors = CurrentResource->GetCompileErrors();
 			for (int32 ErrorIndex = 0; ErrorIndex < CompileErrors.Num(); ErrorIndex++)
 			{
-				// Always log material errors in an unsuppressed category
-				UE_LOG(LogMaterial, Display, TEXT("	%s"), *CompileErrors[ErrorIndex]);
+				ErrorString += FString::Printf(TEXT("	%s\n"), *CompileErrors[ErrorIndex]);
 			}
 #endif
+			
+			if (bIsDefaultMaterial)
+			{
+				UE_ASSET_LOG(LogMaterial, Fatal, this, TEXT("%s"), *ErrorString);
+			}
+			else
+			{
+				UE_ASSET_LOG(LogMaterial, Warning, this, TEXT("%s"), *ErrorString);
+			}
 		}
 	}
 }
diff --git a/Engine/Source/Runtime/Engine/Private/Materials/MaterialInstance.cpp b/Engine/Source/Runtime/Engine/Private/Materials/MaterialInstance.cpp
--- a/Engine/Source/Runtime/Engine/Private/Materials/MaterialInstance.cpp
+++ b/Engine/Source/Runtime/Engine/Private/Materials/MaterialInstance.cpp
@@ -2084,8 +2084,10 @@ void UMaterialInstance::CacheShadersForResources(EShaderPlatform ShaderPlatform,
 
 		if (!bSuccess)
 		{
-			UE_ASSET_LOG(LogMaterial, Warning, this,
-				TEXT("Failed to compile Material Instance with Base %s for platform %s, Default Material will be used in game."), 
+			FString ErrorString;
+
+			ErrorString += FString::Printf(
+				TEXT("Failed to compile Material Instance with Base %s for platform %s, Default Material will be used in game.\n"), 
 				BaseMaterial ? *BaseMaterial->GetName() : TEXT("Null"), 
 				*LegacyShaderPlatformToShaderFormat(ShaderPlatform).ToString()
 				);
@@ -2094,9 +2096,11 @@ void UMaterialInstance::CacheShadersForResources(EShaderPlatform ShaderPlatform,
 			const TArray<FString>& CompileErrors = CurrentResource->GetCompileErrors();
 			for (int32 ErrorIndex = 0; ErrorIndex < CompileErrors.Num(); ErrorIndex++)
 			{
-				UE_LOG(LogMaterial, Display, TEXT("	%s"), *CompileErrors[ErrorIndex]);
+				ErrorString += FString::Printf(TEXT("	%s\n"), *CompileErrors[ErrorIndex]);
 			}
 #endif // WITH_EDITOR
+
+			UE_ASSET_LOG(LogMaterial, Warning, this, TEXT("%s"), *ErrorString);
 		}
 	}
 }
diff --git a/Engine/Source/Runtime/Engine/Private/ShaderCompiler/ShaderCompiler.cpp b/Engine/Source/Runtime/Engine/Private/ShaderCompiler/ShaderCompiler.cpp
--- a/Engine/Source/Runtime/Engine/Private/ShaderCompiler/ShaderCompiler.cpp
+++ b/Engine/Source/Runtime/Engine/Private/ShaderCompiler/ShaderCompiler.cpp
@@ -4358,17 +4358,22 @@ void FShaderCompilingManager::ProcessCompiledShaderMaps(
 
 					if (Material->IsDefaultMaterial())
 					{
+						FString ErrorString;
+
 						// Log the errors unsuppressed before the fatal error, so it's always obvious from the log what the compile error was
 						for (int32 ErrorIndex = 0; ErrorIndex < Errors.Num(); ErrorIndex++)
 						{
-							UE_LOG(LogShaderCompilers, Warning, TEXT("	%s"), *Errors[ErrorIndex]);
+							ErrorString += FString::Printf(TEXT("	%s\n"), *Errors[ErrorIndex]);
 						}
 
+						ErrorString += FString::Printf(TEXT("Failed to compile default material %s!"), *Material->GetBaseMaterialPathName());
 						// Assert if a default material could not be compiled, since there will be nothing for other failed materials to fall back on.
-						UE_LOG(LogShaderCompilers, Fatal, TEXT("Failed to compile default material %s!"), *Material->GetBaseMaterialPathName());
+						UE_LOG(LogShaderCompilers, Fatal, TEXT("%s"), *ErrorString);
 					}
+					
+					FString ErrorString;
 
-					UE_LOG(LogShaderCompilers, Warning, TEXT("Failed to compile Material %s for platform %s, Default Material will be used in game."),
+					ErrorString += FString::Printf(TEXT("Failed to compile Material %s for platform %s, Default Material will be used in game.\n"),
 						*Material->GetDebugName(), *LegacyShaderPlatformToShaderFormat(CompilingShaderMap->GetShaderPlatform()).ToString());
 
 					for (int32 ErrorIndex = 0; ErrorIndex < Errors.Num(); ErrorIndex++)
@@ -4376,8 +4381,10 @@ void FShaderCompilingManager::ProcessCompiledShaderMaps(
 						FString ErrorMessage = Errors[ErrorIndex];
 						// Work around build machine string matching heuristics that will cause a cook to fail
 						ErrorMessage.ReplaceInline(TEXT("error "), TEXT("err0r "), ESearchCase::CaseSensitive);
-						UE_LOG(LogShaderCompilers, Display, TEXT("%s"), *ErrorMessage);
+						ErrorString += FString::Printf(TEXT("%s\n"), *ErrorMessage);
 					}
+
+					UE_LOG(LogShaderCompilers, Warning, TEXT("%s"), *ErrorString);
 				}
 
 				if (bReleaseCompilingId)
