Fixed starta assert when deleting all strata node plugged into FrontMaterial

diff --git a/Engine/Source/Runtime/Engine/Private/ShaderCompiler/ShaderGenerationUtil.cpp b/Engine/Source/Runtime/Engine/Private/ShaderCompiler/ShaderGenerationUtil.cpp
--- a/Engine/Source/Runtime/Engine/Private/ShaderCompiler/ShaderGenerationUtil.cpp
+++ b/Engine/Source/Runtime/Engine/Private/ShaderCompiler/ShaderGenerationUtil.cpp
@@ -1710,7 +1710,7 @@ static void DetermineUsedMaterialSlots(
 	bool bHasTangent = SrcGlobal.GBUFFER_HAS_TANGENT;
 	bool bHasVelocity = Dst.WRITES_VELOCITY_TO_GBUFFER;
 	bool bHasStaticLighting = Dst.GBUFFER_HAS_PRECSHADOWFACTOR || Dst.WRITES_PRECSHADOWFACTOR_TO_GBUFFER;
-	bool bIsStrataMaterial = Mat.PROJECT_STRATA && Mat.MATERIAL_IS_STRATA;
+	bool bIsStrataMaterial = Mat.PROJECT_STRATA; // Similarly to FetchFullGBufferInfo, we do not check for MATERIAL_IS_STRATA as this is decided per project.
 
 	// Strata doesn't use gbuffer, and thus doesn't need CustomData
 	const bool bUseCustomData = !bIsStrataMaterial;
