Avoid recomputing memory for shader cache for each insert leading to long game-thread pause due to contention
Fix memory leak when trimming the shader cache
  - Saves many minutes of cook time when shader compilation is required

diff --git a/Engine/Source/Runtime/Engine/Private/ShaderCompiler/ShaderCompiler.cpp b/Engine/Source/Runtime/Engine/Private/ShaderCompiler/ShaderCompiler.cpp
--- a/Engine/Source/Runtime/Engine/Private/ShaderCompiler/ShaderCompiler.cpp
+++ b/Engine/Source/Runtime/Engine/Private/ShaderCompiler/ShaderCompiler.cpp
@@ -7641,6 +7641,11 @@ uint64 FShaderJobCache::GetCurrentMemoryBudget() const
 	return FMath::Min(AbsoluteLimit, RelativeLimit);
 }
 
+FShaderJobCache::FShaderJobCache()
+{
+	CurrentlyAllocatedMemory = sizeof(*this) + InputHashToOutput.GetAllocatedSize() + Outputs.GetAllocatedSize();
+}
+
 void FShaderJobCache::Add(const FJobInputHash& Hash, const FJobCachedOutput& Contents, int32 InitialHitCount)
 {
 	if (!ShaderCompiler::IsJobCacheEnabled())
@@ -7658,6 +7663,7 @@ void FShaderJobCache::Add(const FJobInputHash& Hash, const FJobCachedOutput& Con
 	FJobOutputHash OutputHash = FBlake3::HashBuffer(Contents.GetData(), Contents.Num());
 
 	// add the record
+	const uint64 InputHashToOutputOriginalSize = InputHashToOutput.GetAllocatedSize();
 	InputHashToOutput.Add(Hash, OutputHash);
 
 	FStoredOutput** CannedOutput = Outputs.Find(OutputHash);
@@ -7668,79 +7674,78 @@ void FShaderJobCache::Add(const FJobInputHash& Hash, const FJobCachedOutput& Con
 	}
 	else
 	{
+		const uint64 OutputsOriginalSize = Outputs.GetAllocatedSize();
+
+		FStoredOutput* NewStoredOutput = new FStoredOutput();
+		NewStoredOutput->NumHits = InitialHitCount;
+		NewStoredOutput->NumReferences = 1;
+		NewStoredOutput->JobOutput = Contents;
+		Outputs.Add(OutputHash, NewStoredOutput);
+
+		CurrentlyAllocatedMemory += NewStoredOutput->JobOutput.GetAllocatedSize() + sizeof(FStoredOutput);
+
 		// delete the previous cache entries if we have a budget
-		uint64 MemoryBudgetBytes = GetCurrentMemoryBudget();
+		const uint64 MemoryBudgetBytes = GetCurrentMemoryBudget();
 		if (MemoryBudgetBytes)
 		{
-			uint64 MemoryThatWillBeUsed = GetAllocatedMemory() + Contents.Num();
-			while (MemoryThatWillBeUsed >= MemoryBudgetBytes)
+			if (CurrentlyAllocatedMemory > MemoryBudgetBytes)
 			{
-				// heuristics: delete the entry that has the smallest hits. Don't account for references as if something is referenced often but not hit, it's of no value for us.
-				// (consider other heuristics: hits * memory, time it took to produce the output, last hit time)
-				int32 MinHits = INT_MAX;
-
-				// find the (new) min
-				for (TMap<FJobOutputHash, FStoredOutput*>::TConstIterator Iter(Outputs); Iter; ++Iter)
+				TRACE_CPUPROFILER_EVENT_SCOPE(FShaderJobCache::Trim);
+			
+				while (CurrentlyAllocatedMemory > MemoryBudgetBytes)
 				{
-					MinHits = FMath::Min(Iter.Value()->NumHits, MinHits);
-				}
+					// heuristics: delete the entry that has the smallest hits. Don't account for references as if something is referenced often but not hit, it's of no value for us.
+					// (consider other heuristics: hits * memory, time it took to produce the output, last hit time)
+					int32 MinHits = INT_MAX;
 
-				// remove all matching this minimum until there's enough memory
-				for (TMap<FJobOutputHash, FStoredOutput*>::TIterator Iter(Outputs); Iter; ++Iter)
-				{
-					if (Iter.Value()->NumHits == MinHits)
+					// find the (new) min
+					for (TMap<FJobOutputHash, FStoredOutput*>::TConstIterator Iter(Outputs); Iter; ++Iter)
 					{
-						MemoryThatWillBeUsed -= static_cast<uint64>(Iter.Value()->JobOutput.Num());
-
-						FJobOutputHash RemovedOutputHash = Iter.Key();
-						Iter.RemoveCurrent();
+						MinHits = FMath::Min(Iter.Value()->NumHits, MinHits);
+					}
 
-						// remove all mappings
-						for (TMap<FJobInputHash, FJobOutputHash>::TIterator IterInputs(InputHashToOutput); IterInputs; ++IterInputs)
+					// remove all matching this minimum until there's enough memory
+					for (TMap<FJobOutputHash, FStoredOutput*>::TIterator Iter(Outputs); Iter; ++Iter)
+					{
+						if (Iter.Value()->NumHits == MinHits)
 						{
-							if (IterInputs.Value() == RemovedOutputHash)
+							CurrentlyAllocatedMemory -= static_cast<uint64>(Iter.Value()->JobOutput.GetAllocatedSize() + sizeof(FStoredOutput));
+							delete Iter.Value();
+
+							FJobOutputHash RemovedOutputHash = Iter.Key();
+							Iter.RemoveCurrent();
+
+							// remove all mappings
+							for (TMap<FJobInputHash, FJobOutputHash>::TIterator IterInputs(InputHashToOutput); IterInputs; ++IterInputs)
 							{
-								IterInputs.RemoveCurrent();
+								if (IterInputs.Value() == RemovedOutputHash)
+								{
+									IterInputs.RemoveCurrent();
+								}
 							}
-						}
 
-						// don't remove too much
-						if (MemoryThatWillBeUsed < MemoryBudgetBytes)
-						{
-							break;
+							// don't remove too much
+							if (CurrentlyAllocatedMemory < MemoryBudgetBytes)
+							{
+								break;
+							}
 						}
 					}
 				}
 			}
 		}
 
-		FStoredOutput* NewStoredOutput = new FStoredOutput();
-		NewStoredOutput->NumHits = InitialHitCount;
-		NewStoredOutput->NumReferences = 1;
-		NewStoredOutput->JobOutput = Contents;
-		Outputs.Add(OutputHash, NewStoredOutput);
-
-		// invalidate currently allocated memory only if we added something substantial. We ignore memory increase due to TMap size
-		CurrentlyAllocatedMemory = 0;
+		CurrentlyAllocatedMemory += Outputs.GetAllocatedSize();
+		CurrentlyAllocatedMemory -= OutputsOriginalSize;
 	}
+
+	CurrentlyAllocatedMemory += InputHashToOutput.GetAllocatedSize();
+	CurrentlyAllocatedMemory -= InputHashToOutputOriginalSize;
 }
 
 /** Calculates memory used by the cache*/
 uint64 FShaderJobCache::GetAllocatedMemory()
 {
-	if (!CurrentlyAllocatedMemory)
-	{
-		uint64 MemoryUsed = sizeof(*this) + InputHashToOutput.GetAllocatedSize() + Outputs.GetAllocatedSize();
-
-		// go through all the outputs and sum them
-		for (TMap<FJobOutputHash, FStoredOutput*>::TConstIterator Iter(Outputs); Iter; ++Iter)
-		{
-			MemoryUsed += Iter->Value->JobOutput.GetAllocatedSize();
-		}
-
-		CurrentlyAllocatedMemory = MemoryUsed;
-	}
-
 	return CurrentlyAllocatedMemory;
 }
 
@@ -7755,7 +7760,6 @@ void FShaderJobCache::LogStats()
 	UE_LOG(LogShaderCompilers, Display, TEXT("Tracking %d distinct input hashes that result in %d distinct outputs (%.2f%%)"),
 		InputHashToOutput.Num(), Outputs.Num(), (InputHashToOutput.Num() > 0) ? 100.0 * static_cast<double>(Outputs.Num()) / static_cast<double>(InputHashToOutput.Num()) : 0.0);
 
-	CurrentlyAllocatedMemory = 0;	// get accurate data by invalidating cache
 	uint64 MemUsed = GetAllocatedMemory();
 	double MemUsedMB = FUnitConversion::Convert(static_cast<double>(MemUsed), EUnit::Bytes, EUnit::Megabytes);
 	double MemUsedGB = FUnitConversion::Convert(static_cast<double>(MemUsed), EUnit::Bytes, EUnit::Gigabytes);
diff --git a/Engine/Source/Runtime/Engine/Public/ShaderCompiler.h b/Engine/Source/Runtime/Engine/Public/ShaderCompiler.h
--- a/Engine/Source/Runtime/Engine/Public/ShaderCompiler.h
+++ b/Engine/Source/Runtime/Engine/Public/ShaderCompiler.h
@@ -334,6 +334,8 @@ inline void FShaderCommonCompileJob::Destroy() const
 class FShaderJobCache
 {
 public:
+	FShaderJobCache();
+
 	using FJobInputHash = FShaderCommonCompileJob::FInputHash;
 	using FJobCachedOutput = TArray<uint8>;
 
