Added DistanceTo to the ConvexVolume

#rnx

diff --git a/Engine/Source/Runtime/Engine/Private/ConvexVolume.cpp b/Engine/Source/Runtime/Engine/Private/ConvexVolume.cpp
--- a/Engine/Source/Runtime/Engine/Private/ConvexVolume.cpp
+++ b/Engine/Source/Runtime/Engine/Private/ConvexVolume.cpp
@@ -154,7 +154,7 @@ FOutcode FConvexVolume::GetBoxIntersectionOutcode(const FVector& Origin,const FV
 }
 
 //
-//	FConvexVolume::IntersectBox
+//	FConvexVolume::IntersectBoxWithPermutedPlanes
 //
 
 static FORCEINLINE bool IntersectBoxWithPermutedPlanes(
@@ -210,6 +210,10 @@ static FORCEINLINE bool IntersectBoxWithPermutedPlanes(
 	return true;
 }
 
+//
+//	FConvexVolume::IntersectBox
+//
+
 bool FConvexVolume::IntersectBox(const FVector& Origin,const FVector& Extent) const
 {
 	// Load the origin & extent
@@ -218,13 +222,10 @@ bool FConvexVolume::IntersectBox(const FVector& Origin,const FVector& Extent) co
 	return IntersectBoxWithPermutedPlanes( PermutedPlanes, Orig, Ext );
 }
 
-/**
- * Intersection test with a translated axis-aligned box.
- * @param Origin - Origin of the box.
- * @param Translation - Translation to apply to the box.
- * @param Extent - Extent of the box along each axis.
- * @returns true if this convex volume intersects the given translated box.
- */
+//
+//	FConvexVolume::IntersectBox
+//
+
 bool FConvexVolume::IntersectBox( const FVector& Origin,const FVector& Translation,const FVector& Extent ) const
 {
 	const VectorRegister Orig = VectorLoadFloat3( &Origin );
@@ -234,6 +235,9 @@ bool FConvexVolume::IntersectBox( const FVector& Origin,const FVector& Translati
 	return IntersectBoxWithPermutedPlanes( PermutedPlanes, BoxOrigin, BoxExtent );
 }
 
+//
+//	FConvexVolume::IntersectSphere with the addition check of if the sphere is COMPLETELY contained or only partially contained
+//
 
 bool FConvexVolume::IntersectBox(const FVector& Origin,const FVector& Extent, bool& bOutFullyContained) const
 {
@@ -401,6 +405,11 @@ bool FConvexVolume::IntersectSphere(const FVector& Origin,const float& Radius, b
 	return true;
 }
 
+
+//
+//	FConvexVolume::IntersectLineSegment
+//
+
 bool FConvexVolume::IntersectLineSegment(const FVector& InStart, const FVector& InEnd) const
 {
 	// @todo: not optimized
@@ -446,6 +455,59 @@ bool FConvexVolume::IntersectLineSegment(const FVector& InStart, const FVector&
 	return true;
 }
 
+
+//
+//	FConvexVolume::DistanceTo
+//
+
+float FConvexVolume::DistanceTo(const FVector& Point) const
+{
+	checkSlow(PermutedPlanes.Num() % 4 == 0);
+
+	constexpr VectorRegister4Float VMinimumDistance = MakeVectorRegisterFloatConstant(-BIG_NUMBER, -BIG_NUMBER, -BIG_NUMBER, -BIG_NUMBER);
+
+	// Load the origin & radius
+	VectorRegister VPoint = VectorLoadFloat3(&Point);
+	VectorRegister VMinDistance = VMinimumDistance;
+	// Splat point into 3 vectors
+	VectorRegister VPointX = VectorReplicate(VPoint, 0);
+	VectorRegister VPointY = VectorReplicate(VPoint, 1);
+	VectorRegister VPointZ = VectorReplicate(VPoint, 2);
+	// Since we are moving straight through get a pointer to the data
+	const FPlane* RESTRICT PermutedPlanePtr = (FPlane*)PermutedPlanes.GetData();
+	// Process four planes at a time until we have < 4 left
+	for (int32 Count = 0; Count < PermutedPlanes.Num(); Count += 4)
+	{
+		// Load 4 planes that are already all Xs, Ys, ...
+		VectorRegister PlanesX = VectorLoadAligned(PermutedPlanePtr);
+		PermutedPlanePtr++;
+		VectorRegister PlanesY = VectorLoadAligned(PermutedPlanePtr);
+		PermutedPlanePtr++;
+		VectorRegister PlanesZ = VectorLoadAligned(PermutedPlanePtr);
+		PermutedPlanePtr++;
+		VectorRegister PlanesW = VectorLoadAligned(PermutedPlanePtr);
+		PermutedPlanePtr++;
+		// Calculate the distance (x * x) + (y * y) + (z * z) - w
+		VectorRegister DistX = VectorMultiply(VPointX, PlanesX);
+		VectorRegister DistY = VectorMultiplyAdd(VPointY, PlanesY, DistX);
+		VectorRegister DistZ = VectorMultiplyAdd(VPointZ, PlanesZ, DistY);
+		VectorRegister Distance = VectorSubtract(DistZ, PlanesW);
+
+		VMinDistance = VectorMax(Distance, VMinDistance);
+	}
+
+	const VectorRegister VMinDistanceWXYZ = VectorSwizzle(VMinDistance, 3, 0, 1, 2);
+	const VectorRegister t0 = VectorMax(VMinDistance, VMinDistanceWXYZ);
+	const VectorRegister VMinDistanceZWXY = VectorSwizzle(VMinDistance, 2, 3, 0, 1);
+	const VectorRegister t1 = VectorMax(t0, VMinDistanceZWXY);
+	const VectorRegister VMinDistanceYZWX = VectorSwizzle(VMinDistance, 1, 2, 3, 0);
+	const VectorRegister t2 = VectorMax(t1, VMinDistanceYZWX);
+
+	float MinDistance;
+	VectorStoreFloat1(t2, &MinDistance);
+	return MinDistance;
+}
+
 void GetViewFrustumBoundsInternal(FConvexVolume& OutResult, const FMatrix& ViewProjectionMatrix, bool bUseNearPlane, bool bUseFarPlane, const FPlane* InFarPlaneOverride)
 {
 	OutResult.Planes.Empty(6);
diff --git a/Engine/Source/Runtime/Engine/Public/ConvexVolume.h b/Engine/Source/Runtime/Engine/Public/ConvexVolume.h
--- a/Engine/Source/Runtime/Engine/Public/ConvexVolume.h
+++ b/Engine/Source/Runtime/Engine/Public/ConvexVolume.h
@@ -87,12 +87,59 @@ public:
 
 	FOutcode GetBoxIntersectionOutcode(const FVector& Origin,const FVector& Extent) const;
 
+    /**
+     * Intersection test with a translated axis-aligned box.
+     * @param Origin of the box.
+     * @param Translation -to apply to the box.
+     * @param Extent of the box along each axis.
+     * @returns true if this convex volume intersects the given translated box.
+     */
 	bool IntersectBox(const FVector& Origin,const FVector& Extent) const;
+
+    /**
+     * Intersection test with a translated axis-aligned box.
+     * @param Origin of the box.
+     * @param Translation -to apply to the box.
+     * @param Extent of the box along each axis.
+	 * param bOutFullyContained to know if the box was fully contained 
+     * @returns true if this convex volume intersects the given translated box.
+     */
 	bool IntersectBox(const FVector& Origin,const FVector& Extent, bool& bOutFullyContained) const;
+
+	/**
+     * Intersection test with a sphere
+     * @param Origin of the sphere.
+     * @param Radius of the sphere.
+     * @returns true if this convex volume intersects the given sphere (the result is conservative at the corners)
+     */
 	bool IntersectSphere(const FVector& Origin,const float& Radius) const;
+
+	/**
+     * Intersection test with a sphere
+     * @param Origin of the sphere.
+     * @param Radius of the sphere.
+	 * param bOutFullyContained to know if the sphere was fully contained 
+     * @returns true if this convex volume intersects the given sphere (the result is conservative at the corners)
+     */
 	bool IntersectSphere(const FVector& Origin,const float& Radius, bool& bOutFullyContained) const;
+
+	/**
+     * Intersection test with line segment
+     * @param Start of the segment.
+     * @param End of the segment.
+	 * param bOutFullyContained to know if the sphere was fully contained 
+     * @returns true if this convex volume intersects the given line segment
+     */
 	bool IntersectLineSegment(const FVector& Start, const FVector& End) const;
 
+	/**
+	 * Calculates the maximum perpendicular distance of a point to the plains of the convex volume.
+	 * The distance can be used to see if the sphere is touching the volume
+	 * @param 	Point to calculate the distance to.
+	 * @return 	Returns the maximum perpendicular distance to then plains (the distance is conservative at the corners)
+	 */
+	float DistanceTo(const FVector& Point) const;
+
 	/**
 	 * Intersection test with a translated axis-aligned box.
 	 * @param Origin - Origin of the box.
@@ -102,7 +149,11 @@ public:
 	 */
 	bool IntersectBox(const FVector& Origin,const FVector& Translation,const FVector& Extent) const;
 
-	/** Determines whether the given point lies inside the convex volume */
+	/** 
+	 * Determines whether the given point lies inside the convex volume
+	 * @param Point to test against.
+	 * @returns true if the point is inside the convex volume.
+	 */
 	bool IntersectPoint(const FVector& Point) const
 	{
 		return IntersectSphere(Point, 0.0f);
