Fixed crash on GetEditorBounds when the editor hash is null.

diff --git a/Engine/Source/Runtime/Engine/Private/WorldPartition/WorldPartition.cpp b/Engine/Source/Runtime/Engine/Private/WorldPartition/WorldPartition.cpp
--- a/Engine/Source/Runtime/Engine/Private/WorldPartition/WorldPartition.cpp
+++ b/Engine/Source/Runtime/Engine/Private/WorldPartition/WorldPartition.cpp
@@ -1354,7 +1354,9 @@ FBox UWorldPartition::GetWorldBounds() const
 
 FBox UWorldPartition::GetEditorWorldBounds() const
 {
-	return EditorHash->GetEditorWorldBounds();
+	// This is not a valid call in PIE
+	check(EditorHash);
+	return EditorHash ? EditorHash->GetEditorWorldBounds() : FBox(EForceInit::ForceInit);
 }
 #endif
 
