Add missing return causing wrong message to be returned

diff --git a/Engine/Source/Runtime/Engine/Private/LevelInstance/LevelInstanceSubsystem.cpp b/Engine/Source/Runtime/Engine/Private/LevelInstance/LevelInstanceSubsystem.cpp
--- a/Engine/Source/Runtime/Engine/Private/LevelInstance/LevelInstanceSubsystem.cpp
+++ b/Engine/Source/Runtime/Engine/Private/LevelInstance/LevelInstanceSubsystem.cpp
@@ -1506,6 +1506,7 @@ bool ULevelInstanceSubsystem::CanEditLevelInstance(const ALevelInstance* LevelIn
 			{
 				*OutReason = LOCTEXT("CanEditLevelInstanceAlreadyBeingEdited", "Level Instance already being edited");
 			}
+			return false;
 		}
 
 		if (IsLevelInstanceEditDirty(LevelInstanceEdit.Get()))
