Chaos Cloth - Fix crash when rendering with an invalid sim LOD.

diff --git a/Engine/Plugins/Experimental/ChaosCloth/Source/Chaos/Private/ChaosCloth/ChaosClothingSimulation.cpp b/Engine/Plugins/Experimental/ChaosCloth/Source/Chaos/Private/ChaosCloth/ChaosClothingSimulation.cpp
--- a/Engine/Plugins/Experimental/ChaosCloth/Source/Chaos/Private/ChaosCloth/ChaosClothingSimulation.cpp
+++ b/Engine/Plugins/Experimental/ChaosCloth/Source/Chaos/Private/ChaosCloth/ChaosClothingSimulation.cpp
@@ -616,10 +616,10 @@ void FClothingSimulation::GetSimulationData(
 	for (const TUniquePtr<FClothingSimulationCloth>& Cloth : Cloths)
 	{
 		const int32 AssetIndex = Cloth->GetGroupId();
-		FClothSimulData& Data = OutData.FindOrAdd(AssetIndex);
 
 		if (!Cloth->GetMesh())
 		{
+			OutData.Remove(AssetIndex);  // Ensures that the cloth vertex factory won't run unnecessarily
 			continue;  // Invalid or empty cloth
 		}
 
@@ -633,6 +633,7 @@ void FClothingSimulation::GetSimulationData(
 
 		if (Cloth->GetOffset(Solver.Get()) == INDEX_NONE || Cloth->GetLODIndex(Solver.Get()) == INDEX_NONE)
 		{
+			OutData.Remove(AssetIndex);  // Ensures that the cloth vertex factory won't run unnecessarily
 			continue;  // No valid LOD, there's nothing to write out
 		}
 
@@ -652,6 +653,7 @@ void FClothingSimulation::GetSimulationData(
 		ReferenceBoneTransform.SetScale3D(FVector(1.0f));  // Scale is already baked in the cloth mesh
 
 		// Set the world space transform to be this cloth's reference bone
+		FClothSimulData& Data = OutData.FindOrAdd(AssetIndex);
 		Data.Transform = ReferenceBoneTransform;
 		Data.ComponentRelativeTransform = ReferenceBoneTransform.GetRelativeTransform(OwnerTransform);
 
diff --git a/Engine/Source/Runtime/Engine/Private/SkeletalRenderGPUSkin.cpp b/Engine/Source/Runtime/Engine/Private/SkeletalRenderGPUSkin.cpp
--- a/Engine/Source/Runtime/Engine/Private/SkeletalRenderGPUSkin.cpp
+++ b/Engine/Source/Runtime/Engine/Private/SkeletalRenderGPUSkin.cpp
@@ -673,7 +673,8 @@ void FSkeletalMeshObjectGPUSkin::ProcessUpdatedDynamicData(EGPUSkinCacheEntryMod
 		for (int32 SectionIdx = 0; SectionIdx < Sections.Num(); SectionIdx++)
 		{
 			const FSkelMeshRenderSection& Section = Sections[SectionIdx];
-			const bool bClothFactory = bClothEnabled && (DynamicData->ClothingSimData.Num() > 0) && Section.HasClothingData();
+			const FClothSimulData* const SimData = DynamicData->ClothingSimData.Find(Section.CorrespondClothAssetIndex);
+			const bool bClothFactory = bClothEnabled && SimData && Section.HasClothingData();
 
 			FGPUBaseSkinVertexFactory* VertexFactory;
 			{
@@ -743,12 +744,9 @@ void FSkeletalMeshObjectGPUSkin::ProcessUpdatedDynamicData(EGPUSkinCacheEntryMod
 			{				
 				FGPUBaseSkinAPEXClothVertexFactory::ClothShaderType& ClothShaderData = VertexFactoryData.ClothVertexFactories[SectionIdx]->GetClothShaderData();
 				ClothShaderData.ClothBlendWeight = DynamicData->ClothBlendWeight;
-				int16 ActorIdx = Section.CorrespondClothAssetIndex;
-				if(FClothSimulData* SimData = DynamicData->ClothingSimData.Find(ActorIdx))
-				{
-					bNeedFence = ClothShaderData.UpdateClothSimulData(RHICmdList, DynamicData->ClothingSimData[ActorIdx].Positions, DynamicData->ClothingSimData[ActorIdx].Normals, FrameNumberToPrepare, FeatureLevel) || bNeedFence;
-					ClothShaderData.GetClothLocalToWorldForWriting(FrameNumberToPrepare) = DynamicData->ClothingSimData[ActorIdx].ComponentRelativeTransform.ToMatrixWithScale() * DynamicData->ClothObjectLocalToWorld;
-				}
+
+				bNeedFence = ClothShaderData.UpdateClothSimulData(RHICmdList, SimData->Positions, SimData->Normals, FrameNumberToPrepare, FeatureLevel) || bNeedFence;
+				ClothShaderData.GetClothLocalToWorldForWriting(FrameNumberToPrepare) = SimData->ComponentRelativeTransform.ToMatrixWithScale() * DynamicData->ClothObjectLocalToWorld;
 			}
 #endif // WITH_APEX_CLOTHING || WITH_CHAOS_CLOTHING
 
