World Partition: added dump when generating streaming.

diff --git a/Engine/Source/Runtime/Engine/Private/WorldPartition/ErrorHandling/WorldPartitionStreamingGenerationErrorHandler.cpp b/Engine/Source/Runtime/Engine/Private/WorldPartition/ErrorHandling/WorldPartitionStreamingGenerationErrorHandler.cpp
deleted file mode 100644
diff --git a/Engine/Source/Runtime/Engine/Private/WorldPartition/ErrorHandling/WorldPartitionStreamingGenerationLogErrorHandler.cpp b/Engine/Source/Runtime/Engine/Private/WorldPartition/ErrorHandling/WorldPartitionStreamingGenerationLogErrorHandler.cpp
--- a/Engine/Source/Runtime/Engine/Private/WorldPartition/ErrorHandling/WorldPartitionStreamingGenerationLogErrorHandler.cpp
+++ b/Engine/Source/Runtime/Engine/Private/WorldPartition/ErrorHandling/WorldPartitionStreamingGenerationLogErrorHandler.cpp
@@ -7,7 +7,7 @@
 
 void FStreamingGenerationLogErrorHandler::OnInvalidReference(const FWorldPartitionActorDescView& ActorDescView, const FGuid& ReferenceGuid)
 {
-	UE_LOG(LogWorldPartition, Log, TEXT("Actor %s have missing reference to %s"), *GetActorLabel(ActorDescView), *ReferenceGuid.ToString());
+	UE_LOG(LogWorldPartition, Log, TEXT("Actor %s have missing reference to %s"), *ActorDescView.GetActorLabelOrName().ToString(), *ReferenceGuid.ToString());
 }
 
 void FStreamingGenerationLogErrorHandler::OnInvalidReferenceGridPlacement(const FWorldPartitionActorDescView& ActorDescView, const FWorldPartitionActorDescView& ReferenceActorDescView)
@@ -16,21 +16,21 @@ void FStreamingGenerationLogErrorHandler::OnInvalidReferenceGridPlacement(const
 	const FString AlwaysLoadedActor(TEXT("Always loaded actor"));
 	const bool bIsActorDescAlwaysLoaded = ActorDescView.GetGridPlacement() == EActorGridPlacement::AlwaysLoaded;
 	const bool bIsActorDescRefAlwaysLoaded = ReferenceActorDescView.GetGridPlacement() == EActorGridPlacement::AlwaysLoaded;
-	UE_LOG(LogWorldPartition, Log, TEXT("%s %s reference %s %s"), bIsActorDescAlwaysLoaded ? *AlwaysLoadedActor : *StreamedActor, *GetActorLabel(ActorDescView), bIsActorDescRefAlwaysLoaded ? *AlwaysLoadedActor : *StreamedActor, *GetActorLabel(ReferenceActorDescView));
+	UE_LOG(LogWorldPartition, Log, TEXT("%s %s reference %s %s"), bIsActorDescAlwaysLoaded ? *AlwaysLoadedActor : *StreamedActor, *ActorDescView.GetActorLabelOrName().ToString(), bIsActorDescRefAlwaysLoaded ? *AlwaysLoadedActor : *StreamedActor, *ReferenceActorDescView.GetActorLabelOrName().ToString());
 }
 
 void FStreamingGenerationLogErrorHandler::OnInvalidReferenceDataLayers(const FWorldPartitionActorDescView& ActorDescView, const FWorldPartitionActorDescView& ReferenceActorDescView)
 {
-	UE_LOG(LogWorldPartition, Log, TEXT("Actor %s references an actor in a different set of data layers %s"), *GetActorLabel(ActorDescView), *GetActorLabel(ReferenceActorDescView));
+	UE_LOG(LogWorldPartition, Log, TEXT("Actor %s references an actor in a different set of data layers %s"), *ActorDescView.GetActorLabelOrName().ToString(), *ReferenceActorDescView.GetActorLabelOrName().ToString());
 }
 
 void FStreamingGenerationLogErrorHandler::OnInvalidReferenceLevelScriptStreamed(const FWorldPartitionActorDescView& ActorDescView)
 {
-	UE_LOG(LogWorldPartition, Log, TEXT("Level Script Blueprint references streamed actor %s"), *GetActorLabel(ActorDescView));
+	UE_LOG(LogWorldPartition, Log, TEXT("Level Script Blueprint references streamed actor %s"), *ActorDescView.GetActorLabelOrName().ToString());
 }
 
 void FStreamingGenerationLogErrorHandler::OnInvalidReferenceLevelScriptDataLayers(const FWorldPartitionActorDescView& ActorDescView)
 {
-	UE_LOG(LogWorldPartition, Log, TEXT("Level Script Blueprint references streamed actor %s with a non empty set of data layers"), *GetActorLabel(ActorDescView));
+	UE_LOG(LogWorldPartition, Log, TEXT("Level Script Blueprint references streamed actor %s with a non empty set of data layers"), *ActorDescView.GetActorLabelOrName().ToString());
 }
 #endif
diff --git a/Engine/Source/Runtime/Engine/Private/WorldPartition/ErrorHandling/WorldPartitionStreamingGenerationMapCheckErrorHandler.cpp b/Engine/Source/Runtime/Engine/Private/WorldPartition/ErrorHandling/WorldPartitionStreamingGenerationMapCheckErrorHandler.cpp
--- a/Engine/Source/Runtime/Engine/Private/WorldPartition/ErrorHandling/WorldPartitionStreamingGenerationMapCheckErrorHandler.cpp
+++ b/Engine/Source/Runtime/Engine/Private/WorldPartition/ErrorHandling/WorldPartitionStreamingGenerationMapCheckErrorHandler.cpp
@@ -13,7 +13,7 @@ void FStreamingGenerationMapCheckErrorHandler::OnInvalidReference(const FWorldPa
 {
 	TSharedRef<FTokenizedMessage> Error = FMessageLog("MapCheck").Warning()
 		->AddToken(FTextToken::Create(LOCTEXT("MapCheck_WorldPartition_Actor", "Actor")))
-		->AddToken(FAssetNameToken::Create(GetActorLabel(ActorDescView)))
+		->AddToken(FAssetNameToken::Create(ActorDescView.GetActorLabelOrName().ToString()))
 		->AddToken(FTextToken::Create(LOCTEXT("MapCheck_WorldPartition_HaveMissingRefsTo", "have missing references to")))
 		->AddToken(FTextToken::Create(FText::FromString(ReferenceGuid.ToString())))
 		->AddToken(FMapErrorToken::Create(FName(TEXT("WorldPartition_MissingActorReference_CheckForErrors"))));
@@ -29,10 +29,10 @@ void FStreamingGenerationMapCheckErrorHandler::OnInvalidReferenceGridPlacement(c
 
 	TSharedRef<FTokenizedMessage> Error = FMessageLog("MapCheck").Error()
 		->AddToken(FTextToken::Create(bIsActorDescAlwaysLoaded ? AlwaysLoadedActor : StreamedActor))
-		->AddToken(FAssetNameToken::Create(GetActorLabel(ActorDescView)))
+		->AddToken(FAssetNameToken::Create(ActorDescView.GetActorLabelOrName().ToString()))
 		->AddToken(FTextToken::Create(LOCTEXT("MapCheck_WorldPartition_References", "references")))
 		->AddToken(FTextToken::Create(bIsActorDescRefAlwaysLoaded ? AlwaysLoadedActor : StreamedActor)								)
-		->AddToken(FAssetNameToken::Create(GetActorLabel(ReferenceActorDescView)))
+		->AddToken(FAssetNameToken::Create(ReferenceActorDescView.GetActorLabelOrName().ToString()))
 		->AddToken(FMapErrorToken::Create(FName(TEXT("WorldPartition_StreamedActorReferenceAlwaysLoadedActor_CheckForErrors"))));
 }
 
@@ -40,9 +40,9 @@ void FStreamingGenerationMapCheckErrorHandler::OnInvalidReferenceDataLayers(cons
 {
 	TSharedRef<FTokenizedMessage> Error = FMessageLog("MapCheck").Error()
 		->AddToken(FTextToken::Create(LOCTEXT("MapCheck_WorldPartition_Actor", "Actor")))
-		->AddToken(FAssetNameToken::Create(GetActorLabel(ActorDescView)))
+		->AddToken(FAssetNameToken::Create(ActorDescView.GetActorLabelOrName().ToString()))
 		->AddToken(FTextToken::Create(LOCTEXT("MapCheck_WorldPartition_ReferenceActorInOtherDataLayers", "references an actor in a different set of data layers")))
-		->AddToken(FAssetNameToken::Create(GetActorLabel(ReferenceActorDescView)))
+		->AddToken(FAssetNameToken::Create(ReferenceActorDescView.GetActorLabelOrName().ToString()))
 		->AddToken(FMapErrorToken::Create(FName(TEXT("WorldPartition_ActorReferenceActorInAnotherDataLayer_CheckForErrors"))));
 }
 
@@ -50,7 +50,7 @@ void FStreamingGenerationMapCheckErrorHandler::OnInvalidReferenceLevelScriptStre
 {
 	TSharedRef<FTokenizedMessage> Error = FMessageLog("MapCheck").Error()
 		->AddToken(FTextToken::Create(LOCTEXT("MapCheck_WorldPartition_LevelScriptBlueprintStreamedActorReference", "Level Script Blueprint references streamed actor")))
-		->AddToken(FAssetNameToken::Create(GetActorLabel(ActorDescView)))
+		->AddToken(FAssetNameToken::Create(ActorDescView.GetActorLabelOrName().ToString()))
 		->AddToken(FMapErrorToken::Create(FName(TEXT("WorldPartition_LevelScriptBlueprintRefefenceStreamed_CheckForErrors"))));
 }
 
@@ -58,7 +58,7 @@ void FStreamingGenerationMapCheckErrorHandler::OnInvalidReferenceLevelScriptData
 {
 	TSharedRef<FTokenizedMessage> Error = FMessageLog("MapCheck").Error()
 		->AddToken(FTextToken::Create(LOCTEXT("MapCheck_WorldPartition_LevelScriptBlueprintActorReference", "Level Script Blueprint references actor")))
-		->AddToken(FAssetNameToken::Create(GetActorLabel(ActorDescView)))
+		->AddToken(FAssetNameToken::Create(ActorDescView.GetActorLabelOrName().ToString()))
 		->AddToken(FTextToken::Create(LOCTEXT("MapCheck_WorldPartition_LevelScriptBlueprintDataLayerReference", "with a non empty set of data layers")))
 		->AddToken(FMapErrorToken::Create(FName(TEXT("WorldPartition_LevelScriptBlueprintRefefenceDataLayer_CheckForErrors"))));
 }
diff --git a/Engine/Source/Runtime/Engine/Private/WorldPartition/WorldPartitionActorDesc.cpp b/Engine/Source/Runtime/Engine/Private/WorldPartition/WorldPartitionActorDesc.cpp
--- a/Engine/Source/Runtime/Engine/Private/WorldPartition/WorldPartitionActorDesc.cpp
+++ b/Engine/Source/Runtime/Engine/Private/WorldPartition/WorldPartitionActorDesc.cpp
@@ -180,7 +180,7 @@ void FWorldPartitionActorDesc::TransformInstance(const FString& From, const FStr
 
 FString FWorldPartitionActorDesc::ToString() const
 {
-	return FString::Printf(TEXT("Guid:%s Class:%s Name:%s"), *Guid.ToString(), *Class.ToString(), *FPaths::GetExtension(ActorPath.ToString()));
+	return FString::Printf(TEXT("Guid:%s Class:%s Name:%s SpatiallyLoaded:%s"), *Guid.ToString(), *Class.ToString(), *FPaths::GetExtension(ActorPath.ToString()), (GridPlacement != EActorGridPlacement::AlwaysLoaded) ? TEXT("True") : TEXT("False"));
 }
 
 void FWorldPartitionActorDesc::Serialize(FArchive& Ar)
@@ -246,6 +246,12 @@ FName FWorldPartitionActorDesc::GetActorName() const
 	return *ActorName;
 }
 
+FName FWorldPartitionActorDesc::GetActorLabelOrName() const
+{
+	return GetActorLabel().IsNone() ? GetActorName() : GetActorLabel();
+};
+
+
 TArray<const UDataLayer*> FWorldPartitionActorDesc::GetDataLayerObjects() const
 {
 	if (Container)
diff --git a/Engine/Source/Runtime/Engine/Private/WorldPartition/WorldPartitionActorDescView.cpp b/Engine/Source/Runtime/Engine/Private/WorldPartition/WorldPartitionActorDescView.cpp
--- a/Engine/Source/Runtime/Engine/Private/WorldPartition/WorldPartitionActorDescView.cpp
+++ b/Engine/Source/Runtime/Engine/Private/WorldPartition/WorldPartitionActorDescView.cpp
@@ -99,6 +99,11 @@ const TArray<FGuid>& FWorldPartitionActorDescView::GetReferences() const
 	return ActorDesc->GetReferences();
 }
 
+FString FWorldPartitionActorDescView::ToString() const
+{
+	return ActorDesc->ToString();
+}
+
 uint32 FWorldPartitionActorDescView::GetTag() const
 {
 	return ActorDesc->Tag;
@@ -109,6 +114,11 @@ bool FWorldPartitionActorDescView::GetContainerInstance(const UActorDescContaine
 	return ActorDesc->GetContainerInstance(OutLevelContainer, OutLevelTransform, OutClusterMode);
 }
 
+FName FWorldPartitionActorDescView::GetActorLabelOrName() const
+{
+	return ActorDesc->GetActorLabelOrName();
+}
+
 void FWorldPartitionActorDescView::SetGridPlacement(EActorGridPlacement InGridPlacement)
 {
 	if (GridPlacement != InGridPlacement)
diff --git a/Engine/Source/Runtime/Engine/Private/WorldPartition/WorldPartitionStreamingGeneration.cpp b/Engine/Source/Runtime/Engine/Private/WorldPartition/WorldPartitionStreamingGeneration.cpp
--- a/Engine/Source/Runtime/Engine/Private/WorldPartition/WorldPartitionStreamingGeneration.cpp
+++ b/Engine/Source/Runtime/Engine/Private/WorldPartition/WorldPartitionStreamingGeneration.cpp
@@ -12,6 +12,7 @@
 #include "WorldPartition/ErrorHandling/WorldPartitionStreamingGenerationLogErrorHandler.h"
 #include "WorldPartition/ErrorHandling/WorldPartitionStreamingGenerationMapCheckErrorHandler.h"
 #include "WorldPartition/HLOD/HLODActor.h"
+#include "HAL/FileManager.h"
 
 #define LOCTEXT_NAMESPACE "WorldPartition"
 
@@ -85,7 +86,7 @@ class FWorldPartitionStreamingGenerator
 		}
 	}
 
-	void CreateActorDescriptorViewsRecursive(const UActorDescContainer* InContainer, const FTransform& InTransform, const TSet<FName>& InDataLayers, const FActorContainerID& InContainerID, const FActorContainerID& InParentContainerID, EContainerClusterMode InClusterMode)
+	void CreateActorDescriptorViewsRecursive(const UActorDescContainer* InContainer, const FTransform& InTransform, const TSet<FName>& InDataLayers, const FActorContainerID& InContainerID, const FActorContainerID& InParentContainerID, EContainerClusterMode InClusterMode, const TCHAR* OwnerName)
 	{
 		TMap<FGuid, FWorldPartitionActorDescView> ActorDescViewMap;
 		
@@ -117,7 +118,7 @@ class FWorldPartitionStreamingGenerator
 					SubDataLayers = &ParentDataLayers;
 				}
 
-				CreateActorDescriptorViewsRecursive(SubContainer, SubTransform * InTransform, *SubDataLayers, SubContainerID, InContainerID, SubClusterMode);
+				CreateActorDescriptorViewsRecursive(SubContainer, SubTransform * InTransform, *SubDataLayers, SubContainerID, InContainerID, SubClusterMode, *ActorDescView.GetActorLabelOrName().ToString());
 
 				// The container actor can be removed now that its contained actors has been registered
 				It.RemoveCurrent();
@@ -138,6 +139,7 @@ class FWorldPartitionStreamingGenerator
 		ContainerDescriptor.ClusterMode = InClusterMode;
 		ContainerDescriptor.ActorDescViewMap = MoveTemp(ActorDescViewMap);
 		ContainerDescriptor.DataLayers = InDataLayers;
+		ContainerDescriptor.OwnerName = OwnerName;
 
 		// Maintain containers hierarchy, bottom up
 		if (InContainerID != InParentContainerID)
@@ -151,7 +153,7 @@ class FWorldPartitionStreamingGenerator
 	 */
 	void CreateActorDescriptorViews(const UActorDescContainer* InContainer)
 	{
-		CreateActorDescriptorViewsRecursive(InContainer, FTransform::Identity, TSet<FName>(), FActorContainerID(), FActorContainerID(), EContainerClusterMode::Partitioned);
+		CreateActorDescriptorViewsRecursive(InContainer, FTransform::Identity, TSet<FName>(), FActorContainerID(), FActorContainerID(), EContainerClusterMode::Partitioned, TEXT("MainContainer"));
 	}
 
 	/** 
@@ -287,9 +289,9 @@ class FWorldPartitionStreamingGenerator
 		// Update parent containers bounds, this relies on the fact that ContainersHierarchy is built bottom up
 		for (auto ContainerPairIt = ContainersHierarchy.CreateIterator(); ContainerPairIt; ++ContainerPairIt)
 		{
-			const FContainerDescriptor& CurrentContainerID = ContainerDescriptorsMap.FindChecked(ContainerPairIt.Key());
-			FContainerDescriptor& ParentContainerID = ContainerDescriptorsMap.FindChecked(ContainerPairIt.Value());
-			ParentContainerID.Bounds += CurrentContainerID.Bounds;
+			const FContainerDescriptor& CurrentContainer = ContainerDescriptorsMap.FindChecked(ContainerPairIt.Key());
+			FContainerDescriptor& ParentContainer = ContainerDescriptorsMap.FindChecked(ContainerPairIt.Value());
+			ParentContainer.Bounds += CurrentContainer.Bounds;
 		}
 	}
 
@@ -328,6 +330,85 @@ public:
 		return FActorClusterContext(MoveTemp(ContainerInstances), InFilterActorDescViewFunc);
 	}
 
+	void DumpStateLog()
+	{
+		FString StateLogOutputPath = FPaths::Combine(FPaths::ProjectDir(), TEXT("Saved"), TEXT("WorldPartition"));
+		FString StateLogOutputFilename = FPaths::Combine(StateLogOutputPath, *FString::Printf(TEXT("StreamingGeneration-%s.log"), *FDateTime::Now().ToString()));
+		
+		if (FArchive* LogFile = IFileManager::Get().CreateFileWriter(*StateLogOutputFilename))
+		{
+			auto SerializeLine = [LogFile](const FString& Line)
+			{
+				LogFile->Serialize(TCHAR_TO_ANSI(*Line), Line.Len());
+				LogFile->Serialize(LINE_TERMINATOR_ANSI, sizeof(LINE_TERMINATOR_ANSI));
+			};
+
+			// Build the containers tree representation
+			TMultiMap<FActorContainerID, FActorContainerID> InvertedContainersHierarchy;
+			for (auto ContainerPairIt = ContainersHierarchy.CreateIterator(); ContainerPairIt; ++ContainerPairIt)
+			{
+				const FActorContainerID& ChildContainerID = ContainerPairIt.Key();
+				const FActorContainerID& ParentContainerID = ContainerPairIt.Value();				
+				InvertedContainersHierarchy.Add(ParentContainerID, ChildContainerID);
+			}
+
+			SerializeLine(TEXT("Containers:"));
+
+			auto DumpContainers = [this, &InvertedContainersHierarchy, &SerializeLine](const FActorContainerID& ContainerID)
+			{
+				auto DumpContainersRecursive = [this, &InvertedContainersHierarchy, &SerializeLine](const FActorContainerID& ContainerID, FString Prefix, auto& RecursiveFunc) -> void
+				{
+					const FContainerDescriptor& ContainerDescriptor = ContainerDescriptorsMap.FindChecked(ContainerID);
+					SerializeLine(FString::Printf(TEXT("%s[+] %s:"), *Prefix, *ContainerDescriptor.OwnerName));
+
+					Prefix += TEXT(" | ");
+
+					SerializeLine(FString::Printf(TEXT("%s           ID: 0x%016llx"), *Prefix, ContainerID.ID));
+					SerializeLine(FString::Printf(TEXT("%s       Bounds: %s"), *Prefix, *ContainerDescriptor.Bounds.ToString()));
+					SerializeLine(FString::Printf(TEXT("%s    Transform: %s"), *Prefix, *ContainerDescriptor.Transform.ToString()));
+					SerializeLine(FString::Printf(TEXT("%s    Container: %s"), *Prefix, *ContainerDescriptor.Container->GetContainerPackage().ToString()));
+					SerializeLine(FString::Printf(TEXT("%s  ClusterMode: %s"), *Prefix, *StaticEnum<EContainerClusterMode>()->GetNameStringByValue((uint64)ContainerDescriptor.ClusterMode)));
+
+					if (ContainerDescriptor.ActorDescViewMap.Num())
+					{
+						SerializeLine(FString::Printf(TEXT("%s ActorDescs:"), *Prefix));
+
+						TMap<FGuid, FWorldPartitionActorDescView> SortedActorDescViewMap = ContainerDescriptor.ActorDescViewMap;
+						SortedActorDescViewMap.KeySort([](const FGuid& GuidA, const FGuid& GuidB) { return GuidA < GuidB; });
+
+						for (auto ActorDescIt = SortedActorDescViewMap.CreateConstIterator(); ActorDescIt; ++ActorDescIt)
+						{
+							const FWorldPartitionActorDescView& ActorDescView = ActorDescIt.Value();
+							SerializeLine(FString::Printf(TEXT("%s     - %s"), *Prefix, *ActorDescView.ToString()));
+						}
+					}
+
+					TArray<FActorContainerID> ChildContainersIDs;
+					InvertedContainersHierarchy.MultiFind(ContainerID, ChildContainersIDs);
+					ChildContainersIDs.Sort([](const FActorContainerID& ActorContainerIDA, const FActorContainerID& ActorContainerIDB) { return ActorContainerIDA.ID < ActorContainerIDB.ID; });
+
+					if (ChildContainersIDs.Num())
+					{
+						SerializeLine(FString::Printf(TEXT("%sSubContainers:"), *Prefix));
+						
+						Prefix += TEXT("  ");
+						for (const FActorContainerID& ChildContainerID : ChildContainersIDs)
+						{
+							RecursiveFunc(ChildContainerID, Prefix, RecursiveFunc);
+						}
+					}
+				};
+
+				DumpContainersRecursive(ContainerID, TEXT("  "), DumpContainersRecursive);
+			};
+
+			DumpContainers(FActorContainerID());
+
+			LogFile->Close();
+			delete LogFile;
+		}
+	}
+
 private:
 	UWorldPartitionRuntimeHash* RuntimeHash;
 	FActorDescList* ModifiedActorsDescList;
@@ -341,11 +422,12 @@ private:
 		{}
 
 		FBox Bounds;
-		FTransform Transform;		
+		FTransform Transform;
 		const UActorDescContainer* Container;
 		EContainerClusterMode ClusterMode;
 		TMap<FGuid, FWorldPartitionActorDescView> ActorDescViewMap;
 		TSet<FName> DataLayers;
+		FString OwnerName;
 	};
 
 	/** Maps containers IDs to their container descriptor */
@@ -380,6 +462,7 @@ bool UWorldPartition::GenerateStreaming(TArray<FString>* OutPackagesToGenerate)
 
 		// Preparation Phase
 		StreamingGenerator.PreparationPhase(this);
+		StreamingGenerator.DumpStateLog();
 
 		// Preparation Phase :: Actor Clusters Creation
 		ActorClusterContext = StreamingGenerator.CreateActorClusters();
@@ -425,4 +508,4 @@ void UWorldPartition::CheckForErrors(IStreamingGenerationErrorHandler* ErrorHand
 }
 #endif // WITH_EDITOR
 
-#undef LOCTEXT_NAMESPACE
\ No newline at end of file
+#undef LOCTEXT_NAMESPACE
diff --git a/Engine/Source/Runtime/Engine/Public/WorldPartition/ErrorHandling/WorldPartitionStreamingGenerationErrorHandler.h b/Engine/Source/Runtime/Engine/Public/WorldPartition/ErrorHandling/WorldPartitionStreamingGenerationErrorHandler.h
--- a/Engine/Source/Runtime/Engine/Public/WorldPartition/ErrorHandling/WorldPartitionStreamingGenerationErrorHandler.h
+++ b/Engine/Source/Runtime/Engine/Public/WorldPartition/ErrorHandling/WorldPartitionStreamingGenerationErrorHandler.h
@@ -7,9 +7,6 @@
 
 class ENGINE_API IStreamingGenerationErrorHandler
 {
-protected:
-	FString GetActorLabel(const FWorldPartitionActorDescView& ActorDescView);
-
 public:
 	virtual ~IStreamingGenerationErrorHandler() {}
 
diff --git a/Engine/Source/Runtime/Engine/Public/WorldPartition/WorldPartitionActorDesc.h b/Engine/Source/Runtime/Engine/Public/WorldPartition/WorldPartitionActorDesc.h
--- a/Engine/Source/Runtime/Engine/Public/WorldPartition/WorldPartitionActorDesc.h
+++ b/Engine/Source/Runtime/Engine/Public/WorldPartition/WorldPartitionActorDesc.h
@@ -21,6 +21,7 @@ struct FWorldPartitionActorDescInitData
 
 class UActorDescContainer;
 
+UENUM()
 enum class EContainerClusterMode : uint8
 {
 	Partitioned, // Per Actor Partitioning
@@ -77,6 +78,8 @@ public:
 
 	FName GetActorName() const;
 
+	FName GetActorLabelOrName() const;
+
 	virtual bool GetContainerInstance(const UActorDescContainer*& OutLevelContainer, FTransform& OutLevelTransform, EContainerClusterMode& OutClusterMode) const { return false; }
 
 	bool operator==(const FWorldPartitionActorDesc& Other) const
diff --git a/Engine/Source/Runtime/Engine/Public/WorldPartition/WorldPartitionActorDescView.h b/Engine/Source/Runtime/Engine/Public/WorldPartition/WorldPartitionActorDescView.h
--- a/Engine/Source/Runtime/Engine/Public/WorldPartition/WorldPartitionActorDescView.h
+++ b/Engine/Source/Runtime/Engine/Public/WorldPartition/WorldPartitionActorDescView.h
@@ -39,10 +39,13 @@ public:
 	FName GetActorLabel() const;
 	FBox GetBounds() const;
 	const TArray<FGuid>& GetReferences() const;
+	FString ToString() const;
 	uint32 GetTag() const;
 
 	bool GetContainerInstance(const UActorDescContainer*& OutLevelContainer, FTransform& OutLevelTransform, EContainerClusterMode& OutClusterMode) const;
 
+	FName GetActorLabelOrName() const;
+
 	void SetGridPlacement(EActorGridPlacement InGridPlacement);
 
 	void SetInvalidDataLayers();
