Fix Linux build: remove UENUM from EContainerClusterMode and handle line terminator differently.

diff --git a/Engine/Source/Runtime/Engine/Private/WorldPartition/WorldPartitionStreamingGeneration.cpp b/Engine/Source/Runtime/Engine/Private/WorldPartition/WorldPartitionStreamingGeneration.cpp
--- a/Engine/Source/Runtime/Engine/Private/WorldPartition/WorldPartitionStreamingGeneration.cpp
+++ b/Engine/Source/Runtime/Engine/Private/WorldPartition/WorldPartitionStreamingGeneration.cpp
@@ -339,8 +339,9 @@ public:
 		{
 			auto SerializeLine = [LogFile](const FString& Line)
 			{
-				LogFile->Serialize(TCHAR_TO_ANSI(*Line), Line.Len());
-				LogFile->Serialize(LINE_TERMINATOR_ANSI, sizeof(LINE_TERMINATOR_ANSI));
+				FString LineEntry(Line);
+				LineEntry += LINE_TERMINATOR;
+				LogFile->Serialize(TCHAR_TO_ANSI(*LineEntry), LineEntry.Len());
 			};
 
 			// Build the containers tree representation
@@ -367,7 +368,6 @@ public:
 					SerializeLine(FString::Printf(TEXT("%s       Bounds: %s"), *Prefix, *ContainerDescriptor.Bounds.ToString()));
 					SerializeLine(FString::Printf(TEXT("%s    Transform: %s"), *Prefix, *ContainerDescriptor.Transform.ToString()));
 					SerializeLine(FString::Printf(TEXT("%s    Container: %s"), *Prefix, *ContainerDescriptor.Container->GetContainerPackage().ToString()));
-					SerializeLine(FString::Printf(TEXT("%s  ClusterMode: %s"), *Prefix, *StaticEnum<EContainerClusterMode>()->GetNameStringByValue((uint64)ContainerDescriptor.ClusterMode)));
 
 					if (ContainerDescriptor.ActorDescViewMap.Num())
 					{
diff --git a/Engine/Source/Runtime/Engine/Public/WorldPartition/WorldPartitionActorDesc.h b/Engine/Source/Runtime/Engine/Public/WorldPartition/WorldPartitionActorDesc.h
--- a/Engine/Source/Runtime/Engine/Public/WorldPartition/WorldPartitionActorDesc.h
+++ b/Engine/Source/Runtime/Engine/Public/WorldPartition/WorldPartitionActorDesc.h
@@ -21,7 +21,6 @@ struct FWorldPartitionActorDescInitData
 
 class UActorDescContainer;
 
-UENUM()
 enum class EContainerClusterMode : uint8
 {
 	Partitioned, // Per Actor Partitioning
