World Partition: set the AsyncLoading flag on the dynamic level package to ensure it's not discoverable while streaming.
This fixes the case where the replication code was resolving an incomplete cell package and failed replication of non-streamed actors.

diff --git a/Engine/Source/Runtime/Engine/Private/WorldPartition/WorldPartitionLevelStreamingDynamic.cpp b/Engine/Source/Runtime/Engine/Private/WorldPartition/WorldPartitionLevelStreamingDynamic.cpp
--- a/Engine/Source/Runtime/Engine/Private/WorldPartition/WorldPartitionLevelStreamingDynamic.cpp
+++ b/Engine/Source/Runtime/Engine/Private/WorldPartition/WorldPartitionLevelStreamingDynamic.cpp
@@ -192,6 +192,9 @@ bool UWorldPartitionLevelStreamingDynamic::RequestLevel(UWorld* InPersistentWorl
 	
 			if (IssueLoadRequests())
 			{
+				// Make sure this package isn't discoverable until it is fully loaded
+				CellLevelPackage->SetInternalFlags(EInternalObjectFlags::AsyncLoading);
+
 				// Editor immediately blocks on load and we also block if background level streaming is disabled.
 				if (InBlockPolicy == AlwaysBlock || (ShouldBeAlwaysLoaded() && InBlockPolicy != NeverBlock))
 				{
@@ -345,6 +348,9 @@ void UWorldPartitionLevelStreamingDynamic::FinalizeRuntimeLevel()
 		}
 	}
 
+	// Now that the loading is done, clear the async loading flag so this level is discoverable
+	RuntimePackage->ClearInternalFlags(EInternalObjectFlags::AsyncLoading);
+
 	SetLoadedLevel(RuntimeLevel);
 
 	// Broadcast level loaded event to blueprints
