Engine Changes required for diskless textures

#rnx

diff --git a/Engine/Source/Runtime/Engine/Classes/Engine/TextureMipDataProviderFactory.h b/Engine/Source/Runtime/Engine/Classes/Engine/TextureMipDataProviderFactory.h
--- a/Engine/Source/Runtime/Engine/Classes/Engine/TextureMipDataProviderFactory.h
+++ b/Engine/Source/Runtime/Engine/Classes/Engine/TextureMipDataProviderFactory.h
@@ -33,4 +33,10 @@ public:
 	* @param InAsset - the texture on which the stream in operation will be performed.
 	*/
 	virtual FTextureMipDataProvider* AllocateMipDataProvider(UTexture* Asset) { return nullptr; }
+
+	/**
+	* Returns true if TextureMipDataProviders allocated by this factory can provide MipData by themselves,
+	* even without loading from disk at all, so streaming can be enabled for their textures.
+	*/
+	virtual bool WillProvideMipDataWithoutDisk() const { return false; }
 };
diff --git a/Engine/Source/Runtime/Engine/Private/Texture.cpp b/Engine/Source/Runtime/Engine/Private/Texture.cpp
--- a/Engine/Source/Runtime/Engine/Private/Texture.cpp
+++ b/Engine/Source/Runtime/Engine/Private/Texture.cpp
@@ -14,6 +14,8 @@
 #include "UObject/ObjectSaveContext.h"
 #include "TextureResource.h"
 #include "Engine/Texture2D.h"
+#include "Streaming/TextureMipDataProvider.h"
+#include "Engine/TextureMipDataProviderFactory.h"
 #include "ContentStreaming.h"
 #include "EngineUtils.h"
 #include "Engine/AssetUserData.h"
@@ -1053,11 +1055,28 @@ FStreamableRenderResourceState UTexture::GetResourcePostInitState(const FTexture
 	int32 NumRequestedMips = 0;
 
 #if PLATFORM_SUPPORTS_TEXTURE_STREAMING
+	bool bWillProvideMipDataWithoutDisk = false;
+
+	// Check if any of the CustomMipData providers associated with this texture can provide mip data even without DDC or disk,
+	// if so, enable streaming for this texture
+	for (UAssetUserData* UserData : AssetUserData)
+	{
+		UTextureMipDataProviderFactory* CustomMipDataProviderFactory = Cast<UTextureMipDataProviderFactory>(UserData);
+		if (CustomMipDataProviderFactory)
+		{
+			bWillProvideMipDataWithoutDisk = CustomMipDataProviderFactory->WillProvideMipDataWithoutDisk();
+			if (bWillProvideMipDataWithoutDisk)
+			{
+				break;
+			}
+		}
+	}
+
 	if (!NeverStream && 
 		NumOfNonStreamingMips < NumMips && 
 		LODGroup != TEXTUREGROUP_UI && 
 		bAllowStreaming &&
-		(bSkipCanBeLoaded || PlatformData->CanBeLoaded()))
+		(bSkipCanBeLoaded || PlatformData->CanBeLoaded() || bWillProvideMipDataWithoutDisk))
 	{
 		bMakeStreamble  = true;
 	}
