[Scalability] Forcing set default scalability values so base scalability group variables from config can be applied properly.

diff --git a/Engine/Source/Runtime/Engine/Private/Scalability.cpp b/Engine/Source/Runtime/Engine/Private/Scalability.cpp
--- a/Engine/Source/Runtime/Engine/Private/Scalability.cpp
+++ b/Engine/Source/Runtime/Engine/Private/Scalability.cpp
@@ -555,7 +555,7 @@ void InitScalabilitySystem()
 	CVarShadingQuality.AsVariable()->SetOnChangedCallback(FConsoleVariableDelegate::CreateStatic(&OnChangeShadingQuality));
 
 	// Set defaults
-	SetQualityLevels(FQualityLevels());
+	SetQualityLevels(FQualityLevels(), true);
 	GScalabilityBackupQualityLevels = FQualityLevels();
 	GScalabilityQualityLevelsOverride = FQualityLevels();
 	GScalabilityUsingTemporaryQualityLevels = false;
