When sampling texture with LWC, we manually compute texture address mode to change to float precision
When we do this, also need to manually compute ddx/ddy based on original LWC coordinate

diff --git a/Engine/Source/Runtime/Engine/Private/Materials/HLSLMaterialTranslator.cpp b/Engine/Source/Runtime/Engine/Private/Materials/HLSLMaterialTranslator.cpp
--- a/Engine/Source/Runtime/Engine/Private/Materials/HLSLMaterialTranslator.cpp
+++ b/Engine/Source/Runtime/Engine/Private/Materials/HLSLMaterialTranslator.cpp
@@ -5822,7 +5822,8 @@ int32 FHLSLMaterialTranslator::TextureSample(
 	}
 	
 	int32 NonLWCCoordinateIndex = CoordinateIndex;
-	if (IsLWCType(GetParameterType(CoordinateIndex)))
+	const bool bLWCCoordinates = IsLWCType(GetParameterType(CoordinateIndex));
+	if (bLWCCoordinates)
 	{
 		// Apply texture address math manually, using LWC-scale operations, then convert the result to float
 		// This could potentially cause problems if content is relying on SSM_FromTextureAsset, and having texture parameters change address mode in MI
@@ -5862,6 +5863,25 @@ int32 FHLSLMaterialTranslator::TextureSample(
 			checkf(false, TEXT("Invalid number of components %d"), NumComponents);
 			break;
 		}
+
+		if (MipValueMode == TMVM_None || MipValueMode == TMVM_MipBias)
+		{
+			int32 MipScaleIndex = INDEX_NONE;
+			if (MipValueMode == TMVM_MipBias)
+			{
+				MipScaleIndex = AddCodeChunkZeroDeriv(UVsType, TEXT("exp2(%s)"), *CoerceParameter(MipValue0Index, MCT_Float1));
+			}
+
+			MipValue0Index = AddCodeChunkZeroDeriv(UVsType, TEXT("LWCDdx(%s)"), *GetParameterCode(CoordinateIndex));
+			MipValue1Index = AddCodeChunkZeroDeriv(UVsType, TEXT("LWCDdy(%s)"), *GetParameterCode(CoordinateIndex));
+			if (MipScaleIndex != INDEX_NONE)
+			{
+				MipValue0Index = Mul(MipValue0Index, MipScaleIndex);
+				MipValue1Index = Mul(MipValue1Index, MipScaleIndex);
+			}
+
+			MipValueMode = TMVM_Derivative;
+		}
 	}
 
 	FString SamplerStateCode;
