Additional fixes gamepad state not maintained

diff --git a/Engine/Source/Runtime/Engine/Private/PlayerController.cpp b/Engine/Source/Runtime/Engine/Private/PlayerController.cpp
--- a/Engine/Source/Runtime/Engine/Private/PlayerController.cpp
+++ b/Engine/Source/Runtime/Engine/Private/PlayerController.cpp
@@ -2308,7 +2308,7 @@ bool APlayerController::InputKey(const FInputKeyParams& Params)
 		if (GEngine->XRSystem.IsValid())
 		{
 			auto XRInput = GEngine->XRSystem->GetXRInput();
-			if (XRInput && XRInput->HandleInputKey(PlayerInput, Params.Key, Params.Event, Params.Delta.X, Params.Key.IsGamepadKey()))
+			if (XRInput && XRInput->HandleInputKey(PlayerInput, Params.Key, Params.Event, Params.Delta.X, Params.IsGamepad()))
 			{
 				return true;
 			}
diff --git a/Engine/Source/Runtime/Engine/Private/UserInterface/PlayerInput.cpp b/Engine/Source/Runtime/Engine/Private/UserInterface/PlayerInput.cpp
--- a/Engine/Source/Runtime/Engine/Private/UserInterface/PlayerInput.cpp
+++ b/Engine/Source/Runtime/Engine/Private/UserInterface/PlayerInput.cpp
@@ -203,6 +203,7 @@ bool UPlayerInput::InputKey(FKey Key, EInputEvent Event, float AmountDepressed,
 	Params.Key = Key;
 	Params.Event = Event;
 	Params.Delta = FVector((double)AmountDepressed, 0.0, 0.0);
+	Params.bIsGamepadOverride = bGamepad;
 	
 	return InputKey(Params);
 }
@@ -370,6 +371,7 @@ bool UPlayerInput::InputAxis(FKey Key, float Delta, float DeltaTime, int32 NumSa
 	Params.Key = Key;
 	Params.Delta = FVector((double)Delta, 0.0, 0.0);
 	Params.NumSamples = NumSamples;
+	Params.bIsGamepadOverride = bGamepad;
 	
 	return InputKey(Params);
 }
