The default type elements api assume that the objects pointers are always valid. So I made some k2_vesions of the function for the blueprints. It does a null checks before calling the native function.

Added some additional error message to some function of the editor actor subsystem.

Changed the signature of UEditorActorSubsystem::SetComponentTransform to take a SceneComponent since ActorComponents doesn't have a transform.

diff --git a/Engine/Source/Editor/UnrealEd/Private/LevelViewportClickHandlers.cpp b/Engine/Source/Editor/UnrealEd/Private/LevelViewportClickHandlers.cpp
--- a/Engine/Source/Editor/UnrealEd/Private/LevelViewportClickHandlers.cpp
+++ b/Engine/Source/Editor/UnrealEd/Private/LevelViewportClickHandlers.cpp
@@ -253,7 +253,7 @@ namespace LevelViewportClickHandlers
 				FlushRenderingCommands();
 			}
 
-			PrivateSummonContextMenu(ViewportClient, UEngineElementsLibrary::AcquireEditorActorElementHandle(Actor));
+			PrivateSummonContextMenu(ViewportClient, Actor ? UEngineElementsLibrary::AcquireEditorActorElementHandle(Actor) : FTypedElementHandle());
 			return true;
 		}
 		else if( Click.GetEvent() == IE_DoubleClick && Click.GetKey() == EKeys::LeftMouseButton && !Click.IsControlDown() && !Click.IsShiftDown() )
diff --git a/Engine/Source/Editor/UnrealEd/Private/Subsystems/EditorActorSubsystem.cpp b/Engine/Source/Editor/UnrealEd/Private/Subsystems/EditorActorSubsystem.cpp
--- a/Engine/Source/Editor/UnrealEd/Private/Subsystems/EditorActorSubsystem.cpp
+++ b/Engine/Source/Editor/UnrealEd/Private/Subsystems/EditorActorSubsystem.cpp
@@ -1,6 +1,7 @@
 // Copyright Epic Games, Inc. All Rights Reserved.
 
 #include "Subsystems/EditorActorSubsystem.h"
+
 #include "CoreMinimal.h"
 #include "Engine/MeshMerging.h"
 #include "GameFramework/Actor.h"
@@ -21,6 +22,7 @@
 #include "Engine/Selection.h"
 #include "LevelEditorViewport.h"
 #include "Subsystems/EditorElementSubsystem.h"
+#include "UObject/Stack.h"
 
 #include "Elements/Framework/EngineElementsLibrary.h"
 #include "Elements/Framework/TypedElementHandle.h"
@@ -472,7 +474,7 @@ AActor* UEditorActorSubsystem::SpawnActorFromClass(TSubclassOf<class AActor> Act
 
 	if (!ActorClass.Get())
 	{
-		UE_LOG(LogUtils, Error, TEXT("SpawnActorFromClass. ActorClass is not valid."));
+		FFrame::KismetExecutionMessage(TEXT("SpawnActorFromClass. ActorClass is not valid."), ELogVerbosity::Error);
 		return nullptr;
 	}
 
@@ -703,6 +705,12 @@ TArray<AActor*> UEditorActorSubsystem::DuplicateActors(const TArray<AActor*>& Ac
 
 bool UEditorActorSubsystem::SetActorTransform(AActor* InActor, const FTransform& InWorldTransform)
 {
+	if (!InActor)
+	{
+		FFrame::KismetExecutionMessage(TEXT("Cannot set the transfrom of a nullptr actor."), ELogVerbosity::Error);
+		return false;
+	}
+
 	if (UEditorElementSubsystem* ElementSubsystem = GEditor->GetEditorSubsystem<UEditorElementSubsystem>())
 	{
 		if (FTypedElementHandle ActorElementHandle = UEngineElementsLibrary::AcquireEditorActorElementHandle(InActor))
@@ -714,11 +722,17 @@ bool UEditorActorSubsystem::SetActorTransform(AActor* InActor, const FTransform&
 	return false;
 }
 
-bool UEditorActorSubsystem::SetComponentTransform(UActorComponent* InActorComponent, const FTransform& InWorldTransform)
+bool UEditorActorSubsystem::SetComponentTransform(USceneComponent* InSceneComponent, const FTransform& InWorldTransform)
 {
+	if (!InSceneComponent)
+	{
+		FFrame::KismetExecutionMessage(TEXT("Cannot set the transform of a nullptr SceneComponent."), ELogVerbosity::Error);
+		return false;
+	}
+
 	if (UEditorElementSubsystem* ElementSubsystem = GEditor->GetEditorSubsystem<UEditorElementSubsystem>())
 	{
-		if (FTypedElementHandle ComponentElementHandle = UEngineElementsLibrary::AcquireEditorComponentElementHandle(InActorComponent))
+		if (FTypedElementHandle ComponentElementHandle = UEngineElementsLibrary::AcquireEditorComponentElementHandle(InSceneComponent))
 		{
 			return ElementSubsystem->SetElementTransform(ComponentElementHandle, InWorldTransform);
 		}
diff --git a/Engine/Source/Editor/UnrealEd/Public/Subsystems/EditorActorSubsystem.h b/Engine/Source/Editor/UnrealEd/Public/Subsystems/EditorActorSubsystem.h
--- a/Engine/Source/Editor/UnrealEd/Public/Subsystems/EditorActorSubsystem.h
+++ b/Engine/Source/Editor/UnrealEd/Public/Subsystems/EditorActorSubsystem.h
@@ -236,7 +236,7 @@ public:
 	 * @returns false if the world transform could not be set.
 	 */
 	UFUNCTION(BlueprintCallable, Category = "Editor Scripting | Level Utility")
-	bool SetComponentTransform(UActorComponent* InActorComponent, const FTransform& InWorldTransform);
+	bool SetComponentTransform(USceneComponent* InSceneComponent, const FTransform& InWorldTransform);
 
 private:
 
diff --git a/Engine/Source/Runtime/Engine/Private/Elements/Framework/EngineElementsLibrary.cpp b/Engine/Source/Runtime/Engine/Private/Elements/Framework/EngineElementsLibrary.cpp
--- a/Engine/Source/Runtime/Engine/Private/Elements/Framework/EngineElementsLibrary.cpp
+++ b/Engine/Source/Runtime/Engine/Private/Elements/Framework/EngineElementsLibrary.cpp
@@ -1,6 +1,7 @@
 // Copyright Epic Games, Inc. All Rights Reserved.
 
 #include "Elements/Framework/EngineElementsLibrary.h"
+
 #include "Elements/Framework/TypedElementRegistry.h"
 #include "Elements/Framework/TypedElementOwnerStore.h"
 #include "Elements/Framework/TypedElementUtil.h"
@@ -18,6 +19,10 @@
 #include "Elements/SMInstance/SMInstanceElementData.h"
 #include "Components/InstancedStaticMeshComponent.h"
 
+#include "Containers/UnrealString.h"
+#include "UObject/Stack.h"
+
+
 // The editor requires ref-counting for object replacement to function correctly
 static_assert(!WITH_EDITOR || UE_TYPED_ELEMENT_HAS_REFCOUNTING, "The editor requires that ref-counting is enabled for typed elements!");
 
@@ -321,6 +326,17 @@ FTypedElementHandle UEngineElementsLibrary::AcquireEditorObjectElementHandle(con
 	return EngineElementsLibraryUtil::AcquireEditorTypedElementHandle<UObject, FObjectElementData>(Object, GObjectElementOwnerStore, &UEngineElementsLibrary::CreateObjectElement, bAllowCreate);
 }
 
+FTypedElementHandle UEngineElementsLibrary::K2_AcquireEditorObjectElementHandle(const UObject* Object, const bool bAllowCreate)
+{
+	if (!Object)
+	{
+		FFrame::KismetExecutionMessage(TEXT("Cannot acquire an editor object element handle from a nullptr object."), ELogVerbosity::Error);
+		return FTypedElementHandle();
+	}
+
+	return AcquireEditorObjectElementHandle(Object , bAllowCreate);
+}
+
 void UEngineElementsLibrary::ReplaceEditorObjectElementHandles(const TMap<const UObject*, const UObject*>& ReplacementObjects)
 {
 	TArray<FTypedElementHandle> UpdatedElements;
@@ -367,6 +383,17 @@ FTypedElementHandle UEngineElementsLibrary::AcquireEditorActorElementHandle(cons
 	return EngineElementsLibraryUtil::AcquireEditorTypedElementHandle<AActor, FActorElementData>(Actor, GActorElementOwnerStore, &UEngineElementsLibrary::CreateActorElement, bAllowCreate);
 }
 
+FTypedElementHandle UEngineElementsLibrary::K2_AcquireEditorActorElementHandle(const AActor* Actor, const bool bAllowCreate)
+{
+	if (!Actor)
+	{
+		FFrame::KismetExecutionMessage(TEXT("Cannot acquire an editor actor element handle from a nullptr actor."), ELogVerbosity::Error);
+		return FTypedElementHandle();
+	}
+
+	return AcquireEditorActorElementHandle(Actor, bAllowCreate);
+}
+
 void UEngineElementsLibrary::ReplaceEditorActorElementHandles(const TMap<const AActor*, const AActor*>& ReplacementActors)
 {
 	TArray<FTypedElementHandle> UpdatedElements;
@@ -413,6 +440,17 @@ FTypedElementHandle UEngineElementsLibrary::AcquireEditorComponentElementHandle(
 	return EngineElementsLibraryUtil::AcquireEditorTypedElementHandle<UActorComponent, FComponentElementData>(Component, GComponentElementOwnerStore, &UEngineElementsLibrary::CreateComponentElement, bAllowCreate);
 }
 
+FTypedElementHandle UEngineElementsLibrary::K2_AcquireEditorComponentElementHandle(const UActorComponent* Component, const bool bAllowCreate)
+{
+	if (!Component)
+	{
+		FFrame::KismetExecutionMessage(TEXT("Cannot acquire an editor component element handle from a nullptr component."), ELogVerbosity::Error);
+		return FTypedElementHandle();
+	}
+
+	return AcquireEditorComponentElementHandle(Component, bAllowCreate);
+}
+
 void UEngineElementsLibrary::ReplaceEditorComponentElementHandles(const TMap<const UActorComponent*, const UActorComponent*>& ReplacementComponents)
 {
 	TArray<FTypedElementHandle> UpdatedElements;
@@ -483,6 +521,28 @@ FTypedElementHandle UEngineElementsLibrary::AcquireEditorSMInstanceElementHandle
 	return AcquireEditorSMInstanceElementHandle(FSMInstanceId{ const_cast<UInstancedStaticMeshComponent*>(ISMComponent), InstanceIndex }, bAllowCreate);
 }
 
+FTypedElementHandle UEngineElementsLibrary::K2_AcquireEditorSMInstanceElementHandle(const UInstancedStaticMeshComponent* ISMComponent, const int32 InstanceIndex, const bool bAllowCreate)
+{
+	if (!ISMComponent)
+	{
+		FFrame::KismetExecutionMessage(TEXT("Cannot acquire an editor SMInstance element handle from a nullptr ISMComponent."), ELogVerbosity::Error);
+		return FTypedElementHandle();
+	}
+
+	if (ISMComponent->GetInstanceCount() <= InstanceIndex || InstanceIndex < 0)
+	{
+		FStringFormatOrderedArguments FormatOrderedArguments;
+		FormatOrderedArguments.Reserve(2);
+		FormatOrderedArguments.Add(InstanceIndex);
+		FormatOrderedArguments.Add(ISMComponent->GetInstanceCount());
+		const FString Message = FString::Format(TEXT("Cannot acquire an editor SMInstance element handle from a instance that doesn't exist in the ISMComponent. (InstanceIndex: {0}, Instance Count: {1})"), FormatOrderedArguments);
+		FFrame::KismetExecutionMessage(*Message, ELogVerbosity::Error);
+		return FTypedElementHandle();
+	}
+
+	return AcquireEditorSMInstanceElementHandle(ISMComponent, InstanceIndex, bAllowCreate);
+}
+
 FTypedElementHandle UEngineElementsLibrary::AcquireEditorSMInstanceElementHandle(const FSMInstanceId& SMInstanceId, const bool bAllowCreate)
 {
 	if (GIsEditor)
diff --git a/Engine/Source/Runtime/Engine/Public/Elements/Framework/EngineElementsLibrary.h b/Engine/Source/Runtime/Engine/Public/Elements/Framework/EngineElementsLibrary.h
--- a/Engine/Source/Runtime/Engine/Public/Elements/Framework/EngineElementsLibrary.h
+++ b/Engine/Source/Runtime/Engine/Public/Elements/Framework/EngineElementsLibrary.h
@@ -34,8 +34,9 @@ public:
 #if WITH_EDITOR
 	static void CreateEditorObjectElement(const UObject* Object);
 	static void DestroyEditorObjectElement(const UObject* Object);
-	UFUNCTION(BlueprintCallable, BlueprintPure=false, Category="TypedElementFramework|Object", meta=(ScriptMethod="AcquireEditorElementHandle"))
 	static FTypedElementHandle AcquireEditorObjectElementHandle(const UObject* Object, const bool bAllowCreate = true);
+	UFUNCTION(BlueprintCallable, BlueprintPure=false, Category="TypedElementFramework|Object", meta=(DisplayName="Acquire Editor Object Element Handle", ScriptMethod="AcquireEditorElementHandle"))
+	static FTypedElementHandle K2_AcquireEditorObjectElementHandle(const UObject* Object, const bool bAllowCreate = true);
 	static void ReplaceEditorObjectElementHandles(const TMap<const UObject*, const UObject*>& ReplacementObjects);
 #endif
 
@@ -44,8 +45,9 @@ public:
 #if WITH_EDITOR
 	static void CreateEditorActorElement(const AActor* Actor);
 	static void DestroyEditorActorElement(const AActor* Actor);
-	UFUNCTION(BlueprintCallable, Category="TypedElementFramework|Actor", meta=(ScriptMethod="AcquireEditorElementHandle"))
 	static FTypedElementHandle AcquireEditorActorElementHandle(const AActor* Actor, const bool bAllowCreate = true);
+	UFUNCTION(BlueprintCallable, Category="TypedElementFramework|Actor", meta=(DisplayName="Acquire Editor Actor Element Handle", ScriptMethod="AcquireEditorElementHandle"))
+	static FTypedElementHandle K2_AcquireEditorActorElementHandle(const AActor* Actor, const bool bAllowCreate = true);
 	static void ReplaceEditorActorElementHandles(const TMap<const AActor*, const AActor*>& ReplacementActors);
 #endif
 
@@ -54,8 +56,9 @@ public:
 #if WITH_EDITOR
 	static void CreateEditorComponentElement(const UActorComponent* Component);
 	static void DestroyEditorComponentElement(const UActorComponent* Component);
-	UFUNCTION(BlueprintCallable, BlueprintPure=false, Category="TypedElementFramework|Component", meta=(ScriptMethod="AcquireEditorElementHandle"))
 	static FTypedElementHandle AcquireEditorComponentElementHandle(const UActorComponent* Component, const bool bAllowCreate = true);
+	UFUNCTION(BlueprintCallable, BlueprintPure=false, Category="TypedElementFramework|Component", meta=(DisplayName="Acquire Editor Component Element Handle", ScriptMethod = "AcquireEditorElementHandle"))
+	static FTypedElementHandle K2_AcquireEditorComponentElementHandle(const UActorComponent* Component, const bool bAllowCreate = true);
 	static void ReplaceEditorComponentElementHandles(const TMap<const UActorComponent*, const UActorComponent*>& ReplacementComponents);
 #endif
 
@@ -64,8 +67,9 @@ public:
 #if WITH_EDITOR
 	static void CreateEditorSMInstanceElement(const FSMInstanceId& SMInstanceId);
 	static void DestroyEditorSMInstanceElement(const FSMInstanceElementId& SMInstanceElementId);
-	UFUNCTION(BlueprintCallable, BlueprintPure=false, Category="TypedElementFramework|StaticMeshInstance", meta=(ScriptMethod="AcquireEditorElementHandle"))
 	static FTypedElementHandle AcquireEditorSMInstanceElementHandle(const UInstancedStaticMeshComponent* ISMComponent, const int32 InstanceIndex, const bool bAllowCreate = true);
+	UFUNCTION(BlueprintCallable, BlueprintPure=false, Category="TypedElementFramework|StaticMeshInstance", meta=(DisplayName="Acquire Editor SMInstance Element Handle", ScriptMethod = "AcquireEditorElementHandle"))
+	static FTypedElementHandle K2_AcquireEditorSMInstanceElementHandle(const UInstancedStaticMeshComponent* ISMComponent, const int32 InstanceIndex, const bool bAllowCreate = true);
 	static FTypedElementHandle AcquireEditorSMInstanceElementHandle(const FSMInstanceId& SMInstanceId, const bool bAllowCreate = true);
 	static void ReplaceEditorSMInstanceElementHandles(const TMap<FSMInstanceId, FSMInstanceId>& ReplacementSMInstanceIds);
 #endif
