UWorld::IsInstanced()
- ULevel::IsInstancedLevel calls UWorld::IsInstanced
- UWorld::IsInstanced doesn't check for /Temp/Untitled anymore as this is a valid instancing case when source package name is not empty (template)
- UWorld::GetSoftObjectMapping uses UWorld::IsInstanced() fixes unhandled empty source package name case

diff --git a/Engine/Source/Runtime/Engine/Classes/Engine/World.h b/Engine/Source/Runtime/Engine/Classes/Engine/World.h
--- a/Engine/Source/Runtime/Engine/Classes/Engine/World.h
+++ b/Engine/Source/Runtime/Engine/Classes/Engine/World.h
@@ -3883,10 +3883,13 @@ public:
 	/** Given a level script actor, modify the string such that it points to the correct instance of the object. For replays. */
 	bool RemapCompiledScriptActor(FString& Str) const;
 
+	/** Returns true if world package is instanced. */
+	bool IsInstanced() const;
+
 	/** 
 	 * If World Package is instanced return a mapping that can be used to fixup SoftObjectPaths for this world 
 	 *
-	 * returns false if package doesn't need remapping
+	 * returns true if world package is instanced and needs remapping.
 	 */
 	bool GetSoftObjectPathMapping(FString& OutSourceWorldPath, FString& OutRemappedWorldPath) const;
 
diff --git a/Engine/Source/Runtime/Engine/Private/Level.cpp b/Engine/Source/Runtime/Engine/Private/Level.cpp
--- a/Engine/Source/Runtime/Engine/Private/Level.cpp
+++ b/Engine/Source/Runtime/Engine/Private/Level.cpp
@@ -2899,11 +2899,8 @@ bool ULevel::IsBeingRemoved() const
 
 bool ULevel::IsInstancedLevel() const
 {
-	UPackage* LevelPackage = GetOutermost();
-	static FName UnsavedPackageName("/Temp/Untitled");
-	const FPackagePath& PackagePath = LevelPackage->GetLoadedPath();
-	return !PackagePath.IsEmpty() && PackagePath.GetPackageFName() != LevelPackage->GetFName() 
-		&& !LevelPackage->GetFName().IsEqual(UnsavedPackageName, ENameCase::IgnoreCase, /*bCompareNumber*/false);
+	const UWorld* OuterWorld = GetTypedOuter<UWorld>();
+	return OuterWorld && OuterWorld->IsInstanced();
 }
 
 void ULevel::ApplyWorldOffset(const FVector& InWorldOffset, bool bWorldShift)
diff --git a/Engine/Source/Runtime/Engine/Private/World.cpp b/Engine/Source/Runtime/Engine/Private/World.cpp
--- a/Engine/Source/Runtime/Engine/Private/World.cpp
+++ b/Engine/Source/Runtime/Engine/Private/World.cpp
@@ -3292,11 +3292,17 @@ void UWorld::RenameToPIEWorld(int32 PIEInstanceID)
 #endif
 }
 
-bool UWorld::GetSoftObjectPathMapping(FString& OutSourceWorldPath, FString& OutRemappedWorldPath) const
+bool UWorld::IsInstanced() const
 {
 	UPackage* Package = GetPackage();
-	if (Package->GetFName() != Package->GetLoadedPath().GetPackageFName())
+	return !Package->GetLoadedPath().IsEmpty() && Package->GetFName() != Package->GetLoadedPath().GetPackageFName();
+}
+
+bool UWorld::GetSoftObjectPathMapping(FString& OutSourceWorldPath, FString& OutRemappedWorldPath) const
+{
+	if (IsInstanced())
 	{
+		UPackage* Package = GetPackage();
 		const FString SourcePackageName = Package->GetLoadedPath().GetPackageName();
 		const FString SourceWorldName = FPaths::GetBaseFilename(SourcePackageName);
 		const FString RemmappedPackageName = Package->GetName();
