World Partition: make sure the level package is not discoverable while streaming.
This fixes the case where the replication code was resolving an incomplete cell package and failed replication of non-streamed actors.

diff --git a/Engine/Source/Runtime/Engine/Private/WorldPartition/WorldPartitionLevelStreamingDynamic.cpp b/Engine/Source/Runtime/Engine/Private/WorldPartition/WorldPartitionLevelStreamingDynamic.cpp
--- a/Engine/Source/Runtime/Engine/Private/WorldPartition/WorldPartitionLevelStreamingDynamic.cpp
+++ b/Engine/Source/Runtime/Engine/Private/WorldPartition/WorldPartitionLevelStreamingDynamic.cpp
@@ -85,8 +85,10 @@ void UWorldPartitionLevelStreamingDynamic::CreateRuntimeLevel()
 	const UWorld* PlayWorld = GetWorld();
 	check(PlayWorld && PlayWorld->IsGameWorld());
 
-	// Create streaming cell Level package
-	RuntimeLevel = FWorldPartitionLevelHelper::CreateEmptyLevelForRuntimeCell(PlayWorld, GetWorldAsset().ToString());
+	// Create streaming cell Level package (make sure this package isn't discoverable until it is fully loaded)
+	const FString WorldAssetPath = GetWorldAsset().ToString();
+	const FString WorldAssetPackageName = FPackageName::ObjectPathToPackageName(WorldAssetPath) + TEXT("_AsyncLoad.") + FPackageName::ObjectPathToObjectName(WorldAssetPath);
+	RuntimeLevel = FWorldPartitionLevelHelper::CreateEmptyLevelForRuntimeCell(PlayWorld, WorldAssetPackageName);
 	check(RuntimeLevel);
 
 	// Set flag here as this level isn't async loaded
@@ -322,6 +324,11 @@ void UWorldPartitionLevelStreamingDynamic::FinalizeRuntimeLevel()
 	UPackage* RuntimePackage = RuntimeLevel->GetPackage();
 	RuntimePackage->MarkAsFullyLoaded();
 
+	// Now that the loading is done, rename it so this level is discoverable
+	const FString WorldAssetPath = GetWorldAsset().ToString();
+	const FString WorldAssetPackageName = FPackageName::ObjectPathToPackageName(WorldAssetPath);
+	RuntimePackage->Rename(*WorldAssetPackageName, nullptr, REN_ForceNoResetLoaders | REN_DontCreateRedirectors | REN_NonTransactional | REN_DoNotDirty);
+
 	if (!StreamingCell->GetIsHLOD())
 	{
 		TRACE_CPUPROFILER_EVENT_SCOPE(FixupIDs);
