Fix crash creating a new, non-open world level after saving an open world level.

diff --git a/Engine/Source/Runtime/Engine/Private/WorldPartition/WorldPartitionEditorCell.cpp b/Engine/Source/Runtime/Engine/Private/WorldPartition/WorldPartitionEditorCell.cpp
--- a/Engine/Source/Runtime/Engine/Private/WorldPartition/WorldPartitionEditorCell.cpp
+++ b/Engine/Source/Runtime/Engine/Private/WorldPartition/WorldPartitionEditorCell.cpp
@@ -68,7 +68,7 @@ void UWorldPartitionEditorCell::AddActor(const FGuid& Source, const FWorldPartit
 			if (ReferenceHandle.IsValid())
 			{
 				AddActor(ActorHandle->GetGuid(), ReferenceHandle);
-				WorldPartition->EditorHash->AddBackReference(ReferenceHandle, this, Source);
+				WorldPartition->EditorHash->AddBackReference(ReferenceHandle->GetGuid(), this, Source);
 			}
 		}
 	}
@@ -104,7 +104,7 @@ void UWorldPartitionEditorCell::RemoveActor(const FGuid& Source, const FWorldPar
 			if (ReferenceHandle.IsValid())
 			{
 				RemoveActor(ActorHandle->GetGuid(), ReferenceHandle);
-				WorldPartition->EditorHash->RemoveBackReference(ReferenceHandle, this, Source);
+				WorldPartition->EditorHash->RemoveBackReference(ReferenceHandle->GetGuid(), this, Source);
 			}
 		}
 	}
diff --git a/Engine/Source/Runtime/Engine/Private/WorldPartition/WorldPartitionEditorSpatialHash.cpp b/Engine/Source/Runtime/Engine/Private/WorldPartition/WorldPartitionEditorSpatialHash.cpp
--- a/Engine/Source/Runtime/Engine/Private/WorldPartition/WorldPartitionEditorSpatialHash.cpp
+++ b/Engine/Source/Runtime/Engine/Private/WorldPartition/WorldPartitionEditorSpatialHash.cpp
@@ -319,8 +319,8 @@ void UWorldPartitionEditorSpatialHash::UnhashActor(FWorldPartitionHandle& InActo
 	// This is important because the FWorldPartitionHandle will soon become invalid and Actors that were still 
 	// referencing this one will not be able to construct valid FWorldPartitionHandle to the actor we're 
 	// removing and it'll become impossible to clean the cells of their dangling actors
-	TArray< TTuple<UWorldPartitionEditorCell*, FGuid>> BackRefs;
-	BackReferences.MultiFind(InActorHandle, BackRefs);
+	TArray<TTuple<UWorldPartitionEditorCell*, FGuid>> BackRefs;
+	BackReferences.MultiFind(InActorHandle->GetGuid(), BackRefs);
 
 	for (auto& BackRef : BackRefs)
 	{
@@ -329,7 +329,7 @@ void UWorldPartitionEditorSpatialHash::UnhashActor(FWorldPartitionHandle& InActo
 
 	if (BackRefs.Num() != 0)
 	{
-		BackReferences.Remove(InActorHandle);
+		BackReferences.Remove(InActorHandle->GetGuid());
 	}
 }
 
@@ -589,17 +589,15 @@ void UWorldPartitionEditorSpatialHash::PostLoad()
 }
 
 
-void UWorldPartitionEditorSpatialHash::AddBackReference(FWorldPartitionHandle& ReferenceHandle,  UWorldPartitionEditorCell* Cell, const FGuid& Source)
+void UWorldPartitionEditorSpatialHash::AddBackReference(const FGuid& Reference, UWorldPartitionEditorCell* Cell, const FGuid& Source)
 {
-	BackReferences.Add(ReferenceHandle) = MakeTuple(Cell, Source);
+	BackReferences.Add(Reference, MakeTuple(Cell, Source));
 }
 
-void UWorldPartitionEditorSpatialHash::RemoveBackReference(FWorldPartitionHandle& ReferenceHandle,  UWorldPartitionEditorCell* Cell, const FGuid& Source)
+void UWorldPartitionEditorSpatialHash::RemoveBackReference(const FGuid& Reference, UWorldPartitionEditorCell* Cell, const FGuid& Source)
 {
-	BackReferences.RemoveSingle(ReferenceHandle, MakeTuple(Cell, Source));
+	BackReferences.RemoveSingle(Reference, MakeTuple(Cell, Source));
 }
-
-
 #endif
 
 #undef LOCTEXT_NAMESPACE
diff --git a/Engine/Source/Runtime/Engine/Public/WorldPartition/WorldPartitionEditorHash.h b/Engine/Source/Runtime/Engine/Public/WorldPartition/WorldPartitionEditorHash.h
--- a/Engine/Source/Runtime/Engine/Public/WorldPartition/WorldPartitionEditorHash.h
+++ b/Engine/Source/Runtime/Engine/Public/WorldPartition/WorldPartitionEditorHash.h
@@ -37,9 +37,8 @@ public:
 	virtual uint32 GetWantedEditorCellSize() const PURE_VIRTUAL(UWorldPartitionEditorHash::GetWantedEditorCellSize, return 0;);
 	virtual void SetEditorWantedCellSize(uint32 InCellSize) PURE_VIRTUAL(UWorldPartitionEditorHash::SetEditorWantedCellSize, );
 
-	virtual void AddBackReference(FWorldPartitionHandle& ReferenceHandle,  UWorldPartitionEditorCell* Cell, const FGuid& Source) PURE_VIRTUAL(UWorldPartitionEditorHash::AddBackReference, ;);
-	virtual void RemoveBackReference(FWorldPartitionHandle& ReferenceHandle,  UWorldPartitionEditorCell* Cell, const FGuid& Source) PURE_VIRTUAL(UWorldPartitionEditorHash::RemoveBackReference, ;);
-
+	virtual void AddBackReference(const FGuid& Reference, UWorldPartitionEditorCell* Cell, const FGuid& Source) PURE_VIRTUAL(UWorldPartitionEditorHash::AddBackReference, ;);
+	virtual void RemoveBackReference(const FGuid& Reference, UWorldPartitionEditorCell* Cell, const FGuid& Source) PURE_VIRTUAL(UWorldPartitionEditorHash::RemoveBackReference, ;);
 
 	// Helpers
 	inline int32 GetIntersectingActors(const FBox& Box, TArray<FWorldPartitionActorDesc*>& OutActors)
diff --git a/Engine/Source/Runtime/Engine/Public/WorldPartition/WorldPartitionEditorSpatialHash.h b/Engine/Source/Runtime/Engine/Public/WorldPartition/WorldPartitionEditorSpatialHash.h
--- a/Engine/Source/Runtime/Engine/Public/WorldPartition/WorldPartitionEditorSpatialHash.h
+++ b/Engine/Source/Runtime/Engine/Public/WorldPartition/WorldPartitionEditorSpatialHash.h
@@ -231,8 +231,8 @@ public:
 	virtual uint32 GetWantedEditorCellSize() const override;
 	virtual void SetEditorWantedCellSize(uint32 InCellSize) override;
 
-	virtual void AddBackReference(FWorldPartitionHandle& ReferenceHandle,  UWorldPartitionEditorCell* Cell, const FGuid& Source) override;
-	virtual void RemoveBackReference(FWorldPartitionHandle& ReferenceHandle,  UWorldPartitionEditorCell* Cell, const FGuid& Source) override;
+	virtual void AddBackReference(const FGuid& Reference, UWorldPartitionEditorCell* Cell, const FGuid& Source) override;
+	virtual void RemoveBackReference(const FGuid& Reference, UWorldPartitionEditorCell* Cell, const FGuid& Source) override;
 		
 	// UWorldPartitionEditorHash interface end
 #endif
@@ -255,7 +255,7 @@ private:
 	TMap<FCellCoord, FCellNode> HashNodes;
 	TMap<FCellCoord, UWorldPartitionEditorCell*> HashCells;
 	
-	TMultiMap<FWorldPartitionHandle, TTuple<UWorldPartitionEditorCell*, FGuid>>	BackReferences;
+	TMultiMap<FGuid, TTuple<UWorldPartitionEditorCell*, FGuid>>	BackReferences;
 
 	UPROPERTY(Transient)
 	TSet<TObjectPtr<UWorldPartitionEditorCell>> Cells;
