-Added ParentActor to WorldPartitionActorDesc so that we can validate WP settings coherency on attachments in Map Check & other validations (without loading the whole map)
(validation is coming in a seperate CL)

diff --git a/Engine/Source/Runtime/Core/Public/UObject/UE5ReleaseStreamObjectVersion.h b/Engine/Source/Runtime/Core/Public/UObject/UE5ReleaseStreamObjectVersion.h
--- a/Engine/Source/Runtime/Core/Public/UObject/UE5ReleaseStreamObjectVersion.h
+++ b/Engine/Source/Runtime/Core/Public/UObject/UE5ReleaseStreamObjectVersion.h
@@ -60,6 +60,9 @@ struct CORE_API FUE5ReleaseStreamObjectVersion
 		// Add support for different external actor packaging schemes
 		AddLevelActorPackagingScheme,
 
+		// Add support for linking to the attached parent actor in WorldPartitionActorDesc
+		WorldPartitionActorDescSerializeAttachParent,
+
 		// -----<new versions can be added above this line>-------------------------------------------------
 		VersionPlusOne,
 		LatestVersion = VersionPlusOne - 1
diff --git a/Engine/Source/Runtime/Engine/Private/WorldPartition/WorldPartitionActorDesc.cpp b/Engine/Source/Runtime/Engine/Private/WorldPartition/WorldPartitionActorDesc.cpp
--- a/Engine/Source/Runtime/Engine/Private/WorldPartition/WorldPartitionActorDesc.cpp
+++ b/Engine/Source/Runtime/Engine/Private/WorldPartition/WorldPartitionActorDesc.cpp
@@ -64,6 +64,16 @@ void FWorldPartitionActorDesc::Init(const AActor* InActor)
 	ActorPackage = InActor->GetPackage()->GetFName();
 	ActorPath = *InActor->GetPathName();
 	FolderPath = InActor->GetFolderPath();
+
+	const AActor* AttachParentActor = InActor->GetAttachParentActor();
+	if (AttachParentActor)
+	{
+		ParentActor = AttachParentActor->GetActorGuid();
+	}
+	else
+	{
+		ParentActor = FGuid();
+	}
 	
 	TArray<AActor*> ActorReferences = ActorsReferencesUtils::GetExternalActorReferences((AActor*)InActor);
 
@@ -231,6 +241,11 @@ void FWorldPartitionActorDesc::Serialize(FArchive& Ar)
 	{
 		Ar << FolderPath;
 	}
+
+	if (Ar.CustomVer(FUE5ReleaseStreamObjectVersion::GUID) >= FUE5ReleaseStreamObjectVersion::WorldPartitionActorDescSerializeAttachParent)
+	{
+		Ar << ParentActor;
+	}
 }
 
 FBox FWorldPartitionActorDesc::GetBounds() const
diff --git a/Engine/Source/Runtime/Engine/Private/WorldPartition/WorldPartitionActorDescView.cpp b/Engine/Source/Runtime/Engine/Private/WorldPartition/WorldPartitionActorDescView.cpp
--- a/Engine/Source/Runtime/Engine/Private/WorldPartition/WorldPartitionActorDescView.cpp
+++ b/Engine/Source/Runtime/Engine/Private/WorldPartition/WorldPartitionActorDescView.cpp
@@ -89,6 +89,12 @@ FName FWorldPartitionActorDescView::GetActorLabel() const
 	return ActorDesc->GetActorLabel();
 }
 
+FName FWorldPartitionActorDescView::GetActorName() const
+{
+	return ActorDesc->GetActorName();
+}
+
+
 FBox FWorldPartitionActorDescView::GetBounds() const
 {
 	return ActorDesc->GetBounds();
@@ -109,6 +115,11 @@ uint32 FWorldPartitionActorDescView::GetTag() const
 	return ActorDesc->Tag;
 }
 
+const FGuid& FWorldPartitionActorDescView::GetParentActor() const
+{
+	return ActorDesc->GetParentActor();
+}
+
 bool FWorldPartitionActorDescView::GetContainerInstance(const UActorDescContainer*& OutLevelContainer, FTransform& OutLevelTransform, EContainerClusterMode& OutClusterMode) const
 {
 	return ActorDesc->GetContainerInstance(OutLevelContainer, OutLevelTransform, OutClusterMode);
diff --git a/Engine/Source/Runtime/Engine/Public/WorldPartition/WorldPartitionActorDesc.h b/Engine/Source/Runtime/Engine/Public/WorldPartition/WorldPartitionActorDesc.h
--- a/Engine/Source/Runtime/Engine/Public/WorldPartition/WorldPartitionActorDesc.h
+++ b/Engine/Source/Runtime/Engine/Public/WorldPartition/WorldPartitionActorDesc.h
@@ -74,6 +74,7 @@ public:
 	inline FName GetActorLabel() const { return ActorLabel; }
 	inline FName GetFolderPath() const { return FolderPath; }
 	FBox GetBounds() const;
+	inline const FGuid& GetParentActor() const { return ParentActor; }
 
 	FName GetActorName() const;
 
@@ -180,6 +181,8 @@ protected:
 	TArray<FName>					DataLayers;
 	TArray<FGuid>					References;
 	FName							FolderPath;
+
+	FGuid							ParentActor; // Used to validate settings against parent (to warn on layer/placement compatibility issues)
 	
 	// Transient
 	mutable uint32					SoftRefCount;
diff --git a/Engine/Source/Runtime/Engine/Public/WorldPartition/WorldPartitionActorDescView.h b/Engine/Source/Runtime/Engine/Public/WorldPartition/WorldPartitionActorDescView.h
--- a/Engine/Source/Runtime/Engine/Public/WorldPartition/WorldPartitionActorDescView.h
+++ b/Engine/Source/Runtime/Engine/Public/WorldPartition/WorldPartitionActorDescView.h
@@ -41,6 +41,8 @@ public:
 	const TArray<FGuid>& GetReferences() const;
 	FString ToString() const;
 	uint32 GetTag() const;
+	const FGuid& GetParentActor() const;
+	FName GetActorName() const;
 
 	bool GetContainerInstance(const UActorDescContainer*& OutLevelContainer, FTransform& OutLevelTransform, EContainerClusterMode& OutClusterMode) const;
 
