Fix for virtual textures not applying RDO when compressed with Oodle.

Since the properties weren't copied, this means that the DDC key doesn't represent the actual bits, so the VT ddc key is invalidated.

diff --git a/Engine/Source/Runtime/Engine/Private/TextureDerivedData.cpp b/Engine/Source/Runtime/Engine/Private/TextureDerivedData.cpp
--- a/Engine/Source/Runtime/Engine/Private/TextureDerivedData.cpp
+++ b/Engine/Source/Runtime/Engine/Private/TextureDerivedData.cpp
@@ -65,8 +65,7 @@
 
 // This GUID is mixed into DDC version for virtual textures only, this allows updating DDC version for VT without invalidating DDC for all textures
 // This is useful during development, but once large numbers of VT are present in shipped content, it will have the same problem as TEXTURE_DERIVEDDATA_VER
-#define TEXTURE_VT_DERIVEDDATA_VER	TEXT("93F54B64FA524F78BCA4EA2CD30B1E9F")
-
+#define TEXTURE_VT_DERIVEDDATA_VER	TEXT("A5849A269D64445585DDD909F2D49882")
 
 #if ENABLE_COOK_STATS
 namespace TextureCookStats
diff --git a/Engine/Source/Runtime/Engine/Private/VT/VirtualTextureDataBuilder.cpp b/Engine/Source/Runtime/Engine/Private/VT/VirtualTextureDataBuilder.cpp
--- a/Engine/Source/Runtime/Engine/Private/VT/VirtualTextureDataBuilder.cpp
+++ b/Engine/Source/Runtime/Engine/Private/VT/VirtualTextureDataBuilder.cpp
@@ -701,6 +701,8 @@ void FVirtualTextureDataBuilder::BuildTiles(const TArray<FVTSourceTileEntry>& Ti
 		TBSettings.OodleEncodeEffort = BuildSettingsForLayer.OodleEncodeEffort;
 		TBSettings.LossyCompressionAmount = BuildSettingsForLayer.LossyCompressionAmount;
 		TBSettings.OodleUniversalTiling = BuildSettingsForLayer.OodleUniversalTiling;
+		TBSettings.bOodleUsesRDO = BuildSettingsForLayer.bOodleUsesRDO;
+		TBSettings.OodleRDO = BuildSettingsForLayer.OodleRDO;
 
 		check(TBSettings.GetGammaSpace() == BuildSettingsForLayer.GetGammaSpace());
  
