Adding ability to override buffer size and number of buffers in device init
* setting override to a small buffer for non-realtime audio rendering

diff --git a/Engine/Source/Editor/UnrealEd/Private/PlayLevel.cpp b/Engine/Source/Editor/UnrealEd/Private/PlayLevel.cpp
--- a/Engine/Source/Editor/UnrealEd/Private/PlayLevel.cpp
+++ b/Engine/Source/Editor/UnrealEd/Private/PlayLevel.cpp
@@ -2853,11 +2853,9 @@ UGameInstance* UEditorEngine::CreateInnerProcessPIEGameInstance(FRequestPlaySess
 
 	if (!InPIEParameters.bRunAsDedicated)
 	{
-		bool bCreateNewAudioDevice = InParams.EditorPlaySettings->IsCreateAudioDeviceForEveryPlayer();
-
 		// Create an instance of the Game Viewport Client, with the class specified by the Engine.
 		ViewportClient = NewObject<UGameViewportClient>(this, GameViewportClientClass);
-		ViewportClient->Init(*PieWorldContext, GameInstance, bCreateNewAudioDevice);
+		ViewportClient->Init(*PieWorldContext, GameInstance);
 
 		ULevelEditorPlaySettings* PlayInSettings = GetMutableDefault<ULevelEditorPlaySettings>();
 		ViewportClient->EngineShowFlags.SetServerDrawDebug(PlayInSettings->ShowServerDebugDrawingByDefault());
@@ -2977,6 +2975,9 @@ UGameInstance* UEditorEngine::CreateInnerProcessPIEGameInstance(FRequestPlaySess
 			DeviceParams.Scope = EAudioDeviceScope::Unique;
 			DeviceParams.AssociatedWorld = PlayWorld;
 			DeviceParams.bIsNonRealtime = true;
+			// For NRT rendering, don't need a large buffer and don't need to double buffer
+			DeviceParams.BufferSizeOverride = 32;
+			DeviceParams.NumBuffersOverride = 2;
 			FAudioDeviceHandle AudioDevice = AudioDeviceManager->RequestAudioDevice(DeviceParams);
 			check(AudioDevice.IsValid());
 			if (PlayWorld)
diff --git a/Engine/Source/Runtime/Engine/Private/AudioDevice.cpp b/Engine/Source/Runtime/Engine/Private/AudioDevice.cpp
--- a/Engine/Source/Runtime/Engine/Private/AudioDevice.cpp
+++ b/Engine/Source/Runtime/Engine/Private/AudioDevice.cpp
@@ -333,7 +333,7 @@ const FAudioQualitySettings& FAudioDevice::GetQualityLevelSettings()
 	return AudioSettings->GetQualityLevelSettings(QualityLevel);
 }
 
-bool FAudioDevice::Init(Audio::FDeviceId InDeviceID, int32 InMaxSources)
+bool FAudioDevice::Init(Audio::FDeviceId InDeviceID, int32 InMaxSources, int32 InBufferSizeOverride, int32 InNumBuffersOverride)
 {
 	SCOPED_BOOT_TIMING("FAudioDevice::Init");
 
@@ -369,6 +369,21 @@ bool FAudioDevice::Init(Audio::FDeviceId InDeviceID, int32 InMaxSources)
 		UE_LOG(LogAudioMixer, Display, TEXT("Command Line NumBuffersToEnqueue Override: %d"), PlatformSettings.NumBuffers);
 	}
 
+	// Override arguments passed into this function trump command line argument overrides
+	if (InBufferSizeOverride != INDEX_NONE)
+	{
+		PlatformSettings.CallbackBufferFrameSize = InBufferSizeOverride;
+	}
+	if (InNumBuffersOverride != INDEX_NONE)
+	{
+		PlatformSettings.NumBuffers = InNumBuffersOverride;
+	}
+
+	// Validate buffer size and num buffers
+ 	PlatformSettings.CallbackBufferFrameSize = FMath::Max(4, PlatformSettings.CallbackBufferFrameSize);
+ 	PlatformSettings.CallbackBufferFrameSize = PlatformSettings.CallbackBufferFrameSize - (PlatformSettings.CallbackBufferFrameSize % 4);
+ 	PlatformSettings.NumBuffers = FMath::Max(1, PlatformSettings.NumBuffers);
+
 	// MaxSources is the max value supplied to Init call (quality settings), unless overwritten by the platform settings.
 	// This does not have to be the minimum value in this case (nor is it desired, so platforms can potentially scale up)
 	// as the Sources array has yet to be initialized. If the cvar is largest, take that value to allow for testing
diff --git a/Engine/Source/Runtime/Engine/Private/AudioDeviceManager.cpp b/Engine/Source/Runtime/Engine/Private/AudioDeviceManager.cpp
--- a/Engine/Source/Runtime/Engine/Private/AudioDeviceManager.cpp
+++ b/Engine/Source/Runtime/Engine/Private/AudioDeviceManager.cpp
@@ -1692,7 +1692,7 @@ FAudioDeviceManager::FAudioDeviceContainer::FAudioDeviceContainer(const FAudioDe
 	// runtime is supported.
 	const UAudioSettings* AudioSettings = GetDefault<UAudioSettings>();
 	const int32 HighestMaxChannels = AudioSettings ? AudioSettings->GetHighestMaxChannels() : 0;
-	if (Device->Init(InDeviceID, HighestMaxChannels))
+	if (Device->Init(InDeviceID, HighestMaxChannels, InParams.BufferSizeOverride, InParams.NumBuffersOverride))
 	{
 		const FAudioQualitySettings& QualitySettings = Device->GetQualityLevelSettings();
 		Device->SetMaxChannels(QualitySettings.MaxChannels);
diff --git a/Engine/Source/Runtime/Engine/Public/AudioDevice.h b/Engine/Source/Runtime/Engine/Public/AudioDevice.h
--- a/Engine/Source/Runtime/Engine/Public/AudioDevice.h
+++ b/Engine/Source/Runtime/Engine/Public/AudioDevice.h
@@ -514,7 +514,7 @@ public:
 	/**
 	 * Basic initialization of the platform agnostic layer of the audio system
 	 */
-	bool Init(Audio::FDeviceId InDeviceID, int32 InMaxSources);
+	bool Init(Audio::FDeviceId InDeviceID, int32 InMaxSources, int32 BufferSizeOverride = INDEX_NONE, int32 NumBuffersOverride = INDEX_NONE);
 
 	/**
 	 * Called after FAudioDevice creation and init
diff --git a/Engine/Source/Runtime/Engine/Public/AudioDeviceManager.h b/Engine/Source/Runtime/Engine/Public/AudioDeviceManager.h
--- a/Engine/Source/Runtime/Engine/Public/AudioDeviceManager.h
+++ b/Engine/Source/Runtime/Engine/Public/AudioDeviceManager.h
@@ -63,20 +63,17 @@ enum class EAudioDeviceScope : uint8
 struct FAudioDeviceParams
 {
 	// Optional world parameter. This allows tools to surface information about which worlds are being rendered through which audio devices.
-	UWorld* AssociatedWorld;
+	UWorld* AssociatedWorld = nullptr;
 	// This should be set to EAudioDeviceScope::Unique if you'd like to force a new device to be created from scratch, or use EAudioDeviceScope::Shared to use an existing device if possible.
-	EAudioDeviceScope Scope;
+	EAudioDeviceScope Scope = EAudioDeviceScope::Default;
 	// Set this to true to get a handle to a non realtime audio renderer.
-	bool bIsNonRealtime;
+	bool bIsNonRealtime = false;
 	// Use this to force this audio device to use a specific audio module. If nullptr, uses the default audio module.
-	IAudioDeviceModule* AudioModule;
-
-	FAudioDeviceParams()
-		: AssociatedWorld(nullptr)
-		, Scope(EAudioDeviceScope::Default)
-		, bIsNonRealtime(false)
-		, AudioModule(nullptr)
-	{}
+	IAudioDeviceModule* AudioModule = nullptr;
+	// Buffer size override
+	int32 BufferSizeOverride = INDEX_NONE;
+	// Num buffers override
+	int32 NumBuffersOverride = INDEX_NONE;
 };
 
 // Strong handle to an audio device. Guarantees that the audio device it references will stay alive while it is in scope.
