Fix potential out-of-bound read when skin cache is running out of memory.

diff --git a/Engine/Source/Runtime/Engine/Private/GPUSkinCache.cpp b/Engine/Source/Runtime/Engine/Private/GPUSkinCache.cpp
--- a/Engine/Source/Runtime/Engine/Private/GPUSkinCache.cpp
+++ b/Engine/Source/Runtime/Engine/Private/GPUSkinCache.cpp
@@ -369,18 +369,21 @@ public:
 	inline FCachedGeometry::Section GetCachedGeometry(int32 SectionIndex) const
 	{
 		FCachedGeometry::Section MeshSection;
-		const FSkelMeshRenderSection& Section = *DispatchData[SectionIndex].Section;
-		MeshSection.PositionBuffer	= DispatchData[SectionIndex].PositionBuffer->Buffer.SRV;
-		MeshSection.UVsBuffer		= DispatchData[SectionIndex].UVsBufferSRV;
-		MeshSection.TotalVertexCount= DispatchData[SectionIndex].PositionBuffer->Buffer.NumBytes / (sizeof(float)*3);
-		MeshSection.NumPrimitives	= Section.NumTriangles;
-		MeshSection.NumVertices		= Section.NumVertices;
-		MeshSection.IndexBaseIndex	= Section.BaseIndex;
-		MeshSection.VertexBaseIndex = Section.BaseVertexIndex;
-		MeshSection.IndexBuffer		= nullptr;
-		MeshSection.TotalIndexCount	= 0;
-		MeshSection.LODIndex		= 0;
-		MeshSection.SectionIndex	= SectionIndex;
+		if (SectionIndex >= 0 && SectionIndex < DispatchData.Num())
+		{
+			const FSkelMeshRenderSection& Section = *DispatchData[SectionIndex].Section;
+			MeshSection.PositionBuffer = DispatchData[SectionIndex].PositionBuffer->Buffer.SRV;
+			MeshSection.UVsBuffer = DispatchData[SectionIndex].UVsBufferSRV;
+			MeshSection.TotalVertexCount = DispatchData[SectionIndex].PositionBuffer->Buffer.NumBytes / (sizeof(float) * 3);
+			MeshSection.NumPrimitives = Section.NumTriangles;
+			MeshSection.NumVertices = Section.NumVertices;
+			MeshSection.IndexBaseIndex = Section.BaseIndex;
+			MeshSection.VertexBaseIndex = Section.BaseVertexIndex;
+			MeshSection.IndexBuffer = nullptr;
+			MeshSection.TotalIndexCount = 0;
+			MeshSection.LODIndex = 0;
+			MeshSection.SectionIndex = SectionIndex;
+		}
 		return MeshSection;
 	}
 
@@ -2298,21 +2301,24 @@ FCachedGeometry FGPUSkinCache::GetCachedGeometry(uint32 ComponentId) const
 	{
 		if (Entry && Entry->GPUSkin && Entry->GPUSkin->GetComponentId() == ComponentId && Entry->GPUSkin->HaveValidDynamicData())
 		{
-			const uint32 LODIndex = Entry->LOD;
 			const FSkeletalMeshRenderData& RenderData = Entry->GPUSkin->GetSkeletalMeshRenderData();
-			const FSkeletalMeshLODRenderData& LODData = RenderData.LODRenderData[LODIndex];
-			const uint32 SectionCount = LODData.RenderSections.Num();
-			for (uint32 SectionIdx=0; SectionIdx< SectionCount;++SectionIdx)
+			const int32 LODIndex = Entry->LOD;
+			if (LODIndex >= 0 && LODIndex < RenderData.LODRenderData.Num())
 			{
-				FCachedGeometry::Section CachedSection = Entry->GetCachedGeometry(SectionIdx);
-				CachedSection.IndexBuffer		= LODData.MultiSizeIndexContainer.GetIndexBuffer()->GetSRV();
-				CachedSection.TotalIndexCount	= LODData.MultiSizeIndexContainer.GetIndexBuffer()->Num();
-				CachedSection.LODIndex			= LODIndex;
-				CachedSection.UVsChannelOffset	= 0; // Assume that we needs to pair meshes based on UVs 0
-				CachedSection.UVsChannelCount	= LODData.StaticVertexBuffers.StaticMeshVertexBuffer.GetNumTexCoords();
-				Out.Sections.Add(CachedSection);
+				const FSkeletalMeshLODRenderData& LODData = RenderData.LODRenderData[LODIndex];
+				const uint32 SectionCount = LODData.RenderSections.Num();
+				for (uint32 SectionIdx=0; SectionIdx< SectionCount;++SectionIdx)
+				{
+					FCachedGeometry::Section CachedSection = Entry->GetCachedGeometry(SectionIdx);
+					CachedSection.IndexBuffer		= LODData.MultiSizeIndexContainer.GetIndexBuffer()->GetSRV();
+					CachedSection.TotalIndexCount	= LODData.MultiSizeIndexContainer.GetIndexBuffer()->Num();
+					CachedSection.LODIndex			= LODIndex;
+					CachedSection.UVsChannelOffset	= 0; // Assume that we needs to pair meshes based on UVs 0
+					CachedSection.UVsChannelCount	= LODData.StaticVertexBuffers.StaticMeshVertexBuffer.GetNumTexCoords();
+					Out.Sections.Add(CachedSection);
+				}
+				break;
 			}
-			break;
 		}
 	}
 
