Added profile tags for worldpartition debug draw

diff --git a/Engine/Source/Runtime/Engine/Private/WorldPartition/DataLayer/DataLayerSubsystem.cpp b/Engine/Source/Runtime/Engine/Private/WorldPartition/DataLayer/DataLayerSubsystem.cpp
--- a/Engine/Source/Runtime/Engine/Private/WorldPartition/DataLayer/DataLayerSubsystem.cpp
+++ b/Engine/Source/Runtime/Engine/Private/WorldPartition/DataLayer/DataLayerSubsystem.cpp
@@ -221,6 +221,8 @@ bool UDataLayerSubsystem::IsAnyDataLayerInEffectiveRuntimeState(const TArray<FNa
 
 void UDataLayerSubsystem::DrawDataLayersStatus(UCanvas* Canvas, FVector2D& Offset) const
 {
+	TRACE_CPUPROFILER_EVENT_SCOPE(UDataLayerSubsystem::DrawDataLayersStatus);
+
 	if (!Canvas || !Canvas->SceneView)
 	{
 		return;
diff --git a/Engine/Source/Runtime/Engine/Private/WorldPartition/WorldPartitionLevelStreamingPolicy.cpp b/Engine/Source/Runtime/Engine/Private/WorldPartition/WorldPartitionLevelStreamingPolicy.cpp
--- a/Engine/Source/Runtime/Engine/Private/WorldPartition/WorldPartitionLevelStreamingPolicy.cpp
+++ b/Engine/Source/Runtime/Engine/Private/WorldPartition/WorldPartitionLevelStreamingPolicy.cpp
@@ -179,6 +179,8 @@ UObject* UWorldPartitionLevelStreamingPolicy::GetSubObject(const TCHAR* SubObjec
 
 void UWorldPartitionLevelStreamingPolicy::DrawRuntimeCellsDetails(UCanvas* Canvas, FVector2D& Offset)
 {
+	TRACE_CPUPROFILER_EVENT_SCOPE(UWorldPartitionLevelStreamingPolicy::DrawRuntimeCellsDetails);
+
 	UWorld* World = WorldPartition->GetWorld();
 	struct FCellsPerStreamingStatus
 	{
@@ -236,6 +238,8 @@ void UWorldPartitionLevelStreamingPolicy::DrawRuntimeCellsDetails(UCanvas* Canva
  */
 void UWorldPartitionLevelStreamingPolicy::DrawStreamingStatusLegend(UCanvas* Canvas, FVector2D& Offset)
 {
+	TRACE_CPUPROFILER_EVENT_SCOPE(UWorldPartitionLevelStreamingPolicy::DrawStreamingStatusLegend);
+
 	check(Canvas);
 
 	// Cumulate counter stats
diff --git a/Engine/Source/Runtime/Engine/Private/WorldPartition/WorldPartitionRuntimeSpatialHash.cpp b/Engine/Source/Runtime/Engine/Private/WorldPartition/WorldPartitionRuntimeSpatialHash.cpp
--- a/Engine/Source/Runtime/Engine/Private/WorldPartition/WorldPartitionRuntimeSpatialHash.cpp
+++ b/Engine/Source/Runtime/Engine/Private/WorldPartition/WorldPartitionRuntimeSpatialHash.cpp
@@ -408,6 +408,8 @@ void FSpatialHashStreamingGrid::DrawStreamingSource3D(UWorld* World, const FSphe
 
 void FSpatialHashStreamingGrid::Draw2D(UCanvas* Canvas, UWorld* World, const TArray<FWorldPartitionStreamingSource>& Sources, const FBox& Region, const FBox2D& GridScreenBounds, TFunctionRef<FVector2D(const FVector2D&)> WorldToScreen) const
 {
+	TRACE_CPUPROFILER_EVENT_SCOPE(FSpatialHashStreamingGrid::Draw2D);
+
 	const EWorldPartitionRuntimeCellVisualizeMode VisualizeMode = GetStreamingCellVisualizeMode();
 	const UDataLayerSubsystem* DataLayerSubsystem = World->GetSubsystem<UDataLayerSubsystem>();
 	TMap<FName, FColor> DataLayerDebugColors;
@@ -1175,6 +1177,8 @@ EWorldPartitionStreamingPerformance UWorldPartitionRuntimeSpatialHash::GetStream
 
 void UWorldPartitionRuntimeSpatialHash::Draw2D(UCanvas* Canvas, const TArray<FWorldPartitionStreamingSource>& Sources, const FVector2D& PartitionCanvasSize, FVector2D& Offset) const
 {
+	TRACE_CPUPROFILER_EVENT_SCOPE(UWorldPartitionRuntimeSpatialHash::Draw2D);
+
 	TArray<const FSpatialHashStreamingGrid*> FilteredStreamingGrids = GetFilteredStreamingGrids();
 	if (FilteredStreamingGrids.Num() == 0 || Sources.Num() == 0)
 	{
diff --git a/Engine/Source/Runtime/Engine/Private/WorldPartition/WorldPartitionSubsystem.cpp b/Engine/Source/Runtime/Engine/Private/WorldPartition/WorldPartitionSubsystem.cpp
--- a/Engine/Source/Runtime/Engine/Private/WorldPartition/WorldPartitionSubsystem.cpp
+++ b/Engine/Source/Runtime/Engine/Private/WorldPartition/WorldPartitionSubsystem.cpp
@@ -260,6 +260,7 @@ void UWorldPartitionSubsystem::UpdateStreamingState()
 
 void UWorldPartitionSubsystem::Draw(UCanvas* Canvas, class APlayerController* PC)
 {
+	TRACE_CPUPROFILER_EVENT_SCOPE(UWorldPartitionSubsystem::Draw);
 	if (!Canvas || !Canvas->SceneView)
 	{
 		return;
@@ -327,6 +328,8 @@ void UWorldPartitionSubsystem::Draw(UCanvas* Canvas, class APlayerController* PC
 
 	if (GDrawStreamingSources || GDrawRuntimeHash2D)
 	{
+		TRACE_CPUPROFILER_EVENT_SCOPE(UWorldPartitionSubsystem::DrawStreamingSources);
+
 		const UWorldPartition* WorldPartition = GetWorldPartition();
 		const TArray<FWorldPartitionStreamingSource>* StreamingSources = WorldPartition ? &WorldPartition->GetStreamingSources() : nullptr;
 		if (StreamingSources && (StreamingSources->Num() > 0))
