Expose GetPredictedLODLevel() to BP

diff --git a/Engine/Source/Runtime/Engine/Classes/Components/SkinnedMeshComponent.h b/Engine/Source/Runtime/Engine/Classes/Components/SkinnedMeshComponent.h
--- a/Engine/Source/Runtime/Engine/Classes/Components/SkinnedMeshComponent.h
+++ b/Engine/Source/Runtime/Engine/Classes/Components/SkinnedMeshComponent.h
@@ -1122,6 +1122,7 @@ public:
 	virtual bool UpdateLODStatus();
 
 	/** Get predicted LOD level. This value is usually calculated in UpdateLODStatus, but can be modified by skeletal mesh streaming. */
+	UFUNCTION(BlueprintCallable, Category = "Components|SkinnedMesh")
 	int32 GetPredictedLODLevel() const
 	{
 		PRAGMA_DISABLE_DEPRECATION_WARNINGS
