Add more vertex buffer memory stats to skeletal mesh report

diff --git a/Engine/Source/Runtime/Engine/Private/UnrealEngine.cpp b/Engine/Source/Runtime/Engine/Private/UnrealEngine.cpp
--- a/Engine/Source/Runtime/Engine/Private/UnrealEngine.cpp
+++ b/Engine/Source/Runtime/Engine/Private/UnrealEngine.cpp
@@ -7009,6 +7009,15 @@ bool UEngine::HandleSkeletalMeshReportCommand(const TCHAR* Cmd, FOutputDevice& A
 
 				int32 MorphDataKB = GetResourceSizeKB(RenderData.MorphTargetVertexInfoBuffers.MorphDataBuffer);
 				int32 ClothVertexBufferKB = GetResourceSizeKB(RenderData.ClothVertexBuffer.VertexBufferRHI);
+				int32 SkinWeightDataVertexBufferKB = GetResourceSizeKB(RenderData.SkinWeightVertexBuffer.GetDataVertexBuffer()->VertexBufferRHI);
+				int32 SkinWeightLookupVertexBufferKB = GetResourceSizeKB(RenderData.SkinWeightVertexBuffer.GetLookupVertexBuffer()->VertexBufferRHI);
+				int32 PositionVertexBufferKB = GetResourceSizeKB(RenderData.StaticVertexBuffers.PositionVertexBuffer.VertexBufferRHI);
+				int32 TangentVertexBufferKB = GetResourceSizeKB(RenderData.StaticVertexBuffers.StaticMeshVertexBuffer.TangentsVertexBuffer.VertexBufferRHI);
+				int32 TexCoordVertexBufferKB = GetResourceSizeKB(RenderData.StaticVertexBuffers.StaticMeshVertexBuffer.TexCoordVertexBuffer.VertexBufferRHI);
+				int32 IndexBufferKB = GetResourceSizeKB(RenderData.MultiSizeIndexContainer.GetIndexBuffer()->IndexBufferRHI);
+				int32 ColorBufferKB = GetResourceSizeKB(RenderData.StaticVertexBuffers.ColorVertexBuffer.VertexBufferRHI);
+				int32 TotalResourceKB = MorphDataKB + ClothVertexBufferKB + SkinWeightDataVertexBufferKB + SkinWeightLookupVertexBufferKB + PositionVertexBufferKB + TangentVertexBufferKB 
+										+ TexCoordVertexBufferKB + IndexBufferKB + ColorBufferKB;
 
 				FString RecomputeTangentSections = TEXT("");
 				for (int32 SectionIdx = 0; SectionIdx < RenderData.RenderSections.Num(); SectionIdx++)
@@ -7032,12 +7041,25 @@ bool UEngine::HandleSkeletalMeshReportCommand(const TCHAR* Cmd, FOutputDevice& A
 					RecomputeTangentSections = TEXT("Off");
 				}
 
-				Ar.Logf(TEXT("  [LOD%d]: RefSkeletonBones=%3d, RequiredBones=%3d, ActiveBones=%3d, MaxBoneInfluences=%2d, NumMorphTargets=%7d, MorphData=%4dKB, ClothVertexBuffer=%4dKB, SkinCache=%s, RecomputeTangent=%s, SupportRayTracing=%s, RayTracingMinLOD=%d"), LODIdx,
+				const int32 MaxBoneInfluences = RenderData.SkinWeightVertexBuffer.GetMaxBoneInfluences();
+				const bool bUseUnlimitedBoneInfluences = FGPUBaseSkinVertexFactory::UseUnlimitedBoneInfluences(MaxBoneInfluences);
+
+				Ar.Logf(TEXT("  [LOD%d]: RefSkeletonBones=%3d, RequiredBones=%3d, ActiveBones=%3d, NumMorphTargets=%4d, MaxBoneInfluences=%2d, 16BitBoneIndex=%s, UnlimitedBoneInfluence=%s, TotalResource=%5dKB, IndexBuffer=%4dKB, PositionVB=%4dKB, TangentVB=%4dKB, TexCoordVB=%4dKB, ColorVB=%4dKB, SkinWeightDataVB=%4dKB, SkinWeightLookupVB=%4dKB, MorphData=%4dKB, ClothVB=%4dKB, SkinCache=%s, RecomputeTangent=%s, SupportRayTracing=%s, RayTracingMinLOD=%d"), LODIdx,
 					Mesh->GetRefSkeleton().GetNum(),
 					RenderData.RequiredBones.Num(),
 					RenderData.ActiveBoneIndices.Num(),
-					RenderData.SkinWeightVertexBuffer.GetMaxBoneInfluences(),
 					Mesh->GetMorphTargets().Num(),
+					RenderData.SkinWeightVertexBuffer.GetMaxBoneInfluences(),
+					RenderData.SkinWeightVertexBuffer.Use16BitBoneIndex() ? TEXT("Y") : TEXT("N"),
+					bUseUnlimitedBoneInfluences ? TEXT("Y") : TEXT("N"),
+					TotalResourceKB,
+					IndexBufferKB,
+					PositionVertexBufferKB,
+					TangentVertexBufferKB,
+					TexCoordVertexBufferKB,
+					ColorBufferKB,
+					SkinWeightDataVertexBufferKB,
+					SkinWeightLookupVertexBufferKB,
 					MorphDataKB,
 					ClothVertexBufferKB,
 					bSkinCachedLODEnabled ? TEXT("Y") : TEXT("N"),
