Add if we are running on a listen server when displaying world partition streaming status, since it forces all cells to be loaded.

diff --git a/Engine/Source/Runtime/Engine/Private/WorldPartition/WorldPartitionSubsystem.cpp b/Engine/Source/Runtime/Engine/Private/WorldPartition/WorldPartitionSubsystem.cpp
--- a/Engine/Source/Runtime/Engine/Private/WorldPartition/WorldPartitionSubsystem.cpp
+++ b/Engine/Source/Runtime/Engine/Private/WorldPartition/WorldPartitionSubsystem.cpp
@@ -298,6 +298,13 @@ void UWorldPartitionSubsystem::Draw(UCanvas* Canvas, class APlayerController* PC
 			if (IsIncrementalUnhashPending()) { StatusText += TEXT("(Unhashing) "); }
 			if (IsAsyncLoading()) { StatusText += TEXT("(AsyncLoading) "); }
 			if (StatusText.IsEmpty()) { StatusText = TEXT("(Idle)"); }
+
+			const ENetMode NetMode = GetWorld()->GetNetMode();
+			if (NetMode == NM_ListenServer)
+			{
+				StatusText += TEXT("(ListenServer)");
+			}
+			
 			const FString Text = FString::Printf(TEXT("Streaming Status: %s"), *StatusText);
 			FWorldPartitionDebugHelper::DrawText(Canvas, Text, GEngine->GetSmallFont(), FColor::White, CurrentOffset);
 		}
