Added wp.Runtime.FilterRuntimeSpatialHashGridLevel to filter grid cells based on their hierarchical level in the grid.

diff --git a/Engine/Source/Runtime/Engine/Private/WorldPartition/WorldPartitionRuntimeSpatialHash.cpp b/Engine/Source/Runtime/Engine/Private/WorldPartition/WorldPartitionRuntimeSpatialHash.cpp
--- a/Engine/Source/Runtime/Engine/Private/WorldPartition/WorldPartitionRuntimeSpatialHash.cpp
+++ b/Engine/Source/Runtime/Engine/Private/WorldPartition/WorldPartitionRuntimeSpatialHash.cpp
@@ -79,6 +79,14 @@ static FAutoConsoleVariableRef CVarBlockOnSlowStreamingWarningFactor(
 	GBlockOnSlowStreamingWarningFactor,
 	TEXT("Factor of wp.Runtime.BlockOnSlowStreamingRatio we want to start notifying the user"));
 
+#if !UE_BUILD_SHIPPING
+static int32 GFilterRuntimeSpatialHashGridLevel = INDEX_NONE;
+static FAutoConsoleVariableRef CVarFilterRuntimeSpatialHashGridLevel(
+	TEXT("wp.Runtime.FilterRuntimeSpatialHashGridLevel"),
+	GFilterRuntimeSpatialHashGridLevel,
+	TEXT("Used to choose filter a single world partition runtime hash grid level."));
+#endif
+
 // ------------------------------------------------------------------------------------------------
 FSpatialHashStreamingGrid::FSpatialHashStreamingGrid()
 	: Origin(ForceInitToZero)
@@ -194,27 +202,32 @@ void FSpatialHashStreamingGrid::GetCells(const TArray<FWorldPartitionStreamingSo
 
 			Helper.ForEachIntersectingCells(Shape, [&](const FIntVector& Coords)
 			{ 
-				if (const int32* LayerCellIndexPtr = GridLevels[Coords.Z].LayerCellsMapping.Find(Coords.Y * Helper.Levels[Coords.Z].GridSize + Coords.X))
+#if !UE_BUILD_SHIPPING
+				if ((GFilterRuntimeSpatialHashGridLevel == INDEX_NONE) || (GFilterRuntimeSpatialHashGridLevel == Coords.Z))
+#endif
 				{
-					const FSpatialHashStreamingGridLayerCell& LayerCell = GridLevels[Coords.Z].LayerCells[*LayerCellIndexPtr];
-					for (const UWorldPartitionRuntimeCell* Cell : LayerCell.GridCells)
+					if (const int32* LayerCellIndexPtr = GridLevels[Coords.Z].LayerCellsMapping.Find(Coords.Y * Helper.Levels[Coords.Z].GridSize + Coords.X))
 					{
-						if (!Cell->HasDataLayers() || (DataLayerSubsystem && DataLayerSubsystem->IsAnyDataLayerInEffectiveRuntimeState(Cell->GetDataLayers(), EDataLayerRuntimeState::Activated)))
+						const FSpatialHashStreamingGridLayerCell& LayerCell = GridLevels[Coords.Z].LayerCells[*LayerCellIndexPtr];
+						for (const UWorldPartitionRuntimeCell* Cell : LayerCell.GridCells)
 						{
-							if (Source.TargetState == EStreamingSourceTargetState::Loaded)
+						if (!Cell->HasDataLayers() || (DataLayerSubsystem && DataLayerSubsystem->IsAnyDataLayerInEffectiveRuntimeState(Cell->GetDataLayers(), EDataLayerRuntimeState::Activated)))
 							{
-								OutLoadCells.AddCell(Cell, Info);
+								if (Source.TargetState == EStreamingSourceTargetState::Loaded)
+								{
+									OutLoadCells.AddCell(Cell, Info);
+								}
+								else
+								{
+									check(Source.TargetState == EStreamingSourceTargetState::Activated);
+									OutActivateCells.AddCell(Cell, Info);
+								}
 							}
-							else
+						else if (DataLayerSubsystem && DataLayerSubsystem->IsAnyDataLayerInEffectiveRuntimeState(Cell->GetDataLayers(), EDataLayerRuntimeState::Loaded))
 							{
-								check(Source.TargetState == EStreamingSourceTargetState::Activated);
-								OutActivateCells.AddCell(Cell, Info);
+								OutLoadCells.AddCell(Cell, Info);
 							}
 						}
-						else if (DataLayerSubsystem && DataLayerSubsystem->IsAnyDataLayerInEffectiveRuntimeState(Cell->GetDataLayers(), EDataLayerRuntimeState::Loaded))
-						{
-							OutLoadCells.AddCell(Cell, Info);
-						}
 					}
 				}
 			});
@@ -1240,8 +1253,14 @@ void UWorldPartitionRuntimeSpatialHash::Draw2D(UCanvas* Canvas, const TArray<FWo
 			FString GridInfoText = FString::Printf(TEXT("%s | %d m"), *StreamingGrid->GridName.ToString(), int32(StreamingGrid->GetLoadingRange() * 0.01f));
 			if (StreamingGrid->bClientOnlyVisible)
 			{
-				GridInfoText += "| Client Only";
+				GridInfoText += TEXT(" | Client Only");
 			}
+#if !UE_BUILD_SHIPPING
+			if (GFilterRuntimeSpatialHashGridLevel != INDEX_NONE)
+			{
+				GridInfoText += FString::Printf(TEXT(" | GridLevelFilter %d"), GFilterRuntimeSpatialHashGridLevel);
+			}
+#endif
 			Canvas->SetDrawColor(255, 255, 0);
 			Canvas->DrawText(GEngine->GetTinyFont(), GridInfoText, GridInfoPos.X, GridInfoPos.Y);
 		}
@@ -1284,4 +1303,4 @@ TArray<const FSpatialHashStreamingGrid*> UWorldPartitionRuntimeSpatialHash::GetF
 	return FilteredStreamingGrids;
 }
 
-#undef LOCTEXT_NAMESPACE
+#undef LOCTEXT_NAMESPACE
\ No newline at end of file
