Cache calls to Ar.CustomVer out ot for loops in FSoftSkinVertex serialization code, saves ~1s loading large maps.

diff --git a/Engine/Source/Runtime/Engine/Private/SkeletalMeshLODModel.cpp b/Engine/Source/Runtime/Engine/Private/SkeletalMeshLODModel.cpp
--- a/Engine/Source/Runtime/Engine/Private/SkeletalMeshLODModel.cpp
+++ b/Engine/Source/Runtime/Engine/Private/SkeletalMeshLODModel.cpp
@@ -82,9 +82,10 @@ FArchive& operator<<(FArchive& Ar, FSoftSkinVertex& V)
 
 	// serialize bone and weight uint8 arrays in order
 	// this is required when serializing as bulk data memory (see TArray::BulkSerialize notes)
+	const bool bBeforeIncreaseBoneIndexLimitPerChunk = Ar.CustomVer(FAnimObjectVersion::GUID) < FAnimObjectVersion::IncreaseBoneIndexLimitPerChunk;
 	for (uint32 InfluenceIndex = 0; InfluenceIndex < MAX_INFLUENCES_PER_STREAM; InfluenceIndex++)
 	{
-		if (Ar.IsLoading() && Ar.CustomVer(FAnimObjectVersion::GUID) < FAnimObjectVersion::IncreaseBoneIndexLimitPerChunk)
+		if (Ar.IsLoading() && bBeforeIncreaseBoneIndexLimitPerChunk)
 		{
 			uint8 BoneIndex = 0;
 			Ar << BoneIndex;
@@ -100,7 +101,7 @@ FArchive& operator<<(FArchive& Ar, FSoftSkinVertex& V)
 	{
 		for (uint32 InfluenceIndex = MAX_INFLUENCES_PER_STREAM; InfluenceIndex < EXTRA_BONE_INFLUENCES; InfluenceIndex++)
 		{
-			if (Ar.IsLoading() && Ar.CustomVer(FAnimObjectVersion::GUID) < FAnimObjectVersion::IncreaseBoneIndexLimitPerChunk)
+			if (Ar.IsLoading() && bBeforeIncreaseBoneIndexLimitPerChunk)
 			{
 				uint8 BoneIndex = 0;
 				Ar << BoneIndex;
