World Partition Editor
- Don't maintain loaded mask through the editor cells octree nodes anymore.
- Display empty cells darker, then partially loaded cells less dark than loaded cells.

diff --git a/Engine/Source/Editor/WorldPartitionEditor/Private/WorldPartition/SWorldPartitionEditorGridSpatialHash.cpp b/Engine/Source/Editor/WorldPartitionEditor/Private/WorldPartition/SWorldPartitionEditorGridSpatialHash.cpp
--- a/Engine/Source/Editor/WorldPartitionEditor/Private/WorldPartition/SWorldPartitionEditorGridSpatialHash.cpp
+++ b/Engine/Source/Editor/WorldPartitionEditor/Private/WorldPartition/SWorldPartitionEditorGridSpatialHash.cpp
@@ -32,8 +32,6 @@
 
 #define LOCTEXT_NAMESPACE "WorldPartitionEditor"
 
-TAutoConsoleVariable<bool> CVarDebugDrawOctree(TEXT("wp.Editor.DebugDrawOctree"), false, TEXT("Whether to debug draw the World Partition octree"), ECVF_Default);
-
 SWorldPartitionEditorGridSpatialHash::SWorldPartitionEditorGridSpatialHash()
 	: WorldMiniMapBounds(ForceInit)
 	, WorldMinimapUsesVirtualTexture(false)
@@ -214,46 +212,88 @@ int32 SWorldPartitionEditorGridSpatialHash::PaintGrid(const FGeometry& AllottedG
 	{
 		FCellDesc2D()
 			: Bounds(ForceInitToZero)
+			, NumLoaded(0)
+			, NumUnloaded(0)
+			, NumEmpty(0)
 		{}
 
 		FBox2D Bounds;
+		int32 NumLoaded;
+		int32 NumUnloaded;
+		int32 NumEmpty;
 	};
 
 	// Draw shadowed regions
 	{
-		TMap<uint32, FCellDesc2D> UniqueCells2D;
-		EditorSpatialHash->ForEachIntersectingUnloadedRegion(VisibleGridRectWorld, [&](const UWorldPartitionEditorSpatialHash::FCellCoord& CellCoord)
+		TMap<uint32, FCellDesc2D> FlattenedCells2D;
+		EditorSpatialHash->ForEachIntersectingCells(VisibleGridRectWorld.ExpandBy(-1), 0, [&](const UWorldPartitionEditorSpatialHash::FCellCoord& CellCoord)
 		{
-			FBox CellBounds = EditorSpatialHash->GetCellBounds(CellCoord);
-			CellBounds = CellBounds.Overlap(EditorSpatialHash->Bounds);
-
 			FHashBuilder HashBuilder;
-			HashBuilder << CellCoord.X << CellCoord.Y << CellCoord.Level;
-			uint32 CellHash2D = HashBuilder.GetHash();
-
-			FCellDesc2D& CellDesc2D = UniqueCells2D.Add(CellHash2D);
+			HashBuilder << CellCoord.X << CellCoord.Y;
+			uint32 CellHash2D = HashBuilder.GetHash();			
+			
+			FCellDesc2D& CellDesc2D = FlattenedCells2D.FindOrAdd(CellHash2D);
+			
+			FBox CellBounds = EditorSpatialHash->GetCellBounds(CellCoord);
 			CellDesc2D.Bounds = FBox2D(FVector2D(CellBounds.Min), FVector2D(CellBounds.Max));
-		});
 
-		const bool bDebugDrawOctree = CVarDebugDrawOctree.GetValueOnAnyThread();
+			if (UWorldPartitionEditorCell** CellPtr = EditorSpatialHash->HashCells.Find(CellCoord))
+			{
+				if ((*CellPtr)->IsEmpty())
+				{
+					CellDesc2D.NumEmpty++;
+				}
+				else if ((*CellPtr)->IsLoaded())
+				{
+					CellDesc2D.NumLoaded++;
+				}
+				else
+				{
+					CellDesc2D.NumUnloaded++;
+				}
+			}
+			else
+			{
+				CellDesc2D.NumEmpty++;
+			}
+		});
 
-		for(auto& UniqueCell: UniqueCells2D)
+		for(auto& UniqueCell: FlattenedCells2D)
 		{
 			const FCellDesc2D& Cell = UniqueCell.Value;
 
-			FPaintGeometry CellGeometry = AllottedGeometry.ToPaintGeometry(
+			// Fully loaded
+			if (Cell.NumLoaded && !Cell.NumUnloaded)
+			{
+				continue;
+			}
+
+			const FPaintGeometry CellGeometry = AllottedGeometry.ToPaintGeometry(
 				WorldToScreen.TransformPoint(Cell.Bounds.Min),
 				WorldToScreen.TransformPoint(Cell.Bounds.Max) - WorldToScreen.TransformPoint(Cell.Bounds.Min)
 			);
 
-			FSlateColorBrush CellBrush(FLinearColor::White);
-			FLinearColor CellColor(0, 0, 0, 0.5f);
+			const FSlateColorBrush CellBrush(FLinearColor::White);
+			const FLinearColor FullyUnloadedCellColor(0, 0, 0, 0.5f);
+			const FLinearColor PartiallyLoadedCellColor(0, 0, 0, 0.25f);
+			const FLinearColor EmptyCellColor(0, 0, 0, 0.75f);
+			
+			FLinearColor CellColor;
 
-			if (bDebugDrawOctree)
+			// Partially loaded aread
+			if (Cell.NumLoaded && Cell.NumUnloaded)
 			{
-				CellColor = FLinearColor::MakeFromHSV8((uint8)UniqueCell.Key, 255, 255);
+				CellColor = PartiallyLoadedCellColor;
 			}
-			
+			else if (Cell.NumEmpty && !Cell.NumUnloaded && !Cell.NumLoaded)
+			{
+				CellColor = EmptyCellColor;
+			}
+			else
+			{
+				CellColor = FullyUnloadedCellColor;
+			}
+
 			FSlateDrawElement::MakeBox(
 				OutDrawElements,
 				++LayerId,
@@ -272,8 +312,7 @@ int32 SWorldPartitionEditorGridSpatialHash::PaintGrid(const FGeometry& AllottedG
 
 		if (VisibleSelectBox.IsValid)
 		{
-			TMap<uint32, FCellDesc2D> UniqueCells2D;
-
+			TMap<uint32, FCellDesc2D> FlattenedCells2D;
 			WorldPartition->EditorHash->ForEachIntersectingCell(VisibleSelectBox, [&](UWorldPartitionEditorCell* Cell)
 			{
 				UWorldPartitionEditorSpatialHash::FCellCoord CellCoord = EditorSpatialHash->GetCellCoords(Cell->Bounds.GetCenter(), 0);
@@ -281,13 +320,12 @@ int32 SWorldPartitionEditorGridSpatialHash::PaintGrid(const FGeometry& AllottedG
 				FHashBuilder HashBuilder;
 				HashBuilder << CellCoord.X << CellCoord.Y;
 				uint32 CellHash2D = HashBuilder.GetHash();
-
-				FCellDesc2D& CellDesc2D = UniqueCells2D.Add(CellHash2D);
+				FCellDesc2D& CellDesc2D = FlattenedCells2D.Add(CellHash2D);
 
 				CellDesc2D.Bounds = FBox2D(FVector2D(Cell->Bounds.Min), FVector2D(Cell->Bounds.Max));
 			});
 
-			for (auto& UniqueCell : UniqueCells2D)
+			for (auto& UniqueCell : FlattenedCells2D)
 			{
 				const FCellDesc2D& Cell = UniqueCell.Value;
 
diff --git a/Engine/Source/Runtime/Engine/Private/WorldPartition/WorldPartition.cpp b/Engine/Source/Runtime/Engine/Private/WorldPartition/WorldPartition.cpp
--- a/Engine/Source/Runtime/Engine/Private/WorldPartition/WorldPartition.cpp
+++ b/Engine/Source/Runtime/Engine/Private/WorldPartition/WorldPartition.cpp
@@ -1013,15 +1013,6 @@ void UWorldPartition::UpdateLoadingEditorCell(UWorldPartitionEditorCell* Cell, b
 	if (Cell->IsLoaded() != bShouldBeLoaded)
 	{
 		Cell->SetLoaded(bShouldBeLoaded, bShouldBeLoaded && bFromUserOperation);
-
-		if (Cell->IsLoaded())
-		{
-			EditorHash->OnCellLoaded(Cell);
-		}
-		else
-		{
-			EditorHash->OnCellUnloaded(Cell);
-		}
 	}
 
 	if (bPotentiallyUnloadedActors)
diff --git a/Engine/Source/Runtime/Engine/Private/WorldPartition/WorldPartitionEditorSpatialHash.cpp b/Engine/Source/Runtime/Engine/Private/WorldPartition/WorldPartitionEditorSpatialHash.cpp
--- a/Engine/Source/Runtime/Engine/Private/WorldPartition/WorldPartitionEditorSpatialHash.cpp
+++ b/Engine/Source/Runtime/Engine/Private/WorldPartition/WorldPartitionEditorSpatialHash.cpp
@@ -87,17 +87,8 @@ void UWorldPartitionEditorSpatialHash::Tick(float DeltaSeconds)
 						LevelCellCoord = LevelCellCoord.GetParentCellCoord();
 						const FCellNode& CellNode = HashNodes.FindChecked(LevelCellCoord);
 						check(CellNode.HasChildNode(ChildIndex));
-						check(CellNode.HasChildLoadedNode(ChildIndex));
 					}
 				}
-				else
-				{
-					FCellCoord LevelCellCoord = CellCoord;
-					const uint32 ChildIndex = LevelCellCoord.GetChildIndex();
-					LevelCellCoord = LevelCellCoord.GetParentCellCoord();
-					const FCellNode& CellNode = HashNodes.FindChecked(LevelCellCoord);
-					check(!CellNode.HasChildLoadedNode(ChildIndex));
-				}
 			}
 		});
 	}
@@ -197,18 +188,6 @@ void UWorldPartitionEditorSpatialHash::HashActor(FWorldPartitionHandle& InActorH
 			{
 				if (CurrentLevel ? HashNodes.Contains(CellCoord) : HashCells.Contains(CellCoord))
 				{
-					bool bSetLoadedMask;
-					if (CellCoord.Level)
-					{
-						FCellNode& ChildCellNode = HashNodes.FindChecked(CellCoord);
-						bSetLoadedMask = ChildCellNode.HasChildLoadedNodes();
-					}
-					else
-					{
-						UWorldPartitionEditorCell* EditorCell = HashCells.FindChecked(CellCoord);
-						bSetLoadedMask = EditorCell->IsLoaded();
-					}
-
 					FCellCoord LevelCellCoord = CellCoord;
 					while (LevelCellCoord.Level < NewLevel)
 					{
@@ -219,7 +198,7 @@ void UWorldPartitionEditorSpatialHash::HashActor(FWorldPartitionHandle& InActorH
 						FCellNode& CellNode = HashNodes.FindOrAdd(LevelCellCoord);
 
 						// We can break updating when aggregated flags are already properly set for parent nodes
-						const bool bShouldBreak = CellNode.HasChildNodes() && (!bSetLoadedMask || CellNode.HasChildLoadedNodes());
+						const bool bShouldBreak = CellNode.HasChildNodes();
 
 						// Propagate the child mask
 						if (!CellNode.HasChildNode(ChildIndex))
@@ -227,12 +206,6 @@ void UWorldPartitionEditorSpatialHash::HashActor(FWorldPartitionHandle& InActorH
 							CellNode.AddChildNode(ChildIndex);
 						}				
 				
-						// Propagate child loaded mask
-						if (bSetLoadedMask && !CellNode.HasChildLoadedNode(ChildIndex))
-						{
-							CellNode.AddChildLoadedNode(ChildIndex);
-						}
-
 						if (bShouldBreak)
 						{
 							break;
@@ -269,7 +242,6 @@ void UWorldPartitionEditorSpatialHash::UnhashActor(FWorldPartitionHandle& InActo
 				verify(HashCells.Remove(CellCoord));
 
 				bool bClearChildMask = true;
-				bool bClearLoadedMask = EditorCell->IsLoaded();
 
 				FCellCoord LevelCellCoord = CellCoord;
 				while (LevelCellCoord.Level < CurrentLevel)
@@ -280,16 +252,6 @@ void UWorldPartitionEditorSpatialHash::UnhashActor(FWorldPartitionHandle& InActo
 
 					FCellNode& CellNode = HashNodes.FindChecked(LevelCellCoord);
 
-					if (bClearLoadedMask)
-					{
-						CellNode.RemoveChildLoadedNode(ChildIndex);
-
-						if (CellNode.HasChildLoadedNodes())
-						{
-							bClearLoadedMask = false;
-						}
-					}
-
 					if (bClearChildMask)
 					{
 						CellNode.RemoveChildNode(ChildIndex);
@@ -304,7 +266,7 @@ void UWorldPartitionEditorSpatialHash::UnhashActor(FWorldPartitionHandle& InActo
 						}
 					}
 
-					if (!bClearChildMask && !bClearLoadedMask)
+					if (!bClearChildMask)
 					{
 						break;
 					}
@@ -333,68 +295,6 @@ void UWorldPartitionEditorSpatialHash::UnhashActor(FWorldPartitionHandle& InActo
 	}
 }
 
-void UWorldPartitionEditorSpatialHash::OnCellLoaded(const UWorldPartitionEditorCell* Cell)
-{
-	if (Cell == AlwaysLoadedCell)
-	{
-		return;
-	}
-
-	check(Cell->IsLoaded());
-
-	const FCellCoord CellCoord = GetCellCoords(Cell->Bounds.GetCenter(), 0);
-	const int32 CurrentLevel = GetLevelForBox(Bounds);
-
-	FCellCoord LevelCellCoord = CellCoord;
-	while (LevelCellCoord.Level < CurrentLevel)
-	{
-		const uint32 ChildIndex = LevelCellCoord.GetChildIndex();
-
-		LevelCellCoord = LevelCellCoord.GetParentCellCoord();
-
-		FCellNode& CellNode = HashNodes.FindChecked(LevelCellCoord);
-
-		check(!CellNode.HasChildLoadedNode(ChildIndex) || (CellCoord.Level != 1));
-
-		if (CellNode.HasChildLoadedNode(ChildIndex))
-		{
-			break;
-		}
-
-		CellNode.AddChildLoadedNode(ChildIndex);
-	}
-}
-
-void UWorldPartitionEditorSpatialHash::OnCellUnloaded(const UWorldPartitionEditorCell* Cell)
-{
-	if (Cell == AlwaysLoadedCell)
-	{
-		return;
-	}
-
-	check(!Cell->IsLoaded());
-
-	const FCellCoord CellCoord = GetCellCoords(Cell->Bounds.GetCenter(), 0);
-	const int32 CurrentLevel = GetLevelForBox(Bounds);
-
-	FCellCoord LevelCellCoord = CellCoord;
-	while (LevelCellCoord.Level < CurrentLevel)
-	{
-		const uint32 ChildIndex = LevelCellCoord.GetChildIndex();
-
-		LevelCellCoord = LevelCellCoord.GetParentCellCoord();
-
-		FCellNode& CellNode = HashNodes.FindChecked(LevelCellCoord);
-
-		CellNode.RemoveChildLoadedNode(ChildIndex);
-
-		if (CellNode.HasChildLoadedNodes())
-		{
-			break;
-		}
-	}
-}
-
 int32 UWorldPartitionEditorSpatialHash::ForEachIntersectingActor(const FBox& Box, TFunctionRef<void(FWorldPartitionActorDesc*)> InOperation)
 {
 	int32 NumIntersecting = 0;
@@ -488,69 +388,6 @@ int32 UWorldPartitionEditorSpatialHash::ForEachIntersectingCell(const FBox& Box,
 	return NumIntersecting;
 }
 
-int32 UWorldPartitionEditorSpatialHash::ForEachIntersectingUnloadedRegionInner(const FBox& Box, const FCellCoord& CellCoord, TFunctionRef<void(const FCellCoord&)> InOperation)
-{
-	int32 NumIntersecting = 0;
-
-	if (CellCoord.Level)
-	{
-		FCellNode CellNode;
-		if (FCellNode* CellNodePtr = HashNodes.Find(CellCoord))
-		{
-			CellNode = *CellNodePtr;
-
-			check(CellNode.HasChildNodes());
-
-			if (!CellNode.HasChildLoadedNodes())
-			{
-				InOperation(CellCoord);
-			}
-			else
-			{
-				CellNode.ForEachChild([&](uint32 ChildIndex)
-				{
-					const FCellCoord ChildCellCoord = CellCoord.GetChildCellCoord(ChildIndex);
-					const FBox CellBounds = GetCellBounds(ChildCellCoord);
-
-					if (Box.Intersect(CellBounds))
-					{
-						NumIntersecting += ForEachIntersectingUnloadedRegionInner(Box, ChildCellCoord, InOperation);
-					}
-				});
-			}
-		}
-	}
-	else
-	{
-		UWorldPartitionEditorCell* EditorCell = nullptr;
-		if (UWorldPartitionEditorCell** EditorCellPtr = HashCells.Find(CellCoord))
-		{
-			if (!(*EditorCellPtr)->IsLoaded())
-			{
-				InOperation(CellCoord);
-				NumIntersecting++;
-			}			
-		}
-	}
-
-	return NumIntersecting;
-}
-
-int32 UWorldPartitionEditorSpatialHash::ForEachIntersectingUnloadedRegion(const FBox& Box, TFunctionRef<void(const FCellCoord&)> InOperation)
-{
-	int32 NumIntersecting = 0;
-
-	const FBox SearchBox = Box.Overlap(Bounds);
-	const int32 SearchLevel = GetLevelForBox(SearchBox);
-
-	ForEachIntersectingCells(SearchBox, SearchLevel, [&](const FCellCoord& CellCoord)
-	{
-		NumIntersecting += ForEachIntersectingUnloadedRegionInner(Box, CellCoord, InOperation);
-	});
-
-	return NumIntersecting;
-}
-
 int32 UWorldPartitionEditorSpatialHash::ForEachCell(TFunctionRef<void(UWorldPartitionEditorCell*)> InOperation)
 {
 	InOperation(AlwaysLoadedCell);
diff --git a/Engine/Source/Runtime/Engine/Public/WorldPartition/WorldPartitionEditorCell.h b/Engine/Source/Runtime/Engine/Public/WorldPartition/WorldPartitionEditorCell.h
--- a/Engine/Source/Runtime/Engine/Public/WorldPartition/WorldPartitionEditorCell.h
+++ b/Engine/Source/Runtime/Engine/Public/WorldPartition/WorldPartitionEditorCell.h
@@ -93,6 +93,11 @@ public:
 		return bLoadedChangedByUserOperation;
 	}
 
+	inline bool IsEmpty() const
+	{
+		return !Actors.Num();
+	}
+
 	FBox						Bounds;
 
 private:
diff --git a/Engine/Source/Runtime/Engine/Public/WorldPartition/WorldPartitionEditorHash.h b/Engine/Source/Runtime/Engine/Public/WorldPartition/WorldPartitionEditorHash.h
--- a/Engine/Source/Runtime/Engine/Public/WorldPartition/WorldPartitionEditorHash.h
+++ b/Engine/Source/Runtime/Engine/Public/WorldPartition/WorldPartitionEditorHash.h
@@ -26,9 +26,6 @@ public:
 	virtual void HashActor(FWorldPartitionHandle& InActorHandle) PURE_VIRTUAL(UWorldPartitionEditorHash::HashActor, ;);
 	virtual void UnhashActor(FWorldPartitionHandle& InActorHandle) PURE_VIRTUAL(UWorldPartitionEditorHash::UnhashActor, ;);
 
-	virtual void OnCellLoaded(const UWorldPartitionEditorCell* Cell) PURE_VIRTUAL(UWorldPartitionEditorHash::OnCellLoaded, ;);
-	virtual void OnCellUnloaded(const UWorldPartitionEditorCell* Cell) PURE_VIRTUAL(UWorldPartitionEditorHash::OnCellLoaded, ;);
-
 	virtual int32 ForEachIntersectingActor(const FBox& Box, TFunctionRef<void(FWorldPartitionActorDesc*)> InOperation) PURE_VIRTUAL(UWorldPartitionEditorHash::ForEachIntersectingActor, return 0;);
 	virtual int32 ForEachIntersectingCell(const FBox& Box, TFunctionRef<void(UWorldPartitionEditorCell*)> InOperation) PURE_VIRTUAL(UWorldPartitionEditorHash::ForEachIntersectingCell, return 0;);
 	virtual int32 ForEachCell(TFunctionRef<void(UWorldPartitionEditorCell*)> InOperation) PURE_VIRTUAL(UWorldPartitionEditorHash::ForEachCell, return 0;);
diff --git a/Engine/Source/Runtime/Engine/Public/WorldPartition/WorldPartitionEditorSpatialHash.h b/Engine/Source/Runtime/Engine/Public/WorldPartition/WorldPartitionEditorSpatialHash.h
--- a/Engine/Source/Runtime/Engine/Public/WorldPartition/WorldPartitionEditorSpatialHash.h
+++ b/Engine/Source/Runtime/Engine/Public/WorldPartition/WorldPartitionEditorSpatialHash.h
@@ -125,7 +125,6 @@ class ENGINE_API UWorldPartitionEditorSpatialHash : public UWorldPartitionEditor
 	{
 		FCellNode()
 			: ChildNodesMask(0)
-			, ChildNodesLoadedMask(0)
 		{}
 
 		inline bool HasChildNodes() const
@@ -170,37 +169,7 @@ class ENGINE_API UWorldPartitionEditorSpatialHash : public UWorldPartitionEditor
 			}
 		}
 
-		inline bool HasChildLoadedNodes() const
-		{
-			return !!ChildNodesLoadedMask;
-		}
-
-		inline bool HasChildLoadedNode(uint32 ChildIndex) const
-		{
-			check(ChildIndex < 8);
-			return !!(ChildNodesLoadedMask & (1 << ChildIndex));
-		}
-
-		inline void AddChildLoadedNode(uint32 ChildIndex)
-		{
-			check(ChildIndex < 8);
-			uint32 ChildMask = 1 << ChildIndex;
-			check(ChildNodesMask & ChildMask);
-			check(!(ChildNodesLoadedMask & ChildMask));
-			ChildNodesLoadedMask |= ChildMask;
-		}
-
-		inline void RemoveChildLoadedNode(uint32 ChildIndex)
-		{
-			check(ChildIndex < 8);
-			uint32 ChildMask = 1 << ChildIndex;
-			check(ChildNodesMask & ChildMask);
-			check(ChildNodesLoadedMask & ChildMask);
-			ChildNodesLoadedMask &= ~ChildMask;
-		}
-
 		uint8 ChildNodesMask;
-		uint8 ChildNodesLoadedMask;
 	};
 
 public:
@@ -220,9 +189,6 @@ public:
 	virtual void HashActor(FWorldPartitionHandle& InActorHandle) override;
 	virtual void UnhashActor(FWorldPartitionHandle& InActorHandle) override;
 
-	virtual void OnCellLoaded(const UWorldPartitionEditorCell* Cell) override;
-	virtual void OnCellUnloaded(const UWorldPartitionEditorCell* Cell) override;
-
 	virtual int32 ForEachIntersectingActor(const FBox& Box, TFunctionRef<void(FWorldPartitionActorDesc*)> InOperation) override;
 	virtual int32 ForEachIntersectingCell(const FBox& Box, TFunctionRef<void(UWorldPartitionEditorCell*)> InOperation) override;
 	virtual int32 ForEachCell(TFunctionRef<void(UWorldPartitionEditorCell*)> InOperation) override;
@@ -239,9 +205,6 @@ public:
 
 #if WITH_EDITORONLY_DATA
 private:
-	int32 ForEachIntersectingUnloadedRegion(const FBox& Box, TFunctionRef<void(const FCellCoord&)> InOperation);
-	int32 ForEachIntersectingUnloadedRegionInner(const FBox& Box, const FCellCoord& CellCoord, TFunctionRef<void(const FCellCoord&)> InOperation);
-
 	FBox GetActorBounds(const FWorldPartitionHandle& InActorHandle) const;
 	bool IsActorAlwaysLoaded(const FWorldPartitionHandle& InActorHandle) const;
 	int32 ForEachIntersectingCellInner(const FBox& Box, const FCellCoord& CellCoord, TFunctionRef<void(UWorldPartitionEditorCell*)> InOperation);
