Player Controllers acting as streaming sources can now specify if they should block on slow streaming or not.

diff --git a/Engine/Source/Runtime/Engine/Classes/GameFramework/PlayerController.h b/Engine/Source/Runtime/Engine/Classes/GameFramework/PlayerController.h
--- a/Engine/Source/Runtime/Engine/Classes/GameFramework/PlayerController.h
+++ b/Engine/Source/Runtime/Engine/Classes/GameFramework/PlayerController.h
@@ -504,14 +504,18 @@ public:
 	UPROPERTY(EditAnywhere, BlueprintReadWrite, Category="Game|Feedback")
 	uint32 bForceFeedbackEnabled:1;
 
-	/** Whether the PlayerController should be used as a World Partiton streaming source.  */
+	/** Whether the PlayerController should be used as a World Partiton streaming source. */
 	UPROPERTY(EditAnywhere, BlueprintReadWrite, Category = WorldPartition)
 	uint32 bEnableStreamingSource:1;
 
-	/** Whether the PlayerController streaming source should activate cells after loading.  */
+	/** Whether the PlayerController streaming source should activate cells after loading. */
 	UPROPERTY(EditAnywhere, BlueprintReadWrite, Category = WorldPartition, meta=(EditCondition="bEnableStreamingSource"))
 	uint32 bStreamingSourceShouldActivate:1;
 
+	/** Whether the PlayerController streaming source should block on slow streaming. */
+	UPROPERTY(EditAnywhere, BlueprintReadWrite, Category = WorldPartition, meta=(EditCondition="bEnableStreamingSource"))
+	uint32 bStreamingSourceShouldBlockOnSlowStreaming:1;
+
 	/** Scale applied to force feedback values */
 	UPROPERTY(config)
 	float ForceFeedbackScale;
@@ -720,6 +724,14 @@ public:
 	UFUNCTION(BlueprintCallable, Category = WorldPartition)
 	virtual bool StreamingSourceShouldActivate() const { return bEnableStreamingSource && bStreamingSourceShouldActivate; }
 
+	/**
+	* Whether the PlayerController streaming source should block on slow streaming.
+	* Default implementation returns bStreamingSourceShouldBlockOnSlowStreaming but can be overriden in child classes.
+	* @return true if it should.
+	*/
+	UFUNCTION(BlueprintCallable, Category = WorldPartition)
+	virtual bool StreamingSourceShouldBlockOnSlowStreaming() const { return bEnableStreamingSource && bStreamingSourceShouldBlockOnSlowStreaming; }
+
 protected:
 	/** Pawn has been possessed, so changing state to NAME_Playing. Start it walking and begin playing with it. */
 	virtual void BeginPlayingState();
diff --git a/Engine/Source/Runtime/Engine/Private/PlayerController.cpp b/Engine/Source/Runtime/Engine/Private/PlayerController.cpp
--- a/Engine/Source/Runtime/Engine/Private/PlayerController.cpp
+++ b/Engine/Source/Runtime/Engine/Private/PlayerController.cpp
@@ -128,7 +128,9 @@ APlayerController::APlayerController(const FObjectInitializer& ObjectInitializer
 	bForceFeedbackEnabled = true;
 	ForceFeedbackScale = 1.f;
 
+	bEnableStreamingSource = true;
 	bStreamingSourceShouldActivate = true;
+	bStreamingSourceShouldBlockOnSlowStreaming = true;
 
 	bAutoManageActiveCameraTarget = true;
 	bRenderPrimitiveComponents = true;
@@ -137,9 +139,7 @@ APlayerController::APlayerController(const FObjectInitializer& ObjectInitializer
 
 	bIsPlayerController = true;
 	bIsLocalPlayerController = false;
-	bDisableHaptics = false;
-
-	bEnableStreamingSource = true;
+	bDisableHaptics = false;	
 
 	ClickEventKeys.Add(EKeys::LeftMouseButton);
 
diff --git a/Engine/Source/Runtime/Engine/Private/WorldPartition/WorldPartitionStreamingPolicy.cpp b/Engine/Source/Runtime/Engine/Private/WorldPartition/WorldPartitionStreamingPolicy.cpp
--- a/Engine/Source/Runtime/Engine/Private/WorldPartition/WorldPartitionStreamingPolicy.cpp
+++ b/Engine/Source/Runtime/Engine/Private/WorldPartition/WorldPartitionStreamingPolicy.cpp
@@ -126,8 +126,10 @@ void UWorldPartitionStreamingPolicy::UpdateStreamingSources()
 						FRotator ViewRotation;
 						Player->PlayerController->GetPlayerViewPoint(ViewLocation, ViewRotation);
 
-						EStreamingSourceTargetState TargetState = Player->PlayerController->StreamingSourceShouldActivate() ? EStreamingSourceTargetState::Activated : EStreamingSourceTargetState::Loaded;
-						StreamingSources.Add(FWorldPartitionStreamingSource(Player->PlayerController->GetFName(), ViewLocation, ViewRotation, TargetState, /*bBlockOnSlowLoading=*/true, EStreamingSourcePriority::Default));
+				
+						const EStreamingSourceTargetState TargetState = Player->PlayerController->StreamingSourceShouldActivate() ? EStreamingSourceTargetState::Activated : EStreamingSourceTargetState::Loaded;
+						const bool bBlockOnSlowLoading = Player->PlayerController->StreamingSourceShouldBlockOnSlowStreaming();
+						StreamingSources.Add(FWorldPartitionStreamingSource(Player->PlayerController->GetFName(), ViewLocation, ViewRotation, TargetState, bBlockOnSlowLoading, EStreamingSourcePriority::Default));
 					}
 				}
 			}
