Time Dilation Multiplier
Unshelved files from [REDACTED_HANDLE]
CL 17427780

diff --git a/Engine/Source/Runtime/Engine/Classes/GameFramework/WorldSettings.h b/Engine/Source/Runtime/Engine/Classes/GameFramework/WorldSettings.h
--- a/Engine/Source/Runtime/Engine/Classes/GameFramework/WorldSettings.h
+++ b/Engine/Source/Runtime/Engine/Classes/GameFramework/WorldSettings.h
@@ -700,6 +700,9 @@ public:
 	UPROPERTY(transient, replicated)
 	float TimeDilation;
 
+	UPROPERTY(transient, replicated)
+	float TimeDilationMultiplier;
+
 	// Additional time dilation used by Matinee (or Sequencer) slomo track.  Transient because this is often 
 	// temporarily modified by the editor when previewing slow motion effects, yet we don't want it saved or loaded from level packages.
 	UPROPERTY(transient, replicated)
@@ -787,7 +790,7 @@ public:
 
 	virtual float GetEffectiveTimeDilation() const
 	{
-		return TimeDilation * MatineeTimeDilation * DemoPlayTimeDilation;
+		return TimeDilation * MatineeTimeDilation * DemoPlayTimeDilation * TimeDilationMultiplier;
 	}
 
 	/**
diff --git a/Engine/Source/Runtime/Engine/Classes/Kismet/GameplayStatics.h b/Engine/Source/Runtime/Engine/Classes/Kismet/GameplayStatics.h
--- a/Engine/Source/Runtime/Engine/Classes/Kismet/GameplayStatics.h
+++ b/Engine/Source/Runtime/Engine/Classes/Kismet/GameplayStatics.h
@@ -348,6 +348,20 @@ class ENGINE_API UGameplayStatics : public UBlueprintFunctionLibrary
 	UFUNCTION(BlueprintCallable, Category="Utilities|Time", meta=(WorldContext="WorldContextObject") )
 	static void SetGlobalTimeDilation(const UObject* WorldContextObject, float TimeDilation);
 
+	/**
+	 * Gets the current global time dilation multiplier.
+	 * @return Current time dilation multiplier.
+	 */
+	UFUNCTION(BlueprintPure, Category="Utilities|Time", meta=(WorldContext="WorldContextObject") )
+	static float GetGlobalTimeDilationMultiplier(const UObject* WorldContextObject);
+
+	/**
+	 * Sets the global time dilation multiplier
+	 * @param	TimeDilationMultiplier	value to set the global time dilation to
+	 */
+	UFUNCTION(BlueprintCallable, Category="Utilities|Time", meta=(WorldContext="WorldContextObject") )
+	static void SetGlobalTimeDilationMultiplier(const UObject* WorldContextObject, float TimeDilationMultiplier);
+
 	/**
 	 * Sets the game's paused state
 	 * @param	bPaused		Whether the game should be paused or not
diff --git a/Engine/Source/Runtime/Engine/Private/GameplayStatics.cpp b/Engine/Source/Runtime/Engine/Private/GameplayStatics.cpp
--- a/Engine/Source/Runtime/Engine/Private/GameplayStatics.cpp
+++ b/Engine/Source/Runtime/Engine/Private/GameplayStatics.cpp
@@ -475,6 +475,25 @@ void UGameplayStatics::SetGlobalTimeDilation(const UObject* WorldContextObject,
 	}
 }
 
+float UGameplayStatics::GetGlobalTimeDilationMultiplier(const UObject* WorldContextObject)
+{
+	UWorld* const World = GEngine->GetWorldFromContextObject(WorldContextObject, EGetWorldErrorMode::LogAndReturnNull);
+	return World ? World->GetWorldSettings()->TimeDilationMultiplier : 1.f;
+}
+
+void UGameplayStatics::SetGlobalTimeDilationMultiplier(const UObject* WorldContextObject, float TimeDilationMultiplier)
+{
+	UWorld* const World = GEngine->GetWorldFromContextObject(WorldContextObject, EGetWorldErrorMode::LogAndReturnNull);
+	if (World != nullptr)
+	{
+		AWorldSettings* const WorldSettings = World->GetWorldSettings();
+		if (WorldSettings != nullptr)
+		{
+			WorldSettings->TimeDilationMultiplier = TimeDilationMultiplier;
+		}
+	}
+}
+
 bool UGameplayStatics::SetGamePaused(const UObject* WorldContextObject, bool bPaused)
 {
 	UGameInstance* const GameInstance = GetGameInstance( WorldContextObject );
diff --git a/Engine/Source/Runtime/Engine/Private/WorldSettings.cpp b/Engine/Source/Runtime/Engine/Private/WorldSettings.cpp
--- a/Engine/Source/Runtime/Engine/Private/WorldSettings.cpp
+++ b/Engine/Source/Runtime/Engine/Private/WorldSettings.cpp
@@ -103,6 +103,7 @@ PRAGMA_ENABLE_DEPRECATION_WARNINGS
 	bReplicates = true;
 	bAlwaysRelevant = true;
 	TimeDilation = 1.0f;
+	TimeDilationMultiplier = 1.0f;
 	MatineeTimeDilation = 1.0f;
 	DemoPlayTimeDilation = 1.0f;
 	PackedLightAndShadowMapTextureSize = 1024;
