Don't flush async loading when flushing streaming levels visibility, as it would include all other pending async load requests resulting in large frame spikes.

diff --git a/Engine/Source/Runtime/Engine/Private/World.cpp b/Engine/Source/Runtime/Engine/Private/World.cpp
--- a/Engine/Source/Runtime/Engine/Private/World.cpp
+++ b/Engine/Source/Runtime/Engine/Private/World.cpp
@@ -3723,6 +3723,8 @@ void FStreamingLevelsToConsider::AddReferencedObjects(UObject* InThis, FReferenc
 
 void UWorld::BlockTillLevelStreamingCompleted()
 {
+	TRACE_CPUPROFILER_EVENT_SCOPE(UWorld::BlockTillLevelStreamingCompleted);
+
 	const double StartTime = FPlatformTime::Seconds();
 	TScopeCounter<uint32> IsInBlockTillLevelStreamingCompletedCounter(IsInBlockTillLevelStreamingCompleted);
 
@@ -4059,32 +4061,35 @@ void UWorld::FlushLevelStreaming(EFlushLevelStreamingType FlushType)
 
 	TGuardValue<EFlushLevelStreamingType> FlushingLevelStreamingGuard(FlushLevelStreamingType, FlushType);
 
-	// Update internals with current loaded/ visibility flags.
-	UpdateLevelStreaming();
+	if (FlushLevelStreamingType == EFlushLevelStreamingType::Full)
+	{
+		// Update internals with current loaded/ visibility flags.
+		UpdateLevelStreaming();
+	}
 
-	// Make sure all outstanding loads are taken care of, other than ones associated with the excluded type
-	FlushAsyncLoading();
+	auto TickLevelStreaming = [this]()
+	{
+		// Only flush async loading if we're performing a full flush.
+		if (FlushLevelStreamingType == EFlushLevelStreamingType::Full)
+		{
+			// Make sure all outstanding loads are taken care of, other than ones associated with the excluded type
+			FlushAsyncLoading();
+		}
 
-	// Kick off making levels visible if loading finished by flushing.
-	UpdateLevelStreaming();
+		// Kick off making levels visible if loading finished by flushing.
+		UpdateLevelStreaming();
+	};
+
+	TickLevelStreaming();
 
 	// Making levels visible is spread across several frames so we simply loop till it is done.
 	bool bLevelsPendingVisibility = IsVisibilityRequestPending();
 	while( bLevelsPendingVisibility )
 	{
-		// Tick level streaming to make levels visible.
-
 		const EFlushLevelStreamingType LastStreamingType = FlushLevelStreamingType;
 
-			// Only flush async loading if we're performing a full flush.
-			if (FlushLevelStreamingType == EFlushLevelStreamingType::Full)
-			{
-				// Make sure all outstanding loads are taken care of...
-				FlushAsyncLoading();
-			}
-	
-			// Update level streaming.
-			UpdateLevelStreaming();
+		// Tick level streaming to make levels visible.
+		TickLevelStreaming();
 
 		// If FlushLevelStreaming reentered as a result of FlushAsyncLoading or UpdateLevelStreaming and upgraded
 		// the flush type, we'll need to do at least one additional loop to be certain all processing is complete
