Add data layers state changes to the CSV profiler.

diff --git a/Engine/Source/Runtime/Engine/Private/WorldPartition/DataLayer/WorldDataLayers.cpp b/Engine/Source/Runtime/Engine/Private/WorldPartition/DataLayer/WorldDataLayers.cpp
--- a/Engine/Source/Runtime/Engine/Private/WorldPartition/DataLayer/WorldDataLayers.cpp
+++ b/Engine/Source/Runtime/Engine/Private/WorldPartition/DataLayer/WorldDataLayers.cpp
@@ -162,10 +162,11 @@ void AWorldDataLayers::SetDataLayerRuntimeState(FActorDataLayer InDataLayer, EDa
 
 			++DataLayersStateEpoch;
 
-			UE_LOG(LogWorldPartition, Log, TEXT("Data Layer '%s' state changed: %s -> %s"), 
-				*DataLayer->GetDataLayerLabel().ToString(), 
-				*StaticEnum<EDataLayerRuntimeState>()->GetDisplayNameTextByValue((int64)CurrentState).ToString(),
-				*StaticEnum<EDataLayerRuntimeState>()->GetDisplayNameTextByValue((int64)InState).ToString());
+#if !NO_LOGGING || CSV_PROFILER
+			const FString DataLayerLabel = DataLayer->GetDataLayerLabel().ToString();
+			UE_LOG(LogWorldPartition, Log, TEXT("Data Layer '%s' state changed: %s -> %s"),  *DataLayerLabel,  GetDataLayerStateName(CurrentState), GetDataLayerStateName(InState));
+			CSV_EVENT_GLOBAL(TEXT("DataLayer-%s-%s"), *DataLayerLabel, GetDataLayerStateName(InState));
+#endif
 
 			ResolveEffectiveRuntimeState(DataLayer);
 		}
