WorldPartition: Fix AttachParent being trashed before Child can register itself.

This is a temporary fix for a bigger issue where editor loaded actors in a WP world need to be sorted by hierarchy then all registered then all reconstructed so that we can preserve attachments.
They are currently processed one by one no sorting and register/construction is done per actor one at a time

diff --git a/Engine/Source/Runtime/Engine/Classes/Components/SceneComponent.h b/Engine/Source/Runtime/Engine/Classes/Components/SceneComponent.h
--- a/Engine/Source/Runtime/Engine/Classes/Components/SceneComponent.h
+++ b/Engine/Source/Runtime/Engine/Classes/Components/SceneComponent.h
@@ -313,6 +313,23 @@ public:
 	/** Returns the current scoped movement update, or NULL if there is none. @see FScopedMovementUpdate */
 	class FScopedMovementUpdate* GetCurrentScopedMovement() const;
 
+#if WITH_EDITORONLY_DATA
+	/**
+	 * @todo_ow: This is needed because of order of registration of Actors
+	 * 
+	 * In World Partition Levels loaded actors are Registered in an atomic fashion meaning we register all their
+	 * components and then call RerunConstructionScripts before loading the next actor. This means that a Parent can be Reconstructed and trash its components
+	 * without notifying its attached actors.
+	 * 
+	 * In Non World Partition Levels when adding actors to world we sort all actors based on hierarchy, register all actors and then RerunConstructionScripts on all actors
+	 * which allows handling of Attachments as attached actors are already registered when the parent gets reconstructed and so attachment is preserved.
+	 * 
+	 * ReplacementSceneComponent is there as a temp solution to allow finding of the replacement component of a trashed scene component. 
+	 */
+	UPROPERTY(Transient)
+	USceneComponent* ReplacementSceneComponent;
+#endif
+
 private:
 	/** Stack of current movement scopes. */
 	TArray<class FScopedMovementUpdate*> ScopedMovementStack;
@@ -918,6 +935,8 @@ public:
 	virtual void Serialize(FArchive& Ar) override;
 #if WITH_EDITORONLY_DATA
 	static void AddReferencedObjects(UObject* InThis, FReferenceCollector& Collector);
+
+	virtual void PostLoad() override;
 #endif
 
 #if WITH_EDITOR
diff --git a/Engine/Source/Runtime/Engine/Private/ActorConstruction.cpp b/Engine/Source/Runtime/Engine/Private/ActorConstruction.cpp
--- a/Engine/Source/Runtime/Engine/Private/ActorConstruction.cpp
+++ b/Engine/Source/Runtime/Engine/Private/ActorConstruction.cpp
@@ -652,6 +652,18 @@ void AActor::RerunConstructionScripts()
 					}
 				}
 			}
+#if WITH_EDITORONLY_DATA
+			if (OwningWorld->GetWorldPartition() != nullptr)
+			{
+				if (USceneComponent* OldSceneComponent = Cast<USceneComponent>(ComponentData.OldComponent))
+				{
+					if (USceneComponent* NewSceneComponent = Cast<USceneComponent>(ResolvedNewComponent))
+					{
+						OldSceneComponent->ReplacementSceneComponent = NewSceneComponent;
+					}
+				}
+			}
+#endif
 
 			OldToNewComponentMapping.Add(ComponentData.OldComponent, ResolvedNewComponent);
 		}
diff --git a/Engine/Source/Runtime/Engine/Private/Components/SceneComponent.cpp b/Engine/Source/Runtime/Engine/Private/Components/SceneComponent.cpp
--- a/Engine/Source/Runtime/Engine/Private/Components/SceneComponent.cpp
+++ b/Engine/Source/Runtime/Engine/Private/Components/SceneComponent.cpp
@@ -98,6 +98,25 @@ void USceneComponent::AddReferencedObjects(UObject* InThis, FReferenceCollector&
 
 	Super::AddReferencedObjects(InThis, Collector);
 }
+
+void USceneComponent::PostLoad()
+{
+	Super::PostLoad();
+
+	if (AttachParent)
+	{
+		USceneComponent* TopReplacementComponent = AttachParent->ReplacementSceneComponent;
+		while (TopReplacementComponent && TopReplacementComponent->ReplacementSceneComponent)
+		{
+			TopReplacementComponent = TopReplacementComponent->ReplacementSceneComponent;
+		}
+
+		if (TopReplacementComponent)
+		{
+			SetAttachParent(TopReplacementComponent);
+		}
+	}
+}
 #endif
 
 #if WITH_EDITOR
