Switched USceneComponent::ReplacementSceneComponent to use TObjectPtr

diff --git a/Engine/Source/Runtime/Engine/Classes/Components/SceneComponent.h b/Engine/Source/Runtime/Engine/Classes/Components/SceneComponent.h
--- a/Engine/Source/Runtime/Engine/Classes/Components/SceneComponent.h
+++ b/Engine/Source/Runtime/Engine/Classes/Components/SceneComponent.h
@@ -327,7 +327,7 @@ public:
 	 * ReplacementSceneComponent is there as a temp solution to allow finding of the replacement component of a trashed scene component. 
 	 */
 	UPROPERTY(Transient)
-	USceneComponent* ReplacementSceneComponent;
+	TObjectPtr<USceneComponent> ReplacementSceneComponent;
 #endif
 
 private:
