Fix morph target broken on underweight female in cooked build.
When the first LOD is initialised with render data, it destroys the source morph data for all LODs, so the fix is to delay destroying source morphs until all LODs are initialised

diff --git a/Engine/Source/Runtime/Engine/Private/SkeletalMeshLODRenderData.cpp b/Engine/Source/Runtime/Engine/Private/SkeletalMeshLODRenderData.cpp
--- a/Engine/Source/Runtime/Engine/Private/SkeletalMeshLODRenderData.cpp
+++ b/Engine/Source/Runtime/Engine/Private/SkeletalMeshLODRenderData.cpp
@@ -531,15 +531,6 @@ void FSkeletalMeshLODRenderData::InitResources(bool bNeedsVertexColors, int32 LO
 			MorphTargetVertexInfoBuffers.BatchesPerMorph.Add(MorphNumBatches);
 			MorphTargetVertexInfoBuffers.MaximumValuePerMorph.Add(FVector4(MaximumValues[0], MaximumValues[1], MaximumValues[2], MaximumValues[3]));
 			MorphTargetVertexInfoBuffers.MinimumValuePerMorph.Add(FVector4(MinimumValues[0], MinimumValues[1], MinimumValues[2], MinimumValues[3]));
-
-#if !WITH_EDITOR
-			if (NumSrcDeltas > 0)
-			{
-				// A CPU copy of the morph deltas has beenA  made so it is safe to 
-				// discard the original data.  Keep CPU buffers when in the editor.
-				MorphTarget->DiscardVertexData();
-			}
-#endif
 		}
 
 
diff --git a/Engine/Source/Runtime/Engine/Private/SkeletalMeshRenderData.cpp b/Engine/Source/Runtime/Engine/Private/SkeletalMeshRenderData.cpp
--- a/Engine/Source/Runtime/Engine/Private/SkeletalMeshRenderData.cpp
+++ b/Engine/Source/Runtime/Engine/Private/SkeletalMeshRenderData.cpp
@@ -728,6 +728,15 @@ void FSkeletalMeshRenderData::InitResources(bool bNeedsVertexColors, TArray<UMor
 			}
 		}
 
+#if !WITH_EDITOR
+		for (UMorphTarget* MorphTarget : InMorphTargets)
+		{
+			// A CPU copy of the morph deltas has been  made so it is safe to 
+			// discard the original data.  Keep CPU buffers when in the editor.
+			MorphTarget->DiscardVertexData();
+		}
+#endif
+
 		ENQUEUE_RENDER_COMMAND(CmdSetSkeletalMeshReadyForStreaming)(
 			[this, Owner](FRHICommandListImmediate&)
 		{
