Data Layers: add a delegate on data layer state change to allow a subsystem to filter out changes.

diff --git a/Engine/Source/Runtime/Engine/Private/WorldPartition/DataLayer/WorldDataLayers.cpp b/Engine/Source/Runtime/Engine/Private/WorldPartition/DataLayer/WorldDataLayers.cpp
--- a/Engine/Source/Runtime/Engine/Private/WorldPartition/DataLayer/WorldDataLayers.cpp
+++ b/Engine/Source/Runtime/Engine/Private/WorldPartition/DataLayer/WorldDataLayers.cpp
@@ -144,6 +144,21 @@ void AWorldDataLayers::SetDataLayerRuntimeState(FActorDataLayer InDataLayer, EDa
 		EDataLayerRuntimeState CurrentState = GetDataLayerRuntimeStateByName(InDataLayer.Name);
 		if (CurrentState != InState)
 		{
+PRAGMA_DISABLE_DEPRECATION_WARNINGS
+			if (GetWorld()->IsGameWorld())
+			{
+				FName DataLayerLabelName = DataLayer->GetDataLayerLabel();
+				if (DataLayersFilterDelegate.IsBound())
+				{
+					if (!DataLayersFilterDelegate.Execute(DataLayerLabelName, CurrentState, InState))
+					{
+						UE_LOG(LogWorldPartition, Log, TEXT("Data Layer '%s' was filtered out: %s -> %s"),  *DataLayerLabelName.ToString(), GetDataLayerStateName(CurrentState), GetDataLayerStateName(InState));
+						return;
+					}
+				}
+			}
+PRAGMA_ENABLE_DEPRECATION_WARNINGS
+
 			LoadedDataLayerNames.Remove(InDataLayer.Name);
 			ActiveDataLayerNames.Remove(InDataLayer.Name);
 
diff --git a/Engine/Source/Runtime/Engine/Public/WorldPartition/DataLayer/WorldDataLayers.h b/Engine/Source/Runtime/Engine/Public/WorldPartition/DataLayer/WorldDataLayers.h
--- a/Engine/Source/Runtime/Engine/Public/WorldPartition/DataLayer/WorldDataLayers.h
+++ b/Engine/Source/Runtime/Engine/Public/WorldPartition/DataLayer/WorldDataLayers.h
@@ -142,4 +142,10 @@ private:
 	TSet<FName> EffectiveLoadedDataLayerNames;
 
 	static int32 DataLayersStateEpoch;
+
+public:
+	DECLARE_DELEGATE_RetVal_ThreeParams(bool, FDataLayersFilterDelegate, FName /*DataLayerName*/, EDataLayerState /*CurrentState*/, EDataLayerState /*TargetState*/);
+
+	UE_DEPRECATED(5.00, "do not use, will be replaced by another mechanism for initial release.")
+	FDataLayersFilterDelegate DataLayersFilterDelegate;
 };
\ No newline at end of file
