[RenderTarget2DArray] - fix up calculation of reported memory size to allow for >2GB

diff --git a/Engine/Source/Runtime/Engine/Private/TextureRenderTarget2DArray.cpp b/Engine/Source/Runtime/Engine/Private/TextureRenderTarget2DArray.cpp
--- a/Engine/Source/Runtime/Engine/Private/TextureRenderTarget2DArray.cpp
+++ b/Engine/Source/Runtime/Engine/Private/TextureRenderTarget2DArray.cpp
@@ -78,12 +78,12 @@ void UTextureRenderTarget2DArray::GetResourceSizeEx(FResourceSizeEx& CumulativeR
 
 	// Calculate size based on format.
 	const EPixelFormat Format = GetFormat();
-	const int32 BlockSizeX	= GPixelFormats[Format].BlockSizeX;
-	const int32 BlockSizeY	= GPixelFormats[Format].BlockSizeY;
-	const int32 BlockBytes	= GPixelFormats[Format].BlockBytes;
-	const int32 NumBlocksX	= (SizeX + BlockSizeX - 1) / BlockSizeX;
-	const int32 NumBlocksY	= (SizeY + BlockSizeY - 1) / BlockSizeY;
-	const int32 NumBytes	= NumBlocksX * NumBlocksY * Slices * BlockBytes;
+	const int64 BlockSizeX	= GPixelFormats[Format].BlockSizeX;
+	const int64 BlockSizeY	= GPixelFormats[Format].BlockSizeY;
+	const int64 BlockBytes	= GPixelFormats[Format].BlockBytes;
+	const int64 NumBlocksX	= (SizeX + BlockSizeX - 1) / BlockSizeX;
+	const int64 NumBlocksY	= (SizeY + BlockSizeY - 1) / BlockSizeY;
+	const int64 NumBytes	= NumBlocksX * NumBlocksY * Slices * BlockBytes;
 
 	CumulativeResourceSize.AddUnknownMemoryBytes(NumBytes);
 }
