-Allow parallel animation eval when when starting in a different tick group
-Allow per-component override for tick animation on init

diff --git a/Engine/Source/Runtime/Engine/Classes/Components/SkeletalMeshComponent.h b/Engine/Source/Runtime/Engine/Classes/Components/SkeletalMeshComponent.h
--- a/Engine/Source/Runtime/Engine/Classes/Components/SkeletalMeshComponent.h
+++ b/Engine/Source/Runtime/Engine/Classes/Components/SkeletalMeshComponent.h
@@ -671,7 +671,7 @@ public:
 	UPROPERTY(EditAnywhere, AdvancedDisplay, BlueprintReadWrite, Category=Animation)
 	uint8 bPauseAnims:1;
 
-	/** On InitAnim should we set to ref pose (if false use first tick of animation data)*/
+	/** On InitAnim should we set to ref pose (if false use first tick of animation data). If enabled, takes precedence over UAnimationSettings::bTickAnimationOnSkeletalMeshInit*/
 	UPROPERTY(EditAnywhere, AdvancedDisplay, Category = Animation)
 	uint8 bUseRefPoseOnInitAnim:1;
 
diff --git a/Engine/Source/Runtime/Engine/Private/Components/SkeletalMeshComponent.cpp b/Engine/Source/Runtime/Engine/Private/Components/SkeletalMeshComponent.cpp
--- a/Engine/Source/Runtime/Engine/Private/Components/SkeletalMeshComponent.cpp
+++ b/Engine/Source/Runtime/Engine/Private/Components/SkeletalMeshComponent.cpp
@@ -740,8 +740,9 @@ void USkeletalMeshComponent::InitAnim(bool bForceReinit)
 		// In Editor, animations won't get ticked. So Update once to get accurate representation instead of T-Pose.
 		// Also allow this to be an option to support pre-4.19 games that might need it..
 		const bool bTickAnimationNow =
-			((GetWorld()->WorldType == EWorldType::Editor) && !bUseRefPoseOnInitAnim && !bForceRefpose)
-			|| UAnimationSettings::Get()->bTickAnimationOnSkeletalMeshInit;
+			(((GetWorld()->WorldType == EWorldType::Editor) && !bForceRefpose)
+			|| UAnimationSettings::Get()->bTickAnimationOnSkeletalMeshInit)
+			&& !bUseRefPoseOnInitAnim;
 
 		const bool bInitializedAnimInstance = InitializeAnimScriptInstance(bForceReinit, !bTickAnimationNow);
 
