[Backout] - CL17587421
Original CL Desc
-----------------------------------------------------------------
Fix RayTracingMinLOD setting being reverted on reimport
Implement proper async build handling for bSupportRayTracing and RayTracingMinLOD

diff --git a/Engine/Source/Runtime/Engine/Private/Components/SkinnedMeshComponent.cpp b/Engine/Source/Runtime/Engine/Private/Components/SkinnedMeshComponent.cpp
--- a/Engine/Source/Runtime/Engine/Private/Components/SkinnedMeshComponent.cpp
+++ b/Engine/Source/Runtime/Engine/Private/Components/SkinnedMeshComponent.cpp
@@ -1484,7 +1484,7 @@ bool USkinnedMeshComponent::IsSkinCacheAllowed(int32 LodIdx) const
 	bool bIsRayTracing = FGPUSkinCache::IsGPUSkinCacheRayTracingSupported();
 	if (bIsRayTracing && SkeletalMesh)
 	{
-		bIsRayTracing = SkeletalMesh->GetSupportRayTracing();
+		bIsRayTracing = SkeletalMesh->bSupportRayTracing;
 	}
 
 	bool bGlobalDefault = CVarDefaultGPUSkinCacheBehavior && ESkinCacheDefaultBehavior(CVarDefaultGPUSkinCacheBehavior->GetInt()) == ESkinCacheDefaultBehavior::Inclusive;
diff --git a/Engine/Source/Runtime/Engine/Private/SkeletalMesh.cpp b/Engine/Source/Runtime/Engine/Private/SkeletalMesh.cpp
--- a/Engine/Source/Runtime/Engine/Private/SkeletalMesh.cpp
+++ b/Engine/Source/Runtime/Engine/Private/SkeletalMesh.cpp
@@ -513,8 +513,8 @@ PRAGMA_ENABLE_DEPRECATION_WARNINGS
 #endif
 	SetMinLod(FPerPlatformInt(0));
 	SetDisableBelowMinLodStripping(FPerPlatformBool(false));
-	SetSupportRayTracing(true);
-	SetRayTracingMinLOD(0);
+	bSupportRayTracing = true;
+	RayTracingMinLOD = 0;
 }
 
 USkeletalMesh::USkeletalMesh(FVTableHelper& Helper)
diff --git a/Engine/Source/Runtime/Engine/Private/SkeletalMeshRenderData.cpp b/Engine/Source/Runtime/Engine/Private/SkeletalMeshRenderData.cpp
--- a/Engine/Source/Runtime/Engine/Private/SkeletalMeshRenderData.cpp
+++ b/Engine/Source/Runtime/Engine/Private/SkeletalMeshRenderData.cpp
@@ -710,7 +710,7 @@ void FSkeletalMeshRenderData::Serialize(FArchive& Ar, USkeletalMesh* Owner)
 	CurrentFirstLODIdx = LODRenderData.Num() - NumInlinedLODs;
 	PendingFirstLODIdx = CurrentFirstLODIdx;
 
-	bSupportRayTracing = Owner->GetSupportRayTracing();
+	bSupportRayTracing = Owner->bSupportRayTracing;
 }
 
 void FSkeletalMeshRenderData::InitResources(bool bNeedsVertexColors, TArray<UMorphTarget*>& InMorphTargets, USkeletalMesh* Owner)
diff --git a/Engine/Source/Runtime/Engine/Private/UnrealEngine.cpp b/Engine/Source/Runtime/Engine/Private/UnrealEngine.cpp
--- a/Engine/Source/Runtime/Engine/Private/UnrealEngine.cpp
+++ b/Engine/Source/Runtime/Engine/Private/UnrealEngine.cpp
@@ -7064,8 +7064,8 @@ bool UEngine::HandleSkeletalMeshReportCommand(const TCHAR* Cmd, FOutputDevice& A
 					ClothVertexBufferKB,
 					bSkinCachedLODEnabled ? TEXT("Y") : TEXT("N"),
 					*RecomputeTangentSections,
-					Mesh->GetSupportRayTracing() ? TEXT("Y") : TEXT("N"),
-					Mesh->GetRayTracingMinLOD());
+					Mesh->bSupportRayTracing ? TEXT("Y") : TEXT("N"),
+					Mesh->RayTracingMinLOD);
 			}
 		}
 
