Adding an UpdateInstances function to UInstancedStaticMeshComponent.  This avoids doing a full MarkRenderStateDirty() call in favor of a MarkRenderTransformDirty() call.  This will only update the instance data for the primitive instead of fully removing and adding the primitive to the scene.

- Added a function, ReallocateGPUSceneInstances, which will free and then alloc all the instances in GPU Scene for the primitives passed into it.  This can be done more efficiently especially if we skip doing an free/alloc if the instance count doesn't change.
- Fixed a bug in UpdateFromCommandBuffer_RenderThread that was setting FStaticMeshInstanceData::NumCustomDataFloats to FStaticMeshInstanceData::NumCustomDataFloats (resulting in the same value being used each time).  Now the new float data count is pulled from the Cmd Buffer which is the single source of truth to update from.
-  FInstanceUpdateCmdBuffer has now been moved to its own header and is the "public" interface for pushing new transforms from the ISM component, through the scene, and into the Proxy.
- Adding tracking for updates/removes and an option to set previous transforms and setting custom float data in bulk on the FInstanceUpdateCmdBuffer.
- Added a new function called UpdateInstances on UInstancedStaticMeshComponent that under the hood does a minimal update of just the transform data instead of add/remove the whole component to the scene.
- The UPrimtiveComponent has the FInstanceUpdateCmdBuffer to push instance transforms through instead of the UInstancedStaticMeshComponent.  This way the code in FScene which pushes data from the PrimtiveComponent down to the proxy can stay generic.  This avoids adding a virtual function call or special casing the code with a cast.
- Adding an argument to avoid freeing the contents of the FInstanceUpdateCmdBuffer.  This way the data can be re-used by the update transform code.
- Adding stats to measure update instance time.

diff --git a/Engine/Plugins/Experimental/GPULightmass/Source/GPULightmass/Private/Scene/Scene.cpp b/Engine/Plugins/Experimental/GPULightmass/Source/GPULightmass/Private/Scene/Scene.cpp
--- a/Engine/Plugins/Experimental/GPULightmass/Source/GPULightmass/Private/Scene/Scene.cpp
+++ b/Engine/Plugins/Experimental/GPULightmass/Source/GPULightmass/Private/Scene/Scene.cpp
@@ -1018,7 +1018,7 @@ void FScene::AddGeometryInstanceFromComponent(UInstancedStaticMeshComponent* InC
 		}
 	}
 
-	InComponent->FlushInstanceUpdateCommands();
+	InComponent->FlushInstanceUpdateCommands(true);
 
 	FInstanceGroupRenderState InstanceRenderState;
 	InstanceRenderState.ComponentUObject = Instance->ComponentUObject;
@@ -1145,7 +1145,7 @@ void FScene::RemoveGeometryInstanceFromComponent(UInstancedStaticMeshComponent*
 		HISMC->BuildTreeIfOutdated(false, true);
 	}
 
-	InComponent->FlushInstanceUpdateCommands();
+	InComponent->FlushInstanceUpdateCommands(true);
 
 	ENQUEUE_RENDER_COMMAND(RenderThreadRemove)(
 		[ElementId, &RenderState = RenderState](FRHICommandListImmediate&) mutable
diff --git a/Engine/Source/Runtime/Engine/Classes/Components/InstancedStaticMeshComponent.h b/Engine/Source/Runtime/Engine/Classes/Components/InstancedStaticMeshComponent.h
--- a/Engine/Source/Runtime/Engine/Classes/Components/InstancedStaticMeshComponent.h
+++ b/Engine/Source/Runtime/Engine/Classes/Components/InstancedStaticMeshComponent.h
@@ -11,6 +11,7 @@
 #include "Engine/TextureStreamingTypes.h"
 #include "Components/StaticMeshComponent.h"
 #include "Elements/SMInstance/SMInstanceManager.h"
+#include "PrimitiveInstanceUpdateCommand.h"
 #include "InstancedStaticMeshComponent.generated.h"
 
 class FLightingBuildOptions;
@@ -29,56 +30,6 @@ class FInstancedLightMap2D;
 class FInstancedShadowMap2D;
 class FStaticMeshInstanceData;
 
-struct FInstanceUpdateCmdBuffer
-{
-	enum EUpdateCommandType
-	{
-		Add,
-		Update,
-		Hide,
-		EditorData,
-		LightmapData,
-		CustomData,
-	};
-	
-	struct FInstanceUpdateCommand
-	{
-		int32 InstanceIndex;
-		EUpdateCommandType Type;
-		FMatrix XForm;
-		
-		FColor HitProxyColor;
-		bool bSelected;
-
-		FVector2D LightmapUVBias;
-		FVector2D ShadowmapUVBias;
-
-		TArray<float> CustomDataFloats;
-	};
-	
-	FInstanceUpdateCmdBuffer();
-	
-	// Commands that can modify render data in place
-	void HideInstance(int32 RenderIndex);
-	void AddInstance(const FMatrix& InTransform);
-	void UpdateInstance(int32 RenderIndex, const FMatrix& InTransform);
-	void SetEditorData(int32 RenderIndex, const FColor& Color, bool bSelected);
-	void SetLightMapData(int32 RenderIndex, const FVector2D& LightmapUVBias);
-	void SetShadowMapData(int32 RenderIndex, const FVector2D& ShadowmapUVBias);
-	void SetCustomData(int32 RenderIndex, const TArray<float>& CustomDataFloats);
-	void ResetInlineCommands();
-	int32 NumInlineCommands() const { return Cmds.Num(); }
-
-	// Command that can't be in-lined and should cause full buffer rebuild
-	void Edit();
-	void Reset();
-	int32 NumTotalCommands() const { return NumEdits; };
-	
-	TArray<FInstanceUpdateCommand> Cmds;
-	int32 NumAdds;
-	int32 NumEdits;
-};
-
 USTRUCT()
 struct FInstancedStaticMeshInstanceData
 {
@@ -268,6 +219,23 @@ class ENGINE_API UInstancedStaticMeshComponent : public UStaticMeshComponent, pu
 	*/
 	virtual bool BatchUpdateInstancesTransforms(int32 StartInstanceIndex, const TArray<FTransform>& NewInstancesTransforms, const TArray<FTransform>& NewInstancesPrevTransforms, bool bWorldSpace = false, bool bMarkRenderStateDirty = false, bool bTeleport = false);
 
+	/**
+	* Lightweight interface to add, remove and update instances.
+	*
+	* @param AddInstanceTransforms				The transforms of the new instances to add.
+	* @param RemoveInstanceIds					The ids of the instances to remove.
+	* @param UpdateInstanceIds					The ids of the new instances to update.
+	* @param UpdateInstanceTransforms			The transforms of the new instances to update.
+	* @param UpdateInstancePreviousTransforms	The transforms of the new instances to update.
+	* @return									True on success
+	*/
+	UFUNCTION(BlueprintCallable, Category = "Components|InstancedStaticMesh")
+	virtual bool UpdateInstances(
+		const TArray<int32>& UpdateInstanceIds, 
+		const TArray<FTransform>& UpdateInstanceTransforms, 
+		const TArray<FTransform>& UpdateInstancePreviousTransforms,
+		int32 NumCustomFloats,
+		const TArray<float>& CustomFloatData);
 
 	/**
 	* Update the transform for a number of instances.
@@ -327,9 +295,6 @@ public:
 	/** Render data will be initialized on PostLoad or on demand. Released on the rendering thread. */
 	TSharedPtr<FPerInstanceRenderData, ESPMode::ThreadSafe> PerInstanceRenderData;
 
-	/** Recorded modifications to per-instance data */
-	FInstanceUpdateCmdBuffer InstanceUpdateCmdBuffer;
-
 	/** 
 	 *  Buffers with per-instance data laid out for rendering. 
 	 *  Serialized for cooked content. Used to create PerInstanceRenderData. 
@@ -351,6 +316,7 @@ public:
 	virtual TStructOnScope<FActorComponentInstanceData> GetComponentInstanceData() const override;
 	virtual void GetComponentChildElements(TArray<FTypedElementHandle>& OutElementHandles, const bool bAllowCreate = true) override;
 	virtual bool IsHLODRelevant() const override;
+	virtual void SendRenderTransform_Concurrent() override;
 	//~ End UActorComponent Interface
 
 	//~ Begin UPrimitiveComponent Interface
@@ -414,7 +380,7 @@ public:
 
 	void GetInstancesMinMaxScale(FVector& MinScale, FVector& MaxScale) const;
 
-	void FlushInstanceUpdateCommands();
+	void FlushInstanceUpdateCommands(bool bFlushInstanceUpdateCmdBuffer);
 
 private:
 
@@ -424,6 +390,12 @@ private:
 	/** Update instance body with a new transform */
 	void UpdateInstanceBodyTransform(int32 InstanceIndex, const FTransform& WorldSpaceInstanceTransform, bool bTeleport);
 
+	/** Used to cache a unique identifier for each instance.  These are provided 
+	*	by the interface UpdateInstances.  This is a map from unique id to index
+	*	into the PerInstanceSMData array.
+	*/
+	TMap<int32, int32> InstanceIdToInstanceIndexMap;
+
 protected:
 	/** Creates body instances for all instances owned by this component. */
 	void CreateAllInstanceBodies();
diff --git a/Engine/Source/Runtime/Engine/Classes/Components/PrimitiveComponent.h b/Engine/Source/Runtime/Engine/Classes/Components/PrimitiveComponent.h
--- a/Engine/Source/Runtime/Engine/Classes/Components/PrimitiveComponent.h
+++ b/Engine/Source/Runtime/Engine/Classes/Components/PrimitiveComponent.h
@@ -23,6 +23,7 @@
 #include "HitProxies.h"
 #include "Interfaces/Interface_AsyncCompilation.h"
 #include "HLOD/HLODBatchingPolicy.h"
+#include "PrimitiveInstanceUpdateCommand.h"
 #include "PrimitiveComponent.generated.h"
 
 class AController;
@@ -799,6 +800,11 @@ public:
 	/** Last time the component was submitted for rendering (called FScene::AddPrimitive). */
 	float LastSubmitTime;
 
+	// This is here because in void FScene::UpdatePrimitiveInstance(UPrimitiveComponent* Primitive) we want access
+	// to the instances updates through the primitive component in a generic way.
+	/** Recorded modifications to per-instance data */
+	FInstanceUpdateCmdBuffer InstanceUpdateCmdBuffer;
+
 private:
 	/**
 	 * The value of WorldSettings->TimeSeconds for the frame when this component was last rendered.  This is written
diff --git a/Engine/Source/Runtime/Engine/Classes/Engine/InstancedStaticMesh.h b/Engine/Source/Runtime/Engine/Classes/Engine/InstancedStaticMesh.h
--- a/Engine/Source/Runtime/Engine/Classes/Engine/InstancedStaticMesh.h
+++ b/Engine/Source/Runtime/Engine/Classes/Engine/InstancedStaticMesh.h
@@ -445,6 +445,7 @@ public:
 
 	FInstancedStaticMeshRenderData(UInstancedStaticMeshComponent* InComponent, ERHIFeatureLevel::Type InFeatureLevel)
 	  : Component(InComponent)
+	  , LightMapCoordinateIndex(Component->GetStaticMesh()->GetLightMapCoordinateIndex())
 	  , PerInstanceRenderData(InComponent->PerInstanceRenderData)
 	  , LODModels(Component->GetStaticMesh()->GetRenderData()->LODResources)
 	  , FeatureLevel(InFeatureLevel)
@@ -475,6 +476,9 @@ public:
 	/** Source component */
 	UInstancedStaticMeshComponent* Component;
 
+	/** Cache off some component data. */
+	int32 LightMapCoordinateIndex;
+
 	/** Per instance render data, could be shared with component */
 	TSharedPtr<FPerInstanceRenderData, ESPMode::ThreadSafe> PerInstanceRenderData;
 
@@ -487,6 +491,8 @@ public:
 	/** Feature level used when creating instance data */
 	ERHIFeatureLevel::Type FeatureLevel;
 
+	void BindBuffersToVertexFactories();
+
 private:
 	void InitVertexFactories();
 
@@ -549,6 +555,8 @@ public:
 
 	virtual void OnTransformChanged() override;
 
+	virtual bool UpdateInstances_RenderThread(const FInstanceUpdateCmdBuffer& CmdBuffer, const FBoxSphereBounds& InBounds, const FBoxSphereBounds& InLocalBounds, const FBoxSphereBounds& InStaticMeshBounds) override;
+
 	virtual FPrimitiveViewRelevance GetViewRelevance(const FSceneView* View) const override
 	{
 		FPrimitiveViewRelevance Result;
diff --git a/Engine/Source/Runtime/Engine/Private/InstancedStaticMesh.cpp b/Engine/Source/Runtime/Engine/Private/InstancedStaticMesh.cpp
--- a/Engine/Source/Runtime/Engine/Private/InstancedStaticMesh.cpp
+++ b/Engine/Source/Runtime/Engine/Private/InstancedStaticMesh.cpp
@@ -174,7 +174,9 @@ FTypedElementHandle HInstancedStaticMeshInstance::GetElementHandle() const
 }
 
 FInstanceUpdateCmdBuffer::FInstanceUpdateCmdBuffer()
-	: NumAdds(0)
+	: NumAdds(0),
+	  NumUpdates(0),
+	  NumRemoves(0)
 	, NumEdits(0)
 {
 }
@@ -187,6 +189,7 @@ void FInstanceUpdateCmdBuffer::HideInstance(int32 RenderIndex)
 	Cmd.InstanceIndex = RenderIndex;
 	Cmd.Type = FInstanceUpdateCmdBuffer::Hide;
 
+	NumRemoves++;
 	Edit();
 }
 
@@ -201,6 +204,18 @@ void FInstanceUpdateCmdBuffer::AddInstance(const FMatrix& InTransform)
 	Edit();
 }
 
+void FInstanceUpdateCmdBuffer::AddInstance(const FMatrix& InTransform, const FMatrix& InPreviousTransform)
+{
+	FInstanceUpdateCommand& Cmd = Cmds.AddDefaulted_GetRef();
+	Cmd.InstanceIndex = INDEX_NONE;
+	Cmd.Type = FInstanceUpdateCmdBuffer::Add;
+	Cmd.XForm = InTransform;
+	Cmd.PreviousXForm = InPreviousTransform;
+
+	NumAdds++;
+	Edit();
+}
+
 void FInstanceUpdateCmdBuffer::UpdateInstance(int32 RenderIndex, const FMatrix& InTransform)
 {
 	FInstanceUpdateCommand& Cmd = Cmds.AddDefaulted_GetRef();
@@ -208,6 +223,19 @@ void FInstanceUpdateCmdBuffer::UpdateInstance(int32 RenderIndex, const FMatrix&
 	Cmd.Type = FInstanceUpdateCmdBuffer::Update;
 	Cmd.XForm = InTransform;
 
+	NumUpdates++;
+	Edit();
+}
+
+void FInstanceUpdateCmdBuffer::UpdateInstance(int32 RenderIndex, const FMatrix& InTransform, const FMatrix& InPreviousTransform)
+{
+	FInstanceUpdateCommand& Cmd = Cmds.AddDefaulted_GetRef();
+	Cmd.InstanceIndex = RenderIndex;
+	Cmd.Type = FInstanceUpdateCmdBuffer::Update;
+	Cmd.XForm = InTransform;
+	Cmd.PreviousXForm = FMatrix(InPreviousTransform);
+
+	NumUpdates++;
 	Edit();
 }
 
@@ -297,13 +325,25 @@ void FInstanceUpdateCmdBuffer::SetCustomData(int32 RenderIndex, const TArray<flo
 		Cmd.CustomDataFloats = CustomDataFloats;
 	}
 
+	// #todo (jnadro) Should really ensure these are all the same size for all instances.
+	NumCustomDataFloats = CustomDataFloats.Num();
+
 	Edit();
 }
 
+void FInstanceUpdateCmdBuffer::SetCustomDataFloatsBulk(const TArray<float>& CustomDataFloats, int32 InNumCustomDataFloats)
+{
+	CustomDataFloatsBulk = CustomDataFloats;
+	NumCustomDataFloats = InNumCustomDataFloats;
+}
+
 void FInstanceUpdateCmdBuffer::ResetInlineCommands()
 {
 	Cmds.Empty();
+	NumCustomDataFloats = 0;
 	NumAdds = 0;
+	NumUpdates = 0;
+	NumRemoves = 0;
 }
 
 void FInstanceUpdateCmdBuffer::Edit()
@@ -314,7 +354,11 @@ void FInstanceUpdateCmdBuffer::Edit()
 void FInstanceUpdateCmdBuffer::Reset()
 {
 	Cmds.Empty();
+	CustomDataFloatsBulk.Empty();
+	NumCustomDataFloats = 0;
 	NumAdds = 0;
+	NumUpdates = 0;
+	NumRemoves = 0;
 	NumEdits = 0;
 }
 
@@ -379,7 +423,7 @@ void FStaticMeshInstanceBuffer::UpdateFromCommandBuffer_RenderThread(FInstanceUp
 		AddIndex = InstanceData->GetNumInstances();
 		int32 NewNumInstances = NumAdds + InstanceData->GetNumInstances();
 
-		InstanceData->AllocateInstances(NewNumInstances, InstanceData->GetNumCustomDataFloats(), GIsEditor ? EResizeBufferFlags::AllowSlackOnGrow | EResizeBufferFlags::AllowSlackOnReduce : EResizeBufferFlags::None, false); // In Editor always permit overallocation, to prevent too much realloc
+		InstanceData->AllocateInstances(NewNumInstances, CmdBuffer.NumCustomDataFloats, GIsEditor ? EResizeBufferFlags::AllowSlackOnGrow | EResizeBufferFlags::AllowSlackOnReduce : EResizeBufferFlags::None, false); // In Editor always permit overallocation, to prevent too much realloc
 	}
 
 	for (int32 i = 0; i < NumCommands; ++i)
@@ -904,6 +948,46 @@ IMPLEMENT_VERTEX_FACTORY_TYPE(FInstancedStaticMeshVertexFactory,"/Engine/Private
 	| EVertexFactoryFlags::SupportsPrimitiveIdStream
 );
 
+void FInstancedStaticMeshRenderData::BindBuffersToVertexFactories()
+{
+	check(IsInRenderingThread());
+
+	PerInstanceRenderData->InstanceBuffer.FlushGPUUpload();
+
+	for (int32 LODIndex = 0; LODIndex < VertexFactories.Num(); LODIndex++)
+	{
+		const FStaticMeshLODResources* RenderData = &LODModels[LODIndex];
+
+		FInstancedStaticMeshVertexFactory::FDataType Data;
+		// Assign to the vertex factory for this LOD.
+		FInstancedStaticMeshVertexFactory& VertexFactory = VertexFactories[LODIndex];
+
+		RenderData->VertexBuffers.PositionVertexBuffer.BindPositionVertexBuffer(&VertexFactory, Data);
+		RenderData->VertexBuffers.StaticMeshVertexBuffer.BindTangentVertexBuffer(&VertexFactory, Data);
+		RenderData->VertexBuffers.StaticMeshVertexBuffer.BindPackedTexCoordVertexBuffer(&VertexFactory, Data);
+		if (LightMapCoordinateIndex < (int32)RenderData->VertexBuffers.StaticMeshVertexBuffer.GetNumTexCoords() && LightMapCoordinateIndex >= 0)
+		{
+			RenderData->VertexBuffers.StaticMeshVertexBuffer.BindLightMapVertexBuffer(&VertexFactory, Data, LightMapCoordinateIndex);
+		}
+
+		if (RenderData->bHasColorVertexData)
+		{
+			RenderData->VertexBuffers.ColorVertexBuffer.BindColorVertexBuffer(&VertexFactory, Data);
+		}
+		else
+		{
+			FColorVertexBuffer::BindDefaultColorVertexBuffer(&VertexFactory, Data, FColorVertexBuffer::NullBindStride::FColorSizeForComponentOverride);
+		}
+
+		check(PerInstanceRenderData);
+
+		PerInstanceRenderData->InstanceBuffer.BindInstanceVertexBuffer(&VertexFactory, Data);
+
+		VertexFactory.SetData(Data);
+		VertexFactory.InitResource();
+	}
+}
+
 void FInstancedStaticMeshRenderData::InitVertexFactories()
 {
 	// Allocate the vertex factories for each LOD
@@ -912,45 +996,11 @@ void FInstancedStaticMeshRenderData::InitVertexFactories()
 		VertexFactories.Add(new FInstancedStaticMeshVertexFactory(FeatureLevel));
 	}
 
-	const int32 LightMapCoordinateIndex = Component->GetStaticMesh()->GetLightMapCoordinateIndex();
 	ENQUEUE_RENDER_COMMAND(InstancedStaticMeshRenderData_InitVertexFactories)(
-		[this, LightMapCoordinateIndex](FRHICommandListImmediate& RHICmdList)
-	{
-		PerInstanceRenderData->InstanceBuffer.FlushGPUUpload();
-
-		for (int32 LODIndex = 0; LODIndex < VertexFactories.Num(); LODIndex++)
+		[this](FRHICommandListImmediate& RHICmdList)
 		{
-			const FStaticMeshLODResources* RenderData = &LODModels[LODIndex];
-
-			FInstancedStaticMeshVertexFactory::FDataType Data;
-			// Assign to the vertex factory for this LOD.
-			FInstancedStaticMeshVertexFactory& VertexFactory = VertexFactories[LODIndex];
-
-			RenderData->VertexBuffers.PositionVertexBuffer.BindPositionVertexBuffer(&VertexFactory, Data);
-			RenderData->VertexBuffers.StaticMeshVertexBuffer.BindTangentVertexBuffer(&VertexFactory, Data);
-			RenderData->VertexBuffers.StaticMeshVertexBuffer.BindPackedTexCoordVertexBuffer(&VertexFactory, Data);
-			if (LightMapCoordinateIndex < (int32)RenderData->VertexBuffers.StaticMeshVertexBuffer.GetNumTexCoords() && LightMapCoordinateIndex >= 0)
-			{
-				RenderData->VertexBuffers.StaticMeshVertexBuffer.BindLightMapVertexBuffer(&VertexFactory, Data, LightMapCoordinateIndex);
-			}
-
-			if (RenderData->bHasColorVertexData)
-			{
-				RenderData->VertexBuffers.ColorVertexBuffer.BindColorVertexBuffer(&VertexFactory, Data);
-			}
-			else
-			{
-				FColorVertexBuffer::BindDefaultColorVertexBuffer(&VertexFactory, Data, FColorVertexBuffer::NullBindStride::FColorSizeForComponentOverride);
-			}
-			
-			check(PerInstanceRenderData);
-			
-			PerInstanceRenderData->InstanceBuffer.BindInstanceVertexBuffer(&VertexFactory, Data);
-
-			VertexFactory.SetData(Data);
-			VertexFactory.InitResource();
-		}
-	});
+			BindBuffersToVertexFactories();
+		});
 }
 
 
@@ -1431,6 +1481,14 @@ void FInstancedStaticMeshSceneProxy::OnTransformChanged()
 	}
 }
 
+bool FInstancedStaticMeshSceneProxy::UpdateInstances_RenderThread(const FInstanceUpdateCmdBuffer& CmdBuffer, const FBoxSphereBounds& InBounds, const FBoxSphereBounds& InLocalBounds, const FBoxSphereBounds& InStaticMeshBounds)
+{
+	// This will flush GPU instance data, create buffers, srvs and vertex factories.
+	InstancedRenderData.BindBuffersToVertexFactories();
+
+	return FPrimitiveSceneProxy::UpdateInstances_RenderThread(CmdBuffer, InBounds, InLocalBounds, InStaticMeshBounds);
+}
+
 void FInstancedStaticMeshSceneProxy::SetupInstancedMeshBatch(int32 LODIndex, int32 BatchIndex, FMeshBatch& OutMeshBatch) const
 {
 	OutMeshBatch.VertexFactory = &InstancedRenderData.VertexFactories[LODIndex];
@@ -2019,9 +2077,12 @@ void UInstancedStaticMeshComponent::ApplyComponentInstanceData(FInstancedStaticM
 #endif
 }
 
-void UInstancedStaticMeshComponent::FlushInstanceUpdateCommands()
+void UInstancedStaticMeshComponent::FlushInstanceUpdateCommands(bool bFlushInstanceUpdateCmdBuffer)
 {
-	InstanceUpdateCmdBuffer.Reset();
+	if (bFlushInstanceUpdateCmdBuffer)
+	{
+		InstanceUpdateCmdBuffer.Reset();
+	}
 
 	FStaticMeshInstanceData RenderInstanceData = FStaticMeshInstanceData(GVertexElementTypeSupport.IsSupported(VET_Half2));
 	BuildRenderData(RenderInstanceData, PerInstanceRenderData->HitProxies);
@@ -2038,6 +2099,20 @@ bool UInstancedStaticMeshComponent::IsHLODRelevant() const
 	return Super::IsHLODRelevant();
 }
 
+void UInstancedStaticMeshComponent::SendRenderTransform_Concurrent()
+{
+	// If the primitive isn't hidden update its transform.
+	const bool bDetailModeAllowsRendering = DetailMode <= GetCachedScalabilityCVars().DetailMode;
+	if ((InstanceUpdateCmdBuffer.NumEdits || InstanceUpdateCmdBuffer.NumAdds) && bDetailModeAllowsRendering && (ShouldRender() || bCastHiddenShadow || bRayTracingFarField))
+	{
+		UpdateBounds();
+
+		// Update the scene info's transform for this primitive.
+		GetWorld()->Scene->UpdatePrimitiveInstances(this);
+		InstanceUpdateCmdBuffer.Reset();
+	}
+}
+
 FBodyInstance* UInstancedStaticMeshComponent::GetBodyInstance(FName BoneName, bool bGetWelded, int32 Index) const
 {
 	if (Index != INDEX_NONE && IsValidInstance(Index))
@@ -2073,7 +2148,7 @@ FPrimitiveSceneProxy* UInstancedStaticMeshComponent::CreateSceneProxy()
 		// generally happens only in editor 
 		if (InstanceUpdateCmdBuffer.NumTotalCommands() != 0)
 		{
-			FlushInstanceUpdateCommands();
+			FlushInstanceUpdateCommands(true);
 		}
 		
 		ProxySize = PerInstanceRenderData->ResourceSize;
@@ -2705,7 +2780,7 @@ void UInstancedStaticMeshComponent::SerializeRenderData(FArchive& Ar)
 				// This will usually happen when having a BP adding instance through the construct script
 				if (PerInstanceRenderData->InstanceBuffer.GetNumInstances() != PerInstanceSMData.Num() || InstanceUpdateCmdBuffer.NumTotalCommands() > 0)
 				{
-					FlushInstanceUpdateCommands();
+					FlushInstanceUpdateCommands(true);
 					MarkRenderStateDirty();
 				}
 			}
@@ -3252,6 +3327,128 @@ bool UInstancedStaticMeshComponent::BatchUpdateInstancesTransforms(int32 StartIn
 	return true;
 }
 
+bool UInstancedStaticMeshComponent::UpdateInstances(
+	const TArray<int32>& UpdateInstanceIds, 
+	const TArray<FTransform>& UpdateInstanceTransforms, 
+	const TArray<FTransform>& UpdateInstancePreviousTransforms,
+	int32 InNumCustomDataFloats,
+	const TArray<float>& CustomFloatData)
+{
+	TRACE_CPUPROFILER_EVENT_SCOPE_STR("ISM UpdateInstances")
+
+	if (UpdateInstanceIds.Num() != UpdateInstanceTransforms.Num())
+	{
+		UE_LOG(LogStaticMesh, Warning, TEXT("Add instances to ISM has missmatched update arrays.  UpdateInstanceIds and UpdateInstanceTransforms should match."));
+		return false;
+	}
+
+	if (UpdateInstanceTransforms.Num() != UpdateInstancePreviousTransforms.Num())
+	{
+		UE_LOG(LogStaticMesh, Warning, TEXT("Add instances to ISM has missmatched update arrays.  Current transform and previous transform arrays should match."));
+		return false;
+	}
+
+	if (InNumCustomDataFloats > 0)
+	{
+		if ((CustomFloatData.Num() / InNumCustomDataFloats) != UpdateInstanceIds.Num())
+		{
+			UE_LOG(LogStaticMesh, Warning, TEXT("CustomFloatData only contains %d entires.  Which is not enough for NumCustomFloat count %d * %d instances."), CustomFloatData.Num(), InNumCustomDataFloats, UpdateInstanceIds.Num());
+			return false;
+		}
+	}
+
+	const int32 NewNumInstances = UpdateInstanceIds.Num();
+	const int32 PreviousNumInstances = PerInstanceSMData.Num();
+
+	// If we didn't have any instances and there are no new instances
+	// skip doing any work.
+	if (PreviousNumInstances == 0 && NewNumInstances == 0)
+	{
+		return false;
+	}
+
+	if (CustomFloatData.Num())
+	{
+		int32 MyNumCustomDataFloats = CustomFloatData.Num() / UpdateInstanceIds.Num();
+		check(UpdateInstanceIds.Num() * MyNumCustomDataFloats == CustomFloatData.Num()); // Temp sanity check
+	}
+
+	Modify();
+
+	// Need to empty the command buffer if it already has values.
+	// If this function is called multiple times the last set of values will take
+	// precedence.  Due to the way MarkRenderTransformDirty() works the primitive will
+	// only be processed once.
+	InstanceUpdateCmdBuffer.Reset();
+
+	NumCustomDataFloats = InNumCustomDataFloats;
+	PerInstanceSMCustomData = CustomFloatData;
+
+	InstanceUpdateCmdBuffer.SetCustomDataFloatsBulk(CustomFloatData, NumCustomDataFloats);
+
+	PerInstanceSMData.SetNum(NewNumInstances, /*bAllowShrinking*/false);
+	PerInstancePrevTransform.SetNum(NewNumInstances, /*bAllowShrinking*/false);
+
+	// if we already have values we need to update, remove, and add.
+	TMap<int32, int32> OldInstanceIdToInstanceIndexMap(InstanceIdToInstanceIndexMap);
+	InstanceIdToInstanceIndexMap.Empty();
+
+	// Apply updates
+	for (int32 i = 0; i < UpdateInstanceIds.Num(); ++i)
+	{
+		FInstancedStaticMeshInstanceData& InstanceData = PerInstanceSMData[i];
+		InstanceData.Transform = UpdateInstanceTransforms[i].ToMatrixWithScale();
+		PerInstancePrevTransform[i] = UpdateInstancePreviousTransforms[i].ToMatrixWithScale();
+
+		// This is an update.
+		const int32 InstanceId = UpdateInstanceIds[i];
+		if (int32* pInstanceIndex = OldInstanceIdToInstanceIndexMap.Find(InstanceId))
+		{
+			InstanceUpdateCmdBuffer.UpdateInstance(*pInstanceIndex, InstanceData.Transform, PerInstancePrevTransform[i]);
+			OldInstanceIdToInstanceIndexMap.Remove(InstanceId);
+		}
+		// This is an add.
+		else
+		{
+			InstanceUpdateCmdBuffer.AddInstance(InstanceData.Transform, PerInstancePrevTransform[i]);
+		}
+
+		// need to make sure the indices we tracked are updated.
+		InstanceIdToInstanceIndexMap.Add(InstanceId, i);
+	}
+
+	// Inform the cmd buffer which instances were removed.
+	for (const auto& InstanceMapping : OldInstanceIdToInstanceIndexMap)
+	{
+		const int32& InstanceIndex = InstanceMapping.Value;
+
+		// not sure if we really need to record this information.
+		// this instance index is the OLD one and is no longer valid.
+		InstanceUpdateCmdBuffer.HideInstance(InstanceIndex);
+	}
+
+	// #todo (jnadro) InstanceIndex is ignored by this function so just pass zero.
+	PartialNavigationUpdate(0);
+
+	// #todo (jnadro) Updating the PerInstanceRenderData this way currently does not work.
+	// Updating the PerInstanceRenderData from the InstanceUpdateCmdBuffer causes
+	// the instances to be in a different order than what is in GPU scene.  This 
+	// causes all sorts of visual artifacts because custom float data doesn't match
+	// with the tranforms in GPU Scene.
+	// PerInstanceRenderData->UpdateFromCommandBuffer(InstanceUpdateCmdBuffer, false);
+
+	// if instance data was modified update it.  This pulls data directly
+	// from the component which we have updated above.
+	if (InstanceUpdateCmdBuffer.NumTotalCommands() != 0)
+	{
+		FlushInstanceUpdateCommands(false);
+	}
+
+	MarkRenderTransformDirty();
+
+	return true;
+}
+
 bool UInstancedStaticMeshComponent::BatchUpdateInstancesTransforms(int32 StartInstanceIndex, const TArray<FTransform>& NewInstancesTransforms, bool bWorldSpace, bool bMarkRenderStateDirty, bool bTeleport)
 {
 	if (!PerInstanceSMData.IsValidIndex(StartInstanceIndex) || !PerInstanceSMData.IsValidIndex(StartInstanceIndex + NewInstancesTransforms.Num() - 1))
diff --git a/Engine/Source/Runtime/Engine/Private/PrimitiveSceneProxy.cpp b/Engine/Source/Runtime/Engine/Private/PrimitiveSceneProxy.cpp
--- a/Engine/Source/Runtime/Engine/Private/PrimitiveSceneProxy.cpp
+++ b/Engine/Source/Runtime/Engine/Private/PrimitiveSceneProxy.cpp
@@ -19,6 +19,7 @@
 #include "FoliageHelper.h"
 #include "ObjectCacheEventSink.h"
 #endif
+#include "PrimitiveInstanceUpdateCommand.h"
 
 static TAutoConsoleVariable<int32> CVarForceSingleSampleShadowingFromStationary(
 	TEXT("r.Shadow.ForceSingleSampleShadowingFromStationary"),
@@ -537,6 +538,57 @@ void FPrimitiveSceneProxy::SetTransform(const FMatrix& InLocalToWorld, const FBo
 	OnTransformChanged();
 }
 
+bool FPrimitiveSceneProxy::UpdateInstances_RenderThread(const FInstanceUpdateCmdBuffer& CmdBuffer, const FBoxSphereBounds& InBounds, const FBoxSphereBounds& InLocalBounds, const FBoxSphereBounds& InStaticMeshBounds)
+{
+	SCOPED_NAMED_EVENT(FPrimitiveSceneProxy_UpdateInstances_RenderThread, FColor::Emerald);
+
+	check(IsInRenderingThread());
+
+	bool bInstanceCountChanged = false;
+	if (UseGPUScene(GetScene().GetShaderPlatform(), GetScene().GetFeatureLevel()))
+	{
+		const uint32 NumInstances = CmdBuffer.NumUpdates + CmdBuffer.NumAdds;
+		bInstanceCountChanged = NumInstances != InstanceSceneData.Num();
+
+		InstanceSceneData.SetNum(NumInstances);
+		InstanceLocalBounds.SetNumUninitialized(1);
+		InstanceLocalBounds[0] = InStaticMeshBounds;
+		InstanceDynamicData.SetNumUninitialized(NumInstances);
+		InstanceRandomID.SetNumZeroed(NumInstances);
+		InstanceLightShadowUVBias.SetNumZeroed(NumInstances);
+		InstanceCustomData = CmdBuffer.CustomDataFloatsBulk;
+
+		if (InstanceCustomData.Num())
+		{
+			int32 MyNumCustomDataFloats = InstanceCustomData.Num() / InstanceSceneData.Num();
+			check(InstanceSceneData.Num() * MyNumCustomDataFloats == InstanceCustomData.Num()); // Temp sanity check
+		}
+
+		bHasPerInstanceRandom = InstanceRandomID.Num() > 0;
+		bHasPerInstanceCustomData = InstanceCustomData.Num() > 0;
+		bHasPerInstanceDynamicData = InstanceDynamicData.Num() > 0;
+		bHasPerInstanceLMSMUVBias = InstanceLightShadowUVBias.Num() > 0;
+
+		int32 i = 0;
+		for (const auto& Cmd : CmdBuffer.Cmds)
+		{
+			if (Cmd.Type == FInstanceUpdateCmdBuffer::Add || Cmd.Type == FInstanceUpdateCmdBuffer::Update)
+			{
+				FPrimitiveInstance& SceneData = InstanceSceneData[i];
+				SceneData.LocalToPrimitive = Cmd.XForm;
+
+				FPrimitiveInstanceDynamicData& DynamicData = InstanceDynamicData[i];
+				DynamicData.PrevLocalToPrimitive = Cmd.PreviousXForm;
+
+				i++;
+			}
+		}
+		ensure(NumInstances == i);
+	}
+
+	return bInstanceCountChanged;
+}
+
 bool FPrimitiveSceneProxy::WouldSetTransformBeRedundant_AnyThread(const FMatrix& InLocalToWorld, const FBoxSphereBounds& InBounds, const FBoxSphereBounds& InLocalBounds, const FVector& InActorPosition) const
 {
 	// Order is based on cheapest tests first.
diff --git a/Engine/Source/Runtime/Engine/Public/PrimitiveInstanceUpdateCommand.h b/Engine/Source/Runtime/Engine/Public/PrimitiveInstanceUpdateCommand.h
new file mode 100644
--- /dev/null
+++ b/Engine/Source/Runtime/Engine/Public/PrimitiveInstanceUpdateCommand.h
@@ -0,0 +1,63 @@
+// Copyright Epic Games, Inc. All Rights Reserved.
+
+#pragma once
+
+#include "CoreMinimal.h"
+
+struct FInstanceUpdateCmdBuffer
+{
+	enum EUpdateCommandType
+	{
+		Add,
+		Update,
+		Hide,
+		EditorData,
+		LightmapData,
+		CustomData,
+	};
+
+	struct FInstanceUpdateCommand
+	{
+		int32 InstanceIndex;
+		EUpdateCommandType Type;
+		FMatrix XForm;
+		FMatrix PreviousXForm;
+
+		FColor HitProxyColor;
+		bool bSelected;
+
+		FVector2D LightmapUVBias;
+		FVector2D ShadowmapUVBias;
+
+		TArray<float> CustomDataFloats;
+	};
+
+	ENGINE_API FInstanceUpdateCmdBuffer();
+
+	// Commands that can modify render data in place
+	void HideInstance(int32 RenderIndex);
+	void AddInstance(const FMatrix& InTransform);
+	void AddInstance(const FMatrix& InTransform, const FMatrix& InPreviousTransform);
+	void UpdateInstance(int32 RenderIndex, const FMatrix& InTransform);
+	void UpdateInstance(int32 RenderIndex, const FMatrix& InTransform, const FMatrix& InPreviousTransform);
+	void SetEditorData(int32 RenderIndex, const FColor& Color, bool bSelected);
+	void SetLightMapData(int32 RenderIndex, const FVector2D& LightmapUVBias);
+	void SetShadowMapData(int32 RenderIndex, const FVector2D& ShadowmapUVBias);
+	void SetCustomData(int32 RenderIndex, const TArray<float>& CustomDataFloats);
+	void SetCustomDataFloatsBulk(const TArray<float>& CustomDataFloats, int32 InNumCustomDataFloats);
+	void ResetInlineCommands();
+	int32 NumInlineCommands() const { return Cmds.Num(); }
+
+	// Command that can't be in-lined and should cause full buffer rebuild
+	void Edit();
+	void Reset();
+	int32 NumTotalCommands() const { return NumEdits; };
+
+	TArray<FInstanceUpdateCommand> Cmds;
+	TArray<float> CustomDataFloatsBulk;
+	int32 NumCustomDataFloats;
+	int32 NumAdds;
+	int32 NumUpdates;
+	int32 NumRemoves;
+	int32 NumEdits;
+};
diff --git a/Engine/Source/Runtime/Engine/Public/PrimitiveSceneProxy.h b/Engine/Source/Runtime/Engine/Public/PrimitiveSceneProxy.h
--- a/Engine/Source/Runtime/Engine/Public/PrimitiveSceneProxy.h
+++ b/Engine/Source/Runtime/Engine/Public/PrimitiveSceneProxy.h
@@ -28,6 +28,7 @@ class UTexture2D;
 enum class ERuntimeVirtualTextureMaterialType : uint8;
 struct FMeshBatch;
 class FColorVertexBuffer;
+struct FInstanceUpdateCmdBuffer;
 
 namespace Nanite
 {
@@ -1349,6 +1350,16 @@ private:
 	void SetHiddenEdViews_RenderThread( uint64 InHiddenEditorViews );
 
 protected:
+	/**
+	 * Updates the primitive proxy's cached transforms for all instances given a buffer of instance updates.
+	 * @param CmdBuffer - A record of all the add, update and remove instances for the proxy to apply to its internal data.
+	 * @param InBounds - Primitive world space bounds.
+	 * @param InLocalBounds - Primitive local space bounds.
+	 * @param InStaticMeshBounds - Bounds of the primitive mesh instance.
+	 * @returns True if instance count has changed since last time.
+	 */
+	ENGINE_API virtual bool UpdateInstances_RenderThread(const FInstanceUpdateCmdBuffer& CmdBuffer, const FBoxSphereBounds& InBounds, const FBoxSphereBounds& InLocalBounds, const FBoxSphereBounds& InStaticMeshBounds);
+
 	/** Updates selection for the primitive proxy. This is called in the rendering thread by SetSelection_GameThread. */
 	void SetSelection_RenderThread(const bool bInParentSelected, const bool bInIndividuallySelected);
 
diff --git a/Engine/Source/Runtime/Engine/Public/SceneInterface.h b/Engine/Source/Runtime/Engine/Public/SceneInterface.h
--- a/Engine/Source/Runtime/Engine/Public/SceneInterface.h
+++ b/Engine/Source/Runtime/Engine/Public/SceneInterface.h
@@ -84,6 +84,14 @@ public:
 	virtual void UpdatePrimitiveOcclusionBoundsSlack(UPrimitiveComponent* Primitive, float NewSlack) = 0;
 	/** Updates primitive attachment state. */
 	virtual void UpdatePrimitiveAttachment(UPrimitiveComponent* Primitive) = 0;
+
+	/**
+	 * Updates all the instances that have been updated through the InstanceUpdateCmdBuffer on the UPrimitiveComponent.
+	 *
+	 * @param Primitive - primitive component to update
+	 */
+	virtual void UpdatePrimitiveInstances(UPrimitiveComponent* Primitive) = 0;
+
 	/** 
 	 * Updates the custom primitive data of a primitive component which has already been added to the scene. 
 	 * 
diff --git a/Engine/Source/Runtime/RenderCore/Private/RenderCore.cpp b/Engine/Source/Runtime/RenderCore/Private/RenderCore.cpp
--- a/Engine/Source/Runtime/RenderCore/Private/RenderCore.cpp
+++ b/Engine/Source/Runtime/RenderCore/Private/RenderCore.cpp
@@ -177,11 +177,13 @@ DEFINE_STAT(STAT_RemoveScenePrimitiveTime);
 DEFINE_STAT(STAT_AddScenePrimitiveRenderThreadTime);
 DEFINE_STAT(STAT_UpdateScenePrimitiveRenderThreadTime);
 DEFINE_STAT(STAT_UpdatePrimitiveTransformRenderThreadTime);
+DEFINE_STAT(STAT_UpdatePrimitiveInstanceRenderThreadTime);
 DEFINE_STAT(STAT_FlushAsyncLPICreation);
 
 DEFINE_STAT(STAT_RemoveScenePrimitiveGT);
 DEFINE_STAT(STAT_AddScenePrimitiveGT);
 DEFINE_STAT(STAT_UpdatePrimitiveTransformGT);
+DEFINE_STAT(STAT_UpdatePrimitiveInstanceGT);
 DEFINE_STAT(STAT_UpdateCustomPrimitiveDataGT);
 
 DEFINE_STAT(STAT_Scene_SetShaderMapsOnMaterialResources_RT);
diff --git a/Engine/Source/Runtime/RenderCore/Public/RenderCore.h b/Engine/Source/Runtime/RenderCore/Public/RenderCore.h
--- a/Engine/Source/Runtime/RenderCore/Public/RenderCore.h
+++ b/Engine/Source/Runtime/RenderCore/Public/RenderCore.h
@@ -123,6 +123,7 @@ DECLARE_CYCLE_STAT_EXTERN(TEXT("RemovePrimitive"),STAT_RemoveScenePrimitiveTime,
 DECLARE_CYCLE_STAT_EXTERN(TEXT("RemoveLight"),STAT_RemoveSceneLightTime,STATGROUP_SceneUpdate, RENDERCORE_API);
 DECLARE_CYCLE_STAT_EXTERN(TEXT("UpdateLight"),STAT_UpdateSceneLightTime,STATGROUP_SceneUpdate, RENDERCORE_API);
 DECLARE_CYCLE_STAT_EXTERN(TEXT("UpdatePrimitiveTransform"),STAT_UpdatePrimitiveTransformRenderThreadTime,STATGROUP_SceneUpdate, RENDERCORE_API);
+DECLARE_CYCLE_STAT_EXTERN(TEXT("UpdatePrimitiveInstance"), STAT_UpdatePrimitiveInstanceRenderThreadTime, STATGROUP_SceneUpdate, RENDERCORE_API);
 DECLARE_CYCLE_STAT_EXTERN(TEXT("Update CPU Skin"),STAT_CPUSkinUpdateRTTime,STATGROUP_SceneUpdate, RENDERCORE_API);
 DECLARE_CYCLE_STAT_EXTERN(TEXT("Update GPU Skin"),STAT_GPUSkinUpdateRTTime,STATGROUP_SceneUpdate, RENDERCORE_API);
 DECLARE_CYCLE_STAT_EXTERN(TEXT("Update particles"),STAT_ParticleUpdateRTTime,STATGROUP_SceneUpdate, RENDERCORE_API);
@@ -132,6 +133,7 @@ DECLARE_CYCLE_STAT_EXTERN(TEXT("Flush async LPI creation"), STAT_FlushAsyncLPICr
 DECLARE_CYCLE_STAT_EXTERN(TEXT("RemovePrimitive (GT)"),STAT_RemoveScenePrimitiveGT,STATGROUP_SceneUpdate, RENDERCORE_API);
 DECLARE_CYCLE_STAT_EXTERN(TEXT("AddPrimitive (GT)"),STAT_AddScenePrimitiveGT,STATGROUP_SceneUpdate, RENDERCORE_API);
 DECLARE_CYCLE_STAT_EXTERN(TEXT("UpdatePrimitiveTransform (GT)"),STAT_UpdatePrimitiveTransformGT,STATGROUP_SceneUpdate, RENDERCORE_API);
+DECLARE_CYCLE_STAT_EXTERN(TEXT("UpdatePrimitiveInstance (GT)"), STAT_UpdatePrimitiveInstanceGT, STATGROUP_SceneUpdate, RENDERCORE_API);
 DECLARE_CYCLE_STAT_EXTERN(TEXT("UpdateCustomPrimitiveData (GT)"),STAT_UpdateCustomPrimitiveDataGT,STATGROUP_SceneUpdate, RENDERCORE_API);
 
 DECLARE_CYCLE_STAT_EXTERN(TEXT("Set Shader Maps On Material Resources (RT)"),STAT_Scene_SetShaderMapsOnMaterialResources_RT,STATGROUP_SceneUpdate, RENDERCORE_API);
diff --git a/Engine/Source/Runtime/Renderer/Private/PrimitiveSceneInfo.cpp b/Engine/Source/Runtime/Renderer/Private/PrimitiveSceneInfo.cpp
--- a/Engine/Source/Runtime/Renderer/Private/PrimitiveSceneInfo.cpp
+++ b/Engine/Source/Runtime/Renderer/Private/PrimitiveSceneInfo.cpp
@@ -1087,6 +1087,20 @@ void FPrimitiveSceneInfo::AllocateGPUSceneInstances(FScene* Scene, const TArrayV
 	}
 }
 
+void FPrimitiveSceneInfo::ReallocateGPUSceneInstances(FScene* Scene, const TArrayView<FPrimitiveSceneInfo*>& SceneInfos)
+{
+	SCOPED_NAMED_EVENT(ReallocateGPUSceneInstances, FColor::Emerald);
+
+	// Free each scene info.
+	for (FPrimitiveSceneInfo* SceneInfo : SceneInfos)
+	{
+		SceneInfo->FreeGPUSceneInstances();
+	}
+
+	// Allocate them all.
+	AllocateGPUSceneInstances(Scene, SceneInfos);
+}
+
 void FPrimitiveSceneInfo::FreeGPUSceneInstances()
 {
 	if (!Scene->GPUScene.IsEnabled())
diff --git a/Engine/Source/Runtime/Renderer/Private/RendererScene.cpp b/Engine/Source/Runtime/Renderer/Private/RendererScene.cpp
--- a/Engine/Source/Runtime/Renderer/Private/RendererScene.cpp
+++ b/Engine/Source/Runtime/Renderer/Private/RendererScene.cpp
@@ -1470,6 +1470,50 @@ void FScene::UpdatePrimitiveOcclusionBoundsSlack(UPrimitiveComponent* Primitive,
 	}
 }
 
+void FScene::UpdatePrimitiveInstances(UPrimitiveComponent* Primitive)
+{
+	SCOPE_CYCLE_COUNTER(STAT_UpdatePrimitiveInstanceGT);
+	SCOPED_NAMED_EVENT(FScene_UpdatePrimitiveInstance, FColor::Yellow);
+
+	if (Primitive->SceneProxy)
+	{
+		FUpdateInstanceCommand UpdateParams;
+		UpdateParams.PrimitiveSceneProxy = Primitive->SceneProxy;
+		UpdateParams.WorldBounds = Primitive->Bounds;
+		UpdateParams.LocalBounds = Primitive->GetLocalBounds();
+
+		// #todo (jnadro) This code should not be dependent on static mesh bounds.
+		UStaticMeshComponent* StaticMeshComponent = Cast<UStaticMeshComponent>(Primitive);
+		UpdateParams.StaticMeshBounds = StaticMeshComponent ? StaticMeshComponent->GetStaticMesh()->GetBounds() : FBoxSphereBounds();
+		UpdateParams.CmdBuffer = Primitive->InstanceUpdateCmdBuffer;	// Copy
+
+		ENQUEUE_RENDER_COMMAND(UpdateInstanceCommand)(
+			[this, UpdateParams = MoveTemp(UpdateParams)](FRHICommandListImmediate& RHICmdList)
+			{
+#if VALIDATE_PRIMITIVE_PACKED_INDEX
+				if (AddedPrimitiveSceneInfos.Contains(UpdateParams.PrimitiveSceneProxy->GetPrimitiveSceneInfo()))
+				{
+					check(UpdateParams.PrimitiveSceneProxy->GetPrimitiveSceneInfo()->PackedIndex == INDEX_NONE);
+				}
+				else
+				{
+					check(UpdateParams.PrimitiveSceneProxy->GetPrimitiveSceneInfo()->PackedIndex != INDEX_NONE);
+				}
+
+				check(!RemovedPrimitiveSceneInfos.Contains(UpdateParams.PrimitiveSceneProxy->GetPrimitiveSceneInfo()));
+#endif
+				FScopeCycleCounter Context(UpdateParams.PrimitiveSceneProxy->GetStatId());
+				UpdatedInstances.Add(UpdateParams.PrimitiveSceneProxy, UpdateParams);
+			}
+		);
+	}
+ 	else
+ 	{
+ 		// If the primitive doesn't have a scene info object yet, it must be added from scratch.
+ 		AddPrimitive(Primitive);
+ 	}
+}
+
 void FScene::UpdatePrimitiveSelectedState_RenderThread(const FPrimitiveSceneInfo* PrimitiveSceneInfo, bool bIsSelected)
 {
 	check(IsInRenderingThread());
@@ -4639,6 +4683,74 @@ void FScene::UpdateAllPrimitiveSceneInfos(FRDGBuilder& GraphBuilder, bool bAsync
 		}
 	}
 
+	{
+		CSV_SCOPED_TIMING_STAT_EXCLUSIVE(UpdatePrimitiveInstances);
+		SCOPED_NAMED_EVENT(FScene_UpdatePrimitiveInstances, FColor::Emerald);
+		SCOPE_CYCLE_COUNTER(STAT_UpdatePrimitiveInstanceRenderThreadTime);
+
+		TArray<FPrimitiveSceneInfo*> UpdatedInstancesLocalPrimitiveSceneInfos;
+
+		TArray<FPrimitiveSceneProxy*> SameSceneInfos;
+
+		for (const auto& UpdateInstance : UpdatedInstances)
+		{
+			FPrimitiveSceneProxy* PrimitiveSceneProxy = UpdateInstance.Key;
+			if (DeletedSceneInfos.Contains(PrimitiveSceneProxy->GetPrimitiveSceneInfo()))
+			{
+				continue;
+			}
+
+			check(!SameSceneInfos.Contains(PrimitiveSceneProxy));
+			SameSceneInfos.Add(PrimitiveSceneProxy);
+
+			check(PrimitiveSceneProxy->GetPrimitiveSceneInfo()->PackedIndex != INDEX_NONE);
+
+			FScopeCycleCounter Context(PrimitiveSceneProxy->GetStatId());
+			FPrimitiveSceneInfo* PrimitiveSceneInfo = PrimitiveSceneProxy->GetPrimitiveSceneInfo();
+
+			// Update the Proxy's data.
+			PrimitiveSceneProxy->UpdateInstances_RenderThread(UpdateInstance.Value.CmdBuffer, UpdateInstance.Value.WorldBounds, UpdateInstance.Value.LocalBounds, UpdateInstance.Value.StaticMeshBounds);
+
+			if (!RHISupportsVolumeTextures(GetFeatureLevel())
+				&& (PrimitiveSceneProxy->IsMovable() || PrimitiveSceneProxy->NeedsUnbuiltPreviewLighting() || PrimitiveSceneProxy->GetLightmapType() == ELightmapType::ForceVolumetric))
+			{
+				PrimitiveSceneInfo->MarkIndirectLightingCacheBufferDirty();
+			}
+
+			UpdatedInstancesLocalPrimitiveSceneInfos.Add(PrimitiveSceneInfo);
+		}
+
+		if (UpdatedInstancesLocalPrimitiveSceneInfos.Num() > 0)
+		{
+			{
+				SCOPED_NAMED_EVENT(CacheDrawCommands, FColor::Emerald);
+
+				// Update all the static mesh draw lists...just so we can patch the instance count.
+				for (FPrimitiveSceneInfo* SceneInfo : UpdatedInstancesLocalPrimitiveSceneInfos)
+				{
+					SceneInfo->RemoveStaticMeshes();
+				}
+				FPrimitiveSceneInfo::AddStaticMeshes(GraphBuilder.RHICmdList, this, UpdatedInstancesLocalPrimitiveSceneInfos);
+			}
+
+			// Do system updates in the same order that Update Primitive Transform does.
+			// 1. GPU Scene
+			// 2. Distance Fields
+			// 3. Lumen
+			// 4. Ray Tracing (This gets updated in FPrimitiveSceneInfo::AddStaticMeshes above).
+
+			// GPU Scene
+			FPrimitiveSceneInfo::ReallocateGPUSceneInstances(this, TArrayView<FPrimitiveSceneInfo*>(&UpdatedInstancesLocalPrimitiveSceneInfos[0], UpdatedInstancesLocalPrimitiveSceneInfos.Num()));
+
+			// Distance Fields
+			for (FPrimitiveSceneInfo* PrimitiveSceneInfo : UpdatedInstancesLocalPrimitiveSceneInfos)
+			{
+				// This clobbers any existing instance data and pulls new data from the proxy.
+				DistanceFieldSceneData.UpdatePrimitive(PrimitiveSceneInfo);
+			}
+		}
+	}
+
 	for (const auto& Attachments : UpdatedAttachmentRoots)
 	{
 		FPrimitiveSceneInfo* PrimitiveSceneInfo = Attachments.Key;
@@ -4725,6 +4837,7 @@ void FScene::UpdateAllPrimitiveSceneInfos(FRDGBuilder& GraphBuilder, bool bAsync
 
 	UpdatedAttachmentRoots.Reset();
 	UpdatedTransforms.Reset();
+	UpdatedInstances.Reset();
 	UpdatedCustomPrimitiveParams.Reset();
 	OverridenPreviousTransforms.Reset();
 	UpdatedOcclusionBoundsSlacks.Reset();
@@ -4790,6 +4903,7 @@ public:
 
 	/** Updates the transform of a primitive which has already been added to the scene. */
 	virtual void UpdatePrimitiveTransform(UPrimitiveComponent* Primitive) override {}
+	virtual void UpdatePrimitiveInstances(UPrimitiveComponent* Primitive) override {}
 	virtual void UpdatePrimitiveOcclusionBoundsSlack(UPrimitiveComponent* Primitive, float NewSlack) override {}
 	virtual void UpdatePrimitiveAttachment(UPrimitiveComponent* Primitive) override {}
 	virtual void UpdateCustomPrimitiveData(UPrimitiveComponent* Primitive) override {}
diff --git a/Engine/Source/Runtime/Renderer/Private/ScenePrivate.h b/Engine/Source/Runtime/Renderer/Private/ScenePrivate.h
--- a/Engine/Source/Runtime/Renderer/Private/ScenePrivate.h
+++ b/Engine/Source/Runtime/Renderer/Private/ScenePrivate.h
@@ -57,6 +57,7 @@
 #include "VolumetricRenderTargetViewStateData.h"
 #include "GPUScene.h"
 #include "DynamicBVH.h"
+#include "PrimitiveInstanceUpdateCommand.h"
 #include "OIT/OIT.h"
 #include "ShadingEnergyConservation.h"
 
@@ -3030,7 +3031,6 @@ public:
 	virtual void ReleasePrimitive(UPrimitiveComponent* Primitive) override;
 	virtual void UpdateAllPrimitiveSceneInfos(FRDGBuilder& GraphBuilder, bool bAsyncCreateLPIs = false) override;
 	virtual void UpdatePrimitiveTransform(UPrimitiveComponent* Primitive) override;
-	virtual void UpdatePrimitiveOcclusionBoundsSlack(UPrimitiveComponent* Primitive, float NewSlack) override;
 	virtual void UpdatePrimitiveAttachment(UPrimitiveComponent* Primitive) override;
 	virtual void UpdateCustomPrimitiveData(UPrimitiveComponent* Primitive) override;
 	virtual void UpdatePrimitiveDistanceFieldSceneData_GameThread(UPrimitiveComponent* Primitive) override;
@@ -3462,9 +3462,19 @@ private:
 		FVector AttachmentRootPosition;
 	};
 
+	struct FUpdateInstanceCommand
+	{
+		FPrimitiveSceneProxy* PrimitiveSceneProxy{ nullptr };
+		FInstanceUpdateCmdBuffer CmdBuffer;
+		FBoxSphereBounds WorldBounds;
+		FBoxSphereBounds LocalBounds;
+		FBoxSphereBounds StaticMeshBounds;
+	};
+
 	TMap<FPrimitiveSceneInfo*, FPrimitiveComponentId> UpdatedAttachmentRoots;
 	TMap<FPrimitiveSceneProxy*, FCustomPrimitiveData> UpdatedCustomPrimitiveParams;
 	TMap<FPrimitiveSceneProxy*, FUpdateTransformCommand> UpdatedTransforms;
+	TMap<FPrimitiveSceneProxy*, FUpdateInstanceCommand> UpdatedInstances;
 	TMap<FPrimitiveSceneInfo*, FMatrix> OverridenPreviousTransforms;
 	TMap<const FPrimitiveSceneProxy*, float> UpdatedOcclusionBoundsSlacks;
 	TSet<FPrimitiveSceneInfo*> AddedPrimitiveSceneInfos;
diff --git a/Engine/Source/Runtime/Renderer/Public/PrimitiveSceneInfo.h b/Engine/Source/Runtime/Renderer/Public/PrimitiveSceneInfo.h
--- a/Engine/Source/Runtime/Renderer/Public/PrimitiveSceneInfo.h
+++ b/Engine/Source/Runtime/Renderer/Public/PrimitiveSceneInfo.h
@@ -419,6 +419,7 @@ public:
 
 	/** Allocate/Free slots for instance data in GPU-Scene */
 	static void AllocateGPUSceneInstances(FScene* Scene, const TArrayView<FPrimitiveSceneInfo*>& SceneInfos);
+	static void ReallocateGPUSceneInstances(FScene* Scene, const TArrayView<FPrimitiveSceneInfo*>& SceneInfos);
 	void FreeGPUSceneInstances();
 
 	/** return true if we need to call ConditionalUpdateStaticMeshes */
