Set name on xbox created textures correctly

diff --git a/Engine/Source/Runtime/Engine/Private/TextureRenderTarget2D.cpp b/Engine/Source/Runtime/Engine/Private/TextureRenderTarget2D.cpp
--- a/Engine/Source/Runtime/Engine/Private/TextureRenderTarget2D.cpp
+++ b/Engine/Source/Runtime/Engine/Private/TextureRenderTarget2D.cpp
@@ -552,7 +552,8 @@ void FTextureRenderTarget2DResource::InitDynamicRHI()
 		// Create the RHI texture. Only one mip is used and the texture is targetable for resolve.
 		ETextureCreateFlags TexCreateFlags = Owner->IsSRGB() ? TexCreate_SRGB : TexCreate_None;
 		TexCreateFlags |= Owner->bGPUSharedFlag ? TexCreate_Shared : TexCreate_None;
-		FRHIResourceCreateInfo CreateInfo = FRHIResourceCreateInfo(TEXT("TextureRenderTarget2DResource"), FClearValueBinding(ClearColor));
+		FString ResourceName = Owner ? Owner->GetName() : FString(TEXT("TextureRenderTarget2DResource"));
+		FRHIResourceCreateInfo CreateInfo = FRHIResourceCreateInfo(*ResourceName, FClearValueBinding(ClearColor));
 
 		if (Owner->bAutoGenerateMips)
 		{
