Temporary fix for broken velocity vectors

diff --git a/Engine/Plugins/Experimental/GPULightmass/Source/GPULightmass/Private/LightmapRenderer.cpp b/Engine/Plugins/Experimental/GPULightmass/Source/GPULightmass/Private/LightmapRenderer.cpp
--- a/Engine/Plugins/Experimental/GPULightmass/Source/GPULightmass/Private/LightmapRenderer.cpp
+++ b/Engine/Plugins/Experimental/GPULightmass/Source/GPULightmass/Private/LightmapRenderer.cpp
@@ -320,7 +320,7 @@ void FCachedRayTracingSceneData::SetupViewUniformBufferFromSceneRenderState(FSce
 				InstanceDataOriginalOffsets[InstanceIndex] = InstanceIndex;
 
 				FInstanceSceneShaderData& SceneData = InstanceSceneData.Emplace_GetRef();
-				SceneData.Build
+				SceneData.BuildInternal
 				(
 					InstanceIndex, /* Primitive Id */
 					0, /* Relative Instance Id */
@@ -328,7 +328,8 @@ void FCachedRayTracingSceneData::SetupViewUniformBufferFromSceneRenderState(FSce
 					INVALID_LAST_UPDATE_FRAME,
 					0, /* Custom Data Count */
 					0.0f, /* Random ID */
-					PrimitiveUniformShaderParameters.LocalToWorld
+					PrimitiveUniformShaderParameters.LocalToWorld,
+					PrimitiveUniformShaderParameters.LocalToWorld // TODO: Temporary PrevVelocityHack
 				);
 
 				for (int32 LODIndex = 0; LODIndex < Instance.LODLightmapRenderStates.Num(); LODIndex++)
@@ -403,7 +404,7 @@ void FCachedRayTracingSceneData::SetupViewUniformBufferFromSceneRenderState(FSce
 				InstanceGroup.InstancedRenderData->PerInstanceRenderData->InstanceBuffer.GetInstanceTransform(InstanceIdx, Instance.LocalToPrimitive);
 
 				FInstanceSceneShaderData& SceneData = InstanceSceneData.Emplace_GetRef();
-				SceneData.Build
+				SceneData.BuildInternal
 				(
 					PrimitiveId,
 					0, /* Relative Instance Id */
@@ -411,7 +412,8 @@ void FCachedRayTracingSceneData::SetupViewUniformBufferFromSceneRenderState(FSce
 					INVALID_LAST_UPDATE_FRAME,
 					0, /* Custom Data Count */
 					0.0f, /* Random ID */
-					InstanceGroup.LocalToWorld
+					InstanceGroup.LocalToWorld,
+					InstanceGroup.LocalToWorld // TODO: Temporary PrevVelocityHack
 				);
 			}
 
@@ -461,7 +463,7 @@ void FCachedRayTracingSceneData::SetupViewUniformBufferFromSceneRenderState(FSce
 			}
 
 			FInstanceSceneShaderData& Instance = InstanceSceneData.Emplace_GetRef();
-			Instance.Build
+			Instance.BuildInternal
 			(
 				PrimitiveId,
 				0, /* Relative Instance Id */
@@ -469,7 +471,8 @@ void FCachedRayTracingSceneData::SetupViewUniformBufferFromSceneRenderState(FSce
 				INVALID_LAST_UPDATE_FRAME,
 				0, /* Custom Data Count */
 				0.0f, /* Random ID */
-				PrimitiveUniformShaderParameters.LocalToWorld
+				PrimitiveUniformShaderParameters.LocalToWorld,
+				PrimitiveUniformShaderParameters.LocalToWorld // TODO: Temporary PrevVelocityHack
 			);
 
 			PrimitiveSceneData.Add(FPrimitiveSceneShaderData(PrimitiveUniformShaderParameters));
diff --git a/Engine/Shaders/Private/SceneData.ush b/Engine/Shaders/Private/SceneData.ush
--- a/Engine/Shaders/Private/SceneData.ush
+++ b/Engine/Shaders/Private/SceneData.ush
@@ -682,7 +682,15 @@ FInstanceSceneData GetInstanceSceneData(uint InstanceId, uint SOAStride, bool bC
 		}
 		else
 		{
+		#if INSTANCE_COMPRESSED_TRANSFORMS
+			// TODO: Temporary PrevVelocityHack
+			uint4  PrevRotationScale		= asuint(LoadInstanceSceneDataElement(3 * SOAStride + InstanceId));
+			float3 PrevTranslation			= LoadInstanceSceneDataElement(4 * SOAStride + InstanceId).xyz;
+			float3 PrevScale = 0;
+			InstanceData.PrevLocalToWorld	= DecodeTransform(PrevRotationScale, PrevTranslation, PrevScale);
+		#else
 			InstanceData.PrevLocalToWorld	= InstanceData.LocalToWorld;
+		#endif
 		}
 
 	#if 1 //NEEDS_LIGHTMAP_COORDINATE
diff --git a/Engine/Source/Runtime/Engine/Private/InstanceUniformShaderParameters.cpp b/Engine/Source/Runtime/Engine/Private/InstanceUniformShaderParameters.cpp
--- a/Engine/Source/Runtime/Engine/Private/InstanceUniformShaderParameters.cpp
+++ b/Engine/Source/Runtime/Engine/Private/InstanceUniformShaderParameters.cpp
@@ -21,7 +21,7 @@ void FInstanceSceneShaderData::Build
 	float RandomID
 )
 {
-	Build(PrimitiveId, RelativeId, PayloadDataFlags, LastUpdateFrame, CustomDataCount, RandomID, FRenderTransform::Identity);
+	BuildInternal(PrimitiveId, RelativeId, PayloadDataFlags, LastUpdateFrame, CustomDataCount, RandomID, FRenderTransform::Identity, FRenderTransform::Identity);
 }
 
 void FInstanceSceneShaderData::Build
@@ -33,18 +33,21 @@ void FInstanceSceneShaderData::Build
 	uint32 CustomDataCount,
 	float RandomID,
 	const FRenderTransform& LocalToPrimitive,
-	const FRenderTransform& PrimitiveToWorld
+	const FRenderTransform& PrimitiveToWorld,
+	const FRenderTransform& PrevPrimitiveToWorld // TODO: Temporary PrevVelocityHack
 )
 {
 	FRenderTransform LocalToWorld = LocalToPrimitive * PrimitiveToWorld;
+	FRenderTransform PrevLocalToWorld = LocalToPrimitive * PrevPrimitiveToWorld; // TODO: Temporary PrevVelocityHack
 
 	// Remove shear
 	LocalToWorld.Orthogonalize();
+	PrevLocalToWorld.Orthogonalize(); // TODO: Temporary PrevVelocityHack
 
-	Build(PrimitiveId, RelativeId, PayloadDataFlags, LastUpdateFrame, CustomDataCount, RandomID, LocalToWorld);
+	BuildInternal(PrimitiveId, RelativeId, PayloadDataFlags, LastUpdateFrame, CustomDataCount, RandomID, LocalToWorld, PrevLocalToWorld);
 }
 
-void FInstanceSceneShaderData::Build
+void FInstanceSceneShaderData::BuildInternal
 (
 	uint32 PrimitiveId,
 	uint32 RelativeId,
@@ -52,7 +55,8 @@ void FInstanceSceneShaderData::Build
 	uint32 LastUpdateFrame,
 	uint32 CustomDataCount,
 	float RandomID,
-	const FRenderTransform& LocalToWorld // Assumes shear has been removed already
+	const FRenderTransform& LocalToWorld, // Assumes shear has been removed already
+	const FRenderTransform& PrevLocalToWorld // Assumes shear has been removed already // TODO: Temporary PrevVelocityHack
 )
 {
 	// Note: layout must match GetInstanceData in SceneData.ush
@@ -80,10 +84,15 @@ void FInstanceSceneShaderData::Build
 	Data[0].Z  = *(const float*)&LastUpdateFrame;
 	Data[0].W  = *(const float*)&RandomID;
 
+	// TODO: Temporary PrevVelocityHack
 #if INSTANCE_COMPRESSED_TRANSFORMS
 	FCompressedTransform CompressedLocalToWorld(LocalToWorld);
 	Data[1] = *(const FVector4*)&CompressedLocalToWorld.Rotation[0];
 	Data[2] = *(const FVector3f*)&CompressedLocalToWorld.Translation;
+
+	FCompressedTransform CompressedPrevLocalToWorld(PrevLocalToWorld);
+	Data[3] = *(const FVector4*)&CompressedPrevLocalToWorld.Rotation[0];
+	Data[4] = *(const FVector3f*)&CompressedPrevLocalToWorld.Translation;
 #else
 	// Note: writes 3x float4s
 	LocalToWorld.To3x4MatrixTranspose((float*)&Data[1]);
diff --git a/Engine/Source/Runtime/Engine/Public/InstanceUniformShaderParameters.h b/Engine/Source/Runtime/Engine/Public/InstanceUniformShaderParameters.h
--- a/Engine/Source/Runtime/Engine/Public/InstanceUniformShaderParameters.h
+++ b/Engine/Source/Runtime/Engine/Public/InstanceUniformShaderParameters.h
@@ -87,7 +87,7 @@ struct FInstanceSceneShaderData
 	enum { DataStrideInFloat4s = 4 };
 #else
 	// Compressed transform
-	enum { DataStrideInFloat4s = 3 };
+	enum { DataStrideInFloat4s = 5 }; // TODO: Temporary PrevVelocityHack
 #endif
 
 	TStaticArray<FVector4f, DataStrideInFloat4s> Data;
@@ -115,10 +115,11 @@ struct FInstanceSceneShaderData
 		uint32 CustomDataCount,
 		float RandomID,
 		const FRenderTransform& LocalToPrimitive,
-		const FRenderTransform& PrimitiveToWorld
+		const FRenderTransform& PrimitiveToWorld,
+		const FRenderTransform& PrevPrimitiveToWorld // TODO: Temporary PrevVelocityHack
 	);
 
-	ENGINE_API void Build
+	ENGINE_API void BuildInternal
 	(
 		uint32 PrimitiveId,
 		uint32 RelativeId,
@@ -126,6 +127,7 @@ struct FInstanceSceneShaderData
 		uint32 LastUpdateFrame,
 		uint32 CustomDataCount,
 		float RandomID,
-		const FRenderTransform& LocalToWorld
+		const FRenderTransform& LocalToWorld,
+		const FRenderTransform& PrevLocalToWorld // Assumes shear has been removed already // TODO: Temporary PrevVelocityHack
 	);
 };
diff --git a/Engine/Source/Runtime/Renderer/Private/GPUScene.cpp b/Engine/Source/Runtime/Renderer/Private/GPUScene.cpp
--- a/Engine/Source/Runtime/Renderer/Private/GPUScene.cpp
+++ b/Engine/Source/Runtime/Renderer/Private/GPUScene.cpp
@@ -1263,7 +1263,8 @@ void FGPUScene::UploadGeneral(FRHICommandListImmediate& RHICmdList, FScene *Scen
 											UploadInfo.InstanceCustomDataCount,
 											RandomID,
 											SceneData.LocalToPrimitive,
-											UploadInfo.PrimitiveToWorld
+											UploadInfo.PrimitiveToWorld,
+											UploadInfo.PrevPrimitiveToWorld
 										);
 
 										int32 PayloadDataStart = UploadInfo.InstancePayloadDataOffset + (InstanceIndex * UploadInfo.InstancePayloadDataStride); 
diff --git a/Engine/Source/Runtime/Renderer/Private/Renderer.cpp b/Engine/Source/Runtime/Renderer/Private/Renderer.cpp
--- a/Engine/Source/Runtime/Renderer/Private/Renderer.cpp
+++ b/Engine/Source/Runtime/Renderer/Private/Renderer.cpp
@@ -172,7 +172,7 @@ void FRendererModule::DrawTileMesh(FCanvasRenderContext& RenderContext, FMeshPas
 				SinglePrimitiveStructured.ShaderPlatform = View.GetShaderPlatform();
 
 				// Also fill out correct single-primitive instance data, derived from the primitive.
-				SinglePrimitiveStructured.InstanceSceneData.Build
+				SinglePrimitiveStructured.InstanceSceneData.BuildInternal
 				(
 					0 /* Primitive Id */,
 					0 /* Relative Instance Id */,
@@ -180,7 +180,8 @@ void FRendererModule::DrawTileMesh(FCanvasRenderContext& RenderContext, FMeshPas
 					INVALID_LAST_UPDATE_FRAME,
 					0 /* Custom Data Count */,
 					0.0f /* Random ID */,
-					PrimitiveParams.LocalToWorld
+					PrimitiveParams.LocalToWorld,
+					PrimitiveParams.PreviousLocalToWorld
 				);
 
 				// TODO: Payload dummy?
