Allows access to UStreamableRenderAsset::SetForceMipLevelsToBeResident() from within blueprint.
-combined with event tracks this will provide a barebones lever for artists to ensure critical textures are fully resident

diff --git a/Engine/Source/Runtime/Engine/Classes/Engine/StreamableRenderAsset.h b/Engine/Source/Runtime/Engine/Classes/Engine/StreamableRenderAsset.h
--- a/Engine/Source/Runtime/Engine/Classes/Engine/StreamableRenderAsset.h
+++ b/Engine/Source/Runtime/Engine/Classes/Engine/StreamableRenderAsset.h
@@ -157,6 +157,7 @@ public:
 	* @param Seconds					Duration in seconds
 	* @param CinematicTextureGroups	Bitfield indicating which texture groups that use extra high-resolution mips
 	*/
+	UFUNCTION(BlueprintCallable, Category = "Rendering")
 	ENGINE_API void SetForceMipLevelsToBeResident(float Seconds, int32 CinematicLODGroupMask = 0);
 
 	/**
