Prevent HLODs from being always loaded when opening P_Construct in the editor

Cut down loading time of P_Construct_WP by ~45 seconds. First PIE session start time increased slightly by ~8 seconds.

Added virtual method ShouldBeLoadedByEditorCells() to World Partition actor desc
* Defaults to true
* Return false for HLOD actors, which should never be loaded by editor cells (either from the editor always loaded cell or any other cell)
* Removed the need to override IsHiddenEd() for HLOD actors (to have them hidden in editor)

diff --git a/Engine/Source/Runtime/Engine/Private/WorldPartition/HLOD/HLODActor.cpp b/Engine/Source/Runtime/Engine/Private/WorldPartition/HLOD/HLODActor.cpp
--- a/Engine/Source/Runtime/Engine/Private/WorldPartition/HLOD/HLODActor.cpp
+++ b/Engine/Source/Runtime/Engine/Private/WorldPartition/HLOD/HLODActor.cpp
@@ -91,11 +91,6 @@ void AWorldPartitionHLOD::RerunConstructionScripts()
 
 #if WITH_EDITOR
 
-bool AWorldPartitionHLOD::IsHiddenEd() const
-{
-	return !IsRunningCommandlet();
-}
-
 EActorGridPlacement AWorldPartitionHLOD::GetGridPlacement() const
 {
 	return SubActorsHLODLayer && SubActorsHLODLayer->IsAlwaysLoaded() ? EActorGridPlacement::AlwaysLoaded : Super::GetGridPlacement();
@@ -206,4 +201,4 @@ void AWorldPartitionHLOD::BuildHLOD(bool bForceBuild)
 	BuildActorTextureStreamingData(this, EMaterialQualityLevel::High, GMaxRHIFeatureLevel);
 }
 
-#endif // WITH_EDITOR
\ No newline at end of file
+#endif // WITH_EDITOR
diff --git a/Engine/Source/Runtime/Engine/Private/WorldPartition/WorldPartition.cpp b/Engine/Source/Runtime/Engine/Private/WorldPartition/WorldPartition.cpp
--- a/Engine/Source/Runtime/Engine/Private/WorldPartition/WorldPartition.cpp
+++ b/Engine/Source/Runtime/Engine/Private/WorldPartition/WorldPartition.cpp
@@ -926,6 +926,11 @@ bool UWorldPartition::UpdateEditorCells(TFunctionRef<bool(TArray<UWorldPartition
 
 bool UWorldPartition::ShouldActorBeLoadedByEditorCells(const FWorldPartitionActorDesc* ActorDesc) const
 {
+	if (!ActorDesc->ShouldBeLoadedByEditorCells())
+	{
+		return false;
+	}
+
 	if (const AWorldDataLayers* WorldDataLayers = GetWorld()->GetWorldDataLayers())
 	{
 		// Use DataLayers of loaded/dirty Actor if available to handle dirtied actors
diff --git a/Engine/Source/Runtime/Engine/Public/WorldPartition/HLOD/HLODActor.h b/Engine/Source/Runtime/Engine/Public/WorldPartition/HLOD/HLODActor.h
--- a/Engine/Source/Runtime/Engine/Public/WorldPartition/HLOD/HLODActor.h
+++ b/Engine/Source/Runtime/Engine/Public/WorldPartition/HLOD/HLODActor.h
@@ -74,7 +74,6 @@ protected:
 	virtual void BeginPlay() override;
 	virtual void EndPlay(const EEndPlayReason::Type EndPlayReason) override;
 #if WITH_EDITOR
-	virtual bool IsHiddenEd() const override;
 	virtual TUniquePtr<class FWorldPartitionActorDesc> CreateClassActorDesc() const override;
 
 	virtual void GetActorBounds(bool bOnlyCollidingComponents, FVector& Origin, FVector& BoxExtent, bool bIncludeFromChildActors) const override;
diff --git a/Engine/Source/Runtime/Engine/Public/WorldPartition/HLOD/HLODActorDesc.h b/Engine/Source/Runtime/Engine/Public/WorldPartition/HLOD/HLODActorDesc.h
--- a/Engine/Source/Runtime/Engine/Public/WorldPartition/HLOD/HLODActorDesc.h
+++ b/Engine/Source/Runtime/Engine/Public/WorldPartition/HLOD/HLODActorDesc.h
@@ -25,6 +25,8 @@ public:
 
 	static uint64 ComputeCellHash(const FString HLODLayerName, uint64 GridIndexX, uint64 GridIndexY, uint64 GridIndexZ, FDataLayersID DataLayersID);
 
+	virtual bool ShouldBeLoadedByEditorCells() const { return false; }
+
 protected:
 	virtual void Init(const AActor* InActor) override;
 	virtual bool Equals(const FWorldPartitionActorDesc* Other) const override;
diff --git a/Engine/Source/Runtime/Engine/Public/WorldPartition/WorldPartitionActorDesc.h b/Engine/Source/Runtime/Engine/Public/WorldPartition/WorldPartitionActorDesc.h
--- a/Engine/Source/Runtime/Engine/Public/WorldPartition/WorldPartitionActorDesc.h
+++ b/Engine/Source/Runtime/Engine/Public/WorldPartition/WorldPartitionActorDesc.h
@@ -137,6 +137,7 @@ public:
 	AActor* GetActor(bool bEvenIfPendingKill=true, bool bEvenIfUnreachable=false) const;
 	AActor* Load() const;
 	virtual void Unload();
+	virtual bool ShouldBeLoadedByEditorCells() const { return true; }
 
 	void RegisterActor();
 	void UnregisterActor();
