Remove a currently unused function on the ISM component that was breaking attachments.

- Overriding this function on the ISM and not calling the super was causing the bug.
- For fast-ISM update code I'll need to revisit this function, but just deleting for now for simplicity's sake.

diff --git a/Engine/Source/Runtime/Engine/Classes/Components/InstancedStaticMeshComponent.h b/Engine/Source/Runtime/Engine/Classes/Components/InstancedStaticMeshComponent.h
--- a/Engine/Source/Runtime/Engine/Classes/Components/InstancedStaticMeshComponent.h
+++ b/Engine/Source/Runtime/Engine/Classes/Components/InstancedStaticMeshComponent.h
@@ -316,7 +316,6 @@ public:
 	virtual TStructOnScope<FActorComponentInstanceData> GetComponentInstanceData() const override;
 	virtual void GetComponentChildElements(TArray<FTypedElementHandle>& OutElementHandles, const bool bAllowCreate = true) override;
 	virtual bool IsHLODRelevant() const override;
-	virtual void SendRenderTransform_Concurrent() override;
 	//~ End UActorComponent Interface
 
 	//~ Begin UPrimitiveComponent Interface
diff --git a/Engine/Source/Runtime/Engine/Private/InstancedStaticMesh.cpp b/Engine/Source/Runtime/Engine/Private/InstancedStaticMesh.cpp
--- a/Engine/Source/Runtime/Engine/Private/InstancedStaticMesh.cpp
+++ b/Engine/Source/Runtime/Engine/Private/InstancedStaticMesh.cpp
@@ -2099,20 +2099,6 @@ bool UInstancedStaticMeshComponent::IsHLODRelevant() const
 	return Super::IsHLODRelevant();
 }
 
-void UInstancedStaticMeshComponent::SendRenderTransform_Concurrent()
-{
-	// If the primitive isn't hidden update its transform.
-	const bool bDetailModeAllowsRendering = DetailMode <= GetCachedScalabilityCVars().DetailMode;
-	if ((InstanceUpdateCmdBuffer.NumEdits || InstanceUpdateCmdBuffer.NumAdds) && bDetailModeAllowsRendering && (ShouldRender() || bCastHiddenShadow || bRayTracingFarField))
-	{
-		UpdateBounds();
-
-		// Update the scene info's transform for this primitive.
-		GetWorld()->Scene->UpdatePrimitiveInstances(this);
-		InstanceUpdateCmdBuffer.Reset();
-	}
-}
-
 FBodyInstance* UInstancedStaticMeshComponent::GetBodyInstance(FName BoneName, bool bGetWelded, int32 Index) const
 {
 	if (Index != INDEX_NONE && IsValidInstance(Index))
