[Memreport] - add additional memreport types for this project, including adding GetFullName substring search

diff --git a/Engine/Config/BaseEngine.ini b/Engine/Config/BaseEngine.ini
--- a/Engine/Config/BaseEngine.ini
+++ b/Engine/Config/BaseEngine.ini
@@ -1777,7 +1777,7 @@ bRecompileOnLoad=true
 ; These commands are run when memreport is executed, and output to a profile file
 +Cmd="Mem FromReport"
 +Cmd="LogCountedInstances"
-+Cmd="obj list -alphasort"
++Cmd="obj list -resourcesizesort"
 +Cmd="rhi.DumpMemory"
 +Cmd="rhi.DumpResourceMemory"
 +Cmd="rhi.dumpresourcememory summary name=Lumen"
diff --git a/Engine/Source/Runtime/Engine/Private/UnrealEngine.cpp b/Engine/Source/Runtime/Engine/Private/UnrealEngine.cpp
--- a/Engine/Source/Runtime/Engine/Private/UnrealEngine.cpp
+++ b/Engine/Source/Runtime/Engine/Private/UnrealEngine.cpp
@@ -8319,6 +8319,9 @@ bool UEngine::HandleObjCommand( const TCHAR* Cmd, FOutputDevice& Ar )
 		FString ObjectName;
 		FParse::Value(Cmd, TEXT("NAME="), ObjectName);
 
+		FString ObjectNameSubString;
+		FParse::Value(Cmd, TEXT("NAMESUB="), ObjectNameSubString);
+
 		TArray<FItem> List;
 		TArray<FSubItem> Objects;
 		FItem Total;
@@ -8420,6 +8423,15 @@ bool UEngine::HandleObjCommand( const TCHAR* Cmd, FOutputDevice& Ar )
 					continue;
 				}
 
+				if (!ObjectNameSubString.IsEmpty())
+				{
+					FString FullName = *It->GetFullName();
+					if(!FullName.Contains(ObjectNameSubString, ESearchCase::IgnoreCase))
+					{
+						continue;
+					}
+				}
+
 				if ( MetaClass )
 				{
 					UClass* ClassObj = Cast<UClass>(*It);
