Data Layers: add data layers debug display to show accurate data layer loading time so we can determine accurate preroll times for data layers activated by sequences.
Can be activated using the wp.Runtime.ToggleDrawDataLayersLoadTime console command.
Bonus: added drop shadow for data layers debug display for better visibility.

diff --git a/Engine/Source/Runtime/Engine/Private/WorldPartition/DataLayer/DataLayerSubsystem.cpp b/Engine/Source/Runtime/Engine/Private/WorldPartition/DataLayer/DataLayerSubsystem.cpp
--- a/Engine/Source/Runtime/Engine/Private/WorldPartition/DataLayer/DataLayerSubsystem.cpp
+++ b/Engine/Source/Runtime/Engine/Private/WorldPartition/DataLayer/DataLayerSubsystem.cpp
@@ -4,16 +4,21 @@
 #include "WorldPartition/DataLayer/WorldDataLayers.h"
 #include "WorldPartition/DataLayer/DataLayer.h"
 #include "WorldPartition/WorldPartitionDebugHelper.h"
+#include "WorldPartition/WorldPartitionSubsystem.h"
+#include "WorldPartition/WorldPartitionRuntimeCell.h"
 #include "WorldPartition/WorldPartition.h"
 #include "Engine/Engine.h"
 #include "Engine/World.h"
 #include "Engine/Canvas.h"
+#include "Algo/Transform.h"
 #if WITH_EDITOR
 #include "Editor.h"
 #include "Modules/ModuleManager.h"
 #include "WorldPartition/DataLayer/IDataLayerEditorModule.h"
 #endif
 
+extern int32 GDrawDataLayersLoadTime;
+
 static FAutoConsoleCommandWithOutputDevice GDumpDataLayersCmd(
 	TEXT("wp.DumpDataLayers"),
 	TEXT("Dumps data layers to the log"),
@@ -231,20 +236,48 @@ void UDataLayerSubsystem::DrawDataLayersStatus(UCanvas* Canvas, FVector2D& Offse
 	FVector2D Pos = Offset;
 	float MaxTextWidth = 0.f;
 
+	UWorldPartitionSubsystem* WorldPartitionSubsystem = GetWorld()->GetSubsystem<UWorldPartitionSubsystem>();
+
 	auto DrawLayerNames = [this, Canvas, &Pos, &MaxTextWidth](const FString& Title, FColor HeaderColor, FColor TextColor, const TSet<FName>& LayerNames)
 	{
 		if (LayerNames.Num() > 0)
 		{
 			FWorldPartitionDebugHelper::DrawText(Canvas, Title, GEngine->GetSmallFont(), HeaderColor, Pos, &MaxTextWidth);
 
-			UFont* DataLayerFont = GEngine->GetSmallFont();
+			TArray<UDataLayer*> DataLayers;
+			DataLayers.Reserve(LayerNames.Num());
 			for (const FName& DataLayerName : LayerNames)
 			{
 				if (UDataLayer* DataLayer = GetDataLayerFromName(DataLayerName))
 				{
-					FWorldPartitionDebugHelper::DrawLegendItem(Canvas, *DataLayer->GetDataLayerLabel().ToString(), DataLayerFont, DataLayer->GetDebugColor(), TextColor, Pos, &MaxTextWidth);
+					DataLayers.Add(DataLayer);
 				}
 			}
+
+			DataLayers.Sort([](const UDataLayer& A, const UDataLayer& B) { return A.GetDataLayerLabel().LexicalLess(B.GetDataLayerLabel()); });
+
+			UFont* DataLayerFont = GEngine->GetSmallFont();
+			for (UDataLayer* DataLayer : DataLayers)
+			{
+				FString DataLayerString = DataLayer->GetDataLayerLabel().ToString();
+
+				if (GDrawDataLayersLoadTime)
+				{
+					if (double* DataLayerLoadTime = ActiveDataLayersLoadTime.Find(DataLayer))
+					{
+						if (*DataLayerLoadTime < 0)
+						{
+							DataLayerString += FString::Printf(TEXT(" (streaming %s)"), *FPlatformTime::PrettyTime(FPlatformTime::Seconds() + *DataLayerLoadTime));
+						}
+						else
+						{
+							DataLayerString += FString::Printf(TEXT(" (took %s)"), *FPlatformTime::PrettyTime(*DataLayerLoadTime));
+						}
+					}
+				}
+
+				FWorldPartitionDebugHelper::DrawLegendItem(Canvas, *DataLayerString, DataLayerFont, DataLayer->GetDebugColor(), TextColor, Pos, &MaxTextWidth);
+			}
 		}
 	};
 
@@ -254,9 +287,9 @@ void UDataLayerSubsystem::DrawDataLayersStatus(UCanvas* Canvas, FVector2D& Offse
 	DrawLayerNames(TEXT("Loaded Data Layers"), FColor::Cyan, FColor::White, LoadedDataLayers);
 	DrawLayerNames(TEXT("Active Data Layers"), FColor::Green, FColor::White, ActiveDataLayers);
 
+	TSet<FName> UnloadedDataLayers;
 	if (AWorldDataLayers* WorldDataLayers = GetWorld()->GetWorldDataLayers())
-	{
-		TSet<FName> UnloadedDataLayers;
+	{		
 		WorldDataLayers->ForEachDataLayer([&LoadedDataLayers, &ActiveDataLayers, &UnloadedDataLayers](UDataLayer* DataLayer)
 		{
 			if (DataLayer->IsRuntime())
@@ -273,6 +306,60 @@ void UDataLayerSubsystem::DrawDataLayersStatus(UCanvas* Canvas, FVector2D& Offse
 	}
 
 	Offset.X += MaxTextWidth + 10;
+
+	// Update data layers load times
+	if (GDrawDataLayersLoadTime)
+	{
+		for (FName DataLayerName : UnloadedDataLayers)
+		{
+			if (UDataLayer* DataLayer = GetDataLayerFromName(DataLayerName))
+			{
+				ActiveDataLayersLoadTime.Remove(DataLayer);
+			}
+		}
+
+		TArray<UDataLayer*> LoadingDataLayers;
+		LoadingDataLayers.Reserve(LoadedDataLayers.Num() + ActiveDataLayers.Num());
+		auto CopyLambda = [this](FName DataLayerName) { return GetDataLayerFromName(DataLayerName); };
+		Algo::Transform(LoadedDataLayers, LoadingDataLayers, CopyLambda);
+		Algo::Transform(ActiveDataLayers, LoadingDataLayers, CopyLambda);
+
+		for (UDataLayer* DataLayer : LoadingDataLayers)
+		{
+			double* DataLayerLoadTime = ActiveDataLayersLoadTime.Find(DataLayer);
+
+			auto IsDataLayerReady = [WorldPartitionSubsystem](UDataLayer* DataLayer, EWorldPartitionRuntimeCellState TargetState)
+			{
+				FWorldPartitionStreamingQuerySource QuerySource;
+				QuerySource.bDataLayersOnly = true;
+				QuerySource.bSpatialQuery = false;
+				QuerySource.DataLayers.Add(DataLayer->GetFName());
+				return WorldPartitionSubsystem->IsStreamingCompleted(TargetState, { QuerySource }, true);
+			};
+
+			const EWorldPartitionRuntimeCellState TargetState = ActiveDataLayers.Contains(DataLayer->GetFName()) ? EWorldPartitionRuntimeCellState::Activated : EWorldPartitionRuntimeCellState::Loaded;
+
+			if (!DataLayerLoadTime)
+			{
+				if (!IsDataLayerReady(DataLayer, TargetState))
+				{
+					DataLayerLoadTime = &ActiveDataLayersLoadTime.Add(DataLayer, -FPlatformTime::Seconds());
+				}
+			}
+
+			if (DataLayerLoadTime && (*DataLayerLoadTime < 0))
+			{
+				if (IsDataLayerReady(DataLayer, TargetState))
+				{
+					*DataLayerLoadTime = FPlatformTime::Seconds() + *DataLayerLoadTime;
+				}
+			}
+		}
+	}
+	else
+	{
+		ActiveDataLayersLoadTime.Empty();
+	}
 }
 
 TArray<UDataLayer*> UDataLayerSubsystem::ConvertArgsToDataLayers(UWorld* World, const TArray<FString>& InArgs)
diff --git a/Engine/Source/Runtime/Engine/Private/WorldPartition/WorldPartitionDebugHelper.cpp b/Engine/Source/Runtime/Engine/Private/WorldPartition/WorldPartitionDebugHelper.cpp
--- a/Engine/Source/Runtime/Engine/Private/WorldPartition/WorldPartitionDebugHelper.cpp
+++ b/Engine/Source/Runtime/Engine/Private/WorldPartition/WorldPartitionDebugHelper.cpp
@@ -143,6 +143,8 @@ void FWorldPartitionDebugHelper::DrawText(UCanvas* Canvas, const FString& Text,
 {
 	float TextWidth, TextHeight;
 	Canvas->StrLen(Font, Text, TextWidth, TextHeight);
+	Canvas->SetDrawColor(FColor(0,0,0));
+	Canvas->DrawText(Font, Text, Pos.X + 1, Pos.Y + 1);
 	Canvas->SetDrawColor(Color);
 	Canvas->DrawText(Font, Text, Pos.X, Pos.Y);
 	Pos.Y += TextHeight + 1;
@@ -157,7 +159,13 @@ void FWorldPartitionDebugHelper::DrawLegendItem(UCanvas* Canvas, const FString&
 	static const FVector2D ItemSize(12, 12);
 	
 	float MaxTextWidth = 0.f;
-	
+
+	const FVector2D ShadowItemPos(Pos - FVector2D(1, 1));
+	const FVector2D ShadowItemSize(ItemSize + FVector2D(2, 2));
+	const FColor ShadowItemColor(0,0,0);
+	FCanvasTileItem ShadowItem(ShadowItemPos, GWhiteTexture, ShadowItemSize, ShadowItemColor);
+	Canvas->DrawItem(ShadowItem);
+
 	FCanvasTileItem Item(Pos, GWhiteTexture, ItemSize, Color);
 	Canvas->DrawItem(Item);
 
diff --git a/Engine/Source/Runtime/Engine/Private/WorldPartition/WorldPartitionSubsystem.cpp b/Engine/Source/Runtime/Engine/Private/WorldPartition/WorldPartitionSubsystem.cpp
--- a/Engine/Source/Runtime/Engine/Private/WorldPartition/WorldPartitionSubsystem.cpp
+++ b/Engine/Source/Runtime/Engine/Private/WorldPartition/WorldPartitionSubsystem.cpp
@@ -57,6 +57,12 @@ static FAutoConsoleCommand CVarDrawDataLayers(
 	TEXT("Toggles debug display of active data layers."),
 	FConsoleCommandDelegate::CreateLambda([] { GDrawDataLayers = !GDrawDataLayers; }));
 
+int32 GDrawDataLayersLoadTime = 0;
+static FAutoConsoleCommand CVarDrawDataLayersLoadTime(
+	TEXT("wp.Runtime.ToggleDrawDataLayersLoadTime"),
+	TEXT("Toggles debug display of active data layers load time."),
+	FConsoleCommandDelegate::CreateLambda([] { GDrawDataLayersLoadTime = !GDrawDataLayersLoadTime; }));
+
 int32 GLevelStreamingContinuouslyIncrementalGCWhileLevelsPendingPurgeForWP = 64;
 static FAutoConsoleVariableRef CVarGLevelStreamingContinuouslyIncrementalGCWhileLevelsPendingPurgeForWP(
 	TEXT("wp.Runtime.LevelStreamingContinuouslyIncrementalGCWhileLevelsPendingPurgeForWP"),
@@ -370,7 +376,7 @@ void UWorldPartitionSubsystem::Draw(UCanvas* Canvas, class APlayerController* PC
 			WorldPartition->DrawStreamingStatusLegend(Canvas, CurrentOffset);
 		}
 
-		if (DataLayerSubsystem && (GDrawDataLayers || GDrawRuntimeHash2D))
+		if (DataLayerSubsystem && (GDrawDataLayers || GDrawDataLayersLoadTime || GDrawRuntimeHash2D))
 		{
 			DataLayerSubsystem->DrawDataLayersStatus(Canvas, CurrentOffset);
 		}
diff --git a/Engine/Source/Runtime/Engine/Public/WorldPartition/DataLayer/DataLayerSubsystem.h b/Engine/Source/Runtime/Engine/Public/WorldPartition/DataLayer/DataLayerSubsystem.h
--- a/Engine/Source/Runtime/Engine/Public/WorldPartition/DataLayer/DataLayerSubsystem.h
+++ b/Engine/Source/Runtime/Engine/Public/WorldPartition/DataLayer/DataLayerSubsystem.h
@@ -150,6 +150,9 @@ private:
 	/** Console command used to toggle activation of a DataLayer */
 	static class FAutoConsoleCommand ToggleDataLayerActivation;
 
+	/** Data layers load time */
+	mutable TMap<UDataLayer*, double> ActiveDataLayersLoadTime;
+
 	/** Console command used to set Runtime DataLayer state*/
 	static class FAutoConsoleCommand SetDataLayerRuntimeStateCommand;
 };
\ No newline at end of file
