Better logging to help narrow down skel mesh attachment issues at cook time

diff --git a/Engine/Source/Runtime/Engine/Private/Components/SkinnedMeshComponent.cpp b/Engine/Source/Runtime/Engine/Private/Components/SkinnedMeshComponent.cpp
--- a/Engine/Source/Runtime/Engine/Private/Components/SkinnedMeshComponent.cpp
+++ b/Engine/Source/Runtime/Engine/Private/Components/SkinnedMeshComponent.cpp
@@ -2128,13 +2128,13 @@ class USkeletalMeshSocket const* USkinnedMeshComponent::GetSocketInfoByName(FNam
 	{
 		if (OverrideSocket)
 		{
-			UE_LOG(LogSkinnedMeshComp, Warning, TEXT("GetSocketByName(%s -> override To %s): No SkeletalMesh for Component(%s) Actor(%s)"),
-				*InSocketName.ToString(), *OverrideSocketName.ToString(), *GetName(), *GetNameSafe(GetOuter()));
+			UE_LOG(LogSkinnedMeshComp, Warning, TEXT("GetSocketInfoByName(%s -> override To %s): No SkeletalMesh for Component(%s) Actor(%s)"),
+				*InSocketName.ToString(), *OverrideSocketName.ToString(), *GetName(), *GetFullNameSafe(GetOuter()));
 		}
 		else
 		{
-			UE_LOG(LogSkinnedMeshComp, Warning, TEXT("GetSocketByName(%s): No SkeletalMesh for Component(%s) Actor(%s)"),
-				*OverrideSocketName.ToString(), *GetName(), *GetNameSafe(GetOuter()));
+			UE_LOG(LogSkinnedMeshComp, Warning, TEXT("GetSocketInfoByName(%s): No SkeletalMesh for Component(%s) Actor(%s)"),
+				*OverrideSocketName.ToString(), *GetName(), *GetFullNameSafe(GetOuter()));
 		}
 	}
 
@@ -2157,12 +2157,12 @@ class USkeletalMeshSocket const* USkinnedMeshComponent::GetSocketByName(FName In
 		if (OverrideSocket)
 		{
 			UE_LOG(LogSkinnedMeshComp, Warning, TEXT("GetSocketByName(%s -> override To %s): No SkeletalMesh for Component(%s) Actor(%s)"),
-				*InSocketName.ToString(), *OverrideSocketName.ToString(), *GetName(), *GetNameSafe(GetOuter()));
+				*InSocketName.ToString(), *OverrideSocketName.ToString(), *GetName(), *GetFullNameSafe(GetOuter()));
 		}
 		else
 		{
 			UE_LOG(LogSkinnedMeshComp, Warning, TEXT("GetSocketByName(%s): No SkeletalMesh for Component(%s) Actor(%s)"),
-				*OverrideSocketName.ToString(), *GetName(), *GetNameSafe(GetOuter()));
+				*OverrideSocketName.ToString(), *GetName(), *GetFullNameSafe(GetOuter()));
 		}
 	}
 
