UE5:  AUDIO:  Fixed memory leak in Metasound::FAddressRouter::DataChannelMap, caused by IAudioInstanceTransmitter::Shutdown not being called in some code paths where FActiveSound instances are deleted.  Solves a leak of around 1 MB/hr.

#rnx

diff --git a/Engine/Source/Runtime/Engine/Private/AudioDevice.cpp b/Engine/Source/Runtime/Engine/Private/AudioDevice.cpp
--- a/Engine/Source/Runtime/Engine/Private/AudioDevice.cpp
+++ b/Engine/Source/Runtime/Engine/Private/AudioDevice.cpp
@@ -5121,6 +5121,13 @@ void FAudioDevice::ProcessingPendingActiveSoundStops(bool bForceDelete)
 				}
 				ActiveSound->bAsyncOcclusionPending = false;
 				PendingSoundsToDelete.RemoveAtSwap(i, 1, false);
+
+				// TODO: Phil Popp investigate a more robust fix -- call Shutdown in FActiveSound or transmitter destructor?
+				if (IAudioInstanceTransmitter* Transmitter = ActiveSound->GetTransmitter())
+				{
+					Transmitter->Shutdown();
+				}
+
 				NotifyPendingDelete(*ActiveSound);
 				delete ActiveSound;
 			}
@@ -5175,6 +5182,12 @@ void FAudioDevice::ProcessingPendingActiveSoundStops(bool bForceDelete)
 
 			if (bDeleteActiveSound)
 			{
+				// TODO: Phil Popp investigate a more robust fix -- call Shutdown in FActiveSound or transmitter destructor?
+				if (IAudioInstanceTransmitter* Transmitter = ActiveSound->GetTransmitter())
+				{
+					Transmitter->Shutdown();
+				}
+
 				NotifyPendingDelete(*ActiveSound);
 
 				// Remove from the list of pending sounds to stop
diff --git a/Engine/Source/Runtime/Engine/Private/AudioThread.cpp b/Engine/Source/Runtime/Engine/Private/AudioThread.cpp
--- a/Engine/Source/Runtime/Engine/Private/AudioThread.cpp
+++ b/Engine/Source/Runtime/Engine/Private/AudioThread.cpp
@@ -371,6 +371,8 @@ struct FAudioAsyncBatcher
 		LastBatch = GAudioPipe.Launch(TEXT("AudioBatch"),
 			[WorkItems = MoveTemp(WorkItems)]() mutable
 			{
+				LLM_SCOPE(ELLMTag::AudioMisc);
+
 				for (FWork& Work : WorkItems)
 				{
 					Work();
