Instead of generating the morph target render data at runtime, generate it at cook time and load at runtime.  This prevents stalls when executing postload.

diff --git a/Engine/Source/Runtime/Core/Public/UObject/UE5PrivateFrostyStreamObjectVersion.h b/Engine/Source/Runtime/Core/Public/UObject/UE5PrivateFrostyStreamObjectVersion.h
--- a/Engine/Source/Runtime/Core/Public/UObject/UE5PrivateFrostyStreamObjectVersion.h
+++ b/Engine/Source/Runtime/Core/Public/UObject/UE5PrivateFrostyStreamObjectVersion.h
@@ -30,6 +30,9 @@ struct CORE_API FUE5PrivateFrostyStreamObjectVersion
 		// Re-calculate the long range attachment to prevent kinematic tethers.
 		ChaosClothRemoveKinematicTethers,
 
+		// Serializes the Morph Target render data for cooked platforms and the DDC
+		SerializeSkeletalMeshMorphTargetRenderData,
+
 		// -----<new versions can be added above this line>-------------------------------------------------
 		VersionPlusOne,
 		LatestVersion = VersionPlusOne - 1
diff --git a/Engine/Source/Runtime/Engine/Private/Rendering/MorphTargetVertexInfoBuffers.cpp b/Engine/Source/Runtime/Engine/Private/Rendering/MorphTargetVertexInfoBuffers.cpp
--- a/Engine/Source/Runtime/Engine/Private/Rendering/MorphTargetVertexInfoBuffers.cpp
+++ b/Engine/Source/Runtime/Engine/Private/Rendering/MorphTargetVertexInfoBuffers.cpp
@@ -2,6 +2,11 @@
 
 #include "Rendering/MorphTargetVertexInfoBuffers.h"
 #include "ProfilingDebugging/LoadTimeTracker.h"
+#include "Rendering/SkeletalMeshLODRenderData.h"
+#include "Animation/MorphTarget.h"
+
+extern int32 GForceRecomputeTangents;
+extern int32 GSkinCacheRecomputeTangents;
 
 void FMorphTargetVertexInfoBuffers::InitRHI()
 {
@@ -19,10 +24,427 @@ void FMorphTargetVertexInfoBuffers::InitRHI()
 	MorphDataSRV = RHICreateShaderResourceView(MorphDataBuffer);
 
 	MorphData.Empty();
+	bIsMorphCPUDataValid = false;
+	bRHIIntialized = true;
 }
 
 void FMorphTargetVertexInfoBuffers::ReleaseRHI()
 {
 	MorphDataBuffer.SafeRelease();
 	MorphDataSRV.SafeRelease();
-}
\ No newline at end of file
+	bRHIIntialized = false;
+}
+
+const float FMorphTargetVertexInfoBuffers::CalculatePositionPrecision(float TargetPositionErrorTolerance)
+{
+	const float UnrealUnitPerMeter = 100.0f;
+	return TargetPositionErrorTolerance * 2.0f * 1e-6f * UnrealUnitPerMeter;	// * 2.0 because correct rounding guarantees error is at most half of the cell size.
+}
+
+void FMorphTargetVertexInfoBuffers::ValidateVertexBuffers(bool bMorphTargetsShouldBeValid)
+{
+#if DO_CHECK
+	check(BatchesPerMorph.Num() == BatchStartOffsetPerMorph.Num());
+	check(BatchesPerMorph.Num() == MaximumValuePerMorph.Num());
+	check(BatchesPerMorph.Num() == MinimumValuePerMorph.Num());
+
+	if (bMorphTargetsShouldBeValid)
+	{
+		if (NumTotalBatches > 0)
+		{
+			check(MorphData.Num() > 0);
+		}
+		else
+		{
+			check(MorphData.Num() == 0);
+		}
+	}
+	
+#endif //DO_CHECK
+}
+
+void FMorphTargetVertexInfoBuffers::Serialize(FArchive& Ar)
+{
+	if (Ar.IsSaving())
+	{
+		check(bResourcesInitialized);
+		check(bIsMorphCPUDataValid);
+		ValidateVertexBuffers(true);
+	}
+	else
+	{
+		check(!bRHIIntialized);
+	}
+
+	Ar << MorphData
+	   << MinimumValuePerMorph
+	   << MaximumValuePerMorph
+	   << BatchStartOffsetPerMorph
+	   << BatchesPerMorph
+	   << NumTotalBatches
+	   << PositionPrecision
+	   << TangentZPrecision;
+
+	if (Ar.IsLoading())
+	{
+		bRHIIntialized = false;
+		bIsMorphCPUDataValid = true;
+		bResourcesInitialized = true;
+		ValidateVertexBuffers(true);
+	}
+}
+
+FArchive& operator<<(FArchive& Ar, FMorphTargetVertexInfoBuffers& MorphTargetVertexInfoBuffers)
+{
+	MorphTargetVertexInfoBuffers.Serialize(Ar);
+	return Ar;
+}
+
+class FDwordBitWriter
+{
+public:
+	FDwordBitWriter(TArray<uint32>& Buffer) :
+		Buffer(Buffer),
+		PendingBits(0ull),
+		NumPendingBits(0)
+	{
+	}
+
+	void PutBits(uint32 Bits, uint32 NumBits)
+	{
+		check((uint64)Bits < (1ull << NumBits));
+		PendingBits |= (uint64)Bits << NumPendingBits;
+		NumPendingBits += NumBits;
+
+		while (NumPendingBits >= 32)
+		{
+			Buffer.Add((uint32)PendingBits);
+			PendingBits >>= 32;
+			NumPendingBits -= 32;
+		}
+	}
+
+	void Flush()
+	{
+		if (NumPendingBits > 0)
+			Buffer.Add((uint32)PendingBits);
+		PendingBits = 0;
+		NumPendingBits = 0;
+	}
+
+private:
+	TArray<uint32>& Buffer;
+	uint64 			PendingBits;
+	int32 			NumPendingBits;
+};
+
+void FMorphTargetVertexInfoBuffers::InitMorphResources(EShaderPlatform ShaderPlatform, const TArray<FSkelMeshRenderSection>& RenderSections, const TArray<UMorphTarget*>& MorphTargets, int NumVertices, int32 LODIndex, float TargetPositionErrorTolerance)
+{
+	check(!IsRHIIntialized());
+	check(!IsMorphResourcesInitialized());
+	check(!IsMorphCPUDataValid());
+
+	bResourcesInitialized = true;
+
+	// UseGPUMorphTargets() can be toggled only on SM5 atm
+	if (!IsFeatureLevelSupported(ShaderPlatform, ERHIFeatureLevel::SM5) || MorphTargets.Num() == 0)
+	{
+		return;
+	}
+
+	bIsMorphCPUDataValid = true;
+
+	// Simple Morph compression 0.1
+	// Instead of storing vertex deltas individually they are organized into batches of 64.
+	// Each batch has a header that describes how many bits are allocated to each of the vertex components.
+	// Batches also store an explicit offset to its associated data. This makes it trivial to decode batches
+	// in parallel, and because deltas are fixed-width inside a batch, deltas can also be decoded in parallel.
+	// The result is a semi-adaptive encoding that functions as a crude substitute for entropy coding, that is
+	// fast to decode on parallel hardware.
+
+	// Quantization still happens globally to avoid issues with cracks at duplicate vertices.
+	// The quantization is artist controlled on a per LOD basis. Higher error tolerance results in smaller deltas
+	// and a smaller compressed size.
+
+	PositionPrecision = CalculatePositionPrecision(TargetPositionErrorTolerance);
+	const float RcpPositionPrecision = 1.0f / PositionPrecision;
+
+	TangentZPrecision = 1.0f / 2048.0f;				// Object scale irrelevant here. Let's assume ~12bits per component is plenty.
+	const float RcpTangentZPrecision = 1.0f / TangentZPrecision;
+
+	const uint32 BatchSize = 64;
+	const uint32 NumBatchHeaderDwords = 10u;
+
+	const uint32 IndexMaxBits = 31u;
+
+	const uint32 PositionMaxBits = 28u;				// Probably more than we need, but let's just allow it to go this high to be safe for now.
+													// For larger deltas this can even be more precision than what was in the float input data!
+													// Maybe consider float-like or exponential encoding of large values?
+	const float PositionMinValue = -134217728.0f;	// -2^(MaxBits-1)
+	const float PositionMaxValue = 134217720.0f;	// Largest float smaller than 2^(MaxBits-1)-1
+													// Using 134217727.0f would NOT work as it would be rounded up to 134217728.0f, which is
+													// outside the range.
+
+	const uint32 TangentZMaxBits = 16u;
+	const float TangentZMinValue = -32768.0f;		// -2^(MaxBits-1)
+	const float TangentZMaxValue = 32767.0f;		//  2^(MaxBits-1)-1
+
+	struct FBatchHeader
+	{
+		uint32		DataOffset;
+		uint32		NumElements;
+		bool		bTangents;
+
+		uint32		IndexBits;
+		FIntVector	PositionBits;
+		FIntVector	TangentZBits;
+
+		uint32		IndexMin;
+		FIntVector	PositionMin;
+		FIntVector	TangentZMin;
+	};
+
+	// uint32 StartTime = FPlatformTime::Cycles();
+
+	MorphData.Empty();
+	NumTotalBatches = 0;
+
+	BatchStartOffsetPerMorph.Empty(MorphTargets.Num());
+	BatchesPerMorph.Empty(MorphTargets.Num());
+	MaximumValuePerMorph.Empty(MorphTargets.Num());
+	MinimumValuePerMorph.Empty(MorphTargets.Num());
+
+	// Mark vertices that are in a section that doesn't recompute tangents as needing tangents
+	const int32 RecomputeTangentsMode = GForceRecomputeTangents > 0 ? 1 : GSkinCacheRecomputeTangents;
+	TBitArray<> VertexNeedsTangents;
+	VertexNeedsTangents.Init(false, NumVertices);
+	for (const FSkelMeshRenderSection& RenderSection : RenderSections)
+	{
+		const bool bRecomputeTangents = RecomputeTangentsMode > 0 && (RenderSection.bRecomputeTangent || RecomputeTangentsMode == 1);
+		if (!bRecomputeTangents)
+		{
+			for (uint32 i = 0; i < RenderSection.NumVertices; i++)
+			{
+				VertexNeedsTangents[RenderSection.BaseVertexIndex + i] = true;
+			}
+		}
+	}
+
+	// Populate the arrays to be filled in later in the render thread
+	TArray<FBatchHeader> BatchHeaders;
+	TArray<uint32> BitstreamData;
+	for (int32 AnimIdx = 0; AnimIdx < MorphTargets.Num(); ++AnimIdx)
+	{
+		uint32 BatchStartOffset = NumTotalBatches;
+
+		float MaximumValues[4] = { -FLT_MAX, -FLT_MAX, -FLT_MAX, -FLT_MAX };
+		float MinimumValues[4] = { +FLT_MAX, +FLT_MAX, +FLT_MAX, +FLT_MAX };
+		UMorphTarget* MorphTarget = MorphTargets[AnimIdx];
+		int32 NumSrcDeltas = 0;
+		FMorphTargetDelta* MorphDeltas = MorphTarget->GetMorphTargetDelta(LODIndex, NumSrcDeltas);
+
+		if (NumSrcDeltas == 0)
+		{
+			MaximumValues[0] = 0.0f;
+			MaximumValues[1] = 0.0f;
+			MaximumValues[2] = 0.0f;
+			MaximumValues[3] = 0.0f;
+
+			MinimumValues[0] = 0.0f;
+			MinimumValues[1] = 0.0f;
+			MinimumValues[2] = 0.0f;
+			MinimumValues[3] = 0.0f;
+		}
+		else
+		{
+			struct FQuantizedDelta
+			{
+				FIntVector	Position;
+				FIntVector	TangentZ;
+				uint32		Index;
+			};
+			TArray<FQuantizedDelta> QuantizedDeltas;
+			QuantizedDeltas.Reserve(NumSrcDeltas);
+
+			bool bVertexIndicesSorted = true;
+
+			int32 PrevVertexIndex = -1;
+			for (int32 DeltaIndex = 0; DeltaIndex < NumSrcDeltas; DeltaIndex++)
+			{
+				const auto& MorphDelta = MorphDeltas[DeltaIndex];
+				const FVector3f TangentZDelta = VertexNeedsTangents[MorphDelta.SourceIdx] ? MorphDelta.TangentZDelta : FVector3f::ZeroVector;
+
+				// when import, we do check threshold, and also when adding weight, we do have threshold for how smaller weight can fit in
+				// so no reason to check here another threshold
+				MaximumValues[0] = FMath::Max(MaximumValues[0], MorphDelta.PositionDelta.X);
+				MaximumValues[1] = FMath::Max(MaximumValues[1], MorphDelta.PositionDelta.Y);
+				MaximumValues[2] = FMath::Max(MaximumValues[2], MorphDelta.PositionDelta.Z);
+				MaximumValues[3] = FMath::Max(MaximumValues[3], FMath::Max(TangentZDelta.X, FMath::Max(TangentZDelta.Y, TangentZDelta.Z)));
+
+				MinimumValues[0] = FMath::Min(MinimumValues[0], MorphDelta.PositionDelta.X);
+				MinimumValues[1] = FMath::Min(MinimumValues[1], MorphDelta.PositionDelta.Y);
+				MinimumValues[2] = FMath::Min(MinimumValues[2], MorphDelta.PositionDelta.Z);
+				MinimumValues[3] = FMath::Min(MinimumValues[3], FMath::Min(TangentZDelta.X, FMath::Min(TangentZDelta.Y, TangentZDelta.Z)));
+
+				// Check if input is sorted. It usually is, but it might not be.
+				if ((int32)MorphDelta.SourceIdx < PrevVertexIndex)
+					bVertexIndicesSorted = false;
+				PrevVertexIndex = (int32)MorphDelta.SourceIdx;
+
+				// Quantize delta
+				FQuantizedDelta QuantizedDelta;
+				const FVector3f& PositionDelta = MorphDelta.PositionDelta;
+				QuantizedDelta.Position.X = FMath::RoundToInt(FMath::Clamp(PositionDelta.X * RcpPositionPrecision, PositionMinValue, PositionMaxValue));
+				QuantizedDelta.Position.Y = FMath::RoundToInt(FMath::Clamp(PositionDelta.Y * RcpPositionPrecision, PositionMinValue, PositionMaxValue));
+				QuantizedDelta.Position.Z = FMath::RoundToInt(FMath::Clamp(PositionDelta.Z * RcpPositionPrecision, PositionMinValue, PositionMaxValue));
+				QuantizedDelta.TangentZ.X = FMath::RoundToInt(FMath::Clamp(TangentZDelta.X * RcpTangentZPrecision, TangentZMinValue, TangentZMaxValue));
+				QuantizedDelta.TangentZ.Y = FMath::RoundToInt(FMath::Clamp(TangentZDelta.Y * RcpTangentZPrecision, TangentZMinValue, TangentZMaxValue));
+				QuantizedDelta.TangentZ.Z = FMath::RoundToInt(FMath::Clamp(TangentZDelta.Z * RcpTangentZPrecision, TangentZMinValue, TangentZMaxValue));
+				QuantizedDelta.Index = MorphDelta.SourceIdx;
+
+				if (QuantizedDelta.Position != FIntVector::ZeroValue || QuantizedDelta.TangentZ != FIntVector::ZeroValue)
+				{
+					// Only add delta if it is non-zero
+					QuantizedDeltas.Add(QuantizedDelta);
+				}
+			}
+
+			// Sort deltas if the source wasn't already sorted
+			if (!bVertexIndicesSorted)
+			{
+				Algo::Sort(QuantizedDeltas, [](const FQuantizedDelta& A, const FQuantizedDelta& B) { return A.Index < B.Index; });
+			}
+
+			// Encode batch deltas
+			const uint32 MorphNumBatches = (QuantizedDeltas.Num() + BatchSize - 1u) / BatchSize;
+			for (uint32 BatchIndex = 0; BatchIndex < MorphNumBatches; BatchIndex++)
+			{
+				const uint32 BatchFirstElementIndex = BatchIndex * BatchSize;
+				const uint32 NumElements = FMath::Min(BatchSize, QuantizedDeltas.Num() - BatchFirstElementIndex);
+
+				// Calculate batch min/max bounds
+				uint32 IndexMin = MAX_uint32;
+				uint32 IndexMax = MIN_uint32;
+				FIntVector PositionMin = FIntVector(MAX_int32);
+				FIntVector PositionMax = FIntVector(MIN_int32);
+				FIntVector TangentZMin = FIntVector(MAX_int32);
+				FIntVector TangentZMax = FIntVector(MIN_int32);
+
+				for (uint32 LocalElementIndex = 0; LocalElementIndex < NumElements; LocalElementIndex++)
+				{
+					const FQuantizedDelta& Delta = QuantizedDeltas[BatchFirstElementIndex + LocalElementIndex];
+
+					// Trick: Deltas are sorted by index, so the index increase by at least one per delta.
+					//		  Naively this would mean that a batch always spans at least 64 index values and
+					//		  indices would have to use at least 6 bits per index.
+					//		  If instead of storing the raw index, we store the index relative to its position in the batch,
+					//		  then the spanned range becomes 63 smaller.
+					//		  For a consecutive range this even gets us down to 0 bits per index!
+					check(Delta.Index >= LocalElementIndex);
+					const uint32 AdjustedIndex = Delta.Index - LocalElementIndex;
+					IndexMin = FMath::Min(IndexMin, AdjustedIndex);
+					IndexMax = FMath::Max(IndexMax, AdjustedIndex);
+
+					PositionMin.X = FMath::Min(PositionMin.X, Delta.Position.X);
+					PositionMin.Y = FMath::Min(PositionMin.Y, Delta.Position.Y);
+					PositionMin.Z = FMath::Min(PositionMin.Z, Delta.Position.Z);
+
+					PositionMax.X = FMath::Max(PositionMax.X, Delta.Position.X);
+					PositionMax.Y = FMath::Max(PositionMax.Y, Delta.Position.Y);
+					PositionMax.Z = FMath::Max(PositionMax.Z, Delta.Position.Z);
+
+					TangentZMin.X = FMath::Min(TangentZMin.X, Delta.TangentZ.X);
+					TangentZMin.Y = FMath::Min(TangentZMin.Y, Delta.TangentZ.Y);
+					TangentZMin.Z = FMath::Min(TangentZMin.Z, Delta.TangentZ.Z);
+
+					TangentZMax.X = FMath::Max(TangentZMax.X, Delta.TangentZ.X);
+					TangentZMax.Y = FMath::Max(TangentZMax.Y, Delta.TangentZ.Y);
+					TangentZMax.Z = FMath::Max(TangentZMax.Z, Delta.TangentZ.Z);
+				}
+
+				const uint32 IndexDelta = IndexMax - IndexMin;
+				const FIntVector PositionDelta = PositionMax - PositionMin;
+				const FIntVector TangentZDelta = TangentZMax - TangentZMin;
+				const bool bBatchHasTangents = TangentZMin != FIntVector::ZeroValue || TangentZMax != FIntVector::ZeroValue;
+
+				FBatchHeader BatchHeader;
+				BatchHeader.DataOffset = BitstreamData.Num() * sizeof(uint32);
+				BatchHeader.bTangents = bBatchHasTangents;
+				BatchHeader.NumElements = NumElements;
+				BatchHeader.IndexBits = FMath::CeilLogTwo(IndexDelta + 1);
+				BatchHeader.PositionBits.X = FMath::CeilLogTwo(uint32(PositionDelta.X) + 1);
+				BatchHeader.PositionBits.Y = FMath::CeilLogTwo(uint32(PositionDelta.Y) + 1);
+				BatchHeader.PositionBits.Z = FMath::CeilLogTwo(uint32(PositionDelta.Z) + 1);
+				BatchHeader.TangentZBits.X = FMath::CeilLogTwo(uint32(TangentZDelta.X) + 1);
+				BatchHeader.TangentZBits.Y = FMath::CeilLogTwo(uint32(TangentZDelta.Y) + 1);
+				BatchHeader.TangentZBits.Z = FMath::CeilLogTwo(uint32(TangentZDelta.Z) + 1);
+				check(BatchHeader.IndexBits <= IndexMaxBits);
+				check(BatchHeader.PositionBits.X <= PositionMaxBits);
+				check(BatchHeader.PositionBits.Y <= PositionMaxBits);
+				check(BatchHeader.PositionBits.Z <= PositionMaxBits);
+				check(BatchHeader.TangentZBits.X <= TangentZMaxBits);
+				check(BatchHeader.TangentZBits.Y <= TangentZMaxBits);
+				check(BatchHeader.TangentZBits.Z <= TangentZMaxBits);
+				BatchHeader.IndexMin = IndexMin;
+				BatchHeader.PositionMin = PositionMin;
+				BatchHeader.TangentZMin = TangentZMin;
+
+				// Write quantized bits
+				FDwordBitWriter BitWriter(BitstreamData);
+				for (uint32 LocalElementIndex = 0; LocalElementIndex < NumElements; LocalElementIndex++)
+				{
+					const FQuantizedDelta& Delta = QuantizedDeltas[BatchFirstElementIndex + LocalElementIndex];
+					const uint32 AdjustedIndex = Delta.Index - LocalElementIndex;
+					BitWriter.PutBits(AdjustedIndex - IndexMin, BatchHeader.IndexBits);
+					BitWriter.PutBits(uint32(Delta.Position.X - PositionMin.X), BatchHeader.PositionBits.X);
+					BitWriter.PutBits(uint32(Delta.Position.Y - PositionMin.Y), BatchHeader.PositionBits.Y);
+					BitWriter.PutBits(uint32(Delta.Position.Z - PositionMin.Z), BatchHeader.PositionBits.Z);
+					if (bBatchHasTangents)
+					{
+						BitWriter.PutBits(uint32(Delta.TangentZ.X - TangentZMin.X), BatchHeader.TangentZBits.X);
+						BitWriter.PutBits(uint32(Delta.TangentZ.Y - TangentZMin.Y), BatchHeader.TangentZBits.Y);
+						BitWriter.PutBits(uint32(Delta.TangentZ.Z - TangentZMin.Z), BatchHeader.TangentZBits.Z);
+					}
+				}
+				BitWriter.Flush();
+
+				BatchHeaders.Add(BatchHeader);
+			}
+			NumTotalBatches += MorphNumBatches;
+		}
+
+		const uint32 MorphNumBatches = NumTotalBatches - BatchStartOffset;
+		BatchStartOffsetPerMorph.Add(BatchStartOffset);
+		BatchesPerMorph.Add(MorphNumBatches);
+		MaximumValuePerMorph.Add(FVector4(MaximumValues[0], MaximumValues[1], MaximumValues[2], MaximumValues[3]));
+		MinimumValuePerMorph.Add(FVector4(MinimumValues[0], MinimumValues[1], MinimumValues[2], MinimumValues[3]));
+	}
+
+	// Write packed batch headers
+	for (const FBatchHeader& BatchHeader : BatchHeaders)
+	{
+		const uint32 DataOffset = BatchHeader.DataOffset + BatchHeaders.Num() * NumBatchHeaderDwords * sizeof(uint32);
+
+		MorphData.Add(DataOffset);
+		MorphData.Add(BatchHeader.IndexBits |
+			(BatchHeader.PositionBits.X << 5) | (BatchHeader.PositionBits.Y << 10) | (BatchHeader.PositionBits.Z << 15) |
+			(BatchHeader.bTangents ? (1u << 20) : 0u) |
+			(BatchHeader.NumElements << 21));
+		MorphData.Add(BatchHeader.IndexMin);
+		MorphData.Add(BatchHeader.PositionMin.X);
+		MorphData.Add(BatchHeader.PositionMin.Y);
+		MorphData.Add(BatchHeader.PositionMin.Z);
+
+		MorphData.Add(BatchHeader.TangentZBits.X | (BatchHeader.TangentZBits.Y << 5) | (BatchHeader.TangentZBits.Z << 10));
+		MorphData.Add(BatchHeader.TangentZMin.X);
+		MorphData.Add(BatchHeader.TangentZMin.Y);
+		MorphData.Add(BatchHeader.TangentZMin.Z);
+	}
+
+	// Append bitstream data
+	MorphData.Append(BitstreamData);
+
+	// UE_LOG(LogStaticMesh, Log, TEXT("Morph compression time: [%.2fs]"), FPlatformTime::ToMilliseconds(FPlatformTime::Cycles() - StartTime) / 1000.0f);
+
+	ValidateVertexBuffers(true);
+}
diff --git a/Engine/Source/Runtime/Engine/Private/SkeletalMesh.cpp b/Engine/Source/Runtime/Engine/Private/SkeletalMesh.cpp
--- a/Engine/Source/Runtime/Engine/Private/SkeletalMesh.cpp
+++ b/Engine/Source/Runtime/Engine/Private/SkeletalMesh.cpp
@@ -1726,6 +1726,7 @@ void USkeletalMesh::Serialize( FArchive& Ar )
 {
 	LLM_SCOPE_BYNAME(TEXT("SkeletalMesh/Serialize")); // This is an important test case for SCOPE_BYNAME with a matching LLM_DEFINE_TAG
 	DECLARE_SCOPE_CYCLE_COUNTER( TEXT("USkeletalMesh::Serialize"), STAT_SkeletalMesh_Serialize, STATGROUP_LoadTime );
+	TRACE_CPUPROFILER_EVENT_SCOPE(USkeletalMesh::Serialize);
 
 #if WITH_EDITOR
 	if (IsCompiling())
@@ -5542,6 +5543,7 @@ FGuid FSkeletalMeshLODInfo::ComputeDeriveDataCacheKey(const FSkeletalMeshLODGrou
 	Ar << BonesToRemove;
 	Ar << BonesToPrioritize;
 	Ar << WeightOfPrioritization;
+	Ar << MorphTargetPositionErrorTolerance;
 
 	//TODO: Ask the derivedata key of the UObject reference by FSoftObjectPath. So if someone change the UObject, this LODs will get dirty
 	//and will be rebuild.
diff --git a/Engine/Source/Runtime/Engine/Private/SkeletalMeshLODRenderData.cpp b/Engine/Source/Runtime/Engine/Private/SkeletalMeshLODRenderData.cpp
--- a/Engine/Source/Runtime/Engine/Private/SkeletalMeshLODRenderData.cpp
+++ b/Engine/Source/Runtime/Engine/Private/SkeletalMeshLODRenderData.cpp
@@ -13,6 +13,7 @@
 #include "UObject/PropertyPortFlags.h"
 #include "UObject/UE5ReleaseStreamObjectVersion.h"
 #include "UObject/UE5MainStreamObjectVersion.h"
+#include "UObject/UE5PrivateFrostyStreamObjectVersion.h"
 #include "GPUSkinCache.h"
 
 #if WITH_EDITOR
@@ -28,9 +29,6 @@ static FAutoConsoleVariableRef CVarStripSkeletalMeshLodsBelowMinLod(
 	TEXT("If set will strip skeletal mesh LODs under the minimum renderable LOD for the target platform during cooking.")
 );
 
-extern int32 GForceRecomputeTangents;
-extern int32 GSkinCacheRecomputeTangents;
-
 namespace
 {
 	struct FReverseOrderBitArraysBySetBits
@@ -160,50 +158,10 @@ FArchive& operator<<(FArchive& Ar, FSkelMeshRenderSection& S)
 	return Ar;
 }
 
-class FDwordBitWriter
-{
-public:
-	FDwordBitWriter(TArray<uint32>& Buffer) :
-		Buffer(Buffer),
-		PendingBits(0ull),
-		NumPendingBits(0)
-	{
-	}
-
-	void PutBits(uint32 Bits, uint32 NumBits)
-	{
-		check((uint64)Bits < (1ull << NumBits));
-		PendingBits |= (uint64)Bits << NumPendingBits;
-		NumPendingBits += NumBits;
-
-		while (NumPendingBits >= 32)
-		{
-			Buffer.Add((uint32)PendingBits);
-			PendingBits >>= 32;
-			NumPendingBits -= 32;
-		}
-	}
-
-	void Flush()
-	{
-		if (NumPendingBits > 0)
-			Buffer.Add((uint32)PendingBits);
-		PendingBits = 0;
-		NumPendingBits = 0;
-	}
-
-private:
-	TArray<uint32>&	Buffer;
-	uint64 			PendingBits;
-	int32 			NumPendingBits;
-};
-
-
 void FSkeletalMeshLODRenderData::InitResources(bool bNeedsVertexColors, int32 LODIndex, TArray<UMorphTarget*>& InMorphTargets, USkeletalMesh* Owner)
 {
 	IncrementMemoryStats(bNeedsVertexColors);
 
-	MorphTargetVertexInfoBuffers.Reset();
 	MultiSizeIndexContainer.InitResources();
 
 	BeginInitResource(&StaticVertexBuffers.PositionVertexBuffer);
@@ -244,329 +202,19 @@ void FSkeletalMeshLODRenderData::InitResources(bool bNeedsVertexColors, int32 LO
 #endif
 			}
 		}
-    }
+	}
 
-	// UseGPUMorphTargets() can be toggled only on SM5 atm
-	if (IsFeatureLevelSupported(GMaxRHIShaderPlatform, ERHIFeatureLevel::SM5) && InMorphTargets.Num() > 0)
+	// Always make sure the morph target resources are in an initialized state. This could should not get hit since the data should come with the load.
+	bool bNeedsMorphTargetRenderData = Owner->GetMorphTargets().Num() > 0 && !MorphTargetVertexInfoBuffers.IsMorphResourcesInitialized();
+	if (bNeedsMorphTargetRenderData)
 	{
-		// Simple Morph compression 0.1
-		// Instead of storing vertex deltas individually they are organized into batches of 64.
-		// Each batch has a header that describes how many bits are allocated to each of the vertex components.
-		// Batches also store an explicit offset to its associated data. This makes it trivial to decode batches
-		// in parallel, and because deltas are fixed-width inside a batch, deltas can also be decoded in parallel.
-		// The result is a semi-adaptive encoding that functions as a crude substitute for entropy coding, that is
-		// fast to decode on parallel hardware.
-		
-		// Quantization still happens globally to avoid issues with cracks at duplicate vertices.
-		// The quantization is artist controlled on a per LOD basis. Higher error tolerance results in smaller deltas
-		// and a smaller compressed size.
-
 		const FSkeletalMeshLODInfo* SkeletalMeshLODInfo = Owner->GetLODInfo(LODIndex);
-		
-		const float UnrealUnitPerMeter		= 100.0f;
-		const float PositionPrecision		= SkeletalMeshLODInfo->MorphTargetPositionErrorTolerance * 2.0f * 1e-6f * UnrealUnitPerMeter;	// * 2.0 because correct rounding guarantees error is at most half of the cell size.
-		const float RcpPositionPrecision	= 1.0f / PositionPrecision;
-
-		const float TangentZPrecision		= 1.0f / 2048.0f;				// Object scale irrelevant here. Let's assume ~12bits per component is plenty.
-		const float RcpTangentZPrecision	= 1.0f / TangentZPrecision;
-
-		const uint32 BatchSize				= 64;
-		const uint32 NumBatchHeaderDwords	= 10u;
-		
-		const uint32 IndexMaxBits			= 31u;
-
-		const uint32 PositionMaxBits		= 28u;				// Probably more than we need, but let's just allow it to go this high to be safe for now.
-																// For larger deltas this can even be more precision than what was in the float input data!
-																// Maybe consider float-like or exponential encoding of large values?
-		const float PositionMinValue		= -134217728.0f;	// -2^(MaxBits-1)
-		const float PositionMaxValue		=  134217720.0f;	// Largest float smaller than 2^(MaxBits-1)-1
-																// Using 134217727.0f would NOT work as it would be rounded up to 134217728.0f, which is
-																// outside the range.
-
-		const uint32 TangentZMaxBits		= 16u;
-		const float TangentZMinValue		= -32768.0f;		// -2^(MaxBits-1)
-		const float TangentZMaxValue		=  32767.0f;		//  2^(MaxBits-1)-1
-
-		struct FBatchHeader
-		{
-			uint32		DataOffset;
-			uint32		NumElements;
-			bool		bTangents;
-
-			uint32		IndexBits;
-			FIntVector	PositionBits;
-			FIntVector	TangentZBits;
-
-			uint32		IndexMin;
-			FIntVector	PositionMin;
-			FIntVector	TangentZMin;
-		};
-
-		// uint32 StartTime = FPlatformTime::Cycles();
-
-		MorphTargetVertexInfoBuffers.MorphData.Empty();
-		MorphTargetVertexInfoBuffers.NumTotalBatches = 0;
-		MorphTargetVertexInfoBuffers.PositionPrecision = PositionPrecision;
-		MorphTargetVertexInfoBuffers.TangentZPrecision = TangentZPrecision;
-
-		MorphTargetVertexInfoBuffers.BatchStartOffsetPerMorph.Empty(InMorphTargets.Num());
-		MorphTargetVertexInfoBuffers.BatchesPerMorph.Empty(InMorphTargets.Num());
-		MorphTargetVertexInfoBuffers.MaximumValuePerMorph.Empty(InMorphTargets.Num());
-		MorphTargetVertexInfoBuffers.MinimumValuePerMorph.Empty(InMorphTargets.Num());
-
-		// Mark vertices that are in a section that doesn't recompute tangents as needing tangents
-		const int32 RecomputeTangentsMode = GForceRecomputeTangents > 0 ? 1 : GSkinCacheRecomputeTangents;
-		TBitArray<> VertexNeedsTangents;
-		VertexNeedsTangents.Init(false, StaticVertexBuffers.PositionVertexBuffer.GetNumVertices());
-		for (const FSkelMeshRenderSection& RenderSection : RenderSections)
-		{
-			const bool bRecomputeTangents = RecomputeTangentsMode > 0 && (RenderSection.bRecomputeTangent || RecomputeTangentsMode == 1);
-			if (!bRecomputeTangents)
-			{
-				for (uint32 i = 0; i < RenderSection.NumVertices; i++)
-				{
-					VertexNeedsTangents[RenderSection.BaseVertexIndex + i] = true;
-				}
-			}
-		}
-
-		// Populate the arrays to be filled in later in the render thread
-		TArray<FBatchHeader> BatchHeaders;
-		TArray<uint32> BitstreamData;
-		for (int32 AnimIdx = 0; AnimIdx < InMorphTargets.Num(); ++AnimIdx)
-		{
-			uint32 BatchStartOffset = MorphTargetVertexInfoBuffers.NumTotalBatches;
-
-			float MaximumValues[4] = { -FLT_MAX, -FLT_MAX, -FLT_MAX, -FLT_MAX };
-			float MinimumValues[4] = { +FLT_MAX, +FLT_MAX, +FLT_MAX, +FLT_MAX };
-			UMorphTarget* MorphTarget = InMorphTargets[AnimIdx];
-			int32 NumSrcDeltas = 0;
-			const FMorphTargetDelta* MorphDeltas = MorphTarget->GetMorphTargetDelta(LODIndex, NumSrcDeltas);
-
-			//Make sure the morphtarget data vertex indices fit the geometry
-			//If a missmatch happen, set the NumSrcDelta to 0 so the morph target is skipped
-			for (int32 DeltaIndex = 0; DeltaIndex < NumSrcDeltas; DeltaIndex++)
-			{
-				const auto& MorphDelta = MorphDeltas[DeltaIndex];
-				if (!VertexNeedsTangents.IsValidIndex(MorphDelta.SourceIdx))
-				{
-					NumSrcDeltas = 0;
-					UE_ASSET_LOG(LogSkeletalMesh, Warning, Owner, TEXT("Skipping morph target %s for LOD %d. The morph target data is incompatible with the mesh data"), *MorphTarget->GetName(), LODIndex);
-					break;
-				}
-			}
-
-			if (NumSrcDeltas == 0)
-			{
-				MaximumValues[0] = 0.0f;
-				MaximumValues[1] = 0.0f;
-				MaximumValues[2] = 0.0f;
-				MaximumValues[3] = 0.0f;
-
-				MinimumValues[0] = 0.0f;
-				MinimumValues[1] = 0.0f;
-				MinimumValues[2] = 0.0f;
-				MinimumValues[3] = 0.0f;
-			}
-			else
-			{
-				struct FQuantizedDelta
-				{
-					FIntVector	Position;
-					FIntVector	TangentZ;
-					uint32		Index;
-				};
-				TArray<FQuantizedDelta> QuantizedDeltas;
-				QuantizedDeltas.Reserve(NumSrcDeltas);
-
-				bool bVertexIndicesSorted = true;
-
-				int32 PrevVertexIndex = -1;
-				for (int32 DeltaIndex = 0; DeltaIndex < NumSrcDeltas; DeltaIndex++)
-				{
-					const auto& MorphDelta = MorphDeltas[DeltaIndex];
-					const FVector3f TangentZDelta = (VertexNeedsTangents.IsValidIndex(MorphDelta.SourceIdx) && VertexNeedsTangents[MorphDelta.SourceIdx]) ? MorphDelta.TangentZDelta : FVector3f::ZeroVector;
-
-					// when import, we do check threshold, and also when adding weight, we do have threshold for how smaller weight can fit in
-					// so no reason to check here another threshold
-					MaximumValues[0] = FMath::Max(MaximumValues[0], MorphDelta.PositionDelta.X);
-					MaximumValues[1] = FMath::Max(MaximumValues[1], MorphDelta.PositionDelta.Y);
-					MaximumValues[2] = FMath::Max(MaximumValues[2], MorphDelta.PositionDelta.Z);
-					MaximumValues[3] = FMath::Max(MaximumValues[3], FMath::Max(TangentZDelta.X, FMath::Max(TangentZDelta.Y, TangentZDelta.Z)));
-
-					MinimumValues[0] = FMath::Min(MinimumValues[0], MorphDelta.PositionDelta.X);
-					MinimumValues[1] = FMath::Min(MinimumValues[1], MorphDelta.PositionDelta.Y);
-					MinimumValues[2] = FMath::Min(MinimumValues[2], MorphDelta.PositionDelta.Z);
-					MinimumValues[3] = FMath::Min(MinimumValues[3], FMath::Min(TangentZDelta.X, FMath::Min(TangentZDelta.Y, TangentZDelta.Z)));
-
-					// Check if input is sorted. It usually is, but it might not be.
-					if ((int32)MorphDelta.SourceIdx < PrevVertexIndex)
-						bVertexIndicesSorted = false;
-					PrevVertexIndex = (int32)MorphDelta.SourceIdx;
-
-					// Quantize delta
-					FQuantizedDelta QuantizedDelta;
-					const FVector3f& PositionDelta	= MorphDelta.PositionDelta;
-					QuantizedDelta.Position.X		= FMath::RoundToInt(FMath::Clamp(PositionDelta.X * RcpPositionPrecision, PositionMinValue, PositionMaxValue));
-					QuantizedDelta.Position.Y		= FMath::RoundToInt(FMath::Clamp(PositionDelta.Y * RcpPositionPrecision, PositionMinValue, PositionMaxValue));
-					QuantizedDelta.Position.Z		= FMath::RoundToInt(FMath::Clamp(PositionDelta.Z * RcpPositionPrecision, PositionMinValue, PositionMaxValue));
-					QuantizedDelta.TangentZ.X		= FMath::RoundToInt(FMath::Clamp(TangentZDelta.X * RcpTangentZPrecision, TangentZMinValue, TangentZMaxValue));
-					QuantizedDelta.TangentZ.Y		= FMath::RoundToInt(FMath::Clamp(TangentZDelta.Y * RcpTangentZPrecision, TangentZMinValue, TangentZMaxValue));
-					QuantizedDelta.TangentZ.Z		= FMath::RoundToInt(FMath::Clamp(TangentZDelta.Z * RcpTangentZPrecision, TangentZMinValue, TangentZMaxValue));
-					QuantizedDelta.Index			= MorphDelta.SourceIdx;
-
-					if(QuantizedDelta.Position != FIntVector::ZeroValue || QuantizedDelta.TangentZ != FIntVector::ZeroValue)
-					{
-						// Only add delta if it is non-zero
-						QuantizedDeltas.Add(QuantizedDelta);
-					}
-				}
-
-				// Sort deltas if the source wasn't already sorted
-				if (!bVertexIndicesSorted)
-				{
-					Algo::Sort(QuantizedDeltas, [](const FQuantizedDelta& A, const FQuantizedDelta& B) { return A.Index < B.Index; });
-				}
-
-				// Encode batch deltas
-				const uint32 MorphNumBatches = (QuantizedDeltas.Num() + BatchSize - 1u) / BatchSize;
-				for (uint32 BatchIndex = 0; BatchIndex < MorphNumBatches; BatchIndex++)
-				{
-					const uint32 BatchFirstElementIndex = BatchIndex * BatchSize;
-					const uint32 NumElements = FMath::Min(BatchSize, QuantizedDeltas.Num() - BatchFirstElementIndex);
-
-					// Calculate batch min/max bounds
-					uint32 IndexMin = MAX_uint32;
-					uint32 IndexMax = MIN_uint32;
-					FIntVector PositionMin = FIntVector(MAX_int32);
-					FIntVector PositionMax = FIntVector(MIN_int32);
-					FIntVector TangentZMin = FIntVector(MAX_int32);
-					FIntVector TangentZMax = FIntVector(MIN_int32);
-
-					for (uint32 LocalElementIndex = 0; LocalElementIndex < NumElements; LocalElementIndex++)
-					{
-						const FQuantizedDelta& Delta = QuantizedDeltas[BatchFirstElementIndex + LocalElementIndex];
-
-						// Trick: Deltas are sorted by index, so the index increase by at least one per delta.
-						//		  Naively this would mean that a batch always spans at least 64 index values and
-						//		  indices would have to use at least 6 bits per index.
-						//		  If instead of storing the raw index, we store the index relative to its position in the batch,
-						//		  then the spanned range becomes 63 smaller.
-						//		  For a consecutive range this even gets us down to 0 bits per index!
-						check(Delta.Index >= LocalElementIndex);
-						const uint32 AdjustedIndex = Delta.Index - LocalElementIndex;
-						IndexMin = FMath::Min(IndexMin, AdjustedIndex);
-						IndexMax = FMath::Max(IndexMax, AdjustedIndex);
-					
-						PositionMin.X = FMath::Min(PositionMin.X, Delta.Position.X);
-						PositionMin.Y = FMath::Min(PositionMin.Y, Delta.Position.Y);
-						PositionMin.Z = FMath::Min(PositionMin.Z, Delta.Position.Z);
-
-						PositionMax.X = FMath::Max(PositionMax.X, Delta.Position.X);
-						PositionMax.Y = FMath::Max(PositionMax.Y, Delta.Position.Y);
-						PositionMax.Z = FMath::Max(PositionMax.Z, Delta.Position.Z);
-						
-						TangentZMin.X = FMath::Min(TangentZMin.X, Delta.TangentZ.X);
-						TangentZMin.Y = FMath::Min(TangentZMin.Y, Delta.TangentZ.Y);
-						TangentZMin.Z = FMath::Min(TangentZMin.Z, Delta.TangentZ.Z);
-
-						TangentZMax.X = FMath::Max(TangentZMax.X, Delta.TangentZ.X);
-						TangentZMax.Y = FMath::Max(TangentZMax.Y, Delta.TangentZ.Y);
-						TangentZMax.Z = FMath::Max(TangentZMax.Z, Delta.TangentZ.Z);
-					}
-
-					const uint32 IndexDelta			= IndexMax - IndexMin;
-					const FIntVector PositionDelta	= PositionMax - PositionMin;
-					const FIntVector TangentZDelta	= TangentZMax - TangentZMin;
-					const bool bBatchHasTangents	= TangentZMin != FIntVector::ZeroValue || TangentZMax != FIntVector::ZeroValue;
-					
-					FBatchHeader BatchHeader;
-					BatchHeader.DataOffset			= BitstreamData.Num() * sizeof(uint32);
-					BatchHeader.bTangents			= bBatchHasTangents;
-					BatchHeader.NumElements			= NumElements;
-					BatchHeader.IndexBits			= FMath::CeilLogTwo(IndexDelta + 1);
-					BatchHeader.PositionBits.X		= FMath::CeilLogTwo(uint32(PositionDelta.X) + 1);
-					BatchHeader.PositionBits.Y		= FMath::CeilLogTwo(uint32(PositionDelta.Y) + 1);
-					BatchHeader.PositionBits.Z		= FMath::CeilLogTwo(uint32(PositionDelta.Z) + 1);
-					BatchHeader.TangentZBits.X		= FMath::CeilLogTwo(uint32(TangentZDelta.X) + 1);
-					BatchHeader.TangentZBits.Y		= FMath::CeilLogTwo(uint32(TangentZDelta.Y) + 1);
-					BatchHeader.TangentZBits.Z		= FMath::CeilLogTwo(uint32(TangentZDelta.Z) + 1);
-					check(BatchHeader.IndexBits <= IndexMaxBits);
-					check(BatchHeader.PositionBits.X <= PositionMaxBits);
-					check(BatchHeader.PositionBits.Y <= PositionMaxBits);
-					check(BatchHeader.PositionBits.Z <= PositionMaxBits);
-					check(BatchHeader.TangentZBits.X <= TangentZMaxBits);
-					check(BatchHeader.TangentZBits.Y <= TangentZMaxBits);
-					check(BatchHeader.TangentZBits.Z <= TangentZMaxBits);
-					BatchHeader.IndexMin			= IndexMin;
-					BatchHeader.PositionMin			= PositionMin;
-					BatchHeader.TangentZMin			= TangentZMin;
-					
-					// Write quantized bits
-					FDwordBitWriter BitWriter(BitstreamData);
-					for (uint32 LocalElementIndex = 0; LocalElementIndex < NumElements; LocalElementIndex++)
-					{
-						const FQuantizedDelta& Delta = QuantizedDeltas[BatchFirstElementIndex + LocalElementIndex];
-						const uint32 AdjustedIndex = Delta.Index - LocalElementIndex;
-						BitWriter.PutBits(AdjustedIndex - IndexMin,					BatchHeader.IndexBits);
-						BitWriter.PutBits(uint32(Delta.Position.X - PositionMin.X), BatchHeader.PositionBits.X);
-						BitWriter.PutBits(uint32(Delta.Position.Y - PositionMin.Y), BatchHeader.PositionBits.Y);
-						BitWriter.PutBits(uint32(Delta.Position.Z - PositionMin.Z), BatchHeader.PositionBits.Z);
-						if (bBatchHasTangents)
-						{
-							BitWriter.PutBits(uint32(Delta.TangentZ.X - TangentZMin.X), BatchHeader.TangentZBits.X);
-							BitWriter.PutBits(uint32(Delta.TangentZ.Y - TangentZMin.Y), BatchHeader.TangentZBits.Y);
-							BitWriter.PutBits(uint32(Delta.TangentZ.Z - TangentZMin.Z), BatchHeader.TangentZBits.Z);
-						}
-					}
-					BitWriter.Flush();
-
-					BatchHeaders.Add(BatchHeader);
-				}
-				MorphTargetVertexInfoBuffers.NumTotalBatches += MorphNumBatches;
-			}
-
-			const uint32 MorphNumBatches = MorphTargetVertexInfoBuffers.NumTotalBatches - BatchStartOffset;
-			MorphTargetVertexInfoBuffers.BatchStartOffsetPerMorph.Add(BatchStartOffset);
-			MorphTargetVertexInfoBuffers.BatchesPerMorph.Add(MorphNumBatches);
-			MorphTargetVertexInfoBuffers.MaximumValuePerMorph.Add(FVector4(MaximumValues[0], MaximumValues[1], MaximumValues[2], MaximumValues[3]));
-			MorphTargetVertexInfoBuffers.MinimumValuePerMorph.Add(FVector4(MinimumValues[0], MinimumValues[1], MinimumValues[2], MinimumValues[3]));
-		}
-
-
-		// Write packed batch headers
-		for (const FBatchHeader& BatchHeader : BatchHeaders)
-		{
-			const uint32 DataOffset = BatchHeader.DataOffset + BatchHeaders.Num() * NumBatchHeaderDwords * sizeof(uint32);
-
-			MorphTargetVertexInfoBuffers.MorphData.Add(DataOffset);
-			MorphTargetVertexInfoBuffers.MorphData.Add(	BatchHeader.IndexBits | 
-														(BatchHeader.PositionBits.X << 5) | (BatchHeader.PositionBits.Y << 10) | (BatchHeader.PositionBits.Z << 15) |
-														(BatchHeader.bTangents ? (1u << 20) : 0u) | 
-														(BatchHeader.NumElements << 21));
-			MorphTargetVertexInfoBuffers.MorphData.Add(BatchHeader.IndexMin);
-			MorphTargetVertexInfoBuffers.MorphData.Add(BatchHeader.PositionMin.X);
-			MorphTargetVertexInfoBuffers.MorphData.Add(BatchHeader.PositionMin.Y);
-			MorphTargetVertexInfoBuffers.MorphData.Add(BatchHeader.PositionMin.Z);
-
-			MorphTargetVertexInfoBuffers.MorphData.Add(BatchHeader.TangentZBits.X | (BatchHeader.TangentZBits.Y << 5) | (BatchHeader.TangentZBits.Z << 10));
-			MorphTargetVertexInfoBuffers.MorphData.Add(BatchHeader.TangentZMin.X);
-			MorphTargetVertexInfoBuffers.MorphData.Add(BatchHeader.TangentZMin.Y);
-			MorphTargetVertexInfoBuffers.MorphData.Add(BatchHeader.TangentZMin.Z);
-		}
-
-		// Append bitstream data
-		MorphTargetVertexInfoBuffers.MorphData.Append(BitstreamData);
-
-		// UE_LOG(LogStaticMesh, Log, TEXT("Morph compression time: [%.2fs]"), FPlatformTime::ToMilliseconds(FPlatformTime::Cycles() - StartTime) / 1000.0f);
-
-		check(MorphTargetVertexInfoBuffers.BatchesPerMorph.Num() == MorphTargetVertexInfoBuffers.BatchStartOffsetPerMorph.Num());
-		check(MorphTargetVertexInfoBuffers.BatchesPerMorph.Num() == MorphTargetVertexInfoBuffers.MaximumValuePerMorph.Num());
-		check(MorphTargetVertexInfoBuffers.BatchesPerMorph.Num() == MorphTargetVertexInfoBuffers.MinimumValuePerMorph.Num());
-		if (MorphTargetVertexInfoBuffers.NumTotalBatches > 0)
-		{
-			BeginInitResource(&MorphTargetVertexInfoBuffers);
-		}
+		MorphTargetVertexInfoBuffers.InitMorphResources(GMaxRHIShaderPlatform, RenderSections, Owner->GetMorphTargets(), StaticVertexBuffers.StaticMeshVertexBuffer.GetNumVertices(), LODIndex, SkeletalMeshLODInfo->MorphTargetPositionErrorTolerance);
+	}
+	
+	if (!MorphTargetVertexInfoBuffers.IsRHIIntialized() && MorphTargetVertexInfoBuffers.IsMorphCPUDataValid() && MorphTargetVertexInfoBuffers.NumTotalBatches > 0)
+	{
+		BeginInitResource(&MorphTargetVertexInfoBuffers);
 	}
 
 #if RHI_RAYTRACING
@@ -650,7 +298,6 @@ void FSkeletalMeshLODRenderData::DecrementMemoryStats()
 }
 
 #if WITH_EDITOR
-
 void FSkeletalMeshLODRenderData::BuildFromLODModel(const FSkeletalMeshLODModel* ImportedModel, uint32 BuildFlags)
 {
 	bool bUseFullPrecisionUVs = (BuildFlags & ESkeletalMeshVertexFlags::UseFullPrecisionUVs) != 0;
@@ -949,7 +596,9 @@ private:
 
 void FSkeletalMeshLODRenderData::SerializeStreamedData(FArchive& Ar, USkeletalMesh* Owner, int32 LODIdx, uint8 ClassDataStripFlags, bool bNeedsCPUAccess, bool bForceKeepCPUResources)
 {
+	Ar.UsingCustomVersion(FUE5PrivateFrostyStreamObjectVersion::GUID);
 	Ar.UsingCustomVersion(FUE5ReleaseStreamObjectVersion::GUID);
+
 	FStripDataFlags StripFlags(Ar, ClassDataStripFlags);
 
 	// TODO: A lot of data in a render section is needed during initialization but maybe some can still be streamed
@@ -1001,6 +650,81 @@ void FSkeletalMeshLODRenderData::SerializeStreamedData(FArchive& Ar, USkeletalMe
 		}
 	}
 	Ar << SourceRayTracingGeometry.RawData;
+
+	// Determine if morph target data should be cooked out for this platform
+	EShaderPlatform MorphTargetShaderPlatform = GMaxRHIShaderPlatform;
+	bool bSerializeCompressedMorphTargets = Ar.IsSaving() && Owner->GetMorphTargets().Num() > 0;
+	if (Ar.IsCooking())
+	{
+		const ITargetPlatform* Platform = Ar.CookingTarget();
+		
+		bSerializeCompressedMorphTargets = false;
+
+		// Make sure to avoid cooking the morph target data when build a server only executable
+		if (!Platform->IsServerOnly() && Owner->GetMorphTargets().Num() > 0)
+		{
+			// Test if any of the supported shader formats supports SM5
+			TArray<FName> ShaderFormats;
+			Ar.CookingTarget()->GetAllTargetedShaderFormats(ShaderFormats);
+			for (FName ShaderFormat : ShaderFormats)
+			{
+				const EShaderPlatform ShaderPlatform = ShaderFormatToLegacyShaderPlatform(ShaderFormat);
+				if (IsFeatureLevelSupported(ShaderPlatform, ERHIFeatureLevel::SM5))
+				{
+					MorphTargetShaderPlatform = ShaderPlatform;
+					bSerializeCompressedMorphTargets = true;
+					break;
+				}
+			}
+		}
+	}
+
+	if(Ar.CustomVer(FUE5PrivateFrostyStreamObjectVersion::GUID) >= FUE5PrivateFrostyStreamObjectVersion::SerializeSkeletalMeshMorphTargetRenderData)
+	{
+		Ar << bSerializeCompressedMorphTargets;
+
+		if (bSerializeCompressedMorphTargets)
+		{
+#if WITH_EDITOR
+			if (Ar.IsSaving())
+			{
+				FMorphTargetVertexInfoBuffers LocalMorphTargetVertexInfoBuffers;
+				FMorphTargetVertexInfoBuffers* TargetMorphTargetVertexInfoBuffers;
+				const FSkeletalMeshLODInfo* SkeletalMeshLODInfo = Owner->GetLODInfo(LODIdx);
+				const TArray<UMorphTarget*>& MorphTargets = Owner->GetMorphTargets();
+
+				// The CPU data could have already been destroyed by this point, which happens when the RHI is initialized.  If possible, use the MorphTargetVertexInfoBuffers.
+				if (!MorphTargetVertexInfoBuffers.IsMorphResourcesInitialized())
+				{
+					TargetMorphTargetVertexInfoBuffers = &MorphTargetVertexInfoBuffers;
+					TargetMorphTargetVertexInfoBuffers->InitMorphResources(MorphTargetShaderPlatform, RenderSections, MorphTargets, StaticVertexBuffers.StaticMeshVertexBuffer.GetNumVertices(), LODIdx, SkeletalMeshLODInfo->MorphTargetPositionErrorTolerance);
+				}
+				else if (MorphTargetVertexInfoBuffers.IsMorphCPUDataValid() && !MorphTargetVertexInfoBuffers.IsRHIIntialized())
+				{
+					TargetMorphTargetVertexInfoBuffers = &MorphTargetVertexInfoBuffers;
+				}
+				else
+				{
+					// Fallback for when the RHI data has already been created.  Create a local version of the morph data, and serialize that
+					TargetMorphTargetVertexInfoBuffers = &LocalMorphTargetVertexInfoBuffers;
+					TargetMorphTargetVertexInfoBuffers->InitMorphResources(MorphTargetShaderPlatform, RenderSections, MorphTargets, StaticVertexBuffers.StaticMeshVertexBuffer.GetNumVertices(), LODIdx, SkeletalMeshLODInfo->MorphTargetPositionErrorTolerance);
+				}
+
+				Ar << *TargetMorphTargetVertexInfoBuffers;
+			}
+			else
+#endif
+			{
+				Ar << MorphTargetVertexInfoBuffers;
+
+#if DO_CHECK
+				const FSkeletalMeshLODInfo* SkeletalMeshLODInfo = Owner->GetLODInfo(LODIdx);
+				float PositionPrecision = FMorphTargetVertexInfoBuffers::CalculatePositionPrecision(SkeletalMeshLODInfo->MorphTargetPositionErrorTolerance);
+				checkf(PositionPrecision == MorphTargetVertexInfoBuffers.GetPositionPrecision(), TEXT("Morph target render data position tollerance for the skeleton %s does not match the expected tolerance in the LOD info."), *Owner->GetName());
+#endif
+			}
+		}
+	}
 }
 
 void FSkeletalMeshLODRenderData::SerializeAvailabilityInfo(FArchive& Ar, USkeletalMesh* Owner, int32 LODIdx, bool bAdjacencyDataStripped, bool bNeedsCPUAccess)
diff --git a/Engine/Source/Runtime/Engine/Private/SkeletalMeshRenderData.cpp b/Engine/Source/Runtime/Engine/Private/SkeletalMeshRenderData.cpp
--- a/Engine/Source/Runtime/Engine/Private/SkeletalMeshRenderData.cpp
+++ b/Engine/Source/Runtime/Engine/Private/SkeletalMeshRenderData.cpp
@@ -186,7 +186,7 @@ static void SerializeLODInfoForDDC(USkeletalMesh* SkeletalMesh, FString& KeySuff
 // differences, etc.) replace the version GUID below with a new one.
 // In case of merge conflicts with DDC versions, you MUST generate a new GUID
 // and set this new GUID as the version.
-#define SKELETALMESH_DERIVEDDATA_VER TEXT("7BA0E85CEABA42BF8930246280B6F7E8")
+#define SKELETALMESH_DERIVEDDATA_VER TEXT("61C5045866H572HEB3C0686189A9E7ED")
 
 const FString& GetSkeletalMeshDerivedDataVersion()
 {
@@ -362,7 +362,7 @@ void FSkeletalMeshRenderData::Cache(const ITargetPlatform* TargetPlatform, USkel
 		if(bAllowDdcFetch && DDCUtils64Bit::GetSynchronous(DerivedDataKey, Owner, DerivedData))
 		{
 			COOK_STAT(Timer.AddHit(DerivedData.Num()));
-			
+
 			FLargeMemoryReader Ar(DerivedData.GetData(), DerivedData.Num(), ELargeMemoryReaderFlags::Persistent);
 
 			//With skeletal mesh build refactor we serialize the LODModel data into the DDC
@@ -472,7 +472,7 @@ void FSkeletalMeshRenderData::Cache(const ITargetPlatform* TargetPlatform, USkel
 			}
 
 			FLargeMemoryWriter Ar(0, /*bIsPersistent=*/ true);
-			
+
 			//If we load an old asset we want to be sure the serialize ddc will be the same has before the skeletalmesh build refactor
 			//So we do not serialize the LODModel sections.
 			if (!Owner->GetUseLegacyMeshDerivedDataKey())
diff --git a/Engine/Source/Runtime/Engine/Public/Rendering/MorphTargetVertexInfoBuffers.h b/Engine/Source/Runtime/Engine/Public/Rendering/MorphTargetVertexInfoBuffers.h
--- a/Engine/Source/Runtime/Engine/Public/Rendering/MorphTargetVertexInfoBuffers.h
+++ b/Engine/Source/Runtime/Engine/Public/Rendering/MorphTargetVertexInfoBuffers.h
@@ -4,6 +4,8 @@
 
 #include "RenderResource.h"
 
+struct FSkelMeshRenderSection;
+class UMorphTarget;
 
 class FMorphTargetVertexInfoBuffers : public FRenderResource
 {
@@ -12,6 +14,12 @@ public:
 	{
 	}
 
+	void InitMorphResources(EShaderPlatform ShaderPlatform, const TArray<FSkelMeshRenderSection>& RenderSections, const TArray<UMorphTarget*>& MorphTargets, int NumVertices, int32 LODIndex, float TargetPositionErrorTolerance);
+
+	inline bool IsMorphResourcesInitialized() const { return bResourcesInitialized; }
+	inline bool IsRHIIntialized() const { return bRHIIntialized; }
+	inline bool IsMorphCPUDataValid() const{ return bIsMorphCPUDataValid; }
+
 	ENGINE_API virtual void InitRHI() override;
 	ENGINE_API virtual void ReleaseRHI() override;
 
@@ -56,6 +64,8 @@ public:
 		return PositionPrecision;
 	}
 
+	static const float CalculatePositionPrecision(float TargetPositionErrorTolerance);
+
 	const float GetTangentZPrecision() const
 	{
 		return TangentZPrecision;
@@ -70,14 +80,20 @@ public:
 		MaximumValuePerMorph.Empty();
 		MinimumValuePerMorph.Empty();
 		BatchStartOffsetPerMorph.Empty();
-		BatchesPerMorph.Empty();		
+		BatchesPerMorph.Empty();
 		NumTotalBatches = 0;
 		PositionPrecision = 0.0f;
 		TangentZPrecision = 0.0f;
+		bResourcesInitialized = false;
+		bRHIIntialized = false;
+		bIsMorphCPUDataValid = false;
 	}
 
 protected:
 
+	void ValidateVertexBuffers(bool bMorphTargetsShouldBeValid);
+	void Serialize(FArchive& Ar);
+
 	// Transient data. Gets deleted as soon as the GPU resource has been initialized.
 	TArray<uint32> MorphData;
 
@@ -91,5 +107,12 @@ protected:
 	float PositionPrecision;
 	float TangentZPrecision;
 
+	bool bIsMorphCPUDataValid = false;
+	bool bResourcesInitialized = false;
+	bool bRHIIntialized = false;
+
 	friend class FSkeletalMeshLODRenderData;
+	friend FArchive& operator<<(FArchive& Ar, FMorphTargetVertexInfoBuffers& MorphTargetVertexInfoBuffers);
 };
+
+FArchive& operator<<(FArchive& Ar, FMorphTargetVertexInfoBuffers& MorphTargetVertexInfoBuffers);
