Re-add changes to MorphTargetVertexInfoBuffers
Integration: CL-17964279

diff --git a/Engine/Source/Runtime/Engine/Private/Rendering/MorphTargetVertexInfoBuffers.cpp b/Engine/Source/Runtime/Engine/Private/Rendering/MorphTargetVertexInfoBuffers.cpp
--- a/Engine/Source/Runtime/Engine/Private/Rendering/MorphTargetVertexInfoBuffers.cpp
+++ b/Engine/Source/Runtime/Engine/Private/Rendering/MorphTargetVertexInfoBuffers.cpp
@@ -241,7 +241,20 @@ void FMorphTargetVertexInfoBuffers::InitMorphResources(EShaderPlatform ShaderPla
 		float MinimumValues[4] = { +FLT_MAX, +FLT_MAX, +FLT_MAX, +FLT_MAX };
 		UMorphTarget* MorphTarget = MorphTargets[AnimIdx];
 		int32 NumSrcDeltas = 0;
-		FMorphTargetDelta* MorphDeltas = MorphTarget->GetMorphTargetDelta(LODIndex, NumSrcDeltas);
+		const FMorphTargetDelta* MorphDeltas = MorphTarget->GetMorphTargetDelta(LODIndex, NumSrcDeltas);
+
+		//Make sure the morphtarget data vertex indices fit the geometry
+		//If a missmatch happen, set the NumSrcDelta to 0 so the morph target is skipped
+		for (int32 DeltaIndex = 0; DeltaIndex < NumSrcDeltas; DeltaIndex++)
+		{
+			const auto& MorphDelta = MorphDeltas[DeltaIndex];
+			if (!VertexNeedsTangents.IsValidIndex(MorphDelta.SourceIdx))
+			{
+				NumSrcDeltas = 0;
+				UE_ASSET_LOG(LogSkeletalMesh, Warning, Owner, TEXT("Skipping morph target %s for LOD %d. The morph target data is incompatible with the mesh data"), *MorphTarget->GetName(), LODIndex);
+				break;
+			}
+		}
 
 		if (NumSrcDeltas == 0)
 		{
@@ -272,7 +285,7 @@ void FMorphTargetVertexInfoBuffers::InitMorphResources(EShaderPlatform ShaderPla
 			for (int32 DeltaIndex = 0; DeltaIndex < NumSrcDeltas; DeltaIndex++)
 			{
 				const auto& MorphDelta = MorphDeltas[DeltaIndex];
-				const FVector3f TangentZDelta = VertexNeedsTangents[MorphDelta.SourceIdx] ? MorphDelta.TangentZDelta : FVector3f::ZeroVector;
+				const FVector3f TangentZDelta = (VertexNeedsTangents.IsValidIndex(MorphDelta.SourceIdx) && VertexNeedsTangents[MorphDelta.SourceIdx]) ? MorphDelta.TangentZDelta : FVector3f::ZeroVector;
 
 				// when import, we do check threshold, and also when adding weight, we do have threshold for how smaller weight can fit in
 				// so no reason to check here another threshold
