Only call FreeSkeletalMeshBuffersSinkCallback when the value actually changed, since it results in execution times of ~500ms in certain circumpstances.

diff --git a/Engine/Source/Runtime/Engine/Private/SkeletalMesh.cpp b/Engine/Source/Runtime/Engine/Private/SkeletalMesh.cpp
--- a/Engine/Source/Runtime/Engine/Private/SkeletalMesh.cpp
+++ b/Engine/Source/Runtime/Engine/Private/SkeletalMesh.cpp
@@ -289,10 +289,13 @@ FreeSkeletalMeshBuffersSinkCallback
 void FreeSkeletalMeshBuffersSinkCallback()
 {
 	// If r.FreeSkeletalMeshBuffers==1 then CPU buffer copies are to be released.
-	static const auto CVar = IConsoleManager::Get().FindTConsoleVariableDataInt(TEXT("r.FreeSkeletalMeshBuffers"));
+	static TOptional<bool> bLastFreeSkeletalMeshBuffers;
+	static const auto CVar = IConsoleManager::Get().FindTConsoleVariableDataInt(TEXT("r.FreeSkeletalMeshBuffers"));	
 	bool bFreeSkeletalMeshBuffers = CVar->GetValueOnGameThread() == 1;
-	if(bFreeSkeletalMeshBuffers)
+	if(!bLastFreeSkeletalMeshBuffers.IsSet() || (bFreeSkeletalMeshBuffers != bLastFreeSkeletalMeshBuffers.GetValue()))
 	{
+		TRACE_CPUPROFILER_EVENT_SCOPE(FreeSkeletalMeshBuffersSinkCallback);
+
 		FlushRenderingCommands();
 		for (TObjectIterator<USkeletalMesh> It;It;++It)
 		{
@@ -301,6 +304,7 @@ void FreeSkeletalMeshBuffersSinkCallback()
 				It->ReleaseCPUResources();
 			}
 		}
+		bLastFreeSkeletalMeshBuffers = bFreeSkeletalMeshBuffers;
 	}
 }
 
