UE5:  MEMORY:  In FCollisionResponse::UpdateArrayFromResponseContainer (under ELLMTag::StaticMesh category), gather elements into a local array before filling in the output array, to avoid memory waste.  Saves 21 MB total.

#rnx

diff --git a/Engine/Source/Runtime/Engine/Private/PhysicsEngine/BodyInstance.cpp b/Engine/Source/Runtime/Engine/Private/PhysicsEngine/BodyInstance.cpp
--- a/Engine/Source/Runtime/Engine/Private/PhysicsEngine/BodyInstance.cpp
+++ b/Engine/Source/Runtime/Engine/Private/PhysicsEngine/BodyInstance.cpp
@@ -253,11 +253,13 @@ bool FCollisionResponse::AddReponseToArray(ECollisionChannel Channel, ECollision
 
 void FCollisionResponse::UpdateArrayFromResponseContainer()
 {
-	ResponseArray.Empty(UE_ARRAY_COUNT(ResponseToChannels.EnumArray));
+	FResponseChannel LocalResponseArray[UE_ARRAY_COUNT(ResponseToChannels.EnumArray)];
+	int32 LocalResponseCount = 0;
 
 	const FCollisionResponseContainer& DefaultResponse = FCollisionResponseContainer::GetDefaultResponseContainer();
 	const UCollisionProfile* CollisionProfile = UCollisionProfile::Get();
 
+	// Do an initial pass to gather elements, to allow the array to be sized exactly, avoiding memory waste
 	for (int32 i = 0; i < UE_ARRAY_COUNT(ResponseToChannels.EnumArray); i++)
 	{
 		// if not same as default
@@ -266,13 +268,16 @@ void FCollisionResponse::UpdateArrayFromResponseContainer()
 			FName ChannelName = CollisionProfile->ReturnChannelNameFromContainerIndex(i);
 			if (ChannelName != NAME_None)
 			{
-				FResponseChannel NewResponse;
+				FResponseChannel& NewResponse = LocalResponseArray[LocalResponseCount++];
+				check(LocalResponseCount <= UE_ARRAY_COUNT(LocalResponseArray));
+
 				NewResponse.Channel = ChannelName;
 				NewResponse.Response = (ECollisionResponse)ResponseToChannels.EnumArray[i];
-				ResponseArray.Add(NewResponse);
 			}
 		}
 	}
+
+	ResponseArray.Insert(LocalResponseArray, LocalResponseCount, 0);
 }
 
 #endif // WITH_EDITOR
