UE5:  MEMORY:  Fix excess slack allocation for matrix arrays.  Call "Empty" with an explicit slack value to set the exact number of elements to allocate, rather than allowing the slack growth algorithm to pad the number arbitrarily.  Also allows the sizes to drop for LODs, rather than staying allocated at the maximum count.  Saves 23 MiB in a test case.

#rnx

diff --git a/Engine/Source/Runtime/Engine/Private/SkeletalRender.cpp b/Engine/Source/Runtime/Engine/Private/SkeletalRender.cpp
--- a/Engine/Source/Runtime/Engine/Private/SkeletalRender.cpp
+++ b/Engine/Source/Runtime/Engine/Private/SkeletalRender.cpp
@@ -325,6 +325,8 @@ void UpdateRefToLocalMatricesInner(TArray<FMatrix44f>& ReferenceToLocal, const T
  */
 void UpdateRefToLocalMatrices( TArray<FMatrix44f>& ReferenceToLocal, const USkinnedMeshComponent* InMeshComponent, const FSkeletalMeshRenderData* InSkeletalMeshRenderData, int32 LODIndex, const TArray<FBoneIndexType>* ExtraRequiredBoneIndices )
 {
+	LLM_SCOPE((ELLMTag)((uint32)ELLMTag::ProjectTagStart + 1));
+
 	const USkeletalMesh* const ThisMesh = InMeshComponent->SkeletalMesh;
 	const USkinnedMeshComponent* const MasterComp = InMeshComponent->MasterPoseComponent.Get();
 	const FSkeletalMeshLODRenderData& LOD = InSkeletalMeshRenderData->LODRenderData[LODIndex];
@@ -348,7 +350,7 @@ void UpdateRefToLocalMatrices( TArray<FMatrix44f>& ReferenceToLocal, const USkin
 
 	if(ReferenceToLocal.Num() != RefBasesInvMatrix->Num())
 	{
-		ReferenceToLocal.Reset();
+		ReferenceToLocal.Empty(RefBasesInvMatrix->Num());
 		ReferenceToLocal.AddUninitialized(RefBasesInvMatrix->Num());
 
 		for (int32 Index = 0; Index < ReferenceToLocal.Num(); ++Index)
@@ -401,7 +403,7 @@ void UpdatePreviousRefToLocalMatrices(TArray<FMatrix44f>& ReferenceToLocal, cons
 
 	if (ReferenceToLocal.Num() != RefBasesInvMatrix->Num())
 	{
-		ReferenceToLocal.Reset();
+		ReferenceToLocal.Empty(RefBasesInvMatrix->Num());
 		ReferenceToLocal.AddUninitialized(RefBasesInvMatrix->Num());
 
 		for (int32 Index = 0; Index < ReferenceToLocal.Num(); ++Index)
