Use transient allocation for GPU skincache BLAS scratch allocations - saves lots of committed resource allocations and also saves upto 25MB of buffer memory

diff --git a/Engine/Source/Runtime/D3D12RHI/Private/D3D12RayTracing.cpp b/Engine/Source/Runtime/D3D12RHI/Private/D3D12RayTracing.cpp
--- a/Engine/Source/Runtime/D3D12RHI/Private/D3D12RayTracing.cpp
+++ b/Engine/Source/Runtime/D3D12RHI/Private/D3D12RayTracing.cpp
@@ -3625,7 +3625,7 @@ void FD3D12RayTracingGeometry::SetupHitGroupSystemParameters(uint32 InGPUIndex)
 	}
 }
 
-void FD3D12RayTracingGeometry::CreateAccelerationStructureBuildDesc(FD3D12CommandContext& CommandContext, EAccelerationStructureBuildMode BuildMode, const TRefCountPtr<FD3D12Buffer>& ScratchBuffer, D3D12_BUILD_RAYTRACING_ACCELERATION_STRUCTURE_DESC& OutDesc, TArrayView<D3D12_RAYTRACING_GEOMETRY_DESC>& OutGeometryDescs) const
+void FD3D12RayTracingGeometry::CreateAccelerationStructureBuildDesc(FD3D12CommandContext& CommandContext, EAccelerationStructureBuildMode BuildMode, D3D12_GPU_VIRTUAL_ADDRESS ScratchBufferAddress, D3D12_BUILD_RAYTRACING_ACCELERATION_STRUCTURE_DESC& OutDesc, TArrayView<D3D12_RAYTRACING_GEOMETRY_DESC>& OutGeometryDescs) const
 {
 	if (RHIIndexBuffer)
 	{
@@ -3762,7 +3762,7 @@ void FD3D12RayTracingGeometry::CreateAccelerationStructureBuildDesc(FD3D12Comman
 	D3D12_BUILD_RAYTRACING_ACCELERATION_STRUCTURE_DESC BuildDesc = {};
 	BuildDesc.Inputs = PrebuildDescInputs;
 	BuildDesc.DestAccelerationStructureData = AccelerationStructureBuffers[GPUIndex]->ResourceLocation.GetGPUVirtualAddress();
-	BuildDesc.ScratchAccelerationStructureData = ScratchBuffer->ResourceLocation.GetGPUVirtualAddress();
+	BuildDesc.ScratchAccelerationStructureData = ScratchBufferAddress;
 	BuildDesc.SourceAccelerationStructureData = bIsUpdate
 		? AccelerationStructureBuffers[GPUIndex]->ResourceLocation.GetGPUVirtualAddress()
 		: D3D12_GPU_VIRTUAL_ADDRESS(0);
@@ -4246,30 +4246,52 @@ void FD3D12CommandContext::RHIBuildAccelerationStructures(const TArrayView<const
 		Geometry->SetupHitGroupSystemParameters(GPUIndex);
 
 		if (Geometry->IsDirty(GPUIndex))
-		{
-			TRefCountPtr<FD3D12Buffer> ScratchBuffer;
-			if (P.BuildMode == EAccelerationStructureBuildMode::Update)
+		{			
+			D3D12_GPU_VIRTUAL_ADDRESS ScratchBufferAddress;
+			if (P.ScratchBuffer)
 			{
-				static const FName BufferName("UpdateScratchBLAS");
-				ScratchBuffer = CreateRayTracingBuffer(GetParentAdapter(), GPUIndex, Geometry->SizeInfo.UpdateScratchSize, ERayTracingBufferType::Scratch, BufferName);
+				FD3D12Buffer* ScratchBuffer = FD3D12DynamicRHI::ResourceCast(P.ScratchBuffer);
+
+				uint64 ScratchBufferRequiredSize = P.BuildMode == EAccelerationStructureBuildMode::Update ? Geometry->SizeInfo.UpdateScratchSize : Geometry->SizeInfo.BuildScratchSize;
+				checkf(ScratchBufferRequiredSize + P.ScratchBufferOffset <= ScratchBuffer->GetSize(),
+					TEXT("BLAS scratch buffer size is %lld bytes with offset %lld (%lld bytes available), but the build requires %lld bytes. "),
+					ScratchBuffer->GetSize(), P.ScratchBufferOffset, ScratchBuffer->GetSize() - P.ScratchBufferOffset, ScratchBufferRequiredSize);
+
+				ScratchBufferAddress = ScratchBuffer->ResourceLocation.GetGPUVirtualAddress() + P.ScratchBufferOffset;
+				ScratchBuffer->GetResource()->UpdateResidency(CommandListHandle);
 			}
 			else
 			{
-				static const FName BufferName("BuildScratchBLAS");
-				ScratchBuffer = CreateRayTracingBuffer(GetParentAdapter(), GPUIndex, Geometry->SizeInfo.BuildScratchSize, ERayTracingBufferType::Scratch, BufferName);
+				TRefCountPtr<FD3D12Buffer> ScratchBuffer;
+				if (P.BuildMode == EAccelerationStructureBuildMode::Update)
+				{
+					static const FName BufferName("UpdateScratchBLAS");
+					ScratchBuffer = CreateRayTracingBuffer(GetParentAdapter(), GPUIndex, Geometry->SizeInfo.UpdateScratchSize, ERayTracingBufferType::Scratch, BufferName);
+				}
+				else
+				{
+					static const FName BufferName("BuildScratchBLAS");
+					ScratchBuffer = CreateRayTracingBuffer(GetParentAdapter(), GPUIndex, Geometry->SizeInfo.BuildScratchSize, ERayTracingBufferType::Scratch, BufferName);
+				}
+
+				ScratchBufferAddress = ScratchBuffer->ResourceLocation.GetGPUVirtualAddress();
+				ScratchBuffers.Add(ScratchBuffer);
+				ScratchBuffer->GetResource()->UpdateResidency(CommandListHandle);
 			}
 
+			checkf(ScratchBufferAddress % GRHIRayTracingAccelerationStructureAlignment == 0,
+				TEXT("BLAS scratch buffer (plus offset) must be aligned to %lld bytes."),
+				GRHIRayTracingAccelerationStructureAlignment);
+
 			// We need to keep D3D12_RAYTRACING_GEOMETRY_DESCs that are used in D3D12_BUILD_RAYTRACING_ACCELERATION_STRUCTURE_DESC alive.
 			const uint32 NumGeometryDescs = Geometry->GeometryDescs.Num();
 			D3D12_RAYTRACING_GEOMETRY_DESC* LocalGeometryDescsMemory = (D3D12_RAYTRACING_GEOMETRY_DESC*)FMemStack::Get().Alloc(NumGeometryDescs * sizeof(D3D12_RAYTRACING_GEOMETRY_DESC), alignof(D3D12_RAYTRACING_GEOMETRY_DESC));
 			TArrayView<D3D12_RAYTRACING_GEOMETRY_DESC> LocalGeometryDescs = MakeArrayView(LocalGeometryDescsMemory, NumGeometryDescs);
 
 			D3D12_BUILD_RAYTRACING_ACCELERATION_STRUCTURE_DESC& BuildDesc = BuildDescs.AddZeroed_GetRef();			
-			Geometry->CreateAccelerationStructureBuildDesc(*this, P.BuildMode, ScratchBuffer, BuildDesc, LocalGeometryDescs);
+			Geometry->CreateAccelerationStructureBuildDesc(*this, P.BuildMode, ScratchBufferAddress, BuildDesc, LocalGeometryDescs);
 
 			Geometry->UpdateResidency(*this);
-			ScratchBuffer->GetResource()->UpdateResidency(CommandListHandle);
-			ScratchBuffers.Add(ScratchBuffer);
 
 			if (P.BuildMode == EAccelerationStructureBuildMode::Update)
 			{
@@ -4338,8 +4360,6 @@ void FD3D12CommandContext::RHIClearRayTracingBindings(FRHIRayTracingScene* InSce
 		delete Item.Value;
 	}
 	Table.Reset();
-
-	Scene->ReleaseBuffer();
 }
 
 struct FD3D12RayTracingGlobalResourceBinder
diff --git a/Engine/Source/Runtime/D3D12RHI/Private/D3D12RayTracing.h b/Engine/Source/Runtime/D3D12RHI/Private/D3D12RayTracing.h
--- a/Engine/Source/Runtime/D3D12RHI/Private/D3D12RayTracing.h
+++ b/Engine/Source/Runtime/D3D12RHI/Private/D3D12RayTracing.h
@@ -37,7 +37,7 @@ public:
 	void TransitionBuffers(FD3D12CommandContext& CommandContext);
 	void UpdateResidency(FD3D12CommandContext& CommandContext);
 	void CompactAccelerationStructure(FD3D12CommandContext& CommandContext, uint32 InGPUIndex, uint64 InSizeAfterCompaction);
-	void CreateAccelerationStructureBuildDesc(FD3D12CommandContext& CommandContext, EAccelerationStructureBuildMode BuildMode, const TRefCountPtr<FD3D12Buffer>& ScratchBuffer,
+	void CreateAccelerationStructureBuildDesc(FD3D12CommandContext& CommandContext, EAccelerationStructureBuildMode BuildMode, D3D12_GPU_VIRTUAL_ADDRESS ScratchBufferAddress,
 											D3D12_BUILD_RAYTRACING_ACCELERATION_STRUCTURE_DESC& OutDesc, TArrayView<D3D12_RAYTRACING_GEOMETRY_DESC>& OutGeometryDescs) const;
 	
 	// Implement FD3D12ShaderResourceRenameListener interface
diff --git a/Engine/Source/Runtime/Engine/Private/GPUSkinCache.cpp b/Engine/Source/Runtime/Engine/Private/GPUSkinCache.cpp
--- a/Engine/Source/Runtime/Engine/Private/GPUSkinCache.cpp
+++ b/Engine/Source/Runtime/Engine/Private/GPUSkinCache.cpp
@@ -823,23 +823,46 @@ void FGPUSkinCache::TransitionAllToReadable(FRHICommandList& RHICmdList)
 
 #if RHI_RAYTRACING
 
-void FGPUSkinCache::AddRayTracingGeometryToUpdate(FRayTracingGeometry* InRayTracingGeometry, EAccelerationStructureBuildMode InBuildMode)
+void FGPUSkinCache::AddRayTracingGeometryToUpdate(FRayTracingGeometry* InRayTracingGeometry, const FRayTracingAccelerationStructureSize& StructureSize, EAccelerationStructureBuildMode InBuildMode)
 {
-	EAccelerationStructureBuildMode* CurrentBuildMode = RayTracingGeometriesToUpdate.Find(InRayTracingGeometry);
-	if (CurrentBuildMode == nullptr)
+	FRayTracingUpdateInfo* CurrentUpdateInfo = RayTracingGeometriesToUpdate.Find(InRayTracingGeometry);
+	if (CurrentUpdateInfo == nullptr)
 	{
-		RayTracingGeometriesToUpdate.Add(InRayTracingGeometry, InBuildMode);
+		FRayTracingUpdateInfo UpdateInfo;
+		UpdateInfo.BuildMode = InBuildMode;
+		UpdateInfo.ScratchSize = InBuildMode == EAccelerationStructureBuildMode::Build ? StructureSize.BuildScratchSize : StructureSize.UpdateScratchSize;
+		RayTracingGeometriesToUpdate.Add(InRayTracingGeometry, UpdateInfo);
 	}
 	// If currently updating but need full rebuild then update the stored build mode
-	else if (*CurrentBuildMode == EAccelerationStructureBuildMode::Update && InBuildMode == EAccelerationStructureBuildMode::Build)
+	else if (CurrentUpdateInfo->BuildMode == EAccelerationStructureBuildMode::Update && InBuildMode == EAccelerationStructureBuildMode::Build)
 	{
-		RayTracingGeometriesToUpdate[InRayTracingGeometry] = InBuildMode;
+		CurrentUpdateInfo->BuildMode = InBuildMode;
+		CurrentUpdateInfo->ScratchSize = StructureSize.BuildScratchSize;
 	}
 }
 
+uint32 FGPUSkinCache::ComputeRayTracingGeometryScratchBufferSize()
+{
+	TRACE_CPUPROFILER_EVENT_SCOPE(FGPUSkinCache::ComputeRayTracingGeometryScratchBufferSize);
+
+	const uint64 ScratchAlignment = GRHIRayTracingAccelerationStructureAlignment;
+	uint32 ScratchBLASSize = 0;
+
+	if (RayTracingGeometriesToUpdate.Num())
+	{
+		for (TMap<FRayTracingGeometry*, FRayTracingUpdateInfo>::TRangedForIterator Iter = RayTracingGeometriesToUpdate.begin(); Iter != RayTracingGeometriesToUpdate.end(); ++Iter)
+		{			
+			FRayTracingUpdateInfo& UpdateInfo = Iter.Value();
+			ScratchBLASSize = Align(ScratchBLASSize + UpdateInfo.ScratchSize, ScratchAlignment);
+		}
+	}
+
+	return ScratchBLASSize;
+}
+
 DECLARE_GPU_STAT(GPUSkinCacheBuildBLAS);
 DECLARE_GPU_STAT(GPUSkinCacheUpdateBLAS);
-void FGPUSkinCache::CommitRayTracingGeometryUpdates(FRHICommandListImmediate& RHICmdList)
+void FGPUSkinCache::CommitRayTracingGeometryUpdates(FRHICommandListImmediate& RHICmdList, FRHIBuffer* ScratchBuffer)
 {
 	TRACE_CPUPROFILER_EVENT_SCOPE(FGPUSkinCache::CommitRayTracingGeometryUpdates);
 
@@ -882,14 +905,24 @@ void FGPUSkinCache::CommitRayTracingGeometryUpdates(FRHICommandListImmediate& RH
 			PrimitivesToUpdates = 0;
 		};
 
+		const uint64 ScratchAlignment = GRHIRayTracingAccelerationStructureAlignment;
+		uint32 ScratchBLASOffset = 0;
+
 		// Iterate all the geometries which need an update
-		for (TMap<FRayTracingGeometry*, EAccelerationStructureBuildMode>::TRangedForIterator Iter = RayTracingGeometriesToUpdate.begin(); Iter != RayTracingGeometriesToUpdate.end(); ++Iter)
+		for (TMap<FRayTracingGeometry*, FRayTracingUpdateInfo>::TRangedForIterator Iter = RayTracingGeometriesToUpdate.begin(); Iter != RayTracingGeometriesToUpdate.end(); ++Iter)
 		{
 			FRayTracingGeometry* RayTracingGeometry = Iter.Key();
+			FRayTracingUpdateInfo& UpdateInfo = Iter.Value();
+
 			FRayTracingGeometryBuildParams BuildParams;
 			BuildParams.Geometry	= RayTracingGeometry->RayTracingGeometryRHI;
-			BuildParams.BuildMode   = Iter.Value();
+			BuildParams.BuildMode   = UpdateInfo.BuildMode;
 			BuildParams.Segments	= RayTracingGeometry->Initializer.Segments;
+			BuildParams.ScratchBuffer = ScratchBuffer;
+			BuildParams.ScratchBufferOffset = ScratchBLASOffset;
+
+			// Update the offset
+			ScratchBLASOffset = Align(ScratchBLASOffset + UpdateInfo.ScratchSize, ScratchAlignment);
 
 			// Make 'Build' 10 times more expensive than 1 'Update' of the BVH
 			uint32 PrimitiveCount = RayTracingGeometry->Initializer.TotalPrimitiveCount;
@@ -1813,11 +1846,14 @@ void FGPUSkinCache::ProcessRayTracingGeometryToUpdate(
 			// Update the new init data
 			RayTracingGeometry.SetInitializer(Initializer);
 
+			// Get the scratch sizes used for build & update
+			SkinCacheEntry->GPUSkin->RayTracingGeometryStructureSize = RHICalcRayTracingGeometrySize(Initializer);
+
 			// Only create RHI object but enqueue actual BLAS creation so they can be accumulated
 			RayTracingGeometry.CreateRayTracingGeometry(ERTAccelerationStructureBuildPriority::Skip);
 
 			// Request the build
-			AddRayTracingGeometryToUpdate(&RayTracingGeometry, EAccelerationStructureBuildMode::Build);
+			AddRayTracingGeometryToUpdate(&RayTracingGeometry, SkinCacheEntry->GPUSkin->RayTracingGeometryStructureSize, EAccelerationStructureBuildMode::Build);
 		}
 		else
 		{
@@ -1828,7 +1864,7 @@ void FGPUSkinCache::ProcessRayTracingGeometryToUpdate(
 				FGPUSkinCache::GetRayTracingSegmentVertexBuffers(*SkinCacheEntry, RayTracingGeometry.Initializer.Segments);
 
 				// Request the update
-				AddRayTracingGeometryToUpdate(&RayTracingGeometry, EAccelerationStructureBuildMode::Update);
+				AddRayTracingGeometryToUpdate(&RayTracingGeometry, SkinCacheEntry->GPUSkin->RayTracingGeometryStructureSize, EAccelerationStructureBuildMode::Update);
 			}
 			else
 			{
diff --git a/Engine/Source/Runtime/Engine/Private/LevelTick.cpp b/Engine/Source/Runtime/Engine/Private/LevelTick.cpp
--- a/Engine/Source/Runtime/Engine/Private/LevelTick.cpp
+++ b/Engine/Source/Runtime/Engine/Private/LevelTick.cpp
@@ -1023,13 +1023,6 @@ void EndSendEndOfFrameUpdatesDrawEvent(FSendAllEndOfFrameUpdates& SendAllEndOfFr
 
 				// Flush any remaining pending resource barriers.
 				GPUSkinCache->TransitionAllToReadable(RHICmdList);
-
-#if RHI_RAYTRACING
-				{
-					SCOPED_GPU_STAT(RHICmdList, GPUSkinCacheRayTracingGeometry);
-					GPUSkinCache->CommitRayTracingGeometryUpdates(RHICmdList);
-				}
-#endif // RHI_RAYTRACING
 			}
 
 			if (ComputeTaskWorkers.Num() > 0)
diff --git a/Engine/Source/Runtime/Engine/Private/SkeletalRenderGPUSkin.h b/Engine/Source/Runtime/Engine/Private/SkeletalRenderGPUSkin.h
--- a/Engine/Source/Runtime/Engine/Private/SkeletalRenderGPUSkin.h
+++ b/Engine/Source/Runtime/Engine/Private/SkeletalRenderGPUSkin.h
@@ -331,6 +331,7 @@ public:
 #if RHI_RAYTRACING
 	/** Geometry for ray tracing. */
 	FRayTracingGeometry RayTracingGeometry;
+	FRayTracingAccelerationStructureSize RayTracingGeometryStructureSize;
 	FRWBuffer RayTracingDynamicVertexBuffer;
 
 	virtual FRayTracingGeometry* GetRayTracingGeometry() { return &RayTracingGeometry; }
diff --git a/Engine/Source/Runtime/Engine/Public/GPUSkinCache.h b/Engine/Source/Runtime/Engine/Public/GPUSkinCache.h
--- a/Engine/Source/Runtime/Engine/Public/GPUSkinCache.h
+++ b/Engine/Source/Runtime/Engine/Public/GPUSkinCache.h
@@ -429,8 +429,9 @@ public:
 	ENGINE_API FRHIShaderResourceView* GetBoneBuffer(uint32 ComponentId, uint32 SectionIndex) const;
 
 #if RHI_RAYTRACING
-	void AddRayTracingGeometryToUpdate(FRayTracingGeometry* InRayTracingGeometry, EAccelerationStructureBuildMode InBuildMode);
-	ENGINE_API void CommitRayTracingGeometryUpdates(FRHICommandListImmediate& RHICmdList);
+	void AddRayTracingGeometryToUpdate(FRayTracingGeometry* InRayTracingGeometry, const FRayTracingAccelerationStructureSize& StructureSize, EAccelerationStructureBuildMode InBuildMode);
+	ENGINE_API uint32 ComputeRayTracingGeometryScratchBufferSize();
+	ENGINE_API void CommitRayTracingGeometryUpdates(FRHICommandListImmediate& RHICmdList, FRHIBuffer* ScratchBuffer);
 	void RemoveRayTracingGeometryUpdate(FRayTracingGeometry* RayTracingGeometry)
 	{
 		if (RayTracingGeometriesToUpdate.Find(RayTracingGeometry) != nullptr)
@@ -460,7 +461,12 @@ protected:
 
 	TSet<FSkinCacheRWBuffer*> BuffersToTransitionToRead;
 #if RHI_RAYTRACING
-	TMap<FRayTracingGeometry*, EAccelerationStructureBuildMode> RayTracingGeometriesToUpdate;
+	struct FRayTracingUpdateInfo
+	{
+		EAccelerationStructureBuildMode BuildMode;
+		uint32 ScratchSize;
+	};
+	TMap<FRayTracingGeometry*, FRayTracingUpdateInfo> RayTracingGeometriesToUpdate;
 	uint64 RayTracingGeometryMemoryPendingRelease = 0;
 #endif // RHI_RAYTRACING
 
diff --git a/Engine/Source/Runtime/RHI/Public/RHIContext.h b/Engine/Source/Runtime/RHI/Public/RHIContext.h
--- a/Engine/Source/Runtime/RHI/Public/RHIContext.h
+++ b/Engine/Source/Runtime/RHI/Public/RHIContext.h
@@ -402,6 +402,11 @@ struct FRayTracingGeometryBuildParams
 	// Optional array of geometry segments that can be used to change per-segment vertex buffers.
 	// Only fields related to vertex buffer are used. If empty, then geometry vertex buffers are not changed.
 	TArrayView<const FRayTracingGeometrySegment> Segments;
+
+	// Optional explicit build scratch buffer. Must be in UAV state.
+	// Scratch buffer will be created automatically if this is not provided.
+	FRHIBuffer* ScratchBuffer = nullptr;
+	uint32 ScratchBufferOffset = 0;
 };
 
 struct FRayTracingSceneBuildParams
diff --git a/Engine/Source/Runtime/Renderer/Private/DeferredShadingRenderer.cpp b/Engine/Source/Runtime/Renderer/Private/DeferredShadingRenderer.cpp
--- a/Engine/Source/Runtime/Renderer/Private/DeferredShadingRenderer.cpp
+++ b/Engine/Source/Runtime/Renderer/Private/DeferredShadingRenderer.cpp
@@ -1430,6 +1430,11 @@ bool FDeferredShadingSceneRenderer::SetupRayTracingPipelineStates(FRHICommandLis
 
 	return true;
 }
+
+BEGIN_SHADER_PARAMETER_STRUCT(FGPUSkinCacheBLASUpdateParams, )
+	RDG_BUFFER_ACCESS(SharedScratchBuffer, ERHIAccess::UAVCompute)
+END_SHADER_PARAMETER_STRUCT()
+
 bool FDeferredShadingSceneRenderer::DispatchRayTracingWorldUpdates(FRDGBuilder& GraphBuilder)
 {
 	bool bAnyRayTracingPassEnabled = false;
@@ -1454,10 +1459,29 @@ bool FDeferredShadingSceneRenderer::DispatchRayTracingWorldUpdates(FRDGBuilder&
 	// top level BVH
 	if (FGPUSkinCache* GPUSkinCache = Scene->GetGPUSkinCache())
 	{
-		AddPass(GraphBuilder, RDG_EVENT_NAME("CommitRayTracingGeometryUpdates"), [this](FRHICommandListImmediate& RHICmdList)
+		// Find out the total BLAS scratch size and allocate RDG buffer which will be transient allocated then
+		uint32 BLASScratchSize = GPUSkinCache->ComputeRayTracingGeometryScratchBufferSize();
+		FRDGBufferRef SharedScratchBuffer = nullptr;
+		if (BLASScratchSize > 0)
 		{
-			Scene->GetGPUSkinCache()->CommitRayTracingGeometryUpdates(RHICmdList);
-		});
+			const uint32 ScratchAlignment = GRHIRayTracingAccelerationStructureAlignment;
+			FRDGBufferDesc ScratchBufferDesc;
+			ScratchBufferDesc.UnderlyingType = FRDGBufferDesc::EUnderlyingType::StructuredBuffer;
+			ScratchBufferDesc.Usage = BUF_UnorderedAccess;
+			ScratchBufferDesc.BytesPerElement = ScratchAlignment;
+			ScratchBufferDesc.NumElements = FMath::DivideAndRoundUp(BLASScratchSize, ScratchAlignment);
+
+			SharedScratchBuffer = GraphBuilder.CreateBuffer(ScratchBufferDesc, TEXT("GPUSkinCache.BLASSharedScratchBuffer"));			
+		}
+
+		FGPUSkinCacheBLASUpdateParams* BLASUpdateParams = GraphBuilder.AllocParameters<FGPUSkinCacheBLASUpdateParams>();
+		BLASUpdateParams->SharedScratchBuffer = SharedScratchBuffer;
+
+		GraphBuilder.AddPass(RDG_EVENT_NAME("CommitRayTracingGeometryUpdates"), BLASUpdateParams, ERDGPassFlags::Compute | ERDGPassFlags::NeverCull,
+			[this, SharedScratchBuffer](FRHICommandListImmediate& RHICmdList)
+			{
+				Scene->GetGPUSkinCache()->CommitRayTracingGeometryUpdates(RHICmdList, SharedScratchBuffer ? SharedScratchBuffer->GetRHI() : nullptr);
+			});
 	}
 
 	GRayTracingGeometryManager.ProcessBuildRequests(GraphBuilder.RHICmdList);
