UE5:  Debug code accidentally checked in (dummy ELLMTag scope), reverting.  Followup to CL 17663476.

#rnx

diff --git a/Engine/Source/Runtime/Engine/Private/SkeletalRender.cpp b/Engine/Source/Runtime/Engine/Private/SkeletalRender.cpp
--- a/Engine/Source/Runtime/Engine/Private/SkeletalRender.cpp
+++ b/Engine/Source/Runtime/Engine/Private/SkeletalRender.cpp
@@ -325,8 +325,6 @@ void UpdateRefToLocalMatricesInner(TArray<FMatrix44f>& ReferenceToLocal, const T
  */
 void UpdateRefToLocalMatrices( TArray<FMatrix44f>& ReferenceToLocal, const USkinnedMeshComponent* InMeshComponent, const FSkeletalMeshRenderData* InSkeletalMeshRenderData, int32 LODIndex, const TArray<FBoneIndexType>* ExtraRequiredBoneIndices )
 {
-	LLM_SCOPE((ELLMTag)((uint32)ELLMTag::ProjectTagStart + 1));
-
 	const USkeletalMesh* const ThisMesh = InMeshComponent->SkeletalMesh;
 	const USkinnedMeshComponent* const MasterComp = InMeshComponent->MasterPoseComponent.Get();
 	const FSkeletalMeshLODRenderData& LOD = InSkeletalMeshRenderData->LODRenderData[LODIndex];
