Disable transient allocation for BLAS Scratch for GPU skin caches by default on PC because of possible GPU crashes during MovieRendering in Editor

diff --git a/Engine/Source/Runtime/Engine/Private/GPUSkinCache.cpp b/Engine/Source/Runtime/Engine/Private/GPUSkinCache.cpp
--- a/Engine/Source/Runtime/Engine/Private/GPUSkinCache.cpp
+++ b/Engine/Source/Runtime/Engine/Private/GPUSkinCache.cpp
@@ -164,6 +164,14 @@ FAutoConsoleVariableRef CVarGPUSkinCacheMemoryLimitForBatchedRayTracingGeometryU
 	ECVF_RenderThreadSafe
 );
 
+static int32 GRayTracingUseTransientForScratch = 0;
+FAutoConsoleVariableRef CVarGPUSkinCacheRayTracingUseTransientForScratch(
+	TEXT("r.SkinCache.RayTracingUseTransientForScratch"),
+	GRayTracingUseTransientForScratch,
+	TEXT("Use Transient memory for BLAS scratch allocation to reduce memory footprint and allocation overhead."),
+	ECVF_RenderThreadSafe
+);
+
 static int32 GMaxRayTracingPrimitivesPerCmdList = -1;
 FAutoConsoleVariableRef CVarGPUSkinCacheMaxRayTracingPrimitivesPerCmdList(
 	TEXT("r.SkinCache.MaxRayTracingPrimitivesPerCmdList"),
@@ -848,7 +856,7 @@ uint32 FGPUSkinCache::ComputeRayTracingGeometryScratchBufferSize()
 	const uint64 ScratchAlignment = GRHIRayTracingAccelerationStructureAlignment;
 	uint32 ScratchBLASSize = 0;
 
-	if (RayTracingGeometriesToUpdate.Num())
+	if (RayTracingGeometriesToUpdate.Num() && GRayTracingUseTransientForScratch > 0)
 	{
 		for (TMap<FRayTracingGeometry*, FRayTracingUpdateInfo>::TRangedForIterator Iter = RayTracingGeometriesToUpdate.begin(); Iter != RayTracingGeometriesToUpdate.end(); ++Iter)
 		{			
