Partial backout of CL 17665451
 - Fixes corrupt BVHs and GPU perf regression

diff --git a/Engine/Source/Runtime/Engine/Private/LevelTick.cpp b/Engine/Source/Runtime/Engine/Private/LevelTick.cpp
--- a/Engine/Source/Runtime/Engine/Private/LevelTick.cpp
+++ b/Engine/Source/Runtime/Engine/Private/LevelTick.cpp
@@ -1023,6 +1023,13 @@ void EndSendEndOfFrameUpdatesDrawEvent(FSendAllEndOfFrameUpdates& SendAllEndOfFr
 
 				// Flush any remaining pending resource barriers.
 				GPUSkinCache->TransitionAllToReadable(RHICmdList);
+
+#if RHI_RAYTRACING
+				{
+					SCOPED_GPU_STAT(RHICmdList, GPUSkinCacheRayTracingGeometry);
+					GPUSkinCache->CommitRayTracingGeometryUpdates(RHICmdList, nullptr);
+				}
+#endif // RHI_RAYTRACING
 			}
 
 			if (ComputeTaskWorkers.Num() > 0)
