Add EGPUSkinCacheGPUMode filtering when retrieving a GPUSkin geometry cache data.

This ensures that the requester get appropriate data if a skel. mesh has both Rasterize & Raytraced entries. If there is no entry with the desired mode, it will fallback on what it available.

diff --git a/Engine/Plugins/Runtime/HairStrands/Source/HairStrandsCore/Private/GroomDebug.cpp b/Engine/Plugins/Runtime/HairStrands/Source/HairStrandsCore/Private/GroomDebug.cpp
--- a/Engine/Plugins/Runtime/HairStrands/Source/HairStrandsCore/Private/GroomDebug.cpp
+++ b/Engine/Plugins/Runtime/HairStrands/Source/HairStrandsCore/Private/GroomDebug.cpp
@@ -80,7 +80,7 @@ static void GetGroomInterpolationData(
 			{
 				if (SkinCache)
 				{
-					CachedGeometry = SkinCache->GetCachedGeometry(SkeletalMeshComponent->ComponentId.PrimIDValue);
+					CachedGeometry = SkinCache->GetCachedGeometry(SkeletalMeshComponent->ComponentId.PrimIDValue, EGPUSkinCacheEntryMode::Raster);
 				}
 
 				if (IsHairStrandsSkinCacheEnable() && CachedGeometry.Sections.Num() == 0)
diff --git a/Engine/Plugins/Runtime/HairStrands/Source/HairStrandsCore/Private/GroomManager.cpp b/Engine/Plugins/Runtime/HairStrands/Source/HairStrandsCore/Private/GroomManager.cpp
--- a/Engine/Plugins/Runtime/HairStrands/Source/HairStrandsCore/Private/GroomManager.cpp
+++ b/Engine/Plugins/Runtime/HairStrands/Source/HairStrandsCore/Private/GroomManager.cpp
@@ -81,7 +81,7 @@ FCachedGeometry GetCacheGeometryForHair(
 		{
 			if (SkinCache)
 			{
-				Out = SkinCache->GetCachedGeometry(SkeletalMeshComponent->ComponentId.PrimIDValue);
+				Out = SkinCache->GetCachedGeometry(SkeletalMeshComponent->ComponentId.PrimIDValue, EGPUSkinCacheEntryMode::Raster);
 			}
 
 			if (IsHairStrandsSkinCacheEnable() && Out.Sections.Num() == 0)
diff --git a/Engine/Source/Runtime/Engine/Private/GPUSkinCache.cpp b/Engine/Source/Runtime/Engine/Private/GPUSkinCache.cpp
--- a/Engine/Source/Runtime/Engine/Private/GPUSkinCache.cpp
+++ b/Engine/Source/Runtime/Engine/Private/GPUSkinCache.cpp
@@ -2338,35 +2338,45 @@ FRHIShaderResourceView* FGPUSkinCache::GetBoneBuffer(uint32 ComponentId, uint32
 	return nullptr;
 }
 
-FCachedGeometry FGPUSkinCache::GetCachedGeometry(uint32 ComponentId) const
+FCachedGeometry FGPUSkinCache::GetCachedGeometry(uint32 ComponentId, EGPUSkinCacheEntryMode Mode) const
 {
-	FCachedGeometry Out;
-	for (FGPUSkinCacheEntry* Entry : Entries)
+	auto FindEntry = [&](EGPUSkinCacheEntryMode InMode, bool bByPassModeCheck, FCachedGeometry& Out) -> bool
 	{
-		if (Entry && Entry->GPUSkin && Entry->GPUSkin->GetComponentId() == ComponentId && Entry->GPUSkin->HaveValidDynamicData())
+		for (FGPUSkinCacheEntry* Entry : Entries)
 		{
-			const FSkeletalMeshRenderData& RenderData = Entry->GPUSkin->GetSkeletalMeshRenderData();
-			const int32 LODIndex = Entry->LOD;
-			if (LODIndex >= 0 && LODIndex < RenderData.LODRenderData.Num())
+			if (Entry && (bByPassModeCheck || Entry->Mode == InMode) && Entry->GPUSkin && Entry->GPUSkin->GetComponentId() == ComponentId && Entry->GPUSkin->HaveValidDynamicData())
 			{
-				const FSkeletalMeshLODRenderData& LODData = RenderData.LODRenderData[LODIndex];
-				const uint32 SectionCount = LODData.RenderSections.Num();
-				for (uint32 SectionIdx=0; SectionIdx< SectionCount;++SectionIdx)
+				const FSkeletalMeshRenderData& RenderData = Entry->GPUSkin->GetSkeletalMeshRenderData();
+				const int32 LODIndex = Entry->LOD;
+				if (LODIndex >= 0 && LODIndex < RenderData.LODRenderData.Num())
 				{
-					FCachedGeometry::Section CachedSection = Entry->GetCachedGeometry(SectionIdx);
-					CachedSection.IndexBuffer		= LODData.MultiSizeIndexContainer.GetIndexBuffer()->GetSRV();
-					CachedSection.TotalIndexCount	= LODData.MultiSizeIndexContainer.GetIndexBuffer()->Num();
-					CachedSection.LODIndex			= LODIndex;
-					CachedSection.UVsChannelOffset	= 0; // Assume that we needs to pair meshes based on UVs 0
-					CachedSection.UVsChannelCount	= LODData.StaticVertexBuffers.StaticMeshVertexBuffer.GetNumTexCoords();
-					Out.Sections.Add(CachedSection);
+					const FSkeletalMeshLODRenderData& LODData = RenderData.LODRenderData[LODIndex];
+					const uint32 SectionCount = LODData.RenderSections.Num();
+					for (uint32 SectionIdx = 0; SectionIdx < SectionCount; ++SectionIdx)
+					{
+						FCachedGeometry::Section CachedSection = Entry->GetCachedGeometry(SectionIdx);
+						CachedSection.IndexBuffer = LODData.MultiSizeIndexContainer.GetIndexBuffer()->GetSRV();
+						CachedSection.TotalIndexCount = LODData.MultiSizeIndexContainer.GetIndexBuffer()->Num();
+						CachedSection.LODIndex = LODIndex;
+						CachedSection.UVsChannelOffset = 0; // Assume that we needs to pair meshes based on UVs 0
+						CachedSection.UVsChannelCount = LODData.StaticVertexBuffers.StaticMeshVertexBuffer.GetNumTexCoords();
+						Out.Sections.Add(CachedSection);
+					}
+					Out.LocalToWorld = FTransform(Entry->GPUSkin->GetTransform());
+					return true;
 				}
-				Out.LocalToWorld = FTransform(Entry->GPUSkin->GetTransform());
-				break;
 			}
 		}
-	}
+		return false;
+	};
 
+	// 1. Try to find a Skin cache entry which matches the requested mode (Raster/Raytracing)
+	// 2. If we can't find an entry with a matching mode, use any mode type
+	FCachedGeometry Out;
+	if (!FindEntry(Mode, false, Out))
+	{
+		FindEntry(Mode, true, Out);
+	}
 	return Out;
 }
 
diff --git a/Engine/Source/Runtime/Engine/Public/GPUSkinCache.h b/Engine/Source/Runtime/Engine/Public/GPUSkinCache.h
--- a/Engine/Source/Runtime/Engine/Public/GPUSkinCache.h
+++ b/Engine/Source/Runtime/Engine/Public/GPUSkinCache.h
@@ -157,7 +157,7 @@ public:
 	ENGINE_API FGPUSkinCache(ERHIFeatureLevel::Type InFeatureLevel, bool bInRequiresMemoryLimit, UWorld* InWorld);
 	ENGINE_API ~FGPUSkinCache();
 
-	ENGINE_API FCachedGeometry GetCachedGeometry(uint32 ComponentId) const;
+	ENGINE_API FCachedGeometry GetCachedGeometry(uint32 ComponentId, EGPUSkinCacheEntryMode InMode) const;
 	FCachedGeometry::Section GetCachedGeometry(FGPUSkinCacheEntry* InOutEntry, uint32 SectionId);
 	void UpdateSkinWeightBuffer(FGPUSkinCacheEntry* Entry);
 
