Make sure to force a GC pass when unloading a Data Layer.

diff --git a/Engine/Source/Runtime/Engine/Classes/Engine/CoreSettings.h b/Engine/Source/Runtime/Engine/Classes/Engine/CoreSettings.h
--- a/Engine/Source/Runtime/Engine/Classes/Engine/CoreSettings.h
+++ b/Engine/Source/Runtime/Engine/Classes/Engine/CoreSettings.h
@@ -146,8 +146,10 @@ extern ENGINE_API int32 GLevelStreamingComponentsUnregistrationGranularity;
 extern ENGINE_API float GLevelStreamingUnregisterComponentsTimeLimit;
 /** Whether to force a GC after levels are streamed out to instantly reclaim the memory at the expensive of a hitch. */
 extern ENGINE_API int32 GLevelStreamingForceGCAfterLevelStreamedOut;
-/** Whether to repeatedly kick off incremental GC when there are levels still waiting to be purged. */
+/** Whether to kick off incremental GC when there are over the specified amount of levels still waiting to be purged. */
 extern ENGINE_API int32 GLevelStreamingContinuouslyIncrementalGCWhileLevelsPendingPurge;
+/** Whether to override GLevelStreamingContinuouslyIncrementalGCWhileLevelsPendingPurge temporarily until at least one level needs to be purged.  */
+extern ENGINE_API int32 GLevelStreamingContinuouslyIncrementalGCWhileLevelsPendingPurgeOverride;
 /** Enables level streaming requests while async loading (of anything) while the match is already in progress and no loading screen is up. */
 extern ENGINE_API int32 GLevelStreamingAllowLevelRequestsWhileAsyncLoadingInMatch;
 /** When we're already loading this many levels and actively in match, don't allow any more requests until one of those completes.  Set to zero to disable. */
diff --git a/Engine/Source/Runtime/Engine/Private/CoreSettings.cpp b/Engine/Source/Runtime/Engine/Private/CoreSettings.cpp
--- a/Engine/Source/Runtime/Engine/Private/CoreSettings.cpp
+++ b/Engine/Source/Runtime/Engine/Private/CoreSettings.cpp
@@ -22,6 +22,7 @@ int32 GLevelStreamingAddPrimitiveGranularity = 120;
 int32 GLevelStreamingComponentsUnregistrationGranularity = 5;
 int32 GLevelStreamingForceGCAfterLevelStreamedOut = 1;
 int32 GLevelStreamingContinuouslyIncrementalGCWhileLevelsPendingPurge = 1;
+int32 GLevelStreamingContinuouslyIncrementalGCWhileLevelsPendingPurgeOverride = 0;
 int32 GLevelStreamingAllowLevelRequestsWhileAsyncLoadingInMatch = 1;
 int32 GLevelStreamingMaxLevelRequestsAtOnceWhileInMatch = 0;
 int32 GLevelStreamingForceVerifyLevelsGotRemovedByGC = 0;
diff --git a/Engine/Source/Runtime/Engine/Private/LevelStreaming.cpp b/Engine/Source/Runtime/Engine/Private/LevelStreaming.cpp
--- a/Engine/Source/Runtime/Engine/Private/LevelStreaming.cpp
+++ b/Engine/Source/Runtime/Engine/Private/LevelStreaming.cpp
@@ -538,20 +538,6 @@ void ULevelStreaming::UpdateStreamingState(bool& bOutUpdateAgain, bool& bOutRede
 
 	auto UpdateStreamingState_RequestLevel = [&]()
 	{
-		if (GLevelStreamingContinuouslyIncrementalGCWhileLevelsPendingPurge)
-		{
-			// Figure out whether there are any levels we haven't collected garbage yet.
-			const bool bAreLevelsPendingPurge = FLevelStreamingGCHelper::GetNumLevelsPendingPurge() >= GLevelStreamingContinuouslyIncrementalGCWhileLevelsPendingPurge;
-
-			// Request a 'soft' GC if there are levels pending purge and there are levels to be loaded. In the case of a blocking
-			// load this is going to guarantee GC firing first thing afterwards and otherwise it is going to sneak in right before
-			// kicking off the async load.
-			if (bAreLevelsPendingPurge)
-			{
-				GEngine->ForceGarbageCollection(false);
-			}
-		}
-
 		bool bBlockOnLoad = (bShouldBlockOnLoad || ShouldBeAlwaysLoaded());
 		const bool bAllowLevelLoadRequests = (bBlockOnLoad || World->AllowLevelLoadRequests());
 		bBlockOnLoad |= (!GUseBackgroundLevelStreaming || !World->IsGameWorld());
diff --git a/Engine/Source/Runtime/Engine/Private/World.cpp b/Engine/Source/Runtime/Engine/Private/World.cpp
--- a/Engine/Source/Runtime/Engine/Private/World.cpp
+++ b/Engine/Source/Runtime/Engine/Private/World.cpp
@@ -3864,6 +3864,21 @@ void UWorld::UpdateLevelStreaming()
 
 	StreamingLevelsToConsider.EndConsideration();
 
+	if (GLevelStreamingContinuouslyIncrementalGCWhileLevelsPendingPurge)
+	{
+		// Figure out whether there are any levels we haven't collected garbage yet.
+		const int32 LevelStreamingContinuouslyIncrementalGCWhileLevelsPendingPurge = GLevelStreamingContinuouslyIncrementalGCWhileLevelsPendingPurgeOverride ? 1 : GLevelStreamingContinuouslyIncrementalGCWhileLevelsPendingPurge;
+		const bool bAreLevelsPendingPurge = FLevelStreamingGCHelper::GetNumLevelsPendingPurge() >= LevelStreamingContinuouslyIncrementalGCWhileLevelsPendingPurge;
+
+		// Request a 'soft' GC if there are levels pending purge and there are levels to be loaded. In the case of a blocking
+		// load this is going to guarantee GC firing first thing afterwards and otherwise it is going to sneak in right before
+		// kicking off the async load.
+		if (bAreLevelsPendingPurge)
+		{
+			GEngine->ForceGarbageCollection(false);
+			GLevelStreamingContinuouslyIncrementalGCWhileLevelsPendingPurgeOverride = 0;
+		}
+	}
 
 	// In case more levels has been requested to unload, force GC on next tick 
 	if (GLevelStreamingForceGCAfterLevelStreamedOut != 0)
diff --git a/Engine/Source/Runtime/Engine/Private/WorldPartition/DataLayer/WorldDataLayers.cpp b/Engine/Source/Runtime/Engine/Private/WorldPartition/DataLayer/WorldDataLayers.cpp
--- a/Engine/Source/Runtime/Engine/Private/WorldPartition/DataLayer/WorldDataLayers.cpp
+++ b/Engine/Source/Runtime/Engine/Private/WorldPartition/DataLayer/WorldDataLayers.cpp
@@ -8,6 +8,7 @@
 #include "WorldPartition/DataLayer/DataLayer.h"
 #include "WorldPartition/DataLayer/DataLayerSubsystem.h"
 #include "EngineUtils.h"
+#include "Engine/CoreSettings.h"
 #include "Net/UnrealNetwork.h"
 #include "WorldPartition/WorldPartition.h"
 #if WITH_EDITOR
@@ -170,6 +171,10 @@ PRAGMA_ENABLE_DEPRECATION_WARNINGS
 			{
 				ActiveDataLayerNames.Add(InDataLayer.Name);
 			}
+			else if (InState == EDataLayerState::Unloaded)
+			{
+				GLevelStreamingContinuouslyIncrementalGCWhileLevelsPendingPurgeOverride = 1;
+			}
 
 			// Update Replicated Properties
 			RepActiveDataLayerNames = ActiveDataLayerNames.Array();
