Fix Show.VisualizeInstanceUpdates for ISMC proxy meshes.

- Having that flag turned on should mark the prim for dynamic relevance so it gets routed through GetDynamicMeshElements.

diff --git a/Engine/Source/Runtime/Engine/Private/StaticMeshRender.cpp b/Engine/Source/Runtime/Engine/Private/StaticMeshRender.cpp
--- a/Engine/Source/Runtime/Engine/Private/StaticMeshRender.cpp
+++ b/Engine/Source/Runtime/Engine/Private/StaticMeshRender.cpp
@@ -1851,6 +1851,7 @@ FPrimitiveViewRelevance FStaticMeshSceneProxy::GetViewRelevance(const FSceneView
 		View->Family->EngineShowFlags.Collision ||
 		bInCollisionView ||
 		View->Family->EngineShowFlags.Bounds ||
+		View->Family->EngineShowFlags.VisualizeInstanceUpdates ||
 #endif
 #if WITH_EDITOR
 		(IsSelected() && View->Family->EngineShowFlags.VertexColors) ||
