Discard source data from morph targets when cooking.

diff --git a/Engine/Source/Runtime/Core/Public/UObject/UE5PrivateFrostyStreamObjectVersion.h b/Engine/Source/Runtime/Core/Public/UObject/UE5PrivateFrostyStreamObjectVersion.h
--- a/Engine/Source/Runtime/Core/Public/UObject/UE5PrivateFrostyStreamObjectVersion.h
+++ b/Engine/Source/Runtime/Core/Public/UObject/UE5PrivateFrostyStreamObjectVersion.h
@@ -33,6 +33,9 @@ struct CORE_API FUE5PrivateFrostyStreamObjectVersion
 		// Serializes the Morph Target render data for cooked platforms and the DDC
 		SerializeSkeletalMeshMorphTargetRenderData,
 
+		// Strip the Morph Target source data for cooked builds
+		StripMorphTargetSourceDataForCookedBuilds,
+
 		// -----<new versions can be added above this line>-------------------------------------------------
 		VersionPlusOne,
 		LatestVersion = VersionPlusOne - 1
diff --git a/Engine/Source/Runtime/Engine/Classes/Animation/MorphTarget.h b/Engine/Source/Runtime/Engine/Classes/Animation/MorphTarget.h
--- a/Engine/Source/Runtime/Engine/Classes/Animation/MorphTarget.h
+++ b/Engine/Source/Runtime/Engine/Classes/Animation/MorphTarget.h
@@ -8,7 +8,9 @@
 #include "PackedNormal.h"
 #include "UObject/EditorObjectVersion.h"
 #include "UObject/FortniteMainBranchObjectVersion.h"
+#include "UObject/UE5PrivateFrostyStreamObjectVersion.h"
 #include "Serialization/MemoryArchive.h"
+
 #include "MorphTarget.generated.h"
 
 class USkeletalMesh;
@@ -76,6 +78,7 @@ struct FMorphTargetLODModel
 	{
 		Ar.UsingCustomVersion(FEditorObjectVersion::GUID);
 		Ar.UsingCustomVersion(FFortniteMainBranchObjectVersion::GUID);
+		Ar.UsingCustomVersion(FUE5PrivateFrostyStreamObjectVersion::GUID);
 
 		if (!Ar.IsObjectReferenceCollector())
 		{
@@ -91,16 +94,32 @@ struct FMorphTargetLODModel
 			}
 			else
 			{
-				Ar << M.Vertices << M.NumBaseMeshVerts << M.SectionIndices << M.bGeneratedByEngine;
+				bool bIsCooked = false;
+				if (Ar.IsPersistent() && (Ar.CustomVer(FUE5PrivateFrostyStreamObjectVersion::GUID) >= FUE5PrivateFrostyStreamObjectVersion::StripMorphTargetSourceDataForCookedBuilds))
+				{
+					bIsCooked = Ar.IsCooking();
+					Ar << bIsCooked;
+				}
+
+				if (bIsCooked)
+				{
+					M.NumVertices = M.Vertices.Num();
+					Ar << M.NumVertices;
+				}
+				else
+				{
+					Ar << M.Vertices;
+
+					if (Ar.IsLoading())
+					{
+						M.NumVertices = M.Vertices.Num();
+					}
+				}
+
+				Ar << M.NumBaseMeshVerts << M.SectionIndices << M.bGeneratedByEngine;
 			}
 		}
 
-		// Cache Vertices array size, no need to serialize NumVertices value.
-		if (Ar.IsLoading())
-		{
-			M.NumVertices = M.Vertices.Num();
-		}
-
 		return Ar;
 	}
 
@@ -119,8 +138,11 @@ struct FMorphTargetLODModel
 	{
 		// Vertices array can be discarded after render data is initialized and if CPU data is no longer required,
 		// but cache the array size so that UMorphTarget::HasValidData() can still query if morph target is valid. 
-		NumVertices = Vertices.Num();
-		Vertices.Empty();
+		if (Vertices.Num())
+		{
+			NumVertices = Vertices.Num();
+			Vertices.Empty();
+		}
 	}
 };
 
diff --git a/Engine/Source/Runtime/Engine/Private/MorphTools.cpp b/Engine/Source/Runtime/Engine/Private/MorphTools.cpp
--- a/Engine/Source/Runtime/Engine/Private/MorphTools.cpp
+++ b/Engine/Source/Runtime/Engine/Private/MorphTools.cpp
@@ -53,6 +53,13 @@ bool UMorphTarget::HasValidData() const
 		{
 			return true;
 		}
+#if !WITH_EDITOR
+		// In cooked builds, Model.Vertices is stripped but Model.NumVertices is valid
+		else if (Model.NumVertices > 0)
+		{
+			return true;
+		}
+#endif
 	}
 
 	return false;
diff --git a/Engine/Source/Runtime/Engine/Private/SkeletalMeshRenderData.cpp b/Engine/Source/Runtime/Engine/Private/SkeletalMeshRenderData.cpp
--- a/Engine/Source/Runtime/Engine/Private/SkeletalMeshRenderData.cpp
+++ b/Engine/Source/Runtime/Engine/Private/SkeletalMeshRenderData.cpp
@@ -731,8 +731,9 @@ void FSkeletalMeshRenderData::InitResources(bool bNeedsVertexColors, TArray<UMor
 #if !WITH_EDITOR
 		for (UMorphTarget* MorphTarget : InMorphTargets)
 		{
-			// A CPU copy of the morph deltas has been  made so it is safe to 
-			// discard the original data.  Keep CPU buffers when in the editor.
+			// A CPU copy of the morph deltas has been made so it is safe to 
+			// discard the original data. Keep CPU buffers when in the editor.
+			// Note: this data is stripped at cook time so it's probably empty.
 			MorphTarget->DiscardVertexData();
 		}
 #endif
