Adding InstancedStaticMeshComponent CSV Profiler category to track various ISM related stats.

- NumInstancesUpdated
- NumBytesSubmitted
- NumBytesUpdated

diff --git a/Engine/Source/Runtime/Engine/Private/InstancedStaticMesh.cpp b/Engine/Source/Runtime/Engine/Private/InstancedStaticMesh.cpp
--- a/Engine/Source/Runtime/Engine/Private/InstancedStaticMesh.cpp
+++ b/Engine/Source/Runtime/Engine/Private/InstancedStaticMesh.cpp
@@ -57,6 +57,7 @@
 
 IMPLEMENT_TYPE_LAYOUT(FInstancedStaticMeshVertexFactoryShaderParameters);
 
+CSV_DEFINE_CATEGORY_MODULE(ENGINE_API, InstancedStaticMeshComponent, true);
 
 const int32 InstancedStaticMeshMaxTexCoord = 8;
 static const int32 MaxSimulatedInstances = 256;
@@ -3437,6 +3438,10 @@ bool UInstancedStaticMeshComponent::UpdateInstances(
 
 	Modify();
 
+#if (CSV_PROFILER)
+	int32 TotalSizeUpdateBytes = 0;
+#endif
+
 	// Need to empty the command buffer if it already has values.
 	// If this function is called multiple times the last set of values will take
 	// precedence.  Due to the way MarkRenderTransformDirty() works the primitive will
@@ -3474,6 +3479,11 @@ bool UInstancedStaticMeshComponent::UpdateInstances(
 
 				// Record in a command buffer for future use.
 				InstanceUpdateCmdBuffer.UpdateInstance(InstanceIndex, PerInstanceSMData[InstanceIndex].Transform, PerInstancePrevTransform[InstanceIndex]);
+
+#if (CSV_PROFILER)
+				// We are updating a current and previous transform.
+				TotalSizeUpdateBytes += (sizeof(FTransform) + sizeof(FTransform));
+#endif
 			}
 
 			// Did the custom data actually change?
@@ -3494,6 +3504,10 @@ bool UInstancedStaticMeshComponent::UpdateInstances(
 					// Record in a command buffer for future use.
 					InstanceUpdateCmdBuffer.SetCustomData(InstanceIndex, NewCustomFloatData);
 
+#if (CSV_PROFILER)
+					TotalSizeUpdateBytes += NumCustomDataFloats * sizeof(float);
+#endif
+
 					break;
 				}
 			}
@@ -3576,6 +3590,20 @@ bool UInstancedStaticMeshComponent::UpdateInstances(
 
 	MarkRenderInstancesDirty();
 
+#if (CSV_PROFILER)
+	const int32 TotalSizeBytes = (UpdateInstanceTransforms.Num() * UpdateInstanceTransforms.GetTypeSize()) +
+								 (UpdateInstancePreviousTransforms.Num() * UpdateInstancePreviousTransforms.GetTypeSize()) +
+								 (CustomFloatData.Num() * CustomFloatData.GetTypeSize());
+
+	static const FName NumInstancesUpdated(TEXT("NumInstancesUpdated"));
+	static const FName NumBytesSubmitted(TEXT("NumBytesSubmitted"));
+	static const FName NumBytesUpdated(TEXT("NumBytesUpdated"));
+
+	FCsvProfiler::Get()->RecordCustomStat(NumInstancesUpdated, CSV_CATEGORY_INDEX(InstancedStaticMeshComponent), PerInstanceSMData.Num(), ECsvCustomStatOp::Accumulate);
+	FCsvProfiler::Get()->RecordCustomStat(NumBytesSubmitted, CSV_CATEGORY_INDEX(InstancedStaticMeshComponent), TotalSizeBytes, ECsvCustomStatOp::Accumulate);
+	FCsvProfiler::Get()->RecordCustomStat(NumBytesUpdated, CSV_CATEGORY_INDEX(InstancedStaticMeshComponent), TotalSizeUpdateBytes, ECsvCustomStatOp::Accumulate);
+#endif
+
 	return true;
 }
 
