Add stat for virtual texture IO requests completed.
Stat is an approximate guideline since it doesn't take into account FileCache or platform compression.

diff --git a/Engine/Source/Runtime/Engine/Private/VT/UploadingVirtualTexture.cpp b/Engine/Source/Runtime/Engine/Private/VT/UploadingVirtualTexture.cpp
--- a/Engine/Source/Runtime/Engine/Private/VT/UploadingVirtualTexture.cpp
+++ b/Engine/Source/Runtime/Engine/Private/VT/UploadingVirtualTexture.cpp
@@ -15,6 +15,7 @@ DECLARE_MEMORY_STAT(TEXT("Total Header Size"), STAT_TotalHeaderSize, STATGROUP_V
 DECLARE_DWORD_ACCUMULATOR_STAT(TEXT("Total Disk Size (MB)"), STAT_TotalDiskSize, STATGROUP_VirtualTextureMemory);
 DECLARE_DWORD_ACCUMULATOR_STAT(TEXT("Num Tile Headers"), STAT_NumTileHeaders, STATGROUP_VirtualTextureMemory);
 DECLARE_DWORD_ACCUMULATOR_STAT(TEXT("Num Codecs"), STAT_NumCodecs, STATGROUP_VirtualTextureMemory);
+DECLARE_FLOAT_COUNTER_STAT(TEXT("IO Requests Completed (MB)"), STAT_IORequestsComplete, STATGROUP_VirtualTexturing);
 
 static TAutoConsoleVariable<int32> CVarVTCodecAgeThreshold(
 	TEXT("r.VT.CodecAgeThreshold"),
@@ -454,6 +455,10 @@ FVTDataAndStatus FUploadingVirtualTexture::ReadData(FGraphEventArray& OutComplet
 	{
 		return EVTRequestPageStatus::Saturated;
 	}
+
+	const float SizeMB = (float)Size / (float)(1024 * 1024);
+	INC_FLOAT_STAT_BY(STAT_IORequestsComplete, SizeMB);
+	
 	return FVTDataAndStatus(EVTRequestPageStatus::Pending, ReadData);
 }
 
