Support always caching local bounds instead of world-space bounds that are transformed.
- reduces local bounds significantly on large rotated meshes, improving VSM culling.

diff --git a/Engine/Source/Runtime/Engine/Classes/Components/SkinnedMeshComponent.h b/Engine/Source/Runtime/Engine/Classes/Components/SkinnedMeshComponent.h
--- a/Engine/Source/Runtime/Engine/Classes/Components/SkinnedMeshComponent.h
+++ b/Engine/Source/Runtime/Engine/Classes/Components/SkinnedMeshComponent.h
@@ -599,6 +599,8 @@ protected:
 	/** true when CachedLocalBounds is up to date. */
 	UPROPERTY(Transient)
 	mutable uint8 bCachedLocalBoundsUpToDate:1;
+	UPROPERTY(Transient)
+	mutable uint8 bCachedWorldSpaceBoundsUpToDate:1;
 
 	/** Whether we have updated bone visibility this tick */
 	uint8 bBoneVisibilityDirty:1;
@@ -1206,8 +1208,15 @@ public:
 
 	FBoxSphereBounds GetCachedLocalBounds()  const
 	{ 
-		ensure(bCachedLocalBoundsUpToDate);
-		return CachedWorldSpaceBounds.TransformBy(CachedWorldToLocalTransform);
+		ensure(bCachedLocalBoundsUpToDate || bCachedWorldSpaceBoundsUpToDate);
+		if (bCachedWorldSpaceBoundsUpToDate)
+		{
+			return CachedWorldOrLocalSpaceBounds.TransformBy(CachedWorldToLocalTransform);
+		}
+		else
+		{
+			return CachedWorldOrLocalSpaceBounds;
+		}
 	} 
 
 	/**
@@ -1238,9 +1247,9 @@ protected:
 	virtual bool AllocateTransformData();
 	virtual void DeallocateTransformData();
 
-	/** Bounds cached, so they're computed just once. */
+	/** Bounds cached, so they're computed just once, either in local or worldspace depending on cvar 'a.CacheLocalSpaceBounds'. */
 	UPROPERTY(Transient)
-	mutable FBoxSphereBounds CachedWorldSpaceBounds;
+	mutable FBoxSphereBounds CachedWorldOrLocalSpaceBounds;
 	UPROPERTY(Transient)
 	mutable FMatrix CachedWorldToLocalTransform;
 
diff --git a/Engine/Source/Runtime/Engine/Private/Components/SkeletalMeshComponent.cpp b/Engine/Source/Runtime/Engine/Private/Components/SkeletalMeshComponent.cpp
--- a/Engine/Source/Runtime/Engine/Private/Components/SkeletalMeshComponent.cpp
+++ b/Engine/Source/Runtime/Engine/Private/Components/SkeletalMeshComponent.cpp
@@ -65,6 +65,12 @@ static TAutoConsoleVariable<float> CVarStallParallelAnimation(
 	TEXT("Sleep for the given time in each parallel animation task. Time is given in ms. This is a debug option used for critical path analysis and forcing a change in the critical path."));
 
 
+static TAutoConsoleVariable<int32> CVarCacheLocalSpaceBounds(
+	TEXT("a.CacheLocalSpaceBounds"),
+	1,
+	TEXT("If 1 (default) local-space bounds are calculated and cached, otherwise worldspace bounds are built and cached (and inverse transformed to produce local bounds)."));
+
+
 DECLARE_CYCLE_STAT_EXTERN(TEXT("Anim Instance Spawn Time"), STAT_AnimSpawnTime, STATGROUP_Anim, );
 DEFINE_STAT(STAT_AnimSpawnTime);
 DEFINE_STAT(STAT_PostAnimEvaluation);
@@ -2781,17 +2787,26 @@ FBoxSphereBounds USkeletalMeshComponent::CalcBounds(const FTransform& LocalToWor
 	// fixme laurent - extend concept of LocalBounds to all SceneComponent
 	// as rendered calls CalcBounds*() directly in FScene::UpdatePrimitiveTransform, which is pretty expensive for SkelMeshes.
 	// No need to calculated that again, just use cached local bounds.
-	if (bCachedLocalBoundsUpToDate)
+	if (bCachedWorldSpaceBoundsUpToDate || bCachedLocalBoundsUpToDate)
 	{
+		FBoxSphereBounds Result;
+		if (bCachedLocalBoundsUpToDate)
+		{
+			Result = CachedWorldOrLocalSpaceBounds.TransformBy(LocalToWorld);
+		}
+		else
+		{
+			Result = CachedWorldOrLocalSpaceBounds.TransformBy(CachedWorldToLocalTransform * LocalToWorld.ToMatrixWithScale());
+		}
+
 		if (bIncludeComponentLocationIntoBounds)
 		{
 			const FVector ComponentLocation = GetComponentLocation();
-			return CachedWorldSpaceBounds.TransformBy(CachedWorldToLocalTransform * LocalToWorld.ToMatrixWithScale()) 
-					+ FBoxSphereBounds(ComponentLocation, FVector(1.0f), 1.0f);
+			return Result + FBoxSphereBounds(ComponentLocation, FVector(1.0f), 1.0f);
 		}
 		else
 		{
-			return CachedWorldSpaceBounds.TransformBy(CachedWorldToLocalTransform * LocalToWorld.ToMatrixWithScale());
+			return Result;
 		}
 	}
 	// Calculate new bounds
@@ -2814,7 +2829,11 @@ FBoxSphereBounds USkeletalMeshComponent::CalcBounds(const FTransform& LocalToWor
 			}
 		}
 
-		FBoxSphereBounds NewBounds = CalcMeshBound(RootBoneOffset, bHasValidBodies, LocalToWorld);
+		const bool bCacheLocalSpaceBounds = CVarCacheLocalSpaceBounds.GetValueOnGameThread() != 0;
+		
+		const FTransform CachedBoundsTransform = bCacheLocalSpaceBounds ? FTransform::Identity : LocalToWorld;
+
+		FBoxSphereBounds NewBounds = CalcMeshBound(RootBoneOffset, bHasValidBodies, CachedBoundsTransform);
 
 		if (bIncludeComponentLocationIntoBounds)
 		{
@@ -2823,14 +2842,24 @@ FBoxSphereBounds USkeletalMeshComponent::CalcBounds(const FTransform& LocalToWor
 		}
 
 #if WITH_APEX_CLOTHING || WITH_CHAOS_CLOTHING
-		AddClothingBounds(NewBounds, LocalToWorld);
+		AddClothingBounds(NewBounds, CachedBoundsTransform);
 #endif// #if WITH_APEX_CLOTHING || WITH_CHAOS_CLOTHING
 
-		bCachedLocalBoundsUpToDate = true;
-		CachedWorldSpaceBounds = NewBounds;
-		CachedWorldToLocalTransform = LocalToWorld.ToInverseMatrixWithScale();
+		CachedWorldOrLocalSpaceBounds = NewBounds;
+		bCachedLocalBoundsUpToDate = bCacheLocalSpaceBounds;
+		bCachedWorldSpaceBoundsUpToDate = !bCacheLocalSpaceBounds;
+
+		if (bCacheLocalSpaceBounds)
+		{ 
+			CachedWorldToLocalTransform.SetIdentity();
+			return NewBounds.TransformBy(LocalToWorld);
+		}
+		else
+		{
+			CachedWorldToLocalTransform = LocalToWorld.ToInverseMatrixWithScale();
+			return NewBounds;
+		}
 
-		return NewBounds;
 	}
 }
 
diff --git a/Engine/Source/Runtime/Engine/Private/Components/SkinnedMeshComponent.cpp b/Engine/Source/Runtime/Engine/Private/Components/SkinnedMeshComponent.cpp
--- a/Engine/Source/Runtime/Engine/Private/Components/SkinnedMeshComponent.cpp
+++ b/Engine/Source/Runtime/Engine/Private/Components/SkinnedMeshComponent.cpp
@@ -1861,13 +1861,14 @@ void USkinnedMeshComponent::RemoveSlavePoseComponent(USkinnedMeshComponent* Skin
 void USkinnedMeshComponent::InvalidateCachedBounds()
 {
 	bCachedLocalBoundsUpToDate = false;
-
+	bCachedWorldSpaceBoundsUpToDate = false;
 	// Also invalidate all slave components.
 	for (const TWeakObjectPtr<USkinnedMeshComponent>& SkinnedMeshComp : SlavePoseComponents)
 	{
 		if (USkinnedMeshComponent* SkinnedMeshCompPtr = SkinnedMeshComp.Get())
 		{
 			SkinnedMeshCompPtr->bCachedLocalBoundsUpToDate = false;
+			SkinnedMeshCompPtr->bCachedWorldSpaceBoundsUpToDate = false;
 		}
 	}
 
