Defer BLAS builds from GPUSkinCache to DynamicGeometryCollection when using WPO
- Objects using WorldPositionOffset have their list of segments overwritten in FRayTracingDynamicGeometryCollection::AddDynamicMeshBatchForGeometryUpdate.
- This causes an issue for objects using both GPUSkinCache and WPO when the BLAS commands from GPUSkinCache are committed in DispatchRayTracingWorldUpdates since the vertex buffers generated by FRayTracingDynamicGeometryCollection are not yet ready.
- In order to prevent this issue, this CL moves the responsibility of creating BLAS of objects using both GPUSkinCache and WPO to FRayTracingDynamicGeometryCollection (similar to how BLAS updates are currently handled).
- This allows moving all GPUSkinCache BLAS updates to DispatchRayTracingWorldUpdates, enabling access to main render graph and bundling them with the rest of BLAS operations in async.
Also reenable optimization from CL 17665451

diff --git a/Engine/Source/Runtime/Engine/Private/GPUSkinCache.cpp b/Engine/Source/Runtime/Engine/Private/GPUSkinCache.cpp
--- a/Engine/Source/Runtime/Engine/Private/GPUSkinCache.cpp
+++ b/Engine/Source/Runtime/Engine/Private/GPUSkinCache.cpp
@@ -1857,27 +1857,27 @@ void FGPUSkinCache::ProcessRayTracingGeometryToUpdate(
 			// Get the scratch sizes used for build & update
 			SkinCacheEntry->GPUSkin->RayTracingGeometryStructureSize = RHICalcRayTracingGeometrySize(Initializer);
 
-			// Only create RHI object but enqueue actual BLAS creation so they can be accumulated
-			RayTracingGeometry.CreateRayTracingGeometry(ERTAccelerationStructureBuildPriority::Skip);
+			if (!bAnySegmentUsesWorldPositionOffset)
+			{
+				// Only create RHI object but enqueue actual BLAS creation so they can be accumulated
+				RayTracingGeometry.CreateRayTracingGeometry(ERTAccelerationStructureBuildPriority::Skip);
+			}
+		}
+		else if (!bAnySegmentUsesWorldPositionOffset)
+		{
+			// Refit BLAS with new vertex buffer data
+			FGPUSkinCache::GetRayTracingSegmentVertexBuffers(*SkinCacheEntry, RayTracingGeometry.Initializer.Segments);
+		}
 
-			// Request the build
-			AddRayTracingGeometryToUpdate(&RayTracingGeometry, SkinCacheEntry->GPUSkin->RayTracingGeometryStructureSize, EAccelerationStructureBuildMode::Build);
+		// If we are not using world position offset in material, handle BLAS build/refit here
+		if (!bAnySegmentUsesWorldPositionOffset)
+		{
+			EAccelerationStructureBuildMode BuildMode = bRequireRecreatingRayTracingGeometry ? EAccelerationStructureBuildMode::Build : EAccelerationStructureBuildMode::Update;
+			AddRayTracingGeometryToUpdate(&RayTracingGeometry, SkinCacheEntry->GPUSkin->RayTracingGeometryStructureSize, BuildMode);
 		}
 		else
 		{
-			// If we are not using world position offset in material, handle BLAS refit here
-			if (!bAnySegmentUsesWorldPositionOffset)
-			{
-				// Refit BLAS with new vertex buffer data
-				FGPUSkinCache::GetRayTracingSegmentVertexBuffers(*SkinCacheEntry, RayTracingGeometry.Initializer.Segments);
-
-				// Request the update
-				AddRayTracingGeometryToUpdate(&RayTracingGeometry, SkinCacheEntry->GPUSkin->RayTracingGeometryStructureSize, EAccelerationStructureBuildMode::Update);
-			}
-			else
-			{
-				// Otherwise, we will run the dynamic ray tracing geometry path, i.e. runnning VSinCS and refit geometry there, so do nothing here
-			}
+			// Otherwise, we will run the dynamic ray tracing geometry path, i.e. runnning VSinCS and build/refit geometry there, so do nothing here
 		}
 	}
 }
diff --git a/Engine/Source/Runtime/Engine/Private/LevelTick.cpp b/Engine/Source/Runtime/Engine/Private/LevelTick.cpp
--- a/Engine/Source/Runtime/Engine/Private/LevelTick.cpp
+++ b/Engine/Source/Runtime/Engine/Private/LevelTick.cpp
@@ -1024,12 +1024,6 @@ void EndSendEndOfFrameUpdatesDrawEvent(FSendAllEndOfFrameUpdates& SendAllEndOfFr
 				// Flush any remaining pending resource barriers.
 				GPUSkinCache->TransitionAllToReadable(RHICmdList);
 
-#if RHI_RAYTRACING
-				{
-					SCOPED_GPU_STAT(RHICmdList, GPUSkinCacheRayTracingGeometry);
-					GPUSkinCache->CommitRayTracingGeometryUpdates(RHICmdList, nullptr);
-				}
-#endif // RHI_RAYTRACING
 			}
 
 			if (ComputeTaskWorkers.Num() > 0)
diff --git a/Engine/Source/Runtime/Engine/Private/SkeletalMesh.cpp b/Engine/Source/Runtime/Engine/Private/SkeletalMesh.cpp
--- a/Engine/Source/Runtime/Engine/Private/SkeletalMesh.cpp
+++ b/Engine/Source/Runtime/Engine/Private/SkeletalMesh.cpp
@@ -6528,8 +6528,7 @@ void FSkeletalMeshSceneProxy::GetDynamicRayTracingInstances(FRayTracingMaterialG
 	// GetRayTracingGeometry()->IsInitialized() is checked as a workaround for UE-92634. FSkeletalMeshSceneProxy's resources may have already been released, but proxy has not removed yet)
 	if (MeshObject->GetRayTracingGeometry() && MeshObject->GetRayTracingGeometry()->IsInitialized() && MeshObject->GetSkinCacheEntryForRayTracing())
 	{
-		// #dxr: the only case where RayTracingGeometryRHI is invalid is the very first frame - if that's not the case we have a bug somewhere else
-		if (MeshObject->GetRayTracingGeometry()->RayTracingGeometryRHI.IsValid())
+		checkf(MeshObject->GetRayTracingGeometry()->RayTracingGeometryRHI.IsValid() || bAnySegmentUsesWorldPositionOffset, TEXT("RayTracingGeometryRHI should be valid at this point unless object uses WorldPositionOffset"));
 		{
 			check(MeshObject->GetRayTracingGeometry()->Initializer.IndexBuffer.IsValid());
 			
@@ -6589,6 +6588,8 @@ void FSkeletalMeshSceneProxy::GetDynamicRayTracingInstances(FRayTracingMaterialG
 				RayTracingInstance.InstanceTransforms.Add(GetLocalToWorld());
 			}
 
+			const uint32 VertexBufferStride = LODData.StaticVertexBuffers.PositionVertexBuffer.GetStride();
+
 			const FVertexFactory* VertexFactory = MeshObject->GetSkinVertexFactory(Context.ReferenceView, LODIndex, 0, ESkinVertexFactoryMode::RayTracing);
 			const FVertexFactoryType* VertexFactoryType = VertexFactory->GetType();
 			if (bAnySegmentUsesWorldPositionOffset 
@@ -6605,6 +6606,8 @@ void FSkeletalMeshSceneProxy::GetDynamicRayTracingInstances(FRayTracingMaterialG
 					const FSectionElementInfo& SectionElementInfo = Iter.GetSectionElementInfo();
 
 					FRayTracingGeometrySegment Segment;
+					Segment.VertexBufferStride = VertexBufferStride;
+					Segment.MaxVertices = Section.GetNumVertices();
 					Segment.FirstPrimitive = Section.BaseIndex / 3;
 					Segment.NumPrimitives = Section.NumTriangles;
 					Segment.bEnabled = !MeshObject->IsMaterialHidden(LODIndex, SectionElementInfo.UseMaterialIndex) && !Section.bDisabled && Section.bVisibleInRayTracing;
