Support custom float data being optional when using UpdateInstances.

- Also support custom float count going from N to 0.
- Need to account for what what the previous num custom floats were so we can properly remove the old data.
- When updating existing instances the custom float count must stay the same.

diff --git a/Engine/Source/Runtime/Engine/Private/InstancedStaticMesh.cpp b/Engine/Source/Runtime/Engine/Private/InstancedStaticMesh.cpp
--- a/Engine/Source/Runtime/Engine/Private/InstancedStaticMesh.cpp
+++ b/Engine/Source/Runtime/Engine/Private/InstancedStaticMesh.cpp
@@ -3454,6 +3454,9 @@ bool UInstancedStaticMeshComponent::UpdateInstances(
 
 	InstanceUpdateCmdBuffer.NumCustomDataFloats = NumCustomDataFloats = InNumCustomDataFloats;
 
+	const bool bPreviouslyHadCustomFloatData = PrevNumCustomDataFloats > 0;
+	const bool bHasCustomFloatData = InstanceUpdateCmdBuffer.NumCustomDataFloats > 0;
+
 	// if we already have values we need to update, remove, and add.
 	TArray<int32> OldInstanceIds(PerInstanceIds);
 
@@ -3487,28 +3490,33 @@ bool UInstancedStaticMeshComponent::UpdateInstances(
 			}
 
 			// Did the custom data actually change?
-			const int32 CustomDataOffset = InstanceIndex * NumCustomDataFloats;
-			const int32 SrcCustomDataOffset = i * NumCustomDataFloats;
-			for (int32 FloatIndex = 0; FloatIndex < NumCustomDataFloats; ++FloatIndex)
+			if (bHasCustomFloatData)
 			{
-				if (FMath::Abs(CustomFloatData[SrcCustomDataOffset + FloatIndex] - PerInstanceSMCustomData[CustomDataOffset + FloatIndex]) > EqualTolerance)
+				// We are updating an existing instance.  They should have the same custom float data count.
+				check(NumCustomDataFloats == PrevNumCustomDataFloats);
+
+				const int32 CustomDataOffset = InstanceIndex * NumCustomDataFloats;
+				const int32 SrcCustomDataOffset = i * NumCustomDataFloats;
+				for (int32 FloatIndex = 0; FloatIndex < NumCustomDataFloats; ++FloatIndex)
 				{
-					// If any custom float data changed update it all.
-					TArray<float> NewCustomFloatData;
-					NewCustomFloatData.SetNumUninitialized(NumCustomDataFloats);
-					FMemory::Memcpy(&NewCustomFloatData[0], &CustomFloatData[SrcCustomDataOffset], NumCustomDataFloats * sizeof(float));
+					if (FMath::Abs(CustomFloatData[SrcCustomDataOffset + FloatIndex] - PerInstanceSMCustomData[CustomDataOffset + FloatIndex]) > EqualTolerance)
+					{
+						// If any custom float data changed update it all.
+						TArray<float> NewCustomFloatData;
+						NewCustomFloatData.SetNumUninitialized(NumCustomDataFloats);
+						FMemory::Memcpy(&NewCustomFloatData[0], &CustomFloatData[SrcCustomDataOffset], NumCustomDataFloats * sizeof(float));
 
-					// Update the component's data in place.
-					FMemory::Memcpy(&PerInstanceSMCustomData[CustomDataOffset], &NewCustomFloatData[0], NumCustomDataFloats * sizeof(float));
+						// Update the component's data in place.
+						FMemory::Memcpy(&PerInstanceSMCustomData[CustomDataOffset], &NewCustomFloatData[0], NumCustomDataFloats * sizeof(float));
 
-					// Record in a command buffer for future use.
-					InstanceUpdateCmdBuffer.SetCustomData(InstanceIndex, NewCustomFloatData);
+						// Record in a command buffer for future use.
+						InstanceUpdateCmdBuffer.SetCustomData(InstanceIndex, NewCustomFloatData);
 
 #if (CSV_PROFILER)
-					TotalSizeUpdateBytes += NumCustomDataFloats * sizeof(float);
+						TotalSizeUpdateBytes += NumCustomDataFloats * sizeof(float);
 #endif
-
-					break;
+						break;
+					}
 				}
 			}
 
@@ -3519,10 +3527,13 @@ bool UInstancedStaticMeshComponent::UpdateInstances(
 		{
 			// This is an add.
 			TArray<float> NewCustomFloatData;
-			NewCustomFloatData.SetNumUninitialized(NumCustomDataFloats);
+			if (bHasCustomFloatData)
+			{
+				NewCustomFloatData.SetNumUninitialized(NumCustomDataFloats);
 
-			const int32 SrcCustomDataOffset = i * NumCustomDataFloats;
-			FMemory::Memcpy(&NewCustomFloatData[0], &CustomFloatData[SrcCustomDataOffset], NumCustomDataFloats * sizeof(float));
+				const int32 SrcCustomDataOffset = i * NumCustomDataFloats;
+				FMemory::Memcpy(&NewCustomFloatData[0], &CustomFloatData[SrcCustomDataOffset], NumCustomDataFloats * sizeof(float));
+			}
 
 			InstanceUpdateCmdBuffer.AddInstance(InstanceId, UpdateInstanceTransforms[i].ToMatrixWithScale(), UpdateInstancePreviousTransforms[i].ToMatrixWithScale(), NewCustomFloatData);
 		}
@@ -3545,7 +3556,12 @@ bool UInstancedStaticMeshComponent::UpdateInstances(
 			PerInstanceSMData.RemoveAt(InstanceIndex, 1, false);
 			PerInstancePrevTransform.RemoveAt(InstanceIndex, 1, false);
 			PerInstanceIds.RemoveAt(InstanceIndex, 1, false);
-			PerInstanceSMCustomData.RemoveAt((InstanceIndex * PrevNumCustomDataFloats), PrevNumCustomDataFloats, false);
+
+			// Only remove the custom float data from this instance if it previously had it.
+			if (bPreviouslyHadCustomFloatData)
+			{
+				PerInstanceSMCustomData.RemoveAt((InstanceIndex * PrevNumCustomDataFloats), PrevNumCustomDataFloats, false);
+			}
 
 			OldInstanceIds.RemoveAt(InstanceIndex, 1, false);
 			InstanceIndex--;
@@ -3561,8 +3577,11 @@ bool UInstancedStaticMeshComponent::UpdateInstances(
 			PerInstancePrevTransform.Add(Cmd.PreviousXForm);
 			PerInstanceIds.Add(Cmd.InstanceId);
 
-			const int32 Index = PerInstanceSMCustomData.AddUninitialized(NumCustomDataFloats);
-			FMemory::Memcpy(&PerInstanceSMCustomData[Index], &Cmd.CustomDataFloats[0], NumCustomDataFloats * sizeof(float));
+			if (bHasCustomFloatData)
+			{
+				const int32 Index = PerInstanceSMCustomData.AddUninitialized(NumCustomDataFloats);
+				FMemory::Memcpy(&PerInstanceSMCustomData[Index], &Cmd.CustomDataFloats[0], NumCustomDataFloats * sizeof(float));
+			}
 		}
 	}
 
diff --git a/Engine/Source/Runtime/Engine/Private/PrimitiveSceneProxy.cpp b/Engine/Source/Runtime/Engine/Private/PrimitiveSceneProxy.cpp
--- a/Engine/Source/Runtime/Engine/Private/PrimitiveSceneProxy.cpp
+++ b/Engine/Source/Runtime/Engine/Private/PrimitiveSceneProxy.cpp
@@ -562,6 +562,9 @@ void FPrimitiveSceneProxy::UpdateInstances_RenderThread(const FInstanceUpdateCmd
 		const int32 PrevNumCustomDataFloats = InstanceSceneData.Num() ? InstanceCustomData.Num() / InstanceSceneData.Num() : 0;
 		const int32 NumCustomDataFloats = CmdBuffer.NumCustomDataFloats;
 
+		const bool bPreviouslyHadCustomFloatData = PrevNumCustomDataFloats > 0;
+		const bool bHasCustomFloatData = CmdBuffer.NumCustomDataFloats > 0;
+
 		// Apply all updates.
 		for (const auto& Cmd : CmdBuffer.Cmds)
 		{
@@ -577,6 +580,8 @@ void FPrimitiveSceneProxy::UpdateInstances_RenderThread(const FInstanceUpdateCmd
 			}
 			else if (Cmd.Type == FInstanceUpdateCmdBuffer::CustomData)
 			{
+				check(bHasCustomFloatData);
+
 				// update custom data because it changed.
 				check(PrevNumCustomDataFloats == NumCustomDataFloats);
 				const int32 DstCustomDataOffset = Cmd.InstanceIndex * NumCustomDataFloats;
@@ -615,8 +620,13 @@ void FPrimitiveSceneProxy::UpdateInstances_RenderThread(const FInstanceUpdateCmd
 				{
 					InstanceLightShadowUVBias.RemoveAt(i, 1, false);
 				}
-				InstanceCustomData.RemoveAt((i * PrevNumCustomDataFloats), PrevNumCustomDataFloats, false);
-				check(InstanceCustomData.Num() == (PrevNumCustomDataFloats * InstanceSceneData.Num()));
+
+				// Only remove the custom float data from this instance if it previously had it.
+				if (bPreviouslyHadCustomFloatData)
+				{
+					InstanceCustomData.RemoveAt((i * PrevNumCustomDataFloats), PrevNumCustomDataFloats, false);
+					check(InstanceCustomData.Num() == (PrevNumCustomDataFloats * InstanceSceneData.Num()));
+				}
 
 				RemoveBits.RemoveAt(i);
 
@@ -646,8 +656,11 @@ void FPrimitiveSceneProxy::UpdateInstances_RenderThread(const FInstanceUpdateCmd
 					InstanceLightShadowUVBias.AddZeroed();
 				}
 
-				const int32 DstCustomDataOffset = InstanceCustomData.AddUninitialized(NumCustomDataFloats);
-				FMemory::Memcpy(&InstanceCustomData[DstCustomDataOffset], &Cmd.CustomDataFloats[0], NumCustomDataFloats * sizeof(float));
+				if (bHasCustomFloatData)
+				{
+					const int32 DstCustomDataOffset = InstanceCustomData.AddUninitialized(NumCustomDataFloats);
+					FMemory::Memcpy(&InstanceCustomData[DstCustomDataOffset], &Cmd.CustomDataFloats[0], NumCustomDataFloats * sizeof(float));
+				}
 
 #if !(UE_BUILD_SHIPPING || UE_BUILD_TEST)
 				InstanceXFormUpdatedThisFrame.Add(true);
