Change warning in USkeletalMeshComponent::ValidateAnimation() to have full object name so validate commandlet can be used to track down assets with generic names

diff --git a/Engine/Source/Runtime/Engine/Private/Components/SkeletalMeshComponent.cpp b/Engine/Source/Runtime/Engine/Private/Components/SkeletalMeshComponent.cpp
--- a/Engine/Source/Runtime/Engine/Private/Components/SkeletalMeshComponent.cpp
+++ b/Engine/Source/Runtime/Engine/Private/Components/SkeletalMeshComponent.cpp
@@ -3718,7 +3718,7 @@ void USkeletalMeshComponent::ValidateAnimation()
 {
 	if (SkeletalMesh && SkeletalMesh->GetSkeleton() == nullptr)
 	{
-		UE_LOG(LogAnimation, Warning, TEXT("SkeletalMesh %s has no skeleton. This needs to fixed before an animation can be set"), *SkeletalMesh->GetName());
+		UE_LOG(LogAnimation, Warning, TEXT("SkeletalMesh %s has no skeleton. This needs to fixed before an animation can be set"), *SkeletalMesh->GetFullName());
 		if (AnimationMode == EAnimationMode::AnimationSingleNode)
 		{
 			AnimationData.AnimToPlay = nullptr;
