Drastically reduce the time it takes to process search for material instance parameters.  Removes large stalls across chase and when loading streaming in levels in sandbox.

diff --git a/Engine/Source/Runtime/Engine/Private/Materials/MaterialInstance.cpp b/Engine/Source/Runtime/Engine/Private/Materials/MaterialInstance.cpp
--- a/Engine/Source/Runtime/Engine/Private/Materials/MaterialInstance.cpp
+++ b/Engine/Source/Runtime/Engine/Private/Materials/MaterialInstance.cpp
@@ -312,9 +312,23 @@ void FMaterialInstanceResource::GameThread_UpdateCachedData(const FMaterialInsta
 		});
 }
 
+template <typename TInstanceType>
+static bool SortMaterialInstanceParametersPredicate(const FMaterialInstanceResource::TNamedParameter<TInstanceType>& Left, const FMaterialInstanceResource::TNamedParameter<TInstanceType>& Right)
+{
+	return GetTypeHash(Left.Info) < GetTypeHash(Right.Info);
+}
+
 void FMaterialInstanceResource::InitMIParameters(FMaterialInstanceParameterSet& ParameterSet)
 {
 	InvalidateUniformExpressionCache(false);
+
+	// Sort the parameters so that a binary lookup and be used.
+	ParameterSet.ScalarParameters.Sort(SortMaterialInstanceParametersPredicate<float>);
+	ParameterSet.VectorParameters.Sort(SortMaterialInstanceParametersPredicate<FLinearColor>);
+	ParameterSet.DoubleVectorParameters.Sort(SortMaterialInstanceParametersPredicate<FVector4d>);
+	ParameterSet.TextureParameters.Sort(SortMaterialInstanceParametersPredicate<const UTexture*>);
+	ParameterSet.RuntimeVirtualTextureParameters.Sort(SortMaterialInstanceParametersPredicate<const URuntimeVirtualTexture*>);
+
 	Swap(ScalarParameterArray, ParameterSet.ScalarParameters);
 	Swap(VectorParameterArray, ParameterSet.VectorParameters);
 	Swap(DoubleVectorParameterArray, ParameterSet.DoubleVectorParameters);
diff --git a/Engine/Source/Runtime/Engine/Private/Materials/MaterialInstanceSupport.h b/Engine/Source/Runtime/Engine/Private/Materials/MaterialInstanceSupport.h
--- a/Engine/Source/Runtime/Engine/Private/Materials/MaterialInstanceSupport.h
+++ b/Engine/Source/Runtime/Engine/Private/Materials/MaterialInstanceSupport.h
@@ -129,21 +129,21 @@ public:
 		LLM_SCOPE(ELLMTag::MaterialInstance);
 
 		InvalidateUniformExpressionCache(false);
+		bool bWasFound;
 		TArray<TNamedParameter<ValueType> >& ValueArray = GetValueArray<ValueType>();
-		const int32 ParameterCount = ValueArray.Num();
-		for (int32 ParameterIndex = 0; ParameterIndex < ParameterCount; ++ParameterIndex)
+		int Index = RenderThread_FindParameterByNameInternal<ValueType>(ParameterInfo, bWasFound);
+
+		if (bWasFound)
 		{
-			TNamedParameter<ValueType>& Parameter = ValueArray[ParameterIndex];
-			if (Parameter.Info == ParameterInfo)
-			{
-				Parameter.Value = Value;
-				return;
-			}
+			ValueArray[Index].Value = Value;
+		}
+		else
+		{
+			TNamedParameter<ValueType> NewParameter;
+			NewParameter.Info = ParameterInfo;
+			NewParameter.Value = Value;
+			ValueArray.Insert(NewParameter, Index);
 		}
-		TNamedParameter<ValueType> NewParameter;
-		NewParameter.Info = ParameterInfo;
-		NewParameter.Value = Value;
-		ValueArray.Add(NewParameter);
 	}
 
 	/**
@@ -152,16 +152,15 @@ public:
 	template <typename ValueType>
 	bool RenderThread_GetParameterValue(const FHashedMaterialParameterInfo& ParameterInfo, FMaterialParameterValue& OutValue) const
 	{
-		const TArray<TNamedParameter<ValueType>>& ValueArray = GetValueArray<ValueType>();
-		for (const TNamedParameter<ValueType>& Parameter : ValueArray)
+		bool bWasFound;
+		const TArray<TNamedParameter<ValueType> >& ValueArray = GetValueArray<ValueType>();
+		int Index = RenderThread_FindParameterByNameInternal<ValueType>(ParameterInfo, bWasFound);
+
+		if (bWasFound)
 		{
-			if (Parameter.Info == ParameterInfo && IsValidParameterValue(Parameter.Value))
-			{
-				OutValue = Parameter.Value;
-				return true;
-			}
+			OutValue = ValueArray[Index].Value;
 		}
-		return false;
+		return bWasFound;
 	}
 
 private:
@@ -177,6 +176,37 @@ private:
 	static bool IsValidParameterValue(const UTexture* Value) { return Value != nullptr; }
 	static bool IsValidParameterValue(const URuntimeVirtualTexture* Value) { return Value != nullptr; }
 
+	template <typename ValueType>
+	int RenderThread_FindParameterByNameInternal(const FHashedMaterialParameterInfo& ParameterInfo, bool& OutWasFound) const
+	{
+		const TArray<TNamedParameter<ValueType> >& ValueArray = GetValueArray<ValueType>();
+
+		TNamedParameter<ValueType> SearchParam;
+		SearchParam.Info = ParameterInfo;
+
+		int Index = Algo::LowerBound(ValueArray, SearchParam,
+			[](const TNamedParameter<ValueType>& Left, const TNamedParameter<ValueType>& Right)
+			{
+				return GetTypeHash(Left.Info) < GetTypeHash(Right.Info);
+			});
+
+		uint32 SearchHash = GetTypeHash(ParameterInfo);
+
+		while (ValueArray.IsValidIndex(Index) && GetTypeHash(ValueArray[Index].Info) == SearchHash)
+		{
+			if (ValueArray[Index].Info == ParameterInfo)
+			{
+				OutWasFound = true;
+				return Index;
+			}
+
+			Index++;
+		}
+
+		OutWasFound = false;
+		return Index;
+	}
+
 	/** The parent of the material instance. */
 	UMaterialInterface* Parent;
 
diff --git a/Engine/Source/Runtime/Engine/Private/Materials/MaterialShared.cpp b/Engine/Source/Runtime/Engine/Private/Materials/MaterialShared.cpp
--- a/Engine/Source/Runtime/Engine/Private/Materials/MaterialShared.cpp
+++ b/Engine/Source/Runtime/Engine/Private/Materials/MaterialShared.cpp
@@ -3183,6 +3183,8 @@ void FMaterialRenderProxy::EvaluateUniformExpressions(FUniformExpressionCache& O
 	check(IsInParallelRenderingThread());
 
 	SCOPE_CYCLE_COUNTER(STAT_CacheUniformExpressions);
+
+	TRACE_CPUPROFILER_EVENT_SCOPE(FMaterialRenderProxy::EvaluateUniformExpressions);
 	
 	// Retrieve the material's uniform expression set.
 	FMaterialShaderMap* ShaderMap = Context.Material.GetRenderingThreadShaderMap();
