Drastically reduce the cost of UpdateInstances when dealing with a large number of instances or a large number of removes.

- Perform RemoveSwaps instead of RemoveAts (we don't care about order).
- Cache a mapping from InstanceID to InstanceIndex (in the component's instance data arrays) to provide faster lookup when decided if we've "seen" this instance before.  This is instead of a linear walk of the instance ID array.
- Add some profile markers to the old BatchUpdateTransform function.

diff --git a/Engine/Source/Runtime/Engine/Classes/Components/InstancedStaticMeshComponent.h b/Engine/Source/Runtime/Engine/Classes/Components/InstancedStaticMeshComponent.h
--- a/Engine/Source/Runtime/Engine/Classes/Components/InstancedStaticMeshComponent.h
+++ b/Engine/Source/Runtime/Engine/Classes/Components/InstancedStaticMeshComponent.h
@@ -383,6 +383,12 @@ public:
 
 	TArray<int32> PerInstanceIds;
 
+	/** Used to cache a unique identifier for each instance.  These are provided
+	*	by the interface UpdateInstances.  This is a map from unique id to index
+	*	into the PerInstanceSMData array.
+	*/
+	TMap<int32, int32> InstanceIdToInstanceIndexMap;
+
 private:
 
 	/** Sets up new instance data to sensible defaults, creates physics counterparts if possible. */
diff --git a/Engine/Source/Runtime/Engine/Private/InstancedStaticMesh.cpp b/Engine/Source/Runtime/Engine/Private/InstancedStaticMesh.cpp
--- a/Engine/Source/Runtime/Engine/Private/InstancedStaticMesh.cpp
+++ b/Engine/Source/Runtime/Engine/Private/InstancedStaticMesh.cpp
@@ -3465,9 +3465,11 @@ bool UInstancedStaticMeshComponent::UpdateInstances(
 	{
 		// This is an update.
 		const int32 InstanceId = UpdateInstanceIds[i];
-		int InstanceIndex = PerInstanceIds.Find(InstanceId);
-		if (InstanceIndex != INDEX_NONE)
+		int32* pInstanceIndex = InstanceIdToInstanceIndexMap.Find(InstanceId);
+		if (pInstanceIndex != nullptr)
 		{
+			const int32 InstanceIndex = *pInstanceIndex;
+
 			// Determine what data changed.
 
 			// Did the transform actually change?
@@ -3553,17 +3555,17 @@ bool UInstancedStaticMeshComponent::UpdateInstances(
 	{
 		if (OldInstanceIds[InstanceIndex] != INDEX_NONE)
 		{
-			PerInstanceSMData.RemoveAt(InstanceIndex, 1, false);
-			PerInstancePrevTransform.RemoveAt(InstanceIndex, 1, false);
-			PerInstanceIds.RemoveAt(InstanceIndex, 1, false);
+			PerInstanceSMData.RemoveAtSwap(InstanceIndex, 1, false);
+			PerInstancePrevTransform.RemoveAtSwap(InstanceIndex, 1, false);
+			PerInstanceIds.RemoveAtSwap(InstanceIndex, 1, false);
 
 			// Only remove the custom float data from this instance if it previously had it.
 			if (bPreviouslyHadCustomFloatData)
 			{
-				PerInstanceSMCustomData.RemoveAt((InstanceIndex * PrevNumCustomDataFloats), PrevNumCustomDataFloats, false);
+				PerInstanceSMCustomData.RemoveAtSwap((InstanceIndex * PrevNumCustomDataFloats), PrevNumCustomDataFloats, false);
 			}
 
-			OldInstanceIds.RemoveAt(InstanceIndex, 1, false);
+			OldInstanceIds.RemoveAtSwap(InstanceIndex, 1, false);
 			InstanceIndex--;
 		}
 	}
@@ -3585,6 +3587,13 @@ bool UInstancedStaticMeshComponent::UpdateInstances(
 		}
 	}
 
+	// Rebuild the mapping from ID to InstanceIndex.
+	InstanceIdToInstanceIndexMap.Reset();
+	for (int32 InstanceIndex = 0; InstanceIndex < PerInstanceIds.Num(); ++InstanceIndex)
+	{
+		InstanceIdToInstanceIndexMap.Add(PerInstanceIds[InstanceIndex], InstanceIndex);
+	}
+
 	// Ensure our data is in sync.
 	check(PerInstanceSMData.Num() == PerInstancePrevTransform.Num());
 	check(PerInstanceSMData.Num() == PerInstanceIds.Num());
@@ -3628,6 +3637,9 @@ bool UInstancedStaticMeshComponent::UpdateInstances(
 
 bool UInstancedStaticMeshComponent::BatchUpdateInstancesTransforms(int32 StartInstanceIndex, const TArray<FTransform>& NewInstancesTransforms, bool bWorldSpace, bool bMarkRenderStateDirty, bool bTeleport)
 {
+	CSV_SCOPED_TIMING_STAT_EXCLUSIVE(UInstancedStaticMeshComponent_BatchUpdateInstancesTransforms);
+	TRACE_CPUPROFILER_EVENT_SCOPE_STR("UInstancedStaticMeshComponent::BatchUpdateInstancesTransforms");
+
 	if (!PerInstanceSMData.IsValidIndex(StartInstanceIndex) || !PerInstanceSMData.IsValidIndex(StartInstanceIndex + NewInstancesTransforms.Num() - 1))
 	{
 		return false;
diff --git a/Engine/Source/Runtime/Engine/Private/PrimitiveSceneProxy.cpp b/Engine/Source/Runtime/Engine/Private/PrimitiveSceneProxy.cpp
--- a/Engine/Source/Runtime/Engine/Private/PrimitiveSceneProxy.cpp
+++ b/Engine/Source/Runtime/Engine/Private/PrimitiveSceneProxy.cpp
@@ -609,30 +609,30 @@ void FPrimitiveSceneProxy::UpdateInstances_RenderThread(const FInstanceUpdateCmd
 		{
 			if (RemoveBits[i])
 			{
-				InstanceSceneData.RemoveAt(i, 1, false);
-				InstanceDynamicData.RemoveAt(i, 1, false);
+				InstanceSceneData.RemoveAtSwap(i, 1, false);
+				InstanceDynamicData.RemoveAtSwap(i, 1, false);
 
 				if (bHasPerInstanceRandom)
 				{
-					InstanceRandomID.RemoveAt(i, 1, false);
+					InstanceRandomID.RemoveAtSwap(i, 1, false);
 				}
 				if (bHasPerInstanceLMSMUVBias)
 				{
-					InstanceLightShadowUVBias.RemoveAt(i, 1, false);
+					InstanceLightShadowUVBias.RemoveAtSwap(i, 1, false);
 				}
 
 				// Only remove the custom float data from this instance if it previously had it.
 				if (bPreviouslyHadCustomFloatData)
 				{
-					InstanceCustomData.RemoveAt((i * PrevNumCustomDataFloats), PrevNumCustomDataFloats, false);
+					InstanceCustomData.RemoveAtSwap((i * PrevNumCustomDataFloats), PrevNumCustomDataFloats, false);
 					check(InstanceCustomData.Num() == (PrevNumCustomDataFloats * InstanceSceneData.Num()));
 				}
 
-				RemoveBits.RemoveAt(i);
+				RemoveBits.RemoveAtSwap(i);
 
 #if !(UE_BUILD_SHIPPING || UE_BUILD_TEST)
-				InstanceXFormUpdatedThisFrame.RemoveAt(i);
-				InstanceCustomDataUpdatedThisFrame.RemoveAt(i);
+				InstanceXFormUpdatedThisFrame.RemoveAtSwap(i);
+				InstanceCustomDataUpdatedThisFrame.RemoveAtSwap(i);
 #endif
 				i--;
 			}
