WorldPartitionEditorPerProjectUserSettings: Fix mismatching return type between declaration/definition

diff --git a/Engine/Source/Runtime/Engine/Public/WorldPartition/WorldPartitionEditorPerProjectUserSettings.h b/Engine/Source/Runtime/Engine/Public/WorldPartition/WorldPartitionEditorPerProjectUserSettings.h
--- a/Engine/Source/Runtime/Engine/Public/WorldPartition/WorldPartitionEditorPerProjectUserSettings.h
+++ b/Engine/Source/Runtime/Engine/Public/WorldPartition/WorldPartitionEditorPerProjectUserSettings.h
@@ -102,8 +102,8 @@ public:
 		}
 	}
 
-	const TArray<FName>& GetWorldDataLayersNotLoadedInEditor(UWorld* InWorld) const;
-	const TArray<FName>& GetWorldDataLayersLoadedInEditor(UWorld* InWorld) const;
+	TArray<FName> GetWorldDataLayersNotLoadedInEditor(UWorld* InWorld) const;
+	TArray<FName> GetWorldDataLayersLoadedInEditor(UWorld* InWorld) const;
 	
 	void SetWorldDataLayersNonDefaultEditorLoadStates(UWorld* InWorld, const TArray<FName>& InDataLayersLoadedInEditor, const TArray<FName>& InDataLayersNotLoadedInEditor);
 
