Added flag to FRayTracingGeometry to track if geometry requires an initial BLAS build.
- For example, if it was created using ERTAccelerationStructureBuildPriority::Skip.
- Fix assert for skinned meshes using WorldPositionOffset.

diff --git a/Engine/Source/Runtime/Engine/Private/GPUSkinCache.cpp b/Engine/Source/Runtime/Engine/Private/GPUSkinCache.cpp
--- a/Engine/Source/Runtime/Engine/Private/GPUSkinCache.cpp
+++ b/Engine/Source/Runtime/Engine/Private/GPUSkinCache.cpp
@@ -1857,11 +1857,8 @@ void FGPUSkinCache::ProcessRayTracingGeometryToUpdate(
 			// Get the scratch sizes used for build & update
 			SkinCacheEntry->GPUSkin->RayTracingGeometryStructureSize = RHICalcRayTracingGeometrySize(Initializer);
 
-			if (!bAnySegmentUsesWorldPositionOffset)
-			{
-				// Only create RHI object but enqueue actual BLAS creation so they can be accumulated
-				RayTracingGeometry.CreateRayTracingGeometry(ERTAccelerationStructureBuildPriority::Skip);
-			}
+			// Only create RHI object but enqueue actual BLAS creation so they can be accumulated
+			RayTracingGeometry.CreateRayTracingGeometry(ERTAccelerationStructureBuildPriority::Skip);
 		}
 		else if (!bAnySegmentUsesWorldPositionOffset)
 		{
@@ -1874,6 +1871,7 @@ void FGPUSkinCache::ProcessRayTracingGeometryToUpdate(
 		{
 			EAccelerationStructureBuildMode BuildMode = bRequireRecreatingRayTracingGeometry ? EAccelerationStructureBuildMode::Build : EAccelerationStructureBuildMode::Update;
 			AddRayTracingGeometryToUpdate(&RayTracingGeometry, SkinCacheEntry->GPUSkin->RayTracingGeometryStructureSize, BuildMode);
+			RayTracingGeometry.bRequiresBuild = false;
 		}
 		else
 		{
diff --git a/Engine/Source/Runtime/Engine/Private/SkeletalMesh.cpp b/Engine/Source/Runtime/Engine/Private/SkeletalMesh.cpp
--- a/Engine/Source/Runtime/Engine/Private/SkeletalMesh.cpp
+++ b/Engine/Source/Runtime/Engine/Private/SkeletalMesh.cpp
@@ -6528,7 +6528,7 @@ void FSkeletalMeshSceneProxy::GetDynamicRayTracingInstances(FRayTracingMaterialG
 	// GetRayTracingGeometry()->IsInitialized() is checked as a workaround for UE-92634. FSkeletalMeshSceneProxy's resources may have already been released, but proxy has not removed yet)
 	if (MeshObject->GetRayTracingGeometry() && MeshObject->GetRayTracingGeometry()->IsInitialized() && MeshObject->GetSkinCacheEntryForRayTracing())
 	{
-		checkf(MeshObject->GetRayTracingGeometry()->RayTracingGeometryRHI.IsValid() || bAnySegmentUsesWorldPositionOffset, TEXT("RayTracingGeometryRHI should be valid at this point unless object uses WorldPositionOffset"));
+		if(MeshObject->GetRayTracingGeometry()->RayTracingGeometryRHI.IsValid())
 		{
 			check(MeshObject->GetRayTracingGeometry()->Initializer.IndexBuffer.IsValid());
 			
diff --git a/Engine/Source/Runtime/RenderCore/Private/RenderResource.cpp b/Engine/Source/Runtime/RenderCore/Private/RenderResource.cpp
--- a/Engine/Source/Runtime/RenderCore/Private/RenderResource.cpp
+++ b/Engine/Source/Runtime/RenderCore/Private/RenderResource.cpp
@@ -416,10 +416,17 @@ void FRayTracingGeometry::CreateRayTracingGeometry(ERTAccelerationStructureBuild
 			if (InBuildPriority != ERTAccelerationStructureBuildPriority::Skip)
 			{
 				RayTracingBuildRequestIndex = GRayTracingGeometryManager.RequestBuildAccelerationStructure(this, InBuildPriority);
+				bRequiresBuild = false;
+			}
+			else
+			{
+				bRequiresBuild = true;
 			}
 		}
 		else
 		{
+			bRequiresBuild = false;
+
 			// Offline data ownership is transferred to the RHI, which discards it after use.
 			// It is no longer valid to use it after this point.
 			Initializer.OfflineData = nullptr;
diff --git a/Engine/Source/Runtime/RenderCore/Public/RenderResource.h b/Engine/Source/Runtime/RenderCore/Public/RenderResource.h
--- a/Engine/Source/Runtime/RenderCore/Public/RenderResource.h
+++ b/Engine/Source/Runtime/RenderCore/Public/RenderResource.h
@@ -692,6 +692,7 @@ public:
 
 	FRayTracingGeometryInitializer Initializer;
 	FRayTracingGeometryRHIRef RayTracingGeometryRHI;
+	bool bRequiresBuild = false;
 
 	void SetInitializer(const FRayTracingGeometryInitializer& InInitializer)
 	{
diff --git a/Engine/Source/Runtime/Renderer/Private/RayTracing/RayTracingDynamicGeometry.cpp b/Engine/Source/Runtime/Renderer/Private/RayTracing/RayTracingDynamicGeometry.cpp
--- a/Engine/Source/Runtime/Renderer/Private/RayTracing/RayTracingDynamicGeometry.cpp
+++ b/Engine/Source/Runtime/Renderer/Private/RayTracing/RayTracingDynamicGeometry.cpp
@@ -302,13 +302,16 @@ void FRayTracingDynamicGeometryCollection::AddDynamicMeshBatchForGeometryUpdate(
 	{
 		checkf(Geometry.Initializer.OfflineData == nullptr, TEXT("Dynamic geometry is not expected to have offline acceleration structure data"));
 		Geometry.RayTracingGeometryRHI = RHICreateRayTracingGeometry(Geometry.Initializer);
+		Geometry.bRequiresBuild = true;
 	}
 
 	FRayTracingGeometryBuildParams Params;
 	Params.Geometry = Geometry.RayTracingGeometryRHI;
-	Params.BuildMode = bRefit
-		? EAccelerationStructureBuildMode::Update
-		: EAccelerationStructureBuildMode::Build;
+	Params.BuildMode = Geometry.bRequiresBuild
+		? EAccelerationStructureBuildMode::Build
+		: EAccelerationStructureBuildMode::Update;
+
+	Geometry.bRequiresBuild = false;
 
 	if (bUseSharedVertexBuffer)
 	{
