Ensure that primitives that have WPO render during the velocity pass.

- The code was previously calling GetMaterial_Concurrent on the Mat Interface which for UMaterialInstance was returning the parent material.
- It would then call GetMaterialResource on the parent material.  It would then check the _parent_ material to see if it was modifying mesh position.  This state would be incorrect for MaterialInstances that had static parameters, which have their own material resources.
- The fix is to ensure we call GetMaterialResource on the MaterialInstance instead of the parent material.  The simplest way to do this is just call GetRelevance_Concurrent on the MaterialInterface which will do the right thing.
- Check the MaterialRelevance to see if the primitive has WPO enabled instead of asking the material resource.

diff --git a/Engine/Source/Runtime/Engine/Private/PrimitiveSceneProxy.cpp b/Engine/Source/Runtime/Engine/Private/PrimitiveSceneProxy.cpp
--- a/Engine/Source/Runtime/Engine/Private/PrimitiveSceneProxy.cpp
+++ b/Engine/Source/Runtime/Engine/Private/PrimitiveSceneProxy.cpp
@@ -321,22 +321,10 @@ FPrimitiveSceneProxy::FPrimitiveSceneProxy(const UPrimitiveComponent* InComponen
 		{
 			if (MaterialInterface)
 			{
-				const UMaterial* Material = MaterialInterface->GetMaterial_Concurrent();
-				if (const FMaterialResource* MaterialResource = Material->GetMaterialResource(FeatureLevel))
+				bAlwaysHasVelocity = MaterialInterface->GetRelevance_Concurrent(FeatureLevel).bUsesWorldPositionOffset;
+				if (bAlwaysHasVelocity)
 				{
-					if (IsInParallelGameThread() || IsInGameThread())
-					{
-						bAlwaysHasVelocity = MaterialResource->MaterialModifiesMeshPosition_GameThread();
-					}
-					else
-					{
-						bAlwaysHasVelocity = MaterialResource->MaterialModifiesMeshPosition_RenderThread();
-					}
-
-					if (bAlwaysHasVelocity)
-					{
-						break;
-					}
+					break;
 				}
 			}
 		}
