Prefetch render resources for HLODs before unloading a WP cell
* Warmup will trigger resource prefetching and delay unloading of cells for a number of frames
* Enable HLOD assets warmup: wp.Runtime.HLOD.WarmupEnabled 0|1
* Number of warmup frames:    wp.Runtime.HLOD.WarmupNumFrames n
* Currently only prefetching VT, Nanite prefetching to be added in another CL

Currently disabled (see "wp.Runtime.HLOD.WarmupNumFrames=0" in the device profiles ini)

diff --git a/Engine/Source/Developer/WorldPartitionHLODUtilities/Private/WorldPartitionHLODUtilities.cpp b/Engine/Source/Developer/WorldPartitionHLODUtilities/Private/WorldPartitionHLODUtilities.cpp
--- a/Engine/Source/Developer/WorldPartitionHLODUtilities/Private/WorldPartitionHLODUtilities.cpp
+++ b/Engine/Source/Developer/WorldPartitionHLODUtilities/Private/WorldPartitionHLODUtilities.cpp
@@ -281,6 +281,13 @@ TArray<AWorldPartitionHLOD*> FWorldPartitionHLODUtilities::CreateHLODActors(FHLO
 			bIsDirty = true;
 		}
 
+		// Require warmup
+		if (HLODActor->DoesRequireWarmup() != HLODLayer->DoesRequireWarmup())
+		{
+			HLODActor->SetRequireWarmup(HLODLayer->DoesRequireWarmup());
+			bIsDirty = true;
+		}
+
 		// Parent HLOD layer
 		UHLODLayer* ParentHLODLayer = HLODLayer->GetParentLayer().LoadSynchronous();
 		if (HLODActor->GetHLODLayer() != ParentHLODLayer)
diff --git a/Engine/Source/Developer/WorldPartitionHLODUtilities/Public/HLODBuilderInstancing.h b/Engine/Source/Developer/WorldPartitionHLODUtilities/Public/HLODBuilderInstancing.h
--- a/Engine/Source/Developer/WorldPartitionHLODUtilities/Public/HLODBuilderInstancing.h
+++ b/Engine/Source/Developer/WorldPartitionHLODUtilities/Public/HLODBuilderInstancing.h
@@ -16,6 +16,7 @@ class WORLDPARTITIONHLODUTILITIES_API UHLODBuilderInstancing : public UHLODBuild
 
 public:
 	virtual bool RequiresCompiledAssets() const override { return false; }
+	virtual bool RequiresWarmup() const override { return false; }
 
 	virtual TArray<UPrimitiveComponent*> CreateComponents(AWorldPartitionHLOD* InHLODActor, const UHLODLayer* InHLODLayer, const TArray<UPrimitiveComponent*>& InSubComponents) const override;
 };
diff --git a/Engine/Source/Runtime/Engine/Private/WorldPartition/HLOD/HLODActor.cpp b/Engine/Source/Runtime/Engine/Private/WorldPartition/HLOD/HLODActor.cpp
--- a/Engine/Source/Runtime/Engine/Private/WorldPartition/HLOD/HLODActor.cpp
+++ b/Engine/Source/Runtime/Engine/Private/WorldPartition/HLOD/HLODActor.cpp
@@ -21,6 +21,7 @@
 
 AWorldPartitionHLOD::AWorldPartitionHLOD(const FObjectInitializer& ObjectInitializer)
 	: Super(ObjectInitializer)
+	, bRequireWarmup(false)
 {
 	SetCanBeDamaged(false);
 	SetActorEnableCollision(false);
@@ -151,6 +152,12 @@ const TArray<FGuid>& AWorldPartitionHLOD::GetSubActors() const
 	return SubActors;
 }
 
+void AWorldPartitionHLOD::SetSubActorsHLODLayer(const UHLODLayer* InSubActorsHLODLayer)
+{
+	SubActorsHLODLayer = InSubActorsHLODLayer;
+	bRequireWarmup = SubActorsHLODLayer->DoesRequireWarmup();
+}
+
 void AWorldPartitionHLOD::SetSourceCellName(FName InSourceCellName)
 {
 	SourceCellName = InSourceCellName;
diff --git a/Engine/Source/Runtime/Engine/Private/WorldPartition/HLOD/HLODBuilder.cpp b/Engine/Source/Runtime/Engine/Private/WorldPartition/HLOD/HLODBuilder.cpp
--- a/Engine/Source/Runtime/Engine/Private/WorldPartition/HLOD/HLODBuilder.cpp
+++ b/Engine/Source/Runtime/Engine/Private/WorldPartition/HLOD/HLODBuilder.cpp
@@ -36,6 +36,11 @@ bool UHLODBuilder::RequiresCompiledAssets() const
 	return true;
 }
 
+bool UHLODBuilder::RequiresWarmup() const
+{
+	return true;
+}
+
 TArray<UPrimitiveComponent*> UHLODBuilder::CreateComponents(AWorldPartitionHLOD* InHLODActor, const UHLODLayer* InHLODLayer, const TArray<UPrimitiveComponent*>& InSubComponents) const
 {
 	return TArray<UPrimitiveComponent*>();
diff --git a/Engine/Source/Runtime/Engine/Private/WorldPartition/HLOD/HLODSubsystem.cpp b/Engine/Source/Runtime/Engine/Private/WorldPartition/HLOD/HLODSubsystem.cpp
--- a/Engine/Source/Runtime/Engine/Private/WorldPartition/HLOD/HLODSubsystem.cpp
+++ b/Engine/Source/Runtime/Engine/Private/WorldPartition/HLOD/HLODSubsystem.cpp
@@ -10,14 +10,30 @@
 #include "WorldPartition/WorldPartitionRuntimeHash.h"
 #include "UObject/UObjectHash.h"
 #include "UObject/UObjectGlobals.h"
-#include "Engine/World.h"
 #include "Engine/Engine.h"
+#include "Engine/World.h"
+#include "EngineModule.h"
 #include "EngineUtils.h"
+#include "Components/StaticMeshComponent.h"
 
 #define LOCTEXT_NAMESPACE "HLODSubsystem"
 
 DEFINE_LOG_CATEGORY_STATIC(LogHLODSubsystem, Log, All);
 
+static TAutoConsoleVariable<int32> CVarHLODWarmupNumFrames(
+	TEXT("wp.Runtime.HLOD.WarmupNumFrames"),
+	5,
+	TEXT("Delay unloading of a cell for this amount of frames to ensure HLOD assets are ready to be shown at the proper resolution. Set to 0 to force disable warmup."),
+	ECVF_Default
+);
+
+static TAutoConsoleVariable<int32> CVarHLODWarmupEnabled(
+	TEXT("wp.Runtime.HLOD.WarmupEnabled"),
+	1,
+	TEXT("Enable HLOD assets warmup. Will delay unloading of cells & transition to HLODs for wp.Runtime.HLOD.WarmupNumFrames frames."),
+	ECVF_Default
+);
+
 UHLODSubsystem::UHLODSubsystem()
 	: UWorldSubsystem()
 {
@@ -41,11 +57,11 @@ FAutoConsoleCommand UHLODSubsystem::EnableHLODCommand(
 			if (World && World->IsGameWorld())
 			{
 				UHLODSubsystem* HLODSubSystem = World->GetSubsystem<UHLODSubsystem>();
-				for (const auto& CellHLODMapping : HLODSubSystem->CellsHLODMapping)
+				for (const auto& CellHLODMapping : HLODSubSystem->CellsData)
 				{
-					const FCellHLODMapping& CellHLODs = CellHLODMapping.Value;
-					bool bIsHLODVisible = UHLODSubsystem::WorldPartitionHLODEnabled && !CellHLODs.bIsCellVisible;
-					for (AWorldPartitionHLOD* HLODActor : CellHLODs.LoadedHLODs)
+					const FCellData& CellData = CellHLODMapping.Value;
+					bool bIsHLODVisible = UHLODSubsystem::WorldPartitionHLODEnabled && !CellData.bIsCellVisible;
+					for (AWorldPartitionHLOD* HLODActor : CellData.LoadedHLODs)
 					{
 						HLODActor->SetVisibility(bIsHLODVisible);
 					}
@@ -81,17 +97,21 @@ void UHLODSubsystem::Initialize(FSubsystemCollectionBase& Collection)
 
 	Super::Initialize(Collection);
 
-	if (GetWorld()->IsGameWorld())
+	UWorld* World = GetWorld();
+
+	if (World->IsGameWorld())
 	{
-		GetWorld()->GetSubsystem<UWorldPartitionSubsystem>()->OnWorldPartitionRegistered.AddUObject(this, &UHLODSubsystem::OnWorldPartitionRegistered);
-		GetWorld()->GetSubsystem<UWorldPartitionSubsystem>()->OnWorldPartitionUnregistered.AddUObject(this, &UHLODSubsystem::OnWorldPartitionUnregistered);
+		World->GetSubsystem<UWorldPartitionSubsystem>()->OnWorldPartitionRegistered.AddUObject(this, &UHLODSubsystem::OnWorldPartitionRegistered);
+		World->GetSubsystem<UWorldPartitionSubsystem>()->OnWorldPartitionUnregistered.AddUObject(this, &UHLODSubsystem::OnWorldPartitionUnregistered);
+
+		SceneViewExtension = FSceneViewExtensions::NewExtension<FHLODResourcesResidencySceneViewExtension>(World);
 	}
 }
 
 void UHLODSubsystem::OnWorldPartitionRegistered(UWorldPartition* InWorldPartition)
 {
 	check(InWorldPartition == GetWorld()->GetWorldPartition());
-	check(CellsHLODMapping.IsEmpty());
+	check(CellsData.IsEmpty());
 
 	TSet<const UWorldPartitionRuntimeCell*> StreamingCells;
 	InWorldPartition->RuntimeHash->GetAllStreamingCells(StreamingCells, /*bAllDataLayers*/ true);
@@ -99,14 +119,14 @@ void UHLODSubsystem::OnWorldPartitionRegistered(UWorldPartition* InWorldPartitio
 	// Build cell to HLOD mapping
 	for (const UWorldPartitionRuntimeCell* Cell : StreamingCells)
 	{
-		CellsHLODMapping.Emplace(Cell->GetFName());
+		CellsData.Emplace(Cell->GetFName());
 	}
 }
 
 void UHLODSubsystem::OnWorldPartitionUnregistered(UWorldPartition* InWorldPartition)
 {
 	check(InWorldPartition == GetWorld()->GetWorldPartition());
-	CellsHLODMapping.Reset();
+	CellsData.Reset();
 }
 
 void UHLODSubsystem::RegisterHLODActor(AWorldPartitionHLOD* InWorldPartitionHLOD)
@@ -114,16 +134,16 @@ void UHLODSubsystem::RegisterHLODActor(AWorldPartitionHLOD* InWorldPartitionHLOD
 	TRACE_CPUPROFILER_EVENT_SCOPE(UHLODSubsystem::RegisterHLODActor);
 
 	const FName CellName = InWorldPartitionHLOD->GetSourceCellName();
-	FCellHLODMapping* CellHLODs = CellsHLODMapping.Find(CellName);
+	FCellData* CellData = CellsData.Find(CellName);
 
 #if WITH_EDITOR
 	UE_LOG(LogHLODSubsystem, Verbose, TEXT("Registering HLOD %s (%s) for cell %s"), *InWorldPartitionHLOD->GetActorLabel(), *InWorldPartitionHLOD->GetActorGuid().ToString(), *CellName.ToString());
 #endif
 
-	if (CellHLODs)
+	if (CellData)
 	{
-		CellHLODs->LoadedHLODs.Add(InWorldPartitionHLOD);
-		InWorldPartitionHLOD->SetVisibility(UHLODSubsystem::WorldPartitionHLODEnabled && !CellHLODs->bIsCellVisible);
+		CellData->LoadedHLODs.Add(InWorldPartitionHLOD);
+		InWorldPartitionHLOD->SetVisibility(UHLODSubsystem::WorldPartitionHLODEnabled && !CellData->bIsCellVisible);
 	}
 	else
 	{
@@ -137,29 +157,29 @@ void UHLODSubsystem::UnregisterHLODActor(AWorldPartitionHLOD* InWorldPartitionHL
 	TRACE_CPUPROFILER_EVENT_SCOPE(UHLODSubsystem::UnregisterHLODActor);
 
 	const FName CellName = InWorldPartitionHLOD->GetSourceCellName();
-	FCellHLODMapping* CellHLODs = CellsHLODMapping.Find(CellName);
+	FCellData* CellData = CellsData.Find(CellName);
 
 #if WITH_EDITOR
 	UE_LOG(LogHLODSubsystem, Verbose, TEXT("Unregistering HLOD %s (%s) for cell %s"), *InWorldPartitionHLOD->GetActorLabel(), *InWorldPartitionHLOD->GetActorGuid().ToString(), *CellName.ToString());
 #endif
 
-	if (CellHLODs)
+	if (CellData)
 	{
-		int32 NumRemoved = CellHLODs->LoadedHLODs.Remove(InWorldPartitionHLOD);
+		int32 NumRemoved = CellData->LoadedHLODs.Remove(InWorldPartitionHLOD);
 		check(NumRemoved == 1);
 	}
 }
 
 void UHLODSubsystem::OnCellShown(const UWorldPartitionRuntimeCell* InCell)
 {
-	FCellHLODMapping& CellHLODs = CellsHLODMapping.FindChecked(InCell->GetFName());
-	CellHLODs.bIsCellVisible = true;
+	FCellData& CellData = CellsData.FindChecked(InCell->GetFName());
+	CellData.bIsCellVisible = true;
 
 #if WITH_EDITOR
-	UE_LOG(LogHLODSubsystem, Verbose, TEXT("Cell shown - %s - hiding %d HLOD actors"), *InCell->GetName(), CellHLODs.LoadedHLODs.Num());
+	UE_LOG(LogHLODSubsystem, Verbose, TEXT("Cell shown - %s - hiding %d HLOD actors"), *InCell->GetName(), CellData.LoadedHLODs.Num());
 #endif
 
-	for (AWorldPartitionHLOD* HLODActor : CellHLODs.LoadedHLODs)
+	for (AWorldPartitionHLOD* HLODActor : CellData.LoadedHLODs)
 	{
 #if WITH_EDITOR
 		UE_LOG(LogHLODSubsystem, Verbose, TEXT("\t\t%s - %s"), *HLODActor->GetActorLabel(), *HLODActor->GetActorGuid().ToString());
@@ -170,14 +190,14 @@ void UHLODSubsystem::OnCellShown(const UWorldPartitionRuntimeCell* InCell)
 
 void UHLODSubsystem::OnCellHidden(const UWorldPartitionRuntimeCell* InCell)
 {
-	FCellHLODMapping& CellHLODs = CellsHLODMapping.FindChecked(InCell->GetFName());
-	CellHLODs.bIsCellVisible = false;
+	FCellData& CellData = CellsData.FindChecked(InCell->GetFName());
+	CellData.bIsCellVisible = false;
 
 #if WITH_EDITOR
-	UE_LOG(LogHLODSubsystem, Verbose, TEXT("Cell hidden - %s - showing %d HLOD actors"), *InCell->GetName(), CellHLODs.LoadedHLODs.Num());
+	UE_LOG(LogHLODSubsystem, Verbose, TEXT("Cell hidden - %s - showing %d HLOD actors"), *InCell->GetName(), CellData.LoadedHLODs.Num());
 #endif
 
-	for (AWorldPartitionHLOD* HLODActor : CellHLODs.LoadedHLODs)
+	for (AWorldPartitionHLOD* HLODActor : CellData.LoadedHLODs)
 	{
 #if WITH_EDITOR
 		UE_LOG(LogHLODSubsystem, Verbose, TEXT("\t\t%s - %s"), *HLODActor->GetActorLabel(), *HLODActor->GetActorGuid().ToString());
@@ -186,4 +206,141 @@ void UHLODSubsystem::OnCellHidden(const UWorldPartitionRuntimeCell* InCell)
 	}
 }
 
+void PrepareVTRequests(TMap<UMaterialInterface*, float>& InOutVTRequests, UStaticMeshComponent* InStaticMeshComponent, int32 InPixelSize)
+{
+	for (UMaterialInterface* MaterialInterface : InStaticMeshComponent->GetMaterials())
+	{
+		// We have a VT we'd like to prefetch, add or update a request in our request map.
+		// If the texture was already requested by another component, fetch the highest required resolution only.
+		float& CurrentMaxPixel = InOutVTRequests.FindOrAdd(MaterialInterface);
+		CurrentMaxPixel = FMath::Max(CurrentMaxPixel, InPixelSize);
+	}
+}
+
+static float EstimateScreenSize(UStaticMeshComponent* InStaticMeshComponent, const FSceneViewFamily& InViewFamily)
+{
+	float MaxPixels = 0;
+
+	// Estimate the highest screen pixel size of this component in the provided views
+	for (const FSceneView* View : InViewFamily.Views)
+	{
+		// Make sure the HLOD actor we're about to show is actually in the frustum
+		if (View->ViewFrustum.IntersectSphere(InStaticMeshComponent->Bounds.Origin, InStaticMeshComponent->Bounds.SphereRadius))
+		{
+			float ScreenDiameter = ComputeBoundsScreenSize(InStaticMeshComponent->Bounds.Origin, InStaticMeshComponent->Bounds.SphereRadius, *View);
+			float PixelSize = ScreenDiameter * View->ViewMatrices.GetScreenScale() * 2.0f;
+
+			MaxPixels = FMath::Max(MaxPixels, PixelSize);
+		}
+	}
+
+	return MaxPixels;
+}
+
+void UHLODSubsystem::MakeRenderResourcesResident(const FCellData& CellData, const FSceneViewFamily& InViewFamily)
+{
+	TRACE_CPUPROFILER_EVENT_SCOPE(UHLODSubsystem::MakeRenderResourcesResident)
+
+	TMap<UMaterialInterface*, float> VTRequests;
+
+	// For each HLOD actor representing this cell
+	for(AWorldPartitionHLOD* HLODActor : CellData.LoadedHLODs)
+	{
+		// Skip HLOD actors that doesn't require warmup.
+		// For example, instanced HLODs, as they reuse the same meshes/textures as their source actors.
+		// These resources should already be resident & at the proper resolution.
+		if (!HLODActor->DoesRequireWarmup())
+		{
+			continue;
+		}
+
+		HLODActor->ForEachComponent<UStaticMeshComponent>(false, [&](UStaticMeshComponent* SMC)
+		{
+			float PixelSize = EstimateScreenSize(SMC, InViewFamily);
+
+			if (PixelSize > 0)
+			{
+				PrepareVTRequests(VTRequests, SMC, PixelSize);
+				
+				// TODO: Nanite
+				//PrepareNaniteRequests(NaniteRequests, SMC, MaxPixels);
+			}
+		});
+	}
+
+	if (!VTRequests.IsEmpty())
+	{
+		ENQUEUE_RENDER_COMMAND(MakeHLODRenderResourcesResident)(
+			[VTRequests = MoveTemp(VTRequests), FeatureLevel = InViewFamily.GetFeatureLevel()](FRHICommandListImmediate& RHICmdList)
+			{
+				for(const TPair<UMaterialInterface*, float>& VTRequest : VTRequests)
+				{
+					UMaterialInterface* Material = VTRequest.Key;
+					FMaterialRenderProxy* MaterialRenderProxy = Material->GetRenderProxy();
+
+					GetRendererModule().RequestVirtualTextureTiles(MaterialRenderProxy, FVector2D(VTRequest.Value, VTRequest.Value), FeatureLevel);
+				}
+			});
+	}
+}
+
+bool UHLODSubsystem::RequestUnloading(const UWorldPartitionRuntimeCell* InCell)
+{
+	// Test if warmup is disabled globally.
+	bool bPerformWarmpup = CVarHLODWarmupEnabled.GetValueOnGameThread() != 0 && CVarHLODWarmupNumFrames.GetValueOnGameThread() > 0;
+	if (!bPerformWarmpup)
+	{
+		return true;
+	}
+
+	FCellData& CellData = CellsData.FindChecked(InCell->GetFName());
+
+	// If cell wasn't even visible yet or doesn't have HLOD actors, skip warmup.
+	bPerformWarmpup = !CellData.LoadedHLODs.IsEmpty() && CellData.bIsCellVisible;
+	if (!bPerformWarmpup)
+	{
+		return true;
+	}
+
+	// Verify that at least one HLOD actor associated with this cell needs warmup.
+	bPerformWarmpup = Algo::AnyOf(CellData.LoadedHLODs, [](const AWorldPartitionHLOD* HLODActor) { return HLODActor->DoesRequireWarmup(); });
+	if (!bPerformWarmpup)
+	{
+		return true;
+	}
+
+	uint32 CurrentFrameNumber = GetWorld()->Scene->GetFrameNumber();
+
+	// Trigger warmup on the first request to unload
+	if (CellData.WarmupCompleteFrame == INDEX_NONE)
+	{
+		CellData.WarmupCompleteFrame = CurrentFrameNumber + CVarHLODWarmupNumFrames.GetValueOnGameThread();
+		CellsToWarmup.Add(&CellData);
+	}
+
+	// Test if warmup is completed
+	bool bCanUnload = CurrentFrameNumber >= CellData.WarmupCompleteFrame;
+	if (bCanUnload)
+	{
+		CellData.WarmupCompleteFrame = INDEX_NONE;
+	}
+
+	return bCanUnload;
+}
+
+void UHLODSubsystem::OnBeginRenderViews(const FSceneViewFamily& InViewFamily)
+{
+	for (FCellData* CellPendingUnload : CellsToWarmup)
+	{
+		MakeRenderResourcesResident(*CellPendingUnload, InViewFamily);
+	}
+
+	CellsToWarmup.Reset();
+}
+
+void FHLODResourcesResidencySceneViewExtension::BeginRenderViewFamily(FSceneViewFamily& InViewFamily)
+{
+	GetWorld()->GetSubsystem<UHLODSubsystem>()->OnBeginRenderViews(InViewFamily);
+}
+
 #undef LOCTEXT_NAMESPACE
diff --git a/Engine/Source/Runtime/Engine/Private/WorldPartition/WorldPartitionRuntimeLevelStreamingCell.cpp b/Engine/Source/Runtime/Engine/Private/WorldPartition/WorldPartitionRuntimeLevelStreamingCell.cpp
--- a/Engine/Source/Runtime/Engine/Private/WorldPartition/WorldPartitionRuntimeLevelStreamingCell.cpp
+++ b/Engine/Source/Runtime/Engine/Private/WorldPartition/WorldPartitionRuntimeLevelStreamingCell.cpp
@@ -314,6 +314,12 @@ ULevel* UWorldPartitionRuntimeLevelStreamingCell::GetLevel() const
 	return LevelStreaming ? LevelStreaming->GetLoadedLevel() : nullptr;
 }
 
+bool UWorldPartitionRuntimeLevelStreamingCell::CanUnload() const
+{
+	check(LevelStreaming);
+	return LevelStreaming->GetWorld()->GetSubsystem<UHLODSubsystem>()->RequestUnloading(this);
+}
+
 void UWorldPartitionRuntimeLevelStreamingCell::Unload() const
 {
 #if WITH_EDITOR
diff --git a/Engine/Source/Runtime/Engine/Private/WorldPartition/WorldPartitionStreamingPolicy.cpp b/Engine/Source/Runtime/Engine/Private/WorldPartition/WorldPartitionStreamingPolicy.cpp
--- a/Engine/Source/Runtime/Engine/Private/WorldPartition/WorldPartitionStreamingPolicy.cpp
+++ b/Engine/Source/Runtime/Engine/Private/WorldPartition/WorldPartitionStreamingPolicy.cpp
@@ -558,10 +558,13 @@ void UWorldPartitionStreamingPolicy::SetCellsStateToUnloaded(const TSet<const UW
 
 	for (const UWorldPartitionRuntimeCell* Cell : ToUnloadCells)
 	{
-		UE_LOG(LogWorldPartition, Verbose, TEXT("UWorldPartitionStreamingPolicy::UnloadCells %s"), *Cell->GetName());
-		Cell->Unload();
-		ActivatedCells.Remove(Cell);
-		LoadedCells.Remove(Cell);
+		if (Cell->CanUnload())
+		{
+			UE_LOG(LogWorldPartition, Verbose, TEXT("UWorldPartitionStreamingPolicy::UnloadCells %s"), *Cell->GetName());
+			Cell->Unload();
+			ActivatedCells.Remove(Cell);
+			LoadedCells.Remove(Cell);
+		}
 	}
 }
 
diff --git a/Engine/Source/Runtime/Engine/Public/WorldPartition/HLOD/HLODActor.h b/Engine/Source/Runtime/Engine/Public/WorldPartition/HLOD/HLODActor.h
--- a/Engine/Source/Runtime/Engine/Public/WorldPartition/HLOD/HLODActor.h
+++ b/Engine/Source/Runtime/Engine/Public/WorldPartition/HLOD/HLODActor.h
@@ -24,15 +24,19 @@ public:
 
 	virtual bool IsHLODRelevant() const override { return true; }
 
+	bool DoesRequireWarmup() const { return bRequireWarmup; }
+
 #if WITH_EDITOR
 	void SetHLODPrimitives(const TArray<UPrimitiveComponent*>& InHLODPrimitives);
 
 	void SetSubActors(const TArray<FGuid>& InSubActors);
 	const TArray<FGuid>& GetSubActors() const;
 
-	void SetSubActorsHLODLayer(const UHLODLayer* InSubActorsHLODLayer) { SubActorsHLODLayer = InSubActorsHLODLayer; }
+	void SetSubActorsHLODLayer(const UHLODLayer* InSubActorsHLODLayer);
 	const UHLODLayer* GetSubActorsHLODLayer() const { return SubActorsHLODLayer; }
 
+	void SetRequireWarmup(bool InRequireWarmup) { bRequireWarmup = InRequireWarmup; }
+
 	void SetGridIndices(uint64 InGridIndexX, uint64 InGridIndexY, uint64 InGridIndexZ)
 	{
 		GridIndexX = InGridIndexX;
@@ -114,6 +118,9 @@ private:
 	UPROPERTY()
 	uint32 LODLevel;
 
+	UPROPERTY()
+	bool bRequireWarmup;
+
 	UPROPERTY()
 	TSoftObjectPtr<UWorldPartitionRuntimeCell> SourceCell_DEPRECATED;
 
diff --git a/Engine/Source/Runtime/Engine/Public/WorldPartition/HLOD/HLODBuilder.h b/Engine/Source/Runtime/Engine/Public/WorldPartition/HLOD/HLODBuilder.h
--- a/Engine/Source/Runtime/Engine/Public/WorldPartition/HLOD/HLODBuilder.h
+++ b/Engine/Source/Runtime/Engine/Public/WorldPartition/HLOD/HLODBuilder.h
@@ -38,6 +38,8 @@ public:
 
 	virtual bool RequiresCompiledAssets() const;
 
+	virtual bool RequiresWarmup() const;
+
 	virtual TArray<UPrimitiveComponent*> CreateComponents(AWorldPartitionHLOD* InHLODActor, const UHLODLayer* InHLODLayer, const TArray<UPrimitiveComponent*>& InSubComponents) const;
 
 	void Build(AWorldPartitionHLOD* InHLODActor, const UHLODLayer* InHLODLayer, const TArray<FWorldPartitionReference>& InSubActors);
diff --git a/Engine/Source/Runtime/Engine/Public/WorldPartition/HLOD/HLODLayer.h b/Engine/Source/Runtime/Engine/Public/WorldPartition/HLOD/HLODLayer.h
--- a/Engine/Source/Runtime/Engine/Public/WorldPartition/HLOD/HLODLayer.h
+++ b/Engine/Source/Runtime/Engine/Public/WorldPartition/HLOD/HLODLayer.h
@@ -53,6 +53,8 @@ public:
 	const TSoftObjectPtr<UHLODLayer>& GetParentLayer() const;
 	bool IsAlwaysLoaded() const { return bAlwaysLoaded; }
 
+	bool DoesRequireWarmup() const { return HLODBuilderClass ? HLODBuilderClass->GetDefaultObject<UHLODBuilder>()->RequiresWarmup() : true; }
+
 	static FName GetRuntimeGridName(uint32 InLODLevel, int32 InCellSize, double InLoadingRange);
 #endif
 
diff --git a/Engine/Source/Runtime/Engine/Public/WorldPartition/HLOD/HLODSubsystem.h b/Engine/Source/Runtime/Engine/Public/WorldPartition/HLOD/HLODSubsystem.h
--- a/Engine/Source/Runtime/Engine/Public/WorldPartition/HLOD/HLODSubsystem.h
+++ b/Engine/Source/Runtime/Engine/Public/WorldPartition/HLOD/HLODSubsystem.h
@@ -5,6 +5,7 @@
 #include "Containers/Array.h"
 #include "Containers/Map.h"
 #include "UObject/NameTypes.h"
+#include "SceneViewExtension.h"
 
 #include "HLODSubsystem.generated.h"
 
@@ -12,6 +13,20 @@ class UWorldPartition;
 class UWorldPartitionRuntimeCell;
 class AWorldPartitionHLOD;
 
+class FHLODResourcesResidencySceneViewExtension : public FWorldSceneViewExtension
+{
+public:
+	FHLODResourcesResidencySceneViewExtension(const FAutoRegister& AutoRegister, UWorld* InWorld)
+		: FWorldSceneViewExtension(AutoRegister, InWorld)
+	{
+	}
+
+	virtual void SetupViewFamily(FSceneViewFamily& InViewFamily) override {}
+	virtual void SetupView(FSceneViewFamily& InViewFamily, FSceneView& InView) override {}
+
+	virtual void BeginRenderViewFamily(FSceneViewFamily& InViewFamily) override;
+};
+
 /**
  * UHLODSubsystem
  */
@@ -35,29 +50,39 @@ public:
 	void OnCellShown(const UWorldPartitionRuntimeCell* InCell);
 	void OnCellHidden(const UWorldPartitionRuntimeCell* InCell);
 
+	bool RequestUnloading(const UWorldPartitionRuntimeCell* InCell);
+
 	static bool IsHLODEnabled();
 	
 private:
-	void OnWorldPartitionRegistered(UWorldPartition* InWorldPartition);
-	void OnWorldPartitionUnregistered(UWorldPartition* InWorldPartition);
-
-	/** Console command used to turn on/off loading & rendering of world partition HLODs */
-	static class FAutoConsoleCommand EnableHLODCommand;
-
-private:
-
-	static bool WorldPartitionHLODEnabled;
-
-	struct FCellHLODMapping
+	struct FCellData
 	{
 		bool						bIsCellVisible;
 		TSet<AWorldPartitionHLOD*>	LoadedHLODs;
-		
 
-		FCellHLODMapping() : bIsCellVisible(false)
+		uint32						WarmupCompleteFrame;
+
+		FCellData() 
+			: bIsCellVisible(false)
+			, WarmupCompleteFrame(INDEX_NONE)
 		{
 		}
 	};
 
-	TMap<FName, FCellHLODMapping> CellsHLODMapping;
+	TMap<FName, FCellData> CellsData;
+	TSet<FCellData*> CellsToWarmup;
+
+	void OnWorldPartitionRegistered(UWorldPartition* InWorldPartition);
+	void OnWorldPartitionUnregistered(UWorldPartition* InWorldPartition);
+	void OnBeginRenderViews(const FSceneViewFamily& InViewFamily);
+
+	void MakeRenderResourcesResident(const FCellData& CellHLODs, const FSceneViewFamily& InViewFamily);
+
+	/** Console command used to turn on/off loading & rendering of world partition HLODs */
+	static class FAutoConsoleCommand EnableHLODCommand;
+
+	static bool WorldPartitionHLODEnabled;
+
+	friend class FHLODResourcesResidencySceneViewExtension;
+	TSharedPtr<FHLODResourcesResidencySceneViewExtension, ESPMode::ThreadSafe> SceneViewExtension;
 };
diff --git a/Engine/Source/Runtime/Engine/Public/WorldPartition/WorldPartitionRuntimeCell.h b/Engine/Source/Runtime/Engine/Public/WorldPartition/WorldPartitionRuntimeCell.h
--- a/Engine/Source/Runtime/Engine/Public/WorldPartition/WorldPartitionRuntimeCell.h
+++ b/Engine/Source/Runtime/Engine/Public/WorldPartition/WorldPartitionRuntimeCell.h
@@ -122,6 +122,7 @@ class UWorldPartitionRuntimeCell : public UObject
 
 	virtual void Load() const PURE_VIRTUAL(UWorldPartitionRuntimeCell::Load,);
 	virtual void Unload() const PURE_VIRTUAL(UWorldPartitionRuntimeCell::Unload,);
+	virtual bool CanUnload() const PURE_VIRTUAL(UWorldPartitionRuntimeCell::CanUnload, return true;);
 	virtual void Activate() const PURE_VIRTUAL(UWorldPartitionRuntimeCell::Activate,);
 	virtual void Deactivate() const PURE_VIRTUAL(UWorldPartitionRuntimeCell::Deactivate,);
 	virtual bool IsAddedToWorld() const PURE_VIRTUAL(UWorldPartitionRuntimeCell::IsAddedToWorld, return false;);
diff --git a/Engine/Source/Runtime/Engine/Public/WorldPartition/WorldPartitionRuntimeLevelStreamingCell.h b/Engine/Source/Runtime/Engine/Public/WorldPartition/WorldPartitionRuntimeLevelStreamingCell.h
--- a/Engine/Source/Runtime/Engine/Public/WorldPartition/WorldPartitionRuntimeLevelStreamingCell.h
+++ b/Engine/Source/Runtime/Engine/Public/WorldPartition/WorldPartitionRuntimeLevelStreamingCell.h
@@ -13,6 +13,7 @@ class UWorldPartitionRuntimeLevelStreamingCell : public UWorldPartitionRuntimeSp
 	//~Begin UWorldPartitionRuntimeCell Interface
 	virtual void Load() const override;
 	virtual void Unload() const override;
+	virtual bool CanUnload() const override;
 	virtual void Activate() const override;
 	virtual void Deactivate() const override;
 	virtual bool IsAddedToWorld() const override;
diff --git a/Engine/Source/Runtime/RenderCore/Public/RendererInterface.h b/Engine/Source/Runtime/RenderCore/Public/RendererInterface.h
--- a/Engine/Source/Runtime/RenderCore/Public/RendererInterface.h
+++ b/Engine/Source/Runtime/RenderCore/Public/RendererInterface.h
@@ -33,6 +33,7 @@ struct FSceneTextures;
 class FShader;
 class FShaderMapPointerTable;
 class FRDGBuilder;
+class FMaterialRenderProxy;
 template<typename ShaderType, typename PointerTableType> class TShaderRefBase;
 
 // Shortcut for the allocator used by scene rendering.
@@ -860,6 +861,7 @@ public:
 	virtual void ReleaseVirtualTexturePendingResources() = 0;
 
 	virtual void RequestVirtualTextureTiles(const FVector2D& InScreenSpaceSize, int32 InMipLevel) = 0;
+	virtual void RequestVirtualTextureTiles(const FMaterialRenderProxy* InMaterialRenderProxy, const FVector2D& InScreenSpaceSize, ERHIFeatureLevel::Type InFeatureLevel) = 0;
 	virtual void RequestVirtualTextureTilesForRegion(IAllocatedVirtualTexture* AllocatedVT, const FVector2D& InScreenSpaceSize, const FVector2D& InViewportPosition, const FVector2D& InViewportSize, const FVector2D& InUV0, const FVector2D& InUV1, int32 InMipLevel) = 0;
 
 	/** Ensure that any tiles requested by 'RequestVirtualTextureTilesForRegion' are loaded, must be called from render thread */
diff --git a/Engine/Source/Runtime/Renderer/Private/RendererModule.h b/Engine/Source/Runtime/Renderer/Private/RendererModule.h
--- a/Engine/Source/Runtime/Renderer/Private/RendererModule.h
+++ b/Engine/Source/Runtime/Renderer/Private/RendererModule.h
@@ -101,6 +101,7 @@ public:
 	virtual void ReleaseVirtualTexturePendingResources() override;
 	virtual void RequestVirtualTextureTiles(TArray<uint64>&& InPageRequests) override;
 	virtual void RequestVirtualTextureTiles(const FVector2D& InScreenSpaceSize, int32 InMipLevel) override;
+	virtual void RequestVirtualTextureTiles(const FMaterialRenderProxy* InMaterialRenderProxy, const FVector2D& InScreenSpaceSize, ERHIFeatureLevel::Type InFeatureLevel) override;
 	virtual void RequestVirtualTextureTilesForRegion(IAllocatedVirtualTexture* AllocatedVT, const FVector2D& InScreenSpaceSize, const FVector2D& InViewportPosition, const FVector2D& InViewportSize, const FVector2D& InUV0, const FVector2D& InUV1, int32 InMipLevel) override;
 	virtual void LoadPendingVirtualTextureTiles(FRHICommandListImmediate& RHICmdList, ERHIFeatureLevel::Type FeatureLevel) override;
 	virtual void SetVirtualTextureRequestRecordBuffer(uint64 Handle) override;
diff --git a/Engine/Source/Runtime/Renderer/Private/SceneRendering.cpp b/Engine/Source/Runtime/Renderer/Private/SceneRendering.cpp
--- a/Engine/Source/Runtime/Renderer/Private/SceneRendering.cpp
+++ b/Engine/Source/Runtime/Renderer/Private/SceneRendering.cpp
@@ -4583,6 +4583,11 @@ void FRendererModule::RequestVirtualTextureTiles(const FVector2D& InScreenSpaceS
 	FVirtualTextureSystem::Get().RequestTiles(InScreenSpaceSize, InMipLevel);
 }
 
+void FRendererModule::RequestVirtualTextureTiles(const FMaterialRenderProxy* InMaterialRenderProxy, const FVector2D& InScreenSpaceSize, ERHIFeatureLevel::Type InFeatureLevel)
+{
+	FVirtualTextureSystem::Get().RequestTiles(InMaterialRenderProxy, InScreenSpaceSize, InFeatureLevel);
+}
+
 void FRendererModule::RequestVirtualTextureTilesForRegion(IAllocatedVirtualTexture* AllocatedVT, const FVector2D& InScreenSpaceSize, const FVector2D& InViewportPosition, const FVector2D& InViewportSize, const FVector2D& InUV0, const FVector2D& InUV1, int32 InMipLevel)
 {
 	FVirtualTextureSystem::Get().RequestTilesForRegion(AllocatedVT, InScreenSpaceSize, InViewportPosition, InViewportSize, InUV0, InUV1, InMipLevel);
diff --git a/Engine/Source/Runtime/Renderer/Private/VT/VirtualTextureSystem.cpp b/Engine/Source/Runtime/Renderer/Private/VT/VirtualTextureSystem.cpp
--- a/Engine/Source/Runtime/Renderer/Private/VT/VirtualTextureSystem.cpp
+++ b/Engine/Source/Runtime/Renderer/Private/VT/VirtualTextureSystem.cpp
@@ -6,6 +6,7 @@
 #include "Engine/Canvas.h"
 #include "Debug/DebugDrawService.h"
 #include "HAL/IConsoleManager.h"
+#include "MaterialShared.h"
 #include "PostProcess/SceneRenderTargets.h"
 #include "ProfilingDebugging/CsvProfiler.h"
 #include "ProfilingDebugging/CpuProfilerTrace.h"
@@ -1079,24 +1080,40 @@ void FVirtualTextureSystem::RequestTiles(const FVector2D& InScreenSpaceSize, int
 
 	for (const auto& Pair : AllocatedVTs)
 	{
-		const IAllocatedVirtualTexture* AllocatedVT = Pair.Value;
-		if (InMipLevel < 0)
-		{
-			const uint32 vMaxLevel = AllocatedVT->GetMaxLevel();
-			const float vLevel = ComputeMipLevel(AllocatedVT, InScreenSpaceSize);
-			const int32 vMipLevelDown = FMath::Clamp((int32)FMath::FloorToInt(vLevel), 0, (int32)vMaxLevel);
-			RequestTilesInternal(AllocatedVT, vLevel);
-			if (vMipLevelDown + 1u <= vMaxLevel)
-			{
-				// Need to fetch 2 levels to support trilinear filtering
-				RequestTilesInternal(AllocatedVT, vMipLevelDown);
-			}
-		}
-		else
+		RequestTilesInternal(Pair.Value, InScreenSpaceSize, InMipLevel);
+	}
+}
+
+void FVirtualTextureSystem::RequestTiles(const FMaterialRenderProxy* InMaterialRenderProxy, const FVector2D& InScreenSpaceSize, ERHIFeatureLevel::Type InFeatureLevel)
+{
+	check(IsInRenderingThread());
+
+	FScopeLock Lock(&RequestedTilesLock);
+
+	for (IAllocatedVirtualTexture* AllocatedVT : InMaterialRenderProxy->UniformExpressionCache[InFeatureLevel].AllocatedVTs)
+	{
+		RequestTilesInternal(AllocatedVT, InScreenSpaceSize, INDEX_NONE);
+	}
+}
+
+void FVirtualTextureSystem::RequestTilesInternal(const IAllocatedVirtualTexture* InAllocatedVT, const FVector2D& InScreenSpaceSize, int32 InMipLevel)
+{
+	if (InMipLevel < 0)
+	{
+		const uint32 vMaxLevel = InAllocatedVT->GetMaxLevel();
+		const float vLevel = ComputeMipLevel(InAllocatedVT, InScreenSpaceSize);
+		const int32 vMipLevelDown = FMath::Clamp((int32)FMath::FloorToInt(vLevel), 0, (int32)vMaxLevel);
+		RequestTilesInternal(InAllocatedVT, vLevel);
+		if (vMipLevelDown + 1u <= vMaxLevel)
 		{
-			RequestTilesInternal(AllocatedVT, InMipLevel);
+			// Need to fetch 2 levels to support trilinear filtering
+			RequestTilesInternal(InAllocatedVT, vMipLevelDown);
 		}
 	}
+	else
+	{
+		RequestTilesInternal(InAllocatedVT, InMipLevel);
+	}
 }
 
 void FVirtualTextureSystem::RequestTilesForRegion(IAllocatedVirtualTexture* AllocatedVT, const FVector2D& InScreenSpaceSize, const FVector2D& InViewportPosition, const FVector2D& InViewportSize, const FVector2D& InUV0, const FVector2D& InUV1, int32 InMipLevel)
diff --git a/Engine/Source/Runtime/Renderer/Private/VT/VirtualTextureSystem.h b/Engine/Source/Runtime/Renderer/Private/VT/VirtualTextureSystem.h
--- a/Engine/Source/Runtime/Renderer/Private/VT/VirtualTextureSystem.h
+++ b/Engine/Source/Runtime/Renderer/Private/VT/VirtualTextureSystem.h
@@ -70,6 +70,7 @@ public:
 	void UnlockTile(const FVirtualTextureLocalTile& Tile, const FVirtualTextureProducer* Producer);
 	void ForceUnlockAllTiles(const FVirtualTextureProducerHandle& ProducerHandle, const FVirtualTextureProducer* Producer);
 	void RequestTiles(const FVector2D& InScreenSpaceSize, int32 InMipLevel = -1);
+	void RequestTiles(const FMaterialRenderProxy* InMaterialRenderProxy, const FVector2D& InScreenSpaceSize, ERHIFeatureLevel::Type InFeatureLevel);
 	void RequestTilesForRegion(IAllocatedVirtualTexture* AllocatedVT, const FVector2D& InScreenSpaceSize, const FVector2D& InViewportPosition, const FVector2D& InViewportSize, const FVector2D& InUV0, const FVector2D& InUV1, int32 InMipLevel = -1);
 	void LoadPendingTiles(FRDGBuilder& GraphBuilder, ERHIFeatureLevel::Type FeatureLevel);
 	
@@ -97,6 +98,7 @@ private:
 
 	void RequestTilesForRegionInternal(const IAllocatedVirtualTexture* AllocatedVT, const FVector2D& InScreenSpaceSize, const FVector2D& InViewportPosition, const FVector2D& InViewportSize, const FVector2D& InUV0, const FVector2D& InUV1, int32 InMipLevel);
 	void RequestTilesInternal(const IAllocatedVirtualTexture* AllocatedVT, int32 InMipLevel);
+	void RequestTilesInternal(const IAllocatedVirtualTexture* AllocatedVT, const FVector2D& InScreenSpaceSize, int32 InMipLevel);
 	
 	void SubmitRequestsFromLocalTileList(TArray<FVirtualTextureLocalTile>& OutDeferredTiles, const TSet<FVirtualTextureLocalTile>& LocalTileList, EVTProducePageFlags Flags, FRDGBuilder& GraphBuilder, ERHIFeatureLevel::Type FeatureLevel);
 
