Fix crash when reloading a partitioned world.

diff --git a/Engine/Source/Runtime/Engine/Private/WorldPartition/WorldPartitionRuntimeLevelStreamingCell.cpp b/Engine/Source/Runtime/Engine/Private/WorldPartition/WorldPartitionRuntimeLevelStreamingCell.cpp
--- a/Engine/Source/Runtime/Engine/Private/WorldPartition/WorldPartitionRuntimeLevelStreamingCell.cpp
+++ b/Engine/Source/Runtime/Engine/Private/WorldPartition/WorldPartitionRuntimeLevelStreamingCell.cpp
@@ -316,8 +316,7 @@ ULevel* UWorldPartitionRuntimeLevelStreamingCell::GetLevel() const
 
 bool UWorldPartitionRuntimeLevelStreamingCell::CanUnload() const
 {
-	check(LevelStreaming);
-	return LevelStreaming->GetWorld()->GetSubsystem<UHLODSubsystem>()->RequestUnloading(this);
+	return LevelStreaming ? LevelStreaming->GetWorld()->GetSubsystem<UHLODSubsystem>()->RequestUnloading(this) : true;
 }
 
 void UWorldPartitionRuntimeLevelStreamingCell::Unload() const
