For VT texture sampling only select shared anisotropic sampler if anisotropy is enabled.
Touches material ddc key. Filled DDC in preflight.

diff --git a/Engine/Source/Runtime/Engine/Private/Materials/HLSLMaterialTranslator.cpp b/Engine/Source/Runtime/Engine/Private/Materials/HLSLMaterialTranslator.cpp
--- a/Engine/Source/Runtime/Engine/Private/Materials/HLSLMaterialTranslator.cpp
+++ b/Engine/Source/Runtime/Engine/Private/Materials/HLSLMaterialTranslator.cpp
@@ -3,6 +3,7 @@
 	HLSLMaterialTranslator.cpp: Translates material expressions into HLSL code.
 =============================================================================*/
 #include "HLSLMaterialTranslator.h"
+#include "VT/VirtualTextureScalability.h"
 
 #if WITH_EDITORONLY_DATA
 
@@ -5734,8 +5735,10 @@ int32 FHLSLMaterialTranslator::TextureSample(
 
 		if (SamplerSource != SSM_FromTextureAsset)
 		{
-			// VT doesn't care if the shared sampler is wrap or clamp this is handled in the shader explicitly by our code so we still inherit this from the texture
-			const TCHAR* SharedSamplerName = (MipValueMode == TMVM_MipLevel) ? TEXT("View.SharedBilinearClampedSampler") : TEXT("View.SharedBilinearAnisoClampedSampler");
+			// VT doesn't care if the shared sampler is wrap or clamp. It only cares if it is aniso or not.
+			// The wrap/clamp/mirror operation is handled in the shader explicitly.
+			const bool bUseAnisoSampler = VirtualTextureScalability::IsAnisotropicFilteringEnabled() && MipValueMode != TMVM_MipLevel;
+			const TCHAR* SharedSamplerName = bUseAnisoSampler ? TEXT("View.SharedBilinearAnisoClampedSampler") : TEXT("View.SharedBilinearClampedSampler");
 			TextureName += FString::Printf(TEXT("Material.VirtualTexturePhysical_%d, GetMaterialSharedSampler(Material.VirtualTexturePhysical_%dSampler, %s)")
 				, VirtualTextureIndex, VirtualTextureIndex, SharedSamplerName);
 		}
diff --git a/Engine/Source/Runtime/Engine/Private/ShaderDerivedDataVersion.h b/Engine/Source/Runtime/Engine/Private/ShaderDerivedDataVersion.h
--- a/Engine/Source/Runtime/Engine/Private/ShaderDerivedDataVersion.h
+++ b/Engine/Source/Runtime/Engine/Private/ShaderDerivedDataVersion.h
@@ -10,5 +10,5 @@
 
 // In case of merge conflicts with DDC versions, you MUST generate a new GUID and set this new
 // guid as version
-#define GLOBALSHADERMAP_DERIVEDDATA_VER			TEXT("8D3A12924CDF4D9F84C0A900E9445CAD")
-#define MATERIALSHADERMAP_DERIVEDDATA_VER		TEXT("1DA85DA5733C4A489A1BE117CDC5ACCC")
\ No newline at end of file
+#define GLOBALSHADERMAP_DERIVEDDATA_VER			TEXT("2EE7C1F4F6C74732951E5B47720E1CB9")
+#define MATERIALSHADERMAP_DERIVEDDATA_VER		TEXT("321C290437B44458A11C57A76F63914E")
\ No newline at end of file
