[Nanite] - enable visualizer always, fixes missing in shipping builds bug

diff --git a/Engine/Source/Runtime/Engine/Public/ShowFlagsValues.inl b/Engine/Source/Runtime/Engine/Public/ShowFlagsValues.inl
--- a/Engine/Source/Runtime/Engine/Public/ShowFlagsValues.inl
+++ b/Engine/Source/Runtime/Engine/Public/ShowFlagsValues.inl
@@ -63,7 +63,7 @@ SHOWFLAG_FIXED_IN_SHIPPING(0, ReflectionOverride, SFG_Hidden, NSLOCTEXT("UnrealE
 /** needed for VMI_VisualizeBuffer, Whether to enable the buffer visualization mode. */
 SHOWFLAG_FIXED_IN_SHIPPING(0, VisualizeBuffer, SFG_Hidden, NSLOCTEXT("UnrealEd", "VisualizeBufferSF", "Buffer Visualization"))
 /** Needed for VMI_VisualizeNanite, Whether to enable the Nanite visualization mode. */
-SHOWFLAG_FIXED_IN_SHIPPING(0, VisualizeNanite, SFG_Hidden, NSLOCTEXT("UnrealEd", "VisualizeNaniteSF", "Nanite Visualization"))
+SHOWFLAG_ALWAYS_ACCESSIBLE(VisualizeNanite, SFG_Hidden, NSLOCTEXT("UnrealEd", "VisualizeNaniteSF", "Nanite Visualization"))
 /** Allows to disable all direct lighting (does not affect indirect light) */
 SHOWFLAG_FIXED_IN_SHIPPING(1, DirectLighting, SFG_LightingComponents, NSLOCTEXT("UnrealEd", "DirectLightingSF", "Direct Lighting"))
 /** Allows to disable lighting from Directional Lights */
