Implement Mip Adjustments in ParticleSubUV material expression.

diff --git a/Engine/Source/Runtime/Engine/Private/Materials/HLSLMaterialTranslator.cpp b/Engine/Source/Runtime/Engine/Private/Materials/HLSLMaterialTranslator.cpp
--- a/Engine/Source/Runtime/Engine/Private/Materials/HLSLMaterialTranslator.cpp
+++ b/Engine/Source/Runtime/Engine/Private/Materials/HLSLMaterialTranslator.cpp
@@ -4863,7 +4863,7 @@ int32 FHLSLMaterialTranslator::ParticleMacroUV()
 	return AddCodeChunk(MCT_Float2,TEXT("GetParticleMacroUV(Parameters)"));
 }
 
-int32 FHLSLMaterialTranslator::ParticleSubUV(int32 TextureIndex, EMaterialSamplerType SamplerType, bool bBlend)
+int32 FHLSLMaterialTranslator::ParticleSubUV(int32 TextureIndex, EMaterialSamplerType SamplerType, int32 MipValue0Index, int32 MipValue1Index, ETextureMipValueMode MipValueMode, bool bBlend)
 {
 	if (ShaderFrequency != SF_Pixel && ShaderFrequency != SF_Compute)
 	{
@@ -4889,13 +4889,13 @@ int32 FHLSLMaterialTranslator::ParticleSubUV(int32 TextureIndex, EMaterialSample
 		const int32 TexCoord2 = AddCodeChunk( MCT_Float2,*TexCoordCode,1);
 		const int32 SubImageLerp = AddCodeChunk(MCT_Float, TEXT("Parameters.Particle.SubUVLerp"));
 
-		const int32 TexSampleA = TextureSample(TextureIndex, TexCoord1, SamplerType );
-		const int32 TexSampleB = TextureSample(TextureIndex, TexCoord2, SamplerType );
+		const int32 TexSampleA = TextureSample(TextureIndex, TexCoord1, SamplerType, MipValue0Index, MipValue1Index, MipValueMode);
+		const int32 TexSampleB = TextureSample(TextureIndex, TexCoord2, SamplerType, MipValue0Index, MipValue1Index, MipValueMode);
 		ParticleSubUV = Lerp( TexSampleA,TexSampleB, SubImageLerp);
 	} 
 	else
 	{
-		ParticleSubUV = TextureSample(TextureIndex, TexCoord1, SamplerType );
+		ParticleSubUV = TextureSample(TextureIndex, TexCoord1, SamplerType, MipValue0Index, MipValue1Index, MipValueMode);
 	}
 	
 	bUsesParticleSubUVs = true;
diff --git a/Engine/Source/Runtime/Engine/Private/Materials/HLSLMaterialTranslator.h b/Engine/Source/Runtime/Engine/Private/Materials/HLSLMaterialTranslator.h
--- a/Engine/Source/Runtime/Engine/Private/Materials/HLSLMaterialTranslator.h
+++ b/Engine/Source/Runtime/Engine/Private/Materials/HLSLMaterialTranslator.h
@@ -750,7 +750,7 @@ protected:
 	virtual int32 GetPixelPosition() override;
 
 	virtual int32 ParticleMacroUV() override;
-	virtual int32 ParticleSubUV(int32 TextureIndex, EMaterialSamplerType SamplerType, bool bBlend) override;
+	virtual int32 ParticleSubUV(int32 TextureIndex, EMaterialSamplerType SamplerType, int32 MipValue0Index, int32 MipValue1Index, ETextureMipValueMode MipValueMode, bool bBlend) override;
 	virtual int32 ParticleSubUVProperty(int32 PropertyIndex) override;
 	virtual int32 ParticleColor() override;
 	virtual int32 ParticlePosition() override;
diff --git a/Engine/Source/Runtime/Engine/Private/Materials/MaterialExpressions.cpp b/Engine/Source/Runtime/Engine/Private/Materials/MaterialExpressions.cpp
--- a/Engine/Source/Runtime/Engine/Private/Materials/MaterialExpressions.cpp
+++ b/Engine/Source/Runtime/Engine/Private/Materials/MaterialExpressions.cpp
@@ -3736,9 +3736,16 @@ void UMaterialExpressionTextureSampleParameterVolume::SetDefaultTexture()
  * Performs a SubUV operation, which is doing a texture lookup into a sub rectangle of a texture, and optionally blending with another rectangle.  
  * This supports both sprites and mesh emitters.
  */
-static int32 ParticleSubUV(FMaterialCompiler* Compiler, int32 TextureIndex, UTexture* DefaultTexture, EMaterialSamplerType SamplerType, FExpressionInput& Coordinates, bool bBlend)
+static int32 ParticleSubUV(
+	FMaterialCompiler* Compiler, 
+	int32 TextureIndex, 
+	EMaterialSamplerType SamplerType, 
+	int32 MipValue0Index,
+	int32 MipValue1Index,
+	ETextureMipValueMode MipValueMode,
+	bool bBlend)
 {
-	return Compiler->ParticleSubUV(TextureIndex, SamplerType, bBlend);
+	return Compiler->ParticleSubUV(TextureIndex, SamplerType, MipValue0Index, MipValue1Index, MipValueMode, bBlend);
 }
 
 /** 
@@ -3781,7 +3788,7 @@ int32 UMaterialExpressionTextureSampleParameterSubUV::Compile(class FMaterialCom
 	}
 
 	int32 TextureCodeIndex = Compiler->TextureParameter(ParameterName, Texture, SamplerType);
-	return ParticleSubUV(Compiler, TextureCodeIndex, Texture, SamplerType, Coordinates, bBlend);
+	return ParticleSubUV(Compiler, TextureCodeIndex, SamplerType, CompileMipValue0(Compiler), CompileMipValue1(Compiler), MipValueMode,	bBlend);
 }
 
 void UMaterialExpressionTextureSampleParameterSubUV::GetCaption(TArray<FString>& OutCaptions) const
@@ -9373,7 +9380,7 @@ int32 UMaterialExpressionParticleSubUV::Compile(class FMaterialCompiler* Compile
 			return Compiler->Errorf(TEXT("%s"), *SamplerTypeError);
 		}
 		int32 TextureCodeIndex = Compiler->Texture(Texture, SamplerType);
-		return ParticleSubUV(Compiler, TextureCodeIndex, Texture, SamplerType, Coordinates, bBlend);
+		return ParticleSubUV(Compiler, TextureCodeIndex, SamplerType, CompileMipValue0(Compiler), CompileMipValue1(Compiler), MipValueMode, bBlend);
 	}
 	else
 	{
diff --git a/Engine/Source/Runtime/Engine/Public/MaterialCompiler.h b/Engine/Source/Runtime/Engine/Public/MaterialCompiler.h
--- a/Engine/Source/Runtime/Engine/Public/MaterialCompiler.h
+++ b/Engine/Source/Runtime/Engine/Public/MaterialCompiler.h
@@ -211,7 +211,7 @@ public:
 	virtual int32 DistanceCullFade() = 0;
 	virtual int32 ActorWorldPosition() = 0;
 	virtual int32 ParticleMacroUV() = 0;
-	virtual int32 ParticleSubUV(int32 TextureIndex, EMaterialSamplerType SamplerType, bool bBlend) = 0;
+	virtual int32 ParticleSubUV(int32 TextureIndex, EMaterialSamplerType SamplerType, int32 MipValue0Index, int32 MipValue1Index, ETextureMipValueMode MipValueMode, bool bBlend) = 0;
 	virtual int32 ParticleSubUVProperty(int32 PropertyIndex) = 0;
 	virtual int32 ParticleColor() = 0;
 	virtual int32 ParticlePosition() = 0;
@@ -608,7 +608,10 @@ public:
 	virtual int32 DistanceCullFade() override { return Compiler->DistanceCullFade(); }
 	virtual int32 ActorWorldPosition() override { return Compiler->ActorWorldPosition(); }
 	virtual int32 ParticleMacroUV() override { return Compiler->ParticleMacroUV(); }
-	virtual int32 ParticleSubUV(int32 TextureIndex, EMaterialSamplerType SamplerType, bool bBlend) override { return Compiler->ParticleSubUV(TextureIndex, SamplerType, bBlend); }
+	virtual int32 ParticleSubUV(int32 TextureIndex, EMaterialSamplerType SamplerType, int32 MipValue0Index, int32 MipValue1Index, ETextureMipValueMode MipValueMode, bool bBlend) override 
+	{
+		return Compiler->ParticleSubUV(TextureIndex, SamplerType, MipValue0Index, MipValue1Index, MipValueMode,	bBlend);
+	}
 	virtual int32 ParticleSubUVProperty(int32 PropertyIndex) override { return Compiler->ParticleSubUVProperty(PropertyIndex); }
 	virtual int32 ParticleColor() override { return Compiler->ParticleColor(); }
 	virtual int32 ParticlePosition() override { return Compiler->ParticlePosition(); }
