Thread safety for the LightweightInstancesSubsystem singleton

#autho: peter.knepley

diff --git a/Engine/Source/Runtime/Engine/Classes/GameFramework/LightWeightInstanceSubsystem.h b/Engine/Source/Runtime/Engine/Classes/GameFramework/LightWeightInstanceSubsystem.h
--- a/Engine/Source/Runtime/Engine/Classes/GameFramework/LightWeightInstanceSubsystem.h
+++ b/Engine/Source/Runtime/Engine/Classes/GameFramework/LightWeightInstanceSubsystem.h
@@ -24,8 +24,12 @@ struct ENGINE_API FLightWeightInstanceSubsystem
 	friend struct FActorInstanceHandle;
 	friend class ALightWeightInstanceManager;
 
+	static FCriticalSection GetFunctionCS;
+
 	static FLightWeightInstanceSubsystem& Get()
 	{
+		FScopeLock Lock(&GetFunctionCS);
+
 		if (!LWISubsystem)
 		{
 			LWISubsystem = MakeShareable(new FLightWeightInstanceSubsystem());
@@ -106,13 +110,12 @@ protected:
 
 private:
 	/** Application singleton */
-	static TSharedPtr<FLightWeightInstanceSubsystem> LWISubsystem;
+	static TSharedPtr<FLightWeightInstanceSubsystem, ESPMode::ThreadSafe> LWISubsystem;
 
 	// TODO: preallocate the size of this based on a config variable
 	UPROPERTY()
 	TArray<ALightWeightInstanceManager*> LWInstanceManagers;
 
-
 #ifdef WITH_EDITOR
 private:
 	FDelegateHandle OnLevelActorAddedHandle;
diff --git a/Engine/Source/Runtime/Engine/Private/LightWeightInstanceSubsystem.cpp b/Engine/Source/Runtime/Engine/Private/LightWeightInstanceSubsystem.cpp
--- a/Engine/Source/Runtime/Engine/Private/LightWeightInstanceSubsystem.cpp
+++ b/Engine/Source/Runtime/Engine/Private/LightWeightInstanceSubsystem.cpp
@@ -9,6 +9,7 @@
 DEFINE_LOG_CATEGORY(LogLightWeightInstance);
 
 TSharedPtr<FLightWeightInstanceSubsystem> FLightWeightInstanceSubsystem::LWISubsystem;
+FCriticalSection FLightWeightInstanceSubsystem::GetFunctionCS;
 
 FLightWeightInstanceSubsystem::FLightWeightInstanceSubsystem()
 {
