HLOD Warmup
* Enable warmup feature - No impact on memory (validated on XSS with 5 different runs with/without the feature)
* Fixed VT warmup - must perform the tile requests every frame so that it is effective
* Added scale factor & clamping (as cVars) for the VT requests
* Added debugging tools
* Fixed issue in FVirtualTextureSystem::RequestTilesInternal() where the same mip could be requested twice
* Fixed non-unity build

diff --git a/Engine/Source/Runtime/Engine/Private/WorldPartition/HLOD/HLODSubsystem.cpp b/Engine/Source/Runtime/Engine/Private/WorldPartition/HLOD/HLODSubsystem.cpp
--- a/Engine/Source/Runtime/Engine/Private/WorldPartition/HLOD/HLODSubsystem.cpp
+++ b/Engine/Source/Runtime/Engine/Private/WorldPartition/HLOD/HLODSubsystem.cpp
@@ -11,6 +11,7 @@
 #include "UObject/UObjectHash.h"
 #include "UObject/UObjectGlobals.h"
 #include "Engine/Engine.h"
+#include "Engine/StaticMesh.h"
 #include "Engine/World.h"
 #include "EngineModule.h"
 #include "EngineUtils.h"
@@ -34,6 +35,27 @@ static TAutoConsoleVariable<int32> CVarHLODWarmupEnabled(
 	ECVF_Default
 );
 
+static TAutoConsoleVariable<int32> CVarHLODWarmupDebugDraw(
+	TEXT("wp.Runtime.HLOD.WarmupDebugDraw"),
+	0,
+	TEXT("Draw debug display for the warmup requests"),
+	ECVF_Default
+);
+
+static TAutoConsoleVariable<float> CVarHLODWarmupVTScaleFactor(
+	TEXT("wp.Runtime.HLOD.WarmupVTScaleFactor"),
+	2.0f,
+	TEXT("Scale the VT size we ask to prefetch by this factor."),
+	ECVF_Default
+);
+
+static TAutoConsoleVariable<int32> CVarHLODWarmupVTSizeClamp(
+	TEXT("wp.Runtime.HLOD.WarmupVTSizeClamp"),
+	2048,
+	TEXT("Clamp VT warmup requests for safety."),
+	ECVF_Default
+);
+
 UHLODSubsystem::UHLODSubsystem()
 	: UWorldSubsystem()
 {
@@ -208,12 +230,20 @@ void UHLODSubsystem::OnCellHidden(const UWorldPartitionRuntimeCell* InCell)
 
 static void PrepareVTRequests(TMap<UMaterialInterface*, float>& InOutVTRequests, UStaticMeshComponent* InStaticMeshComponent, float InPixelSize)
 {
+	float PixelSize = InPixelSize;
+
+	// Assume the texture is wrapped around the object, so the screen size is actually less than the resolution we require.
+	PixelSize *= CVarHLODWarmupVTScaleFactor.GetValueOnAnyThread();
+
+	// Clamp for safety
+	PixelSize = FMath::Min(PixelSize, CVarHLODWarmupVTSizeClamp.GetValueOnAnyThread());
+
 	for (UMaterialInterface* MaterialInterface : InStaticMeshComponent->GetMaterials())
 	{
 		// We have a VT we'd like to prefetch, add or update a request in our request map.
 		// If the texture was already requested by another component, fetch the highest required resolution only.
 		float& CurrentMaxPixel = InOutVTRequests.FindOrAdd(MaterialInterface);
-		CurrentMaxPixel = FMath::Max(CurrentMaxPixel, InPixelSize);
+		CurrentMaxPixel = FMath::Max(CurrentMaxPixel, PixelSize);
 	}
 }
 
@@ -272,7 +302,20 @@ void UHLODSubsystem::MakeRenderResourcesResident(const FCellData& CellData, cons
 			if (PixelSize > 0)
 			{
 				PrepareVTRequests(VTRequests, SMC, PixelSize);
-				PrepareNaniteRequests(NaniteRequests, SMC);
+
+				// Only issue Nanite requests on the first warmup frame
+				if (CellData.WarmupStartFrame == InViewFamily.FrameNumber)
+				{
+					PrepareNaniteRequests(NaniteRequests, SMC);
+				}
+
+#if ENABLE_DRAW_DEBUG
+				if (CVarHLODWarmupDebugDraw.GetValueOnAnyThread())
+				{
+					const FBox& Box = SMC->CalcLocalBounds().GetBox();
+					DrawDebugBox(HLODActor->GetWorld(), Box.GetCenter(), Box.GetExtent(), FColor::Yellow, /*bPersistentLine*/ false, /*Lifetime*/ 1.0f);
+				}
+#endif
 			}
 		});
 	}
@@ -327,17 +370,20 @@ bool UHLODSubsystem::RequestUnloading(const UWorldPartitionRuntimeCell* InCell)
 	uint32 CurrentFrameNumber = GetWorld()->Scene->GetFrameNumber();
 
 	// Trigger warmup on the first request to unload
-	if (CellData.WarmupCompleteFrame == INDEX_NONE)
+	if (CellData.WarmupEndFrame == INDEX_NONE)
 	{
-		CellData.WarmupCompleteFrame = CurrentFrameNumber + CVarHLODWarmupNumFrames.GetValueOnGameThread();
+		// Warmup will be triggered in the next BeginRenderView() call, at which point the frame number will have been incremented.
+		CellData.WarmupStartFrame = CurrentFrameNumber + 1; 
+		CellData.WarmupEndFrame = CellData.WarmupStartFrame + CVarHLODWarmupNumFrames.GetValueOnGameThread();
 		CellsToWarmup.Add(&CellData);
 	}
 
 	// Test if warmup is completed
-	bool bCanUnload = CurrentFrameNumber >= CellData.WarmupCompleteFrame;
+	bool bCanUnload = CurrentFrameNumber >= CellData.WarmupEndFrame;
 	if (bCanUnload)
 	{
-		CellData.WarmupCompleteFrame = INDEX_NONE;
+		CellData.WarmupStartFrame = INDEX_NONE;
+		CellData.WarmupEndFrame = INDEX_NONE;
 	}
 
 	return bCanUnload;
@@ -345,12 +391,18 @@ bool UHLODSubsystem::RequestUnloading(const UWorldPartitionRuntimeCell* InCell)
 
 void UHLODSubsystem::OnBeginRenderViews(const FSceneViewFamily& InViewFamily)
 {
-	for (FCellData* CellPendingUnload : CellsToWarmup)
+	for (TSet<FCellData*>::TIterator CellIt(CellsToWarmup); CellIt; ++CellIt)
 	{
+		FCellData* CellPendingUnload = *CellIt;
+	
 		MakeRenderResourcesResident(*CellPendingUnload, InViewFamily);
-	}
 
-	CellsToWarmup.Reset();
+		// Stop processing this cell if warmup is done.
+		if (InViewFamily.FrameNumber >= CellPendingUnload->WarmupEndFrame)
+		{
+			CellIt.RemoveCurrent();
+		}
+	}
 }
 
 void FHLODResourcesResidencySceneViewExtension::BeginRenderViewFamily(FSceneViewFamily& InViewFamily)
diff --git a/Engine/Source/Runtime/Engine/Public/WorldPartition/HLOD/HLODSubsystem.h b/Engine/Source/Runtime/Engine/Public/WorldPartition/HLOD/HLODSubsystem.h
--- a/Engine/Source/Runtime/Engine/Public/WorldPartition/HLOD/HLODSubsystem.h
+++ b/Engine/Source/Runtime/Engine/Public/WorldPartition/HLOD/HLODSubsystem.h
@@ -60,11 +60,13 @@ private:
 		bool						bIsCellVisible;
 		TSet<AWorldPartitionHLOD*>	LoadedHLODs;
 
-		uint32						WarmupCompleteFrame;
+		uint32						WarmupStartFrame;
+		uint32						WarmupEndFrame;
 
 		FCellData() 
 			: bIsCellVisible(false)
-			, WarmupCompleteFrame(INDEX_NONE)
+			, WarmupStartFrame(INDEX_NONE)
+			, WarmupEndFrame(INDEX_NONE)
 		{
 		}
 	};
diff --git a/Engine/Source/Runtime/Renderer/Private/VT/VirtualTextureSystem.cpp b/Engine/Source/Runtime/Renderer/Private/VT/VirtualTextureSystem.cpp
--- a/Engine/Source/Runtime/Renderer/Private/VT/VirtualTextureSystem.cpp
+++ b/Engine/Source/Runtime/Renderer/Private/VT/VirtualTextureSystem.cpp
@@ -1103,11 +1103,12 @@ void FVirtualTextureSystem::RequestTilesInternal(const IAllocatedVirtualTexture*
 		const uint32 vMaxLevel = InAllocatedVT->GetMaxLevel();
 		const float vLevel = ComputeMipLevel(InAllocatedVT, InScreenSpaceSize);
 		const int32 vMipLevelDown = FMath::Clamp((int32)FMath::FloorToInt(vLevel), 0, (int32)vMaxLevel);
-		RequestTilesInternal(InAllocatedVT, vLevel);
+
+		RequestTilesInternal(InAllocatedVT, vMipLevelDown);
 		if (vMipLevelDown + 1u <= vMaxLevel)
 		{
 			// Need to fetch 2 levels to support trilinear filtering
-			RequestTilesInternal(InAllocatedVT, vMipLevelDown);
+			RequestTilesInternal(InAllocatedVT, vMipLevelDown + 1u);
 		}
 	}
 	else
