Fix compile issue introduced in CL-17986013
Fix mismatched shader param size. All members existed and were properly aligned, but the size wasn't updated in merge.

diff --git a/Engine/Source/Runtime/Engine/Public/PrimitiveUniformShaderParameters.h b/Engine/Source/Runtime/Engine/Public/PrimitiveUniformShaderParameters.h
--- a/Engine/Source/Runtime/Engine/Public/PrimitiveUniformShaderParameters.h
+++ b/Engine/Source/Runtime/Engine/Public/PrimitiveUniformShaderParameters.h
@@ -402,7 +402,7 @@ extern ENGINE_API TGlobalResource<FIdentityPrimitiveUniformBuffer> GIdentityPrim
 struct FPrimitiveSceneShaderData
 {
 	// Must match usf
-	enum { DataStrideInFloat4s = 38 };
+	enum { DataStrideInFloat4s = 39 };
 
 	TStaticArray<FVector4f, DataStrideInFloat4s> Data;
 
