Optimization of the new ISM update from Jason Nadro, saves 3ms when displaying Mass entities with cubes.
* Use TConstArrayView<float> when possible
* Added Mobility logic to the ISM Update and to Mass Mesh Descriptor
* Fixed the hash to include Mobility and other parameters in the ISMSharedData map
* Added a clipping of 3km in the frosty display entities

diff --git a/Engine/Source/Runtime/Engine/Private/InstancedStaticMesh.cpp b/Engine/Source/Runtime/Engine/Private/InstancedStaticMesh.cpp
--- a/Engine/Source/Runtime/Engine/Private/InstancedStaticMesh.cpp
+++ b/Engine/Source/Runtime/Engine/Private/InstancedStaticMesh.cpp
@@ -248,7 +248,7 @@ void FInstanceUpdateCmdBuffer::AddInstance(const FMatrix& InTransform)
 	Edit();
 }
 
-void FInstanceUpdateCmdBuffer::AddInstance(int32 InstanceId, const FMatrix& InTransform, const FMatrix& InPreviousTransform, const TArray<float>& InCustomDataFloats)
+void FInstanceUpdateCmdBuffer::AddInstance(int32 InstanceId, const FMatrix& InTransform, const FMatrix& InPreviousTransform, TConstArrayView<float> InCustomDataFloats)
 {
 	FInstanceUpdateCommand& Cmd = Cmds.AddDefaulted_GetRef();
 	Cmd.InstanceIndex = INDEX_NONE;
@@ -348,7 +348,7 @@ void FInstanceUpdateCmdBuffer::SetShadowMapData(int32 RenderIndex, const FVector
 	Edit();
 }
 
-void FInstanceUpdateCmdBuffer::SetCustomData(int32 RenderIndex, const TArray<float>& CustomDataFloats)
+void FInstanceUpdateCmdBuffer::SetCustomData(int32 RenderIndex, TConstArrayView<float> CustomDataFloats)
 {
 	bool CommandExist = false;
 
@@ -3476,7 +3476,7 @@ bool UInstancedStaticMeshComponent::UpdateInstances(
 			// Explicitly check the component's previous transform for this instance because this 
 			// is the position we were previously rendered at.
 			FMatrix NewInstanceTransformMatrix = UpdateInstanceTransforms[i].ToMatrixWithScale();
-			if (!PerInstanceSMData[InstanceIndex].Transform.Equals(NewInstanceTransformMatrix, EqualTolerance))
+			if (Mobility == EComponentMobility::Movable && !PerInstanceSMData[InstanceIndex].Transform.Equals(NewInstanceTransformMatrix, EqualTolerance))
 			{
 				// Update the component's data in place.
 				PerInstanceSMData[InstanceIndex] = NewInstanceTransformMatrix;
@@ -3492,7 +3492,7 @@ bool UInstancedStaticMeshComponent::UpdateInstances(
 			}
 
 			// Did the custom data actually change?
-			if (bHasCustomFloatData)
+			if (bHasCustomFloatData && Mobility != EComponentMobility::Static)
 			{
 				// We are updating an existing instance.  They should have the same custom float data count.
 				check(NumCustomDataFloats == PrevNumCustomDataFloats);
@@ -3503,16 +3503,11 @@ bool UInstancedStaticMeshComponent::UpdateInstances(
 				{
 					if (FMath::Abs(CustomFloatData[SrcCustomDataOffset + FloatIndex] - PerInstanceSMCustomData[CustomDataOffset + FloatIndex]) > EqualTolerance)
 					{
-						// If any custom float data changed update it all.
-						TArray<float> NewCustomFloatData;
-						NewCustomFloatData.SetNumUninitialized(NumCustomDataFloats);
-						FMemory::Memcpy(&NewCustomFloatData[0], &CustomFloatData[SrcCustomDataOffset], NumCustomDataFloats * sizeof(float));
-
 						// Update the component's data in place.
-						FMemory::Memcpy(&PerInstanceSMCustomData[CustomDataOffset], &NewCustomFloatData[0], NumCustomDataFloats * sizeof(float));
+						FMemory::Memcpy(&PerInstanceSMCustomData[CustomDataOffset], &CustomFloatData[SrcCustomDataOffset], NumCustomDataFloats * sizeof(float));
 
 						// Record in a command buffer for future use.
-						InstanceUpdateCmdBuffer.SetCustomData(InstanceIndex, NewCustomFloatData);
+						InstanceUpdateCmdBuffer.SetCustomData(InstanceIndex, MakeArrayView((const float*)&CustomFloatData[SrcCustomDataOffset], NumCustomDataFloats));
 
 #if (CSV_PROFILER)
 						TotalSizeUpdateBytes += NumCustomDataFloats * sizeof(float);
@@ -3528,16 +3523,8 @@ bool UInstancedStaticMeshComponent::UpdateInstances(
 		else
 		{
 			// This is an add.
-			TArray<float> NewCustomFloatData;
-			if (bHasCustomFloatData)
-			{
-				NewCustomFloatData.SetNumUninitialized(NumCustomDataFloats);
-
-				const int32 SrcCustomDataOffset = i * NumCustomDataFloats;
-				FMemory::Memcpy(&NewCustomFloatData[0], &CustomFloatData[SrcCustomDataOffset], NumCustomDataFloats * sizeof(float));
-			}
-
-			InstanceUpdateCmdBuffer.AddInstance(InstanceId, UpdateInstanceTransforms[i].ToMatrixWithScale(), UpdateInstancePreviousTransforms[i].ToMatrixWithScale(), NewCustomFloatData);
+			const int32 SrcCustomDataOffset = i * NumCustomDataFloats;
+			InstanceUpdateCmdBuffer.AddInstance(InstanceId, UpdateInstanceTransforms[i].ToMatrixWithScale(), UpdateInstancePreviousTransforms[i].ToMatrixWithScale(), bHasCustomFloatData ? MakeArrayView((const float*)&CustomFloatData[SrcCustomDataOffset], NumCustomDataFloats) : TConstArrayView<float>());
 		}
 	}
 
diff --git a/Engine/Source/Runtime/Engine/Public/PrimitiveInstanceUpdateCommand.h b/Engine/Source/Runtime/Engine/Public/PrimitiveInstanceUpdateCommand.h
--- a/Engine/Source/Runtime/Engine/Public/PrimitiveInstanceUpdateCommand.h
+++ b/Engine/Source/Runtime/Engine/Public/PrimitiveInstanceUpdateCommand.h
@@ -42,13 +42,13 @@ struct FInstanceUpdateCmdBuffer
 	// Commands that can modify render data in place
 	void HideInstance(int32 RenderIndex);
 	void AddInstance(const FMatrix& InTransform);
-	void AddInstance(int32 InstanceId, const FMatrix& InTransform, const FMatrix& InPreviousTransform, const TArray<float>& InCustomDataFloats);
+	void AddInstance(int32 InstanceId, const FMatrix& InTransform, const FMatrix& InPreviousTransform, TConstArrayView<float> InCustomDataFloats);
 	void UpdateInstance(int32 RenderIndex, const FMatrix& InTransform);
 	void UpdateInstance(int32 RenderIndex, const FMatrix& InTransform, const FMatrix& InPreviousTransform);
 	void SetEditorData(int32 RenderIndex, const FColor& Color, bool bSelected);
 	void SetLightMapData(int32 RenderIndex, const FVector2D& LightmapUVBias);
 	void SetShadowMapData(int32 RenderIndex, const FVector2D& ShadowmapUVBias);
-	void SetCustomData(int32 RenderIndex, const TArray<float>& CustomDataFloats);
+	void SetCustomData(int32 RenderIndex, TConstArrayView<float> CustomDataFloats);
 	void ResetInlineCommands();
 	int32 NumInlineCommands() const { return Cmds.Num(); }
 
