Don't spend time detaching components from their parent when destroying HLOD components.
Saves more than 75% of unhashing time spent on the GT.

diff --git a/Engine/Source/Runtime/Engine/Private/Components/SceneComponent.cpp b/Engine/Source/Runtime/Engine/Private/Components/SceneComponent.cpp
--- a/Engine/Source/Runtime/Engine/Private/Components/SceneComponent.cpp
+++ b/Engine/Source/Runtime/Engine/Private/Components/SceneComponent.cpp
@@ -1053,7 +1053,7 @@ void USceneComponent::OnComponentDestroyed(bool bDestroyingHierarchy)
 	ScopedMovementStack.Reset();
 
 	// If we're just destroying for the exit purge don't bother with any of this
-	if (!GExitPurge)
+	if (!GExitPurge && !bComputeBoundsOnceForGame)
 	{
 		// If we're destroying the hierarchy we only have to make sure that we detach children from other Actor's
 		AActor* MyOwner = GetOwner();
