Fix compile issue introduced in CL-17963702
IAudioInstanceTransmitter was renamed to Audio:IParameterTransmitter.

diff --git a/Engine/Source/Runtime/Engine/Private/AudioDevice.cpp b/Engine/Source/Runtime/Engine/Private/AudioDevice.cpp
--- a/Engine/Source/Runtime/Engine/Private/AudioDevice.cpp
+++ b/Engine/Source/Runtime/Engine/Private/AudioDevice.cpp
@@ -5123,9 +5123,9 @@ void FAudioDevice::ProcessingPendingActiveSoundStops(bool bForceDelete)
 				PendingSoundsToDelete.RemoveAtSwap(i, 1, false);
 
 				// TODO: Phil Popp investigate a more robust fix -- call Shutdown in FActiveSound or transmitter destructor?
-				if (IAudioInstanceTransmitter* Transmitter = ActiveSound->GetTransmitter())
+				if (Audio::IParameterTransmitter* Transmitter = ActiveSound->GetTransmitter())
 				{
-					Transmitter->Shutdown();
+					Transmitter->Reset();
 				}
 
 				NotifyPendingDelete(*ActiveSound);
@@ -5183,9 +5183,9 @@ void FAudioDevice::ProcessingPendingActiveSoundStops(bool bForceDelete)
 			if (bDeleteActiveSound)
 			{
 				// TODO: Phil Popp investigate a more robust fix -- call Shutdown in FActiveSound or transmitter destructor?
-				if (IAudioInstanceTransmitter* Transmitter = ActiveSound->GetTransmitter())
+				if (Audio::IParameterTransmitter* Transmitter = ActiveSound->GetTransmitter())
 				{
-					Transmitter->Shutdown();
+					Transmitter->Reset();
 				}
 
 				NotifyPendingDelete(*ActiveSound);
