If cloth sim data is empty when it reaches skin cache, revert to shader permuation without cloth, so it doesn't crash. Cloth should still render as normal but without physics simulation.

diff --git a/Engine/Source/Runtime/Engine/Private/GPUSkinCache.cpp b/Engine/Source/Runtime/Engine/Private/GPUSkinCache.cpp
--- a/Engine/Source/Runtime/Engine/Private/GPUSkinCache.cpp
+++ b/Engine/Source/Runtime/Engine/Private/GPUSkinCache.cpp
@@ -1678,53 +1678,60 @@ bool FGPUSkinCache::ProcessEntry(
         InOutEntry->ClothBuffer = ClothVertexBuffer->GetSRV();
         check(InOutEntry->ClothBuffer);
 
-        check(SimData->Positions.Num() == SimData->Normals.Num());
-        VertexAndNormalData.ResizeBuffer( SimData->Positions.Num() );
+		if (SimData->Positions.Num() > 0)
+		{
+			check(SimData->Positions.Num() == SimData->Normals.Num());
+			VertexAndNormalData.ResizeBuffer( SimData->Positions.Num() );
 
-        uint8* Data = VertexAndNormalData.GetDataPointer();
-        uint32 Stride = VertexAndNormalData.GetStride();
+			uint8* Data = VertexAndNormalData.GetDataPointer();
+			uint32 Stride = VertexAndNormalData.GetStride();
 
-        // Copy the vertices into the buffer.
-        checkSlow(Stride*VertexAndNormalData.GetNumVertices() == sizeof(FClothSimulEntry) * SimData->Positions.Num());
-        check(sizeof(FClothSimulEntry) == 6 * sizeof(float));
+			// Copy the vertices into the buffer.
+			checkSlow(Stride*VertexAndNormalData.GetNumVertices() == sizeof(FClothSimulEntry) * SimData->Positions.Num());
+			check(sizeof(FClothSimulEntry) == 6 * sizeof(float));
+		
+			if (ClothVertexBuffer && ClothVertexBuffer->GetClothIndexMapping().Num() > Section)
+			{
+				const FClothBufferIndexMapping& ClothBufferIndexMapping = ClothVertexBuffer->GetClothIndexMapping()[Section];
 
-		if (ClothVertexBuffer && ClothVertexBuffer->GetClothIndexMapping().Num() > Section)
-		{
-			const FClothBufferIndexMapping& ClothBufferIndexMapping = ClothVertexBuffer->GetClothIndexMapping()[Section];
+				check(SimData->LODIndex != INDEX_NONE && SimData->LODIndex <= LODIndex);
+				const uint32 ClothLODBias = (uint32)(LODIndex - SimData->LODIndex);
 
-			check(SimData->LODIndex != INDEX_NONE && SimData->LODIndex <= LODIndex);
-			const uint32 ClothLODBias = (uint32)(LODIndex - SimData->LODIndex);
+				const uint32 ClothBufferOffset = ClothBufferIndexMapping.MappingOffset + ClothBufferIndexMapping.LODBiasStride * ClothLODBias;
 
-			const uint32 ClothBufferOffset = ClothBufferIndexMapping.MappingOffset + ClothBufferIndexMapping.LODBiasStride * ClothLODBias;
-
-			// Set the buffer offset depending on whether enough deformer mapping data exists (RaytracingMinLOD/RaytracingLODBias/ClothLODBiasMode settings)
-			const uint32 NumInfluences = NumVertices ? ClothBufferIndexMapping.LODBiasStride / NumVertices : 1;
-			InOutEntry->DispatchData[Section].ClothBufferOffset = (ClothBufferOffset + NumVertices * NumInfluences <= ClothVertexBuffer->GetNumVertices()) ?
+				// Set the buffer offset depending on whether enough deformer mapping data exists (RaytracingMinLOD/RaytracingLODBias/ClothLODBiasMode settings)
+				const uint32 NumInfluences = NumVertices ? ClothBufferIndexMapping.LODBiasStride / NumVertices : 1;
+				InOutEntry->DispatchData[Section].ClothBufferOffset = (ClothBufferOffset + NumVertices * NumInfluences <= ClothVertexBuffer->GetNumVertices()) ?
 				ClothBufferOffset :                     // If the offset is valid, set the calculated LODBias offset
 				ClothBufferIndexMapping.MappingOffset;  // Otherwise fallback to a 0 ClothLODBias to prevent from reading pass the buffer (but still raytrace broken shadows/reflections/etc.)
-		}
+			}
 
-        for (int32 Index = 0;Index < SimData->Positions.Num();Index++)
-        {
-            FClothSimulEntry NewEntry;
-            NewEntry.Position = SimData->Positions[Index];
-            NewEntry.Normal = SimData->Normals[Index];
-            *((FClothSimulEntry*)(Data + Index * Stride)) = NewEntry;
-        }
+			for (int32 Index = 0;Index < SimData->Positions.Num();Index++)
+			{
+				FClothSimulEntry NewEntry;
+				NewEntry.Position = SimData->Positions[Index];
+				NewEntry.Normal = SimData->Normals[Index];
+				*((FClothSimulEntry*)(Data + Index * Stride)) = NewEntry;
+			}
 
-        FResourceArrayInterface* ResourceArray = VertexAndNormalData.GetResourceArray();
-        check(ResourceArray->GetResourceDataSize() > 0);
+			FResourceArrayInterface* ResourceArray = VertexAndNormalData.GetResourceArray();
+			check(ResourceArray->GetResourceDataSize() > 0);
 
-        FRHIResourceCreateInfo CreateInfo(TEXT("ClothPositionAndNormalsBuffer"), ResourceArray);
-        ClothPositionAndNormalsBuffer.VertexBufferRHI = RHICreateVertexBuffer( ResourceArray->GetResourceDataSize(), BUF_Static | BUF_ShaderResource, CreateInfo);
-        ClothPositionAndNormalsBuffer.VertexBufferSRV = RHICreateShaderResourceView(ClothPositionAndNormalsBuffer.VertexBufferRHI, sizeof(FVector2f), PF_G32R32F);
-        InOutEntry->DispatchData[Section].ClothPositionsAndNormalsBuffer = ClothPositionAndNormalsBuffer.VertexBufferSRV;
+			FRHIResourceCreateInfo CreateInfo(TEXT("ClothPositionAndNormalsBuffer"), ResourceArray);
+			ClothPositionAndNormalsBuffer.VertexBufferRHI = RHICreateVertexBuffer( ResourceArray->GetResourceDataSize(), BUF_Static | BUF_ShaderResource, CreateInfo);
+			ClothPositionAndNormalsBuffer.VertexBufferSRV = RHICreateShaderResourceView(ClothPositionAndNormalsBuffer.VertexBufferRHI, sizeof(FVector2f), PF_G32R32F);
+			InOutEntry->DispatchData[Section].ClothPositionsAndNormalsBuffer = ClothPositionAndNormalsBuffer.VertexBufferSRV;
+		}
+		else
+		{
+			UE_LOG(LogSkinCache, Error, TEXT("Cloth sim data is missing on mesh %s"), *GetSkeletalMeshObjectName(Skin));
+		}
 
         InOutEntry->DispatchData[Section].ClothBlendWeight = ClothBlendWeight;
         InOutEntry->DispatchData[Section].ClothLocalToWorld = ClothLocalToWorld;
         InOutEntry->DispatchData[Section].ClothWorldToLocal = ClothLocalToWorld.Inverse();
     }
-    InOutEntry->DispatchData[Section].SkinType = ClothVertexBuffer ? 2 : (bMorph ? 1 : 0);
+    InOutEntry->DispatchData[Section].SkinType = ClothVertexBuffer && InOutEntry->DispatchData[Section].ClothPositionsAndNormalsBuffer ? 2 : (bMorph ? 1 : 0);
 
 	if (bShouldBatchDispatches)
 	{
