Fix compile issue introduced in CL-17964325
Removed some unnecessary code that doesn't exist in either branch.

diff --git a/Engine/Source/Runtime/Engine/Private/Rendering/MorphTargetVertexInfoBuffers.cpp b/Engine/Source/Runtime/Engine/Private/Rendering/MorphTargetVertexInfoBuffers.cpp
--- a/Engine/Source/Runtime/Engine/Private/Rendering/MorphTargetVertexInfoBuffers.cpp
+++ b/Engine/Source/Runtime/Engine/Private/Rendering/MorphTargetVertexInfoBuffers.cpp
@@ -243,19 +243,6 @@ void FMorphTargetVertexInfoBuffers::InitMorphResources(EShaderPlatform ShaderPla
 		int32 NumSrcDeltas = 0;
 		const FMorphTargetDelta* MorphDeltas = MorphTarget->GetMorphTargetDelta(LODIndex, NumSrcDeltas);
 
-		//Make sure the morphtarget data vertex indices fit the geometry
-		//If a missmatch happen, set the NumSrcDelta to 0 so the morph target is skipped
-		for (int32 DeltaIndex = 0; DeltaIndex < NumSrcDeltas; DeltaIndex++)
-		{
-			const auto& MorphDelta = MorphDeltas[DeltaIndex];
-			if (!VertexNeedsTangents.IsValidIndex(MorphDelta.SourceIdx))
-			{
-				NumSrcDeltas = 0;
-				UE_ASSET_LOG(LogSkeletalMesh, Warning, Owner, TEXT("Skipping morph target %s for LOD %d. The morph target data is incompatible with the mesh data"), *MorphTarget->GetName(), LODIndex);
-				break;
-			}
-		}
-
 		if (NumSrcDeltas == 0)
 		{
 			MaximumValues[0] = 0.0f;
