Make sure HLOD warmup requests expire after some time
* Covers the case where a cell unload is requested but ultimatelly kept loaded before the warmup has completed

diff --git a/Engine/Source/Runtime/Engine/Private/WorldPartition/HLOD/HLODSubsystem.cpp b/Engine/Source/Runtime/Engine/Private/WorldPartition/HLOD/HLODSubsystem.cpp
--- a/Engine/Source/Runtime/Engine/Private/WorldPartition/HLOD/HLODSubsystem.cpp
+++ b/Engine/Source/Runtime/Engine/Private/WorldPartition/HLOD/HLODSubsystem.cpp
@@ -367,10 +367,13 @@ bool UHLODSubsystem::RequestUnloading(const UWorldPartitionRuntimeCell* InCell)
 		return true;
 	}
 
+	// In case a previous request to unload was aborted and the cell never unloaded... assume warmup requests are expired after a given amount of frames.
+	const uint32 WarmupExpiredFrames = 30;
+
 	uint32 CurrentFrameNumber = GetWorld()->Scene->GetFrameNumber();
 
-	// Trigger warmup on the first request to unload
-	if (CellData.WarmupEndFrame == INDEX_NONE)
+	// Trigger warmup on the first request to unload, or if a warmup request expired
+	if (CellData.WarmupEndFrame == INDEX_NONE || CurrentFrameNumber > (CellData.WarmupEndFrame + WarmupExpiredFrames))
 	{
 		// Warmup will be triggered in the next BeginRenderView() call, at which point the frame number will have been incremented.
 		CellData.WarmupStartFrame = CurrentFrameNumber + 1; 
