When destroying a VT producer only wait on its tasks, not on all transcoding tasks.
Removes regular 100ms spikes waiting on unrelated work.

diff --git a/Engine/Source/Runtime/Engine/Private/VT/UploadingVirtualTexture.cpp b/Engine/Source/Runtime/Engine/Private/VT/UploadingVirtualTexture.cpp
--- a/Engine/Source/Runtime/Engine/Private/VT/UploadingVirtualTexture.cpp
+++ b/Engine/Source/Runtime/Engine/Private/VT/UploadingVirtualTexture.cpp
@@ -82,9 +82,6 @@ FUploadingVirtualTexture::~FUploadingVirtualTexture()
 	DEC_DWORD_STAT_BY(STAT_TotalDiskSize, Data->GetDiskMemoryFootprint() / (1024 * 1024));
 	DEC_DWORD_STAT_BY(STAT_NumTileHeaders, Data->GetNumTileHeaders());
 
-	// Complete all open transcode requests before deleting IFileCacheHandle objects
-	StreamingManager->WaitTasksFinished();
-
 	for (TUniquePtr<FVirtualTextureCodec>& Codec : CodecPerChunk)
 	{
 		if (Codec)
@@ -145,6 +142,11 @@ IVirtualTextureFinalizer* FUploadingVirtualTexture::ProducePageData(FRHICommandL
 	return StreamingManager->ProduceTile(RHICmdList, SkipBorderSize, Data->GetNumLayers(), LayerMask, RequestHandle, TargetLayers);
 }
 
+void FUploadingVirtualTexture::GatherProducePageDataTasks(FVirtualTextureProducerHandle const& ProducerHandle, FGraphEventArray& InOutTasks) const
+{
+	StreamingManager->GatherProducePageDataTasks(ProducerHandle, InOutTasks);
+}
+
 void FUploadingVirtualTexture::GatherProducePageDataTasks(uint64 RequestHandle, FGraphEventArray& InOutTasks) const
 {
 	StreamingManager->GatherProducePageDataTasks(RequestHandle, InOutTasks);
diff --git a/Engine/Source/Runtime/Engine/Private/VT/UploadingVirtualTexture.h b/Engine/Source/Runtime/Engine/Private/VT/UploadingVirtualTexture.h
--- a/Engine/Source/Runtime/Engine/Private/VT/UploadingVirtualTexture.h
+++ b/Engine/Source/Runtime/Engine/Private/VT/UploadingVirtualTexture.h
@@ -84,6 +84,7 @@ public:
 		const FVirtualTextureProducerHandle& ProducerHandle, uint8 LayerMask, uint8 vLevel, uint64 vAddress,
 		uint64 RequestHandle,
 		const FVTProduceTargetLayer* TargetLayers) override;
+	virtual void GatherProducePageDataTasks(FVirtualTextureProducerHandle const& ProducerHandle, FGraphEventArray& InOutTasks) const override;
 	virtual void GatherProducePageDataTasks(uint64 RequestHandle, FGraphEventArray& InOutTasks) const override;
 	virtual void DumpToConsole(bool verbose) override;
 	// End IVirtualTexture interface
diff --git a/Engine/Source/Runtime/Engine/Private/VT/VirtualTextureChunkManager.cpp b/Engine/Source/Runtime/Engine/Private/VT/VirtualTextureChunkManager.cpp
--- a/Engine/Source/Runtime/Engine/Private/VT/VirtualTextureChunkManager.cpp
+++ b/Engine/Source/Runtime/Engine/Private/VT/VirtualTextureChunkManager.cpp
@@ -139,7 +139,7 @@ FVTRequestPageResult FVirtualTextureChunkStreamingManager::RequestTile(FUploadin
 	TranscodeParams.LayerMask = LayerMask;
 	TranscodeParams.Codec = CodecResult.Codec;
 	TranscodeParams.Name = VTexture->GetName();
-	const FVTTranscodeTileHandle TranscodeHandle = TranscodeCache.SubmitTask(UploadCache, TranscodeKey, TranscodeParams, &GraphCompletionEvents);
+	const FVTTranscodeTileHandle TranscodeHandle = TranscodeCache.SubmitTask(UploadCache, TranscodeKey, ProducerHandle, TranscodeParams, &GraphCompletionEvents);
 	return FVTRequestPageResult(EVTRequestPageStatus::Pending, TranscodeHandle.PackedData);
 }
 
@@ -160,6 +160,11 @@ IVirtualTextureFinalizer* FVirtualTextureChunkStreamingManager::ProduceTile(FRHI
 	return &UploadCache;
 }
 
+void FVirtualTextureChunkStreamingManager::GatherProducePageDataTasks(FVirtualTextureProducerHandle const& ProducerHandle, FGraphEventArray& InOutTasks) const
+{
+	TranscodeCache.GatherProducePageDataTasks(ProducerHandle, InOutTasks);
+}
+
 void FVirtualTextureChunkStreamingManager::GatherProducePageDataTasks(uint64 RequestHandle, FGraphEventArray& InOutTasks) const
 {
 	FGraphEventRef Event = TranscodeCache.GetTaskEvent(FVTTranscodeTileHandle(RequestHandle));
@@ -168,8 +173,3 @@ void FVirtualTextureChunkStreamingManager::GatherProducePageDataTasks(uint64 Req
 		InOutTasks.Emplace(MoveTemp(Event));
 	}
 }
-
-void FVirtualTextureChunkStreamingManager::WaitTasksFinished() const
-{
-	TranscodeCache.WaitTasksFinished();
-}
diff --git a/Engine/Source/Runtime/Engine/Private/VT/VirtualTextureChunkManager.h b/Engine/Source/Runtime/Engine/Private/VT/VirtualTextureChunkManager.h
--- a/Engine/Source/Runtime/Engine/Private/VT/VirtualTextureChunkManager.h
+++ b/Engine/Source/Runtime/Engine/Private/VT/VirtualTextureChunkManager.h
@@ -44,8 +44,8 @@ public:
 
 	FVTRequestPageResult RequestTile(FUploadingVirtualTexture* VTexture, const FVirtualTextureProducerHandle& ProducerHandle, uint8 LayerMask, uint8 vLevel, uint32 vAddress, EVTRequestPagePriority Priority);
 	IVirtualTextureFinalizer* ProduceTile(FRHICommandListImmediate& RHICmdList, uint32 SkipBorderSize, uint8 NumLayers, uint8 LayerMask, uint64 RequestHandle, const FVTProduceTargetLayer* TargetLayers);
+	void GatherProducePageDataTasks(FVirtualTextureProducerHandle const& ProducerHandle, FGraphEventArray& InOutTasks) const;
 	void GatherProducePageDataTasks(uint64 RequestHandle, FGraphEventArray& InOutTasks) const;
-	void WaitTasksFinished() const;
 
 private:
 	FVirtualTextureUploadCache UploadCache;
diff --git a/Engine/Source/Runtime/Engine/Private/VT/VirtualTextureTranscodeCache.cpp b/Engine/Source/Runtime/Engine/Private/VT/VirtualTextureTranscodeCache.cpp
--- a/Engine/Source/Runtime/Engine/Private/VT/VirtualTextureTranscodeCache.cpp
+++ b/Engine/Source/Runtime/Engine/Private/VT/VirtualTextureTranscodeCache.cpp
@@ -310,21 +310,6 @@ void FVirtualTextureTranscodeCache::WaitTaskFinished(FVTTranscodeTileHandle InHa
 	}
 }
 
-void FVirtualTextureTranscodeCache::WaitTasksFinished() const
-{
-	int32 TaskIndex = Tasks[LIST_PENDING].NextIndex;
-	while (TaskIndex != LIST_PENDING)
-	{
-		FTaskEntry const& TaskEntry = Tasks[TaskIndex];
-		const int32 NextIndex = TaskEntry.NextIndex;
-		if (TaskEntry.GraphEvent && !TaskEntry.GraphEvent->IsComplete())
-		{
-			FTaskGraphInterface::Get().WaitUntilTaskCompletes(TaskEntry.GraphEvent, ENamedThreads::GetRenderThread_Local());
-		}
-		TaskIndex = NextIndex;
-	}
-}
-
 FGraphEventRef FVirtualTextureTranscodeCache::GetTaskEvent(FVTTranscodeTileHandle InHandle) const
 {
 	const uint32 TaskIndex = InHandle.Index;
@@ -334,6 +319,21 @@ FGraphEventRef FVirtualTextureTranscodeCache::GetTaskEvent(FVTTranscodeTileHandl
 	return TaskEntry.GraphEvent;
 }
 
+void FVirtualTextureTranscodeCache::GatherProducePageDataTasks(FVirtualTextureProducerHandle const& ProducerHandle, FGraphEventArray& InOutTasks) const
+{
+	// todo[VT]: Add multimap or similar if we need faster lookup here.
+	int32 TaskIndex = Tasks[LIST_PENDING].NextIndex;
+	while (TaskIndex != LIST_PENDING)
+	{
+		FTaskEntry const& Task = Tasks[TaskIndex];
+		if (Task.PackedProducerHandle == ProducerHandle.PackedValue && Task.GraphEvent && !Task.GraphEvent->IsComplete())
+		{
+			InOutTasks.Add(Task.GraphEvent);
+		}
+		TaskIndex = Task.NextIndex;
+	}
+}
+
 const FVTUploadTileHandle* FVirtualTextureTranscodeCache::AcquireTaskResult(FVTTranscodeTileHandle InHandle)
 {
 	const uint32 TaskIndex = InHandle.Index;
@@ -357,8 +357,10 @@ const FVTUploadTileHandle* FVirtualTextureTranscodeCache::AcquireTaskResult(FVTT
 	return TaskEntry.StageTileHandle;
 }
 
-FVTTranscodeTileHandle FVirtualTextureTranscodeCache::SubmitTask(FVirtualTextureUploadCache& InUploadCache,
+FVTTranscodeTileHandle FVirtualTextureTranscodeCache::SubmitTask(
+	FVirtualTextureUploadCache& InUploadCache,
 	const FVTTranscodeKey& InKey,
+	const FVirtualTextureProducerHandle& InProducerHandle,
 	const FVTTranscodeParams& InParams,
 	const FGraphEventArray* InPrerequisites)
 {
@@ -382,6 +384,7 @@ FVTTranscodeTileHandle FVirtualTextureTranscodeCache::SubmitTask(FVirtualTexture
 
 	FTaskEntry& TaskEntry = Tasks[TaskIndex];
 	TaskEntry.Key = InKey.Key;
+	TaskEntry.PackedProducerHandle = InProducerHandle.PackedValue;
 	TaskEntry.Hash = InKey.Hash;
 	TaskEntry.FrameSubmitted = GFrameNumberRenderThread;
 	FMemory::Memzero(TaskEntry.StageTileHandle);
diff --git a/Engine/Source/Runtime/Engine/Private/VT/VirtualTextureTranscodeCache.h b/Engine/Source/Runtime/Engine/Private/VT/VirtualTextureTranscodeCache.h
--- a/Engine/Source/Runtime/Engine/Private/VT/VirtualTextureTranscodeCache.h
+++ b/Engine/Source/Runtime/Engine/Private/VT/VirtualTextureTranscodeCache.h
@@ -73,17 +73,19 @@ public:
 
 	FVTTranscodeTileHandleAndStatus FindTask(const FVTTranscodeKey& InKey) const;
 
-	FVTTranscodeTileHandle SubmitTask(FVirtualTextureUploadCache& InUploadCache,
+	FVTTranscodeTileHandle SubmitTask(
+		FVirtualTextureUploadCache& InUploadCache,
 		const FVTTranscodeKey& InKey,
+		const FVirtualTextureProducerHandle& InProducerHandle,
 		const FVTTranscodeParams& InParams,
 		const FGraphEventArray* InPrerequisites = NULL);
 
 	bool IsTaskFinished(FVTTranscodeTileHandle InHandle) const;
 	void WaitTaskFinished(FVTTranscodeTileHandle InHandle) const;
 	const FVTUploadTileHandle* AcquireTaskResult(FVTTranscodeTileHandle InHandle);
+	
 	FGraphEventRef GetTaskEvent(FVTTranscodeTileHandle InHandle) const;
-
-	void WaitTasksFinished() const;
+	void GatherProducePageDataTasks(FVirtualTextureProducerHandle const& ProducerHandle, FGraphEventArray& InOutTasks) const;
 
 	void RetireOldTasks(FVirtualTextureUploadCache& InUploadCache);
 
@@ -101,6 +103,7 @@ private:
 		uint64 Key;
 		FGraphEventRef GraphEvent;
 		FVTUploadTileHandle StageTileHandle[VIRTUALTEXTURE_SPACE_MAXLAYERS];
+		uint32 PackedProducerHandle;
 		uint32 FrameSubmitted;
 		uint16 Magic;
 		uint16 Hash;
diff --git a/Engine/Source/Runtime/RenderCore/Public/VirtualTexturing.h b/Engine/Source/Runtime/RenderCore/Public/VirtualTexturing.h
--- a/Engine/Source/Runtime/RenderCore/Public/VirtualTexturing.h
+++ b/Engine/Source/Runtime/RenderCore/Public/VirtualTexturing.h
@@ -298,8 +298,13 @@ public:
 		uint64 RequestHandle,
 		const FVTProduceTargetLayer* TargetLayers) = 0;
 
+	/** Collect all task graph events. */
+	virtual void GatherProducePageDataTasks(FVirtualTextureProducerHandle const& ProducerHandle, FGraphEventArray& InOutTasks) const {}
+
+	/** Collect all task graph events related to a request. */
 	virtual void GatherProducePageDataTasks(uint64 RequestHandle, FGraphEventArray& InOutTasks) const {};
 
+	/** Dump any type specific debug info. */
 	virtual void DumpToConsole(bool verbose) {}
 };
 
diff --git a/Engine/Source/Runtime/Renderer/Private/VT/VirtualTextureProducer.cpp b/Engine/Source/Runtime/Renderer/Private/VT/VirtualTextureProducer.cpp
--- a/Engine/Source/Runtime/Renderer/Private/VT/VirtualTextureProducer.cpp
+++ b/Engine/Source/Runtime/Renderer/Private/VT/VirtualTextureProducer.cpp
@@ -23,8 +23,19 @@ void FVirtualTextureProducer::Release(FVirtualTextureSystem* System, const FVirt
 	}
 
 	PhysicalGroups.Reset();
+
+	{
+		// Complete all open tasks before delete.
+		TRACE_CPUPROFILER_EVENT_SCOPE(FVirtualTextureProducer::Release_Wait);
+	
+		FGraphEventArray ProducePageTasks;
+		VirtualTexture->GatherProducePageDataTasks(HandleToSelf, ProducePageTasks);
+		FTaskGraphInterface::Get().WaitUntilTasksComplete(ProducePageTasks, ENamedThreads::GetRenderThread_Local());
+	}
+
 	delete VirtualTexture;
 	VirtualTexture = nullptr;
+	
 	Description = FVTProducerDescription();
 }
 
