Fixed UWorldPartition::GetStreamingSources accessing an unchecked UWorldPartitionStreamingPolicy pointer

diff --git a/Engine/Source/Runtime/Engine/Private/WorldPartition/WorldPartition.cpp b/Engine/Source/Runtime/Engine/Private/WorldPartition/WorldPartition.cpp
--- a/Engine/Source/Runtime/Engine/Private/WorldPartition/WorldPartition.cpp
+++ b/Engine/Source/Runtime/Engine/Private/WorldPartition/WorldPartition.cpp
@@ -741,7 +741,7 @@ UWorldPartition* UWorldPartition::CreateOrRepairWorldPartition(AWorldSettings* W
 
 const TArray<FWorldPartitionStreamingSource>& UWorldPartition::GetStreamingSources() const
 {
-	if (GetWorld()->IsGameWorld())
+	if (StreamingPolicy && GetWorld()->IsGameWorld())
 	{
 		return StreamingPolicy->GetStreamingSources();
 	}
