Faster lookup of tasks per VT producer handle.

diff --git a/Engine/Source/Runtime/Engine/Private/VT/VirtualTextureTranscodeCache.cpp b/Engine/Source/Runtime/Engine/Private/VT/VirtualTextureTranscodeCache.cpp
--- a/Engine/Source/Runtime/Engine/Private/VT/VirtualTextureTranscodeCache.cpp
+++ b/Engine/Source/Runtime/Engine/Private/VT/VirtualTextureTranscodeCache.cpp
@@ -321,16 +321,14 @@ FGraphEventRef FVirtualTextureTranscodeCache::GetTaskEvent(FVTTranscodeTileHandl
 
 void FVirtualTextureTranscodeCache::GatherProducePageDataTasks(FVirtualTextureProducerHandle const& ProducerHandle, FGraphEventArray& InOutTasks) const
 {
-	// todo[VT]: Add multimap or similar if we need faster lookup here.
-	int32 TaskIndex = Tasks[LIST_PENDING].NextIndex;
-	while (TaskIndex != LIST_PENDING)
+	const uint32 Hash = MurmurFinalize32(ProducerHandle.PackedValue);
+	for (uint32 TaskIndex = ProducerToTaskIndex.First(Hash); ProducerToTaskIndex.IsValid(TaskIndex); TaskIndex = ProducerToTaskIndex.Next(TaskIndex))
 	{
-		FTaskEntry const& Task = Tasks[TaskIndex];
-		if (Task.PackedProducerHandle == ProducerHandle.PackedValue && Task.GraphEvent && !Task.GraphEvent->IsComplete())
+		FTaskEntry const& TaskEntry = Tasks[TaskIndex];
+		if (TaskEntry.PackedProducerHandle == ProducerHandle.PackedValue && TaskEntry.GraphEvent && !TaskEntry.GraphEvent->IsComplete())
 		{
-			InOutTasks.Add(Task.GraphEvent);
+			InOutTasks.Add(TaskEntry.GraphEvent);
 		}
-		TaskIndex = Task.NextIndex;
 	}
 }
 
@@ -353,6 +351,7 @@ const FVTUploadTileHandle* FVirtualTextureTranscodeCache::AcquireTaskResult(FVTT
 
 	++TaskEntry.Magic;
 	TileIDToTaskIndex.Remove(TaskEntry.Hash, TaskIndex);
+	ProducerToTaskIndex.Remove(MurmurFinalize32(TaskEntry.PackedProducerHandle), TaskIndex);
 
 	return TaskEntry.StageTileHandle;
 }
@@ -381,6 +380,7 @@ FVTTranscodeTileHandle FVirtualTextureTranscodeCache::SubmitTask(
 
 	AddToList(LIST_PENDING, TaskIndex);
 	TileIDToTaskIndex.Add(InKey.Hash, TaskIndex);
+	ProducerToTaskIndex.Add(MurmurFinalize32(InProducerHandle.PackedValue), TaskIndex);
 
 	FTaskEntry& TaskEntry = Tasks[TaskIndex];
 	TaskEntry.Key = InKey.Key;
@@ -498,6 +498,7 @@ void FVirtualTextureTranscodeCache::RetireOldTasks(FVirtualTextureUploadCache& I
 		RemoveFromList(TaskIndex);
 		AddToList(LIST_FREE, TaskIndex);
 		TileIDToTaskIndex.Remove(TaskEntry.Hash, TaskIndex);
+		ProducerToTaskIndex.Remove(MurmurFinalize32(TaskEntry.PackedProducerHandle), TaskIndex);
 
 		TaskIndex = NextIndex;
 	}
diff --git a/Engine/Source/Runtime/Engine/Private/VT/VirtualTextureTranscodeCache.h b/Engine/Source/Runtime/Engine/Private/VT/VirtualTextureTranscodeCache.h
--- a/Engine/Source/Runtime/Engine/Private/VT/VirtualTextureTranscodeCache.h
+++ b/Engine/Source/Runtime/Engine/Private/VT/VirtualTextureTranscodeCache.h
@@ -139,4 +139,5 @@ private:
 
 	TArray<FTaskEntry> Tasks;
 	FHashTable TileIDToTaskIndex;
+	FHashTable ProducerToTaskIndex;
 };
diff --git a/Engine/Source/Runtime/Renderer/Private/VT/VirtualTextureProducer.cpp b/Engine/Source/Runtime/Renderer/Private/VT/VirtualTextureProducer.cpp
--- a/Engine/Source/Runtime/Renderer/Private/VT/VirtualTextureProducer.cpp
+++ b/Engine/Source/Runtime/Renderer/Private/VT/VirtualTextureProducer.cpp
@@ -24,12 +24,11 @@ void FVirtualTextureProducer::Release(FVirtualTextureSystem* System, const FVirt
 
 	PhysicalGroups.Reset();
 
+	FGraphEventArray ProducePageTasks;
+	VirtualTexture->GatherProducePageDataTasks(HandleToSelf, ProducePageTasks);
+	if (ProducePageTasks.Num())
 	{
-		// Complete all open tasks before delete.
 		TRACE_CPUPROFILER_EVENT_SCOPE(FVirtualTextureProducer::Release_Wait);
-	
-		FGraphEventArray ProducePageTasks;
-		VirtualTexture->GatherProducePageDataTasks(HandleToSelf, ProducePageTasks);
 		FTaskGraphInterface::Get().WaitUntilTasksComplete(ProducePageTasks, ENamedThreads::GetRenderThread_Local());
 	}
 
