Defer merge of spans in GPU-Scene allocators during RemovePrimitiveSceneInfos.

diff --git a/Engine/Source/Runtime/Engine/Private/GrowOnlySpanAllocator.cpp b/Engine/Source/Runtime/Engine/Private/GrowOnlySpanAllocator.cpp
--- a/Engine/Source/Runtime/Engine/Private/GrowOnlySpanAllocator.cpp
+++ b/Engine/Source/Runtime/Engine/Private/GrowOnlySpanAllocator.cpp
@@ -41,6 +41,13 @@ void FGrowOnlySpanAllocator::Free(int32 BaseOffset, int32 Num)
 
 	FLinearAllocation NewFreeSpan(BaseOffset, Num);
 
+	// If we are deferring consolidation, just store the item, merging happens on allocate
+	if (bDeferMerges)
+	{
+		FreeSpans.Add(NewFreeSpan);
+		return;
+	}
+
 #if DO_GUARD_SLOW
 	// Detect double delete
 	for (int32 i = 0; i < FreeSpans.Num(); i++)
@@ -50,8 +57,6 @@ void FGrowOnlySpanAllocator::Free(int32 BaseOffset, int32 Num)
 	}
 #endif
 
-	bool bMergedIntoExisting = false;
-
 	int32 SpanBeforeIndex = INDEX_NONE;
 	int32 SpanAfterIndex = INDEX_NONE;
 
@@ -99,6 +104,48 @@ void FGrowOnlySpanAllocator::Free(int32 BaseOffset, int32 Num)
 	}
 }
 
+
+void FGrowOnlySpanAllocator::BeginDeferMerges()
+{
+	check(!bDeferMerges);
+	bDeferMerges = true;
+
+}
+void FGrowOnlySpanAllocator::EndDeferMerges()
+{
+	check(bDeferMerges);
+	bDeferMerges = false;
+
+	// Consolidation 
+	// 1. Sort the Free list by span start
+	FreeSpans.Sort();
+
+	// alternate free list, used during consolidation to avoid N^2 worst case, retained to avoid re-allocations.
+	TArray<FLinearAllocation, TInlineAllocator<10>> FreeSpansTmp;
+	FreeSpansTmp.Reset(FreeSpans.Num());
+
+	int32 PrevEndOffset = INDEX_NONE;
+
+	// 2. Loop over and fuse all adjacent, copy into new free list (to avoid compaction)
+	for (int32 Index = 0; Index < FreeSpans.Num(); ++Index)
+	{
+		FLinearAllocation Alloc = FreeSpans[Index];
+		// Continues the previous one, fuse
+		if (PrevEndOffset == Alloc.StartOffset)
+		{
+			FreeSpansTmp.Last().Num += Alloc.Num;
+		}
+		else
+		{
+			FreeSpansTmp.Add(Alloc);
+		}
+		PrevEndOffset = FreeSpansTmp.Last().Num + FreeSpansTmp.Last().StartOffset;
+	}
+	// 3. Swap old and new free lists
+	Swap(FreeSpans, FreeSpansTmp);
+}
+
+
 int32 FGrowOnlySpanAllocator::SearchFreeList(int32 Num)
 {
 	// Search free list for first matching
diff --git a/Engine/Source/Runtime/Engine/Public/GrowOnlySpanAllocator.h b/Engine/Source/Runtime/Engine/Public/GrowOnlySpanAllocator.h
--- a/Engine/Source/Runtime/Engine/Public/GrowOnlySpanAllocator.h
+++ b/Engine/Source/Runtime/Engine/Public/GrowOnlySpanAllocator.h
@@ -65,6 +65,13 @@ public:
 	}
 #endif // DO_CHECK
 
+	/** 
+	 * Between these calls to Free just appends the allocation to the free list, rather than trying to merge with existing allocations.
+	 * At EndDeferMerges the free list is consolidated by sorting and merging all spans. This amortises the cost of the merge over many calls.
+	 */
+	ENGINE_API void BeginDeferMerges();
+	ENGINE_API void EndDeferMerges();
+
 private:
 	class FLinearAllocation
 	{
@@ -82,6 +89,11 @@ private:
 		{
 			return StartOffset <= Other.StartOffset && (StartOffset + Num) >= (Other.StartOffset + Other.Num);
 		}
+
+		bool operator<(const FLinearAllocation& Other) const
+		{
+			return StartOffset < Other.StartOffset;
+		}
 	};
 
 	// Size of the linear range used by the allocator
@@ -90,5 +102,7 @@ private:
 	// Unordered free list
 	TArray<FLinearAllocation, TInlineAllocator<10>> FreeSpans;
 
+	bool bDeferMerges = false;
+
 	ENGINE_API int32 SearchFreeList(int32 Num);
 };
\ No newline at end of file
diff --git a/Engine/Source/Runtime/Renderer/Private/GPUScene.cpp b/Engine/Source/Runtime/Renderer/Private/GPUScene.cpp
--- a/Engine/Source/Runtime/Renderer/Private/GPUScene.cpp
+++ b/Engine/Source/Runtime/Renderer/Private/GPUScene.cpp
@@ -97,6 +97,16 @@ static TAutoConsoleVariable<float> CVarGPUSceneDebugDrawRange(
 	ECVF_RenderThreadSafe
 );
 
+
+static int32 GGPUSceneAllowDeferredAllocatorMerges = 1;
+FAutoConsoleVariableRef CVarGPUSceneAllowDeferredAllocatorMerges(
+	TEXT("r.GPUScene.AllowDeferredAllocatorMerges"),
+	GGPUSceneAllowDeferredAllocatorMerges,
+	TEXT(""),
+	ECVF_RenderThreadSafe
+);
+
+
 constexpr uint32 InstanceSceneDataNumArrays = FInstanceSceneShaderData::DataStrideInFloat4s;
 
 int32 GGPUSceneInstanceUploadViaCreate = 1;
@@ -1844,6 +1854,26 @@ void FGPUScene::DebugRender(FRDGBuilder& GraphBuilder, FScene& Scene, FViewInfo&
 }
 
 
+void FGPUScene::BeginDeferAllocatorMerges()
+{
+	if (GGPUSceneAllowDeferredAllocatorMerges != 0)
+	{
+		InstanceSceneDataAllocator.BeginDeferMerges();
+		InstancePayloadDataAllocator.BeginDeferMerges();
+		LightmapDataAllocator.BeginDeferMerges();
+	}
+}
+
+
+void FGPUScene::EndDeferAllocatorMerges()
+{
+	if (GGPUSceneAllowDeferredAllocatorMerges != 0)
+	{
+		InstanceSceneDataAllocator.EndDeferMerges();
+		InstancePayloadDataAllocator.EndDeferMerges();
+		LightmapDataAllocator.EndDeferMerges();
+	}
+}
 
 
 TRange<int32> FGPUScene::CommitPrimitiveCollector(FGPUScenePrimitiveCollector& PrimitiveCollector)
diff --git a/Engine/Source/Runtime/Renderer/Private/GPUScene.h b/Engine/Source/Runtime/Renderer/Private/GPUScene.h
--- a/Engine/Source/Runtime/Renderer/Private/GPUScene.h
+++ b/Engine/Source/Runtime/Renderer/Private/GPUScene.h
@@ -304,6 +304,13 @@ public:
 	 */
 	void DebugRender(FRDGBuilder& GraphBuilder, FScene& Scene, FViewInfo& View);
 
+	/**
+	 * Between these calls to FGrowOnlySpanAllocator::Free just appends the allocation to the free list, rather than trying to merge with existing allocations.
+	 * At EndDeferAllocatorMerges the free list is consolidated by sorting and merging all spans. This amortises the cost of the merge over many calls.
+	 */
+	void BeginDeferAllocatorMerges();
+	void EndDeferAllocatorMerges();
+
 	bool bUpdateAllPrimitives;
 
 	/** Indices of primitives that need to be updated in GPU Scene */
diff --git a/Engine/Source/Runtime/Renderer/Private/RendererScene.cpp b/Engine/Source/Runtime/Renderer/Private/RendererScene.cpp
--- a/Engine/Source/Runtime/Renderer/Private/RendererScene.cpp
+++ b/Engine/Source/Runtime/Renderer/Private/RendererScene.cpp
@@ -4115,6 +4115,9 @@ void FScene::UpdateAllPrimitiveSceneInfos(FRDGBuilder& GraphBuilder, bool bAsync
 		CSV_SCOPED_TIMING_STAT_EXCLUSIVE(RemovePrimitiveSceneInfos);
 		SCOPED_NAMED_EVENT(FScene_RemovePrimitiveSceneInfos, FColor::Red);
 		SCOPE_CYCLE_COUNTER(STAT_RemoveScenePrimitiveTime);
+
+		GPUScene.BeginDeferAllocatorMerges();
+
 		for (FPrimitiveSceneInfo* PrimitiveSceneInfo : RemovedLocalPrimitiveSceneInfos)
 		{
 			// clear it up, parent is getting removed
@@ -4350,6 +4353,8 @@ void FScene::UpdateAllPrimitiveSceneInfos(FRDGBuilder& GraphBuilder, bool bAsync
 			}
 			RemovedLocalPrimitiveSceneInfos.RemoveAt(StartIndex, RemovedLocalPrimitiveSceneInfos.Num() - StartIndex, false);
 		}
+	
+		GPUScene.EndDeferAllocatorMerges();
 	}
 	{
 		CSV_SCOPED_TIMING_STAT_EXCLUSIVE(AddPrimitiveSceneInfos);
@@ -4756,11 +4761,11 @@ void FScene::UpdateAllPrimitiveSceneInfos(FRDGBuilder& GraphBuilder, bool bAsync
 				PrimitiveSceneInfo->MarkIndirectLightingCacheBufferDirty();
 			}
 
-			PrimitiveSceneInfo->FreeGPUSceneInstances();
-			FPrimitiveSceneInfo::AllocateGPUSceneInstances(this, MakeArrayView(&PrimitiveSceneInfo, 1));
-
 			if (UpdateInstance.Value.CmdBuffer.NumAdds > 0 || UpdateInstance.Value.CmdBuffer.NumRemoves > 0)
 			{
+				PrimitiveSceneInfo->FreeGPUSceneInstances();
+				FPrimitiveSceneInfo::AllocateGPUSceneInstances(this, MakeArrayView(&PrimitiveSceneInfo, 1));
+
 				DistanceFieldSceneData.RemovePrimitive(PrimitiveSceneInfo);
 				DistanceFieldSceneData.AddPrimitive(PrimitiveSceneInfo);
 
@@ -4769,6 +4774,8 @@ void FScene::UpdateAllPrimitiveSceneInfos(FRDGBuilder& GraphBuilder, bool bAsync
 			}
 			else
 			{
+				GPUScene.AddPrimitiveToUpdate(PrimitiveSceneInfo->PackedIndex, EPrimitiveDirtyState::ChangedAll);
+
 				DistanceFieldSceneData.UpdatePrimitive(PrimitiveSceneInfo);
 				LumenSceneData->UpdatePrimitive(PrimitiveSceneInfo);
 			}
