Make LocalVF ISM/HISM not need any instance attributes except in editor, move remaining data to GPU-Scene.

- USE_EDITOR_SHADERS moved MaterialTemplate.ush -> Common.ush to allow use in global shaders as well as in includes that happen earlier in MaterialTemplate.ush
- PLATFORM_SUPPORTS_EDITOR_SHADERS default define moved to Platform.ush to enable use in Common.ush
- Fetch vertex shader custom data from GPU-Scene whenever GPU-scene instance culling path is used.
- Plumb through per-instance editor data to upload in the payload buffer of GPU-Scene
- Skip uploading per-instance vertex attributes and creating VF uniform buffer when GPU-Scene is active (ISM)
- Fetch all instance attributes from GPU-Scene, custom data and editor data, and only do so for editor platforms (LocalVF)
- When doing UpdateInstances only recreate mesh draw commands if the instance count changed.  Since we no longer have the per instance render data buffers they no longer are in the MDC and don't need to update.  The only thing left in the MDC that needs updating is the instance count.
- Skip creating/flushing per-instance render data buffers for ISM proxies.

#ddcfill: 61729e14c33baf00011f4e8b

diff --git a/Engine/Shaders/Private/Common.ush b/Engine/Shaders/Private/Common.ush
--- a/Engine/Shaders/Private/Common.ush
+++ b/Engine/Shaders/Private/Common.ush
@@ -146,6 +146,8 @@ float4 FmodFloor(float4 Lhs, float4 Rhs)
 	#define VELOCITY_ENCODE_GAMMA 0
 #endif
 
+//Tie Editor features to platform support and the COMPILE_SHADERS_FOR_DEVELOPMENT which is set via CVAR.
+#define USE_EDITOR_SHADERS (PLATFORM_SUPPORTS_EDITOR_SHADERS && USE_DEVELOPMENT_SHADERS)
 
 // Using SV_ClipDistance has overhead (15% slower base pass in triangle bound test scene on PS4) so projects have to opt-in
 #define USE_GLOBAL_CLIP_PLANE (PLATFORM_SUPPORTS_GLOBAL_CLIP_PLANE && PROJECT_ALLOW_GLOBAL_CLIP_PLANE && !MATERIAL_DOMAIN_POSTPROCESS && !MATERIAL_DOMAIN_UI)
diff --git a/Engine/Shaders/Private/LocalVertexFactory.ush b/Engine/Shaders/Private/LocalVertexFactory.ush
--- a/Engine/Shaders/Private/LocalVertexFactory.ush
+++ b/Engine/Shaders/Private/LocalVertexFactory.ush
@@ -593,6 +593,10 @@ FMaterialVertexParameters GetMaterialVertexParameters(FVertexFactoryInput Input,
 		Result.PerInstanceParams = Intermediates.PerInstanceParams;
 		Result.PrimitiveLocalInstanceIndex = Intermediates.PrimitiveLocalInstanceIndex;
 	#endif
+#if USES_PER_INSTANCE_CUSTOM_DATA
+	Result.CustomDataOffset = Intermediates.SceneData.InstanceData.CustomDataOffset;
+	Result.CustomDataCount = Intermediates.SceneData.InstanceData.CustomDataCount;
+#endif
 #endif
 
 	Result.PreSkinnedPosition = Intermediates.PreSkinPosition.xyz;
@@ -860,28 +864,23 @@ FVertexFactoryIntermediates GetVertexFactoryIntermediates(FVertexFactoryInput In
 		float B = 0;
 
 		Intermediates.InstanceLightmapAndShadowMapUVBias = GetInstanceData(Intermediates).LightMapAndShadowMapUVBias;
-	#elif MANUAL_VERTEX_FETCH
-		Intermediates.PerInstanceRandom = InstanceVF.VertexFetch_InstanceOriginBuffer[Intermediates.PrimitiveLocalInstanceIndex].w;
-
-		float R = InstanceVF.VertexFetch_InstanceTransformBuffer[3 * (Intermediates.PrimitiveLocalInstanceIndex) + 0].w;
-		float G = InstanceVF.VertexFetch_InstanceTransformBuffer[3 * (Intermediates.PrimitiveLocalInstanceIndex) + 1].w;
-		float B = InstanceVF.VertexFetch_InstanceTransformBuffer[3 * (Intermediates.PrimitiveLocalInstanceIndex) + 2].w;
-
-		Intermediates.InstanceLightmapAndShadowMapUVBias = InstanceVF.VertexFetch_InstanceLightmapBuffer[Intermediates.PrimitiveLocalInstanceIndex];
 	#else
-		Intermediates.PerInstanceRandom = Input.InstanceOrigin.w;
-
-		float R = Input.InstanceTransform1.w;
-		float G = Input.InstanceTransform2.w;
-		float B = Input.InstanceTransform3.w;
-
-		Intermediates.InstanceLightmapAndShadowMapUVBias = Input.InstanceLightmapAndShadowMapUVBias;
+	Intermediates.InstanceLightmapAndShadowMapUVBias = Intermediates.SceneData.InstanceData.LightMapAndShadowMapUVBias;
+	Intermediates.PerInstanceRandom = Intermediates.SceneData.InstanceData.RandomID;
+#if USE_EDITOR_SHADERS
+	Intermediates.IsSelected = Intermediates.SceneData.InstanceData.EditorData.bIsSelected ? 1.0f : 0.0f;
+	Intermediates.HitProxyId = float4(Intermediates.SceneData.InstanceData.EditorData.HitProxyId, 0.0f);
+#else // !USE_EDITOR_SHADERS
+	Intermediates.IsSelected = 0.0f;
+	Intermediates.HitProxyId = float4(0.0f, 0.0f, 0.0f, 0.0f);
+#endif // USE_EDITOR_SHADERS
 	#endif
 
 		Intermediates.IsSelected = trunc(R * (1.0 / 256.0));;
 		R = R - 256.0 * Intermediates.IsSelected;
 		Intermediates.HitProxyId = float4(R / 255.0, G / 255.0, B / 255.0, 0);
 
+	{
 		// This repackaging is stupid-seeming but these params are propably packed off to some interpolator or something
 		Intermediates.PerInstanceParams.x = Intermediates.PerInstanceRandom;
 		
diff --git a/Engine/Shaders/Private/MaterialTemplate.ush b/Engine/Shaders/Private/MaterialTemplate.ush
--- a/Engine/Shaders/Private/MaterialTemplate.ush
+++ b/Engine/Shaders/Private/MaterialTemplate.ush
@@ -117,14 +117,6 @@
 	#define USE_STENCIL_LOD_DITHER	USE_STENCIL_LOD_DITHER_DEFAULT
 #endif
 
-//Platforms that don't run the editor shouldn't need editor features in the shaders.
-#ifndef PLATFORM_SUPPORTS_EDITOR_SHADERS
-#define PLATFORM_SUPPORTS_EDITOR_SHADERS 1
-#endif
-
-//Tie Editor features to platform support and the COMPILE_SHADERS_FOR_DEVELOPMENT which is set via CVAR.
-#define USE_EDITOR_SHADERS (PLATFORM_SUPPORTS_EDITOR_SHADERS && USE_DEVELOPMENT_SHADERS)
-
 //Materials also have to opt in to these features.
 #define USE_EDITOR_COMPOSITING (USE_EDITOR_SHADERS && EDITOR_PRIMITIVE_MATERIAL)
 
@@ -479,6 +471,12 @@ struct FMaterialVertexParameters
 	float3 WorldPosition;
 	// TangentToWorld[2] is WorldVertexNormal
 	half3x3 TangentToWorld;
+	half3x3 TangentToWorld;
+
+#if USES_PER_INSTANCE_CUSTOM_DATA
+	uint CustomDataOffset;
+	uint CustomDataCount;
+#endif
 	
 #if USE_INSTANCING || USE_INSTANCE_CULLING || IS_MESHPARTICLE_FACTORY
 	FLWCMatrix InstanceLocalToWorld;
@@ -488,9 +486,7 @@ struct FMaterialVertexParameters
 	/** Per-instance properties. */
 	float3 InstanceLocalPosition;
 	float4 PerInstanceParams;
-#if USE_INSTANCE_CULLING
-	int PrimitiveLocalInstanceIndex;
-#else
+#if !USE_INSTANCE_CULLING
 	uint InstanceId;
 #endif
 	uint InstanceOffset;
@@ -1009,16 +1005,21 @@ float GetPerInstanceCustomData(FMaterialPixelParameters Parameters, int Index, f
 /** Get the per-instance custom data when instancing */
 float GetPerInstanceCustomData(FMaterialVertexParameters Parameters, int Index, float DefaultValue)
 {
-#if (USE_INSTANCING || USE_INSTANCE_CULLING) && USES_PER_INSTANCE_CUSTOM_DATA
+#if USE_INSTANCE_CULLING && USES_PER_INSTANCE_CUSTOM_DATA
+	const uint FloatIndex = uint(Index);
+	BRANCH
+	if (FloatIndex < Parameters.CustomDataCount)
+	{
+		const uint   Float4Offset = Parameters.CustomDataOffset + (FloatIndex >> 2u);
+		const float4 Float4Packed = LoadInstancePayloadDataElement(Float4Offset);
+		return Float4Packed[FloatIndex & 0x3u];
+	}
+#elif USE_INSTANCING && USES_PER_INSTANCE_CUSTOM_DATA
 	const uint FloatIndex = uint(Index);
 	BRANCH
 	if (FloatIndex < InstanceVF.NumCustomDataFloats)
 	{
-	#if USE_INSTANCE_CULLING
-		const uint InstanceDataIndex = Parameters.PrimitiveLocalInstanceIndex;
-	#else
 		const uint InstanceDataIndex = Parameters.InstanceId + Parameters.InstanceOffset;
-	#endif
 		const uint BufferStartIndex = InstanceDataIndex * InstanceVF.NumCustomDataFloats;
 		return InstanceVF.InstanceCustomDataBuffer[BufferStartIndex + FloatIndex];
 	}
@@ -1031,16 +1032,16 @@ float GetPerInstanceCustomData(FMaterialVertexParameters Parameters, int Index,
 /** Get the per-instance custom data when instancing */
 MaterialFloat3 GetPerInstanceCustomData3Vector(FMaterialVertexParameters Parameters, int Index, MaterialFloat3 DefaultValue)
 {
-#if (USE_INSTANCING || USE_INSTANCE_CULLING) && USES_PER_INSTANCE_CUSTOM_DATA
+#if USE_INSTANCE_CULLING && USES_PER_INSTANCE_CUSTOM_DATA
+		return float3(GetPerInstanceCustomData(Parameters, Index, DefaultValue.x), 
+					  GetPerInstanceCustomData(Parameters, Index + 1, DefaultValue.y),
+					  GetPerInstanceCustomData(Parameters, Index + 2, DefaultValue.z));
+#elif USE_INSTANCING && USES_PER_INSTANCE_CUSTOM_DATA
 	const uint FloatIndex = uint(Index);
 	BRANCH
 	if (FloatIndex + 2 < InstanceVF.NumCustomDataFloats)
 	{
-	#if USE_INSTANCE_CULLING
-		const uint InstanceDataIndex = Parameters.PrimitiveLocalInstanceIndex;
-	#else
 		const uint InstanceDataIndex = Parameters.InstanceId + Parameters.InstanceOffset;
-	#endif
 
 		const uint BufferStartIndex = InstanceDataIndex * InstanceVF.NumCustomDataFloats;
 		return float3(InstanceVF.InstanceCustomDataBuffer[BufferStartIndex + FloatIndex], 
diff --git a/Engine/Shaders/Private/SceneData.ush b/Engine/Shaders/Private/SceneData.ush
--- a/Engine/Shaders/Private/SceneData.ush
+++ b/Engine/Shaders/Private/SceneData.ush
@@ -68,6 +68,7 @@
 #define INSTANCE_SCENE_DATA_FLAG_HAS_LIGHTSHADOW_UV_BIAS	0x10
 #define INSTANCE_SCENE_DATA_FLAG_HAS_HIERARCHY_OFFSET		0x20
 #define INSTANCE_SCENE_DATA_FLAG_HAS_LOCAL_BOUNDS			0x40
+#define INSTANCE_SCENE_DATA_FLAG_HAS_EDITOR_DATA			0x80
 
 #define INVALID_PRIMITIVE_ID 0x000FFFFFu
 
@@ -330,6 +331,14 @@ uint GetPrimitive_LightingChannelMask(FPrimitiveSceneData Primitive)
 #endif
 
 // Unpacked AoS layout - see FInstanceSceneShaderData::Setup() for SoA packed layout.
+#if USE_EDITOR_SHADERS
+struct FInstanceSceneEditorData
+{
+	float3 HitProxyId;
+	bool bIsSelected;
+};
+#endif
+
 struct FInstanceSceneData
 {
 	FLWCMatrix LocalToWorld;
@@ -346,7 +355,7 @@ struct FInstanceSceneData
 	uint     LastUpdateSceneFrameNumber;
 	uint     NaniteRuntimeResourceID;
 	uint     NaniteHierarchyOffset;
-#if USES_PER_INSTANCE_RANDOM
+#if 1//USES_PER_INSTANCE_RANDOM
 	float    RandomID;
 #endif
 #if USES_PER_INSTANCE_CUSTOM_DATA
@@ -358,6 +367,10 @@ struct FInstanceSceneData
 #endif
 	bool     ValidInstance;
 	uint     Flags;
+
+#if USE_EDITOR_SHADERS
+	FInstanceSceneEditorData EditorData;
+#endif
 };
 
 #if USE_GLOBAL_GPU_SCENE_DATA
@@ -569,10 +582,10 @@ FInstanceSceneData GetInstanceSceneData(uint InstanceId, uint SOAStride, bool bC
 		const uint RelativeIdAndCustomDataCount					= asuint(LoadInstanceSceneDataElement(0 * SOAStride + InstanceId).y);
 		InstanceData.RelativeId									= BitFieldExtractU32(RelativeIdAndCustomDataCount, 24, 0); // Max of 16,777,216 instances per primitive
 	#if USES_PER_INSTANCE_CUSTOM_DATA
-		InstanceData.CustomDataCount							= BitFieldExtractU32(RelativeIdAndCustomDataCount, 8, 24); // Max of 255 custom data float4 per instance
+		InstanceData.CustomDataCount							= BitFieldExtractU32(RelativeIdAndCustomDataCount, 8, 24); // Max of 255 custom data floats per instance
 	#endif
 		InstanceData.LastUpdateSceneFrameNumber					= asuint(LoadInstanceSceneDataElement(0 * SOAStride + InstanceId).z);
-	#if USES_PER_INSTANCE_RANDOM
+	#if 1//USES_PER_INSTANCE_RANDOM
 		InstanceData.RandomID									= LoadInstanceSceneDataElement(0 * SOAStride + InstanceId).w;
 	#endif
 
@@ -632,11 +645,17 @@ FInstanceSceneData GetInstanceSceneData(uint InstanceId, uint SOAStride, bool bC
 		const bool bHasLightShadowUVBias		= (InstanceData.Flags & INSTANCE_SCENE_DATA_FLAG_HAS_LIGHTSHADOW_UV_BIAS) != 0u;
 		const bool bHasCustomData				= (InstanceData.Flags & INSTANCE_SCENE_DATA_FLAG_HAS_CUSTOM_DATA) != 0u;
 		const bool bHasRandomID					= (InstanceData.Flags & INSTANCE_SCENE_DATA_FLAG_HAS_RANDOM) != 0u;
+#if USE_EDITOR_SHADERS
+		const bool bHasEditorData				= (InstanceData.Flags & INSTANCE_SCENE_DATA_FLAG_HAS_EDITOR_DATA) != 0u;
+#else //
+		const bool bHasEditorData				= false;
+#endif
 
 		// Offsets are in float4s
 		const uint Offset_HierarchyOffset		= PayloadDataGlobalOffset + 0u; // float0.x
-		const uint Offset_LocalBounds			= Offset_HierarchyOffset; // float0.yzw
-		const uint Offset_DynamicData			= CondMask(bHasLocalBounds, 2u, CondMask(bHasHierarchyOffset, 1u, 0u)) + PayloadDataGlobalOffset;
+		const uint Offset_EditorData			= PayloadDataGlobalOffset + 0u; // float0.y
+		const uint Offset_LocalBounds			= Offset_HierarchyOffset; // float0.zw + float1.xyzw
+		const uint Offset_DynamicData			= CondMask(bHasLocalBounds, 2u, CondMask(bHasHierarchyOffset || bHasEditorData, 1u, 0u)) + PayloadDataGlobalOffset;
 		const uint Offset_LightShadowUVBias		= CondMask(bHasDynamicData, TransformSizeFloat4Count, 0u) + Offset_DynamicData;
 	#if USES_PER_INSTANCE_CUSTOM_DATA
 		InstanceData.CustomDataOffset			= CondMask(bHasCustomData, CondMask(bHasLightShadowUVBias, 1u, 0u) + Offset_LightShadowUVBias, 0xFFFFFFFFu);
@@ -654,12 +673,22 @@ FInstanceSceneData GetInstanceSceneData(uint InstanceId, uint SOAStride, bool bC
 			InstanceData.NaniteHierarchyOffset += HierarchyRootOffset;
 		}
 
+#if USE_EDITOR_SHADERS
+		BRANCH
+		if (bHasEditorData)
+		{
+			const uint PackedEditorData = asuint(LoadInstancePayloadDataElement(Offset_EditorData)).y;
+
+			InstanceData.EditorData.bIsSelected = (PackedEditorData >> 24u) != 0;
+			InstanceData.EditorData.HitProxyId = float3(float(PackedEditorData & 0xFF), float((PackedEditorData >> 8u) & 0xFF), float((PackedEditorData >> 16u) & 0xFF)) * (1.0f / 255.0f);
+		}
+#endif
 		// TODO: Permutation? Should be mutually exclusive for specific VFs (i.e. geometry collections), in theory
 		BRANCH
 		if (bHasLocalBounds)
 		{
-			InstanceData.LocalBoundsCenter = LoadInstancePayloadDataElement(Offset_LocalBounds + 0).yzw;
-			InstanceData.LocalBoundsExtent = LoadInstancePayloadDataElement(Offset_LocalBounds + 1).xyz;
+			InstanceData.LocalBoundsCenter = float3(LoadInstancePayloadDataElement(Offset_LocalBounds + 0).zw, LoadInstancePayloadDataElement(Offset_LocalBounds + 1).x);
+			InstanceData.LocalBoundsExtent = LoadInstancePayloadDataElement(Offset_LocalBounds + 1).yzw;
 		}
 		else
 		{
diff --git a/Engine/Shaders/Public/Platform.ush b/Engine/Shaders/Public/Platform.ush
--- a/Engine/Shaders/Public/Platform.ush
+++ b/Engine/Shaders/Public/Platform.ush
@@ -182,6 +182,11 @@
 #define COMPILER_SUPPORTS_WAVE_32_64_MODE 0
 #endif
 
+//Platforms that don't run the editor shouldn't need editor features in the shaders.
+#ifndef PLATFORM_SUPPORTS_EDITOR_SHADERS
+#define PLATFORM_SUPPORTS_EDITOR_SHADERS 1
+#endif
+
 // ---------------------------------------------------- Alternative floating point types
 
 #ifndef FORCE_FLOATS
diff --git a/Engine/Source/Runtime/Engine/Classes/Engine/InstancedStaticMesh.h b/Engine/Source/Runtime/Engine/Classes/Engine/InstancedStaticMesh.h
--- a/Engine/Source/Runtime/Engine/Classes/Engine/InstancedStaticMesh.h
+++ b/Engine/Source/Runtime/Engine/Classes/Engine/InstancedStaticMesh.h
@@ -93,6 +93,14 @@ public:
 		InstanceData->GetInstanceRandomID(InstanceIndex, RandomInstanceID);
 	}
 
+#if WITH_EDITOR
+	FORCEINLINE void GetInstanceEditorData(int32 InstanceIndex, FColor& HitProxyColorOut, bool& bSelectedOut) const
+	{
+		InstanceData->GetInstanceEditorData(InstanceIndex, HitProxyColorOut, bSelectedOut);
+	}
+#endif 
+
+
 	FORCEINLINE void GetInstanceLightMapData(int32 InstanceIndex, FVector4f& InstanceLightmapAndShadowMapUVBias) const
 	{
 		InstanceData->GetInstanceLightMapData(InstanceIndex, InstanceLightmapAndShadowMapUVBias);
diff --git a/Engine/Source/Runtime/Engine/Private/InstancedStaticMesh.cpp b/Engine/Source/Runtime/Engine/Private/InstancedStaticMesh.cpp
--- a/Engine/Source/Runtime/Engine/Private/InstancedStaticMesh.cpp
+++ b/Engine/Source/Runtime/Engine/Private/InstancedStaticMesh.cpp
@@ -868,6 +868,8 @@ void FInstancedStaticMeshVertexFactory::InitRHI()
 		Elements.Add(AccessStreamComponent(Data.PositionComponent,0));
 	}
 
+	const bool bCanUseGPUScene = UseGPUScene(GMaxRHIShaderPlatform, GMaxRHIFeatureLevel);
+
 	// only tangent,normal are used by the stream. the binormal is derived in the shader
 	uint8 TangentBasisAttributes[2] = { 1, 2 };
 	for(int32 AxisIndex = 0;AxisIndex < 2;AxisIndex++)
@@ -962,6 +964,7 @@ void FInstancedStaticMeshVertexFactory::InitRHI()
 	// we don't need per-vertex shadow or lightmap rendering
 	InitDeclaration(Elements);
 
+	if (!bCanUseGPUScene)
 	{
 		FInstancedStaticMeshVertexFactoryUniformShaderParameters UniformParameters;
 		UniformParameters.VertexFetch_InstanceOriginBuffer = GetInstanceOriginSRV();
@@ -997,7 +1000,11 @@ void FInstancedStaticMeshRenderData::BindBuffersToVertexFactories()
 
 	check(IsInRenderingThread());
 
-	PerInstanceRenderData->InstanceBuffer.FlushGPUUpload();
+	const bool bCanUseGPUScene = UseGPUScene(GMaxRHIShaderPlatform, GMaxRHIFeatureLevel);
+	if (!bCanUseGPUScene)
+	{
+		PerInstanceRenderData->InstanceBuffer.FlushGPUUpload();
+	}
 
 	for (int32 LODIndex = 0; LODIndex < VertexFactories.Num(); LODIndex++)
 	{
@@ -1025,8 +1032,10 @@ void FInstancedStaticMeshRenderData::BindBuffersToVertexFactories()
 		}
 
 		check(PerInstanceRenderData);
-
-		PerInstanceRenderData->InstanceBuffer.BindInstanceVertexBuffer(&VertexFactory, Data);
+		if (!bCanUseGPUScene)
+		{
+			PerInstanceRenderData->InstanceBuffer.BindInstanceVertexBuffer(&VertexFactory, Data);
+		}
 
 		VertexFactory.SetData(Data);
 		VertexFactory.InitResource();
@@ -1401,6 +1410,11 @@ void FInstancedStaticMeshSceneProxy::SetupProxy(UInstancedStaticMeshComponent* I
 
 		InstanceRandomID.SetNumZeroed(bHasPerInstanceRandom ? InComponent->GetInstanceCount() : 0); // Only allocate if material bound which uses this
 
+#if WITH_EDITOR
+		bHasPerInstanceEditorData = true;
+		InstanceEditorData.SetNumZeroed(bHasPerInstanceEditorData ? InComponent->GetInstanceCount() : 0);
+#endif
+
 		// Only allocate if material bound which uses this
 		if (bHasPerInstanceCustomData && InComponent->NumCustomDataFloats > 0)
 		{
@@ -1469,7 +1483,11 @@ void FInstancedStaticMeshSceneProxy::CreateRenderThreadResources()
 	if (ensure(InstancedRenderData.PerInstanceRenderData.IsValid()))
 	{
 		FStaticMeshInstanceBuffer& InstanceBuffer = InstancedRenderData.PerInstanceRenderData->InstanceBuffer;
-		InstanceBuffer.FlushGPUUpload();
+		const bool bCanUseGPUScene = UseGPUScene(GMaxRHIShaderPlatform, GMaxRHIFeatureLevel);
+		if (!bCanUseGPUScene)
+		{
+			InstanceBuffer.FlushGPUUpload();
+		}
 	}
 
 	if (UseGPUScene(GetScene().GetShaderPlatform(), GetScene().GetFeatureLevel()))
@@ -1493,6 +1511,9 @@ void FInstancedStaticMeshSceneProxy::CreateRenderThreadResources()
 			{
 				const bool bHasLightMapData = bHasPerInstanceLMSMUVBias && InstanceLightShadowUVBias.Num() == InstanceSceneData.Num();
 				const bool bHasRandomID = bHasPerInstanceRandom && InstanceRandomID.Num() == InstanceSceneData.Num();
+#if WITH_EDITOR
+				const bool bHasEditorData = bHasPerInstanceEditorData && InstanceEditorData.Num() == InstanceSceneData.Num();
+#endif
 
 				for (int32 InstanceIndex = 0; InstanceIndex < InstanceSceneData.Num(); ++InstanceIndex)
 				{
@@ -1508,6 +1529,16 @@ void FInstancedStaticMeshSceneProxy::CreateRenderThreadResources()
 					{
 						InstanceBuffer.GetInstanceLightMapData(InstanceIndex, InstanceLightShadowUVBias[InstanceIndex]);
 					}
+
+#if WITH_EDITOR
+					if (bHasEditorData)
+					{
+						FColor HitProxyColor;
+						bool bSelected;
+						InstanceBuffer.GetInstanceEditorData(InstanceIndex, HitProxyColor, bSelected);
+						InstanceEditorData[InstanceIndex] = FInstanceUpdateCmdBuffer::PackEditorData(HitProxyColor, bSelected);
+					}
+#endif
 				}
 			}
 		}
@@ -4506,21 +4537,24 @@ void FInstancedStaticMeshVertexFactoryShaderParameters::GetElementShaderBindings
 	const auto* InstancedVertexFactory = static_cast<const FInstancedStaticMeshVertexFactory*>(VertexFactory);
 	const int32 InstanceOffsetValue = BatchElement.UserIndex;
 
-	ShaderBindings.Add(Shader->GetUniformBufferParameter<FInstancedStaticMeshVertexFactoryUniformShaderParameters>(), InstancedVertexFactory->GetUniformBuffer());
 	ShaderBindings.Add(InstanceOffset, InstanceOffsetValue);
 	
-	if (InstancedVertexFactory->SupportsManualVertexFetch(FeatureLevel))
+	if (!UseGPUScene(Scene->GetShaderPlatform()))
 	{
-		ShaderBindings.Add(VertexFetch_InstanceOriginBufferParameter, InstancedVertexFactory->GetInstanceOriginSRV());
-		ShaderBindings.Add(VertexFetch_InstanceTransformBufferParameter, InstancedVertexFactory->GetInstanceTransformSRV());
-		ShaderBindings.Add(VertexFetch_InstanceLightmapBufferParameter, InstancedVertexFactory->GetInstanceLightmapSRV());
-	}
-	if (InstanceOffsetValue > 0 && VertexStreams.Num() > 0)
-	{
-		// GPUCULL_TODO: This here can still work together with the instance attributes for index, but note that all instance attributes then must assume they are offset wrt the on-the-fly generate buffer
-		//          so with the new scheme there is no clear way this can work in the vanilla instancing way as there is an indirection. So either other attributes must be loaded in the shader or they
-		//          would have to be copied as the instance ID is now - not good.
-		VertexFactory->OffsetInstanceStreams(InstanceOffsetValue, InputStreamType, VertexStreams);
+		ShaderBindings.Add(Shader->GetUniformBufferParameter<FInstancedStaticMeshVertexFactoryUniformShaderParameters>(), InstancedVertexFactory->GetUniformBuffer());
+		if (InstancedVertexFactory->SupportsManualVertexFetch(FeatureLevel))
+		{
+			ShaderBindings.Add(VertexFetch_InstanceOriginBufferParameter, InstancedVertexFactory->GetInstanceOriginSRV());
+			ShaderBindings.Add(VertexFetch_InstanceTransformBufferParameter, InstancedVertexFactory->GetInstanceTransformSRV());
+			ShaderBindings.Add(VertexFetch_InstanceLightmapBufferParameter, InstancedVertexFactory->GetInstanceLightmapSRV());
+		}
+		if (InstanceOffsetValue > 0 && VertexStreams.Num() > 0)
+		{
+			// GPUCULL_TODO: This here can still work together with the instance attributes for index, but note that all instance attributes then must assume they are offset wrt the on-the-fly generate buffer
+			//          so with the new scheme there is no clear way this can work in the vanilla instancing way as there is an indirection. So either other attributes must be loaded in the shader or they
+			//          would have to be copied as the instance ID is now - not good.
+			VertexFactory->OffsetInstanceStreams(InstanceOffsetValue, InputStreamType, VertexStreams);
+		}
 	}
 
 	if( InstancingWorldViewOriginOneParameter.IsBound() )
diff --git a/Engine/Source/Runtime/Engine/Private/PrimitiveSceneProxy.cpp b/Engine/Source/Runtime/Engine/Private/PrimitiveSceneProxy.cpp
--- a/Engine/Source/Runtime/Engine/Private/PrimitiveSceneProxy.cpp
+++ b/Engine/Source/Runtime/Engine/Private/PrimitiveSceneProxy.cpp
@@ -177,6 +177,9 @@ FPrimitiveSceneProxy::FPrimitiveSceneProxy(const UPrimitiveComponent* InComponen
 ,	bHasPerInstanceLMSMUVBias(false)
 ,	bHasPerInstanceLocalBounds(false)
 ,	bHasPerInstanceHierarchyOffset(false)
+#if WITH_EDITOR
+,	bHasPerInstanceEditorData(false)
+#endif
 ,	bUseAsOccluder(InComponent->bUseAsOccluder)
 ,	bAllowApproximateOcclusion(InComponent->Mobility != EComponentMobility::Movable)
 ,	bSelectable(InComponent->bSelectable)
@@ -479,11 +482,11 @@ uint32 FPrimitiveSceneProxy::GetPayloadDataStride() const
 	// Hierarchy is packed in with local bounds if they are both present (almost always the case)
 	if (HasPerInstanceLocalBounds())
 	{
-		PayloadDataCount += 2; // FRenderBounds and possibly uint32 for hierarchy offset
+		PayloadDataCount += 2; // FRenderBounds and possibly uint32 for hierarchy offset & another uint32 for EditorData
 	}
-	else if (HasPerInstanceHierarchyOffset())
+	else if (HasPerInstanceHierarchyOffset() || HasPerInstanceEditorData())
 	{
-		PayloadDataCount += 1; // uint32 for hierarchy offset (float4 packed)
+		PayloadDataCount += 1; // uint32 for hierarchy offset (float4 packed) & instance editor data is packed in the same float4
 	}
 
 	PayloadDataCount += HasPerInstanceLMSMUVBias() ? 1 : 0; // FVector4
@@ -556,29 +559,44 @@ void FPrimitiveSceneProxy::UpdateInstances_RenderThread(const FInstanceUpdateCmd
 		// Apply all updates.
 		for (const auto& Cmd : CmdBuffer.Cmds)
 		{
-			if (Cmd.Type == FInstanceUpdateCmdBuffer::Update)
+			switch (Cmd.Type)
 			{
-				// update transform data.
-				InstanceSceneData[Cmd.InstanceIndex].LocalToPrimitive = Cmd.XForm;
-				InstanceDynamicData[Cmd.InstanceIndex].PrevLocalToPrimitive = Cmd.PreviousXForm;
+				case FInstanceUpdateCmdBuffer::Update:
+				{
+					// update transform data.
+					InstanceSceneData[Cmd.InstanceIndex].LocalToPrimitive = Cmd.XForm;
+					InstanceDynamicData[Cmd.InstanceIndex].PrevLocalToPrimitive = Cmd.PreviousXForm;
 
 #if !(UE_BUILD_SHIPPING || UE_BUILD_TEST)
-				InstanceXFormUpdatedThisFrame[Cmd.InstanceIndex] = true;
+					InstanceXFormUpdatedThisFrame[Cmd.InstanceIndex] = true;
 #endif
-			}
-			else if (Cmd.Type == FInstanceUpdateCmdBuffer::CustomData)
-			{
-				check(bHasCustomFloatData);
+				}
+				break;
+				case FInstanceUpdateCmdBuffer::CustomData:
+				{
+					check(bHasCustomFloatData);
 
-				// update custom data because it changed.
-				check(PrevNumCustomDataFloats == NumCustomDataFloats);
-				const int32 DstCustomDataOffset = Cmd.InstanceIndex * NumCustomDataFloats;
-				FMemory::Memcpy(&InstanceCustomData[DstCustomDataOffset], &Cmd.CustomDataFloats[0], NumCustomDataFloats * sizeof(float));
+					// update custom data because it changed.
+					check(PrevNumCustomDataFloats == NumCustomDataFloats);
+					const int32 DstCustomDataOffset = Cmd.InstanceIndex * NumCustomDataFloats;
+					FMemory::Memcpy(&InstanceCustomData[DstCustomDataOffset], &Cmd.CustomDataFloats[0], NumCustomDataFloats * sizeof(float));
 
 #if !(UE_BUILD_SHIPPING || UE_BUILD_TEST)
-				InstanceCustomDataUpdatedThisFrame[Cmd.InstanceIndex] = true;
+					InstanceCustomDataUpdatedThisFrame[Cmd.InstanceIndex] = true;
 #endif
-			}
+				}
+				break;
+#if WITH_EDITOR
+				case FInstanceUpdateCmdBuffer::EditorData:
+				{
+					check(bHasPerInstanceEditorData);
+					InstanceEditorData[Cmd.InstanceIndex] = FInstanceUpdateCmdBuffer::PackEditorData(Cmd.HitProxyColor, Cmd.bSelected);
+				}
+				break;
+#endif //WITH_EDITOR
+				default:
+				break;
+			};
 		}
 
 		// Build bit array of commands to remove.
@@ -609,6 +627,13 @@ void FPrimitiveSceneProxy::UpdateInstances_RenderThread(const FInstanceUpdateCmd
 					InstanceLightShadowUVBias.RemoveAtSwap(i, 1, false);
 				}
 
+#if WITH_EDITOR
+				if (bHasPerInstanceEditorData)
+				{
+					InstanceEditorData.RemoveAtSwap(i, 1, false);
+				}
+#endif
+
 				// Only remove the custom float data from this instance if it previously had it.
 				if (bPreviouslyHadCustomFloatData)
 				{
@@ -644,6 +669,12 @@ void FPrimitiveSceneProxy::UpdateInstances_RenderThread(const FInstanceUpdateCmd
 					InstanceLightShadowUVBias.AddZeroed();
 				}
 
+#if WITH_EDITOR
+				if (bHasPerInstanceEditorData)
+				{
+					InstanceEditorData.AddZeroed();
+				}
+#endif
 				if (bHasCustomFloatData)
 				{
 					const int32 DstCustomDataOffset = InstanceCustomData.AddUninitialized(NumCustomDataFloats);
@@ -664,6 +695,9 @@ void FPrimitiveSceneProxy::UpdateInstances_RenderThread(const FInstanceUpdateCmd
 		bHasPerInstanceCustomData = InstanceCustomData.Num() > 0;
 		bHasPerInstanceDynamicData = InstanceDynamicData.Num() > 0;
 		bHasPerInstanceLMSMUVBias = InstanceLightShadowUVBias.Num() > 0;
+#if WITH_EDITOR
+		bHasPerInstanceEditorData = InstanceEditorData.Num() > 0;
+#endif
 
 		// Ensure our data is in sync.
 		check(InstanceSceneData.Num() == InstanceDynamicData.Num());
diff --git a/Engine/Source/Runtime/Engine/Public/InstanceUniformShaderParameters.h b/Engine/Source/Runtime/Engine/Public/InstanceUniformShaderParameters.h
--- a/Engine/Source/Runtime/Engine/Public/InstanceUniformShaderParameters.h
+++ b/Engine/Source/Runtime/Engine/Public/InstanceUniformShaderParameters.h
@@ -19,6 +19,7 @@
 #define INSTANCE_SCENE_DATA_FLAG_HAS_LIGHTSHADOW_UV_BIAS	0x10
 #define INSTANCE_SCENE_DATA_FLAG_HAS_HIERARCHY_OFFSET		0x20
 #define INSTANCE_SCENE_DATA_FLAG_HAS_LOCAL_BOUNDS			0x40
+#define INSTANCE_SCENE_DATA_FLAG_HAS_EDITOR_DATA			0x80
 
 #define INVALID_PRIMITIVE_ID 0x000FFFFFu
 
diff --git a/Engine/Source/Runtime/Engine/Public/PrimitiveInstanceUpdateCommand.h b/Engine/Source/Runtime/Engine/Public/PrimitiveInstanceUpdateCommand.h
--- a/Engine/Source/Runtime/Engine/Public/PrimitiveInstanceUpdateCommand.h
+++ b/Engine/Source/Runtime/Engine/Public/PrimitiveInstanceUpdateCommand.h
@@ -57,6 +57,12 @@ struct FInstanceUpdateCmdBuffer
 	void Reset();
 	int32 NumTotalCommands() const { return NumEdits; };
 
+	static inline uint32 PackEditorData(const FColor& HitProxyColor, bool bSelected)
+	{
+		return uint32(HitProxyColor.R) | uint32(HitProxyColor.G) << 8u | uint32(HitProxyColor.B) << 16u | (bSelected ? 1u << 24u : 0u);
+	}
+
+
 	TArray<FInstanceUpdateCommand> Cmds;
 	int32 NumCustomDataFloats;
 	int32 NumAdds;
diff --git a/Engine/Source/Runtime/Engine/Public/PrimitiveSceneProxy.h b/Engine/Source/Runtime/Engine/Public/PrimitiveSceneProxy.h
--- a/Engine/Source/Runtime/Engine/Public/PrimitiveSceneProxy.h
+++ b/Engine/Source/Runtime/Engine/Public/PrimitiveSceneProxy.h
@@ -642,6 +642,11 @@ public:
 	inline bool HasPerInstanceLMSMUVBias() const { return bHasPerInstanceLMSMUVBias; }
 	inline bool HasPerInstanceLocalBounds() const { return bHasPerInstanceLocalBounds; }
 	inline bool HasPerInstanceHierarchyOffset() const { return bHasPerInstanceHierarchyOffset; }
+#if WITH_EDITOR
+	inline bool HasPerInstanceEditorData() const { return bHasPerInstanceEditorData; }
+#else
+	FORCEINLINE bool HasPerInstanceEditorData() const { return false; }
+#endif // WITH_EDITOR
 
 	inline bool HasAnyPerInstancePayloadData() const
 	{
@@ -651,6 +656,9 @@ public:
 			bHasPerInstanceDynamicData	|
 			bHasPerInstanceLMSMUVBias	|
 			bHasPerInstanceLocalBounds	|
+#if WITH_EDITOR
+			bHasPerInstanceEditorData	|
+#endif
 			bHasPerInstanceHierarchyOffset;
 	}
 
@@ -663,6 +671,10 @@ public:
 		PayloadDataFlags |= HasPerInstanceLMSMUVBias()      ? INSTANCE_SCENE_DATA_FLAG_HAS_LIGHTSHADOW_UV_BIAS	: 0u;
 		PayloadDataFlags |= HasPerInstanceHierarchyOffset() ? INSTANCE_SCENE_DATA_FLAG_HAS_HIERARCHY_OFFSET		: 0u;
 		PayloadDataFlags |= HasPerInstanceLocalBounds()     ? INSTANCE_SCENE_DATA_FLAG_HAS_LOCAL_BOUNDS			: 0u;
+#if WITH_EDITOR
+		PayloadDataFlags |= HasPerInstanceEditorData()      ? INSTANCE_SCENE_DATA_FLAG_HAS_EDITOR_DATA			: 0u;
+#endif
+
 		return PayloadDataFlags;
 	}
 
@@ -806,6 +818,13 @@ public:
 		return InstanceLocalBounds;
 	}
 
+#if WITH_EDITOR
+	FORCEINLINE TConstArrayView<uint32> GetInstanceEditorData() const
+	{
+		return InstanceEditorData;
+	}
+#endif // 
+
 	// Helper function to avoid multiple code paths requesting bounds
 	FORCEINLINE const FRenderBounds& GetInstanceLocalBounds(uint32 Instance) 
 	{
@@ -1189,6 +1208,9 @@ protected:
 	uint8 bHasPerInstanceLMSMUVBias : 1;
 	uint8 bHasPerInstanceLocalBounds : 1;
 	uint8 bHasPerInstanceHierarchyOffset : 1;
+#if WITH_EDITOR
+	uint8 bHasPerInstanceEditorData : 1;
+#endif
 
 private:
 
@@ -1245,6 +1267,10 @@ protected:
 	TArray<float> InstanceRandomID;
 	TArray<FVector4f> InstanceLightShadowUVBias;
 	TArray<uint32> InstanceHierarchyOffset;
+#if WITH_EDITOR
+	TArray<uint32> InstanceEditorData;
+#endif
+
 #if !(UE_BUILD_SHIPPING || UE_BUILD_TEST)
 	/** Whether instance data has changed this frame on the proxy. Currently used for non-shipping debug drawing. */
 	TBitArray<> InstanceXFormUpdatedThisFrame;
diff --git a/Engine/Source/Runtime/Engine/Public/StaticMeshResources.h b/Engine/Source/Runtime/Engine/Public/StaticMeshResources.h
--- a/Engine/Source/Runtime/Engine/Public/StaticMeshResources.h
+++ b/Engine/Source/Runtime/Engine/Public/StaticMeshResources.h
@@ -1196,6 +1196,27 @@ public:
 		RandomInstanceID = Origin.W;
 	}
 
+
+#if WITH_EDITOR
+	FORCEINLINE_DEBUGGABLE void GetInstanceEditorData(int32 InstanceIndex, FColor& HitProxyColorOut, bool & bSelectedOut) const
+	{
+		// TODO: put this into a sensible format
+		FVector4 InstanceTransform[3];
+		if (bUseHalfFloat)
+		{
+			GetInstanceTransformInternal<FFloat16>(InstanceIndex, InstanceTransform);
+		}
+		else
+		{
+			GetInstanceTransformInternal<float>(InstanceIndex, InstanceTransform);
+		}
+		bSelectedOut = InstanceTransform[0][3] > 255.0f;
+		HitProxyColorOut.R = uint8(InstanceTransform[0][3] - (bSelectedOut ? 256.0f : 0.0f));
+		HitProxyColorOut.G = uint8(InstanceTransform[1][3]);
+		HitProxyColorOut.B = uint8(InstanceTransform[2][3]);
+	}
+#endif 
+
 	FORCEINLINE_DEBUGGABLE void GetInstanceLightMapData(int32 InstanceIndex, FVector4f& InstanceLightmapAndShadowMapUVBias) const
 	{
 		GetInstanceLightMapDataInternal(InstanceIndex, InstanceLightmapAndShadowMapUVBias);
diff --git a/Engine/Source/Runtime/Renderer/Private/GPUScene.cpp b/Engine/Source/Runtime/Renderer/Private/GPUScene.cpp
--- a/Engine/Source/Runtime/Renderer/Private/GPUScene.cpp
+++ b/Engine/Source/Runtime/Renderer/Private/GPUScene.cpp
@@ -226,6 +226,9 @@ struct FInstanceUploadInfo
 	TConstArrayView<float> InstanceRandomID;
 	TConstArrayView<uint32> InstanceHierarchyOffset;
 	TConstArrayView<FRenderBounds> InstanceLocalBounds;
+#if WITH_EDITOR
+	TConstArrayView<uint32> InstanceEditorData;
+#endif
 
 	// Used for primitives that need to create a dummy instance (they do not have instance data in the proxy)
 	FPrimitiveInstance DummyInstance;
@@ -247,12 +250,19 @@ void ValidateInstanceUploadInfo(const FInstanceUploadInfo& UploadInfo)
 	const bool bHasDynamicData		= (UploadInfo.PayloadDataFlags & INSTANCE_SCENE_DATA_FLAG_HAS_DYNAMIC_DATA) != 0u;
 	const bool bHasLMSMScaleBias	= (UploadInfo.PayloadDataFlags & INSTANCE_SCENE_DATA_FLAG_HAS_LIGHTSHADOW_UV_BIAS) != 0u;
 	const bool bHasHierarchyOffset	= (UploadInfo.PayloadDataFlags & INSTANCE_SCENE_DATA_FLAG_HAS_HIERARCHY_OFFSET) != 0u;
+#if WITH_EDITOR
+	const bool bHasEditorData		= (UploadInfo.PayloadDataFlags & INSTANCE_SCENE_DATA_FLAG_HAS_EDITOR_DATA) != 0u;
+#endif
 
 	const int32 InstanceCount = UploadInfo.PrimitiveInstances.Num();
 	check(UploadInfo.InstanceRandomID.Num()				== (bHasRandomID		? InstanceCount : 0));
 	check(UploadInfo.InstanceDynamicData.Num()			== (bHasDynamicData		? InstanceCount : 0));
 	check(UploadInfo.InstanceLightShadowUVBias.Num()	== (bHasLMSMScaleBias	? InstanceCount : 0));
 	check(UploadInfo.InstanceHierarchyOffset.Num()		== (bHasHierarchyOffset	? InstanceCount : 0));
+#if WITH_EDITOR
+	check(UploadInfo.InstanceEditorData.Num() == (bHasEditorData ? InstanceCount : 0));
+#endif
+
 	if (bHasCustomData)
 	{
 		check(UploadInfo.InstanceCustomDataCount > 0);
@@ -405,6 +415,10 @@ struct FUploadDataSourceAdapterScenePrimitives
 			InstanceUploadInfo.InstanceCustomData = PrimitiveSceneProxy->GetInstanceCustomData();
 			InstanceUploadInfo.InstanceRandomID = PrimitiveSceneProxy->GetInstanceRandomID();
 			InstanceUploadInfo.InstanceHierarchyOffset = PrimitiveSceneProxy->GetInstanceHierarchyOffset();
+
+#if WITH_EDITOR
+			InstanceUploadInfo.InstanceEditorData = PrimitiveSceneProxy->GetInstanceEditorData();
+#endif
 		}
 		else
 		{
@@ -422,6 +436,9 @@ struct FUploadDataSourceAdapterScenePrimitives
 			InstanceUploadInfo.InstanceCustomData = TConstArrayView<float>();
 			InstanceUploadInfo.InstanceRandomID = TConstArrayView<float>();
 			InstanceUploadInfo.InstanceHierarchyOffset = TConstArrayView<uint32>();
+#if WITH_EDITOR
+			InstanceUploadInfo.InstanceEditorData = TConstArrayView<uint32>();
+#endif
 		}
 
 		InstanceUploadInfo.InstanceCustomDataCount = 0;
@@ -1242,11 +1259,15 @@ void FGPUScene::UploadGeneral(FRHICommandListImmediate& RHICmdList, FScene *Scen
 
 									int32 PayloadPosition = 0;
 
-									if (UploadInfo.PayloadDataFlags & (INSTANCE_SCENE_DATA_FLAG_HAS_HIERARCHY_OFFSET | INSTANCE_SCENE_DATA_FLAG_HAS_LOCAL_BOUNDS))
+									if (UploadInfo.PayloadDataFlags & (INSTANCE_SCENE_DATA_FLAG_HAS_HIERARCHY_OFFSET | INSTANCE_SCENE_DATA_FLAG_HAS_LOCAL_BOUNDS | INSTANCE_SCENE_DATA_FLAG_HAS_EDITOR_DATA))
 									{
 										const uint32 InstanceHierarchyOffset = (UploadInfo.PayloadDataFlags & INSTANCE_SCENE_DATA_FLAG_HAS_HIERARCHY_OFFSET) ? UploadInfo.InstanceHierarchyOffset[InstanceIndex] : 0;
 										InstancePayloadData[PayloadPosition].X = *(const float*)&InstanceHierarchyOffset;
 
+#if WITH_EDITOR
+										const uint32 InstanceEditorData = (UploadInfo.PayloadDataFlags & INSTANCE_SCENE_DATA_FLAG_HAS_EDITOR_DATA) ? UploadInfo.InstanceEditorData[InstanceIndex] : 0;
+										InstancePayloadData[PayloadPosition].Y = *(const float*)&InstanceEditorData;
+#endif
 										if (UploadInfo.PayloadDataFlags & INSTANCE_SCENE_DATA_FLAG_HAS_LOCAL_BOUNDS)
 										{
 											check(UploadInfo.InstanceLocalBounds.Num() == UploadInfo.PrimitiveInstances.Num());
@@ -1254,14 +1275,13 @@ void FGPUScene::UploadGeneral(FRHICommandListImmediate& RHICmdList, FScene *Scen
 											const FVector3f BoundsOrigin = InstanceLocalBounds.GetCenter();
 											const FVector3f BoundsExtent = InstanceLocalBounds.GetExtent();
 
-											InstancePayloadData[PayloadPosition + 0].Y = *(const float*)&BoundsOrigin.X;
-											InstancePayloadData[PayloadPosition + 0].Z = *(const float*)&BoundsOrigin.Y;
-											InstancePayloadData[PayloadPosition + 0].W = *(const float*)&BoundsOrigin.Z;
-
-											InstancePayloadData[PayloadPosition + 1].X = *(const float*)&BoundsExtent.X;
-											InstancePayloadData[PayloadPosition + 1].Y = *(const float*)&BoundsExtent.Y;
-											InstancePayloadData[PayloadPosition + 1].Z = *(const float*)&BoundsExtent.Z;
-											// .W = Unused
+											InstancePayloadData[PayloadPosition + 0].Z = *(const float*)&BoundsOrigin.X;
+											InstancePayloadData[PayloadPosition + 0].W = *(const float*)&BoundsOrigin.Y;
+											
+											InstancePayloadData[PayloadPosition + 1].X = *(const float*)&BoundsOrigin.Z;
+											InstancePayloadData[PayloadPosition + 1].Y = *(const float*)&BoundsExtent.X;
+											InstancePayloadData[PayloadPosition + 1].Z = *(const float*)&BoundsExtent.Y;
+											InstancePayloadData[PayloadPosition + 1].W = *(const float*)&BoundsExtent.Z;
 										}
 
 										PayloadPosition += (UploadInfo.PayloadDataFlags & INSTANCE_SCENE_DATA_FLAG_HAS_LOCAL_BOUNDS) ? 2 : 1;
@@ -1491,6 +1511,9 @@ struct FUploadDataSourceAdapterDynamicPrimitives
 			InstanceUploadInfo.InstanceRandomID				= TConstArrayView<float>();
 			InstanceUploadInfo.InstanceLocalBounds			= TConstArrayView<FRenderBounds>(&InstanceUploadInfo.DummyLocalBounds, 1);
 			InstanceUploadInfo.InstanceHierarchyOffset		= TConstArrayView<uint32>();
+#if WITH_EDITOR
+			InstanceUploadInfo.InstanceEditorData			= TConstArrayView<uint32>();
+#endif
 			InstanceUploadInfo.InstanceSceneDataOffset		= InstanceIDStartOffset + ItemIndex; 
 			InstanceUploadInfo.PayloadDataFlags				= 0;
 			InstanceUploadInfo.InstancePayloadDataOffset	= INDEX_NONE;
diff --git a/Engine/Source/Runtime/Renderer/Private/RendererScene.cpp b/Engine/Source/Runtime/Renderer/Private/RendererScene.cpp
--- a/Engine/Source/Runtime/Renderer/Private/RendererScene.cpp
+++ b/Engine/Source/Runtime/Renderer/Private/RendererScene.cpp
@@ -4731,7 +4731,10 @@ void FScene::UpdateAllPrimitiveSceneInfos(FRDGBuilder& GraphBuilder, bool bAsync
 		SCOPE_CYCLE_COUNTER(STAT_UpdatePrimitiveInstanceRenderThreadTime);
 
 		TArray<FPrimitiveSceneInfo*> UpdatedSceneInfosWithStaticDrawListUpdate;
-		UpdatedSceneInfosWithStaticDrawListUpdate.Reserve(UpdatedTransforms.Num());
+		UpdatedSceneInfosWithStaticDrawListUpdate.Reserve(UpdatedInstances.Num());
+
+		TArray<FPrimitiveSceneInfo*> UpdatedSceneInfosWithoutStaticDrawListUpdate;
+		UpdatedSceneInfosWithoutStaticDrawListUpdate.Reserve(UpdatedInstances.Num());
 
 		for (const auto& UpdateInstance : UpdatedInstances)
 		{
@@ -4745,15 +4748,29 @@ void FScene::UpdateAllPrimitiveSceneInfos(FRDGBuilder& GraphBuilder, bool bAsync
 
 			FScopeCycleCounter Context(PrimitiveSceneProxy->GetStatId());
 			FPrimitiveSceneInfo* PrimitiveSceneInfo = PrimitiveSceneProxy->GetPrimitiveSceneInfo();
+			
+			PrimitiveSceneInfo->FlushRuntimeVirtualTexture();
 
-			UpdatedSceneInfosWithStaticDrawListUpdate.Push(PrimitiveSceneInfo);
+			const bool bUpdateStaticDrawLists = !PrimitiveSceneProxy->StaticElementsAlwaysUseProxyPrimitiveUniformBuffer();
 
-			// Update the Proxy's data.
-			PrimitiveSceneProxy->UpdateInstances_RenderThread(UpdateInstance.Value.CmdBuffer, UpdateInstance.Value.WorldBounds, UpdateInstance.Value.LocalBounds, UpdateInstance.Value.StaticMeshBounds);
+			// If we recorded no adds or removes the instance count has stayed the same.  Therefore the cached mesh draw commands do not
+			// need to be updated.  In situations where the instance count changes the mesh draw command stores the instance count which
+			// would need to be updated.
+			const bool bInstanceCountChanged = (UpdateInstance.Value.CmdBuffer.NumAdds > 0) || (UpdateInstance.Value.CmdBuffer.NumRemoves > 0);
 
-			PrimitiveSceneInfo->FlushRuntimeVirtualTexture();
+			if (bUpdateStaticDrawLists || bInstanceCountChanged)
+			{
+				UpdatedSceneInfosWithStaticDrawListUpdate.Push(PrimitiveSceneInfo);
+				PrimitiveSceneInfo->RemoveFromScene(true);
+			}
+			else
+			{
+				UpdatedSceneInfosWithoutStaticDrawListUpdate.Push(PrimitiveSceneInfo);
+				PrimitiveSceneInfo->RemoveFromScene(false);
+			}
 
-			PrimitiveSceneInfo->RemoveFromScene(true);
+			// Update the Proxy's data.
+			PrimitiveSceneProxy->UpdateInstances_RenderThread(UpdateInstance.Value.CmdBuffer, UpdateInstance.Value.WorldBounds, UpdateInstance.Value.LocalBounds, UpdateInstance.Value.StaticMeshBounds);
 
 			if (!RHISupportsVolumeTextures(GetFeatureLevel())
 				&& (PrimitiveSceneProxy->IsMovable() || PrimitiveSceneProxy->NeedsUnbuiltPreviewLighting() || PrimitiveSceneProxy->GetLightmapType() == ELightmapType::ForceVolumetric))
@@ -4794,6 +4811,11 @@ void FScene::UpdateAllPrimitiveSceneInfos(FRDGBuilder& GraphBuilder, bool bAsync
 		{
 			FPrimitiveSceneInfo::AddToScene(GraphBuilder.RHICmdList, this, UpdatedSceneInfosWithStaticDrawListUpdate, true, true, bAsyncCreateLPIs);
 		}
+
+		if (UpdatedSceneInfosWithoutStaticDrawListUpdate.Num() > 0)
+		{
+			FPrimitiveSceneInfo::AddToScene(GraphBuilder.RHICmdList, this, UpdatedSceneInfosWithoutStaticDrawListUpdate, false, true, bAsyncCreateLPIs);
+		}
 	}
 
 	for (const auto& Attachments : UpdatedAttachmentRoots)
