Fix compile issue introduced in CL-18006634
Change InstanceTransform to use FVector4f

diff --git a/Engine/Source/Runtime/Engine/Public/StaticMeshResources.h b/Engine/Source/Runtime/Engine/Public/StaticMeshResources.h
--- a/Engine/Source/Runtime/Engine/Public/StaticMeshResources.h
+++ b/Engine/Source/Runtime/Engine/Public/StaticMeshResources.h
@@ -1201,7 +1201,7 @@ public:
 	FORCEINLINE_DEBUGGABLE void GetInstanceEditorData(int32 InstanceIndex, FColor& HitProxyColorOut, bool & bSelectedOut) const
 	{
 		// TODO: put this into a sensible format
-		FVector4 InstanceTransform[3];
+		FVector4f InstanceTransform[3];
 		if (bUseHalfFloat)
 		{
 			GetInstanceTransformInternal<FFloat16>(InstanceIndex, InstanceTransform);
