Fix incorrect merge of ISM VF (re-disable checks that should not happen when GPU-Scene is used)

diff --git a/Engine/Source/Runtime/Engine/Private/InstancedStaticMesh.cpp b/Engine/Source/Runtime/Engine/Private/InstancedStaticMesh.cpp
--- a/Engine/Source/Runtime/Engine/Private/InstancedStaticMesh.cpp
+++ b/Engine/Source/Runtime/Engine/Private/InstancedStaticMesh.cpp
@@ -927,13 +927,13 @@ void FInstancedStaticMeshVertexFactory::InitRHI()
 	if (GetFeatureLevel() > ERHIFeatureLevel::ES3_1 || !bMobileUsesGPUScene)
 	{
 		// toss in the instanced location stream
-		check(Data.InstanceOriginComponent.VertexBuffer);
+		check(bCanUseGPUScene || Data.InstanceOriginComponent.VertexBuffer);
 		if (Data.InstanceOriginComponent.VertexBuffer)
 		{
 			Elements.Add(AccessStreamComponent(Data.InstanceOriginComponent, 8));
 		}
 
-		check(Data.InstanceTransformComponent[0].VertexBuffer);
+		check(bCanUseGPUScene || Data.InstanceTransformComponent[0].VertexBuffer);
 		if (Data.InstanceTransformComponent[0].VertexBuffer)
 		{
 			Elements.Add(AccessStreamComponent(Data.InstanceTransformComponent[0], 9));
