Fix accidental LWC coordinate double precision vector cast & unintentional OOB read.

diff --git a/Engine/Source/Runtime/Engine/Private/InstanceUniformShaderParameters.cpp b/Engine/Source/Runtime/Engine/Private/InstanceUniformShaderParameters.cpp
--- a/Engine/Source/Runtime/Engine/Private/InstanceUniformShaderParameters.cpp
+++ b/Engine/Source/Runtime/Engine/Private/InstanceUniformShaderParameters.cpp
@@ -87,11 +87,11 @@ void FInstanceSceneShaderData::BuildInternal
 	// TODO: Temporary PrevVelocityHack
 #if INSTANCE_COMPRESSED_TRANSFORMS
 	FCompressedTransform CompressedLocalToWorld(LocalToWorld);
-	Data[1] = *(const FVector4*)&CompressedLocalToWorld.Rotation[0];
+	Data[1] = *(const FVector4f*)&CompressedLocalToWorld.Rotation[0];
 	Data[2] = *(const FVector3f*)&CompressedLocalToWorld.Translation;
 
 	FCompressedTransform CompressedPrevLocalToWorld(PrevLocalToWorld);
-	Data[3] = *(const FVector4*)&CompressedPrevLocalToWorld.Rotation[0];
+	Data[3] = *(const FVector4f*)&CompressedPrevLocalToWorld.Rotation[0];
 	Data[4] = *(const FVector3f*)&CompressedPrevLocalToWorld.Translation;
 #else
 	// Note: writes 3x float4s
