Fix payload data & flag setup for landscape grass use-case (stores no data in the component)

diff --git a/Engine/Source/Runtime/Engine/Private/InstancedStaticMesh.cpp b/Engine/Source/Runtime/Engine/Private/InstancedStaticMesh.cpp
--- a/Engine/Source/Runtime/Engine/Private/InstancedStaticMesh.cpp
+++ b/Engine/Source/Runtime/Engine/Private/InstancedStaticMesh.cpp
@@ -1504,6 +1504,19 @@ void FInstancedStaticMeshSceneProxy::CreateRenderThreadResources()
 			if (InstanceSceneData.Num() == 0)
 			{
 				InstanceSceneData.SetNum(InstanceBuffer.GetNumInstances());
+				InstanceLightShadowUVBias.SetNumZeroed(bHasPerInstanceLMSMUVBias ? InstanceBuffer.GetNumInstances() : 0);
+
+				// Note: since bHasPerInstanceDynamicData is set based on the number (bHasPerInstanceDynamicData = InComponent->PerInstancePrevTransform.Num() == InComponent->GetInstanceCount();)
+				//       in the ctor, it gets set to true when both are zero as well, for the landscape grass path (or whatever triggers initialization from the InstanceBuffer only)
+				//       there is no component data to get, thus we set the bHasPerInstanceDynamicData to false.
+				bHasPerInstanceDynamicData = false;
+				InstanceDynamicData.Empty();
+
+				InstanceRandomID.SetNumZeroed(bHasPerInstanceRandom ? InstanceBuffer.GetNumInstances() : 0); // Only allocate if material bound which uses this
+
+#if WITH_EDITOR
+				InstanceEditorData.SetNumZeroed(bHasPerInstanceEditorData ? InstanceBuffer.GetNumInstances() : 0);
+#endif
 			}
 
 			// NOTE: we set up partial data in the construction of ISM proxy (yep, awful but the equally awful way the InstanceBuffer is maintained means complete data is not available)
