Fixes crash on incorrect r.AntiAliasingMethod values

diff --git a/Engine/Source/Runtime/Engine/Private/SceneUtils.cpp b/Engine/Source/Runtime/Engine/Private/SceneUtils.cpp
--- a/Engine/Source/Runtime/Engine/Private/SceneUtils.cpp
+++ b/Engine/Source/Runtime/Engine/Private/SceneUtils.cpp
@@ -45,7 +45,7 @@ ENGINE_API EAntiAliasingMethod GetDefaultAntiAliasingMethod(const FStaticFeature
 	if (InFeatureLevel == ERHIFeatureLevel::ES3_1)
 	{
 		static auto* MobileAntiAliasingCvar = IConsoleManager::Get().FindTConsoleVariableDataInt(TEXT("r.Mobile.AntiAliasing"));
-		AntiAliasingMethod = (EAntiAliasingMethod)MobileAntiAliasingCvar->GetValueOnAnyThread();
+		AntiAliasingMethod = EAntiAliasingMethod(FMath::Clamp<int32>(MobileAntiAliasingCvar->GetValueOnAnyThread(), 0, AAM_MAX));
 
 		static auto* MobileHDRCvar = IConsoleManager::Get().FindTConsoleVariableDataInt(TEXT("r.MobileHDR"));
 		// Disable antialiasing in GammaLDR mode to avoid jittering.
@@ -60,6 +60,12 @@ ENGINE_API EAntiAliasingMethod GetDefaultAntiAliasingMethod(const FStaticFeature
 		AntiAliasingMethod = EAntiAliasingMethod(FMath::Clamp<int32>(DefaultAntiAliasingCvar->GetValueOnAnyThread(), 0, AAM_MAX));
 	}
 
+	if (AntiAliasingMethod == AAM_MAX)
+	{
+		ensureMsgf(false, TEXT("Unknown anti-aliasing method."));
+		AntiAliasingMethod = AAM_None;
+	}
+
 	static auto* MSAACountCVar = IConsoleManager::Get().FindTConsoleVariableDataInt(TEXT("r.MSAACount"));
 
 	EShaderPlatform ShaderPlatform = GetFeatureLevelShaderPlatform(InFeatureLevel);
