Fix for virtual textures not applying RDO when compressed with Oodle.

Since the properties weren't copied, this means that the DDC key doesn't represent the actual bits, so the VT ddc key is invalidated.

diff --git a/Engine/Source/Runtime/Engine/Private/TextureDerivedData.cpp b/Engine/Source/Runtime/Engine/Private/TextureDerivedData.cpp
--- a/Engine/Source/Runtime/Engine/Private/TextureDerivedData.cpp
+++ b/Engine/Source/Runtime/Engine/Private/TextureDerivedData.cpp
@@ -65,7 +65,7 @@
 
 // This GUID is mixed into DDC version for virtual textures only, this allows updating DDC version for VT without invalidating DDC for all textures
 // This is useful during development, but once large numbers of VT are present in shipped content, it will have the same problem as TEXTURE_DERIVEDDATA_VER
-#define TEXTURE_VT_DERIVEDDATA_VER	TEXT("4FA538D2450243E380565F65CE8E44B7")
+#define TEXTURE_VT_DERIVEDDATA_VER	TEXT("F36E4B807BC24Y818FE085C655C45176")
 
 #if ENABLE_COOK_STATS
 namespace TextureCookStats
