Fix deprecated usages of USkeletalMesh:: RayTracingMinLOD and bSupportRayTracing in HandleSkeletalMeshReportCommand

diff --git a/Engine/Source/Runtime/Engine/Private/UnrealEngine.cpp b/Engine/Source/Runtime/Engine/Private/UnrealEngine.cpp
--- a/Engine/Source/Runtime/Engine/Private/UnrealEngine.cpp
+++ b/Engine/Source/Runtime/Engine/Private/UnrealEngine.cpp
@@ -7111,8 +7111,8 @@ bool UEngine::HandleSkeletalMeshReportCommand(const TCHAR* Cmd, FOutputDevice& A
 					ClothVertexBufferKB,
 					bSkinCachedLODEnabled ? TEXT("Y") : TEXT("N"),
 					*RecomputeTangentSections,
-					Mesh->bSupportRayTracing ? TEXT("Y") : TEXT("N"),
-					Mesh->RayTracingMinLOD);
+					Mesh->GetSupportRayTracing() ? TEXT("Y") : TEXT("N"),
+					Mesh->GetRayTracingMinLOD());
 			}
 		}
 
