Fix bad merge.

diff --git a/Engine/Source/Runtime/Engine/Private/WorldPartition/WorldPartitionRuntimeSpatialHash.cpp b/Engine/Source/Runtime/Engine/Private/WorldPartition/WorldPartitionRuntimeSpatialHash.cpp
--- a/Engine/Source/Runtime/Engine/Private/WorldPartition/WorldPartitionRuntimeSpatialHash.cpp
+++ b/Engine/Source/Runtime/Engine/Private/WorldPartition/WorldPartitionRuntimeSpatialHash.cpp
@@ -667,7 +667,7 @@ void UWorldPartitionRuntimeSpatialHash::SetDefaultValues()
 	MainGrid.DebugColor = FLinearColor::Gray;
 }
 
-bool UWorldPartitionRuntimeSpatialHash::GenerateStreaming(EWorldPartitionStreamingMode Mode, UWorldPartitionStreamingPolicy* StreamingPolicy, TArray<FString>* OutPackagesToGenerate)
+bool UWorldPartitionRuntimeSpatialHash::GenerateStreaming(class UWorldPartitionStreamingPolicy* StreamingPolicy, const FActorClusterContext& ActorClusterContext, TArray<FString>* OutPackagesToGenerate = nullptr)
 {
 	TRACE_CPUPROFILER_EVENT_SCOPE(UWorldPartitionRuntimeSpatialHash::GenerateStreaming);
 	UWorldPartition* WorldPartition = GetOuterUWorldPartition();
