fix deprecation warnings

diff --git a/Engine/Source/Runtime/Engine/Private/Components/SkinnedMeshComponent.cpp b/Engine/Source/Runtime/Engine/Private/Components/SkinnedMeshComponent.cpp
--- a/Engine/Source/Runtime/Engine/Private/Components/SkinnedMeshComponent.cpp
+++ b/Engine/Source/Runtime/Engine/Private/Components/SkinnedMeshComponent.cpp
@@ -1484,7 +1484,9 @@ bool USkinnedMeshComponent::IsSkinCacheAllowed(int32 LodIdx) const
 	bool bIsRayTracing = FGPUSkinCache::IsGPUSkinCacheRayTracingSupported();
 	if (bIsRayTracing && SkeletalMesh)
 	{
+PRAGMA_DISABLE_DEPRECATION_WARNINGS
 		bIsRayTracing = SkeletalMesh->bSupportRayTracing;
+PRAGMA_ENABLE_DEPRECATION_WARNINGS
 	}
 
 	bool bGlobalDefault = CVarDefaultGPUSkinCacheBehavior && ESkinCacheDefaultBehavior(CVarDefaultGPUSkinCacheBehavior->GetInt()) == ESkinCacheDefaultBehavior::Inclusive;
