fix deprecation warnings

diff --git a/Engine/Source/Runtime/Engine/Private/SkeletalMeshRenderData.cpp b/Engine/Source/Runtime/Engine/Private/SkeletalMeshRenderData.cpp
--- a/Engine/Source/Runtime/Engine/Private/SkeletalMeshRenderData.cpp
+++ b/Engine/Source/Runtime/Engine/Private/SkeletalMeshRenderData.cpp
@@ -710,7 +710,9 @@ void FSkeletalMeshRenderData::Serialize(FArchive& Ar, USkeletalMesh* Owner)
 	CurrentFirstLODIdx = LODRenderData.Num() - NumInlinedLODs;
 	PendingFirstLODIdx = CurrentFirstLODIdx;
 
+PRAGMA_DISABLE_DEPRECATION_WARNINGS
 	bSupportRayTracing = Owner->bSupportRayTracing;
+PRAGMA_ENABLE_DEPRECATION_WARNINGS
 }
 
 void FSkeletalMeshRenderData::InitResources(bool bNeedsVertexColors, TArray<UMorphTarget*>& InMorphTargets, USkeletalMesh* Owner)
