Fix various C6011: Dereferencing NULL pointer 'OwningWorld'

diff --git a/Engine/Source/Runtime/Engine/Private/ActorConstruction.cpp b/Engine/Source/Runtime/Engine/Private/ActorConstruction.cpp
--- a/Engine/Source/Runtime/Engine/Private/ActorConstruction.cpp
+++ b/Engine/Source/Runtime/Engine/Private/ActorConstruction.cpp
@@ -653,7 +653,7 @@ void AActor::RerunConstructionScripts()
 				}
 			}
 #if WITH_EDITORONLY_DATA
-			if (OwningWorld->GetWorldPartition() != nullptr)
+			if (OwningWorld && OwningWorld->GetWorldPartition() != nullptr)
 			{
 				if (USceneComponent* OldSceneComponent = Cast<USceneComponent>(ComponentData.OldComponent))
 				{
diff --git a/Engine/Source/Runtime/Engine/Private/SkeletalRenderGPUSkin.cpp b/Engine/Source/Runtime/Engine/Private/SkeletalRenderGPUSkin.cpp
--- a/Engine/Source/Runtime/Engine/Private/SkeletalRenderGPUSkin.cpp
+++ b/Engine/Source/Runtime/Engine/Private/SkeletalRenderGPUSkin.cpp
@@ -2123,7 +2123,7 @@ void FDynamicSkelMeshObjectDataGPUSkin::InitDynamicSkelMeshObjectDataGPUSkin(
 	// Update the clothing simulation mesh positions and normals
 	UpdateClothSimulationData(InMeshComponent);
 
-	bIsSkinCacheAllowed = InMeshComponent->IsSkinCacheAllowed(InLODIndex);
+	bIsSkinCacheAllowed = InMeshComponent ? InMeshComponent->IsSkinCacheAllowed(InLODIndex) : false;
 
 	if (bIsSkinCacheAllowed && InMeshObject->FeatureLevel == ERHIFeatureLevel::ES3_1)
 	{
diff --git a/Engine/Source/Runtime/Engine/Private/TextureRenderTarget2D.cpp b/Engine/Source/Runtime/Engine/Private/TextureRenderTarget2D.cpp
--- a/Engine/Source/Runtime/Engine/Private/TextureRenderTarget2D.cpp
+++ b/Engine/Source/Runtime/Engine/Private/TextureRenderTarget2D.cpp
@@ -552,7 +552,7 @@ void FTextureRenderTarget2DResource::InitDynamicRHI()
 		// Create the RHI texture. Only one mip is used and the texture is targetable for resolve.
 		ETextureCreateFlags TexCreateFlags = Owner->IsSRGB() ? TexCreate_SRGB : TexCreate_None;
 		TexCreateFlags |= Owner->bGPUSharedFlag ? TexCreate_Shared : TexCreate_None;
-		FString ResourceName = Owner ? Owner->GetName() : FString(TEXT("TextureRenderTarget2DResource"));
+		FString ResourceName = Owner->GetName();
 		FRHIResourceCreateInfo CreateInfo = FRHIResourceCreateInfo(*ResourceName, FClearValueBinding(ClearColor));
 
 		if (Owner->bAutoGenerateMips)
