Linked anim instances that have a mix of interface and non-interface layers now only unlink interface layers when unlinking

diff --git a/Engine/Source/Runtime/Engine/Private/Animation/AnimInstance.cpp b/Engine/Source/Runtime/Engine/Private/Animation/AnimInstance.cpp
--- a/Engine/Source/Runtime/Engine/Private/Animation/AnimInstance.cpp
+++ b/Engine/Source/Runtime/Engine/Private/Animation/AnimInstance.cpp
@@ -2848,13 +2848,17 @@ void UAnimInstance::PerformLinkedLayerOverlayOperation(TSubclassOf<UAnimInstance
 
 		USkeletalMeshComponent* MeshComp = GetSkelMeshComponent();
 
-		auto UnlinkLayerNodesInInstance = [](UAnimInstance* InAnimInstance)
+		auto UnlinkLayerNodesInInstance = [](UAnimInstance* InAnimInstance, TArrayView<FAnimNode_LinkedAnimLayer*> InLayerNodes)
 		{
-			const IAnimClassInterface* NewLinkedInstanceClass = IAnimClassInterface::GetFromClass(InAnimInstance->GetClass());
-			for(const FStructProperty* LayerNodeProperty : NewLinkedInstanceClass->GetLinkedAnimLayerNodeProperties())
+			const IAnimClassInterface* const NewLinkedInstanceClass = IAnimClassInterface::GetFromClass(InAnimInstance->GetClass());
+			for (const FStructProperty* const LayerNodeProperty : NewLinkedInstanceClass->GetLinkedAnimLayerNodeProperties())
 			{
-				FAnimNode_LinkedAnimLayer* LinkedAnimLayerNode = LayerNodeProperty->ContainerPtrToValuePtr<FAnimNode_LinkedAnimLayer>(InAnimInstance);
-				LinkedAnimLayerNode->DynamicUnlink(InAnimInstance);
+				FAnimNode_LinkedAnimLayer* const LinkedAnimLayerNode = LayerNodeProperty->ContainerPtrToValuePtr<FAnimNode_LinkedAnimLayer>(InAnimInstance);
+				const bool bExternalLink = InLayerNodes.ContainsByPredicate([LinkedAnimLayerNode](const FAnimNode_LinkedAnimLayer* Layer) { return Layer->Layer == LinkedAnimLayerNode->Layer; });
+				if (bExternalLink)
+				{
+					LinkedAnimLayerNode->DynamicUnlink(InAnimInstance);
+				}
 			}
 		};
 
@@ -2893,7 +2897,7 @@ void UAnimInstance::PerformLinkedLayerOverlayOperation(TSubclassOf<UAnimInstance
 								NewLinkedInstance->InitializeAnimation();
 
 								// Unlink any layer nodes in the new linked instance, as they may have been hooked up to self in InitializeAnimation above.
-								UnlinkLayerNodesInInstance(NewLinkedInstance);
+								UnlinkLayerNodesInInstance(NewLinkedInstance, LayerPair.Value);
 
 								LayerNode->SetLinkedLayerInstance(this, NewLinkedInstance);
 
@@ -2949,7 +2953,7 @@ void UAnimInstance::PerformLinkedLayerOverlayOperation(TSubclassOf<UAnimInstance
 							NewLinkedInstance->InitializeAnimation();
 
 							// Unlink any layer nodes in the new linked instance, as they may have been hooked up to self in InitializeAnimation above.
-							UnlinkLayerNodesInInstance(NewLinkedInstance);
+							UnlinkLayerNodesInInstance(NewLinkedInstance, LayerPair.Value);
 
 							for(FAnimNode_LinkedAnimLayer* LayerNode : LayerPair.Value)
 							{
