Simplify how we get FFolder::RootObject from LevelInstance (fixes a failed check when reusing levels in pie)

diff --git a/Engine/Source/Editor/UnrealEd/Private/EditorLevelUtils.cpp b/Engine/Source/Editor/UnrealEd/Private/EditorLevelUtils.cpp
--- a/Engine/Source/Editor/UnrealEd/Private/EditorLevelUtils.cpp
+++ b/Engine/Source/Editor/UnrealEd/Private/EditorLevelUtils.cpp
@@ -867,7 +867,6 @@ void UEditorLevelUtils::PrivateRemoveLevelFromWorld(ULevel* InLevel)
 		InLevel->OwningWorld->RemoveFromWorld(InLevel);
 		check(InLevel->bIsVisible == false);
 	}
-	InLevel->bIsEditorBeingRemoved = true;
 
 	InLevel->ReleaseRenderingResources();
 
diff --git a/Engine/Source/Runtime/Engine/Classes/Engine/Level.h b/Engine/Source/Runtime/Engine/Classes/Engine/Level.h
--- a/Engine/Source/Runtime/Engine/Classes/Engine/Level.h
+++ b/Engine/Source/Runtime/Engine/Classes/Engine/Level.h
@@ -617,9 +617,6 @@ public:
 	UPROPERTY()
 	uint8 										bLocked:1;
 
-	/** Whether the level is currently being removed from the world through UEditorLevelUtils::RemoveLevelsFromWorld */
-	uint8										bIsEditorBeingRemoved : 1;
-
 	/** Whether the level has been saved after introducing actor GUIDs */
 	uint8										bContainsStableActorGUIDs:1;
 #endif
@@ -1182,11 +1179,6 @@ public:
 	 */
 	ENGINE_API bool IsInstancedLevel() const;
 
-	/**
-	 * Is level being removed from world
-	 */
-	ENGINE_API bool IsBeingRemoved() const;
-
 	/** 
 	 * Shift level actors by specified offset
 	 * The offset vector will get subtracted from all actors positions and corresponding data structures
diff --git a/Engine/Source/Runtime/Engine/Classes/Engine/LevelStreaming.h b/Engine/Source/Runtime/Engine/Classes/Engine/LevelStreaming.h
--- a/Engine/Source/Runtime/Engine/Classes/Engine/LevelStreaming.h
+++ b/Engine/Source/Runtime/Engine/Classes/Engine/LevelStreaming.h
@@ -11,6 +11,11 @@
 #include "LatentActions.h"
 #include "ProfilingDebugging/ProfilingHelpers.h"
 
+#if WITH_EDITOR
+#include "Folder.h"
+#include "Misc/Optional.h"
+#endif
+
 class ALevelScriptActor;
 class ALevelStreamingVolume;
 class ULevel;
@@ -471,6 +476,8 @@ public:
 
 	/** Sets the folder path for this level in the world browser. Only available in editor builds */
 	void SetFolderPath(const FName& InFolderPath);
+
+	virtual TOptional<FFolder::FRootObject> GetFolderRootObject() const { return FFolder::GetDefaultRootObject(); }
 #endif	// WITH_EDITOR
 
 	//~==============================================================================================
diff --git a/Engine/Source/Runtime/Engine/Private/ActorEditor.cpp b/Engine/Source/Runtime/Engine/Private/ActorEditor.cpp
--- a/Engine/Source/Runtime/Engine/Private/ActorEditor.cpp
+++ b/Engine/Source/Runtime/Engine/Private/ActorEditor.cpp
@@ -29,9 +29,8 @@
 #include "Editor.h"
 #include "EditorActorFolders.h"
 #include "Folder.h"
+#include "Engine/LevelStreaming.h"
 #include "WorldPartition/WorldPartitionActorDesc.h"
-#include "LevelInstance/LevelInstanceSubsystem.h"
-#include "LevelInstance/LevelInstanceActor.h"
 
 #define LOCTEXT_NAMESPACE "ErrorChecking"
 
@@ -953,31 +952,28 @@ FFolder AActor::GetFolder() const
 	return FFolder(GetFolderPath(), GetFolderRootObject());
 }
 
-static bool IsActorFolderValid(const AActor* InActor)
+static TOptional<FFolder::FRootObject> GetOptionalFolderRootObject(const AActor* Actor)
 {
-	const ULevel* Level = InActor->GetLevel();
-	const bool bIsBeingRemoved = Level && Level->IsBeingRemoved();
-	return !bIsBeingRemoved;
-}
-
-FFolder::FRootObject AActor::GetFolderRootObject() const
-{
-	if (IsActorFolderValid(this))
+	if (ULevel* Level = Actor->GetLevel())
 	{
-		UWorld* World = GetWorld();
-		if (const ULevelInstanceSubsystem* LevelInstanceSubsystem = World ? World->GetSubsystem<ULevelInstanceSubsystem>() : nullptr)
+		if (ULevelStreaming* LevelStreaming = ULevelStreaming::FindStreamingLevel(Level))
 		{
-			ALevelInstance* LevelInstance = LevelInstanceSubsystem->GetParentLevelInstance(this);
-			return FFolder::FRootObject(LevelInstance);
+			return LevelStreaming->GetFolderRootObject();
 		}
 	}
+
 	return FFolder::GetDefaultRootObject();
 }
 
+FFolder::FRootObject AActor::GetFolderRootObject() const
+{
+	return GetOptionalFolderRootObject(this).Get(FFolder::GetDefaultRootObject());
+}
+
 const FName& AActor::GetFolderPath() const
 {
 	static const FName EmptyPath = NAME_None;
-	return IsActorFolderValid(this) ? FolderPath : EmptyPath;
+	return GetOptionalFolderRootObject(this) ? FolderPath : EmptyPath;
 }
 
 void AActor::SetFolderPath(const FName& NewFolderPath)
diff --git a/Engine/Source/Runtime/Engine/Private/Level.cpp b/Engine/Source/Runtime/Engine/Private/Level.cpp
--- a/Engine/Source/Runtime/Engine/Private/Level.cpp
+++ b/Engine/Source/Runtime/Engine/Private/Level.cpp
@@ -2888,15 +2888,6 @@ bool ULevel::IsCurrentLevel() const
 	return bIsCurrent;
 }
 
-bool ULevel::IsBeingRemoved() const
-{
-	bool bResult = bIsBeingRemoved;
-#if WITH_EDITOR
-	bResult |= bIsEditorBeingRemoved;
-#endif
-	return bResult;
-}
-
 bool ULevel::IsInstancedLevel() const
 {
 	const UWorld* OuterWorld = GetTypedOuter<UWorld>();
diff --git a/Engine/Source/Runtime/Engine/Private/LevelInstance/LevelInstanceEditorLevelStreaming.cpp b/Engine/Source/Runtime/Engine/Private/LevelInstance/LevelInstanceEditorLevelStreaming.cpp
--- a/Engine/Source/Runtime/Engine/Private/LevelInstance/LevelInstanceEditorLevelStreaming.cpp
+++ b/Engine/Source/Runtime/Engine/Private/LevelInstance/LevelInstanceEditorLevelStreaming.cpp
@@ -33,6 +33,16 @@ ULevelStreamingLevelInstanceEditor::ULevelStreamingLevelInstanceEditor(const FOb
 }
 
 #if WITH_EDITOR
+TOptional<FFolder::FRootObject> ULevelStreamingLevelInstanceEditor::GetFolderRootObject() const
+{
+	if (ALevelInstance* LevelInstance = GetLevelInstanceActor())
+	{
+		return FFolder::FRootObject(LevelInstance);
+	}
+
+	return TOptional<FFolder::FRootObject>();
+}
+
 ALevelInstance* ULevelStreamingLevelInstanceEditor::GetLevelInstanceActor() const
 {
 	if (ULevelInstanceSubsystem* LevelInstanceSubsystem = GetWorld()->GetSubsystem<ULevelInstanceSubsystem>())
diff --git a/Engine/Source/Runtime/Engine/Private/LevelInstance/LevelInstanceLevelStreaming.cpp b/Engine/Source/Runtime/Engine/Private/LevelInstance/LevelInstanceLevelStreaming.cpp
--- a/Engine/Source/Runtime/Engine/Private/LevelInstance/LevelInstanceLevelStreaming.cpp
+++ b/Engine/Source/Runtime/Engine/Private/LevelInstance/LevelInstanceLevelStreaming.cpp
@@ -35,6 +35,16 @@ ALevelInstance* ULevelStreamingLevelInstance::GetLevelInstanceActor() const
 }
 
 #if WITH_EDITOR
+TOptional<FFolder::FRootObject> ULevelStreamingLevelInstance::GetFolderRootObject() const
+{
+	if (ALevelInstance* LevelInstance = GetLevelInstanceActor())
+	{
+		return FFolder::FRootObject(LevelInstance);
+	}
+
+	return TOptional<FFolder::FRootObject>();
+}
+
 FBox ULevelStreamingLevelInstance::GetBounds() const
 {
 	check(GetLoadedLevel());
diff --git a/Engine/Source/Runtime/Engine/Private/LevelInstance/LevelInstanceSubsystem.cpp b/Engine/Source/Runtime/Engine/Private/LevelInstance/LevelInstanceSubsystem.cpp
--- a/Engine/Source/Runtime/Engine/Private/LevelInstance/LevelInstanceSubsystem.cpp
+++ b/Engine/Source/Runtime/Engine/Private/LevelInstance/LevelInstanceSubsystem.cpp
@@ -131,33 +131,12 @@ FLevelInstanceID ULevelInstanceSubsystem::RegisterLevelInstance(ALevelInstance*
 	check(GIsReinstancing || Value == nullptr || Value == LevelInstanceActor);
 	Value = LevelInstanceActor;
 
-#if WITH_EDITOR
-	FObjectKey Level;
-	if (UnregisteringLevelInstances.RemoveAndCopyValue(LevelInstanceActor, Level))
-	{
-		if (ULevel* LevelPtr = Cast<ULevel>(Level.ResolveObjectPtr()))
-		{
-			LevelPtr->bIsEditorBeingRemoved = false;
-		}
-		ensure(UnregisteringLevelInstanceLevels.Remove(Level));
-	}
-#endif
-
 	return LevelInstanceID;
 }
 
 void ULevelInstanceSubsystem::UnregisterLevelInstance(ALevelInstance* LevelInstanceActor)
 {
 	RegisteredLevelInstances.Remove(LevelInstanceActor->GetLevelInstanceID());
-
-#if WITH_EDITOR
-	if (ULevel* Level = GetLevelInstanceLevel(LevelInstanceActor))
-	{
-		Level->bIsEditorBeingRemoved = true;
-		UnregisteringLevelInstanceLevels.Add(Level, LevelInstanceActor);
-		UnregisteringLevelInstances.Add(LevelInstanceActor, Level);
-	}
-#endif
 }
 
 void ULevelInstanceSubsystem::RequestLoadLevelInstance(ALevelInstance* LevelInstanceActor, bool bForce /* = false */)
@@ -169,13 +148,6 @@ void ULevelInstanceSubsystem::RequestLoadLevelInstance(ALevelInstance* LevelInst
 		if (!IsEditingLevelInstance(LevelInstanceActor))
 #endif
 		{
-#if WITH_EDITOR
-			if (ULevel* Level = GetLevelInstanceLevel(LevelInstanceActor))
-			{
-				Level->bIsEditorBeingRemoved = false;
-			}
-#endif
-			
 			LevelInstancesToUnload.Remove(LevelInstanceActor->GetLevelInstanceID());
 
 			bool* bForcePtr = LevelInstancesToLoadOrUpdate.Find(LevelInstanceActor);
@@ -200,13 +172,6 @@ void ULevelInstanceSubsystem::RequestLoadLevelInstance(ALevelInstance* LevelInst
 
 void ULevelInstanceSubsystem::RequestUnloadLevelInstance(ALevelInstance* LevelInstanceActor)
 {
-#if WITH_EDITOR
-	if (ULevel* Level = GetLevelInstanceLevel(LevelInstanceActor))
-	{
-		Level->bIsEditorBeingRemoved = true;
-	}
-#endif
-
 	const FLevelInstanceID& LevelInstanceID = LevelInstanceActor->GetLevelInstanceID();
 	if (LevelInstances.Contains(LevelInstanceID))
 	{
@@ -330,10 +295,6 @@ void ULevelInstanceSubsystem::UnloadLevelInstance(const FLevelInstanceID& LevelI
 	{
 		if (ULevel* LoadedLevel = LevelInstance.LevelStreaming->GetLoadedLevel())
 		{
-#if WITH_EDITOR
-			LoadedLevel->bIsEditorBeingRemoved = true;
-#endif
-
 			ForEachActorInLevel(LoadedLevel, [this](AActor* LevelActor)
 			{
 				if (ALevelInstance* LevelInstanceActor = Cast<ALevelInstance>(LevelActor))
@@ -1213,44 +1174,6 @@ void ULevelInstanceSubsystem::RemoveLevelsFromWorld(const TArray<ULevel*>& InLev
 	}
 	else
 	{
-		TSet<ULevel*> LevelInstanceLevels;
-		for (ULevel* Level : InLevels)
-		{
-			bool bIsAlreadyInSet = false;
-			LevelInstanceLevels.Add(Level, &bIsAlreadyInSet);
-
-			if (!bIsAlreadyInSet)
-			{
-				ForEachActorInLevel(Level, [this, &LevelInstanceLevels](AActor* LevelActor)
-				{
-					if (const ALevelInstance* ChildLevelInstanceActor = Cast<ALevelInstance>(LevelActor))
-					{
-						ForEachLevelInstanceChild(ChildLevelInstanceActor, /*bRecursive*/true, [this, &LevelInstanceLevels](const ALevelInstance* ChildLevelInstanceActor)
-						{
-							if (ULevel* ChildLevel = GetLevelInstanceLevel(ChildLevelInstanceActor))
-							{
-								LevelInstanceLevels.Add(ChildLevel);
-							}
-							return true;
-						});
-					}
-					return true;
-				});
-			}
-		}
-
-		// Flag all levels as bEditorBeingDestroyed. This way, even if child level instances are still pending to be unloaded, TActorIterator won't iterate on them.
-		for (ULevel* Level : LevelInstanceLevels)
-		{
-			Level->bIsEditorBeingRemoved = true;
-			FFolder::FRootObject RootObject;
-			if (UnregisteringLevelInstanceLevels.RemoveAndCopyValue(Level, RootObject))
-			{
-				ensure(UnregisteringLevelInstances.Remove(RootObject));
-				FActorFolders::Get().OnFolderRootObjectRemoved(*GetWorld(), RootObject);
-			}
-		}
-
 		// No need to clear the whole editor selection since actor of this level will be removed from the selection by: UEditorEngine::OnLevelRemovedFromWorld
 		EditorLevelUtils::RemoveLevelsFromWorld(InLevels, /*bClearSelection*/false, bResetTrans);
 	}
diff --git a/Engine/Source/Runtime/Engine/Private/LevelStreaming.cpp b/Engine/Source/Runtime/Engine/Private/LevelStreaming.cpp
--- a/Engine/Source/Runtime/Engine/Private/LevelStreaming.cpp
+++ b/Engine/Source/Runtime/Engine/Private/LevelStreaming.cpp
@@ -858,10 +858,6 @@ ULevelStreaming* ULevelStreaming::FindStreamingLevel(const ULevel* Level)
 
 void ULevelStreaming::SetLoadedLevel(ULevel* Level)
 { 
-#if WITH_EDITORONLY_DATA
-	check(!Level || !Level->IsBeingRemoved());
-#endif
-
 	// Pending level should be unloaded at this point
 	check(PendingUnloadLevel == nullptr);
 	PendingUnloadLevel = LoadedLevel;
diff --git a/Engine/Source/Runtime/Engine/Public/EngineUtils.h b/Engine/Source/Runtime/Engine/Public/EngineUtils.h
--- a/Engine/Source/Runtime/Engine/Public/EngineUtils.h
+++ b/Engine/Source/Runtime/Engine/Public/EngineUtils.h
@@ -424,7 +424,7 @@ protected:
 	{
 		if (EnumHasAnyFlags(Flags, EActorIteratorFlags::OnlyActiveLevels))
 		{
-			const bool bIsLevelVisibleOrAssociating = (Level->bIsVisible && !Level->IsBeingRemoved()) || Level->bIsAssociatingLevel || Level->bIsDisassociatingLevel;
+			const bool bIsLevelVisibleOrAssociating = (Level->bIsVisible && !Level->bIsBeingRemoved) || Level->bIsAssociatingLevel || Level->bIsDisassociatingLevel;
 
 			// Only allow iteration of Level if it's in the currently active level collection of the world, or is a static level.
 			const FLevelCollection* const ActorLevelCollection = Level->GetCachedLevelCollection();
diff --git a/Engine/Source/Runtime/Engine/Public/LevelInstance/LevelInstanceEditorLevelStreaming.h b/Engine/Source/Runtime/Engine/Public/LevelInstance/LevelInstanceEditorLevelStreaming.h
--- a/Engine/Source/Runtime/Engine/Public/LevelInstance/LevelInstanceEditorLevelStreaming.h
+++ b/Engine/Source/Runtime/Engine/Public/LevelInstance/LevelInstanceEditorLevelStreaming.h
@@ -22,6 +22,7 @@ public:
 	const FLevelInstanceID& GetLevelInstanceID() const { return LevelInstanceID; }
 	FBox GetBounds() const;
 
+	virtual TOptional<FFolder::FRootObject> GetFolderRootObject() const override;
 protected:
 	void OnLevelActorAdded(AActor* InActor);
 
diff --git a/Engine/Source/Runtime/Engine/Public/LevelInstance/LevelInstanceLevelStreaming.h b/Engine/Source/Runtime/Engine/Public/LevelInstance/LevelInstanceLevelStreaming.h
--- a/Engine/Source/Runtime/Engine/Public/LevelInstance/LevelInstanceLevelStreaming.h
+++ b/Engine/Source/Runtime/Engine/Public/LevelInstance/LevelInstanceLevelStreaming.h
@@ -21,6 +21,8 @@ public:
 #if WITH_EDITOR
 	virtual bool ShowInLevelCollection() const override { return false; }
 	FBox GetBounds() const;
+
+	virtual TOptional<FFolder::FRootObject> GetFolderRootObject() const override;
 #endif
 	
 protected:
diff --git a/Engine/Source/Runtime/Engine/Public/LevelInstance/LevelInstanceSubsystem.h b/Engine/Source/Runtime/Engine/Public/LevelInstance/LevelInstanceSubsystem.h
--- a/Engine/Source/Runtime/Engine/Public/LevelInstance/LevelInstanceSubsystem.h
+++ b/Engine/Source/Runtime/Engine/Public/LevelInstance/LevelInstanceSubsystem.h
@@ -168,8 +168,6 @@ private:
 	TUniquePtr<FLevelInstanceEdit> LevelInstanceEdit;
 	bool bCreatingLevelInstance = false;
 
-	TMap<FObjectKey, FFolder::FRootObject> UnregisteringLevelInstanceLevels;
-	TMap<FFolder::FRootObject, FObjectKey > UnregisteringLevelInstances;
 	TMap<FLevelInstanceID, int32> ChildEdits;
 #endif
 };
\ No newline at end of file
