Fix FWorldPartitionActorDesc::Equals missing parent actor check.

diff --git a/Engine/Source/Runtime/Engine/Private/WorldPartition/WorldPartitionActorDesc.cpp b/Engine/Source/Runtime/Engine/Private/WorldPartition/WorldPartitionActorDesc.cpp
--- a/Engine/Source/Runtime/Engine/Private/WorldPartition/WorldPartitionActorDesc.cpp
+++ b/Engine/Source/Runtime/Engine/Private/WorldPartition/WorldPartitionActorDesc.cpp
@@ -144,6 +144,7 @@ bool FWorldPartitionActorDesc::Equals(const FWorldPartitionActorDesc* Other) con
 		bActorIsHLODRelevant == Other->bActorIsHLODRelevant && 
 		HLODLayer == Other->HLODLayer && 
 		FolderPath == Other->FolderPath &&
+		ParentActor == Other->ParentActor &&
 		DataLayers.Num() == Other->DataLayers.Num() &&
 		References.Num() == Other->References.Num())
 	{
