Fix redirector issues by adding deprecated explicit override in AudioComponent.
#rnx

diff --git a/Engine/Source/Runtime/Engine/Classes/Components/AudioComponent.h b/Engine/Source/Runtime/Engine/Classes/Components/AudioComponent.h
--- a/Engine/Source/Runtime/Engine/Classes/Components/AudioComponent.h
+++ b/Engine/Source/Runtime/Engine/Classes/Components/AudioComponent.h
@@ -502,6 +502,27 @@ public:
 		, EAudioFaderCurve InFadeCurve = EAudioFaderCurve::Linear
 	);
 
+	// Sets a named Boolean
+	UFUNCTION(BlueprintCallable, meta = (DisplayName = "Set Boolean Parameter", DeprecatedFunction, DeprecationMessage = "5.0 - Moved to AudioParameterControllerInterface"), Category = "Audio|Parameter")
+	virtual void SetBoolParameter(FName InName, bool InBool) override
+	{
+		return ISoundParameterControllerInterface::SetBoolParameter(InName, InBool);
+	}
+
+	// Sets a named Int32
+	UFUNCTION(BlueprintCallable, meta = (DisplayName = "Set Integer Parameter", DeprecatedFunction, DeprecationMessage = "5.0 - Moved to AudioParameterControllerInterface"), Category = "Audio|Parameter")
+	virtual void SetIntParameter(FName InName, int32 InInt) override
+	{
+		return ISoundParameterControllerInterface::SetIntParameter(InName, InInt);
+	}
+
+	// Sets a named Float
+	UFUNCTION(BlueprintCallable, meta = (DisplayName = "Set Float Parameter", DeprecatedFunction, DeprecationMessage = "5.0 - Moved to AudioParameterControllerInterface"), Category = "Audio|Parameter")
+	virtual void SetFloatParameter(FName InName, float InFloat) override
+	{
+		return ISoundParameterControllerInterface::SetFloatParameter(InName, InFloat);
+	}
+
 	virtual void ResetParameters() override;
 
 
