Combine ES3.1 shader platform check and mobile platform check into one for morph target

diff --git a/Engine/Source/Runtime/Engine/Private/GPUSkinVertexFactory.cpp b/Engine/Source/Runtime/Engine/Private/GPUSkinVertexFactory.cpp
--- a/Engine/Source/Runtime/Engine/Private/GPUSkinVertexFactory.cpp
+++ b/Engine/Source/Runtime/Engine/Private/GPUSkinVertexFactory.cpp
@@ -935,7 +935,7 @@ public:
 
 		// Mobile doesn't support motion blur, don't use previous frame morph delta for mobile.
 		const EShaderPlatform ShaderPlatform = GetFeatureLevelShaderPlatform(FeatureLevel);
-		const bool bIsMobile = IsMobilePlatform(ShaderPlatform) || ShaderPlatform == SP_PCD3D_ES3_1;
+		const bool bIsMobile = IsMobilePlatform(ShaderPlatform);
 		if (!bIsMobile)
 		{
 			const FMorphVertexBuffer* MorphVertexBuffer = nullptr;
@@ -969,7 +969,7 @@ void TGPUSkinMorphVertexFactory<BoneInfluenceType>::ModifyCompilationEnvironment
 	Super::ModifyCompilationEnvironment(Parameters, OutEnvironment);
 	OutEnvironment.SetDefine(TEXT("GPUSKIN_MORPH_BLEND"),TEXT("1"));
 	// Mobile doesn't support motion blur, don't use previous frame morph delta for mobile.
-	const bool bIsMobile = IsMobilePlatform(Parameters.Platform) || Parameters.Platform == SP_PCD3D_ES3_1;
+	const bool bIsMobile = IsMobilePlatform(Parameters.Platform);
 	OutEnvironment.SetDefine(TEXT("GPUSKIN_MORPH_USE_PREVIOUS"), !bIsMobile);
 }
 
