Fix rewind debugger behavior in Debug config

diff --git a/Engine/Plugins/Animation/GameplayInsights/Source/RewindDebugger/Private/RewindDebugger.cpp b/Engine/Plugins/Animation/GameplayInsights/Source/RewindDebugger/Private/RewindDebugger.cpp
--- a/Engine/Plugins/Animation/GameplayInsights/Source/RewindDebugger/Private/RewindDebugger.cpp
+++ b/Engine/Plugins/Animation/GameplayInsights/Source/RewindDebugger/Private/RewindDebugger.cpp
@@ -673,28 +673,31 @@ void FRewindDebugger::Tick(float DeltaTime)
 						
 						// until we have actor transforms traced out, the first skeletal mesh component transform on the target actor be used as as the actor position 
 						uint64 TargetActorId = GetTargetActorId();
-						if(UObject* ObjectInstance = FObjectTrace::GetObjectFromId(TargetActorId))
+						if (TargetActorId != 0)
 						{
-							if (AActor* TargetActor = Cast<AActor>(ObjectInstance))
+							if (UObject* ObjectInstance = FObjectTrace::GetObjectFromId(TargetActorId))
 							{
-								TInlineComponentArray<USkeletalMeshComponent*> SkeletalMeshComponents;
-								TargetActor->GetComponents(SkeletalMeshComponents);
-
-								if (SkeletalMeshComponents.Num() > 0)
+								if (AActor* TargetActor = Cast<AActor>(ObjectInstance))
 								{
-									int64 ObjectId = FObjectTrace::GetObjectId(SkeletalMeshComponents[0]);
+									TInlineComponentArray<USkeletalMeshComponent*> SkeletalMeshComponents;
+									TargetActor->GetComponents(SkeletalMeshComponents);
 
-									AnimationProvider->ReadSkeletalMeshPoseTimeline(ObjectId, [this, Frame, ObjectId, AnimationProvider](const IAnimationProvider::SkeletalMeshPoseTimeline& TimelineData, bool bHasCurves)
+									if (SkeletalMeshComponents.Num() > 0)
 									{
-										TimelineData.EnumerateEvents(Frame.StartTime, Frame.EndTime,
-											[this](double InStartTime, double InEndTime, uint32 InDepth, const FSkeletalMeshPoseMessage& PoseMessage)
+										int64 ObjectId = FObjectTrace::GetObjectId(SkeletalMeshComponents[0]);
+
+										AnimationProvider->ReadSkeletalMeshPoseTimeline(ObjectId, [this, Frame, ObjectId, AnimationProvider](const IAnimationProvider::SkeletalMeshPoseTimeline& TimelineData, bool bHasCurves)
 											{
-												bTargetActorPositionValid = true;
-												TargetActorPosition = PoseMessage.ComponentToWorld.GetTranslation();
-										
-												return TraceServices::EEventEnumerate::Stop;
+												TimelineData.EnumerateEvents(Frame.StartTime, Frame.EndTime,
+													[this](double InStartTime, double InEndTime, uint32 InDepth, const FSkeletalMeshPoseMessage& PoseMessage)
+													{
+														bTargetActorPositionValid = true;
+														TargetActorPosition = PoseMessage.ComponentToWorld.GetTranslation();
+
+														return TraceServices::EEventEnumerate::Stop;
+													});
 											});
-									});
+									}
 								}
 							}
 						}
diff --git a/Engine/Source/Runtime/Engine/Private/ObjectTrace.cpp b/Engine/Source/Runtime/Engine/Private/ObjectTrace.cpp
--- a/Engine/Source/Runtime/Engine/Private/ObjectTrace.cpp
+++ b/Engine/Source/Runtime/Engine/Private/ObjectTrace.cpp
@@ -86,36 +86,48 @@ UE_TRACE_EVENT_END()
 struct FTracedObjectAnnotation
 {
 	FTracedObjectAnnotation()
-		: Id(0)
-		, bTraced(false)
+		: bTraced(false)
 	{}
 
-	// Object ID
-	uint64 Id;
-
 	// Whether this object has been traced this session
 	bool bTraced;
 
 	/** Determine if this annotation is default - required for annotations */
 	FORCEINLINE bool IsDefault() const
 	{
-		return bTraced == false && Id == 0;
+		return bTraced == false;
 	}
+};
+
+struct FObjectIdAnnotation
+{
+	FObjectIdAnnotation()
+		: Id(0)
+	{}
+
+	// Object ID
+	uint64 Id;
 
+	/** Determine if this annotation is default - required for annotations */
+	FORCEINLINE bool IsDefault() const
+	{
+		return Id == 0;
+	}
 
-	bool operator == (const FTracedObjectAnnotation& other) const
+	bool operator == (const FObjectIdAnnotation& other) const
 	{
 		return Id == other.Id;
 	}
 };
 
-int32 GetTypeHash(const FTracedObjectAnnotation& Annotation)
+int32 GetTypeHash(const FObjectIdAnnotation& Annotation)
 {
 	return GetTypeHash(Annotation.Id);
 }
 
 // Object annotations used for tracing
-FUObjectAnnotationSparseSearchable<FTracedObjectAnnotation, true> GObjectTraceAnnotations;
+FUObjectAnnotationSparse<FTracedObjectAnnotation, true> GObjectTracedAnnotations;
+FUObjectAnnotationSparseSearchable<FObjectIdAnnotation, true> GObjectIdAnnotations;
 
 // Handle used to hook to world tick
 static FDelegateHandle WorldTickStartHandle;
@@ -162,11 +174,11 @@ uint64 FObjectTrace::GetObjectId(const UObject* InObject)
 	{
 		static uint64 CurrentId = 1;
 
-		FTracedObjectAnnotation Annotation = GObjectTraceAnnotations.GetAnnotation(InObjectInner);
+		FObjectIdAnnotation Annotation = GObjectIdAnnotations.GetAnnotation(InObjectInner);
 		if(Annotation.Id == 0)
 		{
 			Annotation.Id = CurrentId++;
-			GObjectTraceAnnotations.AddAnnotation(InObjectInner, MoveTemp(Annotation));
+			GObjectIdAnnotations.AddAnnotation(InObjectInner, MoveTemp(Annotation));
 		}
 
 		return Annotation.Id;
@@ -189,10 +201,10 @@ uint64 FObjectTrace::GetObjectId(const UObject* InObject)
 
 UObject* FObjectTrace::GetObjectFromId(uint64 Id)
 {
-	FTracedObjectAnnotation FindAnnotation;
+	FObjectIdAnnotation FindAnnotation;
 	// Id used for annotation map doesn't include the parent id in the upper bits, so zero those first
 	FindAnnotation.Id = Id & 0x00000000FFFFFFFFll;
-	return GObjectTraceAnnotations.Find(FindAnnotation);
+	return GObjectIdAnnotations.Find(FindAnnotation);
 }
 
 void FObjectTrace::ResetWorldElapsedTime(const UWorld* World)
@@ -285,17 +297,17 @@ void FObjectTrace::OutputClass(const UClass* InClass)
 		return;
 	}
 
-	OutputClass(InClass->GetSuperClass());
-
-	FTracedObjectAnnotation Annotation = GObjectTraceAnnotations.GetAnnotation(InClass);
+	FTracedObjectAnnotation Annotation = GObjectTracedAnnotations.GetAnnotation(InClass);
 	if(Annotation.bTraced)
 	{
 		// Already traced, so skip
 		return;
 	}
 
+	OutputClass(InClass->GetSuperClass());
+
 	Annotation.bTraced = true;
-	GObjectTraceAnnotations.AddAnnotation(InClass, MoveTemp(Annotation));
+	GObjectTracedAnnotations.AddAnnotation(InClass, MoveTemp(Annotation));
 
 	FString ClassPathName = InClass->GetPathName();
 	TCHAR ClassName[FName::StringBufferSize];
@@ -355,7 +367,7 @@ void FObjectTrace::OutputObject(const UObject* InObject)
 		return;
 	}
 
-	FTracedObjectAnnotation Annotation = GObjectTraceAnnotations.GetAnnotation(InObject);
+	FTracedObjectAnnotation Annotation = GObjectTracedAnnotations.GetAnnotation(InObject);
 	if(Annotation.bTraced)
 	{
 		// Already traced, so skip
@@ -365,7 +377,7 @@ void FObjectTrace::OutputObject(const UObject* InObject)
 	OutputObject(InObject->GetOuter());
 
 	Annotation.bTraced = true;
-	GObjectTraceAnnotations.AddAnnotation(InObject, MoveTemp(Annotation));
+	GObjectTracedAnnotations.AddAnnotation(InObject, MoveTemp(Annotation));
 
 	// Trace the object's class first
 	TRACE_CLASS(InObject->GetClass());
