New HLSL translator supports structs in preshaders....allows constant folding to be tracked across struct assignments

diff --git a/Engine/Source/Runtime/Core/Public/Misc/MemStackUtility.h b/Engine/Source/Runtime/Core/Public/Misc/MemStackUtility.h
new file mode 100644
--- /dev/null
+++ b/Engine/Source/Runtime/Core/Public/Misc/MemStackUtility.h
@@ -0,0 +1,57 @@
+#pragma once
+
+#include "CoreMinimal.h"
+#include "Misc/MemStack.h"
+#include "Misc/StringBuilder.h"
+#include "Containers/ArrayView.h"
+#include "Containers/StringView.h"
+
+namespace UE
+{
+namespace MemStack
+{
+
+inline const TCHAR* AllocateString(FMemStackBase& Allocator, const TCHAR* String, int32 Length)
+{
+	TCHAR* Result = New<TCHAR>(Allocator, Length + 1);
+	FMemory::Memcpy(Result, String, Length * sizeof(TCHAR));
+	Result[Length] = 0;
+	return Result;
+}
+
+inline const TCHAR* AllocateString(FMemStackBase& Allocator, FStringView String)
+{
+	return AllocateString(Allocator, String.GetData(), String.Len());
+}
+
+inline const TCHAR* AllocateString(FMemStackBase& Allocator, const FStringBuilderBase& StringBuilder)
+{
+	return AllocateString(Allocator, StringBuilder.GetData(), StringBuilder.Len());
+}
+
+template <typename FormatType, typename... ArgTypes>
+inline const TCHAR* AllocateStringf(FMemStackBase& Allocator, const FormatType& Format, ArgTypes... Args)
+{
+	TStringBuilder<1024> String;
+	String.Appendf(Format, Forward<ArgTypes>(Args)...);
+	return AllocateString(Allocator, String);
+}
+
+
+template<typename T>
+inline TArrayView<T> AllocateArrayView(FMemStackBase& Allocator, TArrayView<T> View)
+{
+	T* Data = nullptr;
+	if (View.Num() > 0)
+	{
+		Data = new(Allocator) T[View.Num()];
+		for (int32 i = 0; i < View.Num(); ++i)
+		{
+			Data[i] = View[i];
+		}
+	}
+	return MakeArrayView(Data, View.Num());
+}
+
+} // namespace MemStack
+} // namespace UE
diff --git a/Engine/Source/Runtime/Engine/Private/HLSLTree/HLSLTree.cpp b/Engine/Source/Runtime/Engine/Private/HLSLTree/HLSLTree.cpp
--- a/Engine/Source/Runtime/Engine/Private/HLSLTree/HLSLTree.cpp
+++ b/Engine/Source/Runtime/Engine/Private/HLSLTree/HLSLTree.cpp
@@ -1,50 +1,12 @@
 // Copyright Epic Games, Inc. All Rights Reserved.
 #include "HLSLTree/HLSLTree.h"
 #include "Misc/StringBuilder.h"
-#include "Misc/MemStack.h"
+#include "Misc/MemStackUtility.h"
 #include "MaterialShared.h" // TODO - split preshader out into its own module
 
 // From MaterialUniformExpressions.cpp
 extern void WriteMaterialUniformAccess(UE::Shader::EValueComponentType ComponentType, uint32 NumComponents, uint32 UniformOffset, FStringBuilderBase& OutResult);
 
-static const TCHAR* AllocateString(FMemStackBase& Allocator, const TCHAR* String, int32 Length)
-{
-	TCHAR* Result = New<TCHAR>(Allocator, Length + 1);
-	FMemory::Memcpy(Result, String, Length * sizeof(TCHAR));
-	Result[Length] = 0;
-	return Result;
-}
-
-static const TCHAR* AllocateString(FMemStackBase& Allocator, FStringView String)
-{
-	return AllocateString(Allocator, String.GetData(), String.Len());
-}
-
-static const TCHAR* AllocateString(FMemStackBase& Allocator, const FStringBuilderBase& StringBuilder)
-{
-	return AllocateString(Allocator, StringBuilder.GetData(), StringBuilder.Len());
-}
-
-template <typename FormatType, typename... ArgTypes>
-static const TCHAR* AllocateStringf(FMemStackBase& Allocator, const FormatType& Format, ArgTypes... Args)
-{
-	TStringBuilder<1024> String;
-	String.Appendf(Format, Forward<ArgTypes>(Args)...);
-	return AllocateString(Allocator, String);
-}
-
-
-template<typename T, typename AllocatorType>
-static TArrayView<T> AllocateArrayView(FMemStackBase& Allocator, const TArray<T, AllocatorType>& Array)
-{
-	T* Data = new(Allocator) T[Array.Num()];
-	for (int32 i = 0; i < Array.Num(); ++i)
-	{
-		Data[i] = Array[i];
-	}
-	return MakeArrayView(Data, Array.Num());
-}
-
 static FSHAHash HashString(const TCHAR* String, int32 Length)
 {
 	FSHAHash Hash;
@@ -64,7 +26,16 @@ namespace HLSLTree
 
 EExpressionEvaluationType CombineEvaluationTypes(EExpressionEvaluationType Lhs, EExpressionEvaluationType Rhs)
 {
-	if (Lhs == EExpressionEvaluationType::Constant && Rhs == EExpressionEvaluationType::Constant)
+	if (Lhs == EExpressionEvaluationType::None)
+	{
+		// If either is 'None', return the other
+		return Rhs;
+	}
+	else if (Rhs == EExpressionEvaluationType::None)
+	{
+		return Lhs;
+	}
+	else if (Lhs == EExpressionEvaluationType::Constant && Rhs == EExpressionEvaluationType::Constant)
 	{
 		// 2 constants make a constant
 		return EExpressionEvaluationType::Constant;
@@ -99,8 +70,9 @@ void FErrors::AddError(const FNode* InNode, FStringView InError)
 	ensureMsgf(false, TEXT("%s"), Error->Message);
 }
 
-FEmitContext::FEmitContext(FMemStackBase& InAllocator)
+FEmitContext::FEmitContext(FMemStackBase& InAllocator, const Shader::FStructTypeRegistry& InTypeRegistry)
 	: Allocator(&InAllocator)
+	, TypeRegistry(&InTypeRegistry)
 	, Errors(InAllocator)
 {
 }
@@ -111,175 +83,238 @@ FEmitContext::~FEmitContext()
 
 const TCHAR* FEmitContext::AcquireLocalDeclarationCode()
 {
-	return AllocateStringf(*Allocator, TEXT("Local%d"), NumExpressionLocals++);
+	return MemStack::AllocateStringf(*Allocator, TEXT("Local%d"), NumExpressionLocals++);
 }
 
-const TCHAR* FEmitContext::CastShaderValue(const FNode* Node, const TCHAR* Code, const FType& SourceType, const FType& DestType, ECastFlags Flags)
+FEmitShaderValue* FEmitContext::CastShaderValue(FNode* Node, FScope* Scope, FEmitShaderValue* ShaderValue, const Shader::FType& DestType)
 {
-	if (SourceType == DestType)
+	if (ShaderValue->Type == DestType)
 	{
-		return Code;
+		return ShaderValue;
 	}
 
-	const bool bAllowTruncate = EnumHasAnyFlags(Flags, ECastFlags::AllowTruncate);
-	const bool bAllowAppendZeroes = EnumHasAnyFlags(Flags, ECastFlags::AllowAppendZeroes);
-	bool bReplicateScalar = EnumHasAnyFlags(Flags, ECastFlags::ReplicateScalar);
-
-	const Shader::FValueTypeDescription SourceTypeDesc = Shader::GetValueTypeDescription(SourceType);
+	const Shader::FValueTypeDescription SourceTypeDesc = Shader::GetValueTypeDescription(ShaderValue->Type);
 	const Shader::FValueTypeDescription DestTypeDesc = Shader::GetValueTypeDescription(DestType);
 
+	TStringBuilder<1024> FormattedCode;
+	FShaderValue Shader(FormattedCode);
+	Shader.Type = DestType;
+	Shader.bInline = true;
+
 	if (SourceTypeDesc.NumComponents > 0 && DestTypeDesc.NumComponents > 0)
 	{
-		if (SourceTypeDesc.NumComponents != 1)
-		{
-			bReplicateScalar = false;
-		}
-		if (!bReplicateScalar && !bAllowAppendZeroes && DestTypeDesc.NumComponents > SourceTypeDesc.NumComponents)
-		{
-			Errors.AddErrorf(Node, TEXT("Cannot cast from smaller type %s to larger type %s."), SourceTypeDesc.Name, DestTypeDesc.Name);
-			return TEXT("");
-		}
-		if (!bReplicateScalar && !bAllowTruncate && DestTypeDesc.NumComponents < SourceTypeDesc.NumComponents)
-		{
-			Errors.AddErrorf(Node, TEXT("Cannot cast from larger type %s to smaller type %s."), SourceTypeDesc.Name, DestTypeDesc.Name);
-			return TEXT("");
-		}
-
 		const bool bIsSourceLWC = SourceTypeDesc.ComponentType == Shader::EValueComponentType::Double;
 		const bool bIsLWC = DestTypeDesc.ComponentType == Shader::EValueComponentType::Double;
+
 		if (bIsLWC != bIsSourceLWC)
 		{
 			if (bIsLWC)
 			{
 				// float->LWC
-				const TCHAR* CodeAsFloat = CastShaderValue(Node, Code, SourceType, Shader::MakeValueType(Shader::EValueComponentType::Float, DestTypeDesc.NumComponents), Flags);
-				return AllocateStringf(*Allocator, TEXT("LWCPromote(%s)"), CodeAsFloat);
+				ShaderValue = CastShaderValue(Node, Scope, ShaderValue, Shader::MakeValueType(Shader::EValueComponentType::Float, DestTypeDesc.NumComponents));
+				Shader.Code.Appendf(TEXT("LWCPromote(%s)"), ShaderValue->Reference);
 			}
 			else
 			{
 				//LWC->float
-				const TCHAR* CodeAsFloat = AllocateStringf(*Allocator, TEXT("LWCToFloat(%s)"), Code);
-				return CastShaderValue(Node, CodeAsFloat, Shader::MakeValueType(Shader::EValueComponentType::Float, SourceTypeDesc.NumComponents), DestType, Flags);
+				Shader.Code.Appendf(TEXT("LWCToFloat(%s)"), ShaderValue->Reference);
+				Shader.Type = Shader::MakeValueType(Shader::EValueComponentType::Float, SourceTypeDesc.NumComponents);
 			}
 		}
-
-		TStringBuilder<1024> Result;
-		int32 NumComponents = 0;
-		bool bNeedClosingParen = false;
-		if (bIsLWC)
-		{
-			Result.Append(TEXT("MakeLWCVector("));
-			bNeedClosingParen = true;
-		}
 		else
 		{
-			if (bReplicateScalar || SourceTypeDesc.NumComponents == DestTypeDesc.NumComponents)
+			const bool bReplicateScalar = (SourceTypeDesc.NumComponents == 1);
+
+			int32 NumComponents = 0;
+			bool bNeedClosingParen = false;
+			if (bIsLWC)
 			{
-				NumComponents = DestTypeDesc.NumComponents;
-				// Cast the scalar to the correct type, HLSL language will replicate the scalar if needed when performing this cast
-				Result.Appendf(TEXT("((%s)%s)"), DestTypeDesc.Name, Code);
+				Shader.Code.Append(TEXT("MakeLWCVector("));
+				bNeedClosingParen = true;
 			}
 			else
 			{
-				NumComponents = FMath::Min(SourceTypeDesc.NumComponents, DestTypeDesc.NumComponents);
-				if (NumComponents < DestTypeDesc.NumComponents)
-				{
-					Result.Appendf(TEXT("%s("), DestTypeDesc.Name);
-					bNeedClosingParen = true;
-				}
-				if (NumComponents == SourceTypeDesc.NumComponents && SourceTypeDesc.ComponentType == DestTypeDesc.ComponentType)
+				if (SourceTypeDesc.NumComponents == 1 || SourceTypeDesc.NumComponents == DestTypeDesc.NumComponents)
 				{
-					// If we're taking all the components from the source, can avoid adding a swizzle
-					Result.Append(Code);
+					NumComponents = DestTypeDesc.NumComponents;
+					// Cast the scalar to the correct type, HLSL language will replicate the scalar if needed when performing this cast
+					Shader.Code.Appendf(TEXT("((%s)%s)"), DestTypeDesc.Name, ShaderValue->Reference);
 				}
 				else
 				{
-					// Use a cast to truncate the source to the correct number of types
-					const Shader::EValueType IntermediateType = Shader::MakeValueType(DestTypeDesc.ComponentType, NumComponents);
-					const Shader::FValueTypeDescription IntermediateTypeDesc = Shader::GetValueTypeDescription(IntermediateType);
-					Result.Appendf(TEXT("((%s)%s)"), IntermediateTypeDesc.Name, Code);
+					NumComponents = FMath::Min(SourceTypeDesc.NumComponents, DestTypeDesc.NumComponents);
+					if (NumComponents < DestTypeDesc.NumComponents)
+					{
+						Shader.Code.Appendf(TEXT("%s("), DestTypeDesc.Name);
+						bNeedClosingParen = true;
+					}
+					if (NumComponents == SourceTypeDesc.NumComponents && SourceTypeDesc.ComponentType == DestTypeDesc.ComponentType)
+					{
+						// If we're taking all the components from the source, can avoid adding a swizzle
+						Shader.Code.Append(ShaderValue->Reference);
+					}
+					else
+					{
+						// Use a cast to truncate the source to the correct number of types
+						const Shader::EValueType IntermediateType = Shader::MakeValueType(DestTypeDesc.ComponentType, NumComponents);
+						const Shader::FValueTypeDescription IntermediateTypeDesc = Shader::GetValueTypeDescription(IntermediateType);
+						Shader.Code.Appendf(TEXT("((%s)%s)"), IntermediateTypeDesc.Name, ShaderValue->Reference);
+					}
 				}
 			}
-		}
 
-		if (bNeedClosingParen)
-		{
-			const Shader::FValue ZeroValue(DestTypeDesc.ComponentType, 1);
-			for (int32 ComponentIndex = NumComponents; ComponentIndex < DestTypeDesc.NumComponents; ++ComponentIndex)
+			if (bNeedClosingParen)
 			{
-				if (ComponentIndex > 0u)
-				{
-					Result.Append(TEXT(","));
-				}
-				if (bIsLWC)
+				const Shader::FValue ZeroValue(DestTypeDesc.ComponentType, 1);
+				for (int32 ComponentIndex = NumComponents; ComponentIndex < DestTypeDesc.NumComponents; ++ComponentIndex)
 				{
-					if (!bReplicateScalar && ComponentIndex >= SourceTypeDesc.NumComponents)
+					if (ComponentIndex > 0u)
+					{
+						Shader.Code.Append(TEXT(","));
+					}
+					if (bIsLWC)
 					{
-						check(bAllowAppendZeroes);
-						Result.Append(TEXT("LWCPromote(0.0f)"));
+						if (!bReplicateScalar && ComponentIndex >= SourceTypeDesc.NumComponents)
+						{
+							Shader.Code.Append(TEXT("LWCPromote(0.0f)"));
+						}
+						else
+						{
+							Shader.Code.Appendf(TEXT("LWCGetComponent(%s, %d)"), ShaderValue->Reference, bReplicateScalar ? 0 : ComponentIndex);
+						}
 					}
 					else
 					{
-						Result.Appendf(TEXT("LWCGetComponent(%s, %d)"), Code, bReplicateScalar ? 0 : ComponentIndex);
+						// Non-LWC case should only be zero-filling here, other cases should have already been handled
+						check(!bReplicateScalar);
+						check(ComponentIndex >= SourceTypeDesc.NumComponents);
+						ZeroValue.ToString(Shader::EValueStringFormat::HLSL, Shader.Code);
 					}
 				}
-				else
-				{
-					// Non-LWC case should only be zero-filling here, other cases should have already been handled
-					check(bAllowAppendZeroes);
-					check(!bReplicateScalar);
-					check(ComponentIndex >= SourceTypeDesc.NumComponents);
-					ZeroValue.ToString(Shader::EValueStringFormat::HLSL, Result);
-				}
+				NumComponents = DestTypeDesc.NumComponents;
+				Shader.Code.Append(TEXT(")"));
 			}
-			NumComponents = DestTypeDesc.NumComponents;
-			Result.Append(TEXT(")"));
-		}
 
-		check(NumComponents == DestTypeDesc.NumComponents);
-		return AllocateString(*Allocator, Result);
+			check(NumComponents == DestTypeDesc.NumComponents);
+		}
 	}
 	else
 	{
 		Errors.AddErrorf(Node, TEXT("Cannot cast between non-numeric types %s to %s."), SourceTypeDesc.Name, DestTypeDesc.Name);
-		return TEXT("");
+		Shader.Code.Appendf(TEXT("((%s)0)"), DestType.GetName());
+		Shader.bInline = false;
+	}
+
+	check(Shader.Type != ShaderValue->Type);
+	ShaderValue = AcquireShader(Scope, Shader, MakeArrayView(&ShaderValue, 1));
+	if (ShaderValue->Type != DestType)
+	{
+		// May need to cast through multiple intermediate types to reach our destination type
+		ShaderValue = CastShaderValue(Node, Scope, ShaderValue, DestType);
+	}
+	return ShaderValue;
+}
+
+namespace Private
+{
+void MoveToScope(FEmitShaderValue* ShaderValue, FScope* Scope)
+{
+	if (ShaderValue->Scope != Scope)
+	{
+		FScope* NewScope = FScope::FindSharedParent(ShaderValue->Scope, Scope);
+		check(NewScope);
+
+		ShaderValue->Scope = NewScope;
+		for (FEmitShaderValue* Dependency : ShaderValue->Dependencies)
+		{
+			MoveToScope(Dependency, NewScope);
+		}
 	}
 }
+}
 
-const TCHAR* FEmitContext::AcquirePreshader(Shader::EValueType Type, const Shader::FPreshaderData& Preshader)
+FEmitShaderValue* FEmitContext::AcquireShader(FScope* Scope, const FShaderValue& Shader, TArrayView<FEmitShaderValue*> Dependencies)
 {
+	FEmitShaderValue* ShaderValue = nullptr;
+
+	// Check to see if we've already generated code for an equivalent expression
+	const FSHAHash ShaderHash = HashString(Shader.Code);
+	FEmitShaderValue** const PrevShaderValue = ShaderValueMap.Find(ShaderHash);
+	if (PrevShaderValue)
+	{
+		ShaderValue = *PrevShaderValue;
+		check(ShaderValue->Type == Shader.Type);
+		if (!Shader.bInline)
+		{
+			Private::MoveToScope(ShaderValue, Scope);
+		}
+	}
+	else
+	{
+		ShaderValue = new(*Allocator) FEmitShaderValue(Scope, Shader.Type);
+		ShaderValue->Hash = ShaderHash;
+		ShaderValue->Dependencies = MemStack::AllocateArrayView(*Allocator, Dependencies);
+		if (Shader.bInline)
+		{
+			ShaderValue->Reference = MemStack::AllocateString(*Allocator, Shader.Code);
+		}
+		else
+		{
+			ShaderValue->Reference = AcquireLocalDeclarationCode();
+			ShaderValue->Value = MemStack::AllocateString(*Allocator, Shader.Code);
+			ShaderValueMap.Add(ShaderHash, ShaderValue);
+		}
+	}
+
+	return ShaderValue;
+}
+
+FEmitShaderValue* FEmitContext::AcquirePreshader(const FRequestedType& RequestedType, FScope* Scope, FExpression* Expression)
+{
+	Shader::FPreshaderData LocalPreshader;
+	Expression->EmitValuePreshader(*this, RequestedType, LocalPreshader);
+
+	const Shader::FType Type = RequestedType.GetType();
+	check(!Type.IsStruct()); // TODO
+
 	FSHA1 Hasher;
 	Hasher.Update((uint8*)&Type, sizeof(Type));
-	Preshader.AppendHash(Hasher);
+	LocalPreshader.AppendHash(Hasher);
 	const FSHAHash Hash = Hasher.Finalize();
-	TCHAR const* const* PrevPreshader = Preshaders.Find(Hash);
+	FEmitShaderValue* const* PrevPreshader = Preshaders.Find(Hash);
 	if (PrevPreshader)
 	{
 		return *PrevPreshader;
 	}
 
 	const Shader::FValueTypeDescription TypeDesc = Shader::GetValueTypeDescription(Type);
-	ensure(TypeDesc.ComponentType == Shader::EValueComponentType::Float || TypeDesc.ComponentType == Shader::EValueComponentType::Double);
+	Shader::EValueComponentType ComponentType = TypeDesc.ComponentType;
+	if (ComponentType == Shader::EValueComponentType::Bool || ComponentType == Shader::EValueComponentType::Int)
+	{
+		ComponentType = Shader::EValueComponentType::Float; // TODO
+	}
+
+	//ensure(TypeDesc.ComponentType == Shader::EValueComponentType::Float || TypeDesc.ComponentType == Shader::EValueComponentType::Double);
 
 	FUniformExpressionSet& UniformExpressionSet = MaterialCompilationOutput->UniformExpressionSet;
 	FMaterialUniformPreshaderHeader& PreshaderHeader = UniformExpressionSet.UniformPreshaders.AddDefaulted_GetRef();
 
 	const uint32 RegisterOffset = UniformPreshaderOffset % 4;
-	if (TypeDesc.ComponentType == Shader::EValueComponentType::Float && RegisterOffset + TypeDesc.NumComponents > 4u)
+	if (ComponentType == Shader::EValueComponentType::Float && RegisterOffset + TypeDesc.NumComponents > 4u)
 	{
 		// If this uniform would span multiple registers, align offset to the next register to avoid this
 		UniformPreshaderOffset = Align(UniformPreshaderOffset, 4u);
 	}
 
 	PreshaderHeader.OpcodeOffset = UniformExpressionSet.UniformPreshaderData.Num();
-	UniformExpressionSet.UniformPreshaderData.Append(Preshader);
-	PreshaderHeader.OpcodeSize = Preshader.Num();
+	Expression->EmitValuePreshader(*this, RequestedType, UniformExpressionSet.UniformPreshaderData);
+	PreshaderHeader.OpcodeSize = UniformExpressionSet.UniformPreshaderData.Num() - PreshaderHeader.OpcodeOffset;
 	PreshaderHeader.BufferOffset = UniformPreshaderOffset;
-	PreshaderHeader.ComponentType = TypeDesc.ComponentType;
+	PreshaderHeader.ComponentType = ComponentType;
 	PreshaderHeader.NumComponents = TypeDesc.NumComponents;
 
 	TStringBuilder<1024> FormattedCode;
-	if (TypeDesc.ComponentType == Shader::EValueComponentType::Double)
+	if (ComponentType == Shader::EValueComponentType::Double)
 	{
 		if (TypeDesc.NumComponents == 1)
 		{
@@ -299,13 +334,14 @@ const TCHAR* FEmitContext::AcquirePreshader(Shader::EValueType Type, const Shade
 	}
 	else
 	{
-		WriteMaterialUniformAccess(TypeDesc.ComponentType, TypeDesc.NumComponents, UniformPreshaderOffset, FormattedCode);
+		WriteMaterialUniformAccess(ComponentType, TypeDesc.NumComponents, UniformPreshaderOffset, FormattedCode);
 		UniformPreshaderOffset += TypeDesc.NumComponents;
 	}
 
-	const TCHAR* Code = AllocateString(*Allocator, FormattedCode);
-	Preshaders.Add(Hash, Code);
-	return Code;
+	FEmitShaderValue* ShaderValue = new(*Allocator) FEmitShaderValue(Scope, Shader::MakeValueType(ComponentType, TypeDesc.NumComponents));
+	ShaderValue->Reference = MemStack::AllocateString(*Allocator, FormattedCode);
+	Preshaders.Add(Hash, ShaderValue);
+	return ShaderValue;
 }
 
 namespace Private
@@ -323,7 +359,7 @@ void EmitShaderValue(FEmitContext& Context, FEmitShaderValue* ShaderValue)
 		if (!ShaderValue->IsInline())
 		{
 			ShaderValue->Scope->EmitDeclarationf(Context, TEXT("const %s %s = %s;"),
-				ShaderValue->Expression->GetType().GetName(),
+				ShaderValue->Type.GetName(),
 				ShaderValue->Reference,
 				ShaderValue->Value);
 		}
@@ -341,85 +377,13 @@ void FEmitContext::Finalize()
 	{
 		Private::EmitShaderValue(*this, It.Value);
 	}
+
 	ShaderValueMap.Reset();
+	LocalPHIs.Reset();
 
 	MaterialCompilationOutput->UniformExpressionSet.UniformPreshaderBufferSize = (UniformPreshaderOffset + 3u) / 4u;
 }
 
-const TCHAR* FType::GetName() const
-{
-	if (StructType)
-	{
-		return StructType->Name;
-	}
-
-	check(ValueType != Shader::EValueType::Struct);
-	const Shader::FValueTypeDescription TypeDesc = Shader::GetValueTypeDescription(ValueType);
-	return TypeDesc.Name;
-}
-
-int32 FType::GetNumComponents() const
-{
-	if (StructType)
-	{
-		return StructType->ComponentTypes.Num();
-	}
-	const Shader::FValueTypeDescription TypeDesc = Shader::GetValueTypeDescription(ValueType);
-	return TypeDesc.NumComponents;
-}
-
-bool FType::Merge(const FType& OtherType)
-{
-	if (ValueType == Shader::EValueType::Void)
-	{
-		ValueType = OtherType.ValueType;
-		StructType = OtherType.StructType;
-		return true;
-	}
-
-	if (IsStruct() || OtherType.IsStruct())
-	{
-		return StructType == OtherType.StructType;
-	}
-
-	if (ValueType != OtherType.ValueType)
-	{
-		const Shader::FValueTypeDescription TypeDesc = Shader::GetValueTypeDescription(ValueType);
-		const Shader::FValueTypeDescription OtherTypeDesc = Shader::GetValueTypeDescription(OtherType);
-		if (TypeDesc.ComponentType != OtherTypeDesc.ComponentType)
-		{
-			return false;
-		}
-
-		const int8 NumComponents = FMath::Max(TypeDesc.NumComponents, OtherTypeDesc.NumComponents);
-		ValueType = Shader::MakeValueType(TypeDesc.ComponentType, NumComponents);
-	}
-	return true;
-}
-
-FStructFieldRef FStructType::FindFieldByName(const TCHAR* InName) const
-{
-	for (const FStructField& Field : Fields)
-	{
-		if (FCString::Strcmp(Field.Name, InName) == 0)
-		{
-			return FStructFieldRef(Field.Type, Field.ComponentIndex, Field.Type.GetNumComponents());
-		}
-	}
-
-	return FStructFieldRef();
-}
-
-FConstantValue::FConstantValue(const Shader::FValue& InValue) : Type(InValue.GetType())
-{
-	const Shader::FValueTypeDescription TypeDesc = Shader::GetValueTypeDescription(Type);
-	Component.Empty(TypeDesc.NumComponents);
-	for (int32 i = 0; i < TypeDesc.NumComponents; ++i)
-	{
-		Component.Add(InValue.Component[i]);
-	}
-}
-
 void FScope::Reset()
 {
 	State = EScopeState::Uninitialized;
@@ -450,7 +414,7 @@ FScope* FScope::FindSharedParent(FScope* Lhs, FScope* Rhs)
 	return Scope0;
 }
 
-void FExpressionLocalPHI::PrepareValue(FEmitContext& Context, const FRequestedType& RequestedType, FPrepareValueResult& OutResult)
+void FExpressionLocalPHI::PrepareValue(FEmitContext& Context, const FRequestedType& RequestedType, FPrepareValueResult& OutResult) const
 {
 	check(NumValues <= MaxNumPreviousScopes);
 	FExpression* ForwardExpression = Values[0];
@@ -472,7 +436,7 @@ void FExpressionLocalPHI::PrepareValue(FEmitContext& Context, const FRequestedTy
 	if (bForwardExpressionValid)
 	{
 		check(ForwardExpression);
-		return OutResult.SetForwardValue(Context, ForwardExpression, RequestedType);
+		return OutResult.SetForwardValue(Context, RequestedType, ForwardExpression);
 	}
 
 	// 2) PHI has different values in previous scopes, but possible some previous scopes may become dead due to constant folding
@@ -498,24 +462,25 @@ void FExpressionLocalPHI::PrepareValue(FEmitContext& Context, const FRequestedTy
 	if (bForwardExpressionValid)
 	{
 		check(ForwardExpression);
-		return OutResult.SetForwardValue(Context, ForwardExpression, RequestedType);
+		return OutResult.SetForwardValue(Context, RequestedType, ForwardExpression);
 	}
 
-	FType TypePerValue[MaxNumPreviousScopes];
+	FPreparedType TypePerValue[MaxNumPreviousScopes];
 	int32 NumValidTypes = 0;
-	FType CurrentType;
+	FPreparedType CurrentType;
 
 	auto UpdateValueTypes = [&]()
 	{
 		for (int32 i = 0; i < NumValues; ++i)
 		{
-			if (!TypePerValue[i] && PrepareScope(Context, Scopes[i]))
+			if (TypePerValue[i].IsVoid() && PrepareScope(Context, Scopes[i]))
 			{
-				const FPrepareValueResult ValueResult = PrepareExpressionValue(Context, Values[i], RequestedType);
-				if (ValueResult.Type)
+				const FPreparedType& ValueType = PrepareExpressionValue(Context, Values[i], RequestedType);
+				if (!ValueType.IsVoid())
 				{
-					TypePerValue[i] = ValueResult.Type;
-					if (!CurrentType.Merge(ValueResult.Type))
+					TypePerValue[i] = ValueType;
+					CurrentType = MergePreparedTypes(CurrentType, ValueType);
+					if (CurrentType.IsVoid())
 					{
 						return Context.Errors.AddError(this, TEXT("Type mismatch"));
 					}
@@ -530,7 +495,8 @@ void FExpressionLocalPHI::PrepareValue(FEmitContext& Context, const FRequestedTy
 	UpdateValueTypes();
 
 	// Assuming we have at least one value with a valid type, we use that to initialize our type
-	OutResult.SetType(Context, EExpressionEvaluationType::Shader, CurrentType);
+	CurrentType.SetEvaluationType(EExpressionEvaluationType::Shader); // TODO - No support for preshader flow control
+	OutResult.SetType(Context, RequestedType, CurrentType);
 
 	if (NumValidTypes < NumValues)
 	{
@@ -543,62 +509,66 @@ void FExpressionLocalPHI::PrepareValue(FEmitContext& Context, const FRequestedTy
 	}
 }
 
-void FExpressionLocalPHI::EmitValueShader(FEmitContext& Context, FShaderValue& OutShader) const
+void FExpressionLocalPHI::EmitValueShader(FEmitContext& Context, const FRequestedType& RequestedType, FShaderValue& OutShader) const
 {
-	const int32 LocalPHIIndex = Context.NumLocalPHIs++;
-	OutShader.Code.Appendf(TEXT("LocalPHI%d"), LocalPHIIndex);
-	OutShader.bInline = true;
-	OutShader.bHasDependencies = true;
-}
-
-void FExpressionLocalPHI::EmitShaderDependencies(FEmitContext& Context, const FShaderValue& Shader) const
-{
-	// Find the outermost scope to declare our local variable
-	FScope* DeclarationScope = Context.ScopeStack.Last();
-	for (int32 i = 0; i < NumValues; ++i)
+	int32 LocalPHIIndex = Context.LocalPHIs.Find(this);
+	if (LocalPHIIndex == INDEX_NONE)
 	{
-		DeclarationScope = FScope::FindSharedParent(DeclarationScope, Scopes[i]);
-		if (!DeclarationScope)
+		// This is the first time we've emitted shader code for this PHI
+		// First add it to the list, so if this is called recursively, this path will only be taken the first time
+		LocalPHIIndex = Context.LocalPHIs.Add(this);
+
+		// Find the outermost scope to declare our local variable
+		FScope* DeclarationScope = Context.ScopeStack.Last();
+		for (int32 i = 0; i < NumValues; ++i)
 		{
-			return Context.Errors.AddError(this, TEXT("Invalid LocalPHI"));
+			DeclarationScope = FScope::FindSharedParent(DeclarationScope, Scopes[i]);
+			if (!DeclarationScope)
+			{
+				return Context.Errors.AddError(this, TEXT("Invalid LocalPHI"));
+			}
 		}
-	}
 
-	const FType LocalType = GetType();
+		const FRequestedType LocalType = GetRequestedType();
 
-	bool bNeedToAddDeclaration = true;
-	for (int32 i = 0; i < NumValues; ++i)
-	{
-		FScope* ValueScope = Scopes[i];
-		check(IsScopeLive(ValueScope));
+		bool bNeedToAddDeclaration = true;
+		for (int32 i = 0; i < NumValues; ++i)
+		{
+			FScope* ValueScope = Scopes[i];
+			check(IsScopeLive(ValueScope));
 
-		Context.ScopeStack.Add(ValueScope);
-		const TCHAR* ValueCode = Values[i]->GetValueShader(Context, LocalType);
-		Context.ScopeStack.Pop();
+			Context.ScopeStack.Add(ValueScope);
+			const TCHAR* ValueCode = Values[i]->GetValueShader(Context, LocalType);
+			Context.ScopeStack.Pop();
 
-		if (ValueScope == DeclarationScope)
-		{
-			ValueScope->EmitDeclarationf(Context, TEXT("%s %s = %s;"),
-				LocalType.GetName(),
-				Shader.Code.ToString(),
-				ValueCode);
-			bNeedToAddDeclaration = false;
+			if (ValueScope == DeclarationScope)
+			{
+				ValueScope->EmitDeclarationf(Context, TEXT("%s LocalPHI%d = %s;"),
+					LocalType.GetName(),
+					LocalPHIIndex,
+					ValueCode);
+				bNeedToAddDeclaration = false;
+			}
+			else
+			{
+				ValueScope->EmitStatementf(Context, TEXT("LocalPHI%d = %s;"),
+					LocalPHIIndex,
+					ValueCode);
+			}
 		}
-		else
+
+		if (bNeedToAddDeclaration)
 		{
-			ValueScope->EmitStatementf(Context, TEXT("%s = %s;"),
-				Shader.Code.ToString(),
-				ValueCode);
+			check(IsScopeLive(DeclarationScope));
+			DeclarationScope->EmitDeclarationf(Context, TEXT("%s LocalPHI%d;"),
+				LocalType.GetName(),
+				LocalPHIIndex);
 		}
 	}
 
-	if (bNeedToAddDeclaration)
-	{
-		check(IsScopeLive(DeclarationScope));
-		DeclarationScope->EmitDeclarationf(Context, TEXT("%s %s;"),
-			LocalType.GetName(),
-			Shader.Code.ToString());
-	}
+	OutShader.Code.Appendf(TEXT("LocalPHI%d"), LocalPHIIndex);
+	OutShader.Type = GetType();
+	OutShader.bInline = true;
 }
 
 void FStatement::Reset()
@@ -608,17 +578,16 @@ void FStatement::Reset()
 
 void FExpression::Reset()
 {
-	ShaderValue = nullptr;
 	CurrentRequestedType.Reset();
 	PrepareValueResult = FPrepareValueResult();
 }
 
-const FPrepareValueResult& PrepareExpressionValue(FEmitContext& Context, FExpression* InExpression, const FRequestedType& RequestedType)
+const FPreparedType& PrepareExpressionValue(FEmitContext& Context, FExpression* InExpression, const FRequestedType& RequestedType)
 {
-	static FPrepareValueResult EmptyResult;
+	static FPreparedType VoidType;
 	if (!InExpression)
 	{
-		return EmptyResult;
+		return VoidType;
 	}
 
 	if (InExpression->bReentryFlag)
@@ -626,7 +595,7 @@ const FPrepareValueResult& PrepareExpressionValue(FEmitContext& Context, FExpres
 		// Valid for this to be called reentrantly
 		// Code should ensure that the type is set before the reentrant call, otherwise type will not be valid here
 		// LocalPHI nodes rely on this to break loops
-		return InExpression->PrepareValueResult;
+		return InExpression->PrepareValueResult.GetPreparedType();
 	}
 
 	bool bNeedToUpdateType = false;
@@ -635,10 +604,10 @@ const FPrepareValueResult& PrepareExpressionValue(FEmitContext& Context, FExpres
 		InExpression->CurrentRequestedType = RequestedType;
 		bNeedToUpdateType = !RequestedType.IsVoid();
 	}
-	else if(InExpression->CurrentRequestedType.StructType != RequestedType.StructType)
+	else if(InExpression->CurrentRequestedType.GetStructType() != RequestedType.GetStructType())
 	{
 		Context.Errors.AddError(InExpression, TEXT("Type mismatch"));
-		return EmptyResult;
+		return VoidType;
 	}
 	else
 	{
@@ -654,237 +623,473 @@ const FPrepareValueResult& PrepareExpressionValue(FEmitContext& Context, FExpres
 		}
 	}
 
-	FPrepareValueResult& Result = InExpression->PrepareValueResult;
 	if (bNeedToUpdateType)
 	{
 		check(!InExpression->CurrentRequestedType.IsVoid());
 
 		InExpression->bReentryFlag = true;
-		InExpression->PrepareValue(Context, InExpression->CurrentRequestedType, Result);
+		InExpression->PrepareValue(Context, InExpression->CurrentRequestedType, InExpression->PrepareValueResult);
 		InExpression->bReentryFlag = false;
 
-		if (Result.EvaluationType == EExpressionEvaluationType::None)
+		if (InExpression->PrepareValueResult.GetPreparedType().IsVoid())
 		{
 			// If we failed to assign a valid type, reset the requested type as well
 			// This ensures we'll try to compute a type again the next time we're called
 			InExpression->CurrentRequestedType.Reset();
 		}
-		else if (Result.EvaluationType == EExpressionEvaluationType::Constant &&
-			Result.Type != FType(Result.ConstantValue.GetType()))
-		{
-			Shader::FPreshaderData ConstantPreshader;
-			InExpression->bReentryFlag = true;
-			InExpression->EmitValuePreshader(Context, ConstantPreshader);
-			InExpression->bReentryFlag = false;
-
-			// Evaluate the constant preshader and store its value
-			FMaterialRenderContext RenderContext(nullptr, *Context.Material, nullptr);
-			Shader::FValue Value;
-			ConstantPreshader.Evaluate(nullptr, RenderContext, Value);
-			Result.ConstantValue = Shader::Cast(Value, Result.Type);
-		}
 	}
 
-	return Result;
+	return InExpression->PrepareValueResult.GetPreparedType();
 }
 
-void FPrepareValueResult::SetConstant(FEmitContext& Context, const Shader::FValue& InValue)
+FRequestedType::FRequestedType(const Shader::FType& InType, bool bDefaultRequest)
 {
-	EvaluationType = EExpressionEvaluationType::Constant;
-	Type = InValue.GetType();
-	ConstantValue = InValue;
-	ForwardValue = nullptr;
+	int32 NumComponents = 0;
+	if (InType.IsStruct())
+	{
+		StructType = InType.StructType;
+		NumComponents = InType.StructType->ComponentTypes.Num();
+	}
+	else
+	{
+		const Shader::FValueTypeDescription TypeDesc = Shader::GetValueTypeDescription(InType);
+		ValueComponentType = TypeDesc.ComponentType;
+		NumComponents = TypeDesc.NumComponents;
+	}
+	RequestedComponents.Init(bDefaultRequest, NumComponents);
 }
 
-void FPrepareValueResult::SetType(FEmitContext& Context, EExpressionEvaluationType InEvaluationType, const FType& InType)
+FRequestedType::FRequestedType(const Shader::EValueType& InType, bool bDefaultRequest)
 {
-	EvaluationType = InEvaluationType;
-	Type = InType;
-	ForwardValue = nullptr;
+	const Shader::FValueTypeDescription TypeDesc = Shader::GetValueTypeDescription(InType);
+	ValueComponentType = TypeDesc.ComponentType;
+	RequestedComponents.Init(bDefaultRequest, TypeDesc.NumComponents);
 }
 
-void FPrepareValueResult::SetForwardValue(FEmitContext& Context, FExpression* InForwardValue, const FRequestedType& RequestedType)
+const Shader::FType FRequestedType::GetType() const
 {
-	check(InForwardValue);
-	const FPrepareValueResult& ForwardResult = PrepareExpressionValue(Context, InForwardValue, RequestedType);
-	EvaluationType = ForwardResult.EvaluationType;
-	Type = ForwardResult.Type;
-	ConstantValue = ForwardResult.ConstantValue;
-	ForwardValue = InForwardValue;
+	if (IsStruct())
+	{
+		return StructType;
+	}
+	return Shader::MakeValueType(ValueComponentType, GetNumComponents());
 }
 
-void FExpression::EmitValueShader(FEmitContext& Context, FShaderValue& OutShader) const
+int32 FRequestedType::GetNumComponents() const
 {
-	check(false);
+	if (StructType)
+	{
+		return StructType->ComponentTypes.Num();
+	}
+	else
+	{
+		const int32 MaxComponentIndex = RequestedComponents.FindLast(true);
+		if (MaxComponentIndex != INDEX_NONE)
+		{
+			return MaxComponentIndex + 1;
+		}
+	}
+	return 0;
 }
 
-void FExpression::EmitShaderDependencies(FEmitContext& Context, const FShaderValue& Shader) const
+bool FRequestedType::Merge(const FRequestedType& OtherType)
 {
-	check(false);
+	if (StructType || OtherType.StructType && (StructType != OtherType.StructType))
+	{
+		return false;
+	}
+
+	StructType = OtherType.StructType;
+	if (!StructType)
+	{
+		ValueComponentType = Shader::CombineComponentTypes(ValueComponentType, OtherType.ValueComponentType);
+	}
+
+	for (int32 Index = 0; Index < OtherType.RequestedComponents.Num(); ++Index)
+	{
+		if (OtherType.IsComponentRequested(Index))
+		{
+			SetComponentRequested(Index);
+		}
+	}
+
+	return true;
 }
 
-void FExpression::EmitValuePreshader(FEmitContext& Context, Shader::FPreshaderData& OutPreshader) const
+FPreparedType::FPreparedType(const Shader::FType& InType)
 {
-	check(false);
+	if (InType.IsStruct())
+	{
+		StructType = InType.StructType;
+	}
+	else
+	{
+		ValueComponentType = Shader::GetValueTypeDescription(InType).ComponentType;
+	}
 }
 
-namespace Private
+int32 FPreparedType::GetNumComponents() const
 {
-void MoveToScope(FEmitShaderValue* ShaderValue, FScope* Scope)
+	if (StructType)
+	{
+		return StructType->ComponentTypes.Num();
+	}
+	else if (ValueComponentType != Shader::EValueComponentType::Void)
+	{
+		const int32 MaxComponentIndex = ComponentEvaluationType.FindLastByPredicate([](EExpressionEvaluationType InEvaluationType) { return InEvaluationType != EExpressionEvaluationType::None; });
+		if (MaxComponentIndex != INDEX_NONE)
+		{
+			return MaxComponentIndex + 1;
+		}
+	}
+	return 0;
+}
+
+bool FPreparedType::IsVoid() const
 {
-	if (ShaderValue->Scope != Scope)
+	return GetNumComponents() == 0;
+}
+
+Shader::FType FPreparedType::GetType() const
+{
+	if (IsStruct())
 	{
-		FScope* NewScope = FScope::FindSharedParent(ShaderValue->Scope, Scope);
-		check(NewScope);
+		return StructType;
+	}
+	return Shader::MakeValueType(ValueComponentType, GetNumComponents());
+}
 
-		ShaderValue->Scope = NewScope;
-		for (FEmitShaderValue* Dependency : ShaderValue->Dependencies)
+FRequestedType FPreparedType::GetRequestedType() const
+{
+	const int32 NumComponents = GetNumComponents();
+	FRequestedType Result;
+	if (NumComponents > 0)
+	{
+		if (StructType)
 		{
-			MoveToScope(Dependency, NewScope);
+			Result.StructType = StructType;
+		}
+		else
+		{
+			Result.ValueComponentType = ValueComponentType;
+		}
+		for (int32 Index = 0; Index < NumComponents; ++Index)
+		{
+			if (GetComponentEvaluationType(Index) != EExpressionEvaluationType::None)
+			{
+				Result.SetComponentRequested(Index);
+			}
 		}
 	}
+	return Result;
 }
+
+EExpressionEvaluationType FPreparedType::GetEvaluationType(const FRequestedType& RequestedType) const
+{
+	EExpressionEvaluationType Result = EExpressionEvaluationType::None;
+	for(int32 Index = 0; Index < ComponentEvaluationType.Num(); ++Index)
+	{
+		if (RequestedType.IsComponentRequested(Index))
+		{
+			Result = CombineEvaluationTypes(Result, ComponentEvaluationType[Index]);
+		}
+	}
+	return Result;
 }
 
-const TCHAR* FExpression::GetValueShader(FEmitContext& Context)
+void FPreparedType::SetEvaluationType(EExpressionEvaluationType EvaluationType)
 {
-	check(PrepareValueResult.EvaluationType != EExpressionEvaluationType::None);
-	if (PrepareValueResult.ForwardValue)
+	const int32 NumComponents = GetNumComponents();
+	for (int32 Index = 0; Index < NumComponents; ++Index)
 	{
-		return PrepareValueResult.ForwardValue->GetValueShader(Context);
+		if (ComponentEvaluationType[Index] != EExpressionEvaluationType::None)
+		{
+			ComponentEvaluationType[Index] = EvaluationType;
+		}
 	}
+}
 
-	FScope* CurrentScope = Context.ScopeStack.Last();
-	check(IsScopeLive(CurrentScope));
+void FPreparedType::SetField(const Shader::FStructField* Field, const FPreparedType& FieldType)
+{
+	for (int32 Index = 0; Index < Field->GetNumComponents(); ++Index)
+	{
+		SetComponentEvaluationType(Field->ComponentIndex + Index, FieldType.GetComponentEvaluationType(Index));
+	}
+}
+
+FPreparedType FPreparedType::GetFieldType(const Shader::FStructField* Field) const
+{
+	FPreparedType Result(Field->Type);
+	for (int32 Index = 0; Index < Field->GetNumComponents(); ++Index)
+	{
+		Result.SetComponentEvaluationType(Index, GetComponentEvaluationType(Field->ComponentIndex + Index));
+	}
+	return Result;
+}
 
-	if (!ShaderValue)
+void FPreparedType::SetComponentEvaluationType(int32 Index, EExpressionEvaluationType EvaluationType)
+{
+	if (Index >= ComponentEvaluationType.Num())
 	{
-		check(!bReentryFlag);
+		static_assert((uint8)EExpressionEvaluationType::None == 0u, "Assume AddZeroed() will initialize to None");
+		ComponentEvaluationType.AddZeroed(Index + 1 - ComponentEvaluationType.Num());
+	}
+	ComponentEvaluationType[Index] = EvaluationType;
+}
 
-		TStringBuilder<1024> FormattedCode;
-		if (PrepareValueResult.EvaluationType == EExpressionEvaluationType::Shader)
+FPreparedType MergePreparedTypes(const FPreparedType& Lhs, const FPreparedType& Rhs)
+{
+	// If one type is not initialized yet, just use the other type
+	if (!Lhs.IsInitialized())
+	{
+		return Rhs;
+	}
+	else if (!Rhs.IsInitialized())
+	{
+		return Lhs;
+	}
+
+	int32 NumComponents = 0;
+	FPreparedType Result;
+	if (Lhs.IsStruct() || Rhs.IsStruct())
+	{
+		if (Lhs.StructType != Rhs.StructType)
 		{
-			FShaderValue ShaderResult(FormattedCode);
-			FEmitShaderValueDependencies Dependencies;
-			{
-				FEmitShaderValueContext& ValueContext = Context.ShaderValueStack.AddDefaulted_GetRef();
+			// Mismatched structs
+			return Result;
+		}
+		Result.StructType = Lhs.StructType;
+		NumComponents = Result.StructType->ComponentTypes.Num();
+	}
+	else
+	{
+		Result.ValueComponentType = Shader::CombineComponentTypes(Lhs.ValueComponentType, Rhs.ValueComponentType);
+		NumComponents = FMath::Max(Lhs.GetNumComponents(), Rhs.GetNumComponents());
+	}
 
-				bReentryFlag = true;
-				EmitValueShader(Context, ShaderResult);
-				bReentryFlag = false;
+	for (int32 Index = 0; Index < NumComponents; ++Index)
+	{
+		const EExpressionEvaluationType LhsEvaluation = Lhs.GetComponentEvaluationType(Index);
+		const EExpressionEvaluationType RhsEvaluation = Rhs.GetComponentEvaluationType(Index);
+		Result.SetComponentEvaluationType(Index, CombineEvaluationTypes(LhsEvaluation, RhsEvaluation));
+	}
 
-				Dependencies = MoveTemp(ValueContext.Dependencies);
-				Context.ShaderValueStack.Pop(false);
-			}
+	return Result;
+}
 
-			if (ShaderResult.bInline)
-			{
-				ShaderValue = new(*Context.Allocator) FEmitShaderValue(this, CurrentScope);
-				ShaderValue->Reference = AllocateString(*Context.Allocator, FormattedCode);
-				ShaderValue->Dependencies = AllocateArrayView(*Context.Allocator, Dependencies);
-			}
-			else
+bool FPrepareValueResult::TryMergePreparedType(FEmitContext& Context, const Shader::FStructType* StructType, Shader::EValueComponentType ComponentType)
+{
+	if (!PreparedType.IsInitialized())
+	{
+		PreparedType.ValueComponentType = ComponentType;
+		PreparedType.StructType = StructType;
+		return true;
+	}
+
+	if (StructType)
+	{
+		check(ComponentType == Shader::EValueComponentType::Void);
+		if (StructType != PreparedType.StructType)
+		{
+			Context.Errors.AddError(nullptr, TEXT("Invalid type"));
+			return false;
+		}
+	}
+	else
+	{
+		check(ComponentType != Shader::EValueComponentType::Void);
+		PreparedType.ValueComponentType = Shader::CombineComponentTypes(PreparedType.ValueComponentType, ComponentType);
+	}
+
+	return true;
+}
+
+void FPrepareValueResult::SetType(FEmitContext& Context, const FRequestedType& RequestedType, EExpressionEvaluationType EvaluationType, const Shader::FType& Type)
+{
+	if (TryMergePreparedType(Context, Type.StructType, Shader::GetValueTypeDescription(Type.ValueType).ComponentType))
+	{
+		if (EvaluationType != EExpressionEvaluationType::None)
+		{
+			const int32 NumComponents = Type.GetNumComponents();
+			for (int32 Index = 0; Index < NumComponents; ++Index)
 			{
-				// Check to see if we've already generated code for an equivalent expression
-				const FSHAHash ShaderHash = HashString(FormattedCode);
-				FEmitShaderValue** const PrevShaderValue = Context.ShaderValueMap.Find(ShaderHash);
-				if (PrevShaderValue)
+				if (RequestedType.IsComponentRequested(Index))
 				{
-					ShaderValue = *PrevShaderValue;
+					PreparedType.SetComponentEvaluationType(Index, EvaluationType);
 				}
-				else
+			}
+		}
+	}
+}
+
+void FPrepareValueResult::SetType(FEmitContext& Context, const FRequestedType& RequestedType, EExpressionEvaluationType EvaluationType, Shader::EValueComponentType ComponentType)
+{
+	if (TryMergePreparedType(Context, nullptr, ComponentType))
+	{
+		if (EvaluationType != EExpressionEvaluationType::None)
+		{
+			const int32 NumComponents = RequestedType.GetNumComponents();
+			for (int32 Index = 0; Index < NumComponents; ++Index)
+			{
+				if (RequestedType.IsComponentRequested(Index))
 				{
-					ShaderValue = new(*Context.Allocator) FEmitShaderValue(this, CurrentScope);
-					ShaderValue->Hash = ShaderHash;
-					ShaderValue->Reference = Context.AcquireLocalDeclarationCode();
-					ShaderValue->Value = AllocateString(*Context.Allocator, FormattedCode);
-					ShaderValue->Dependencies = AllocateArrayView(*Context.Allocator, Dependencies);
-					Context.ShaderValueMap.Add(ShaderHash, ShaderValue);
+					PreparedType.SetComponentEvaluationType(Index, EvaluationType);
 				}
 			}
+		}
+	}
+}
 
-			check(ShaderValue);
-			if (ShaderResult.bHasDependencies)
+void FPrepareValueResult::SetType(FEmitContext& Context, const FRequestedType& RequestedType, const FPreparedType& Type)
+{
+	if (TryMergePreparedType(Context, Type.StructType, Type.ValueComponentType))
+	{
+		const int32 NumComponents = RequestedType.GetNumComponents();
+		for (int32 Index = 0; Index < NumComponents; ++Index)
+		{
+			if (RequestedType.IsComponentRequested(Index))
 			{
-				bReentryFlag = true;
-				EmitShaderDependencies(Context, ShaderResult);
-				bReentryFlag = false;
+				PreparedType.SetComponentEvaluationType(Index, Type.GetComponentEvaluationType(Index));
 			}
 		}
-		else if (PrepareValueResult.EvaluationType == EExpressionEvaluationType::Preshader)
+	}
+}
+
+void FPrepareValueResult::SetForwardValue(FEmitContext& Context, const FRequestedType& RequestedType, FExpression* InForwardValue)
+{
+	check(InForwardValue);
+	if (InForwardValue != ForwardValue)
+	{
+		PreparedType = PrepareExpressionValue(Context, InForwardValue, RequestedType);
+		ForwardValue = InForwardValue;
+	}
+}
+
+void FExpression::EmitValueShader(FEmitContext& Context, const FRequestedType& RequestedType, FShaderValue& OutShader) const
+{
+	check(false);
+}
+
+void FExpression::EmitValuePreshader(FEmitContext& Context, const FRequestedType& RequestedType, Shader::FPreshaderData& OutPreshader) const
+{
+	check(false);
+}
+
+const TCHAR* FExpression::GetValueShader(FEmitContext& Context, const FRequestedType& RequestedType)
+{
+	if (PrepareValueResult.ForwardValue)
+	{
+		return PrepareValueResult.ForwardValue->GetValueShader(Context, RequestedType);
+	}
+
+	const EExpressionEvaluationType EvaluationType = PrepareValueResult.PreparedType.GetEvaluationType(RequestedType);
+	check(EvaluationType != EExpressionEvaluationType::None);
+
+	FScope* CurrentScope = Context.ScopeStack.Last();
+	check(IsScopeLive(CurrentScope));
+
+	FEmitShaderValue* ShaderValue = nullptr;
+	if (EvaluationType == EExpressionEvaluationType::Preshader)
+	{
+		ShaderValue = Context.AcquirePreshader(RequestedType, CurrentScope, this);
+	}
+	else
+	{
+		TStringBuilder<1024> FormattedCode;
+		FShaderValue ShaderResult(FormattedCode);
+		FEmitShaderValueDependencies Dependencies;
+
+		if (EvaluationType == EExpressionEvaluationType::Shader)
 		{
-			Shader::FPreshaderData Preshader;
+			FEmitShaderValueContext& ValueContext = Context.ShaderValueStack.AddDefaulted_GetRef();
 
-			bReentryFlag = true;
-			EmitValuePreshader(Context, Preshader);
-			bReentryFlag = false;
+			EmitValueShader(Context, RequestedType, ShaderResult);
+			check(!ShaderResult.Type.IsVoid());
 
-			ShaderValue = new(*Context.Allocator) FEmitShaderValue(this, CurrentScope);
-			ShaderValue->Reference = Context.AcquirePreshader(PrepareValueResult.Type, Preshader);
+			Dependencies = MoveTemp(ValueContext.Dependencies);
+			Context.ShaderValueStack.Pop(false);
 		}
 		else
 		{
-			check(PrepareValueResult.EvaluationType == EExpressionEvaluationType::Constant);
-			PrepareValueResult.ConstantValue.ToString(Shader::EValueStringFormat::HLSL, FormattedCode);
-
-			ShaderValue = new(*Context.Allocator) FEmitShaderValue(this, CurrentScope);
-			ShaderValue->Reference = AllocateString(*Context.Allocator, FormattedCode);
+			check(EvaluationType == EExpressionEvaluationType::Constant);
+			const Shader::FValue ConstantValue = GetValueConstant(Context, RequestedType);
+			ShaderResult.Type = ConstantValue.Type;
+			ShaderResult.bInline = !ConstantValue.Type.IsStruct(); // 'struct' constants can't be inline, due to HLSL syntax limitations
+			ConstantValue.ToString(Shader::EValueStringFormat::HLSL, ShaderResult.Code);
 		}
+
+		ShaderValue = Context.AcquireShader(CurrentScope, ShaderResult, Dependencies);
 	}
 
 	check(ShaderValue);
-	Private::MoveToScope(ShaderValue, CurrentScope);
+	ShaderValue = Context.CastShaderValue(this, CurrentScope, ShaderValue, RequestedType.GetType());
 
 	if (Context.ShaderValueStack.Num() > 0)
 	{
 		Context.ShaderValueStack.Last().Dependencies.Add(ShaderValue);
 	}
+
 	return ShaderValue->Reference;
 }
 
-const TCHAR* FExpression::GetValueShader(FEmitContext& Context, const FType& InType)
+const TCHAR* FExpression::GetValueShader(FEmitContext& Context)
 {
-	if (PrepareValueResult.EvaluationType == EExpressionEvaluationType::Constant && PrepareValueResult.Type != InType)
-	{
-		// If we need to cast a constant value, perform the cast first, then convert the result to HLSL
-		const Shader::FValue ConstantValue = Shader::Cast(PrepareValueResult.ConstantValue, InType);
+	return GetValueShader(Context, GetRequestedType());
+}
 
-		TStringBuilder<1024> FormattedCode;
-		ConstantValue.ToString(Shader::EValueStringFormat::HLSL, FormattedCode);
-		return AllocateString(*Context.Allocator, FormattedCode);
+void FExpression::GetValuePreshader(FEmitContext& Context, const FRequestedType& RequestedType, Shader::FPreshaderData& OutPreshader)
+{
+	if (PrepareValueResult.ForwardValue)
+	{
+		return PrepareValueResult.ForwardValue->GetValuePreshader(Context, RequestedType, OutPreshader);
 	}
 
-	const ECastFlags CastFlags = ECastFlags::ReplicateScalar | ECastFlags::AllowAppendZeroes | ECastFlags::AllowTruncate;
-	return Context.CastShaderValue(this, GetValueShader(Context), PrepareValueResult.Type, InType, CastFlags);
+	check(!bReentryFlag);
+	const EExpressionEvaluationType EvaluationType = PrepareValueResult.PreparedType.GetEvaluationType(RequestedType);
+	if (EvaluationType == EExpressionEvaluationType::Preshader)
+	{
+		bReentryFlag = true;
+		EmitValuePreshader(Context, RequestedType, OutPreshader);
+		bReentryFlag = false;
+	}
+	else
+	{
+		check(EvaluationType == EExpressionEvaluationType::Constant);
+		const Shader::FValue ConstantValue = GetValueConstant(Context, RequestedType);
+		OutPreshader.WriteOpcode(Shader::EPreshaderOpcode::Constant).Write(ConstantValue);
+	}
 }
 
-void FExpression::GetValuePreshader(FEmitContext& Context, Shader::FPreshaderData& OutPreshader)
+Shader::FValue FExpression::GetValueConstant(FEmitContext& Context, const FRequestedType& RequestedType)
 {
-	switch (PrepareValueResult.EvaluationType)
+	if (PrepareValueResult.ForwardValue)
 	{
-	case EExpressionEvaluationType::Preshader:
-		if (PrepareValueResult.ForwardValue)
-		{
-			PrepareValueResult.ForwardValue->EmitValuePreshader(Context, OutPreshader);
-		}
-		else
-		{
-			EmitValuePreshader(Context, OutPreshader);
-		}
-		break;
-	case EExpressionEvaluationType::Constant:
-		OutPreshader.WriteOpcode(Shader::EPreshaderOpcode::Constant).Write(PrepareValueResult.ConstantValue);
-		break;
-	default:
-		check(false);
-		break;
+		return PrepareValueResult.ForwardValue->GetValueConstant(Context, RequestedType);
 	}
-}
 
-Shader::FValue FExpression::GetValueConstant(FEmitContext& Context)
-{
-	check(PrepareValueResult.EvaluationType == EExpressionEvaluationType::Constant);
-	return PrepareValueResult.ConstantValue;
+	check(!bReentryFlag);
+	check(PrepareValueResult.PreparedType.GetEvaluationType(RequestedType) == EExpressionEvaluationType::Constant);
+	
+	Shader::FPreshaderData ConstantPreshader;
+	bReentryFlag = true;
+	EmitValuePreshader(Context, RequestedType, ConstantPreshader);
+	bReentryFlag = false;
+
+	// Evaluate the constant preshader and store its value
+	Shader::FPreshaderStack Stack;
+	FMaterialRenderContext RenderContext(nullptr, *Context.Material, nullptr);
+
+	const Shader::FPreshaderValue PreshaderValue = ConstantPreshader.Evaluate(nullptr, RenderContext, Stack);
+	Shader::FValue Result = PreshaderValue.AsShaderValue(Context.TypeRegistry);
+
+	const Shader::FType RequestedConstantType = RequestedType.GetType();
+	if (Result.Type.IsNumeric() && RequestedConstantType.IsNumeric())
+	{
+		Result = Shader::Cast(Result, RequestedConstantType.ValueType);
+	}
+
+	check(Result.Type == RequestedConstantType);
+	return Result;
 }
 
 bool FScope::HasParentScope(const FScope& InParentScope) const
@@ -1073,24 +1278,6 @@ void FScope::WriteHLSL(int32 Indent, FStringBuilderBase& OutString) const
 	}
 }
 
-void FStructType::WriteHLSL(FStringBuilderBase& OutString) const
-{
-	OutString.Appendf(TEXT("struct %s\n"), Name);
-	OutString.Append(TEXT("{\n"));
-	for (const FStructField& Field : Fields)
-	{
-		OutString.Appendf(TEXT("\t%s %s;\n"), Field.Type.GetName(), Field.Name);
-	}
-	OutString.Append(TEXT("};\n"));
-
-	for (const FStructField& Field : Fields)
-	{
-		OutString.Appendf(TEXT("%s %s_Set%s(%s Self, %s Value) { Self.%s = Value; return Self; }\n"),
-			Name, Name, Field.Name, Name, Field.Type.GetName(), Field.Name);
-	}
-	OutString.Append(TEXT("\n"));
-}
-
 void FTree::ResetNodes()
 {
 	FNode* Node = Nodes;
@@ -1102,16 +1289,6 @@ void FTree::ResetNodes()
 	}
 }
 
-void FTree::EmitDeclarationsCode(FStringBuilderBase& OutCode) const
-{
-	const FStructType* StructType = StructTypes;
-	while (StructType)
-	{
-		StructType->WriteHLSL(OutCode);
-		StructType = StructType->NextType;
-	}
-}
-
 bool FTree::EmitShader(FEmitContext& Context, FStringBuilderBase& OutCode) const
 {
 	if (RootScope->ContainedStatement)
@@ -1169,59 +1346,5 @@ FTextureParameterDeclaration* FTree::NewTextureParameterDeclaration(const FName&
 	return Declaration;
 }
 
-namespace Private
-{
-void SetComponentTypes(Shader::EValueComponentType* ComponentTypes, int32 ComponentIndex, const HLSLTree::FType& Type)
-{
-	if (Type.IsStruct())
-	{
-		for (const FStructField& Field : Type.StructType->Fields)
-		{
-			SetComponentTypes(ComponentTypes, ComponentIndex + Field.ComponentIndex, Field.Type);
-		}
-	}
-	else
-	{
-		const Shader::FValueTypeDescription TypeDesc = Shader::GetValueTypeDescription(Type);
-		for (int32 i = 0; i < TypeDesc.NumComponents; ++i)
-		{
-			ComponentTypes[ComponentIndex + i] = TypeDesc.ComponentType;
-		}
-	}
-}
-} // namespace Private
-
-const FStructType* FTree::NewStructType(const FStructTypeInitializer& Initializer)
-{
-	const int32 NumFields = Initializer.Fields.Num();
-	int32 ComponentIndex = 0;
-	FStructField* Fields = new(*Allocator) FStructField[NumFields];
-	for (int32 FieldIndex = 0; FieldIndex < NumFields; ++FieldIndex)
-	{
-		const FStructFieldInitializer& FieldInitializer = Initializer.Fields[FieldIndex];
-		const FType& FieldType = FieldInitializer.Type;
-		FStructField& Field = Fields[FieldIndex];
-		Field.Name = AllocateString(*Allocator, FieldInitializer.Name);
-		Field.Type = FieldType;
-		Field.ComponentIndex = ComponentIndex;
-		ComponentIndex += FieldType.GetNumComponents();
-	}
-
-	Shader::EValueComponentType* ComponentTypes = new(*Allocator) Shader::EValueComponentType[ComponentIndex];
-	for (int32 FieldIndex = 0; FieldIndex < NumFields; ++FieldIndex)
-	{
-		const FStructField& Field = Fields[FieldIndex];
-		Private::SetComponentTypes(ComponentTypes, Field.ComponentIndex, Field.Type);
-	}
-
-	FStructType* StructType = NewNode<FStructType>();
-	StructType->Name = AllocateString(*Allocator, Initializer.Name);
-	StructType->Fields = MakeArrayView(Fields, NumFields);
-	StructType->ComponentTypes = MakeArrayView(ComponentTypes, ComponentIndex);
-	StructType->NextType = StructTypes;
-	StructTypes = StructType;
-	return StructType;
-}
-
 } // namespace HLSLTree
 } // namespace UE
\ No newline at end of file
diff --git a/Engine/Source/Runtime/Engine/Private/HLSLTree/HLSLTreeCommon.cpp b/Engine/Source/Runtime/Engine/Private/HLSLTree/HLSLTreeCommon.cpp
--- a/Engine/Source/Runtime/Engine/Private/HLSLTree/HLSLTreeCommon.cpp
+++ b/Engine/Source/Runtime/Engine/Private/HLSLTree/HLSLTreeCommon.cpp
@@ -62,6 +62,23 @@ FSwizzleParameters::FSwizzleParameters(int8 InR, int8 InG, int8 InB, int8 InA) :
 	++NumComponents;
 }
 
+FRequestedType FSwizzleParameters::GetRequestedInputType(const FRequestedType& RequestedType) const
+{
+	FRequestedType RequestedInputType;
+	RequestedInputType.ValueComponentType = RequestedType.ValueComponentType;
+
+	for (int32 Index = 0; Index < NumComponents; ++Index)
+	{
+		if (RequestedType.IsComponentRequested(Index))
+		{
+			const int32 SwizzledComponentIndex = ComponentIndex[Index];
+			RequestedInputType.SetComponentRequested(SwizzledComponentIndex);
+		}
+	}
+	
+	return RequestedInputType;
+}
+
 FSwizzleParameters MakeSwizzleMask(bool bInR, bool bInG, bool bInB, bool bInA)
 {
 	int8 ComponentIndex[4] = { INDEX_NONE, INDEX_NONE, INDEX_NONE, INDEX_NONE };
@@ -85,57 +102,23 @@ FSwizzleParameters MakeSwizzleMask(bool bInR, bool bInG, bool bInB, bool bInA)
 	return FSwizzleParameters(ComponentIndex[0], ComponentIndex[1], ComponentIndex[2], ComponentIndex[3]);
 }
 
-void FExpressionConstant::PrepareValue(FEmitContext& Context, const FRequestedType& RequestedType, FPrepareValueResult& OutResult)
+void FExpressionConstant::PrepareValue(FEmitContext& Context, const FRequestedType& RequestedType, FPrepareValueResult& OutResult) const
 {
-	if (Value.Type.IsStruct())
-	{
-		if (Value.Type.StructType != RequestedType.GetStructType())
-		{
-			return Context.Errors.AddError(this, TEXT("Type mismatch"));
-		}
-
-		// TODO - Support for preshader structs?
-		OutResult.SetType(Context, EExpressionEvaluationType::Shader, Value.Type);
-	}
-	else
-	{
-		OutResult.SetType(Context, EExpressionEvaluationType::Constant, Shader::MakeValueTypeWithRequestedNumComponents(Value.Type, RequestedType.GetRequestedNumComponents()));
-	}
-}
-
-void FExpressionConstant::EmitValueShader(FEmitContext& Context, FShaderValue& OutShader) const
-{
-	const TArrayView<const Shader::EValueComponentType> ComponentTypes = Value.Type.StructType->ComponentTypes;
-	OutShader.Code.Append(TEXT("{ "));
-	for (int32 ComponentIndex = 0; ComponentIndex < Value.Component.Num(); ++ComponentIndex)
-	{
-		if (ComponentIndex > 0)
-		{
-			OutShader.Code.Append(TEXT(", "));
-		}
-		Value.Component[ComponentIndex].ToString(ComponentTypes[ComponentIndex], OutShader.Code);
-	}
-	OutShader.Code.Append(TEXT(" }"));
+	OutResult.SetType(Context, RequestedType, EExpressionEvaluationType::Constant, Value.Type);
 }
 
-void FExpressionConstant::EmitValuePreshader(FEmitContext& Context, Shader::FPreshaderData& OutPreshader) const
+void FExpressionConstant::EmitValuePreshader(FEmitContext& Context, const FRequestedType& RequestedType, Shader::FPreshaderData& OutPreshader) const
 {
-	Shader::FValue BaseValue(Value.Type.ValueType);
-	for (int32 ComponentIndex = 0; ComponentIndex < Value.Component.Num(); ++ComponentIndex)
-	{
-		BaseValue.Component[ComponentIndex] = Value.Component[ComponentIndex];
-	}
-	OutPreshader.WriteOpcode(Shader::EPreshaderOpcode::Constant).Write(BaseValue);
+	OutPreshader.WriteOpcode(Shader::EPreshaderOpcode::Constant).Write(Value);
 }
 
-void FExpressionMaterialParameter::PrepareValue(FEmitContext& Context, const FRequestedType& RequestedType, FPrepareValueResult& OutResult)
+void FExpressionMaterialParameter::PrepareValue(FEmitContext& Context, const FRequestedType& RequestedType, FPrepareValueResult& OutResult) const
 {
-	const Shader::EValueType CurrentType = Shader::MakeValueTypeWithRequestedNumComponents(GetShaderValueType(ParameterType), RequestedType.GetRequestedNumComponents());
 	const EExpressionEvaluationType EvaluationType = IsStaticMaterialParameter(ParameterType) ? EExpressionEvaluationType::Constant : EExpressionEvaluationType::Preshader;
-	OutResult.SetType(Context, EvaluationType, CurrentType);
+	OutResult.SetType(Context, RequestedType, EvaluationType, GetShaderValueType(ParameterType));
 }
 
-void FExpressionMaterialParameter::EmitValuePreshader(FEmitContext& Context, Shader::FPreshaderData& OutPreshader) const
+void FExpressionMaterialParameter::EmitValuePreshader(FEmitContext& Context, const FRequestedType& RequestedType, Shader::FPreshaderData& OutPreshader) const
 {
 	if (ParameterType == EMaterialParameterType::StaticSwitch)
 	{
@@ -170,11 +153,12 @@ void FExpressionMaterialParameter::EmitValuePreshader(FEmitContext& Context, Sha
 	}
 }
 
-void FExpressionExternalInput::EmitValueShader(FEmitContext& Context, FShaderValue& OutShader) const
+void FExpressionExternalInput::EmitValueShader(FEmitContext& Context, const FRequestedType& RequestedType, FShaderValue& OutShader) const
 {
 	const int32 TypeIndex = (int32)InputType;
 	const int32 TexCoordIndex = TypeIndex - (int32)EExternalInputType::TexCoord0;
 
+	OutShader.Type = GetInputExpressionType(InputType);
 	if (TexCoordIndex >= 0 && TexCoordIndex < 8)
 	{
 		Context.NumTexCoords = FMath::Max(Context.NumTexCoords, TexCoordIndex + 1);
@@ -187,13 +171,13 @@ void FExpressionExternalInput::EmitValueShader(FEmitContext& Context, FShaderVal
 	}
 }
 
-void FExpressionTextureSample::PrepareValue(FEmitContext& Context, const FRequestedType& RequestedType, FPrepareValueResult& OutResult)
+void FExpressionTextureSample::PrepareValue(FEmitContext& Context, const FRequestedType& RequestedType, FPrepareValueResult& OutResult) const
 {
 	PrepareExpressionValue(Context, TexCoordExpression, 2);
-	OutResult.SetType(Context, EExpressionEvaluationType::Shader, Shader::EValueType::Float4);
+	OutResult.SetType(Context, RequestedType, EExpressionEvaluationType::Shader, Shader::EValueType::Float4);
 }
 
-void FExpressionTextureSample::EmitValueShader(FEmitContext& Context, FShaderValue& OutShader) const
+void FExpressionTextureSample::EmitValueShader(FEmitContext& Context, const FRequestedType& RequestedType, FShaderValue& OutShader) const
 {
 	const FTextureDescription& Desc = Declaration->Description;
 	const EMaterialValueType MaterialValueType = Desc.Texture->GetMaterialType();
@@ -319,6 +303,7 @@ void FExpressionTextureSample::EmitValueShader(FEmitContext& Context, FShaderVal
 		}
 	}
 
+	OutShader.Type = Shader::EValueType::Float4;
 	OutShader.Code.Appendf(TEXT("%s(%s(%s, %s, %s))"),
 		SamplerTypeFunction,
 		SampleFunctionName,
@@ -327,171 +312,225 @@ void FExpressionTextureSample::EmitValueShader(FEmitContext& Context, FShaderVal
 		TexCoordExpression->GetValueShader(Context, Shader::EValueType::Float2));
 }
 
-void FExpressionGetStructField::PrepareValue(FEmitContext& Context, const FRequestedType& RequestedType, FPrepareValueResult& OutResult)
+void FExpressionGetStructField::PrepareValue(FEmitContext& Context, const FRequestedType& RequestedType, FPrepareValueResult& OutResult) const
 {
-	const FStructFieldRef FieldRef = StructType->FindFieldByName(FieldName);
-	if (!FieldRef)
-	{
-		return Context.Errors.AddErrorf(this, TEXT("Invalid field %s"), FieldName);
-	}
-
 	FRequestedType RequestedStructType(StructType, false);
-	RequestedStructType.SetField(FieldRef, RequestedType);
+	RequestedStructType.SetField(Field, RequestedType);
 
-	const FPrepareValueResult& StructResult = PrepareExpressionValue(Context, StructExpression, RequestedStructType);
-	if (StructResult.IsValid() && StructResult.Type != FType(StructType))
+	const FPreparedType& StructPreparedType = PrepareExpressionValue(Context, StructExpression, RequestedStructType);
+	if (!StructPreparedType.IsVoid() && StructPreparedType.StructType != StructType)
 	{
 		return Context.Errors.AddErrorf(this, TEXT("Expected type %s"), StructType->Name);
 	}
 
-	// TODO - preshader structs
-	OutResult.SetType(Context, EExpressionEvaluationType::Shader, FieldRef.Type);
+	OutResult.SetType(Context, RequestedType, StructPreparedType.GetFieldType(Field));
 }
 
-void FExpressionGetStructField::EmitValueShader(FEmitContext& Context, FShaderValue& OutShader) const
+void FExpressionGetStructField::EmitValueShader(FEmitContext& Context, const FRequestedType& RequestedType, FShaderValue& OutShader) const
 {
+	FRequestedType RequestedStructType(StructType, false);
+	RequestedStructType.SetField(Field, RequestedType);
+
 	OutShader.bInline = true;
+	OutShader.Type = Field->Type;
 	OutShader.Code.Appendf(TEXT("%s.%s"),
-		StructExpression->GetValueShader(Context, StructType),
-		FieldName);
+		StructExpression->GetValueShader(Context, RequestedStructType),
+		Field->Name);
 }
 
-void FExpressionSetStructField::PrepareValue(FEmitContext& Context, const FRequestedType& RequestedType, FPrepareValueResult& OutResult)
+void FExpressionGetStructField::EmitValuePreshader(FEmitContext& Context, const FRequestedType& RequestedType, Shader::FPreshaderData& OutPreshader) const
 {
-	const FStructFieldRef FieldRef = StructType->FindFieldByName(FieldName);
-	if (!FieldRef)
-	{
-		return Context.Errors.AddErrorf(this, TEXT("Invalid field %s"), FieldName);
-	}
+	FRequestedType RequestedStructType(StructType, false);
+	RequestedStructType.SetField(Field, RequestedType);
 
+	StructExpression->GetValuePreshader(Context, RequestedStructType, OutPreshader);
+	OutPreshader.WriteOpcode(Shader::EPreshaderOpcode::GetField).Write(Field->Type).Write(Field->ComponentIndex);
+}
+
+void FExpressionSetStructField::PrepareValue(FEmitContext& Context, const FRequestedType& RequestedType, FPrepareValueResult& OutResult) const
+{
 	FRequestedType RequestedStructType(RequestedType);
-	RequestedStructType.ClearFieldRequested(FieldRef);
+	RequestedStructType.ClearFieldRequested(Field);
 
-	const FPrepareValueResult& StructResult = PrepareExpressionValue(Context, StructExpression, RequestedStructType);
-	if (StructResult.IsValid() && StructResult.Type != FType(StructType))
+	const FPreparedType& StructPreparedType = PrepareExpressionValue(Context, StructExpression, RequestedStructType);
+	if (!StructPreparedType.IsVoid() && StructPreparedType.StructType != StructType)
 	{
 		return Context.Errors.AddErrorf(this, TEXT("Expected type %s"), StructType->Name);
 	}
 
-	const FPrepareValueResult& FieldResult = PrepareExpressionValue(Context, FieldExpression, RequestedType.GetField(FieldRef));
-	if (!FieldResult.IsValid())
+	const FPreparedType& FieldPreparedType = PrepareExpressionValue(Context, FieldExpression, RequestedType.GetField(Field));
+	if (FieldPreparedType.IsVoid())
 	{
-		OutResult.SetForwardValue(Context, StructExpression, RequestedType);
+		OutResult.SetForwardValue(Context, RequestedType, StructExpression);
 	}
 	else
 	{
-		// TODO - preshader structs
-		OutResult.SetType(Context, EExpressionEvaluationType::Shader, StructType);
+		FPreparedType ResultType(StructPreparedType);
+		if (ResultType.IsVoid())
+		{
+			ResultType = StructType;
+		}
+		ResultType.SetField(Field, FieldPreparedType);
+		OutResult.SetType(Context, RequestedType, ResultType);
 	}
 }
 
-void FExpressionSetStructField::EmitValueShader(FEmitContext& Context, FShaderValue& OutShader) const
+void FExpressionSetStructField::EmitValueShader(FEmitContext& Context, const FRequestedType& RequestedType, FShaderValue& OutShader) const
 {
-	const FStructFieldRef FieldRef = StructType->FindFieldByName(FieldName);
-	check(FieldRef);
-	if (StructExpression->GetEvaluationType() != EExpressionEvaluationType::None)
-	{
-		OutShader.Code.Appendf(TEXT("%s_Set%s(%s, %s)"),
-			StructType->Name,
-			FieldName,
-			StructExpression->GetValueShader(Context, StructType),
-			FieldExpression->GetValueShader(Context, FieldRef.Type));
-	}
-	else
+	FRequestedType RequestedStructType(RequestedType);
+	RequestedStructType.ClearFieldRequested(Field);
+	const EExpressionEvaluationType StructEvaluationType = StructExpression->GetEvaluationType(RequestedStructType);
+
+	const FRequestedType RequestedFieldType = RequestedType.GetField(Field);
+	const EExpressionEvaluationType FieldEvaluationType = FieldExpression->GetEvaluationType(RequestedFieldType);
+
+	OutShader.Type = StructType;
+	if (StructEvaluationType == EExpressionEvaluationType::None)
 	{
+		check(FieldEvaluationType != EExpressionEvaluationType::None);
+
 		// StructExpression is not used, so default to a zero-initialized struct
 		// This will happen if all the accessed struct fields are explicitly defined
 		OutShader.Code.Appendf(TEXT("%s_Set%s((%s)0, %s)"),
 			StructType->Name,
-			FieldName,
+			Field->Name,
+			StructType->Name,
+			FieldExpression->GetValueShader(Context, RequestedFieldType));
+	}
+	else
+	{
+		OutShader.Code.Appendf(TEXT("%s_Set%s(%s, %s)"),
 			StructType->Name,
-			FieldExpression->GetValueShader(Context, FieldRef.Type));
+			Field->Name,
+			StructExpression->GetValueShader(Context, RequestedStructType),
+			FieldExpression->GetValueShader(Context, RequestedFieldType));
 	}
 }
 
-void FExpressionSelect::PrepareValue(FEmitContext& Context, const FRequestedType& RequestedType, FPrepareValueResult& OutResult)
+void FExpressionSetStructField::EmitValuePreshader(FEmitContext& Context, const FRequestedType& RequestedType, Shader::FPreshaderData& OutPreshader) const
 {
-	const FPrepareValueResult& ConditionResult = PrepareExpressionValue(Context, ConditionExpression, 1);
-	if (ConditionResult.EvaluationType == EExpressionEvaluationType::Constant)
-	{
-		const bool bCondition = ConditionResult.ConstantValue.AsBoolScalar();
-		return OutResult.SetForwardValue(Context, bCondition ? TrueExpression : FalseExpression, RequestedType);
-	}
+	FRequestedType RequestedStructType(RequestedType);
+	RequestedStructType.ClearFieldRequested(Field);
+	const EExpressionEvaluationType StructEvaluationType = StructExpression->GetEvaluationType(RequestedStructType);
 
-	const FPrepareValueResult& LhsResult = PrepareExpressionValue(Context, FalseExpression, RequestedType);
-	const FPrepareValueResult& RhsResult = PrepareExpressionValue(Context, TrueExpression, RequestedType);
-	if (!ConditionResult.IsValid() || !LhsResult.IsValid() || !RhsResult.IsValid())
+	const FRequestedType RequestedFieldType = RequestedType.GetField(Field);
+	const EExpressionEvaluationType FieldEvaluationType = FieldExpression->GetEvaluationType(RequestedFieldType);
+
+	if (StructEvaluationType != EExpressionEvaluationType::None)
 	{
-		return;
+		StructExpression->GetValuePreshader(Context, RequestedStructType, OutPreshader);
 	}
-
-	const Shader::FValueTypeDescription LhsTypeDesc = Shader::GetValueTypeDescription(LhsResult.Type);
-	const Shader::FValueTypeDescription RhsTypeDesc = Shader::GetValueTypeDescription(RhsResult.Type);
-	if (LhsTypeDesc.ComponentType != RhsTypeDesc.ComponentType)
+	else
 	{
-		return Context.Errors.AddError(this, TEXT("Type mismatch"));
+		OutPreshader.WriteOpcode(Shader::EPreshaderOpcode::ConstantZero).Write(Shader::FType(StructType));
 	}
 
-	if (LhsResult.Type.IsStruct())
+	FieldExpression->GetValuePreshader(Context, RequestedFieldType, OutPreshader);
+	OutPreshader.WriteOpcode(Shader::EPreshaderOpcode::SetField).Write(Field->ComponentIndex).Write(Field->GetNumComponents());
+}
+
+void FExpressionSelect::PrepareValue(FEmitContext& Context, const FRequestedType& RequestedType, FPrepareValueResult& OutResult) const
+{
+	const FPreparedType& ConditionType = PrepareExpressionValue(Context, ConditionExpression, 1);
+	if (ConditionType.GetEvaluationType(Shader::EValueType::Bool1) == EExpressionEvaluationType::Constant)
 	{
-		OutResult.SetType(Context, EExpressionEvaluationType::Shader, LhsResult.Type);
+		const bool bCondition = ConditionExpression->GetValueConstant(Context, Shader::EValueType::Bool1).AsBoolScalar();
+		return OutResult.SetForwardValue(Context, RequestedType, bCondition ? TrueExpression : FalseExpression);
 	}
-	else
+
+	const FPreparedType& LhsType = PrepareExpressionValue(Context, FalseExpression, RequestedType);
+	const FPreparedType& RhsType = PrepareExpressionValue(Context, TrueExpression, RequestedType);
+	
+	if (LhsType.ValueComponentType != RhsType.ValueComponentType ||
+		LhsType.StructType != RhsType.StructType)
 	{
-		const int8 NumComponents = FMath::Max(LhsTypeDesc.NumComponents, RhsTypeDesc.NumComponents);
-		OutResult.SetType(Context, EExpressionEvaluationType::Shader, Shader::MakeValueType(LhsTypeDesc.ComponentType, NumComponents));
+		return Context.Errors.AddError(this, TEXT("Type mismatch"));
 	}
+
+	OutResult.SetType(Context, RequestedType, MergePreparedTypes(LhsType, RhsType));
 }
 
-void FExpressionSelect::EmitValueShader(FEmitContext& Context, FShaderValue& OutShader) const
+void FExpressionSelect::EmitValueShader(FEmitContext& Context, const FRequestedType& RequestedType, FShaderValue& OutShader) const
 {
+	OutShader.Type = RequestedType.GetType();
 	OutShader.Code.Appendf(TEXT("(%s ? %s : %s)"),
 		ConditionExpression->GetValueShader(Context, Shader::EValueType::Bool1),
-		TrueExpression->GetValueShader(Context, GetType()),
-		FalseExpression->GetValueShader(Context, GetType()));
+		TrueExpression->GetValueShader(Context, RequestedType),
+		FalseExpression->GetValueShader(Context, RequestedType));
 }
 
-void FExpressionSelect::EmitValuePreshader(FEmitContext& Context, Shader::FPreshaderData& OutPreshader) const
+void FExpressionSelect::EmitValuePreshader(FEmitContext& Context, const FRequestedType& RequestedType, Shader::FPreshaderData& OutPreshader) const
 {
 	check(false); // TODO
 }
 
-void FExpressionBinaryOp::PrepareValue(FEmitContext& Context, const FRequestedType& RequestedType, FPrepareValueResult& OutResult)
+void FExpressionBinaryOp::PrepareValue(FEmitContext& Context, const FRequestedType& RequestedType, FPrepareValueResult& OutResult) const
 {
-	const FPrepareValueResult& LhsResult = PrepareExpressionValue(Context, Lhs, RequestedType);
-	const FPrepareValueResult& RhsResult = PrepareExpressionValue(Context, Rhs, RequestedType);
-	if (!LhsResult.IsValid() || !RhsResult.IsValid())
+	const FPreparedType& LhsType = PrepareExpressionValue(Context, Lhs, RequestedType);
+	const FPreparedType& RhsType = PrepareExpressionValue(Context, Rhs, RequestedType);
+	if (LhsType.IsVoid() || RhsType.IsVoid())
 	{
 		return;
 	}
 
-	FString ErrorMessage;
-	Shader::EValueType ResultType;
+	if (!LhsType.IsNumeric() || !RhsType.IsNumeric())
+	{
+		return Context.Errors.AddError(this, TEXT("Invalid arithmetic between non-numeric types"));
+	}
+
+	FPreparedType ResultType = MergePreparedTypes(LhsType, RhsType);
+	if (Op == EBinaryOp::Less)
+	{
+		ResultType.ValueComponentType = Shader::EValueComponentType::Bool;
+	}
+
+	const FBinaryOpDescription OpDesc = GetBinaryOpDesription(Op);
+	if (OpDesc.PreshaderOpcode == Shader::EPreshaderOpcode::Nop)
+	{
+		// No preshader support
+		ResultType.SetEvaluationType(EExpressionEvaluationType::Shader);
+	}
+
+	OutResult.SetType(Context, RequestedType, ResultType);
+}
+
+namespace Private
+{
+struct FBinaryOpTypes
+{
+	Shader::FType ResultType;
+	FRequestedType LhsType;
+	FRequestedType RhsType;
+};
+FBinaryOpTypes GetBinaryOpTypes(const FRequestedType& RequestedType, EBinaryOp Op, const FRequestedType& LhsType, const FRequestedType& RhsType)
+{
+	const int32 NumRequestedComponents = RequestedType.GetNumComponents();
+	const Shader::EValueComponentType InputComponentType = Shader::CombineComponentTypes(LhsType.ValueComponentType, RhsType.ValueComponentType);
+	FBinaryOpTypes Result;
 	switch (Op)
 	{
 	case EBinaryOp::Less:
-		ResultType = Shader::MakeComparisonResultType(LhsResult.Type, RhsResult.Type, ErrorMessage);
+		Result.LhsType = Shader::MakeValueType(InputComponentType, NumRequestedComponents);
+		Result.RhsType = Shader::MakeValueType(InputComponentType, NumRequestedComponents);
+		Result.ResultType = Shader::MakeValueType(Shader::EValueComponentType::Bool, NumRequestedComponents);
 		break;
 	default:
-		ResultType = Shader::MakeArithmeticResultType(LhsResult.Type, RhsResult.Type, ErrorMessage);
+		Result.LhsType = RequestedType;
+		Result.RhsType = RequestedType;
+		Result.ResultType = RequestedType.GetType();
 		break;
 	}
-	if (!ErrorMessage.IsEmpty())
-	{
-		return Context.Errors.AddError(this, ErrorMessage);
-	}
-
-	const FBinaryOpDescription OpDesc = GetBinaryOpDesription(Op);
-	const EExpressionEvaluationType EvaluationType = OpDesc.PreshaderOpcode != Shader::EPreshaderOpcode::Nop
-		? CombineEvaluationTypes(LhsResult.EvaluationType, RhsResult.EvaluationType)
-		: EExpressionEvaluationType::Shader;
-	OutResult.SetType(Context, EvaluationType, ResultType);
+	return Result;
 }
+} // namespace Private
 
-void FExpressionBinaryOp::EmitValueShader(FEmitContext& Context, FShaderValue& OutShader) const
+void FExpressionBinaryOp::EmitValueShader(FEmitContext& Context, const FRequestedType& RequestedType, FShaderValue& OutShader) const
 {
-	const TCHAR* LhsCode = Lhs->GetValueShader(Context);
-	const TCHAR* RhsCode = Rhs->GetValueShader(Context);
+	const Private::FBinaryOpTypes Types = Private::GetBinaryOpTypes(RequestedType, Op, Lhs->GetRequestedType(), Rhs->GetRequestedType());
+	const TCHAR* LhsCode = Lhs->GetValueShader(Context, Types.LhsType);
+	const TCHAR* RhsCode = Rhs->GetValueShader(Context, Types.RhsType);
+
+	OutShader.Type = Types.ResultType;
 	switch (Op)
 	{
 	case EBinaryOp::Add: OutShader.Code.Appendf(TEXT("(%s + %s)"), LhsCode, RhsCode); break;
@@ -503,71 +542,69 @@ void FExpressionBinaryOp::EmitValueShader(FEmitContext& Context, FShaderValue& O
 	}
 }
 
-void FExpressionBinaryOp::EmitValuePreshader(FEmitContext& Context, Shader::FPreshaderData& OutPreshader) const
+void FExpressionBinaryOp::EmitValuePreshader(FEmitContext& Context, const FRequestedType& RequestedType, Shader::FPreshaderData& OutPreshader) const
 {
+	const Private::FBinaryOpTypes Types = Private::GetBinaryOpTypes(RequestedType, Op, Lhs->GetRequestedType(), Rhs->GetRequestedType());
 	const FBinaryOpDescription OpDesc = GetBinaryOpDesription(Op);
 	check(OpDesc.PreshaderOpcode != Shader::EPreshaderOpcode::Nop);
 
-	Lhs->GetValuePreshader(Context, OutPreshader);
-	Rhs->GetValuePreshader(Context, OutPreshader);
+	Lhs->GetValuePreshader(Context, Types.LhsType, OutPreshader);
+	Rhs->GetValuePreshader(Context, Types.RhsType, OutPreshader);
 	OutPreshader.WriteOpcode(OpDesc.PreshaderOpcode);
 }
 
-void FExpressionSwizzle::PrepareValue(FEmitContext& Context, const FRequestedType& RequestedType, FPrepareValueResult& OutResult)
+void FExpressionSwizzle::PrepareValue(FEmitContext& Context, const FRequestedType& RequestedType, FPrepareValueResult& OutResult) const
 {
-	FRequestedType RequestedInputType(4, false);
-	int32 MaxComponentIndex = 0;
+	const FRequestedType RequestedInputType = Parameters.GetRequestedInputType(RequestedType);
+	const FPreparedType& InputType = PrepareExpressionValue(Context, Input, RequestedInputType);
+
+	FPreparedType ResultType(InputType.ValueComponentType);
 	for (int32 ComponentIndex = 0; ComponentIndex < Parameters.NumComponents; ++ComponentIndex)
 	{
-		if (RequestedType.IsComponentRequested(ComponentIndex))
-		{
-			const int32 SwizzledComponentIndex = Parameters.ComponentIndex[ComponentIndex];
-			RequestedInputType.RequestedComponents[SwizzledComponentIndex] = true;
-			MaxComponentIndex = ComponentIndex;
-		}
+		const int32 SwizzledComponentIndex = Parameters.ComponentIndex[ComponentIndex];
+		ResultType.SetComponentEvaluationType(ComponentIndex, InputType.GetComponentEvaluationType(SwizzledComponentIndex));
 	}
 
-	const FPrepareValueResult& InputResult = PrepareExpressionValue(Context, Input, RequestedInputType);
-	OutResult.SetType(Context, InputResult.EvaluationType, Shader::MakeValueType(InputResult.Type, MaxComponentIndex + 1));
+	OutResult.SetType(Context, RequestedType, ResultType);
 }
 
-void FExpressionSwizzle::EmitValueShader(FEmitContext& Context, FShaderValue& OutShader) const
+void FExpressionSwizzle::EmitValueShader(FEmitContext& Context, const FRequestedType& RequestedType, FShaderValue& OutShader) const
 {
-	const Shader::FValueTypeDescription TypeDesc = Shader::GetValueTypeDescription(GetType());
-	const int32 NumSwizzleComponents = FMath::Min<int32>(TypeDesc.NumComponents, Parameters.NumComponents);
+	const FRequestedType RequestedInputType = Parameters.GetRequestedInputType(RequestedType);
 
 	static const TCHAR ComponentName[] = { 'x', 'y', 'z', 'w' };
 	TCHAR Swizzle[5] = TEXT("");
 	bool bHasSwizzleReorder = false;
 
-	for (int32 i = 0; i < NumSwizzleComponents; ++i)
+	for (int32 ComponentIndex = 0; ComponentIndex < Parameters.NumComponents; ++ComponentIndex)
 	{
-		const int32 ComponentIndex = Parameters.ComponentIndex[i];
-		check(ComponentIndex >= 0 && ComponentIndex < 4);
-		Swizzle[i] = ComponentName[ComponentIndex];
-		if (ComponentIndex != i)
+		const int32 SwizzledComponentIndex = Parameters.ComponentIndex[ComponentIndex];
+		Swizzle[ComponentIndex] = ComponentName[SwizzledComponentIndex];
+		if (SwizzledComponentIndex != ComponentIndex)
 		{
 			bHasSwizzleReorder = true;
 		}
 	}
 
+	OutShader.Type = RequestedType.GetType();
 	OutShader.bInline = true;
 	if (bHasSwizzleReorder)
 	{
 		OutShader.Code.Appendf(TEXT("%s.%s"),
-			Input->GetValueShader(Context),
+			Input->GetValueShader(Context, RequestedInputType),
 			Swizzle);
 	}
 	else
 	{
-		const Shader::EValueType InputType = Shader::MakeValueType(GetType(), NumSwizzleComponents);
-		OutShader.Code.Appendf(TEXT("%s"), Input->GetValueShader(Context, InputType));
+		OutShader.Code.Appendf(TEXT("%s"), Input->GetValueShader(Context, RequestedInputType));
 	}
 }
 
-void FExpressionSwizzle::EmitValuePreshader(FEmitContext& Context, Shader::FPreshaderData& OutPreshader) const
+void FExpressionSwizzle::EmitValuePreshader(FEmitContext& Context, const FRequestedType& RequestedType, Shader::FPreshaderData& OutPreshader) const
 {
-	Input->GetValuePreshader(Context, OutPreshader);
+	const FRequestedType RequestedInputType = Parameters.GetRequestedInputType(RequestedType);
+
+	Input->GetValuePreshader(Context, RequestedInputType, OutPreshader);
 	OutPreshader.WriteOpcode(Shader::EPreshaderOpcode::ComponentSwizzle)
 		.Write((uint8)Parameters.NumComponents)
 		.Write((uint8)Parameters.ComponentIndex[0])
@@ -576,58 +613,104 @@ void FExpressionSwizzle::EmitValuePreshader(FEmitContext& Context, Shader::FPres
 		.Write((uint8)Parameters.ComponentIndex[3]);
 }
 
-void FExpressionAppend::PrepareValue(FEmitContext& Context, const FRequestedType& RequestedType, FPrepareValueResult& OutResult)
+void FExpressionAppend::PrepareValue(FEmitContext& Context, const FRequestedType& RequestedType, FPrepareValueResult& OutResult) const
 {
-	const FPrepareValueResult& LhsResult = PrepareExpressionValue(Context, Lhs, 4);
-	const FPrepareValueResult& RhsResult = PrepareExpressionValue(Context, Rhs, 4);
-	if (!LhsResult.IsValid() || !RhsResult.IsValid())
+	// TODO - better handling of requested types for inputs
+	const int32 NumRequestedComponents = RequestedType.GetNumComponents();
+	const FPreparedType& LhsType = PrepareExpressionValue(Context, Lhs, NumRequestedComponents);
+	const int32 NumLhsComponents = FMath::Min(LhsType.GetNumComponents(), NumRequestedComponents);
+
+	const FPreparedType& RhsType = PrepareExpressionValue(Context, Rhs, NumRequestedComponents - NumLhsComponents);
+	const int32 NumRhsComponents = FMath::Min(RhsType.GetNumComponents(), NumRequestedComponents - NumLhsComponents);
+
+	if (LhsType.ValueComponentType != RhsType.ValueComponentType)
 	{
-		return;
+		return Context.Errors.AddError(this, TEXT("Type mismatch"));
 	}
 
-	if (LhsResult.Type.IsStruct() || RhsResult.Type.IsStruct())
+	FPreparedType ResultType(LhsType.ValueComponentType);
+	for (int32 Index = 0; Index < NumLhsComponents; ++Index)
 	{
-		return Context.Errors.AddError(this, TEXT("Not supported for structs"));
+		ResultType.SetComponentEvaluationType(Index, LhsType.GetComponentEvaluationType(Index));
 	}
-
-	const Shader::FValueTypeDescription LhsTypeDesc = Shader::GetValueTypeDescription(LhsResult.Type);
-	const Shader::FValueTypeDescription RhsTypeDesc = Shader::GetValueTypeDescription(RhsResult.Type);
-	if (LhsTypeDesc.ComponentType != RhsTypeDesc.ComponentType)
+	for (int32 Index = 0; Index < NumRhsComponents; ++Index)
 	{
-		return Context.Errors.AddError(this, TEXT("Type mismatch"));
+		ResultType.SetComponentEvaluationType(NumLhsComponents + Index, LhsType.GetComponentEvaluationType(Index));
 	}
 
-	const EExpressionEvaluationType EvaluationType = CombineEvaluationTypes(LhsResult.EvaluationType, RhsResult.EvaluationType);
-	const int32 NumComponents = FMath::Min<int32>(LhsTypeDesc.NumComponents + RhsTypeDesc.NumComponents, RequestedType.GetRequestedNumComponents());
-	OutResult.SetType(Context, EvaluationType, Shader::MakeValueType(LhsTypeDesc.ComponentType, NumComponents));
+	OutResult.SetType(Context, RequestedType, ResultType);
 }
 
-void FExpressionAppend::EmitValueShader(FEmitContext& Context, FShaderValue& OutShader) const
+namespace Private
+{
+struct FAppendTypes
 {
-	const Shader::FValueTypeDescription ResultTypeDesc = Shader::GetValueTypeDescription(GetType());
+	Shader::FType ResultType;
+	FRequestedType LhsType;
+	FRequestedType RhsType;
+};
+FAppendTypes GetAppendTypes(const FRequestedType& RequestedType, const FRequestedType& LhsType, const FRequestedType& RhsType)
+{
+	const int32 NumResultComponents = RequestedType.GetNumComponents();
+	const int32 NumLhsComponents = FMath::Min<int32>(NumResultComponents, LhsType.GetNumComponents());
+	const int32 NumRhsComponents = FMath::Min<int32>(NumResultComponents - NumLhsComponents, RhsType.GetNumComponents());
+
+	FAppendTypes Types;
+	Types.ResultType = Shader::MakeValueType(RequestedType.ValueComponentType, NumLhsComponents + NumRhsComponents);
+	Types.LhsType.ValueComponentType = RequestedType.ValueComponentType;
+	Types.RhsType.ValueComponentType = RequestedType.ValueComponentType;
+	for (int32 Index = 0; Index < NumLhsComponents; ++Index)
+	{
+		Types.LhsType.SetComponentRequested(Index, RequestedType.IsComponentRequested(Index));
+	}
+	for (int32 Index = 0; Index < NumRhsComponents; ++Index)
+	{
+		Types.RhsType.SetComponentRequested(Index, RequestedType.IsComponentRequested(NumLhsComponents + Index));
+	}
+	return Types;
+}
+} // namespace Private
+
+void FExpressionAppend::EmitValueShader(FEmitContext& Context, const FRequestedType& RequestedType, FShaderValue& OutShader) const
+{
+	const Private::FAppendTypes Types = Private::GetAppendTypes(RequestedType, Lhs->GetRequestedType(), Rhs->GetRequestedType());
+
 	OutShader.bInline = true;
-	OutShader.Code.Appendf(TEXT("%s(%s, %s)"),
-		ResultTypeDesc.Name,
-		Lhs->GetValueShader(Context),
-		Rhs->GetValueShader(Context));
+	OutShader.Type = Types.ResultType;
+	if (Types.RhsType.IsVoid())
+	{
+		OutShader.Code.Appendf(TEXT("%s"), Lhs->GetValueShader(Context, Types.LhsType));
+	}
+	else
+	{
+		OutShader.Code.Appendf(TEXT("%s(%s, %s)"),
+			RequestedType.GetName(),
+			Lhs->GetValueShader(Context, Types.LhsType),
+			Rhs->GetValueShader(Context, Types.RhsType));
+	}
 }
 
-void FExpressionAppend::EmitValuePreshader(FEmitContext& Context, Shader::FPreshaderData& OutPreshader) const
+void FExpressionAppend::EmitValuePreshader(FEmitContext& Context, const FRequestedType& RequestedType, Shader::FPreshaderData& OutPreshader) const
 {
-	Lhs->GetValuePreshader(Context, OutPreshader);
-	Rhs->GetValuePreshader(Context, OutPreshader);
-	OutPreshader.WriteOpcode(Shader::EPreshaderOpcode::AppendVector);
+	const Private::FAppendTypes Types = Private::GetAppendTypes(RequestedType, Lhs->GetRequestedType(), Rhs->GetRequestedType());
+	Lhs->GetValuePreshader(Context, Types.LhsType, OutPreshader);
+	if (Types.RhsType.IsVoid())
+	{
+		Rhs->GetValuePreshader(Context, Types.RhsType, OutPreshader);
+		OutPreshader.WriteOpcode(Shader::EPreshaderOpcode::AppendVector);
+	}
 }
 
-void FExpressionReflectionVector::PrepareValue(FEmitContext& Context, const FRequestedType& RequestedType, FPrepareValueResult& OutResult)
+void FExpressionReflectionVector::PrepareValue(FEmitContext& Context, const FRequestedType& RequestedType, FPrepareValueResult& OutResult) const
 {
 	Context.bReadMaterialNormal = true;
-	OutResult.SetType(Context, EExpressionEvaluationType::Shader, Shader::EValueType::Float3);
+	OutResult.SetType(Context, RequestedType, EExpressionEvaluationType::Shader, Shader::EValueType::Float3);
 }
 
-void FExpressionReflectionVector::EmitValueShader(FEmitContext& Context, FShaderValue& OutShader) const
+void FExpressionReflectionVector::EmitValueShader(FEmitContext& Context, const FRequestedType& RequestedType, FShaderValue& OutShader) const
 {
 	OutShader.bInline = true;
+	OutShader.Type = Shader::EValueType::Float3;
 	OutShader.Code.Append(TEXT("Parameters.ReflectionVector"));
 }
 
@@ -651,10 +734,10 @@ void FStatementReturn::EmitShader(FEmitContext& Context) const
 
 void FStatementIf::Prepare(FEmitContext& Context) const
 {
-	const FPrepareValueResult& ConditionResult = PrepareExpressionValue(Context, ConditionExpression, 1);
-	if (ConditionResult.EvaluationType == EExpressionEvaluationType::Constant)
+	const FPreparedType& ConditionType = PrepareExpressionValue(Context, ConditionExpression, 1);
+	if (ConditionType.GetEvaluationType(Shader::EValueType::Bool1) == EExpressionEvaluationType::Constant)
 	{
-		const bool bCondition = ConditionResult.ConstantValue.AsBoolScalar();
+		const bool bCondition = ConditionExpression->GetValueConstant(Context, Shader::EValueType::Bool1).AsBoolScalar();
 		if (bCondition)
 		{
 			MarkScopeLive(ThenScope);
@@ -666,7 +749,7 @@ void FStatementIf::Prepare(FEmitContext& Context) const
 			MarkScopeLive(ElseScope);
 		}
 	}
-	else if(ConditionResult.IsValid())
+	else if(!ConditionType.IsVoid())
 	{
 		MarkScopeLive(ThenScope);
 		MarkScopeLive(ElseScope);
@@ -675,9 +758,9 @@ void FStatementIf::Prepare(FEmitContext& Context) const
 
 void FStatementIf::EmitShader(FEmitContext& Context) const
 {
-	if (ConditionExpression->GetEvaluationType() == EExpressionEvaluationType::Constant)
+	if (ConditionExpression->GetEvaluationType(Shader::EValueType::Bool1) == EExpressionEvaluationType::Constant)
 	{
-		const bool bCondition = ConditionExpression->GetValueConstant(Context).AsBoolScalar();
+		const bool bCondition = ConditionExpression->GetValueConstant(Context, Shader::EValueType::Bool1).AsBoolScalar();
 		if (bCondition && IsScopeLive(ThenScope))
 		{
 			ParentScope->EmitScope(Context, ThenScope);
diff --git a/Engine/Source/Runtime/Engine/Private/Materials/MaterialExpressionHLSL.cpp b/Engine/Source/Runtime/Engine/Private/Materials/MaterialExpressionHLSL.cpp
--- a/Engine/Source/Runtime/Engine/Private/Materials/MaterialExpressionHLSL.cpp
+++ b/Engine/Source/Runtime/Engine/Private/Materials/MaterialExpressionHLSL.cpp
@@ -280,12 +280,7 @@ EMaterialGenerateHLSLStatus UMaterialExpressionGetMaterialAttributes::GenerateHL
 
 	const FGuid& AttributeID = AttributeGetTypes[AttributeIndex];
 	const FString& AttributeName = FMaterialAttributeDefinitionMap::GetAttributeName(AttributeID);
-
-	UE::HLSLTree::FExpressionGetStructField* SetAttributeExpression = Generator.GetTree().NewExpression<UE::HLSLTree::FExpressionGetStructField>();
-	SetAttributeExpression->StructType = Generator.GetMaterialAttributesType();
-	SetAttributeExpression->StructExpression = AttributesExpression;
-	SetAttributeExpression->FieldName = *AttributeName;
-	OutExpression = SetAttributeExpression;
+	OutExpression = Generator.GetTree().NewExpression<UE::HLSLTree::FExpressionGetStructField>(Generator.GetMaterialAttributesType(), *AttributeName, AttributesExpression);
 
 	return EMaterialGenerateHLSLStatus::Success;
 
@@ -317,12 +312,7 @@ EMaterialGenerateHLSLStatus UMaterialExpressionSetMaterialAttributes::GenerateHL
 				UE::HLSLTree::FExpression* ValueExpression = AttributeInput.AcquireHLSLExpression(Generator, Scope);
 				if (ValueExpression)
 				{
-					UE::HLSLTree::FExpressionSetStructField* SetAttributeExpression = Generator.GetTree().NewExpression<UE::HLSLTree::FExpressionSetStructField>();
-					SetAttributeExpression->StructType = Generator.GetMaterialAttributesType();
-					SetAttributeExpression->FieldName = *AttributeName;
-					SetAttributeExpression->StructExpression = AttributesExpression;
-					SetAttributeExpression->FieldExpression = ValueExpression;
-					AttributesExpression = SetAttributeExpression;
+					AttributesExpression = Generator.GetTree().NewExpression<UE::HLSLTree::FExpressionSetStructField>(Generator.GetMaterialAttributesType(), *AttributeName, AttributesExpression, ValueExpression);
 				}
 			}
 		}
diff --git a/Engine/Source/Runtime/Engine/Private/Materials/MaterialHLSLEmitter.cpp b/Engine/Source/Runtime/Engine/Private/Materials/MaterialHLSLEmitter.cpp
--- a/Engine/Source/Runtime/Engine/Private/Materials/MaterialHLSLEmitter.cpp
+++ b/Engine/Source/Runtime/Engine/Private/Materials/MaterialHLSLEmitter.cpp
@@ -693,6 +693,7 @@ bool MaterialEmitHLSL(const FMaterialCompileTargetParameters& InCompilerTarget,
 	TRefCountPtr<FSharedShaderCompilerEnvironment>& OutMaterialEnvironment)
 {
 	using namespace UE::HLSLTree;
+	using namespace UE::Shader;
 
 	FMemory::Memzero(SharedPixelProperties);
 	SharedPixelProperties[MP_Normal] = true;
@@ -717,8 +718,9 @@ bool MaterialEmitHLSL(const FMaterialCompileTargetParameters& InCompilerTarget,
 	UMaterialInterface* MaterialInterface = InOutMaterial.GetMaterialInterface();
 	UMaterial* BaseMaterial = MaterialInterface->GetMaterial();
 
+	FStructTypeRegistry TypeRegistry(Allocator);
 	FTree* HLSLTree = FTree::Create(Allocator);
-	FMaterialHLSLGenerator Generator(BaseMaterial, InCompilerTarget, *HLSLTree);
+	FMaterialHLSLGenerator Generator(BaseMaterial, InCompilerTarget, TypeRegistry, *HLSLTree);
 	FScope& RootScope = HLSLTree->GetRootScope();
 
 	bool bGenerateResult = false;
@@ -743,13 +745,13 @@ bool MaterialEmitHLSL(const FMaterialCompileTargetParameters& InCompilerTarget,
 		return false;
 	}
 
-	FEmitContext EmitContext(Allocator);
+	FEmitContext EmitContext(Allocator, TypeRegistry);
 	EmitContext.Material = &InOutMaterial;
 	EmitContext.StaticParameters = &InStaticParameters;
 	EmitContext.MaterialCompilationOutput = &OutCompilationOutput;
 
-	const FStructFieldRef NormalFieldRef = Generator.GetMaterialAttributesType()->FindFieldByName(*FMaterialAttributeDefinitionMap::GetAttributeName(MP_Normal));
-	check(NormalFieldRef);
+	const FStructField* NormalField = Generator.GetMaterialAttributesType()->FindFieldByName(*FMaterialAttributeDefinitionMap::GetAttributeName(MP_Normal));
+	check(NormalField);
 
 	if (!PrepareScope(EmitContext, Generator.GetResultStatement()->ParentScope))
 	{
@@ -758,7 +760,7 @@ bool MaterialEmitHLSL(const FMaterialCompileTargetParameters& InCompilerTarget,
 
 	// Prepare all fields *except* normal
 	FRequestedType RequestedMaterialAttributesType(Generator.GetMaterialAttributesType());
-	RequestedMaterialAttributesType.ClearFieldRequested(NormalFieldRef);
+	RequestedMaterialAttributesType.ClearFieldRequested(NormalField);
 	PrepareExpressionValue(EmitContext, Generator.GetResultExpression(), RequestedMaterialAttributesType);
 
 	const bool bReadMaterialNormal = EmitContext.bReadMaterialNormal;
@@ -767,7 +769,7 @@ bool MaterialEmitHLSL(const FMaterialCompileTargetParameters& InCompilerTarget,
 	if (!bReadMaterialNormal)
 	{
 		// No access to material normal, can execute everything in phase0
-		RequestedMaterialAttributesType.SetFieldRequested(NormalFieldRef);
+		RequestedMaterialAttributesType.SetFieldRequested(NormalField);
 		PrepareExpressionValue(EmitContext, Generator.GetResultExpression(), RequestedMaterialAttributesType);
 		HLSLTree->EmitShader(EmitContext, CodePhase0);
 	}
@@ -785,7 +787,7 @@ bool MaterialEmitHLSL(const FMaterialCompileTargetParameters& InCompilerTarget,
 
 		// Prepare code for just the normal
 		FRequestedType RequestedMaterialNormal(Generator.GetMaterialAttributesType(), false);
-		RequestedMaterialNormal.SetFieldRequested(NormalFieldRef);
+		RequestedMaterialNormal.SetFieldRequested(NormalField);
 		PrepareExpressionValue(EmitContext, Generator.GetResultExpression(), RequestedMaterialNormal);
 
 		// Execute the normal in phase0
@@ -793,7 +795,7 @@ bool MaterialEmitHLSL(const FMaterialCompileTargetParameters& InCompilerTarget,
 	}
 
 	FStringBuilderMemstack Declarations(Allocator, 64 * 1024);
-	HLSLTree->EmitDeclarationsCode(Declarations);
+	TypeRegistry.EmitDeclarationsCode(Declarations);
 
 	FString MaterialTemplateSource = GenerateMaterialTemplateHLSL(InCompilerTarget.ShaderPlatform,
 		InOutMaterial,
diff --git a/Engine/Source/Runtime/Engine/Private/Materials/MaterialHLSLGenerator.cpp b/Engine/Source/Runtime/Engine/Private/Materials/MaterialHLSLGenerator.cpp
--- a/Engine/Source/Runtime/Engine/Private/Materials/MaterialHLSLGenerator.cpp
+++ b/Engine/Source/Runtime/Engine/Private/Materials/MaterialHLSLGenerator.cpp
@@ -35,10 +35,14 @@ static UE::Shader::EValueType GetShaderType(EMaterialValueType MaterialType)
 	}
 }
 
-FMaterialHLSLGenerator::FMaterialHLSLGenerator(UMaterial* InTargetMaterial, const FMaterialCompileTargetParameters& InCompileTarget, UE::HLSLTree::FTree& InOutTree)
+FMaterialHLSLGenerator::FMaterialHLSLGenerator(UMaterial* InTargetMaterial,
+	const FMaterialCompileTargetParameters& InCompileTarget,
+	UE::Shader::FStructTypeRegistry& InOutTypeRegistry,
+	UE::HLSLTree::FTree& InOutTree)
 	: CompileTarget(InCompileTarget)
 	, TargetMaterial(InTargetMaterial)
 	, HLSLTree(&InOutTree)
+	, TypeRegistry(&InOutTypeRegistry)
 	, bGeneratedResult(false)
 {
 	const EMaterialShadingModel DefaultShadingModel = InTargetMaterial->GetShadingModels().GetFirstShadingModel();
@@ -46,7 +50,7 @@ FMaterialHLSLGenerator::FMaterialHLSLGenerator(UMaterial* InTargetMaterial, cons
 	FFunctionCallEntry* RootFunctionEntry = new(InOutTree.GetAllocator()) FFunctionCallEntry();
 	FunctionCallStack.Add(RootFunctionEntry);
 
-	TArray<UE::HLSLTree::FStructFieldInitializer, TInlineAllocator<MP_MAX>> MaterialAttributeFields;
+	TArray<UE::Shader::FStructFieldInitializer, TInlineAllocator<MP_MAX>> MaterialAttributeFields;
 
 	const TArray<FGuid>& OrderedVisibleAttributes = FMaterialAttributeDefinitionMap::GetOrderedVisibleAttributeList();
 	for (const FGuid& AttributeID : OrderedVisibleAttributes)
@@ -68,19 +72,17 @@ FMaterialHLSLGenerator::FMaterialHLSLGenerator(UMaterial* InTargetMaterial, cons
 			else
 			{
 				const UE::Shader::FValue DefaultValue = UE::Shader::Cast(FMaterialAttributeDefinitionMap::GetDefaultValue(AttributeID), ValueType);
-				for (int32 i = 0; i < DefaultValue.NumComponents; ++i)
-				{
-					MaterialAttributesDefaultValue.Component.Add(DefaultValue.Component[i]);
-				}
+				MaterialAttributesDefaultValue.Component.Append(DefaultValue.Component);
 			}
 		}
 	}
 
-	UE::HLSLTree::FStructTypeInitializer MaterialAttributesInitializer;
+	UE::Shader::FStructTypeInitializer MaterialAttributesInitializer;
 	MaterialAttributesInitializer.Name = TEXT("FMaterialAttributes");
 	MaterialAttributesInitializer.Fields = MaterialAttributeFields;
-	MaterialAttributesType = InOutTree.NewStructType(MaterialAttributesInitializer);
+	MaterialAttributesType = InOutTypeRegistry.NewType(MaterialAttributesInitializer);
 
+	check(MaterialAttributesDefaultValue.Component.Num() == MaterialAttributesType->ComponentTypes.Num());
 	MaterialAttributesDefaultValue.Type = MaterialAttributesType;
 }
 
diff --git a/Engine/Source/Runtime/Engine/Private/Materials/MaterialUniformExpressions.cpp b/Engine/Source/Runtime/Engine/Private/Materials/MaterialUniformExpressions.cpp
--- a/Engine/Source/Runtime/Engine/Private/Materials/MaterialUniformExpressions.cpp
+++ b/Engine/Source/Runtime/Engine/Private/Materials/MaterialUniformExpressions.cpp
@@ -143,7 +143,7 @@ FMaterialUniformExpressionType::FMaterialUniformExpressionType(const TCHAR* InNa
 void FMaterialUniformExpression::WriteNumberOpcodes(UE::Shader::FPreshaderData& OutData) const
 {
 	UE_LOG(LogMaterial, Warning, TEXT("Missing WriteNumberOpcodes impl for %s"), GetType()->GetName());
-	OutData.WriteOpcode(UE::Shader::EPreshaderOpcode::ConstantZero);
+	OutData.WriteOpcode(UE::Shader::EPreshaderOpcode::ConstantZero).Write(UE::Shader::FType(UE::Shader::EValueType::Float1));
 }
 
 void FUniformParameterOverrides::SetNumericOverride(EMaterialParameterType Type, const FHashedMaterialParameterInfo& ParameterInfo, const UE::Shader::FValue& Value, bool bOverride)
@@ -541,11 +541,13 @@ FUniformExpressionSet::FVTPackedStackAndLayerIndex FUniformExpressionSet::GetVTS
 
 void FMaterialUniformExpression::GetNumberValue(const struct FMaterialRenderContext& Context, FLinearColor& OutValue) const
 {
-	UE::Shader::FPreshaderData PreshaderData;
+	using namespace UE::Shader;
+
+	FPreshaderData PreshaderData;
 	WriteNumberOpcodes(PreshaderData);
-	UE::Shader::FValue Value;
-	PreshaderData.Evaluate(nullptr, Context, Value);
-	OutValue = Value.AsLinearColor();
+	FPreshaderStack Stack;
+	const FPreshaderValue Value = PreshaderData.Evaluate(nullptr, Context, Stack);
+	OutValue = Value.AsShaderValue().AsLinearColor();
 }
 
 int32 FUniformExpressionSet::FindOrAddTextureParameter(EMaterialTextureParameterType Type, const FMaterialTextureParameterInfo& Info)
@@ -641,6 +643,7 @@ void FUniformExpressionSet::GetTextureValue(int32 Index, const FMaterialRenderCo
 
 void FUniformExpressionSet::FillUniformBuffer(const FMaterialRenderContext& MaterialRenderContext, const FUniformExpressionCache& UniformExpressionCache, const FRHIUniformBufferLayout* UniformBufferLayout, uint8* TempBuffer, int TempBufferSize) const
 {
+	using namespace UE::Shader;
 	check(IsInParallelRenderingThread());
 
 	if (UniformBufferLayout->ConstantBufferSize > 0)
@@ -714,26 +717,25 @@ void FUniformExpressionSet::FillUniformBuffer(const FMaterialRenderContext& Mate
 
 		// Dump preshader results into buffer.
 		float* PreshaderBuffer = (float*)BufferCursor;
-		UE::Shader::FPreshaderStack PreshaderStack;
-		UE::Shader::FPreshaderDataContext PreshaderBaseContext(UniformPreshaderData);
+		FPreshaderStack PreshaderStack;
+		FPreshaderDataContext PreshaderBaseContext(UniformPreshaderData);
 		for (const FMaterialUniformPreshaderHeader& Preshader : UniformPreshaders)
 		{
-			UE::Shader::FValue Result;
-			UE::Shader::FPreshaderDataContext PreshaderContext(PreshaderBaseContext, Preshader.OpcodeOffset, Preshader.OpcodeSize);
-			UE::Shader::EvaluatePreshader(this, MaterialRenderContext, PreshaderStack, PreshaderContext, Result);
+			FPreshaderDataContext PreshaderContext(PreshaderBaseContext, Preshader.OpcodeOffset, Preshader.OpcodeSize);
+			const FPreshaderValue Result = EvaluatePreshader(this, MaterialRenderContext, PreshaderStack, PreshaderContext);
 
 			float* DestAddress = PreshaderBuffer + Preshader.BufferOffset;
-			if(Preshader.ComponentType == UE::Shader::EValueComponentType::Float)
+			if(Preshader.ComponentType == EValueComponentType::Float)
 			{
-				const UE::Shader::FFloatValue FloatValue = Result.AsFloat();
+				const FFloatValue FloatValue = Result.AsShaderValue().AsFloat();
 				for (uint32 i = 0u; i < Preshader.NumComponents; ++i)
 				{
 					*DestAddress++ = FloatValue[i];
 				}
 			}
-			else if (Preshader.ComponentType == UE::Shader::EValueComponentType::Double)
+			else if (Preshader.ComponentType == EValueComponentType::Double)
 			{
-				const UE::Shader::FDoubleValue DoubleValue = Result.AsDouble();
+				const FDoubleValue DoubleValue = Result.AsShaderValue().AsDouble();
 				float TileValue[4];
 				float OffsetValue[4];
 				for (uint32 i = 0u; i < Preshader.NumComponents; ++i)
diff --git a/Engine/Source/Runtime/Engine/Private/Materials/MaterialUniformExpressions.h b/Engine/Source/Runtime/Engine/Private/Materials/MaterialUniformExpressions.h
--- a/Engine/Source/Runtime/Engine/Private/Materials/MaterialUniformExpressions.h
+++ b/Engine/Source/Runtime/Engine/Private/Materials/MaterialUniformExpressions.h
@@ -199,10 +199,10 @@ public:
 		switch (ValueType)
 		{
 		case MCT_Float:
-		case MCT_Float1: ConstantValue.NumComponents = 1; break;
-		case MCT_Float2: ConstantValue.NumComponents = 2; break;
-		case MCT_Float3: ConstantValue.NumComponents = 3; break;
-		case MCT_Float4: ConstantValue.NumComponents = 4; break;
+		case MCT_Float1: ConstantValue.Type = UE::Shader::EValueType::Float1; break;
+		case MCT_Float2: ConstantValue.Type = UE::Shader::EValueType::Float2; break;
+		case MCT_Float3: ConstantValue.Type = UE::Shader::EValueType::Float3; break;
+		case MCT_Float4: ConstantValue.Type = UE::Shader::EValueType::Float4; break;
 		default: checkNoEntry(); break;
 		}
 		OutData.WriteOpcode(UE::Shader::EPreshaderOpcode::Constant);
diff --git a/Engine/Source/Runtime/Engine/Private/Shader/Preshader.cpp b/Engine/Source/Runtime/Engine/Private/Shader/Preshader.cpp
--- a/Engine/Source/Runtime/Engine/Private/Shader/Preshader.cpp
+++ b/Engine/Source/Runtime/Engine/Private/Shader/Preshader.cpp
@@ -8,12 +8,110 @@
 #include "ExternalTexture.h"
 
 IMPLEMENT_TYPE_LAYOUT(UE::Shader::FPreshaderData);
+IMPLEMENT_TYPE_LAYOUT(UE::Shader::FPreshaderStructType);
 
 namespace UE
 {
 namespace Shader
 {
 
+FPreshaderType::FPreshaderType(const FType& InType) : ValueType(InType.ValueType)
+{
+	if (InType.IsStruct())
+	{
+		StructTypeHash = InType.StructType->Hash;
+		StructComponentTypes = InType.StructType->ComponentTypes;
+	}
+}
+
+EValueComponentType FPreshaderType::GetComponentType(int32 Index) const
+{
+	if (IsStruct())
+	{
+		return StructComponentTypes.IsValidIndex(Index) ? StructComponentTypes[Index] : EValueComponentType::Void;
+	}
+	else
+	{
+		const FValueTypeDescription TypeDesc = GetValueTypeDescription(ValueType);
+		return (Index >= 0 && Index < TypeDesc.NumComponents) ? TypeDesc.ComponentType : EValueComponentType::Void;
+	}
+}
+
+void FPreshaderStack::PushValue(const FValue& InValue)
+{
+	check(InValue.Component.Num() == InValue.Type.GetNumComponents());
+	Values.Emplace(InValue.Type);
+	Components.Append(InValue.Component);
+}
+
+void FPreshaderStack::PushValue(const FPreshaderValue& InValue)
+{
+	check(InValue.Component.Num() == InValue.Type.GetNumComponents());
+	Values.Add(InValue.Type);
+	Components.Append(InValue.Component.GetData(), InValue.Component.Num());
+}
+
+void FPreshaderStack::PushValue(const FPreshaderType& InType, TArrayView<const FValueComponent> InComponents)
+{
+	check(InComponents.Num() == InType.GetNumComponents());
+	Values.Add(InType);
+	Components.Append(InComponents.GetData(), InComponents.Num());
+}
+
+TArrayView<FValueComponent> FPreshaderStack::PushEmptyValue(const FPreshaderType& InType)
+{
+	Values.Add(InType);
+	const int32 NumComponents = InType.GetNumComponents();
+	const int32 ComponentIndex = Components.AddZeroed(NumComponents);
+	return MakeArrayView(Components.GetData() + ComponentIndex, NumComponents);
+}
+
+FPreshaderValue FPreshaderStack::PopValue()
+{
+	FPreshaderValue Value;
+	Value.Type = Values.Pop(false);
+
+	const int32 NumComponents = Value.Type.GetNumComponents();
+	const int32 ComponentIndex = Components.Num() - NumComponents;
+
+	Value.Component = MakeArrayView(Components.GetData() + ComponentIndex, NumComponents);
+	Components.RemoveAt(ComponentIndex, NumComponents, false);
+
+	return Value;
+}
+
+FPreshaderValue FPreshaderStack::PeekValue()
+{
+	FPreshaderValue Value;
+	Value.Type = Values.Last();
+
+	const int32 NumComponents = Value.Type.GetNumComponents();
+	const int32 ComponentIndex = Components.Num() - NumComponents;
+	Value.Component = MakeArrayView(Components.GetData() + ComponentIndex, NumComponents);
+	return Value;
+}
+
+FValue FPreshaderValue::AsShaderValue(const FStructTypeRegistry* TypeRegistry) const
+{
+	FValue Result;
+	if (!Type.IsStruct())
+	{
+		Result.Type = Type.ValueType;
+		Result.Component.Append(Component.GetData(), Component.Num());
+	}
+	else if (ensure(TypeRegistry))
+	{
+		Result.Type = TypeRegistry->FindType(Type.StructTypeHash);
+		if (Result.Type.IsStruct())
+		{
+			check(Result.Type.GetNumComponents() == Component.Num());
+			Result.Component.Append(Component.GetData(), Component.Num());
+		}
+	}
+
+	return Result;
+}
+
 void FPreshaderData::WriteData(const void* Value, uint32 Size)
 {
 	Data.Append((uint8*)Value, Size);
@@ -26,57 +124,80 @@ void FPreshaderData::WriteName(const FScriptName& Name)
 	{
 		Index = Names.Add(Name);
 	}
+
 	check(Index >= 0 && Index <= 0xffff);
-	const uint32 Offset = Data.Num();
-	NameOffsets.Add(Offset);
 	Write((uint16)Index);
 }
 
-void FPreshaderData::Append(const FPreshaderData& InPreshader)
+void FPreshaderData::WriteType(const FType& Type)
 {
-	const uint32 BaseOffset = Data.Num();
-	Data.Append(InPreshader.Data);
-
-	TArray<uint16> NameIndexRemap;
-	NameIndexRemap.Empty(InPreshader.Names.Num());
-	for (const FScriptName& Name : InPreshader.Names)
+	Write(Type.ValueType);
+	if (Type.IsStruct())
 	{
-		const int32 Index = Names.AddUnique(Name);
+		const uint64 Hash = Type.StructType->Hash;
+		int32 Index = INDEX_NONE;
+		for (int32 PrevIndex = 0; PrevIndex < StructTypes.Num(); ++PrevIndex)
+		{
+			if (StructTypes[PrevIndex].Hash == Hash)
+			{
+				Index = PrevIndex;
+				break;
+			}
+		}
+		if (Index == INDEX_NONE)
+		{
+			Index = StructTypes.Num();
+			FPreshaderStructType& PreshaderStructType = StructTypes.AddDefaulted_GetRef();
+			PreshaderStructType.Hash = Hash;
+			PreshaderStructType.ComponentTypeIndex = StructComponentTypes.Num();
+			PreshaderStructType.NumComponents = Type.StructType->ComponentTypes.Num();
+			StructComponentTypes.Append(Type.StructType->ComponentTypes.GetData(), Type.StructType->ComponentTypes.Num());
+		}
+
 		check(Index >= 0 && Index <= 0xffff);
-		NameIndexRemap.Add(Index);
+		Write((uint16)Index);
 	}
+}
 
-	for (uint32 PrevOffset : InPreshader.NameOffsets)
+void FPreshaderData::WriteValue(const FValue& Value)
+{
+	WriteType(Value.Type);
+	for (int32 Index = 0; Index < Value.Component.Num(); ++Index)
 	{
-		const uint32 Offset = BaseOffset + PrevOffset;
-		uint16* PrevNameIndex = (uint16*)(Data.GetData() + Offset);
-		const uint16 RemapIndex = NameIndexRemap[*PrevNameIndex];
-		*PrevNameIndex = RemapIndex;
-		NameOffsets.Add(Offset);
+		const EValueComponentType ComponentType = Value.Type.GetComponentType(Index);
+		const int32 ComponentSize = GetComponentTypeSizeInBytes(ComponentType);
+		WriteData(&Value.Component[Index].Packed, ComponentSize);
 	}
 }
 
 FPreshaderDataContext::FPreshaderDataContext(const FPreshaderData& InData)
 	: Ptr(InData.Data.GetData())
 	, EndPtr(Ptr + InData.Data.Num())
-	, Names(InData.Names.GetData())
-	, NumNames(InData.Names.Num())
+	, Names(InData.Names)
+	, StructTypes(InData.StructTypes)
+	, StructComponentTypes(InData.StructComponentTypes)
 {}
 
 FPreshaderDataContext::FPreshaderDataContext(const FPreshaderDataContext& InContext, uint32 InOffset, uint32 InSize)
 	: Ptr(InContext.Ptr + InOffset)
 	, EndPtr(Ptr + InSize)
 	, Names(InContext.Names)
-	, NumNames(InContext.NumNames)
+	, StructTypes(InContext.StructTypes)
+	, StructComponentTypes(InContext.StructComponentTypes)
 {}
 
+static inline void ReadPreshaderData(FPreshaderDataContext& RESTRICT Data, int32 Size, void* Result)
+{
+	FMemory::Memcpy(Result, Data.Ptr, Size);
+	Data.Ptr += Size;
+	checkSlow(Data.Ptr <= Data.EndPtr);
+}
+
 template<typename T>
 static inline T ReadPreshaderValue(FPreshaderDataContext& RESTRICT Data)
 {
 	T Result;
-	FMemory::Memcpy(&Result, Data.Ptr, sizeof(T));
-	Data.Ptr += sizeof(T);
-	checkSlow(Data.Ptr <= Data.EndPtr);
+	ReadPreshaderData(Data, sizeof(T), &Result);
 	return Result;
 }
 
@@ -91,10 +212,24 @@ template<>
 FScriptName ReadPreshaderValue<FScriptName>(FPreshaderDataContext& RESTRICT Data)
 {
 	const int32 Index = ReadPreshaderValue<uint16>(Data);
-	check(Index >= 0 && Index < Data.NumNames);
 	return Data.Names[Index];
 }
 
+template<>
+FPreshaderType ReadPreshaderValue<FPreshaderType>(FPreshaderDataContext& RESTRICT Data)
+{
+	FPreshaderType Result;
+	Result.ValueType = ReadPreshaderValue<EValueType>(Data);
+	if (Result.ValueType == EValueType::Struct)
+	{
+		const uint16 Index = ReadPreshaderValue<uint16>(Data);
+		const FPreshaderStructType& PreshaderStruct = Data.StructTypes[Index];
+		Result.StructTypeHash = PreshaderStruct.Hash;
+		Result.StructComponentTypes = MakeArrayView(Data.StructComponentTypes.GetData() + PreshaderStruct.ComponentTypeIndex, PreshaderStruct.NumComponents);
+	}
+	return Result;
+}
+
 template<>
 FName ReadPreshaderValue<FName>(FPreshaderDataContext& RESTRICT Data) = delete;
 
@@ -107,18 +242,57 @@ FHashedMaterialParameterInfo ReadPreshaderValue<FHashedMaterialParameterInfo>(FP
 	return FHashedMaterialParameterInfo(Name, Association, Index);
 }
 
-template<>
-FValue ReadPreshaderValue<FValue>(FPreshaderDataContext& RESTRICT Data)
+static void EvaluateConstantZero(FPreshaderStack& Stack, FPreshaderDataContext& RESTRICT Data)
 {
-	const EValueType Type = (EValueType)ReadPreshaderValue<uint8>(Data);
-	uint32 Size = 0u;
-	FValue Result = FValue::FromMemoryImage(Type, Data.Ptr, &Size);
-	Data.Ptr += Size;
-	checkSlow(Data.Ptr <= Data.EndPtr);
-	return Result;
+	const FPreshaderType Type = ReadPreshaderValue<FPreshaderType>(Data);
+	Stack.PushEmptyValue(Type); // Leave the empty value zero-initialized
+}
+
+static void EvaluateConstant(FPreshaderStack& Stack, FPreshaderDataContext& RESTRICT Data)
+{
+	const FPreshaderType Type = ReadPreshaderValue<FPreshaderType>(Data);
+	TArrayView<FValueComponent> Component = Stack.PushEmptyValue(Type);
+	for (int32 Index = 0; Index < Component.Num(); ++Index)
+	{
+		const EValueComponentType ComponentType = Type.GetComponentType(Index);
+		const int32 ComponmentSizeInBytes = GetComponentTypeSizeInBytes(ComponentType);
+		ReadPreshaderData(Data, ComponmentSizeInBytes, &Component[Index].Packed);
+	}
+}
+
+static void EvaluateSetField(FPreshaderStack& Stack, FPreshaderDataContext& RESTRICT Data)
+{
+	const FPreshaderValue Value = Stack.PopValue();
+	const FPreshaderValue StructValue = Stack.PeekValue();
+
+	const int32 ComponentIndex = ReadPreshaderValue<int32>(Data);
+	const int32 ComponentNum = ReadPreshaderValue<int32>(Data);
+	
+	// Modify the struct value in-place
+	for (int32 Index = 0; Index < ComponentNum; ++Index)
+	{
+		StructValue.Component[ComponentIndex + Index] = Value.Component[Index];
+	}
+}
+
+static void EvaluateGetField(FPreshaderStack& Stack, FPreshaderDataContext& RESTRICT Data)
+{
+	const FPreshaderValue StructValue = Stack.PopValue();
+	const FPreshaderType FieldType = ReadPreshaderValue<FPreshaderType>(Data);
+	const int32 ComponentIndex = ReadPreshaderValue<int32>(Data);
+	const int32 ComponentNum = FieldType.GetNumComponents();
+
+	// Need to make a local copy of components, since StructValue.Component is only valid until we push the result
+	TArray<FValueComponent, TInlineAllocator<64>> FieldComponents;
+	FieldComponents.Empty(ComponentNum);
+	for (int32 Index = 0; Index < ComponentNum; ++Index)
+	{
+		FieldComponents.Add(StructValue.Component[ComponentIndex + Index]);
+	}
+	Stack.PushValue(FieldType, FieldComponents);
 }
 
-static void EvaluateParameter(const FUniformExpressionSet& UniformExpressionSet, uint32 ParameterIndex, const FMaterialRenderContext& Context, FValue& OutValue)
+static void EvaluateParameter(FPreshaderStack& Stack, const FUniformExpressionSet& UniformExpressionSet, uint32 ParameterIndex, const FMaterialRenderContext& Context)
 {
 	const FMaterialNumericParameterInfo& Parameter = UniformExpressionSet.GetNumericParameter(ParameterIndex);
 	bool bFoundParameter = false;
@@ -129,7 +303,7 @@ static void EvaluateParameter(const FUniformExpressionSet& UniformExpressionSet,
 		FMaterialParameterValue ParameterValue;
 		if (Context.MaterialRenderProxy->GetParameterValue(Parameter.ParameterType, Parameter.ParameterInfo, ParameterValue, Context))
 		{
-			OutValue = ParameterValue.AsShaderValue();
+			Stack.PushValue(ParameterValue.AsShaderValue());
 			bFoundParameter = true;
 		}
 	}
@@ -138,39 +312,44 @@ static void EvaluateParameter(const FUniformExpressionSet& UniformExpressionSet,
 #if WITH_EDITOR
 	if (!bFoundParameter)
 	{
-		bFoundParameter = Context.Material.TransientOverrides.GetNumericOverride(Parameter.ParameterType, Parameter.ParameterInfo, OutValue);
+		FValue OverrideValue;
+		if (Context.Material.TransientOverrides.GetNumericOverride(Parameter.ParameterType, Parameter.ParameterInfo, OverrideValue))
+		{
+			Stack.PushValue(OverrideValue);
+			bFoundParameter = true;
+		}
 	}
 #endif // WITH_EDITOR
 
 	// Default value
 	if (!bFoundParameter)
 	{
-		OutValue = UniformExpressionSet.GetDefaultParameterValue(Parameter.ParameterType, Parameter.DefaultValueOffset);
+		Stack.PushValue(UniformExpressionSet.GetDefaultParameterValue(Parameter.ParameterType, Parameter.DefaultValueOffset));
 	}
 }
 
 template<typename Operation>
 static inline void EvaluateUnaryOp(FPreshaderStack& Stack, const Operation& Op)
 {
-	const FValue Value = Stack.Pop(false);
-	Stack.Add(Op(Value));
+	const FValue Value = Stack.PopValue().AsShaderValue();
+	Stack.PushValue(Op(Value));
 }
 
 template<typename Operation>
 static inline void EvaluateBinaryOp(FPreshaderStack& Stack, const Operation& Op)
 {
-	const FValue Value1 = Stack.Pop(false);
-	const FValue Value0 = Stack.Pop(false);
-	Stack.Add(Op(Value0, Value1));
+	const FValue Value1 = Stack.PopValue().AsShaderValue();
+	const FValue Value0 = Stack.PopValue().AsShaderValue();
+	Stack.PushValue(Op(Value0, Value1));
 }
 
 template<typename Operation>
 static inline void EvaluateTernaryOp(FPreshaderStack& Stack, const Operation& Op)
 {
-	const FValue Value2 = Stack.Pop(false);
-	const FValue Value1 = Stack.Pop(false);
-	const FValue Value0 = Stack.Pop(false);
-	Stack.Add(Op(Value0, Value1, Value2));
+	const FValue Value2 = Stack.PopValue().AsShaderValue();
+	const FValue Value1 = Stack.PopValue().AsShaderValue();
+	const FValue Value0 = Stack.PopValue().AsShaderValue();
+	Stack.PushValue(Op(Value0, Value1, Value2));
 }
 
 static void EvaluateComponentSwizzle(FPreshaderStack& Stack, FPreshaderDataContext& RESTRICT Data)
@@ -181,11 +360,8 @@ static void EvaluateComponentSwizzle(FPreshaderStack& Stack, FPreshaderDataConte
 	const uint8 IndexB = ReadPreshaderValue<uint8>(Data);
 	const uint8 IndexA = ReadPreshaderValue<uint8>(Data);
 
-	FValue Value = Stack.Pop(false);
-	FValue Result;
-	Result.ComponentType = Value.ComponentType;
-	Result.NumComponents = NumElements;
-
+	FValue Value = Stack.PopValue().AsShaderValue();
+	FValue Result(MakeValueType(Value.Type, NumElements));
 	switch (NumElements)
 	{
 	case 4:
@@ -204,7 +380,7 @@ static void EvaluateComponentSwizzle(FPreshaderStack& Stack, FPreshaderDataConte
 		UE_LOG(LogMaterial, Fatal, TEXT("Invalid number of swizzle elements: %d"), NumElements);
 		break;
 	}
-	Stack.Add(Result);
+	Stack.PushValue(Result);
 }
 
 static const UTexture* GetTextureParameter(const FMaterialRenderContext& Context, FPreshaderDataContext& RESTRICT Data)
@@ -225,11 +401,11 @@ static void EvaluateTextureSize(const FMaterialRenderContext& Context, FPreshade
 		const uint32 SizeX = Texture->GetResource()->GetSizeX();
 		const uint32 SizeY = Texture->GetResource()->GetSizeY();
 		const uint32 SizeZ = Texture->GetResource()->GetSizeZ();
-		Stack.Add(FValue((float)SizeX, (float)SizeY, (float)SizeZ));
+		Stack.PushValue(FValue((float)SizeX, (float)SizeY, (float)SizeZ));
 	}
 	else
 	{
-		Stack.Add(FValue(0.0f, 0.0f, 0.0f));
+		Stack.PushValue(FValue(0.0f, 0.0f, 0.0f));
 	}
 }
 
@@ -241,11 +417,11 @@ static void EvaluateTexelSize(const FMaterialRenderContext& Context, FPreshaderS
 		const uint32 SizeX = Texture->GetResource()->GetSizeX();
 		const uint32 SizeY = Texture->GetResource()->GetSizeY();
 		const uint32 SizeZ = Texture->GetResource()->GetSizeZ();
-		Stack.Add(FValue(1.0f / (float)SizeX, 1.0f / (float)SizeY, (SizeZ > 0 ? 1.0f / (float)SizeZ : 0.0f)));
+		Stack.PushValue(FValue(1.0f / (float)SizeX, 1.0f / (float)SizeY, (SizeZ > 0 ? 1.0f / (float)SizeZ : 0.0f)));
 	}
 	else
 	{
-		Stack.Add(FValue(0.0f, 0.0f, 0.0f));
+		Stack.PushValue(FValue(0.0f, 0.0f, 0.0f));
 	}
 }
 
@@ -265,7 +441,7 @@ static void EvaluateExternalTextureCoordinateScaleRotation(const FMaterialRender
 	{
 		FExternalTextureRegistry::Get().GetExternalTextureCoordinateScaleRotation(GuidToLookup, Result);
 	}
-	Stack.Add(Result);
+	Stack.PushValue(Result);
 }
 
 static void EvaluateExternalTextureCoordinateOffset(const FMaterialRenderContext& Context, FPreshaderStack& Stack, FPreshaderDataContext& RESTRICT Data)
@@ -276,7 +452,7 @@ static void EvaluateExternalTextureCoordinateOffset(const FMaterialRenderContext
 	{
 		FExternalTextureRegistry::Get().GetExternalTextureCoordinateOffset(GuidToLookup, Result);
 	}
-	Stack.Add(Result);
+	Stack.PushValue(Result);
 }
 
 static void EvaluateRuntimeVirtualTextureUniform(const FMaterialRenderContext& Context, FPreshaderStack& Stack, FPreshaderDataContext& RESTRICT Data)
@@ -292,15 +468,15 @@ static void EvaluateRuntimeVirtualTextureUniform(const FMaterialRenderContext& C
 	}
 	if (Texture != nullptr && VectorIndex != INDEX_NONE)
 	{
-		Stack.Add(FValue(Texture->GetUniformParameter(VectorIndex)));
+		Stack.PushValue(FValue(Texture->GetUniformParameter(VectorIndex)));
 	}
 	else
 	{
-		Stack.Add(FValue(0.f, 0.f, 0.f, 0.f));
+		Stack.PushValue(FValue(0.f, 0.f, 0.f, 0.f));
 	}
 }
 
-void EvaluatePreshader(const FUniformExpressionSet* UniformExpressionSet, const FMaterialRenderContext& Context, FPreshaderStack& Stack, FPreshaderDataContext& RESTRICT Data, FValue& OutValue)
+FPreshaderValue EvaluatePreshader(const FUniformExpressionSet* UniformExpressionSet, const FMaterialRenderContext& Context, FPreshaderStack& Stack, FPreshaderDataContext& RESTRICT Data)
 {
 	uint8 const* const DataEnd = Data.EndPtr;
 
@@ -310,15 +486,13 @@ void EvaluatePreshader(const FUniformExpressionSet* UniformExpressionSet, const
 		const EPreshaderOpcode Opcode = (EPreshaderOpcode)ReadPreshaderValue<uint8>(Data);
 		switch (Opcode)
 		{
-		case EPreshaderOpcode::ConstantZero:
-			Stack.Add(FValue(0.0f));
-			break;
-		case EPreshaderOpcode::Constant:
-			Stack.Add(ReadPreshaderValue<FValue>(Data));
-			break;
+		case EPreshaderOpcode::ConstantZero: EvaluateConstantZero(Stack, Data); break;
+		case EPreshaderOpcode::Constant: EvaluateConstant(Stack, Data); break;
+		case EPreshaderOpcode::GetField: EvaluateGetField(Stack, Data); break;
+		case EPreshaderOpcode::SetField: EvaluateSetField(Stack, Data); break;
 		case EPreshaderOpcode::Parameter:
 			check(UniformExpressionSet);
-			EvaluateParameter(*UniformExpressionSet, ReadPreshaderValue<uint16>(Data), Context, Stack.AddDefaulted_GetRef());
+			EvaluateParameter(Stack, *UniformExpressionSet, ReadPreshaderValue<uint16>(Data), Context);
 			break;
 		case EPreshaderOpcode::Add: EvaluateBinaryOp(Stack, Add); break;
 		case EPreshaderOpcode::Sub: EvaluateBinaryOp(Stack, Sub); break;
@@ -366,24 +540,24 @@ void EvaluatePreshader(const FUniformExpressionSet* UniformExpressionSet, const
 	}
 	check(Data.Ptr == DataEnd);
 
-	ensure(Stack.Num() <= 1);
+	FPreshaderValue Result;
 	if (Stack.Num() > 0)
 	{
-		OutValue = Stack.Last();
+		Result = Stack.PopValue();
 	}
+	Stack.CheckEmpty();
+	return Result;
 }
 
-void FPreshaderData::Evaluate(FUniformExpressionSet* UniformExpressionSet, const struct FMaterialRenderContext& Context, FValue& OutValue)
+FPreshaderValue FPreshaderData::Evaluate(FUniformExpressionSet* UniformExpressionSet, const struct FMaterialRenderContext& Context, FPreshaderStack& Stack)
 {
-	FPreshaderStack Stack;
 	FPreshaderDataContext PreshaderContext(*this);
-	EvaluatePreshader(UniformExpressionSet, Context, Stack, PreshaderContext, OutValue);
+	return EvaluatePreshader(UniformExpressionSet, Context, Stack, PreshaderContext);
 }
 
 void FPreshaderData::AppendHash(FSHA1& OutHasher) const
 {
 	OutHasher.Update((uint8*)Names.GetData(), Names.Num() * Names.GetTypeSize());
-	OutHasher.Update((uint8*)NameOffsets.GetData(), NameOffsets.Num() * NameOffsets.GetTypeSize());
 	OutHasher.Update((uint8*)Data.GetData(), Data.Num() * Data.GetTypeSize());
 }
 
diff --git a/Engine/Source/Runtime/Engine/Private/Shader/PreshaderEvaluate.h b/Engine/Source/Runtime/Engine/Private/Shader/PreshaderEvaluate.h
--- a/Engine/Source/Runtime/Engine/Private/Shader/PreshaderEvaluate.h
+++ b/Engine/Source/Runtime/Engine/Private/Shader/PreshaderEvaluate.h
@@ -14,8 +14,6 @@ namespace Shader
 
 class FPreshaderData;
 
-using FPreshaderStack = TArray<FValue, TInlineAllocator<64u>>;
-
 struct FPreshaderDataContext
 {
 	explicit FPreshaderDataContext(const FPreshaderData& InData);
@@ -23,11 +21,12 @@ struct FPreshaderDataContext
 
 	const uint8* RESTRICT Ptr;
 	const uint8* RESTRICT EndPtr;
-	const FScriptName* RESTRICT Names;
-	int32 NumNames;
+	TArrayView<const FScriptName> Names;
+	TArrayView<const FPreshaderStructType> StructTypes;
+	TArrayView<const EValueComponentType> StructComponentTypes;
 };
 
-ENGINE_API void EvaluatePreshader(const FUniformExpressionSet* UniformExpressionSet, const FMaterialRenderContext& Context, FPreshaderStack& Stack, FPreshaderDataContext& RESTRICT Data, FValue& OutValue);
+ENGINE_API FPreshaderValue EvaluatePreshader(const FUniformExpressionSet* UniformExpressionSet, const FMaterialRenderContext& Context, FPreshaderStack& Stack, FPreshaderDataContext& RESTRICT Data);
 
 } // namespace Shader
 } // namespace UE
diff --git a/Engine/Source/Runtime/Engine/Private/Shader/ShaderValue.cpp b/Engine/Source/Runtime/Engine/Private/Shader/ShaderValue.cpp
--- a/Engine/Source/Runtime/Engine/Private/Shader/ShaderValue.cpp
+++ b/Engine/Source/Runtime/Engine/Private/Shader/ShaderValue.cpp
@@ -1,6 +1,8 @@
 // Copyright Epic Games, Inc. All Rights Reserved.
 #include "Shader/ShaderTypes.h"
 #include "Misc/StringBuilder.h"
+#include "Misc/SecureHash.h"
+#include "Misc/MemStackUtility.h"
 #include "Misc/LargeWorldRenderPosition.h"
 
 namespace UE
@@ -8,6 +10,81 @@ namespace UE
 namespace Shader
 {
 
+const TCHAR* FType::GetName() const
+{
+	if (StructType)
+	{
+		check(ValueType == EValueType::Struct);
+		return StructType->Name;
+	}
+
+	check(ValueType != EValueType::Struct);
+	const FValueTypeDescription TypeDesc = GetValueTypeDescription(ValueType);
+	return TypeDesc.Name;
+}
+
+int32 FType::GetNumComponents() const
+{
+	if (StructType)
+	{
+		check(ValueType == EValueType::Struct);
+		return StructType->ComponentTypes.Num();
+	}
+
+	check(ValueType != EValueType::Struct);
+	const FValueTypeDescription TypeDesc = GetValueTypeDescription(ValueType);
+	return TypeDesc.NumComponents;
+}
+
+EValueComponentType FType::GetComponentType(int32 Index) const
+{
+	if (StructType)
+	{
+		check(ValueType == EValueType::Struct);
+		return StructType->ComponentTypes.IsValidIndex(Index) ? StructType->ComponentTypes[Index] : EValueComponentType::Void;
+	}
+	const FValueTypeDescription TypeDesc = GetValueTypeDescription(ValueType);
+	return (Index >= 0 && Index < TypeDesc.NumComponents) ? TypeDesc.ComponentType : EValueComponentType::Void;
+}
+
+bool FType::Merge(const FType& OtherType)
+{
+	if (ValueType == EValueType::Void)
+	{
+		ValueType = OtherType.ValueType;
+		StructType = OtherType.StructType;
+		return true;
+	}
+
+	if (IsStruct() || OtherType.IsStruct())
+	{
+		return StructType == OtherType.StructType;
+	}
+
+	if (ValueType != OtherType.ValueType)
+	{
+		const FValueTypeDescription TypeDesc = GetValueTypeDescription(ValueType);
+		const FValueTypeDescription OtherTypeDesc = GetValueTypeDescription(OtherType);
+		const EValueComponentType ComponentType = CombineComponentTypes(TypeDesc.ComponentType, OtherTypeDesc.ComponentType);
+		const int8 NumComponents = FMath::Max(TypeDesc.NumComponents, OtherTypeDesc.NumComponents);
+		ValueType = MakeValueType(ComponentType, NumComponents);
+	}
+	return true;
+}
+
+const FStructField* FStructType::FindFieldByName(const TCHAR* InName) const
+{
+	for (const FStructField& Field : Fields)
+	{
+		if (FCString::Strcmp(Field.Name, InName) == 0)
+		{
+			return &Field;
+		}
+	}
+
+	return nullptr;
+}
+
 namespace Private
 {
 struct FCastFloat
@@ -78,11 +155,10 @@ template<typename Operation, typename ResultType>
 void AsType(const Operation& Op, const FValue& Value, ResultType& OutResult)
 {
 	using FComponentType = typename Operation::FComponentType;
-	const EValueComponentType ComponentType = Value.ComponentType;
-	const int32 NumComponents = Value.NumComponents;
-	if (NumComponents == 1)
+	const FValueTypeDescription TypeDesc = GetValueTypeDescription(Value.Type);
+	if (TypeDesc.NumComponents == 1)
 	{
-		const FComponentType Component = Op(ComponentType, Value.Component[0]);
+		const FComponentType Component = Op(TypeDesc.ComponentType, Value.Component[0]);
 		for (int32 i = 0; i < 4; ++i)
 		{
 			OutResult[i] = Component;
@@ -90,11 +166,11 @@ void AsType(const Operation& Op, const FValue& Value, ResultType& OutResult)
 	}
 	else
 	{
-		for (int32 i = 0; i < NumComponents; ++i)
+		for (int32 i = 0; i < TypeDesc.NumComponents; ++i)
 		{
-			OutResult[i] = Op(ComponentType, Value.Component[i]);
+			OutResult[i] = Op(TypeDesc.ComponentType, Value.Component[i]);
 		}
-		for (int32 i = NumComponents; i < 4; ++i)
+		for (int32 i = TypeDesc.NumComponents; i < 4; ++i)
 		{
 			OutResult[i] = (FComponentType)0;
 		}
@@ -105,19 +181,20 @@ template<typename Operation>
 void Cast(const Operation& Op, const FValue& Value, FValue& OutResult)
 {
 	using FComponentType = typename Operation::FComponentType;
-	const EValueComponentType ComponentType = Value.ComponentType;
-	const int32 NumResultComponents = OutResult.NumComponents;
-	const int32 NumCopyComponents = FMath::Min<int32>(Value.NumComponents, NumResultComponents);
+	const FValueTypeDescription ValueTypeDesc = GetValueTypeDescription(Value.Type);
+	const FValueTypeDescription ResultTypeDesc = GetValueTypeDescription(OutResult.Type);
+	const int32 NumCopyComponents = FMath::Min<int32>(ValueTypeDesc.NumComponents, ResultTypeDesc.NumComponents);
 	for (int32 i = 0; i < NumCopyComponents; ++i)
 	{
-		OutResult.Component[i] = Op(ComponentType, Value.Component[i]);
+		OutResult.Component.Add(Op(ValueTypeDesc.ComponentType, Value.Component[i]));
 	}
 
-	if (NumCopyComponents == 1 && NumResultComponents > 1)
+	if (NumCopyComponents == 1 && ResultTypeDesc.NumComponents > 1)
 	{
-		for (int32 i = 1; i < NumResultComponents; ++i)
+		const FValueComponent Component = OutResult.Component[0];
+		for (int32 i = 1; i < ResultTypeDesc.NumComponents; ++i)
 		{
-			OutResult.Component[i] = OutResult.Component[0];
+			OutResult.Component.Add(Component);
 		}
 	}
 }
@@ -133,10 +210,10 @@ void FormatComponent_Double(double Value, int32 NumComponents, EValueStringForma
 		// Shorter format for more components
 		switch (NumComponents)
 		{
-		case 4: OutResult.Appendf(TEXT("%.2g"), Value); break;
+		default: OutResult.Appendf(TEXT("%.2g"), Value); break;
 		case 3: OutResult.Appendf(TEXT("%.3g"), Value); break;
 		case 2: OutResult.Appendf(TEXT("%.3g"), Value); break;
-		default: OutResult.Appendf(TEXT("%.4g"), Value); break;
+		case 1: OutResult.Appendf(TEXT("%.4g"), Value); break;
 		}
 	}
 }
@@ -145,12 +222,14 @@ void FormatComponent_Double(double Value, int32 NumComponents, EValueStringForma
 
 FValue FValue::FromMemoryImage(EValueType Type, const void* Data, uint32* OutSizeInBytes)
 {
-	FValue Result(Type);
+	const FValueTypeDescription TypeDesc = GetValueTypeDescription(Type);
+
+	FValue Result(TypeDesc.ComponentType, TypeDesc.NumComponents);
 	const uint8* Bytes = static_cast<const uint8*>(Data);
-	const uint32 ComponentSizeInBytes = GetComponentTypeSizeInBytes(Result.ComponentType);
+	const uint32 ComponentSizeInBytes = GetComponentTypeSizeInBytes(TypeDesc.ComponentType);
 	if (ComponentSizeInBytes > 0u)
 	{
-		for (int32 i = 0u; i < Result.NumComponents; ++i)
+		for (int32 i = 0u; i < TypeDesc.NumComponents; ++i)
 		{
 			FMemory::Memcpy(&Result.Component[i].Packed, Bytes, ComponentSizeInBytes);
 			Bytes += ComponentSizeInBytes;
@@ -165,12 +244,15 @@ FValue FValue::FromMemoryImage(EValueType Type, const void* Data, uint32* OutSiz
 
 FMemoryImageValue FValue::AsMemoryImage() const
 {
+	check(!Type.IsStruct());
+	const FValueTypeDescription TypeDesc = GetValueTypeDescription(Type);
+
 	FMemoryImageValue Result;
 	uint8* Bytes = Result.Bytes;
-	const uint32 ComponentSizeInBytes = GetComponentTypeSizeInBytes(ComponentType);
+	const uint32 ComponentSizeInBytes = GetComponentTypeSizeInBytes(TypeDesc.ComponentType);
 	if (ComponentSizeInBytes > 0u)
 	{
-		for (int32 i = 0u; i < NumComponents; ++i)
+		for (int32 i = 0u; i < TypeDesc.NumComponents; ++i)
 		{
 			FMemory::Memcpy(Bytes, &Component[i].Packed, ComponentSizeInBytes);
 			Bytes += ComponentSizeInBytes;
@@ -231,7 +313,7 @@ float FValue::AsFloatScalar() const
 bool FValue::AsBoolScalar() const
 {
 	const FBoolValue Result = AsBool();
-	for (int32 i = 0; i < NumComponents; ++i)
+	for (int32 i = 0; i < 4; ++i)
 	{
 		if (Result.Component[i])
 		{
@@ -254,6 +336,34 @@ uint32 GetComponentTypeSizeInBytes(EValueComponentType Type)
 	}
 }
 
+EValueComponentType CombineComponentTypes(EValueComponentType Lhs, EValueComponentType Rhs)
+{
+	if (Lhs == Rhs)
+	{
+		return Lhs;
+	}
+	else if (Lhs == EValueComponentType::Void)
+	{
+		return Rhs;
+	}
+	else if (Rhs == EValueComponentType::Void)
+	{
+		return Lhs;
+	}
+	else if (Lhs == EValueComponentType::Double || Rhs == EValueComponentType::Double)
+	{
+		return EValueComponentType::Double;
+	}
+	else if (Lhs == EValueComponentType::Float || Rhs == EValueComponentType::Float)
+	{
+		return EValueComponentType::Float;
+	}
+	else
+	{
+		return EValueComponentType::Int;
+	}
+}
+
 const TCHAR* FValueComponent::ToString(EValueComponentType Type, FStringBuilderBase& OutString) const
 {
 	switch (Type)
@@ -268,7 +378,54 @@ const TCHAR* FValueComponent::ToString(EValueComponentType Type, FStringBuilderB
 
 const TCHAR* FValue::ToString(EValueStringFormat Format, FStringBuilderBase& OutString) const
 {
-	if (Format == EValueStringFormat::HLSL && ComponentType == EValueComponentType::Double)
+	const int32 NumComponents = Type.GetNumComponents();
+	const TCHAR* TypeName = Type.GetName();
+	const TCHAR* ClosingSuffix = nullptr;
+
+	if (Format == EValueStringFormat::HLSL)
+	{
+		if (Type.IsStruct())
+		{
+			OutString.Append(TEXT("{ "));
+			ClosingSuffix = TEXT(" }");
+		}
+		else
+		{
+			const FValueTypeDescription TypeDesc = GetValueTypeDescription(Type.ValueType);
+			if (TypeDesc.ComponentType != EValueComponentType::Double)
+			{
+				OutString.Appendf(TEXT("%s("), TypeDesc.Name);
+				ClosingSuffix = TEXT(")");
+			}
+		}
+	}
+
+	for (int32 Index = 0; Index < NumComponents; ++Index)
+	{
+		if (Index > 0)
+		{
+			OutString.Append(TEXT(", "));
+		}
+		const EValueComponentType ComponentType = Type.GetComponentType(Index);
+		switch (ComponentType)
+		{
+		case EValueComponentType::Int: OutString.Appendf(TEXT("%d"), Component[Index].Int); break;
+		case EValueComponentType::Bool: OutString.Append(Component[Index].Bool ? TEXT("true") : TEXT("false")); break;
+		case EValueComponentType::Float: Private::FormatComponent_Double((double)Component[Index].Float, NumComponents, Format, OutString); break;
+		case EValueComponentType::Double: Private::FormatComponent_Double(Component[Index].Double, NumComponents, Format, OutString); break;
+		default: checkNoEntry(); break;
+		}
+	}
+
+	if (ClosingSuffix)
+	{
+		OutString.Append(ClosingSuffix);
+	}
+
+
+
+
+	/*if (Format == EValueStringFormat::HLSL && ComponentType == EValueComponentType::Double)
 	{
 		// Construct an HLSL LWC Vector
 		TStringBuilder<256> TileValue;
@@ -338,7 +495,7 @@ const TCHAR* FValue::ToString(EValueStringFormat Format, FStringBuilderBase& Out
 		{
 			OutString.Append(TEXT(")"));
 		}
-	}
+	}*/
 
 	return OutString.ToString();
 }
@@ -510,6 +667,112 @@ EValueType MakeComparisonResultType(EValueType Lhs, EValueType Rhs, FString& Out
 	return EValueType::Void;
 }
 
+FStructTypeRegistry::FStructTypeRegistry(FMemStackBase& InAllocator)
+	: Allocator(&InAllocator)
+{
+}
+
+void FStructTypeRegistry::EmitDeclarationsCode(FStringBuilderBase& OutCode) const
+{
+	for (const auto& It : Types)
+	{
+		const FStructType* StructType = It.Value;
+
+		OutCode.Appendf(TEXT("struct %s\n"), StructType->Name);
+		OutCode.Append(TEXT("{\n"));
+		for (const FStructField& Field : StructType->Fields)
+		{
+			OutCode.Appendf(TEXT("\t%s %s;\n"), Field.Type.GetName(), Field.Name);
+		}
+		OutCode.Append(TEXT("};\n"));
+
+		for (const FStructField& Field : StructType->Fields)
+		{
+			OutCode.Appendf(TEXT("%s %s_Set%s(%s Self, %s Value) { Self.%s = Value; return Self; }\n"),
+				StructType->Name, StructType->Name, Field.Name, StructType->Name, Field.Type.GetName(), Field.Name);
+		}
+		OutCode.Append(TEXT("\n"));
+	}
+}
+
+namespace Private
+{
+void SetComponentTypes(EValueComponentType* ComponentTypes, int32 ComponentIndex, const FType& Type)
+{
+	if (Type.IsStruct())
+	{
+		for (const FStructField& Field : Type.StructType->Fields)
+		{
+			SetComponentTypes(ComponentTypes, ComponentIndex + Field.ComponentIndex, Field.Type);
+		}
+	}
+	else
+	{
+		const FValueTypeDescription TypeDesc = GetValueTypeDescription(Type);
+		for (int32 i = 0; i < TypeDesc.NumComponents; ++i)
+		{
+			ComponentTypes[ComponentIndex + i] = TypeDesc.ComponentType;
+		}
+	}
+}
+} // namespace Private
+
+const FStructType* FStructTypeRegistry::NewType(const FStructTypeInitializer& Initializer)
+{
+	FSHA1 Hasher;
+	Hasher.UpdateWithString(Initializer.Name.GetData(), Initializer.Name.Len());
+
+	const int32 NumFields = Initializer.Fields.Num();
+	int32 ComponentIndex = 0;
+	FStructField* Fields = new(*Allocator) FStructField[NumFields];
+	for (int32 FieldIndex = 0; FieldIndex < NumFields; ++FieldIndex)
+	{
+		const FStructFieldInitializer& FieldInitializer = Initializer.Fields[FieldIndex];
+		const FType& FieldType = FieldInitializer.Type;
+
+		Hasher.UpdateWithString(FieldInitializer.Name.GetData(), FieldInitializer.Name.Len());
+		if (FieldType.IsStruct())
+		{
+			Hasher.Update((uint8*)&FieldType.StructType->Hash, sizeof(FieldType.StructType->Hash));
+		}
+		else
+		{
+			Hasher.Update((uint8*)&FieldType.ValueType, sizeof(FieldType.ValueType));
+		}
+
+		FStructField& Field = Fields[FieldIndex];
+		Field.Name = MemStack::AllocateString(*Allocator, FieldInitializer.Name);
+		Field.Type = FieldType;
+		Field.ComponentIndex = ComponentIndex;
+		ComponentIndex += FieldType.GetNumComponents();
+	}
+
+	EValueComponentType* ComponentTypes = new(*Allocator) EValueComponentType[ComponentIndex];
+	for (int32 FieldIndex = 0; FieldIndex < NumFields; ++FieldIndex)
+	{
+		const FStructField& Field = Fields[FieldIndex];
+		Private::SetComponentTypes(ComponentTypes, Field.ComponentIndex, Field.Type);
+	}
+
+	const FSHAHash Hash = Hasher.Finalize();
+
+	FStructType* StructType = new(*Allocator) FStructType();
+	StructType->Name = MemStack::AllocateString(*Allocator, Initializer.Name);
+	StructType->Hash = *reinterpret_cast<const uint64*>(Hash.Hash);
+	StructType->Fields = MakeArrayView(Fields, NumFields);
+	StructType->ComponentTypes = MakeArrayView(ComponentTypes, ComponentIndex);
+
+	Types.Add(StructType->Hash, StructType);
+
+	return StructType;
+}
+
+const FStructType* FStructTypeRegistry::FindType(uint64 Hash) const
+{
+	FStructType const* const* PrevType = Types.Find(Hash);
+	return PrevType ? *PrevType : nullptr;
+}
+
 namespace Private
 {
 
@@ -594,20 +857,23 @@ struct FOpAtan2 : public FOpBaseNoInt { template<typename T> T operator()(T Lhs,
 template<typename Operation>
 inline FValue UnaryOp(const Operation& Op, const FValue& Value)
 {
-	const int8 NumComponents = Value.NumComponents;
+	if (Value.Type.IsStruct())
+	{
+		return FValue();
+	}
+	const FValueTypeDescription TypeDesc = GetValueTypeDescription(Value.Type);
+	const int8 NumComponents = TypeDesc.NumComponents;
 	
 	FValue Result;
-	Result.NumComponents = NumComponents;
-
 	if constexpr (Operation::SupportsDouble)
 	{
-		if (Value.ComponentType == EValueComponentType::Double)
+		if (TypeDesc.ComponentType == EValueComponentType::Double)
 		{
-			Result.ComponentType = EValueComponentType::Double;
+			Result.Type = MakeValueType(EValueComponentType::Double, NumComponents);
 			const FDoubleValue Cast = Value.AsDouble();
 			for (int32 i = 0; i < NumComponents; ++i)
 			{
-				Result.Component[i].Double = Op(Cast.Component[i]);
+				Result.Component.Add(Op(Cast.Component[i]));
 			}
 			return Result;
 		}
@@ -615,23 +881,23 @@ inline FValue UnaryOp(const Operation& Op, const FValue& Value)
 
 	if constexpr (Operation::SupportsInt)
 	{
-		if (Value.ComponentType != EValueComponentType::Float)
+		if (TypeDesc.ComponentType != EValueComponentType::Float)
 		{
-			Result.ComponentType = EValueComponentType::Int;
+			Result.Type = MakeValueType(EValueComponentType::Int, NumComponents);
 			const FIntValue Cast = Value.AsInt();
 			for (int32 i = 0; i < NumComponents; ++i)
 			{
-				Result.Component[i].Int = Op(Cast.Component[i]);
+				Result.Component.Add(Op(Cast.Component[i]));
 			}
 			return Result;
 		}
 	}
 
-	Result.ComponentType = EValueComponentType::Float;
+	Result.Type = MakeValueType(EValueComponentType::Float, NumComponents);
 	const FFloatValue Cast = Value.AsFloat();
 	for (int32 i = 0; i < NumComponents; ++i)
 	{
-		Result.Component[i].Float = Op(Cast.Component[i]);
+		Result.Component.Add(Op(Cast.Component[i]));
 	}
 	return Result;
 }
@@ -646,21 +912,25 @@ inline int8 GetNumComponentsResult(int8 Lhs, int8 Rhs)
 template<typename Operation>
 inline FValue BinaryOp(const Operation& Op, const FValue& Lhs, const FValue& Rhs)
 {
-	const int8 NumComponents = GetNumComponentsResult(Lhs.NumComponents, Rhs.NumComponents);
+	if (Lhs.Type.IsStruct() || Rhs.Type.IsStruct())
+	{
+		return FValue();
+	}
+	const FValueTypeDescription LhsDesc = GetValueTypeDescription(Lhs.Type);
+	const FValueTypeDescription RhsDesc = GetValueTypeDescription(Rhs.Type);
+	const int8 NumComponents = GetNumComponentsResult(LhsDesc.NumComponents, RhsDesc.NumComponents);
 	
 	FValue Result;
-	Result.NumComponents = NumComponents;
-
 	if constexpr (Operation::SupportsDouble)
 	{
-		if (Lhs.ComponentType == EValueComponentType::Double || Rhs.ComponentType == EValueComponentType::Double)
+		if (LhsDesc.ComponentType == EValueComponentType::Double || RhsDesc.ComponentType == EValueComponentType::Double)
 		{
-			Result.ComponentType = EValueComponentType::Double;
+			Result.Type = MakeValueType(EValueComponentType::Double, NumComponents);
 			const FDoubleValue LhsCast = Lhs.AsDouble();
 			const FDoubleValue RhsCast = Rhs.AsDouble();
 			for (int32 i = 0; i < NumComponents; ++i)
 			{
-				Result.Component[i].Double = Op(LhsCast.Component[i], RhsCast.Component[i]);
+				Result.Component.Add(Op(LhsCast.Component[i], RhsCast.Component[i]));
 			}
 			return Result;
 		}
@@ -668,25 +938,25 @@ inline FValue BinaryOp(const Operation& Op, const FValue& Lhs, const FValue& Rhs
 
 	if constexpr (Operation::SupportsInt)
 	{
-		if (Lhs.ComponentType != EValueComponentType::Float && Rhs.ComponentType != EValueComponentType::Float)
+		if (LhsDesc.ComponentType != EValueComponentType::Float && RhsDesc.ComponentType != EValueComponentType::Float)
 		{
-			Result.ComponentType = EValueComponentType::Int;
+			Result.Type = MakeValueType(EValueComponentType::Int, NumComponents);
 			const FIntValue LhsCast = Lhs.AsInt();
 			const FIntValue RhsCast = Rhs.AsInt();
 			for (int32 i = 0; i < NumComponents; ++i)
 			{
-				Result.Component[i].Int = Op(LhsCast.Component[i], RhsCast.Component[i]);
+				Result.Component.Add(Op(LhsCast.Component[i], RhsCast.Component[i]));
 			}
 			return Result;
 		}
 	}
 
-	Result.ComponentType = EValueComponentType::Float;
+	Result.Type = MakeValueType(EValueComponentType::Float, NumComponents);
 	const FFloatValue LhsCast = Lhs.AsFloat();
 	const FFloatValue RhsCast = Rhs.AsFloat();
 	for (int32 i = 0; i < NumComponents; ++i)
 	{
-		Result.Component[i].Float = Op(LhsCast.Component[i], RhsCast.Component[i]);
+		Result.Component.Add(Op(LhsCast.Component[i], RhsCast.Component[i]));
 	}
 	return Result;
 }
@@ -695,11 +965,13 @@ inline FValue BinaryOp(const Operation& Op, const FValue& Lhs, const FValue& Rhs
 
 bool operator==(const FValue& Lhs, const FValue& Rhs)
 {
-	if (Lhs.ComponentType != Rhs.ComponentType || Lhs.NumComponents != Rhs.NumComponents)
+	if (Lhs.Type != Rhs.Type)
 	{
 		return false;
 	}
-	for (int32 i = 0u; i < Lhs.NumComponents; ++i)
+
+	check(Lhs.Component.Num() == Rhs.Component.Num());
+	for (int32 i = 0u; i < Lhs.Component.Num(); ++i)
 	{
 		if (Lhs.Component[i].Packed != Rhs.Component[i].Packed)
 		{
@@ -709,20 +981,29 @@ bool operator==(const FValue& Lhs, const FValue& Rhs)
 	return true;
 }
 
+uint32 GetTypeHash(const FType& Type)
+{
+	uint32 Result = ::GetTypeHash(Type.ValueType);
+	if (Type.IsStruct())
+	{
+		Result = HashCombine(Result, ::GetTypeHash(Type.StructType));
+	}
+	return Result;
+}
+
 uint32 GetTypeHash(const FValue& Value)
 {
-	uint32 Result = ::GetTypeHash(Value.ComponentType);
-	Result = HashCombine(Result, ::GetTypeHash(Value.NumComponents));
-	for (int32 i = 0u; i < Value.NumComponents; ++i)
+	uint32 Result = GetTypeHash(Value.Type);
+	const int32 NumComponents = Value.Type.GetNumComponents();
+	for(int32 Index = 0; Index < NumComponents; ++Index)
 	{
-		const FValueComponent& Component = Value.Component[i];
 		uint32 ComponentHash = 0u;
-		switch (Value.ComponentType)
+		switch (Value.Type.GetComponentType(Index))
 		{
-		case EValueComponentType::Float: ComponentHash = ::GetTypeHash(Component.Float); break;
-		case EValueComponentType::Double: ComponentHash = ::GetTypeHash(Component.Double); break;
-		case EValueComponentType::Int: ComponentHash = ::GetTypeHash(Component.Int); break;
-		case EValueComponentType::Bool: ComponentHash = ::GetTypeHash(Component.Bool); break;
+		case EValueComponentType::Float: ComponentHash = ::GetTypeHash(Value.Component[Index].Float); break;
+		case EValueComponentType::Double: ComponentHash = ::GetTypeHash(Value.Component[Index].Double); break;
+		case EValueComponentType::Int: ComponentHash = ::GetTypeHash(Value.Component[Index].Int); break;
+		case EValueComponentType::Bool: ComponentHash = ::GetTypeHash(Value.Component[Index].Bool); break;
 		default: checkNoEntry(); break;
 		}
 		Result = HashCombine(Result, ComponentHash);
@@ -872,14 +1153,18 @@ FValue Clamp(const FValue& Value, const FValue& Low, const FValue& High)
 
 FValue Dot(const FValue& Lhs, const FValue& Rhs)
 {
-	const int8 NumComponents = Private::GetNumComponentsResult(Lhs.NumComponents, Rhs.NumComponents);
+	if (Lhs.Type.IsStruct() || Rhs.Type.IsStruct())
+	{
+		return FValue();
+	}
+	const FValueTypeDescription LhsDesc = GetValueTypeDescription(Lhs.Type);
+	const FValueTypeDescription RhsDesc = GetValueTypeDescription(Rhs.Type);
+	const int8 NumComponents = Private::GetNumComponentsResult(LhsDesc.NumComponents, RhsDesc.NumComponents);
 
 	FValue Result;
-	Result.NumComponents = 1;
-
-	if (Lhs.ComponentType == EValueComponentType::Double || Rhs.ComponentType == EValueComponentType::Double)
+	if (LhsDesc.ComponentType == EValueComponentType::Double || RhsDesc.ComponentType == EValueComponentType::Double)
 	{
-		Result.ComponentType = EValueComponentType::Double;
+		Result.Type = EValueType::Double1;
 		const FDoubleValue LhsValue = Lhs.AsDouble();
 		const FDoubleValue RhsValue = Rhs.AsDouble();
 		double ComponentValue = 0.0;
@@ -887,11 +1172,11 @@ FValue Dot(const FValue& Lhs, const FValue& Rhs)
 		{
 			ComponentValue += LhsValue.Component[i] * RhsValue.Component[i];
 		}
-		Result.Component[0].Double = ComponentValue;
+		Result.Component.Add(ComponentValue);
 	}
-	else if (Lhs.ComponentType == EValueComponentType::Float || Rhs.ComponentType == EValueComponentType::Float)
+	else if (LhsDesc.ComponentType == EValueComponentType::Float || RhsDesc.ComponentType == EValueComponentType::Float)
 	{
-		Result.ComponentType = EValueComponentType::Float;
+		Result.Type = EValueType::Float1;
 		const FFloatValue LhsValue = Lhs.AsFloat();
 		const FFloatValue RhsValue = Rhs.AsFloat();
 		float ComponentValue = 0.0f;
@@ -899,11 +1184,11 @@ FValue Dot(const FValue& Lhs, const FValue& Rhs)
 		{
 			ComponentValue += LhsValue.Component[i] * RhsValue.Component[i];
 		}
-		Result.Component[0].Float = ComponentValue;
+		Result.Component.Add(ComponentValue);
 	}
 	else
 	{
-		Result.ComponentType = EValueComponentType::Int;
+		Result.Type = EValueType::Int1;
 		const FIntValue LhsValue = Lhs.AsInt();
 		const FIntValue RhsValue = Rhs.AsInt();
 		int32 ComponentValue = 0;
@@ -911,110 +1196,122 @@ FValue Dot(const FValue& Lhs, const FValue& Rhs)
 		{
 			ComponentValue += LhsValue.Component[i] * RhsValue.Component[i];
 		}
-		Result.Component[0].Int = ComponentValue;
+		Result.Component.Add(ComponentValue);
 	}
 	return Result;
 }
 
 FValue Cross(const FValue& Lhs, const FValue& Rhs)
 {
-	FValue Result;
-	Result.NumComponents = 3;
+	if (Lhs.Type.IsStruct() || Rhs.Type.IsStruct())
+	{
+		return FValue();
+	}
+	const FValueTypeDescription LhsDesc = GetValueTypeDescription(Lhs.Type);
+	const FValueTypeDescription RhsDesc = GetValueTypeDescription(Rhs.Type);
 
-	if (Lhs.ComponentType == EValueComponentType::Double || Rhs.ComponentType == EValueComponentType::Double)
+	FValue Result;
+	if (LhsDesc.ComponentType == EValueComponentType::Double || RhsDesc.ComponentType == EValueComponentType::Double)
 	{
-		Result.ComponentType = EValueComponentType::Double;
+		Result.Type = EValueType::Double3;
 		const FDoubleValue LhsValue = Lhs.AsDouble();
 		const FDoubleValue RhsValue = Rhs.AsDouble();
 
-		Result.Component[0].Double = LhsValue.Component[1] * RhsValue.Component[2] - LhsValue.Component[2] * RhsValue.Component[1];
-		Result.Component[1].Double = LhsValue.Component[2] * RhsValue.Component[0] - LhsValue.Component[0] * RhsValue.Component[2];
-		Result.Component[2].Double = LhsValue.Component[0] * RhsValue.Component[1] - LhsValue.Component[1] * RhsValue.Component[0];
+		Result.Component.Add(LhsValue.Component[1] * RhsValue.Component[2] - LhsValue.Component[2] * RhsValue.Component[1]);
+		Result.Component.Add(LhsValue.Component[2] * RhsValue.Component[0] - LhsValue.Component[0] * RhsValue.Component[2]);
+		Result.Component.Add(LhsValue.Component[0] * RhsValue.Component[1] - LhsValue.Component[1] * RhsValue.Component[0]);
 	}
-	else if (Lhs.ComponentType == EValueComponentType::Float || Rhs.ComponentType == EValueComponentType::Float)
+	else if (LhsDesc.ComponentType == EValueComponentType::Float || RhsDesc.ComponentType == EValueComponentType::Float)
 	{
-		Result.ComponentType = EValueComponentType::Float;
+		Result.Type = EValueType::Float3;
 		const FFloatValue LhsValue = Lhs.AsFloat();
 		const FFloatValue RhsValue = Rhs.AsFloat();
 		
-		Result.Component[0].Float = LhsValue.Component[1] * RhsValue.Component[2] - LhsValue.Component[2] * RhsValue.Component[1];
-		Result.Component[1].Float = LhsValue.Component[2] * RhsValue.Component[0] - LhsValue.Component[0] * RhsValue.Component[2];
-		Result.Component[2].Float = LhsValue.Component[0] * RhsValue.Component[1] - LhsValue.Component[1] * RhsValue.Component[0];
+		Result.Component.Add(LhsValue.Component[1] * RhsValue.Component[2] - LhsValue.Component[2] * RhsValue.Component[1]);
+		Result.Component.Add(LhsValue.Component[2] * RhsValue.Component[0] - LhsValue.Component[0] * RhsValue.Component[2]);
+		Result.Component.Add(LhsValue.Component[0] * RhsValue.Component[1] - LhsValue.Component[1] * RhsValue.Component[0]);
 	}
 	else
 	{
-		Result.ComponentType = EValueComponentType::Int;
+		Result.Type = EValueType::Int3;
 		const FIntValue LhsValue = Lhs.AsInt();
 		const FIntValue RhsValue = Rhs.AsInt();
 
-		Result.Component[0].Int = LhsValue.Component[1] * RhsValue.Component[2] - LhsValue.Component[2] * RhsValue.Component[1];
-		Result.Component[1].Int = LhsValue.Component[2] * RhsValue.Component[0] - LhsValue.Component[0] * RhsValue.Component[2];
-		Result.Component[2].Int = LhsValue.Component[0] * RhsValue.Component[1] - LhsValue.Component[1] * RhsValue.Component[0];
+		Result.Component.Add(LhsValue.Component[1] * RhsValue.Component[2] - LhsValue.Component[2] * RhsValue.Component[1]);
+		Result.Component.Add(LhsValue.Component[2] * RhsValue.Component[0] - LhsValue.Component[0] * RhsValue.Component[2]);
+		Result.Component.Add(LhsValue.Component[0] * RhsValue.Component[1] - LhsValue.Component[1] * RhsValue.Component[0]);
 	}
 	return Result;
 }
 
 FValue Append(const FValue& Lhs, const FValue& Rhs)
 {
+	if (Lhs.Type.IsStruct() || Rhs.Type.IsStruct())
+	{
+		return FValue();
+	}
+	const FValueTypeDescription LhsDesc = GetValueTypeDescription(Lhs.Type);
+	const FValueTypeDescription RhsDesc = GetValueTypeDescription(Rhs.Type);
+
 	FValue Result;
-	int32 NumComponents = 0;
-	if (Lhs.ComponentType == Rhs.ComponentType)
+	const int32 NumComponents = FMath::Min<int32>(LhsDesc.NumComponents + RhsDesc.NumComponents, 4);
+	if (LhsDesc.ComponentType == RhsDesc.ComponentType)
 	{
 		// If both values have the same component type, use as-is
 		// (otherwise will need to convert)
-		Result.ComponentType = Lhs.ComponentType;
-		for (int32 i = 0; i < Lhs.NumComponents; ++i)
+		
+		Result.Type = MakeValueType(LhsDesc.ComponentType, NumComponents);
+		for (int32 i = 0; i < LhsDesc.NumComponents && Result.Component.Num() < NumComponents; ++i)
 		{
-			Result.Component[NumComponents++] = Lhs.Component[i];
+			Result.Component.Add(Lhs.Component[i]);
 		}
-		for (int32 i = 0; i < Rhs.NumComponents && NumComponents < 4; ++i)
+		for (int32 i = 0; i < RhsDesc.NumComponents && Result.Component.Num() < NumComponents; ++i)
 		{
-			Result.Component[NumComponents++] = Rhs.Component[i];
+			Result.Component.Add(Rhs.Component[i]);
 		}
 	}
-	else if (Lhs.ComponentType == EValueComponentType::Double || Rhs.ComponentType == EValueComponentType::Double)
+	else if (LhsDesc.ComponentType == EValueComponentType::Double || RhsDesc.ComponentType == EValueComponentType::Double)
 	{
-		Result.ComponentType = EValueComponentType::Double;
+		Result.Type = MakeValueType(EValueComponentType::Double, NumComponents);
 		const FDoubleValue LhsValue = Lhs.AsDouble();
 		const FDoubleValue RhsValue = Rhs.AsDouble();
-		for (int32 i = 0; i < Lhs.NumComponents; ++i)
+		for (int32 i = 0; i < LhsDesc.NumComponents && Result.Component.Num() < NumComponents; ++i)
 		{
-			Result.Component[NumComponents++].Double = LhsValue.Component[i];
+			Result.Component.Add(LhsValue.Component[i]);
 		}
-		for (int32 i = 0; i < Rhs.NumComponents && NumComponents < 4; ++i)
+		for (int32 i = 0; i < RhsDesc.NumComponents && Result.Component.Num() < NumComponents; ++i)
 		{
-			Result.Component[NumComponents++].Double = RhsValue.Component[i];
+			Result.Component.Add(RhsValue.Component[i]);
 		}
 	}
-	else if (Lhs.ComponentType == EValueComponentType::Float || Rhs.ComponentType == EValueComponentType::Float)
+	else if (LhsDesc.ComponentType == EValueComponentType::Float || RhsDesc.ComponentType == EValueComponentType::Float)
 	{
-		Result.ComponentType = EValueComponentType::Float;
+		Result.Type = MakeValueType(EValueComponentType::Float, NumComponents);
 		const FFloatValue LhsValue = Lhs.AsFloat();
 		const FFloatValue RhsValue = Rhs.AsFloat();
-		for (int32 i = 0; i < Lhs.NumComponents; ++i)
+		for (int32 i = 0; i < LhsDesc.NumComponents && Result.Component.Num() < NumComponents; ++i)
 		{
-			Result.Component[NumComponents++].Float = LhsValue.Component[i];
+			Result.Component.Add(LhsValue.Component[i]);
 		}
-		for (int32 i = 0; i < Rhs.NumComponents && NumComponents < 4; ++i)
+		for (int32 i = 0; i < RhsDesc.NumComponents && Result.Component.Num() < NumComponents; ++i)
 		{
-			Result.Component[NumComponents++].Float = RhsValue.Component[i];
+			Result.Component.Add(RhsValue.Component[i]);
 		}
 	}
 	else
 	{
-		Result.ComponentType = EValueComponentType::Int;
+		Result.Type = MakeValueType(EValueComponentType::Int, NumComponents);
 		const FIntValue LhsValue = Lhs.AsInt();
 		const FIntValue RhsValue = Rhs.AsInt();
-		for (int32 i = 0; i < Lhs.NumComponents; ++i)
+		for (int32 i = 0; i < LhsDesc.NumComponents && Result.Component.Num() < NumComponents; ++i)
 		{
-			Result.Component[NumComponents++].Int = LhsValue.Component[i];
+			Result.Component.Add(LhsValue.Component[i]);
 		}
-		for (int32 i = 0; i < Rhs.NumComponents && NumComponents < 4; ++i)
+		for (int32 i = 0; i < RhsDesc.NumComponents && Result.Component.Num() < NumComponents; ++i)
 		{
-			Result.Component[NumComponents++].Int = RhsValue.Component[i];
+			Result.Component.Add(RhsValue.Component[i]);
 		}
 	}
-	Result.NumComponents = NumComponents;
 	return Result;
 }
 
@@ -1026,8 +1323,10 @@ FValue Cast(const FValue& Value, EValueType Type)
 		return Value;
 	}
 
-	FValue Result(Type);
-	switch (Result.ComponentType)
+	FValue Result;
+	Result.Type = Type;
+
+	switch (GetValueTypeDescription(Type).ComponentType)
 	{
 	case EValueComponentType::Float: Private::Cast(Private::FCastFloat(), Value, Result); break;
 	case EValueComponentType::Double: Private::Cast(Private::FCastDouble(), Value, Result); break;
diff --git a/Engine/Source/Runtime/Engine/Private/ShaderDerivedDataVersion.h b/Engine/Source/Runtime/Engine/Private/ShaderDerivedDataVersion.h
--- a/Engine/Source/Runtime/Engine/Private/ShaderDerivedDataVersion.h
+++ b/Engine/Source/Runtime/Engine/Private/ShaderDerivedDataVersion.h
@@ -11,4 +11,4 @@
 // In case of merge conflicts with DDC versions, you MUST generate a new GUID and set this new
 // guid as version
 #define GLOBALSHADERMAP_DERIVEDDATA_VER			TEXT("7429FF5E4C57472FA538321D6777E8BC")
-#define MATERIALSHADERMAP_DERIVEDDATA_VER		TEXT("0A688C67C250488F8EFB8DCD02A544C9")
\ No newline at end of file
+#define MATERIALSHADERMAP_DERIVEDDATA_VER		TEXT("C2D6564DC8B14777B5D6747464213D7E")
\ No newline at end of file
diff --git a/Engine/Source/Runtime/Engine/Public/HLSLTree/HLSLTree.h b/Engine/Source/Runtime/Engine/Public/HLSLTree/HLSLTree.h
--- a/Engine/Source/Runtime/Engine/Public/HLSLTree/HLSLTree.h
+++ b/Engine/Source/Runtime/Engine/Public/HLSLTree/HLSLTree.h
@@ -33,54 +33,13 @@ namespace HLSLTree
 
 class FNode;
 class FScope;
+class FExpressionLocalPHI;
+class FRequestedType;
 struct FEmitShaderValue;
+struct FShaderValue;
 
 static constexpr int32 MaxNumPreviousScopes = 2;
 
-struct FType
-{
-	FType() : StructType(nullptr), ValueType(Shader::EValueType::Void) {}
-	FType(Shader::EValueType InValueType) : StructType(nullptr), ValueType(InValueType) {}
-	FType(const FStructType* InStruct) : StructType(InStruct), ValueType(Shader::EValueType::Struct) { check(InStruct); }
-
-	const TCHAR* GetName() const;
-	bool IsStruct() const { return (bool)StructType; }
-	int32 GetNumComponents() const;
-	bool Merge(const FType& OtherType);
-
-	inline operator Shader::EValueType() const { return ValueType; }
-	inline operator bool() const { return ValueType != Shader::EValueType::Void; }
-	inline bool operator!() const { return ValueType == Shader::EValueType::Void; }
-
-	const FStructType* StructType;
-	Shader::EValueType ValueType;
-};
-
-inline bool operator==(const FType& Lhs, const FType& Rhs)
-{
-	if (Lhs.ValueType != Rhs.ValueType) return false;
-	if (Lhs.ValueType == Shader::EValueType::Struct && Lhs.StructType != Rhs.StructType) return false;
-	return true;
-}
-inline bool operator!=(const FType& Lhs, const FType& Rhs)
-{
-	return !operator==(Lhs, Rhs);
-}
-
-struct FConstantValue
-{
-	FConstantValue() = default;
-	FConstantValue(const Shader::FValue& InValue);
-
-	FType Type;
-
-	/**
-	 * For basic types, this will include 1-4 components
-	 * For struct types, will match the flattened list of struct component types
-	 */
-	TArray<Shader::FValueComponent, TInlineAllocator<16>> Component;
-};
-
 struct FError
 {
 	const FError* Next;
@@ -134,7 +93,7 @@ struct FEmitShaderValueContext
 class FEmitContext
 {
 public:
-	explicit FEmitContext(FMemStackBase& InAllocator);
+	explicit FEmitContext(FMemStackBase& InAllocator, const Shader::FStructTypeRegistry& InTypeRegistry);
 	~FEmitContext();
 
 	void Finalize();
@@ -142,12 +101,15 @@ public:
 	/** Get a unique local variable name */
 	const TCHAR* AcquireLocalDeclarationCode();
 
-	const TCHAR* CastShaderValue(const FNode* Node, const TCHAR* Code, const FType& SourceType, const FType& DestType, ECastFlags Flags);
+	FEmitShaderValue* AcquireShader(FScope* Scope, const FShaderValue& Shader, TArrayView<FEmitShaderValue*> Dependencies);
+	FEmitShaderValue* AcquirePreshader(const FRequestedType& RequestedType, FScope* Scope, FExpression* Expression);
 
-	const TCHAR* AcquirePreshader(Shader::EValueType Type, const Shader::FPreshaderData& Preshader);
+	FEmitShaderValue* CastShaderValue(FNode* Node, FScope* Scope, FEmitShaderValue* ShaderValue, const Shader::FType& DestType);
 
 	FMemStackBase* Allocator = nullptr;
+	const Shader::FStructTypeRegistry* TypeRegistry = nullptr;
 	TMap<FSHAHash, FEmitShaderValue*> ShaderValueMap;
+	TArray<const FExpressionLocalPHI*> LocalPHIs;
 	FErrors Errors;
 
 	// TODO - remove preshader material dependency
@@ -155,14 +117,13 @@ public:
 	const FStaticParameterSet* StaticParameters = nullptr;
 	FMaterialCompilationOutput* MaterialCompilationOutput = nullptr;
 	TMap<Shader::FValue, uint32> DefaultUniformValues;
-	TMap<FSHAHash, const TCHAR*> Preshaders;
+	TMap<FSHAHash, FEmitShaderValue*> Preshaders;
 	TArray<FScope*, TInlineAllocator<16>> ScopeStack;
 	TArray<FEmitShaderValueContext, TInlineAllocator<16>> ShaderValueStack;
 	uint32 UniformPreshaderOffset = 0u;
 	bool bReadMaterialNormal = false;
 
 	int32 NumExpressionLocals = 0;
-	int32 NumLocalPHIs = 0;
 	int32 NumTexCoords = 0;
 };
 
@@ -175,59 +136,6 @@ public:
 	FNode* NextNode = nullptr;
 };
 
-struct FStructField
-{
-	const TCHAR* Name;
-	FType Type;
-	int32 ComponentIndex;
-};
-
-struct FStructFieldRef
-{
-	FStructFieldRef() = default;
-	FStructFieldRef(const FType& InType, int32 InIndex, int32 InNum) : Type(InType), ComponentIndex(InIndex), ComponentNum(InNum) {}
-
-	FType Type;
-	int32 ComponentIndex = INDEX_NONE;
-	int32 ComponentNum = 0;
-
-	inline operator bool() const { return ComponentNum > 0; }
-	inline bool operator!() const { return ComponentNum <= 0; }
-};
-
-class FStructType : public FNode
-{
-public:
-	FStructType* NextType;
-	const TCHAR* Name;
-	TArrayView<const FStructField> Fields;
-
-	/**
-	 * Most code working with HLSLTree views struct types as a flat list of components
-	 * Fields with basic types are represented directly. Fields with struct types are recursively flattened into this list
-	 */
-	TArrayView<const Shader::EValueComponentType> ComponentTypes;
-
-	FStructFieldRef FindFieldByName(const TCHAR* InName) const;
-
-	void WriteHLSL(FStringBuilderBase& OutWriter) const;
-};
-
-struct FStructFieldInitializer
-{
-	FStructFieldInitializer() = default;
-	FStructFieldInitializer(const FStringView& InName, const FType& InType) : Name(InName), Type(InType) {}
-
-	FStringView Name;
-	FType Type;
-};
-
-struct FStructTypeInitializer
-{
-	FStringView Name;
-	TArrayView<const FStructFieldInitializer> Fields;
-};
-
 /**
  * Represents an HLSL statement.  This is a piece of code that doesn't evaluate to any value, but instead should be executed sequentially, and likely has side-effects.
  * Examples include assigning a value, or various control flow structures (if, for, while, etc)
@@ -236,9 +144,6 @@ struct FStructTypeInitializer
 class FStatement : public FNode
 {
 public:
-	//static constexpr bool MarkScopeLive = false;
-	//static constexpr bool MarkScopeLiveRecursive = false;
-
 	virtual void Reset() override;
 	virtual void Prepare(FEmitContext& Context) const = 0;
 	virtual void EmitShader(FEmitContext& Context) const = 0;
@@ -247,92 +152,164 @@ public:
 	bool bEmitShader = false;
 };
 
+/**
+ * Like Shader::FType, but tracks which individual components are needed
+ */
 class FRequestedType
 {
 public:
 	FRequestedType() = default;
 	FRequestedType(int32 NumComponents, bool bDefaultRequest = true) : RequestedComponents(bDefaultRequest, NumComponents) {}
-	FRequestedType(const FType& InType, bool bDefaultRequest = true) : StructType(InType.StructType), RequestedComponents(bDefaultRequest, InType.GetNumComponents()) {}
+	FRequestedType(const Shader::FType& InType, bool bDefaultRequest = true);
+	FRequestedType(const Shader::EValueType& InType, bool bDefaultRequest = true);
+	
+	const TCHAR* GetName() const { return GetType().GetName(); }
+	bool IsStruct() const { return StructType != nullptr; }
+	const Shader::FStructType* GetStructType() const { return StructType; }
+	const Shader::FType GetType() const;
 
-	bool IsStruct() const { return (bool)StructType; }
-	const FStructType* GetStructType() const { return StructType; }
-	int32 GetNumComponents() const { return RequestedComponents.Num(); }
-	int32 GetRequestedNumComponents() const { const int32 LastIndex = RequestedComponents.FindLast(true); return LastIndex + 1; }
+	int32 GetNumComponents() const;
 	bool IsComponentRequested(int32 Index) const { return RequestedComponents.IsValidIndex(Index) ? (bool)RequestedComponents[Index] : false; }
 	bool IsVoid() const { return RequestedComponents.Find(true) == INDEX_NONE; }
 
+	bool Merge(const FRequestedType& OtherType);
+
 	void Reset()
 	{
 		StructType = nullptr;
+		ValueComponentType = Shader::EValueComponentType::Void;
 		RequestedComponents.Reset();
 	}
 
+	void SetComponentRequested(int32 Index, bool bRequested = true)
+	{
+		if (bRequested)
+		{
+			RequestedComponents.PadToNum(Index + 1, false);
+		}
+		if (RequestedComponents.IsValidIndex(Index))
+		{
+			RequestedComponents[Index] = bRequested;
+		}
+	}
+
 	/** Marks the given field as requested (or not) */
-	void SetFieldRequested(const FStructFieldRef& FieldRef, bool bRequested = true)
+	void SetFieldRequested(const Shader::FStructField* Field, bool bRequested = true)
 	{
-		RequestedComponents.SetRange(FieldRef.ComponentIndex, FieldRef.ComponentNum, bRequested);
+		RequestedComponents.SetRange(Field->ComponentIndex, Field->GetNumComponents(), bRequested);
 	}
 
-	void ClearFieldRequested(const FStructFieldRef& FieldRef)
+	void ClearFieldRequested(const Shader::FStructField* Field)
 	{
-		SetFieldRequested(FieldRef, false);
+		SetFieldRequested(Field, false);
 	}
 
 	/** Marks the given field as requested, based on the input request type (which should match the field type) */
-	void SetField(const FStructFieldRef& FieldRef, const FRequestedType& InRequest)
+	void SetField(const Shader::FStructField* Field, const FRequestedType& InRequest)
 	{
-		ensure(InRequest.GetNumComponents() == FieldRef.ComponentNum);
-		RequestedComponents.SetRangeFromRange(FieldRef.ComponentIndex, FieldRef.ComponentNum, InRequest.RequestedComponents, 0);
+		ensure(InRequest.GetNumComponents() == Field->GetNumComponents());
+		RequestedComponents.SetRangeFromRange(Field->ComponentIndex, InRequest.GetNumComponents(), InRequest.RequestedComponents, 0);
 	}
 
 	/** Returns the requested type of the given field */
-	FRequestedType GetField(const FStructFieldRef& FieldRef) const
+	FRequestedType GetField(const Shader::FStructField* Field) const
 	{
-		FRequestedType Result(FieldRef.Type);
-		Result.RequestedComponents.SetRangeFromRange(0, FieldRef.ComponentNum, RequestedComponents, FieldRef.ComponentIndex);
+		FRequestedType Result(Field->Type);
+		Result.RequestedComponents.SetRangeFromRange(0, Field->GetNumComponents(), RequestedComponents, Field->ComponentIndex);
 		return Result;
 	}
 
-	/** The struct type we're requesting, or nullptr if we're requesting a basic type */
-	const FStructType* StructType = nullptr;
+	/**
+	 * If either StructType or ValueComponentType are set, then the request is for an explicit type
+	 * Otherwise, the request is for any type with the given components
+	 */
+	const Shader::FStructType* StructType = nullptr;
+	Shader::EValueComponentType ValueComponentType = Shader::EValueComponentType::Void;
 
 	/** 1 bit per component, a value of 'true' means the specified component is requsted */
 	TBitArray<> RequestedComponents;
 };
 
+/**
+ * Like FRequestedType, but tracks an EExpressionEvaluationType per component, rather than a simple requested flag
+ */
+class FPreparedType
+{
+public:
+	FPreparedType() = default;
+	FPreparedType(Shader::EValueComponentType InComponentType) : ValueComponentType(InComponentType) {}
+	FPreparedType(const Shader::FStructType* InStructType) : StructType(InStructType) {}
+	FPreparedType(const Shader::FType& InType);
+
+	void SetEvaluationType(EExpressionEvaluationType EvaluationType);
+
+	void SetField(const Shader::FStructField* Field, const FPreparedType& FieldType);
+	FPreparedType GetFieldType(const Shader::FStructField* Field) const;
+
+	int32 GetNumComponents() const;
+	FRequestedType GetRequestedType() const;
+	Shader::FType GetType() const;
+	bool IsStruct() const { return !IsVoid() && StructType != nullptr; }
+	bool IsNumeric() const { return !IsVoid() && ValueComponentType != Shader::EValueComponentType::Void; }
+	bool IsInitialized() const { return StructType != nullptr || ValueComponentType != Shader::EValueComponentType::Void; }
+	bool IsVoid() const;
+	EExpressionEvaluationType GetEvaluationType(const FRequestedType& RequestedType) const;
+
+	EExpressionEvaluationType GetComponentEvaluationType(int32 Index) const
+	{
+		return ComponentEvaluationType.IsValidIndex(Index) ? ComponentEvaluationType[Index] : EExpressionEvaluationType::None;
+	}
+
+	void SetComponentEvaluationType(int32 Index, EExpressionEvaluationType EvaluationType);
+
+	/** Unlike FRequestedType, one of these should be set */
+	const Shader::FStructType* StructType = nullptr;
+	Shader::EValueComponentType ValueComponentType = Shader::EValueComponentType::Void;
+
+	/** Evaluation type for each component, may be 'None' for components that are unused */
+	TArray<EExpressionEvaluationType, TInlineAllocator<16>> ComponentEvaluationType;
+};
+
+FPreparedType MergePreparedTypes(const FPreparedType& Lhs, const FPreparedType& Rhs);
+
 struct FShaderValue
 {
 	explicit FShaderValue(FStringBuilderBase& InCode) : Code(InCode) {}
 
 	FStringBuilderBase& Code;
+	Shader::FType Type;
 	bool bInline = false;
-	bool bHasDependencies = false;
 };
 
-struct FPrepareValueResult
+class FPrepareValueResult
 {
-	FExpression* ForwardValue = nullptr;
-	FType Type;
-	EExpressionEvaluationType EvaluationType = EExpressionEvaluationType::None;
-	Shader::FValue ConstantValue;
+public:
+	const FPreparedType& GetPreparedType() const { return PreparedType; }
+
+	void SetType(FEmitContext& Context, const FRequestedType& RequestedType, EExpressionEvaluationType EvaluationType, const Shader::FType& Type);
+	void SetType(FEmitContext& Context, const FRequestedType& RequestedType, EExpressionEvaluationType EvaluationType, Shader::EValueComponentType ComponentType);
+	void SetType(FEmitContext& Context, const FRequestedType& RequestedType, const FPreparedType& Type);
+	void SetForwardValue(FEmitContext& Context, const FRequestedType& RequestedType, FExpression* InValue);
 
-	void SetConstant(FEmitContext& Context, const Shader::FValue& InValue);
-	void SetType(FEmitContext& Context, EExpressionEvaluationType InEvaluationType, const FType& InType);
-	void SetForwardValue(FEmitContext& Context, FExpression* InValue, const FRequestedType& RequestedType);
+private:
+	bool TryMergePreparedType(FEmitContext& Context, const Shader::FStructType* StructType, Shader::EValueComponentType ComponentType);
+
+	FExpression* ForwardValue = nullptr;
+	FPreparedType PreparedType;
 
-	inline bool IsValid() const { return EvaluationType != EExpressionEvaluationType::None; }
+	friend class FExpression;
 };
 
 struct FEmitShaderValue
 {
-	FEmitShaderValue(FExpression* InExpression, FScope* InScope) : Expression(InExpression), Scope(InScope) {}
+	FEmitShaderValue(FScope* InScope, const Shader::FType& InType) : Scope(InScope), Type(InType) {}
 
 	inline bool IsInline() const { return Value == nullptr; }
 
-	FExpression* Expression = nullptr;
 	FScope* Scope = nullptr;
 	const TCHAR* Reference = nullptr;
 	const TCHAR* Value = nullptr;
+	Shader::FType Type;
 	TArrayView<FEmitShaderValue*> Dependencies;
 	FSHAHash Hash;
 };
@@ -346,26 +323,26 @@ struct FEmitShaderValue
 class FExpression : public FNode
 {
 public:
-	const FType& GetType() const { return PrepareValueResult.Type; }
-	EExpressionEvaluationType GetEvaluationType() const { return PrepareValueResult.EvaluationType; }
+	FRequestedType GetRequestedType() const { return PrepareValueResult.PreparedType.GetRequestedType(); }
+	Shader::FType GetType() const { return PrepareValueResult.PreparedType.GetType(); }
+	EExpressionEvaluationType GetEvaluationType(const FRequestedType& RequestedType) const { return PrepareValueResult.PreparedType.GetEvaluationType(RequestedType); }
 
 	virtual void Reset() override;
 
-	friend const FPrepareValueResult& PrepareExpressionValue(FEmitContext& Context, FExpression* InExpression, const FRequestedType& RequestedType);
+	friend const FPreparedType& PrepareExpressionValue(FEmitContext& Context, FExpression* InExpression, const FRequestedType& RequestedType);
+
+	const TCHAR* GetValueShader(FEmitContext& Context, const FRequestedType& RequestedType);
+	void GetValuePreshader(FEmitContext& Context, const FRequestedType& RequestedType, Shader::FPreshaderData& OutPreshader);
+	Shader::FValue GetValueConstant(FEmitContext& Context, const FRequestedType& RequestedType);
 
 	const TCHAR* GetValueShader(FEmitContext& Context);
-	const TCHAR* GetValueShader(FEmitContext& Context, const FType& InType);
-	void GetValuePreshader(FEmitContext& Context, Shader::FPreshaderData& OutPreshader);
-	Shader::FValue GetValueConstant(FEmitContext& Context);
 
 protected:
-	virtual void PrepareValue(FEmitContext& Context, const FRequestedType& RequestedType, FPrepareValueResult& OutResult) = 0;
-	virtual void EmitValueShader(FEmitContext& Context, FShaderValue& OutShader) const;
-	virtual void EmitShaderDependencies(FEmitContext& Context, const FShaderValue& Shader) const;
-	virtual void EmitValuePreshader(FEmitContext& Context, Shader::FPreshaderData& OutPreshader) const;
+	virtual void PrepareValue(FEmitContext& Context, const FRequestedType& RequestedType, FPrepareValueResult& OutResult) const = 0;
+	virtual void EmitValueShader(FEmitContext& Context, const FRequestedType& RequestedType, FShaderValue& OutShader) const;
+	virtual void EmitValuePreshader(FEmitContext& Context, const FRequestedType& RequestedType, Shader::FPreshaderData& OutPreshader) const;
 
 private:
-	FEmitShaderValue* ShaderValue = nullptr;
 	FRequestedType CurrentRequestedType;
 	FPrepareValueResult PrepareValueResult;
 	bool bReentryFlag = false;
@@ -382,9 +359,8 @@ private:
 class FExpressionLocalPHI final : public FExpression
 {
 public:
-	virtual void PrepareValue(FEmitContext& Context, const FRequestedType& RequestedType, FPrepareValueResult& OutResult) override;
-	virtual void EmitValueShader(FEmitContext& Context, FShaderValue& OutShader) const override;
-	virtual void EmitShaderDependencies(FEmitContext& Context, const FShaderValue& Shader) const override;
+	virtual void PrepareValue(FEmitContext& Context, const FRequestedType& RequestedType, FPrepareValueResult& OutResult) const override;
+	virtual void EmitValueShader(FEmitContext& Context, const FRequestedType& RequestedType, FShaderValue& OutShader) const override;
 
 	FName LocalName;
 	FScope* Scopes[MaxNumPreviousScopes];
@@ -548,8 +524,6 @@ public:
 
 	void ResetNodes();
 
-	void EmitDeclarationsCode(FStringBuilderBase& OutCode) const;
-
 	bool EmitShader(FEmitContext& Context, FStringBuilderBase& OutCode) const;
 
 	FScope& GetRootScope() const { return *RootScope; }
@@ -575,8 +549,6 @@ public:
 
 	FTextureParameterDeclaration* NewTextureParameterDeclaration(const FName& Name, const FTextureDescription& DefaultValue);
 
-	const FStructType* NewStructType(const FStructTypeInitializer& Initializer);
-
 private:
 	template<typename T, typename... ArgTypes>
 	inline T* NewNode(ArgTypes&&... Args)
@@ -592,8 +564,6 @@ private:
 
 	FMemStackBase* Allocator = nullptr;
 	FNode* Nodes = nullptr;
-	FExpression* ExpressionsToDeclare = nullptr;
-	FStructType* StructTypes = nullptr;
 	FScope* RootScope = nullptr;
 };
 
diff --git a/Engine/Source/Runtime/Engine/Public/HLSLTree/HLSLTreeCommon.h b/Engine/Source/Runtime/Engine/Public/HLSLTree/HLSLTreeCommon.h
--- a/Engine/Source/Runtime/Engine/Public/HLSLTree/HLSLTreeCommon.h
+++ b/Engine/Source/Runtime/Engine/Public/HLSLTree/HLSLTreeCommon.h
@@ -41,15 +41,14 @@ FBinaryOpDescription GetBinaryOpDesription(EBinaryOp Op);
 class FExpressionConstant : public FExpression
 {
 public:
-	explicit FExpressionConstant(const FConstantValue& InValue)
+	explicit FExpressionConstant(const Shader::FValue& InValue)
 		: Value(InValue)
 	{}
 
-	FConstantValue Value;
+	Shader::FValue Value;
 
-	virtual void PrepareValue(FEmitContext& Context, const FRequestedType& RequestedType, FPrepareValueResult& OutResult) override;
-	virtual void EmitValueShader(FEmitContext& Context, FShaderValue& OutShader) const override;
-	virtual void EmitValuePreshader(FEmitContext& Context, Shader::FPreshaderData& OutPreshader) const override;
+	virtual void PrepareValue(FEmitContext& Context, const FRequestedType& RequestedType, FPrepareValueResult& OutResult) const override;
+	virtual void EmitValuePreshader(FEmitContext& Context, const FRequestedType& RequestedType, Shader::FPreshaderData& OutPreshader) const override;
 };
 
 class FExpressionMaterialParameter : public FExpression
@@ -63,8 +62,8 @@ public:
 	Shader::FValue DefaultValue;
 	EMaterialParameterType ParameterType;
 
-	virtual void PrepareValue(FEmitContext& Context, const FRequestedType& RequestedType, FPrepareValueResult& OutResult) override;
-	virtual void EmitValuePreshader(FEmitContext& Context, Shader::FPreshaderData& OutPreshader) const override;
+	virtual void PrepareValue(FEmitContext& Context, const FRequestedType& RequestedType, FPrepareValueResult& OutResult) const override;
+	virtual void EmitValuePreshader(FEmitContext& Context, const FRequestedType& RequestedType, Shader::FPreshaderData& OutPreshader) const override;
 };
 
 enum class EExternalInputType
@@ -95,8 +94,8 @@ public:
 
 	EExternalInputType InputType;
 
-	virtual void PrepareValue(FEmitContext& Context, const FRequestedType& RequestedType, FPrepareValueResult& OutResult) override { OutResult.SetType(Context, EExpressionEvaluationType::Shader, GetInputExpressionType(InputType)); }
-	virtual void EmitValueShader(FEmitContext& Context, FShaderValue& OutShader) const override;
+	virtual void PrepareValue(FEmitContext& Context, const FRequestedType& RequestedType, FPrepareValueResult& OutResult) const override { OutResult.SetType(Context, RequestedType, EExpressionEvaluationType::Shader, GetInputExpressionType(InputType)); }
+	virtual void EmitValueShader(FEmitContext& Context, const FRequestedType& RequestedType, FShaderValue& OutShader) const override;
 };
 
 class FExpressionTextureSample : public FExpression
@@ -114,35 +113,50 @@ public:
 	ESamplerSourceMode SamplerSource;
 	ETextureMipValueMode MipValueMode;
 
-	virtual void PrepareValue(FEmitContext& Context, const FRequestedType& RequestedType, FPrepareValueResult& OutResult) override;
-	virtual void EmitValueShader(FEmitContext& Context, FShaderValue& OutShader) const override;
+	virtual void PrepareValue(FEmitContext& Context, const FRequestedType& RequestedType, FPrepareValueResult& OutResult) const override;
+	virtual void EmitValueShader(FEmitContext& Context, const FRequestedType& RequestedType, FShaderValue& OutShader) const override;
 };
 
 class FExpressionGetStructField : public FExpression
 {
 public:
-	FExpressionGetStructField() {}
-
-	const FStructType* StructType;
-	const TCHAR* FieldName;
+	FExpressionGetStructField(const Shader::FStructType* InStructType, const TCHAR* InFieldName, FExpression* InStructExpression)
+		: StructType(InStructType)
+		, Field(InStructType->FindFieldByName(InFieldName))
+		, StructExpression(InStructExpression)
+	{
+		check(Field);
+	}
+
+	const Shader::FStructType* StructType;
+	const Shader::FStructField* Field;
 	FExpression* StructExpression;
 
-	virtual void PrepareValue(FEmitContext& Context, const FRequestedType& RequestedType, FPrepareValueResult& OutResult) override;
-	virtual void EmitValueShader(FEmitContext& Context, FShaderValue& OutShader) const override;
+	virtual void PrepareValue(FEmitContext& Context, const FRequestedType& RequestedType, FPrepareValueResult& OutResult) const override;
+	virtual void EmitValueShader(FEmitContext& Context, const FRequestedType& RequestedType, FShaderValue& OutShader) const override;
+	virtual void EmitValuePreshader(FEmitContext& Context, const FRequestedType& RequestedType, Shader::FPreshaderData& OutPreshader) const override;
 };
 
 class FExpressionSetStructField : public FExpression
 {
 public:
-	FExpressionSetStructField() {}
-
-	const FStructType* StructType;
-	const TCHAR* FieldName;
+	FExpressionSetStructField(const Shader::FStructType* InStructType, const TCHAR* InFieldName, FExpression* InStructExpression, FExpression* InFieldExpression)
+		: StructType(InStructType)
+		, Field(InStructType->FindFieldByName(InFieldName))
+		, StructExpression(InStructExpression)
+		, FieldExpression(InFieldExpression)
+	{
+		check(Field);
+	}
+
+	const Shader::FStructType* StructType;
+	const Shader::FStructField* Field;
 	FExpression* StructExpression;
 	FExpression* FieldExpression;
 
-	virtual void PrepareValue(FEmitContext& Context, const FRequestedType& RequestedType, FPrepareValueResult& OutResult) override;
-	virtual void EmitValueShader(FEmitContext& Context, FShaderValue& OutShader) const override;
+	virtual void PrepareValue(FEmitContext& Context, const FRequestedType& RequestedType, FPrepareValueResult& OutResult) const override;
+	virtual void EmitValueShader(FEmitContext& Context, const FRequestedType& RequestedType, FShaderValue& OutShader) const override;
+	virtual void EmitValuePreshader(FEmitContext& Context, const FRequestedType& RequestedType, Shader::FPreshaderData& OutPreshader) const override;
 };
 
 class FExpressionSelect : public FExpression
@@ -158,9 +172,9 @@ public:
 	FExpression* TrueExpression;
 	FExpression* FalseExpression;
 
-	virtual void PrepareValue(FEmitContext& Context, const FRequestedType& RequestedType, FPrepareValueResult& OutResult) override;
-	virtual void EmitValueShader(FEmitContext& Context, FShaderValue& OutShader) const override;
-	virtual void EmitValuePreshader(FEmitContext& Context, Shader::FPreshaderData& OutPreshader) const override;
+	virtual void PrepareValue(FEmitContext& Context, const FRequestedType& RequestedType, FPrepareValueResult& OutResult) const override;
+	virtual void EmitValueShader(FEmitContext& Context, const FRequestedType& RequestedType, FShaderValue& OutShader) const override;
+	virtual void EmitValuePreshader(FEmitContext& Context, const FRequestedType& RequestedType, Shader::FPreshaderData& OutPreshader) const override;
 };
 
 class FExpressionBinaryOp : public FExpression
@@ -176,9 +190,9 @@ public:
 	FExpression* Lhs;
 	FExpression* Rhs;
 
-	virtual void PrepareValue(FEmitContext& Context, const FRequestedType& RequestedType, FPrepareValueResult& OutResult) override;
-	virtual void EmitValueShader(FEmitContext& Context, FShaderValue& OutShader) const override;
-	virtual void EmitValuePreshader(FEmitContext& Context, Shader::FPreshaderData& OutPreshader) const override;
+	virtual void PrepareValue(FEmitContext& Context, const FRequestedType& RequestedType, FPrepareValueResult& OutResult) const override;
+	virtual void EmitValueShader(FEmitContext& Context, const FRequestedType& RequestedType, FShaderValue& OutShader) const override;
+	virtual void EmitValuePreshader(FEmitContext& Context, const FRequestedType& RequestedType, Shader::FPreshaderData& OutPreshader) const override;
 };
 
 struct FSwizzleParameters
@@ -186,6 +200,9 @@ struct FSwizzleParameters
 	FSwizzleParameters() : NumComponents(0) { ComponentIndex[0] = ComponentIndex[1] = ComponentIndex[2] = ComponentIndex[3] = INDEX_NONE; }
 	FSwizzleParameters(int8 IndexR, int8 IndexG, int8 IndexB, int8 IndexA);
 
+	FRequestedType GetRequestedInputType(const FRequestedType& RequestedType) const;
+	bool HasSwizzle() const;
+
 	int8 ComponentIndex[4];
 	int32 NumComponents;
 };
@@ -202,9 +219,9 @@ public:
 	FSwizzleParameters Parameters;
 	FExpression* Input;
 
-	virtual void PrepareValue(FEmitContext& Context, const FRequestedType& RequestedType, FPrepareValueResult& OutResult) override;
-	virtual void EmitValueShader(FEmitContext& Context, FShaderValue& OutShader) const override;
-	virtual void EmitValuePreshader(FEmitContext& Context, Shader::FPreshaderData& OutPreshader) const override;
+	virtual void PrepareValue(FEmitContext& Context, const FRequestedType& RequestedType, FPrepareValueResult& OutResult) const override;
+	virtual void EmitValueShader(FEmitContext& Context, const FRequestedType& RequestedType, FShaderValue& OutShader) const override;
+	virtual void EmitValuePreshader(FEmitContext& Context, const FRequestedType& RequestedType, Shader::FPreshaderData& OutPreshader) const override;
 };
 
 class FExpressionAppend : public FExpression
@@ -218,16 +235,16 @@ public:
 	FExpression* Lhs;
 	FExpression* Rhs;
 
-	virtual void PrepareValue(FEmitContext& Context, const FRequestedType& RequestedType, FPrepareValueResult& OutResult) override;
-	virtual void EmitValueShader(FEmitContext& Context, FShaderValue& OutShader) const override;
-	virtual void EmitValuePreshader(FEmitContext& Context, Shader::FPreshaderData& OutPreshader) const override;
+	virtual void PrepareValue(FEmitContext& Context, const FRequestedType& RequestedType, FPrepareValueResult& OutResult) const override;
+	virtual void EmitValueShader(FEmitContext& Context, const FRequestedType& RequestedType, FShaderValue& OutShader) const override;
+	virtual void EmitValuePreshader(FEmitContext& Context, const FRequestedType& RequestedType, Shader::FPreshaderData& OutPreshader) const override;
 };
 
 class FExpressionReflectionVector : public FExpression
 {
 public:
-	virtual void PrepareValue(FEmitContext& Context, const FRequestedType& RequestedType, FPrepareValueResult& OutResult) override;
-	virtual void EmitValueShader(FEmitContext& Context, FShaderValue& OutShader) const override;
+	virtual void PrepareValue(FEmitContext& Context, const FRequestedType& RequestedType, FPrepareValueResult& OutResult) const override;
+	virtual void EmitValueShader(FEmitContext& Context, const FRequestedType& RequestedType, FShaderValue& OutShader) const override;
 };
 
 class FStatementReturn : public FStatement
diff --git a/Engine/Source/Runtime/Engine/Public/HLSLTree/HLSLTreeTypes.h b/Engine/Source/Runtime/Engine/Public/HLSLTree/HLSLTreeTypes.h
--- a/Engine/Source/Runtime/Engine/Public/HLSLTree/HLSLTreeTypes.h
+++ b/Engine/Source/Runtime/Engine/Public/HLSLTree/HLSLTreeTypes.h
@@ -22,7 +22,7 @@ class FTree;
 
 /**
  * Describes how a given expression needs to be evaluated */
-enum class EExpressionEvaluationType
+enum class EExpressionEvaluationType : uint8
 {
 	/** Invalid/uninitialized */
 	None,
diff --git a/Engine/Source/Runtime/Engine/Public/MaterialHLSLGenerator.h b/Engine/Source/Runtime/Engine/Public/MaterialHLSLGenerator.h
--- a/Engine/Source/Runtime/Engine/Public/MaterialHLSLGenerator.h
+++ b/Engine/Source/Runtime/Engine/Public/MaterialHLSLGenerator.h
@@ -35,6 +35,10 @@ class FExpressionExternalInput;
 class FExpressionSwizzle;
 class FExpressionLocalPHI;
 }
+namespace Shader
+{
+struct FStructType;
+}
 }
 
 enum class EMaterialNewScopeFlag : uint8
@@ -56,7 +60,9 @@ struct TMaterialHLSLGeneratorType;
 class FMaterialHLSLGenerator
 {
 public:
-	FMaterialHLSLGenerator(UMaterial* InTargetMaterial, const FMaterialCompileTargetParameters& InCompilerTarget,
+	FMaterialHLSLGenerator(UMaterial* InTargetMaterial,
+		const FMaterialCompileTargetParameters& InCompilerTarget,
+		UE::Shader::FStructTypeRegistry& InOutTypeRegistry,
 		UE::HLSLTree::FTree& InOutTree);
 
 	const FMaterialCompileTargetParameters& GetCompileTarget() const { return CompileTarget; }
@@ -136,8 +142,8 @@ public:
 		return Data;
 	}
 
-	const UE::HLSLTree::FStructType* GetMaterialAttributesType() const { return MaterialAttributesType; }
-	const UE::HLSLTree::FConstantValue& GetMaterialAttributesDefaultValue() const { return MaterialAttributesDefaultValue; }
+	const UE::Shader::FStructType* GetMaterialAttributesType() const { return MaterialAttributesType; }
+	const UE::Shader::FValue& GetMaterialAttributesDefaultValue() const { return MaterialAttributesDefaultValue; }
 
 private:
 	static constexpr int32 MaxNumPreviousScopes = UE::HLSLTree::MaxNumPreviousScopes;
@@ -215,10 +221,11 @@ private:
 	const FMaterialCompileTargetParameters& CompileTarget;
 	UMaterial* TargetMaterial;
 
-	const UE::HLSLTree::FStructType* MaterialAttributesType;
-	UE::HLSLTree::FConstantValue MaterialAttributesDefaultValue;
+	const UE::Shader::FStructType* MaterialAttributesType;
+	UE::Shader::FValue MaterialAttributesDefaultValue;
 
 	UE::HLSLTree::FTree* HLSLTree;
+	UE::Shader::FStructTypeRegistry* TypeRegistry;
 	UE::HLSLTree::FExpression* ResultExpression = nullptr;
 	UE::HLSLTree::FStatement* ResultStatement = nullptr;
 
diff --git a/Engine/Source/Runtime/Engine/Public/Shader/Preshader.h b/Engine/Source/Runtime/Engine/Public/Shader/Preshader.h
--- a/Engine/Source/Runtime/Engine/Public/Shader/Preshader.h
+++ b/Engine/Source/Runtime/Engine/Public/Shader/Preshader.h
@@ -59,6 +59,16 @@ enum class EPreshaderOpcode : uint8
 	ExternalTextureCoordinateScaleRotation,
 	ExternalTextureCoordinateOffset,
 	RuntimeVirtualTextureUniform,
+	GetField,
+	SetField,
+};
+
+struct FPreshaderStructType
+{
+	DECLARE_TYPE_LAYOUT(FPreshaderStructType, NonVirtual);
+	LAYOUT_FIELD(uint64, Hash);
+	LAYOUT_FIELD(int32, ComponentTypeIndex);
+	LAYOUT_FIELD(int32, NumComponents);
 };
 
 class FPreshaderData
@@ -78,14 +88,14 @@ public:
 	FSHAHash GetHash() const;
 	void AppendHash(FSHA1& OutHasher) const;
 
-	void Evaluate(FUniformExpressionSet* UniformExpressionSet, const struct FMaterialRenderContext& Context, FValue& OutValue);
-
-	void Append(const FPreshaderData& InPreshader);
+	FPreshaderValue Evaluate(FUniformExpressionSet* UniformExpressionSet, const struct FMaterialRenderContext& Context, FPreshaderStack& Stack);
 
 	const int32 Num() const { return Data.Num(); }
 
 	void WriteData(const void* Value, uint32 Size);
 	void WriteName(const FScriptName& Name);
+	void WriteType(const FType& Type);
+	void WriteValue(const FValue& Value);
 
 	template<typename T>
 	FPreshaderData& Write(const T& Value) { WriteData(&Value, sizeof(T)); return *this; }
@@ -94,7 +104,10 @@ public:
 	FPreshaderData& Write<FScriptName>(const FScriptName& Value) { WriteName(Value); return *this; }
 
 	template<>
-	FPreshaderData& Write<FValue>(const FValue& Value);
+	FPreshaderData& Write<FType>(const FType& Value) { WriteType(Value); return *this; }
+
+	template<>
+	FPreshaderData& Write<FValue>(const FValue& Value) { WriteValue(Value); return *this; }
 
 	/** Can't write FName, use FScriptName instead */
 	template<>
@@ -106,22 +119,10 @@ public:
 	inline FPreshaderData& WriteOpcode(EPreshaderOpcode Op) { return Write<uint8>((uint8)Op); }
 
 	LAYOUT_FIELD(TMemoryImageArray<FScriptName>, Names);
-	LAYOUT_FIELD(TMemoryImageArray<uint32>, NameOffsets);
+	LAYOUT_FIELD(TMemoryImageArray<FPreshaderStructType>, StructTypes);
+	LAYOUT_FIELD(TMemoryImageArray<EValueComponentType>, StructComponentTypes);
 	LAYOUT_FIELD(TMemoryImageArray<uint8>, Data);
 };
 
-template<>
-inline FPreshaderData& FPreshaderData::Write<FValue>(const FValue& Value)
-{
-	const EValueType Type = Value.GetType();
-	FMemoryImageValue MemoryValue = Value.AsMemoryImage();
-	Data.Add((uint8)Type);
-	if (MemoryValue.Size > 0u)
-	{
-		Data.Append(MemoryValue.Bytes, MemoryValue.Size);
-	}
-	return *this;
-}
-
 } // namespace Shader
 } // namespace UE
\ No newline at end of file
diff --git a/Engine/Source/Runtime/Engine/Public/Shader/PreshaderTypes.h b/Engine/Source/Runtime/Engine/Public/Shader/PreshaderTypes.h
--- a/Engine/Source/Runtime/Engine/Public/Shader/PreshaderTypes.h
+++ b/Engine/Source/Runtime/Engine/Public/Shader/PreshaderTypes.h
@@ -2,3 +2,71 @@
 #pragma once
 
 #include "Shader/ShaderTypes.h"
+#include "Containers/ArrayView.h"
+
+namespace UE
+{
+namespace Shader
+{
+
+/**
+ * Mirrors Shader::FType, but stores StructType as a hash rather than a pointer to facilitate serialization
+ * Struct's flattened component types are stored directly, as that is the primary thing needed at runtime
+ */
+struct FPreshaderType
+{
+	FPreshaderType() = default;
+	FPreshaderType(const FType& InType);
+
+	bool IsStruct() const { return StructTypeHash != 0u; }
+	int32 GetNumComponents() const { return StructTypeHash != 0u ? StructComponentTypes.Num() : GetValueTypeDescription(ValueType).NumComponents; }
+	EValueComponentType GetComponentType(int32 Index) const;
+
+	uint64 StructTypeHash = 0u;
+	TArrayView<const EValueComponentType> StructComponentTypes;
+	EValueType ValueType = EValueType::Void;
+};
+
+/** Mirrors Shader::FValue, except 'Component' array references memory owned by stack, rather than inline storage */
+struct FPreshaderValue
+{
+	FPreshaderType Type;
+	TArrayView<FValueComponent> Component;
+
+	/** Converts to a regular 'FValue', uses the type registry to resolve StructType hashes into pointers */
+	FValue AsShaderValue(const FStructTypeRegistry* TypeRegistry = nullptr) const;
+};
+
+class FPreshaderStack
+{
+public:
+	int32 Num() const { return Values.Num(); }
+
+	void CheckEmpty() const
+	{
+		check(Values.Num() == 0);
+		check(Components.Num() == 0);
+	}
+
+	void PushValue(const FValue& InValue);
+	void PushValue(const FPreshaderValue& InValue);
+	void PushValue(const FPreshaderType& InType, TArrayView<const FValueComponent> InComponents);
+	TArrayView<FValueComponent> PushEmptyValue(const FPreshaderType& InType);
+
+	/** Returned values are invalidated when anything is pushed onto the stack */
+	FPreshaderValue PopValue();
+	FPreshaderValue PeekValue();
+
+	void Reset()
+	{
+		Values.Reset();
+		Components.Reset();
+	}
+
+private:
+	TArray<FPreshaderType, TInlineAllocator<8>> Values;
+	TArray<FValueComponent, TInlineAllocator<64>> Components;
+};
+
+} // namespace Shader
+} // namespace UE
\ No newline at end of file
diff --git a/Engine/Source/Runtime/Engine/Public/Shader/ShaderTypes.h b/Engine/Source/Runtime/Engine/Public/Shader/ShaderTypes.h
--- a/Engine/Source/Runtime/Engine/Public/Shader/ShaderTypes.h
+++ b/Engine/Source/Runtime/Engine/Public/Shader/ShaderTypes.h
@@ -2,13 +2,19 @@
 #pragma once
 
 #include "CoreMinimal.h"
+#include "Containers/ArrayView.h"
+#include "Containers/StringView.h"
 #include "Serialization/MemoryLayout.h"
 
+class FMemStackBase;
+
 namespace UE
 {
 namespace Shader
 {
 
+struct FStructType;
+
 enum class EValueComponentType : uint8
 {
 	Void,
@@ -18,6 +24,9 @@ enum class EValueComponentType : uint8
 	Bool,
 };
 
+uint32 GetComponentTypeSizeInBytes(EValueComponentType Type);
+EValueComponentType CombineComponentTypes(EValueComponentType Lhs, EValueComponentType Rhs);
+
 struct FValueTypeDescription
 {
 	FValueTypeDescription() : Name(nullptr), ComponentType(EValueComponentType::Void), NumComponents(0) {}
@@ -55,7 +64,6 @@ enum class EValueType : uint8
 	Struct,
 };
 
-uint32 GetComponentTypeSizeInBytes(EValueComponentType Type);
 FValueTypeDescription GetValueTypeDescription(EValueType Type);
 EValueType MakeValueType(EValueComponentType ComponentType, int32 NumComponents);
 EValueType MakeValueType(EValueType BaseType, int32 NumComponents);
@@ -63,6 +71,111 @@ EValueType MakeValueTypeWithRequestedNumComponents(EValueType BaseType, int8 Req
 EValueType MakeArithmeticResultType(EValueType Lhs, EValueType Rhs, FString& OutErrorMessage);
 EValueType MakeComparisonResultType(EValueType Lhs, EValueType Rhs, FString& OutErrorMessage);
 
+struct FType
+{
+	FType() : StructType(nullptr), ValueType(EValueType::Void) {}
+	FType(EValueType InValueType) : StructType(nullptr), ValueType(InValueType) {}
+	FType(const FStructType* InStruct) : StructType(InStruct), ValueType(InStruct ? EValueType::Struct : EValueType::Void) {}
+
+	const TCHAR* GetName() const;
+	bool IsVoid() const { return ValueType == EValueType::Void; }
+	bool IsStruct() const { return ValueType == EValueType::Struct; }
+	bool IsNumeric() const { return !IsVoid() && !IsStruct(); }
+	int32 GetNumComponents() const;
+	EValueComponentType GetComponentType(int32 Index) const;
+	bool Merge(const FType& OtherType);
+
+	inline operator EValueType() const { return ValueType; }
+	inline operator bool() const { return !IsVoid(); }
+	inline bool operator!() const { return IsVoid(); }
+
+	const FStructType* StructType;
+	EValueType ValueType;
+};
+
+inline bool operator==(const FType& Lhs, const FType& Rhs)
+{
+	if (Lhs.ValueType != Rhs.ValueType) return false;
+	if (Lhs.ValueType == EValueType::Struct && Lhs.StructType != Rhs.StructType) return false;
+	return true;
+}
+inline bool operator!=(const FType& Lhs, const FType& Rhs)
+{
+	return !operator==(Lhs, Rhs);
+}
+
+inline bool operator==(const FType& Lhs, const EValueType& Rhs)
+{
+	return !Lhs.IsStruct() && Lhs.ValueType == Rhs;
+}
+inline bool operator!=(const FType& Lhs, const EValueType& Rhs)
+{
+	return !operator==(Lhs, Rhs);
+}
+
+inline bool operator==(const EValueType& Lhs, const FType& Rhs)
+{
+	return !Rhs.IsStruct() && Lhs == Rhs.ValueType;
+}
+inline bool operator!=(const EValueType& Lhs, const FType& Rhs)
+{
+	return !operator==(Lhs, Rhs);
+}
+
+struct FStructField
+{
+	const TCHAR* Name;
+	FType Type;
+	int32 ComponentIndex;
+
+	int32 GetNumComponents() const { return Type.GetNumComponents(); }
+};
+
+struct FStructType
+{
+	uint64 Hash;
+	const TCHAR* Name;
+	TArrayView<const FStructField> Fields;
+
+	/**
+	 * Most code working with HLSLTree views struct types as a flat list of components
+	 * Fields with basic types are represented directly. Fields with struct types are recursively flattened into this list
+	 */
+	TArrayView<const EValueComponentType> ComponentTypes;
+
+	const FStructField* FindFieldByName(const TCHAR* InName) const;
+};
+
+struct FStructFieldInitializer
+{
+	FStructFieldInitializer() = default;
+	FStructFieldInitializer(const FStringView& InName, const FType& InType) : Name(InName), Type(InType) {}
+
+	FStringView Name;
+	FType Type;
+};
+
+struct FStructTypeInitializer
+{
+	FStringView Name;
+	TArrayView<const FStructFieldInitializer> Fields;
+};
+
+class FStructTypeRegistry
+{
+public:
+	explicit FStructTypeRegistry(FMemStackBase& InAllocator);
+
+	void EmitDeclarationsCode(FStringBuilderBase& OutCode) const;
+
+	const FStructType* NewType(const FStructTypeInitializer& Initializer);
+	const FStructType* FindType(uint64 Hash) const;
+
+private:
+	FMemStackBase* Allocator;
+	TMap<uint64, const FStructType*> Types;
+};
+
 template<typename T>
 struct TValue
 {
@@ -112,135 +225,136 @@ static_assert(sizeof(FValueComponent) == sizeof(uint64), "bad packing");
 
 struct FValue
 {
-	inline FValue() : ComponentType(EValueComponentType::Void), NumComponents(0) {}
+	FValue() = default;
 
-	explicit FValue(EValueType InType)
+	explicit FValue(const FType& InType) : Type(InType)
 	{
-		const FValueTypeDescription TypeDesc = GetValueTypeDescription(InType);
-		ComponentType = TypeDesc.ComponentType;
-		NumComponents = TypeDesc.NumComponents;
+		Component.AddDefaulted(InType.GetNumComponents());
 	}
 
-	inline FValue(EValueComponentType InComponentType, int8 InNumComponents) : ComponentType(InComponentType), NumComponents(InNumComponents) {}
+	inline FValue(EValueComponentType InComponentType, int8 InNumComponents) : Type(MakeValueType(InComponentType, InNumComponents))
+	{
+		Component.AddDefaulted(InNumComponents);
+	}
 
-	inline FValue(float v) : ComponentType(EValueComponentType::Float), NumComponents(1)
+	inline FValue(float v) : Type(EValueType::Float1)
 	{
-		Component[0] = v;
+		Component.Add(v);
 	}
 
-	inline FValue(float X, float Y) : ComponentType(EValueComponentType::Float), NumComponents(2)
+	inline FValue(float X, float Y) : Type(EValueType::Float2)
 	{
-		Component[0] = X;
-		Component[1] = Y;
+		Component.Add(X);
+		Component.Add(Y);
 	}
 
-	inline FValue(float X, float Y, float Z) : ComponentType(EValueComponentType::Float), NumComponents(3)
+	inline FValue(float X, float Y, float Z) : Type(EValueType::Float3)
 	{
-		Component[0] = X;
-		Component[1] = Y;
-		Component[2] = Z;
+		Component.Add(X);
+		Component.Add(Y);
+		Component.Add(Z);
 	}
 
-	inline FValue(float X, float Y, float Z, float W) : ComponentType(EValueComponentType::Float), NumComponents(4)
+	inline FValue(float X, float Y, float Z, float W) : Type(EValueType::Float4)
 	{
-		Component[0] = X;
-		Component[1] = Y;
-		Component[2] = Z;
-		Component[3] = W;
+		Component.Add(X);
+		Component.Add(Y);
+		Component.Add(Z);
+		Component.Add(W);
 	}
 
-	inline FValue(double v) : ComponentType(EValueComponentType::Double), NumComponents(1)
+	inline FValue(double v) : Type(EValueType::Double1)
 	{
-		Component[0] = v;
+		Component.Add(v);
 	}
 
-	inline FValue(double X, double Y) : ComponentType(EValueComponentType::Double), NumComponents(2)
+	inline FValue(double X, double Y) : Type(EValueType::Double2)
 	{
-		Component[0] = X;
-		Component[1] = Y;
+		Component.Add(X);
+		Component.Add(Y);
 	}
 
-	inline FValue(double X, double Y, double Z) : ComponentType(EValueComponentType::Double), NumComponents(3)
+	inline FValue(double X, double Y, double Z) : Type(EValueType::Double3)
 	{
-		Component[0] = X;
-		Component[1] = Y;
-		Component[2] = Z;
+		Component.Add(X);
+		Component.Add(Y);
+		Component.Add(Z);
 	}
 
-	inline FValue(double X, double Y, double Z, double W) : ComponentType(EValueComponentType::Double), NumComponents(4)
+	inline FValue(double X, double Y, double Z, double W) : Type(EValueType::Double4)
 	{
-		Component[0] = X;
-		Component[1] = Y;
-		Component[2] = Z;
-		Component[3] = W;
+		Component.Add(X);
+		Component.Add(Y);
+		Component.Add(Z);
+		Component.Add(W);
 	}
 
-	inline FValue(const FLinearColor& Value) : ComponentType(EValueComponentType::Float), NumComponents(4)
+	inline FValue(const FLinearColor& Value) : Type(EValueType::Float4)
 	{
-		Component[0] = Value.R;
-		Component[1] = Value.G;
-		Component[2] = Value.B;
-		Component[3] = Value.A;
+		Component.Add(Value.R);
+		Component.Add(Value.G);
+		Component.Add(Value.B);
+		Component.Add(Value.A);
 	}
 
-	inline FValue(const FVector2f& Value) : ComponentType(EValueComponentType::Float), NumComponents(2)
+	inline FValue(const FVector2f& Value) : Type(EValueType::Float2)
 	{
-		Component[0] = Value.X;
-		Component[1] = Value.Y;
+		Component.Add(Value.X);
+		Component.Add(Value.Y);
 	}
 
-	inline FValue(const FVector3f& Value) : ComponentType(EValueComponentType::Float), NumComponents(3)
+	inline FValue(const FVector3f& Value) : Type(EValueType::Float3)
 	{
-		Component[0] = Value.X;
-		Component[1] = Value.Y;
-		Component[2] = Value.Z;
+		Component.Add(Value.X);
+		Component.Add(Value.Y);
+		Component.Add(Value.Z);
 	}
 
-	inline FValue(const FVector3d& Value) : ComponentType(EValueComponentType::Double), NumComponents(3)
+	inline FValue(const FVector3d& Value) : Type(EValueType::Double3)
 	{
-		Component[0] = Value.X;
-		Component[1] = Value.Y;
-		Component[2] = Value.Z;
+		Component.Add(Value.X);
+		Component.Add(Value.Y);
+		Component.Add(Value.Z);
 	}
 
-	inline FValue(const FVector4f& Value) : ComponentType(EValueComponentType::Float), NumComponents(4)
+	inline FValue(const FVector4f& Value) : Type(EValueType::Float4)
 	{
-		Component[0] = Value.X;
-		Component[1] = Value.Y;
-		Component[2] = Value.Z;
-		Component[3] = Value.W;
+		Component.Add(Value.X);
+		Component.Add(Value.Y);
+		Component.Add(Value.Z);
+		Component.Add(Value.W);
 	}
 
-	inline FValue(const FVector4d& Value) : ComponentType(EValueComponentType::Double), NumComponents(4)
+	inline FValue(const FVector4d& Value) : Type(EValueType::Double4)
 	{
-		Component[0] = Value.X;
-		Component[1] = Value.Y;
-		Component[2] = Value.Z;
-		Component[3] = Value.W;
+		Component.Add(Value.X);
+		Component.Add(Value.Y);
+		Component.Add(Value.Z);
+		Component.Add(Value.W);
 	}
 
-	inline FValue(bool v) : ComponentType(EValueComponentType::Bool), NumComponents(1)
+	inline FValue(bool v) : Type(EValueType::Bool1)
 	{
-		Component[0] = v;
+		Component.Add(v);
 	}
 
-	inline FValue(bool X, bool Y, bool Z, bool W) : ComponentType(EValueComponentType::Bool), NumComponents(4)
+	inline FValue(bool X, bool Y, bool Z, bool W) : Type(EValueType::Bool4)
 	{
-		Component[0] = X;
-		Component[1] = Y;
-		Component[2] = Z;
-		Component[3] = W;
+		Component.Add(X);
+		Component.Add(Y);
+		Component.Add(Z);
+		Component.Add(W);
 	}
 
-	inline FValue(int32 v) : ComponentType(EValueComponentType::Int), NumComponents(1)
+	inline FValue(int32 v) : Type(EValueType::Int1)
 	{
-		Component[0] = v;
+		Component.Add(v);
 	}
 
-	inline EValueType GetType() const { return MakeValueType(ComponentType, NumComponents); }
+	inline const FType& GetType() const { return Type; }
 	inline const FValueComponent& GetComponent(int32 i) const
 	{
-		checkf(i >= 0 && i < NumComponents, TEXT("Invalid component %d/%d, of type '%s'"), i, NumComponents, GetValueTypeDescription(GetType()).Name);
+		checkf(Component.IsValidIndex(i), TEXT("Invalid component %d, of type '%s'"), i, Type.GetName());
 		return Component[i];
 	}
 
@@ -259,13 +373,13 @@ struct FValue
 
 	const TCHAR* ToString(EValueStringFormat Format, FStringBuilderBase& OutString) const;
 
-	FValueComponent Component[4];
-	EValueComponentType ComponentType;
-	int8 NumComponents;
+	FType Type;
+	TArray<FValueComponent, TInlineAllocator<8>> Component;
 };
 
 ENGINE_API bool operator==(const FValue& Lhs, const FValue& Rhs);
 inline bool operator!=(const FValue& Lhs, const FValue& Rhs) { return !operator==(Lhs, Rhs); }
+ENGINE_API uint32 GetTypeHash(const FType& Type);
 ENGINE_API uint32 GetTypeHash(const FValue& Value);
 
 ENGINE_API FValue Abs(const FValue& Value);
