Fix non-MSVC compile error

diff --git a/Engine/Source/Runtime/Engine/Private/HLSLTree/HLSLTree.cpp b/Engine/Source/Runtime/Engine/Private/HLSLTree/HLSLTree.cpp
--- a/Engine/Source/Runtime/Engine/Private/HLSLTree/HLSLTree.cpp
+++ b/Engine/Source/Runtime/Engine/Private/HLSLTree/HLSLTree.cpp
@@ -694,7 +694,7 @@ int32 FRequestedType::GetNumComponents() const
 
 bool FRequestedType::Merge(const FRequestedType& OtherType)
 {
-	if (StructType || OtherType.StructType && (StructType != OtherType.StructType))
+	if ((StructType || OtherType.StructType) && StructType != OtherType.StructType)
 	{
 		return false;
 	}
