Remove r.SkinCache.ForceRecomputeTangents. It can be achieved equivalently with r.SkinCache.CompileShaders=1 & r.SkinCache.Mode=1 & r.SkinCache.RecomputeTangents=1

diff --git a/Engine/Source/Runtime/Engine/Classes/Engine/RendererSettings.h b/Engine/Source/Runtime/Engine/Classes/Engine/RendererSettings.h
--- a/Engine/Source/Runtime/Engine/Classes/Engine/RendererSettings.h
+++ b/Engine/Source/Runtime/Engine/Classes/Engine/RendererSettings.h
@@ -1069,11 +1069,6 @@ class ENGINE_API URendererOverrideSettings : public UDeveloperSettings
 		ConfigRestartRequired = true))
 		uint32 bSupportAllShaderPermutations : 1;
 
-	UPROPERTY(config, EditAnywhere, Category = Miscellaneous, meta = (
-		ConsoleVariable = "r.SkinCache.ForceRecomputeTangents", DisplayName = "Force all skinned meshes to recompute tangents (also forces Compute SkinCache)",
-		ConfigRestartRequired = true))
-		uint32 bForceRecomputeTangents : 1;
-
 public:
 
 	//~ Begin UObject Interface
diff --git a/Engine/Source/Runtime/Engine/Private/GPUSkinCache.cpp b/Engine/Source/Runtime/Engine/Private/GPUSkinCache.cpp
--- a/Engine/Source/Runtime/Engine/Private/GPUSkinCache.cpp
+++ b/Engine/Source/Runtime/Engine/Private/GPUSkinCache.cpp
@@ -76,15 +76,6 @@ TAutoConsoleVariable<int32> CVarGPUSkinCacheRecomputeTangents(
 	ECVF_RenderThreadSafe
 );
 
-int32 GForceRecomputeTangents = 0;
-FAutoConsoleVariableRef CVarGPUSkinCacheForceRecomputeTangents(
-	TEXT("r.SkinCache.ForceRecomputeTangents"),
-	GForceRecomputeTangents,
-	TEXT("0: off (default)\n")
-	TEXT("1: Forces enabling and using the skincache and forces all skinned object to Recompute Tangents\n"),
-	ECVF_RenderThreadSafe | ECVF_ReadOnly
-);
-
 static int32 GNumTangentIntermediateBuffers = 1;
 static TAutoConsoleVariable<float> CVarGPUSkinNumTangentIntermediateBuffers(
 	TEXT("r.SkinCache.NumTangentIntermediateBuffers"),
@@ -191,7 +182,7 @@ static inline bool DoesPlatformSupportGPUSkinCache(const FStaticShaderPlatform P
 
 ENGINE_API bool IsGPUSkinCacheAvailable(EShaderPlatform Platform)
 {
-	return (AreSkinCacheShadersEnabled(Platform) != 0 || GForceRecomputeTangents != 0) && DoesPlatformSupportGPUSkinCache(Platform);
+	return AreSkinCacheShadersEnabled(Platform) != 0 && DoesPlatformSupportGPUSkinCache(Platform);
 }
 
 ENGINE_API bool GPUSkinCacheNeedsDuplicatedVertices()
@@ -217,7 +208,7 @@ ENGINE_API bool DoRecomputeSkinTangentsOnGPU_RT()
 {
 	// currently only implemented and tested on Window SM5 (needs Compute, Atomics, SRV for index buffers, UAV for VertexBuffers)
 	//#todo-gpuskin: Enable on PS4 when SRVs for IB exist
-	return DoesPlatformSupportGPUSkinCache(GMaxRHIShaderPlatform) && GEnableGPUSkinCacheShaders != 0 && ((GEnableGPUSkinCache && GSkinCacheRecomputeTangents != 0) || GForceRecomputeTangents != 0);
+	return DoesPlatformSupportGPUSkinCache(GMaxRHIShaderPlatform) && GEnableGPUSkinCacheShaders != 0 && (GEnableGPUSkinCache && GSkinCacheRecomputeTangents != 0);
 }
 
 // determine if during DispatchUpdateSkinning caching should occur
@@ -510,7 +501,7 @@ public:
 
 		InTargetVertexFactory->InvalidateStreams();
 
-		int32 RecomputeTangentsMode = GForceRecomputeTangents > 0 ? 1 : GSkinCacheRecomputeTangents;
+		int32 RecomputeTangentsMode = GSkinCacheRecomputeTangents;
 		if (RecomputeTangentsMode > 0)
 		{
 			if (Section->bRecomputeTangent || RecomputeTangentsMode == 1)
@@ -1565,7 +1556,7 @@ bool FGPUSkinCache::ProcessEntry(
 		InvalidateAllEntries();
 	}
 
-	int32 RecomputeTangentsMode = GForceRecomputeTangents > 0 ? 1 : GSkinCacheRecomputeTangents;
+	int32 RecomputeTangentsMode = GSkinCacheRecomputeTangents;
 	bool bShouldRecomputeTangent = false;
 
 	// IntermediateAccumulatedTangents buffer is needed if mesh has at least one section needing recomputing tangents.
diff --git a/Engine/Source/Runtime/Engine/Private/Rendering/MorphTargetVertexInfoBuffers.cpp b/Engine/Source/Runtime/Engine/Private/Rendering/MorphTargetVertexInfoBuffers.cpp
--- a/Engine/Source/Runtime/Engine/Private/Rendering/MorphTargetVertexInfoBuffers.cpp
+++ b/Engine/Source/Runtime/Engine/Private/Rendering/MorphTargetVertexInfoBuffers.cpp
@@ -5,7 +5,6 @@
 #include "Rendering/SkeletalMeshLODRenderData.h"
 #include "Animation/MorphTarget.h"
 
-extern int32 GForceRecomputeTangents;
 extern int32 GSkinCacheRecomputeTangents;
 
 void FMorphTargetVertexInfoBuffers::InitRHI()
@@ -215,7 +214,7 @@ void FMorphTargetVertexInfoBuffers::InitMorphResources(EShaderPlatform ShaderPla
 	MinimumValuePerMorph.Empty(MorphTargets.Num());
 
 	// Mark vertices that are in a section that doesn't recompute tangents as needing tangents
-	const int32 RecomputeTangentsMode = GForceRecomputeTangents > 0 ? 1 : GSkinCacheRecomputeTangents;
+	const int32 RecomputeTangentsMode = GSkinCacheRecomputeTangents;
 	TBitArray<> VertexNeedsTangents;
 	VertexNeedsTangents.Init(false, NumVertices);
 	for (const FSkelMeshRenderSection& RenderSection : RenderSections)
