Anim blueprints now dont participate in unneeded dependency compilation

This is because dependent BPs at least need a re-link to pick up layout changes in their property access 'bytecode'.

The also requires some re-arranging of where sparse class data gets created, linked and patched.
 - Previously sparse class data was patched in the linker at the point of re-creation, but this meant that stale TRASHCLASS sparse class data was could still be in the struct hierarchy when we are BATCH compiling (when batch compiling, each compilation phase is performed in turn for all compiling/dependent BPs). This change removes the sparse class data when the class is sanitized prior to full compilation and stashes it for later patching. This ensures that sparse class data doesnt get linked into stale structs that can then get GCed later, invalidating property links.
- Sparse class data is now linked when its containing ABPGC is linked. It also ensures that the hierarchy of sparse class data mirrors that of the ABPGC. The 'mirror hierarchy' is currently only an invariant of anim blueprints.

Part of this change alters the behavior of UClass::SetSparseClassDataStruct when passed a NULL value. This will now ensure that subclasses are left in a valid state when this occurs. Previously if sparse class data was nulled-out, subclasses would be left pointing at a potentially stale struct in the case where the struct was being regenerated.

Added an override point for subsystems to participate in Link() and made sure property access offsets are patched there, to catch cases where dependent BPs get either a bytecode-only or link-only compile

diff --git a/Engine/Plugins/Runtime/PropertyAccess/Source/PropertyAccessEditor/Private/PropertyAccessEditor.cpp b/Engine/Plugins/Runtime/PropertyAccess/Source/PropertyAccessEditor/Private/PropertyAccessEditor.cpp
--- a/Engine/Plugins/Runtime/PropertyAccess/Source/PropertyAccessEditor/Private/PropertyAccessEditor.cpp
+++ b/Engine/Plugins/Runtime/PropertyAccess/Source/PropertyAccessEditor/Private/PropertyAccessEditor.cpp
@@ -748,7 +748,7 @@ bool FPropertyAccessLibraryCompiler::FinishCompilation()
 
 		// Always rebuild the library even if we detected a 'failure'. Otherwise we could fail to copy data for both
 		// valid and invalid copies 
-		PropertyAccess::PostLoadLibrary(*Library);
+		PropertyAccess::PatchPropertyOffsets(*Library);
 
 		return bResult;
 	}
diff --git a/Engine/Source/Editor/AnimGraph/Private/AnimBlueprintCompiler.cpp b/Engine/Source/Editor/AnimGraph/Private/AnimBlueprintCompiler.cpp
--- a/Engine/Source/Editor/AnimGraph/Private/AnimBlueprintCompiler.cpp
+++ b/Engine/Source/Editor/AnimGraph/Private/AnimBlueprintCompiler.cpp
@@ -54,6 +54,7 @@ FAnimBlueprintCompilerContext::FAnimBlueprintCompilerContext(UAnimBlueprint* Sou
 	, NewAnimBlueprintMutables(nullptr)
 	, NewMutablesProperty(nullptr)
 	, AnimBlueprint(SourceSketch)
+	, OldSparseClassDataStruct(nullptr)
 	, bIsDerivedAnimBlueprint(false)
 {
 	// Add the animation graph schema to skip default function processing on them
@@ -1255,8 +1256,6 @@ void FAnimBlueprintCompilerContext::RecreateSparseClassData()
 {
 	UAnimBlueprintGeneratedClass* NewAnimBlueprintClass = GetNewAnimBlueprintClass();
 
-	UScriptStruct* OldSparseClassDataStruct = NewAnimBlueprintClass->GetSparseClassDataStruct();
-	
 	// Set up our sparse class data struct
 	if (bIsDerivedAnimBlueprint)
 	{
@@ -1381,6 +1380,10 @@ void FAnimBlueprintCompilerContext::CleanAndSanitizeClass(UBlueprintGeneratedCla
 	UAnimBlueprint* RootAnimBP = UAnimBlueprint::FindRootAnimBlueprint(AnimBlueprint);
 	bIsDerivedAnimBlueprint = RootAnimBP != NULL;
 
+	// Cleanup sparse class data & stash a reference to patch the linker later in RecreateSparseClassData
+	OldSparseClassDataStruct = AnimBlueprintClassToClean->GetSparseClassDataStruct();
+	AnimBlueprintClassToClean->ClearSparseClassDataStruct();
+	
 	FAnimBlueprintGeneratedClassCompiledData CompiledData(AnimBlueprintClassToClean);
 	FAnimBlueprintCompilationBracketContext CompilerContext(this);
 	UAnimBlueprintExtension::ForEachExtension(AnimBlueprint, [this, &CompiledData, &CompilerContext, AnimBlueprintClassToClean](UAnimBlueprintExtension* InExtension)
diff --git a/Engine/Source/Editor/AnimGraph/Private/AnimBlueprintCompiler.h b/Engine/Source/Editor/AnimGraph/Private/AnimBlueprintCompiler.h
--- a/Engine/Source/Editor/AnimGraph/Private/AnimBlueprintCompiler.h
+++ b/Engine/Source/Editor/AnimGraph/Private/AnimBlueprintCompiler.h
@@ -85,6 +85,9 @@ protected:
 	FStructProperty* NewMutablesProperty;
 	UAnimBlueprint* AnimBlueprint;
 
+	// Old sparse class data stored to patchup linker when doing a full compile
+	UScriptStruct* OldSparseClassDataStruct;
+	
 	UAnimationGraphSchema* AnimSchema;
 
 	// Map of allocated v3 nodes that are members of the class
diff --git a/Engine/Source/Editor/Kismet/Private/BlueprintCompilationManager.cpp b/Engine/Source/Editor/Kismet/Private/BlueprintCompilationManager.cpp
--- a/Engine/Source/Editor/Kismet/Private/BlueprintCompilationManager.cpp
+++ b/Engine/Source/Editor/Kismet/Private/BlueprintCompilationManager.cpp
@@ -34,6 +34,7 @@
 #include "UObject/UObjectHash.h"
 #include "Kismet2/KismetDebugUtilities.h"
 #include "BlueprintEditorModule.h"
+#include "Animation/AnimBlueprint.h"
 #include "Stats/StatsHierarchical.h"
 
 extern UNREALED_API UUnrealEdEngine* GUnrealEd;
@@ -1009,6 +1010,14 @@ void FBlueprintCompilationManagerImpl::FlushCompilationQueueImpl(bool bSuppressB
 						return false;
 					}
 
+					// Anim BPs cannot skip un-needed dependency compilation as their property access bytecode
+					// will need refreshing due to external class layouts changing (they require at least a bytecode recompile or a relink)
+					const bool bIsAnimBlueprintClass = !!Cast<UAnimBlueprint>(Data.BP);
+					if(bIsAnimBlueprintClass)
+					{
+						return false;
+					}
+					
 					// if our parent is still being compiled, then we still need to be compiled:
 					UClass* Iter = Data.BP->ParentClass;
 					while(Iter)
diff --git a/Engine/Source/Runtime/CoreUObject/Private/UObject/Class.cpp b/Engine/Source/Runtime/CoreUObject/Private/UObject/Class.cpp
--- a/Engine/Source/Runtime/CoreUObject/Private/UObject/Class.cpp
+++ b/Engine/Source/Runtime/CoreUObject/Private/UObject/Class.cpp
@@ -5197,7 +5197,8 @@ void UClass::SetSparseClassDataStruct(UScriptStruct* InSparseClassDataStruct)
 { 
 	if (SparseClassDataStruct != InSparseClassDataStruct)
 	{
-		if (SparseClassDataStruct && InSparseClassDataStruct)
+		// Passing nullptr as InSparseClassDataStruct is a valid way to clear sparse class data
+		if (SparseClassDataStruct)
 		{
 			// Find all subclasses and point the SuperClass of their SparseClassDataStruct to point to this new SparseClassDataStruct.
 			// We have to do this when compilation creates a new SparseClassDataStruct that has already loaded subclasses.
@@ -5212,6 +5213,9 @@ void UClass::SetSparseClassDataStruct(UScriptStruct* InSparseClassDataStruct)
 				UScriptStruct* SubClassSparseClassDataStruct = SubClass->GetSparseClassDataStruct();
 				if (SubClassSparseClassDataStruct && SubClassSparseClassDataStruct->GetSuperStruct() == SparseClassDataStruct)
 				{
+					// As we are potentialy completely changing the struct layout, we need to cleanup any data we have 
+					// before the superstruct link gets set
+					SubClass->CleanupSparseClassData();
 					SubClassSparseClassDataStruct->SetSuperStruct(InSparseClassDataStruct);
 				}
 			}
@@ -5223,6 +5227,30 @@ void UClass::SetSparseClassDataStruct(UScriptStruct* InSparseClassDataStruct)
 	}
 }
 
+void UClass::ClearSparseClassDataStruct()
+{ 
+	if (SparseClassDataStruct != nullptr)
+	{
+		// Find all subclasses and clear their sparse class data struct as well.
+		TArray<UClass*> SubClasses;
+		GetDerivedClasses(this, SubClasses, true /* bRecursive */);
+		for (UClass* SubClass : SubClasses)
+		{
+			UScriptStruct* SubClassSparseClassDataStruct = SubClass->GetSparseClassDataStruct();
+			if (SubClassSparseClassDataStruct && SubClassSparseClassDataStruct->GetSuperStruct() == SparseClassDataStruct)
+			{
+				SubClass->CleanupSparseClassData();
+				SubClassSparseClassDataStruct->SetSuperStruct(nullptr);
+			}
+		}
+
+		CleanupSparseClassData();
+
+		SparseClassDataStruct = nullptr;
+	}
+}
+
+
 #if WITH_RELOAD
 
 bool UClass::HotReloadPrivateStaticClass(
diff --git a/Engine/Source/Runtime/CoreUObject/Public/UObject/Class.h b/Engine/Source/Runtime/CoreUObject/Public/UObject/Class.h
--- a/Engine/Source/Runtime/CoreUObject/Public/UObject/Class.h
+++ b/Engine/Source/Runtime/CoreUObject/Public/UObject/Class.h
@@ -2848,6 +2848,9 @@ public:
 
 	void SetSparseClassDataStruct(UScriptStruct* InSparseClassDataStruct);
 
+	/** Clears the sparse class data struct for this and all child classes that directly reference it as a super-struct */
+	void ClearSparseClassDataStruct();
+
 	/** Assemble reference token streams for all classes if they haven't had it assembled already */
 	static void AssembleReferenceTokenStreams();
 
diff --git a/Engine/Source/Runtime/Engine/Private/Animation/AnimBlueprintGeneratedClass.cpp b/Engine/Source/Runtime/Engine/Private/Animation/AnimBlueprintGeneratedClass.cpp
--- a/Engine/Source/Runtime/Engine/Private/Animation/AnimBlueprintGeneratedClass.cpp
+++ b/Engine/Source/Runtime/Engine/Private/Animation/AnimBlueprintGeneratedClass.cpp
@@ -469,6 +469,23 @@ void UAnimBlueprintGeneratedClass::Link(FArchive& Ar, bool bRelinkExistingProper
 	{
 		MutablesStruct->Link(Ar, bRelinkExistingProperties);
 	}
+
+	// Link in sparse class data - it must mirror the generated class hierarchy
+	if(UAnimBlueprintGeneratedClass* ParentClass = Cast<UAnimBlueprintGeneratedClass>(GetSuperClass()))
+	{
+		if(UScriptStruct* ParentSparseClassDataStruct = ParentClass->GetSparseClassDataStruct())
+		{
+			if(SparseClassDataStruct && ParentSparseClassDataStruct != SparseClassDataStruct)
+			{
+				// Ensure parent is linked before setting super
+				ParentSparseClassDataStruct->Link(Ar, bRelinkExistingProperties);
+				SparseClassDataStruct->SetSuperStruct(ParentSparseClassDataStruct);
+
+				// Link sparse class data now it is correctly parented
+				SparseClassDataStruct->Link(Ar, bRelinkExistingProperties);
+			}
+		}
+	}
 #endif	
 	
 	Super::Link(Ar, bRelinkExistingProperties);
@@ -528,6 +545,18 @@ void UAnimBlueprintGeneratedClass::Link(FArchive& Ar, bool bRelinkExistingProper
 			}
 		}
 	}
+
+	// Must build constant properties to be able to iterate subsystems
+	BuildConstantProperties();
+	
+	ForEachSubsystem([this](const FAnimSubsystemContext& InContext)
+	{
+		Subsystems.Add(&InContext.Subsystem);
+
+		FAnimSubsystemLinkContext Context(InContext, *this);
+		const_cast<FAnimSubsystem&>(InContext.Subsystem).OnLink(Context);
+		return EAnimSubsystemEnumeration::Continue;
+	});	
 }
 
 void UAnimBlueprintGeneratedClass::PurgeClass(bool bRecompilingOnLoad)
diff --git a/Engine/Source/Runtime/Engine/Private/Animation/AnimSubsystem_PropertyAccess.cpp b/Engine/Source/Runtime/Engine/Private/Animation/AnimSubsystem_PropertyAccess.cpp
--- a/Engine/Source/Runtime/Engine/Private/Animation/AnimSubsystem_PropertyAccess.cpp
+++ b/Engine/Source/Runtime/Engine/Private/Animation/AnimSubsystem_PropertyAccess.cpp
@@ -31,5 +31,13 @@ void FAnimSubsystem_PropertyAccess::OnPostUpdate_WorkerThread(FAnimSubsystemPara
 void FAnimSubsystem_PropertyAccess::OnPostLoad(FAnimSubsystemPostLoadContext& InContext)
 {
 	// Patch the library on load to fixup property offsets
-	PropertyAccess::PostLoadLibrary(Library);
-}
\ No newline at end of file
+	PropertyAccess::PatchPropertyOffsets(Library);
+}
+
+#if WITH_EDITORONLY_DATA
+void FAnimSubsystem_PropertyAccess::OnLink(FAnimSubsystemLinkContext& InContext)
+{
+	// Patch the library on load to fixup property offsets
+	PropertyAccess::PatchPropertyOffsets(Library);
+}
+#endif
\ No newline at end of file
diff --git a/Engine/Source/Runtime/Engine/Private/PropertyAccess.cpp b/Engine/Source/Runtime/Engine/Private/PropertyAccess.cpp
--- a/Engine/Source/Runtime/Engine/Private/PropertyAccess.cpp
+++ b/Engine/Source/Runtime/Engine/Private/PropertyAccess.cpp
@@ -212,7 +212,7 @@ struct FPropertyAccessSystem
 		return true;
 	}
 
-	static void PostLoadLibrary(FPropertyAccessLibrary& InLibrary)
+	static void PatchPropertyOffsets(FPropertyAccessLibrary& InLibrary)
 	{
 		InLibrary.Indirections.Reset();
 
@@ -395,7 +395,7 @@ struct FPropertyAccessSystem
 		InFunction->Invoke(InObject, Stack, OutRetValue);
 	}
 
-	static void IterateAccess(void* InContainer, const FPropertyAccessLibrary& InLibrary, const FPropertyAccessIndirectionChain& InAccess, TFunctionRef<void(void*)> InAddressFunction, TFunctionRef<void(UObject*)> InPerObjectFunction)
+	static void IterateAccess(void* InContainer, const FPropertyAccessLibrary& InLibrary, const FPropertyAccessIndirectionChain& InAccess, TFunctionRef<void(void*)> InAddressFunction)
 	{
 		void* Address = InContainer;
 
@@ -416,10 +416,6 @@ struct FPropertyAccessSystem
 				{
 					UObject* Object = *reinterpret_cast<UObject**>(static_cast<uint8*>(Address) + Indirection.Offset);
 					Address = static_cast<void*>(Object);
-					if(Object != nullptr)
-					{
-						InPerObjectFunction(Object);
-					}
 					break;
 				}
 				case EPropertyAccessObjectType::WeakObject:
@@ -427,10 +423,6 @@ struct FPropertyAccessSystem
 					TWeakObjectPtr<UObject>& WeakObjectPtr = *reinterpret_cast<TWeakObjectPtr<UObject>*>(static_cast<uint8*>(Address) + Indirection.Offset);
 					UObject* Object = WeakObjectPtr.Get();
 					Address = static_cast<void*>(Object);
-					if(Object != nullptr)
-					{
-						InPerObjectFunction(Object);
-					}
 					break;
 				}
 				case EPropertyAccessObjectType::SoftObject:
@@ -438,10 +430,6 @@ struct FPropertyAccessSystem
 					FSoftObjectPtr& SoftObjectPtr = *reinterpret_cast<FSoftObjectPtr*>(static_cast<uint8*>(Address) + Indirection.Offset);
 					UObject* Object = SoftObjectPtr.Get();
 					Address = static_cast<void*>(Object);
-					if(Object != nullptr)
-					{
-						InPerObjectFunction(Object);
-					}
 					break;
 				}
 				default:
@@ -499,10 +487,6 @@ struct FPropertyAccessSystem
 					{
 						UObject* Object = *static_cast<UObject**>(Address);
 						Address = static_cast<void*>(Object);
-						if(Object != nullptr)
-						{
-							InPerObjectFunction(Object);
-						}
 						break;
 					}
 					case EPropertyAccessObjectType::WeakObject:
@@ -510,10 +494,6 @@ struct FPropertyAccessSystem
 						TWeakObjectPtr<UObject>& WeakObjectPtr = *static_cast<TWeakObjectPtr<UObject>*>(Address);
 						UObject* Object = WeakObjectPtr.Get();
 						Address = static_cast<void*>(Object);
-						if(Object != nullptr)
-						{
-							InPerObjectFunction(Object);
-						}
 						break;
 					}
 					case EPropertyAccessObjectType::SoftObject:
@@ -521,10 +501,6 @@ struct FPropertyAccessSystem
 						FSoftObjectPtr& SoftObjectPtr = *reinterpret_cast<FSoftObjectPtr*>(Address);
 						UObject* Object = SoftObjectPtr.Get();
 						Address = static_cast<void*>(Object);
-						if(Object != nullptr)
-						{
-							InPerObjectFunction(Object);
-						}
 						break;
 					}
 					default:
@@ -547,19 +523,13 @@ struct FPropertyAccessSystem
 			InAddressFunction(Address);
 		}
 	}
-
-	// Iterates all the objects in an access path/indirection graph.
-	static void ForEachObjectInAccess(void* InContainer, const FPropertyAccessLibrary& InLibrary, const FPropertyAccessIndirectionChain& InAccess, TFunctionRef<void(UObject*)> InPerObjectFunction)
-	{
-		IterateAccess(InContainer, InLibrary, InAccess, [](void*){}, InPerObjectFunction);
-	}
-
+	
 	// Gets the address that corresponds to a property access.
 	// Forwards the address onto the passed-in function. This callback-style approach is used because in some cases 
 	// (e.g. functions), the address may be memory allocated on the stack.
 	static void GetAccessAddress(void* InContainer, const FPropertyAccessLibrary& InLibrary, const FPropertyAccessIndirectionChain& InAccess, TFunctionRef<void(void*)> InAddressFunction)
 	{
-		IterateAccess(InContainer, InLibrary, InAccess, InAddressFunction, [](UObject*){});
+		IterateAccess(InContainer, InLibrary, InAccess, InAddressFunction);
 	}
 
 	// Process a single copy
@@ -571,18 +541,18 @@ struct FPropertyAccessSystem
 			{
 				const FPropertyAccessCopy& Copy = InLibrary.CopyBatchArray[InBatchId].Copies[InCopyIndex];
 				const FPropertyAccessIndirectionChain& SrcAccess = InLibrary.SrcAccesses[Copy.AccessIndex];
-				if(SrcAccess.Property.Get())
+				if(FProperty* SourceProperty = SrcAccess.Property.Get())
 				{
-					GetAccessAddress(InContainer, InLibrary, SrcAccess, [InContainer, &InLibrary, &Copy, &SrcAccess, &InPostCopyOperation](void* InSrcAddress)
+					GetAccessAddress(InContainer, InLibrary, SrcAccess, [InContainer, &InLibrary, &Copy, SourceProperty, &InPostCopyOperation](void* InSrcAddress)
 					{
 						for(int32 DestAccessIndex = Copy.DestAccessStartIndex; DestAccessIndex < Copy.DestAccessEndIndex; ++DestAccessIndex)
 						{
 							const FPropertyAccessIndirectionChain& DestAccess = InLibrary.DestAccesses[DestAccessIndex];
-							if(DestAccess.Property.Get())
+							if(FProperty* DestProperty = DestAccess.Property.Get())
 							{
-								GetAccessAddress(InContainer, InLibrary, DestAccess, [&InSrcAddress, &Copy, &SrcAccess, &DestAccess, &InPostCopyOperation](void* InDestAddress)
+								GetAccessAddress(InContainer, InLibrary, DestAccess, [&InSrcAddress, &Copy, SourceProperty, DestProperty, &InPostCopyOperation](void* InDestAddress)
 								{
-									PerformCopy(Copy, SrcAccess.Property.Get(), InSrcAddress, DestAccess.Property.Get(), InDestAddress, InPostCopyOperation);
+									PerformCopy(Copy, SourceProperty, InSrcAddress, DestProperty, InDestAddress, InPostCopyOperation);
 								});
 							}
 						}
@@ -616,7 +586,12 @@ namespace PropertyAccess
 {
 	void PostLoadLibrary(FPropertyAccessLibrary& InLibrary)
 	{
-		::FPropertyAccessSystem::PostLoadLibrary(InLibrary);
+		::FPropertyAccessSystem::PatchPropertyOffsets(InLibrary);
+	}
+
+	void PatchPropertyOffsets(FPropertyAccessLibrary& InLibrary)
+	{
+		::FPropertyAccessSystem::PatchPropertyOffsets(InLibrary);
 	}
 
 	void ProcessCopies(UObject* InObject, const FPropertyAccessLibrary& InLibrary, EPropertyAccessCopyBatch InBatchType)
diff --git a/Engine/Source/Runtime/Engine/Public/Animation/AnimSubsystem.h b/Engine/Source/Runtime/Engine/Public/Animation/AnimSubsystem.h
--- a/Engine/Source/Runtime/Engine/Public/Animation/AnimSubsystem.h
+++ b/Engine/Source/Runtime/Engine/Public/Animation/AnimSubsystem.h
@@ -84,6 +84,17 @@ struct FAnimSubsystemPostLoadDefaultsContext
 	UObject* DefaultAnimInstance;
 };
 
+struct FAnimSubsystemLinkContext
+{
+	FAnimSubsystemLinkContext(const FAnimSubsystemContext& InContext, IAnimClassInterface& InAnimClassInterface)
+		: InnerContext(InContext)
+		, AnimClassInterface(InAnimClassInterface)
+	{}
+
+	const FAnimSubsystemContext& InnerContext;
+	IAnimClassInterface& AnimClassInterface;
+};
+
 /** Base structure for all anim subsystem class data */
 USTRUCT()
 struct FAnimSubsystem
@@ -110,4 +121,7 @@ struct FAnimSubsystem
 	
 	/** Override point to perform subsystem class data initialization post-load */
 	virtual void OnPostLoad(FAnimSubsystemPostLoadContext& InContext) {}
+
+	/** Override point to perform subsystem class data initialization on class link */
+	virtual void OnLink(FAnimSubsystemLinkContext& InContext) {}
 };
\ No newline at end of file
diff --git a/Engine/Source/Runtime/Engine/Public/Animation/AnimSubsystem_PropertyAccess.h b/Engine/Source/Runtime/Engine/Public/Animation/AnimSubsystem_PropertyAccess.h
--- a/Engine/Source/Runtime/Engine/Public/Animation/AnimSubsystem_PropertyAccess.h
+++ b/Engine/Source/Runtime/Engine/Public/Animation/AnimSubsystem_PropertyAccess.h
@@ -40,6 +40,9 @@ struct ENGINE_API FAnimSubsystem_PropertyAccess : public FAnimSubsystem
 	virtual void OnPreUpdate_WorkerThread(FAnimSubsystemParallelUpdateContext& InContext) const override;
 	virtual void OnPostUpdate_WorkerThread(FAnimSubsystemParallelUpdateContext& InContext) const override;
 	virtual void OnPostLoad(FAnimSubsystemPostLoadContext& InContext) override;
+#if WITH_EDITORONLY_DATA
+	virtual void OnLink(FAnimSubsystemLinkContext& InContext) override;
+#endif
 
 	/** Access the property access library */
 	const FPropertyAccessLibrary& GetLibrary() const { return Library; }
diff --git a/Engine/Source/Runtime/Engine/Public/PropertyAccess.h b/Engine/Source/Runtime/Engine/Public/PropertyAccess.h
--- a/Engine/Source/Runtime/Engine/Public/PropertyAccess.h
+++ b/Engine/Source/Runtime/Engine/Public/PropertyAccess.h
@@ -44,11 +44,15 @@ namespace PropertyAccess
 		int32 Id;
 	};
 
+
+	UE_DEPRECATED(5.0, "Please use PatchPropertyOffsets instead")
+	ENGINE_API extern void PostLoadLibrary(FPropertyAccessLibrary& InLibrary);
+
 	/** 
-	 * Called to patch up library after it is loaded.
+	 * Called to patch up library after it is loaded or linked
 	 * This converts all FName-based paths into node-based paths that provide an optimized way of accessing properties.
 	 */
-	ENGINE_API extern void PostLoadLibrary(FPropertyAccessLibrary& InLibrary);
+	ENGINE_API extern void PatchPropertyOffsets(FPropertyAccessLibrary& InLibrary);
 
 	/** 
 	 * Process a 'tick' of a property access instance. 
