Remove r.SkinCache.Mode=2

diff --git a/Engine/Source/Runtime/Engine/Private/GPUSkinCache.cpp b/Engine/Source/Runtime/Engine/Private/GPUSkinCache.cpp
--- a/Engine/Source/Runtime/Engine/Private/GPUSkinCache.cpp
+++ b/Engine/Source/Runtime/Engine/Private/GPUSkinCache.cpp
@@ -51,8 +51,7 @@ static TAutoConsoleVariable<int32> CVarEnableGPUSkinCache(
 	TEXT("This will perform skinning on a compute job and not skin on the vertex shader.\n")
 	TEXT("Requires r.SkinCache.CompileShaders=1\n")
 	TEXT(" 0: off\n")
-	TEXT(" 1: on(default)\n")
-	TEXT(" 2: only use skin cache for skinned meshes that ticked the Recompute Tangents checkbox (unavailable in shipping builds)"),
+	TEXT(" 1: on(default)\n"),
 	ECVF_RenderThreadSafe
 );
 
diff --git a/Engine/Source/Runtime/Engine/Private/SkeletalRenderGPUSkin.cpp b/Engine/Source/Runtime/Engine/Private/SkeletalRenderGPUSkin.cpp
--- a/Engine/Source/Runtime/Engine/Private/SkeletalRenderGPUSkin.cpp
+++ b/Engine/Source/Runtime/Engine/Private/SkeletalRenderGPUSkin.cpp
@@ -705,19 +705,6 @@ void FSkeletalMeshObjectGPUSkin::ProcessUpdatedDynamicData(EGPUSkinCacheEntryMod
 					//INC_DWORD_STAT(STAT_GPUSkinCache_SkippedForZeroInfluences);
 					bUseSkinCache = false;
 				}
-
-				{
-#if (UE_BUILD_DEBUG || UE_BUILD_DEVELOPMENT)
-					// In this mode the SkinCache should only be used for RecomputeTangent meshes
-					if (GEnableGPUSkinCache == 2)
-					{
-						if (!Section.bRecomputeTangent)
-						{
-							bUseSkinCache = false;
-						}
-					}
-#endif
-				}
 			}
 
 		#if RHI_RAYTRACING
