IsGPUSkinCacheAvailable() already covers DoesPlatformSupportGPUSkinCache() so remove the 2nd from ShouldCompilePermutation()

diff --git a/Engine/Source/Runtime/Engine/Private/GPUSkinCache.cpp b/Engine/Source/Runtime/Engine/Private/GPUSkinCache.cpp
--- a/Engine/Source/Runtime/Engine/Private/GPUSkinCache.cpp
+++ b/Engine/Source/Runtime/Engine/Private/GPUSkinCache.cpp
@@ -962,7 +962,7 @@ public:
 	static bool ShouldCompilePermutation(const FGlobalShaderPermutationParameters& Parameters)
 	{
 		// currently only implemented and tested on Window SM5 (needs Compute, Atomics, SRV for index buffers, UAV for VertexBuffers)
-		return DoesPlatformSupportGPUSkinCache(Parameters.Platform) && IsGPUSkinCacheAvailable(Parameters.Platform);
+		return IsGPUSkinCacheAvailable(Parameters.Platform);
 	}
 
 	static const uint32 ThreadGroupSizeX = 64;
@@ -1091,7 +1091,7 @@ public:
 	static bool ShouldCompilePermutation(const FGlobalShaderPermutationParameters& Parameters)
 	{
 		// currently only implemented and tested on Window SM5 (needs Compute, Atomics, SRV for index buffers, UAV for VertexBuffers)
-		return DoesPlatformSupportGPUSkinCache(Parameters.Platform) && IsGPUSkinCacheAvailable(Parameters.Platform);
+		return IsGPUSkinCacheAvailable(Parameters.Platform);
 	}
 
 	static const uint32 ThreadGroupSizeX = 64;
