Fix static analysic warning "The 'InMeshComponent' pointer was utilized before it was verified against nullptr."

diff --git a/Engine/Source/Runtime/Engine/Private/SkeletalRenderGPUSkin.cpp b/Engine/Source/Runtime/Engine/Private/SkeletalRenderGPUSkin.cpp
--- a/Engine/Source/Runtime/Engine/Private/SkeletalRenderGPUSkin.cpp
+++ b/Engine/Source/Runtime/Engine/Private/SkeletalRenderGPUSkin.cpp
@@ -2031,7 +2031,7 @@ void FDynamicSkelMeshObjectDataGPUSkin::InitDynamicSkelMeshObjectDataGPUSkin(
 	NumWeightedActiveMorphTargets = 0;
 
 	// Gather any bones referenced by shadow shapes
-	FSkeletalMeshSceneProxy* SkeletalMeshProxy = (FSkeletalMeshSceneProxy*)InMeshComponent->SceneProxy;
+	FSkeletalMeshSceneProxy* SkeletalMeshProxy = InMeshComponent ? (FSkeletalMeshSceneProxy*)InMeshComponent->SceneProxy : nullptr;
 	const TArray<FBoneIndexType>* ExtraRequiredBoneIndices = SkeletalMeshProxy ? &SkeletalMeshProxy->GetSortedShadowBoneIndices() : nullptr;
 
 #if RHI_RAYTRACING
