obj list: skip the expensive measurement of memory usage when called only to report object counts.
#rnx

diff --git a/Engine/Source/Editor/UnrealEd/Private/CookOnTheFlyServer.cpp b/Engine/Source/Editor/UnrealEd/Private/CookOnTheFlyServer.cpp
--- a/Engine/Source/Editor/UnrealEd/Private/CookOnTheFlyServer.cpp
+++ b/Engine/Source/Editor/UnrealEd/Private/CookOnTheFlyServer.cpp
@@ -3476,7 +3476,7 @@ bool UCookOnTheFlyServer::HasExceededMaxMemory() const
 		if (GEngine && !bPerformedObjListWhenNearMaxObjects)
 		{
 			UE_LOG(LogCook, Display, TEXT("Performing 'obj list' to show counts of types of objects due to low availability of UObject indices."));
-			GEngine->Exec(nullptr, TEXT("OBJ LIST -COUNTSORT"));
+			GEngine->Exec(nullptr, TEXT("OBJ LIST -COUNTSORT -SKIPMEMORYSIZE"));
 			bPerformedObjListWhenNearMaxObjects = true;
 		}
 		return true;
diff --git a/Engine/Source/Runtime/Engine/Private/UnrealEngine.cpp b/Engine/Source/Runtime/Engine/Private/UnrealEngine.cpp
--- a/Engine/Source/Runtime/Engine/Private/UnrealEngine.cpp
+++ b/Engine/Source/Runtime/Engine/Private/UnrealEngine.cpp
@@ -4292,6 +4292,14 @@ struct FItem
 		Max += Ar.GetMax();
 		TrueResourceSize += InTrueResourceSize;
 	}
+
+	void Add(SIZE_T InNumBytes, SIZE_T InMaxBytes, FResourceSizeEx InTrueResourceSize)
+	{
+		Count++;
+		Num += InNumBytes;
+		Max += InMaxBytes;
+		TrueResourceSize += InTrueResourceSize;
+	}
 };
 
 struct FSubItem
@@ -8369,7 +8377,7 @@ bool UEngine::HandleObjCommand( const TCHAR* Cmd, FOutputDevice& Ar )
 		FString ObjectNameSubString;
 		FParse::Value(Cmd, TEXT("NAMESUB="), ObjectNameSubString);
 
-		TArray<FItem> List;
+		TMap<UClass*, FItem> ObjectsByClass;
 		TArray<FSubItem> Objects;
 		FItem Total;
 
@@ -8380,6 +8388,7 @@ bool UEngine::HandleObjCommand( const TCHAR* Cmd, FOutputDevice& Ar )
 		}
 
 		const bool bAll				= FParse::Param( Cmd, TEXT("ALL") );
+		const bool bShowMemory = !FParse::Param(Cmd, TEXT("SKIPMEMORYSIZE"));
 
 		// if we specified a parameter in the command, but no objects of that parameter were found,
 		// and they didn't specify "all", then don't list all objects
@@ -8488,11 +8497,16 @@ bool UEngine::HandleObjCommand( const TCHAR* Cmd, FOutputDevice& Ar )
 					}
 				}
 
-				FArchiveCountMem Count( *It );
+				SIZE_T NumBytes = 0;
+				SIZE_T MaxBytes = 0;
 				FResourceSizeEx TrueResourceSize = FResourceSizeEx(EResourceSizeMode::Exclusive);
-				It->GetResourceSizeEx(TrueResourceSize);
-
-				int32 i;
+				if (bShowMemory)
+				{
+					FArchiveCountMem Count(*It);
+					NumBytes = Count.GetNum();
+					MaxBytes = Count.GetNum();
+					It->GetResourceSizeEx(TrueResourceSize);
+				}
 
 				// which class are we going to file this object under? by default, it's class
 				UClass* ClassToUse = It->GetClass();
@@ -8513,24 +8527,15 @@ bool UEngine::HandleObjCommand( const TCHAR* Cmd, FOutputDevice& Ar )
 					}
 				}
 
-				for( i=0; i<List.Num(); i++ )
-				{
-					if( List[i].Class == ClassToUse )
-					{
-						break;
-					}
-				}
-				if( i==List.Num() )
-				{
-					i = List.Add(FItem( ClassToUse ));
-				}
+				FItem& ClassData = ObjectsByClass.FindOrAdd(ClassToUse);
+				ClassData.Class = ClassToUse;
 
 				if( bShowDetailedObjectInfo )
 				{
-					Objects.Add( FSubItem( *It, Count.GetNum(), Count.GetMax(), TrueResourceSize ) );
+					Objects.Add(FSubItem(*It, NumBytes, MaxBytes, TrueResourceSize));
 				}
-				List[i].Add( Count, TrueResourceSize );
-				Total.Add( Count, TrueResourceSize );
+				ClassData.Add(NumBytes, MaxBytes, TrueResourceSize);
+				Total.Add(NumBytes, MaxBytes, TrueResourceSize);
 			}
 		}
 
@@ -8623,7 +8628,7 @@ bool UEngine::HandleObjCommand( const TCHAR* Cmd, FOutputDevice& Ar )
 			Ar.Log(TEXT(""));
 		}
 
-		if( List.Num() )
+		if (!ObjectsByClass.IsEmpty())
 		{
 			struct FCompareFItem
 			{
@@ -8637,68 +8642,111 @@ bool UEngine::HandleObjCommand( const TCHAR* Cmd, FOutputDevice& Ar )
 					return bAlphaSort ? (A.Class->GetName() < B.Class->GetName()) : bCountSort ? (B.Count < A.Count) : bResourceSizeSort ? (B.TrueResourceSize.GetTotalMemoryBytes() < A.TrueResourceSize.GetTotalMemoryBytes()) : (B.Max < A.Max);
 				}
 			};
-			List.Sort( FCompareFItem( bAlphaSort, bCountSort, bResourceSizeSort ) );
+			TArray<FItem> SortedClasses;
+			ObjectsByClass.GenerateValueArray(SortedClasses);
+			SortedClasses.Sort(FCompareFItem(bAlphaSort, bCountSort, bResourceSizeSort));
 			
 			if (bCSV)
 			{
-				Ar.Logf(TEXT(", Class, Count, NumKB, MaxKB, ResExcKB, ResExcDedSysKB, ResExcDedVidKB, ResExcUnkKB"));
+				if (bShowMemory)
+				{
+					Ar.Logf(TEXT(", Class, Count, NumKB, MaxKB, ResExcKB, ResExcDedSysKB, ResExcDedVidKB, ResExcUnkKB"));
+				}
+				else
+				{
+					Ar.Logf(TEXT(", Class, Count"));
+				}
 			}
 			else
 			{
-				Ar.Logf(
-					TEXT(" %100s %8s %10s %10s %10s %15s %15s %15s"),
-					TEXT("Class"),
-					TEXT("Count"),
-					TEXT("NumKB"),
-					TEXT("MaxKB"),
-					TEXT("ResExcKB"),
-					TEXT("ResExcDedSysKB"),
-					TEXT("ResExcDedVidKB"),
-					TEXT("ResExcUnkKB")
-				);
+				if (bShowMemory)
+				{
+					Ar.Logf(
+						TEXT(" %100s %8s %10s %10s %10s %15s %15s %15s"),
+						TEXT("Class"),
+						TEXT("Count"),
+						TEXT("NumKB"),
+						TEXT("MaxKB"),
+						TEXT("ResExcKB"),
+						TEXT("ResExcDedSysKB"),
+						TEXT("ResExcDedVidKB"),
+						TEXT("ResExcUnkKB")
+					);
+				}
+				else
+				{
+					Ar.Logf(TEXT(" %100s %8s"), TEXT("Class"), TEXT("Count"));
+				}
 			}
 
-			for( int32 i=0; i<List.Num(); i++ )
+			for (const FItem& ClassData : SortedClasses)
 			{
 				if (bCSV)
 				{
-					Ar.Logf(TEXT(", %s, %i, %f, %f, %f, %f, %f, %f"),
-						bShowFullClassName ? *List[i].Class->GetFullName() : *List[i].Class->GetName(),
-						(int32)List[i].Count,
-						List[i].Num / 1024.0f,
-						List[i].Max / 1024.0f,
-						List[i].TrueResourceSize.GetTotalMemoryBytes() / 1024.0f,
-						List[i].TrueResourceSize.GetDedicatedSystemMemoryBytes() / 1024.0f,
-						List[i].TrueResourceSize.GetDedicatedVideoMemoryBytes() / 1024.0f,
-						List[i].TrueResourceSize.GetUnknownMemoryBytes() / 1024.0f
-					);
+					if (bShowMemory)
+					{
+						Ar.Logf(TEXT(", %s, %i, %f, %f, %f, %f, %f, %f"),
+							bShowFullClassName ? *ClassData.Class->GetFullName() : *ClassData.Class->GetName(),
+							(int32)ClassData.Count,
+							ClassData.Num / 1024.0f,
+							ClassData.Max / 1024.0f,
+							ClassData.TrueResourceSize.GetTotalMemoryBytes() / 1024.0f,
+							ClassData.TrueResourceSize.GetDedicatedSystemMemoryBytes() / 1024.0f,
+							ClassData.TrueResourceSize.GetDedicatedVideoMemoryBytes() / 1024.0f,
+							ClassData.TrueResourceSize.GetUnknownMemoryBytes() / 1024.0f
+						);
+					}
+					else
+					{
+						Ar.Logf(TEXT(", %s, %i"),
+							bShowFullClassName ? *ClassData.Class->GetFullName() : *ClassData.Class->GetName(),
+							(int32)ClassData.Count
+						);
+					}
 				}
 				else
 				{
-					Ar.Logf(TEXT(" %100s %8i %10.2f %10.2f %10.2f %15.2f %15.2f %15.2f"),
-						bShowFullClassName ? *List[i].Class->GetFullName() : *List[i].Class->GetName(),
-						(int32)List[i].Count,
-						List[i].Num / 1024.0f,
-						List[i].Max / 1024.0f,
-						List[i].TrueResourceSize.GetTotalMemoryBytes() / 1024.0f,
-						List[i].TrueResourceSize.GetDedicatedSystemMemoryBytes() / 1024.0f,
-						List[i].TrueResourceSize.GetDedicatedVideoMemoryBytes() / 1024.0f,
-						List[i].TrueResourceSize.GetUnknownMemoryBytes() / 1024.0f
-					);
+					if (bShowMemory)
+					{
+						Ar.Logf(TEXT(" %100s %8i %10.2f %10.2f %10.2f %15.2f %15.2f %15.2f"),
+							bShowFullClassName ? *ClassData.Class->GetFullName() : *ClassData.Class->GetName(),
+							(int32)ClassData.Count,
+							ClassData.Num / 1024.0f,
+							ClassData.Max / 1024.0f,
+							ClassData.TrueResourceSize.GetTotalMemoryBytes() / 1024.0f,
+							ClassData.TrueResourceSize.GetDedicatedSystemMemoryBytes() / 1024.0f,
+							ClassData.TrueResourceSize.GetDedicatedVideoMemoryBytes() / 1024.0f,
+							ClassData.TrueResourceSize.GetUnknownMemoryBytes() / 1024.0f
+						);
+					}
+					else
+					{
+						Ar.Logf(TEXT(" %100s %8i"),
+							bShowFullClassName ? *ClassData.Class->GetFullName() : *ClassData.Class->GetName(),
+							(int32)ClassData.Count
+						);
+					}
 				}
 			}
 			Ar.Log( TEXT("") );
 		}
-		Ar.Logf( 
-			TEXT("%i Objects (Total: %.3fM / Max: %.3fM / Res: %.3fM | ResDedSys: %.3fM / ResDedVid: %.3fM / ResUnknown: %.3fM)"),
-			Total.Count, 
-			(float)Total.Num/1024.0/1024.0, 
-			(float)Total.Max/1024.0/1024.0, 
-			(float)Total.TrueResourceSize.GetTotalMemoryBytes()/1024.0/1024.0,
-			(float)Total.TrueResourceSize.GetDedicatedSystemMemoryBytes()/1024.0/1024.0, 
-			(float)Total.TrueResourceSize.GetDedicatedVideoMemoryBytes()/1024.0/1024.0, 
-			(float)Total.TrueResourceSize.GetUnknownMemoryBytes()/1024.0/1024.0
-		);
+		if (bShowMemory)
+		{
+			Ar.Logf(
+				TEXT("%i Objects (Total: %.3fM / Max: %.3fM / Res: %.3fM | ResDedSys: %.3fM / ResDedVid: %.3fM / ResUnknown: %.3fM)"),
+				Total.Count,
+				(float)Total.Num / 1024.0 / 1024.0,
+				(float)Total.Max / 1024.0 / 1024.0,
+				(float)Total.TrueResourceSize.GetTotalMemoryBytes() / 1024.0 / 1024.0,
+				(float)Total.TrueResourceSize.GetDedicatedSystemMemoryBytes() / 1024.0 / 1024.0,
+				(float)Total.TrueResourceSize.GetDedicatedVideoMemoryBytes() / 1024.0 / 1024.0,
+				(float)Total.TrueResourceSize.GetUnknownMemoryBytes() / 1024.0 / 1024.0
+			);
+		}
+		else
+		{
+			Ar.Logf(TEXT("%i Objects"), Total.Count);
+		}
 		return true;
 
 	}
