Fixed actor texture streaming hash being non-deterministic
* Texture->IsStreamable() isn't reliable unless the texture has been compiled.
* Rely on GetTextureIsStreamable() instead (now publicly exposed)
* The BuildActorTextureStreamingData() call made during HLOD builds caused actors to be resaved (& checked out, & submitted)

diff --git a/Engine/Source/Runtime/Engine/Classes/Engine/TextureStreamingTypes.h b/Engine/Source/Runtime/Engine/Classes/Engine/TextureStreamingTypes.h
--- a/Engine/Source/Runtime/Engine/Classes/Engine/TextureStreamingTypes.h
+++ b/Engine/Source/Runtime/Engine/Classes/Engine/TextureStreamingTypes.h
@@ -304,6 +304,12 @@ ENGINE_API void BuildActorTextureStreamingData(AActor* InActor, EMaterialQuality
 
 /** Build the texture streaming component data for a level based on actor's UActorTextureStreamingBuildDataComponent. */
 ENGINE_API bool BuildLevelTextureStreamingComponentDataFromActors(ULevel* InLevel);
+
+/** Resolves whether a texture is streamable independent of compression speed, and avoids unnecessary build settings generation work. */
+ENGINE_API bool GetTextureIsStreamable(const UTexture& Texture);
+
+/** Resolves whether a texture is streamable on the given platform, independent of compression speed, and avoids unnecessary build settings generation work. */
+ENGINE_API bool GetTextureIsStreamableOnPlatform(const UTexture& Texture, const ITargetPlatform& TargetPlatform);
 #endif
 
 /** Build the shaders required for the texture streaming build. Returns whether or not the action was successful. */
diff --git a/Engine/Source/Runtime/Engine/Private/Streaming/TextureStreamingBuild.cpp b/Engine/Source/Runtime/Engine/Private/Streaming/TextureStreamingBuild.cpp
--- a/Engine/Source/Runtime/Engine/Private/Streaming/TextureStreamingBuild.cpp
+++ b/Engine/Source/Runtime/Engine/Private/Streaming/TextureStreamingBuild.cpp
@@ -25,6 +25,11 @@ TextureStreamingBuild.cpp : Contains definitions to build texture streaming data
 #include "UnrealEngine.h"
 #if WITH_EDITOR
 #include "Misc/Crc.h"
+#include "Engine/Texture2D.h"
+#include "Engine/Texture2DArray.h"
+#include "Engine/VolumeTexture.h"
+#include "Interfaces/ITargetPlatform.h"
+#include "Interfaces/ITargetPlatformManagerModule.h"
 #endif
 
 
@@ -67,7 +72,7 @@ static uint32 ComputeHashTextureStreamingDataForActor(AActor* InActor)
 	return Hash;
 }
 
-ENGINE_API void BuildActorTextureStreamingData(AActor* InActor, EMaterialQualityLevel::Type InQualityLevel, ERHIFeatureLevel::Type InFeatureLevel)
+void BuildActorTextureStreamingData(AActor* InActor, EMaterialQualityLevel::Type InQualityLevel, ERHIFeatureLevel::Type InFeatureLevel)
 {
 	if (!InActor || InActor->IsTemplate())
 	{
@@ -102,7 +107,7 @@ ENGINE_API void BuildActorTextureStreamingData(AActor* InActor, EMaterialQuality
 	}
 }
 
-ENGINE_API bool BuildLevelTextureStreamingComponentDataFromActors(ULevel* InLevel)
+bool BuildLevelTextureStreamingComponentDataFromActors(ULevel* InLevel)
 {
 	// If Level already contains texture streaming built data, keep it as it was built in the editor.
 	if (!InLevel->StreamingTextureGuids.IsEmpty())
@@ -149,6 +154,36 @@ ENGINE_API bool BuildLevelTextureStreamingComponentDataFromActors(ULevel* InLeve
 	return true;
 }
 
+bool GetTextureIsStreamable(const UTexture& Texture)
+{
+	const ITargetPlatform& CurrentPlatform = *GetTargetPlatformManagerRef().GetRunningTargetPlatform();
+
+	bool bStreamable = GetTextureIsStreamableOnPlatform(Texture, CurrentPlatform);
+	return bStreamable;
+}
+
+bool GetTextureIsStreamableOnPlatform(const UTexture& Texture, const ITargetPlatform& TargetPlatform)
+{
+	const bool bPlatformSupportsTextureStreaming = TargetPlatform.SupportsFeature(ETargetPlatformFeatures::TextureStreaming);
+	
+	bool bStreamable = false;
+	if (Texture.IsA(UTexture2DArray::StaticClass()))
+	{
+		bStreamable = GSupportsTexture2DArrayStreaming;
+	}
+	else if (Texture.IsA(UVolumeTexture::StaticClass()))
+	{
+		bStreamable = GSupportsVolumeTextureStreaming;
+	}
+	else if (Texture.IsA(UTexture2D::StaticClass()))
+	{
+		bStreamable = true;
+	}
+
+	bStreamable &= Texture.IsCandidateForTextureStreaming(&TargetPlatform);
+	return bStreamable;
+}
+
 #else
 
 static bool GetUnpackedTextureStreamingQualityLevelFeatureLevel(uint32 InPackedValue, EMaterialQualityLevel::Type& OutQualityLevel, ERHIFeatureLevel::Type& OutFeatureLevel)
@@ -634,7 +669,7 @@ void FStreamingTextureLevelContext::ProcessMaterial(const FBoxSphereBounds& Comp
 
 	for (UTexture* Texture : Textures)
 	{
-		if (!Texture->IsStreamable())
+		if (!GetTextureIsStreamable(*Texture))
 		{
 			continue;
 		}
diff --git a/Engine/Source/Runtime/Engine/Private/TextureDerivedData.cpp b/Engine/Source/Runtime/Engine/Private/TextureDerivedData.cpp
--- a/Engine/Source/Runtime/Engine/Private/TextureDerivedData.cpp
+++ b/Engine/Source/Runtime/Engine/Private/TextureDerivedData.cpp
@@ -626,32 +626,6 @@ static ETextureEncodeSpeed GetDesiredEncodeSpeed()
 	return FThreadSafeInitializer.CachedEncodeSpeedOption;
 }
 
-//
-// Resolves whether the texture is streamable independent of compression speed, and avoids
-// unnecessary build settings generation work.
-//
-static bool GetTextureIsStreamable(const UTexture& Texture, const ITargetPlatform& CurrentPlatform)
-{
-	const bool bPlatformSupportsTextureStreaming = CurrentPlatform.SupportsFeature(ETargetPlatformFeatures::TextureStreaming);
-
-	bool bStreamable = false;
-	if (Texture.IsA(UTexture2DArray::StaticClass()))
-	{
-		bStreamable = GSupportsTexture2DArrayStreaming;
-	}
-	else if (Texture.IsA(UVolumeTexture::StaticClass()))
-	{
-		bStreamable = GSupportsVolumeTextureStreaming;
-	}
-	else if (Texture.IsA(UTexture2D::StaticClass()))
-	{
-		bStreamable = true;
-	}
-
-	bStreamable &= bPlatformSupportsTextureStreaming && !Texture.NeverStream && (Texture.LODGroup != TEXTUREGROUP_UI);
-	return bStreamable;
-}
-
 /**
  * Sets texture build settings.
  * @param Texture - The texture for which to build compressor settings.
@@ -762,7 +736,7 @@ static void GetTextureBuildSettings(
 	OutBuildSettings.CompositePower = Texture.CompositePower;
 	OutBuildSettings.LODBias = TextureLODSettings.CalculateLODBias(SourceSize.X, SourceSize.Y, Texture.MaxTextureSize, Texture.LODGroup, Texture.LODBias, Texture.NumCinematicMipLevels, Texture.MipGenSettings, bVirtualTextureStreaming);
 	OutBuildSettings.LODBiasWithCinematicMips = TextureLODSettings.CalculateLODBias(SourceSize.X, SourceSize.Y, Texture.MaxTextureSize, Texture.LODGroup, Texture.LODBias, 0, Texture.MipGenSettings, bVirtualTextureStreaming);
-	OutBuildSettings.bStreamable = GetTextureIsStreamable(Texture, CurrentPlatform);
+	OutBuildSettings.bStreamable = GetTextureIsStreamableOnPlatform(Texture, CurrentPlatform);
 	OutBuildSettings.bVirtualStreamable = bVirtualTextureStreaming;
 	OutBuildSettings.PowerOfTwoMode = Texture.PowerOfTwoMode;
 	OutBuildSettings.PaddingColor = Texture.PaddingColor;
@@ -2895,7 +2869,7 @@ void UTexture::SerializeCookedPlatformData(FArchive& Ar)
 				}
 
 				// Pass streamable flag for inlining mips
-				bool bTextureIsStreamable = GetTextureIsStreamable(*this, *Ar.CookingTarget());
+				bool bTextureIsStreamable = GetTextureIsStreamableOnPlatform(*this, *Ar.CookingTarget());
 				PlatformDataToSave->SerializeCooked(Ar, this, bTextureIsStreamable);
 
 				SkipOffset = Ar.Tell() - SkipOffsetLoc;
