[Backout] - CL18394732
Original CL Desc
-----------------------------------------------------------------
Fixed actor texture streaming hash being non-deterministic
* Texture->IsStreamable() isn't reliable unless the texture has been compiled.
* Rely on GetTextureIsStreamable() instead (now publicly exposed)
* The BuildActorTextureStreamingData() call made during HLOD builds caused actors to be resaved (& checked out, & submitted)

diff --git a/Engine/Source/Runtime/Engine/Classes/Engine/TextureStreamingTypes.h b/Engine/Source/Runtime/Engine/Classes/Engine/TextureStreamingTypes.h
--- a/Engine/Source/Runtime/Engine/Classes/Engine/TextureStreamingTypes.h
+++ b/Engine/Source/Runtime/Engine/Classes/Engine/TextureStreamingTypes.h
@@ -304,12 +304,6 @@ ENGINE_API void BuildActorTextureStreamingData(AActor* InActor, EMaterialQuality
 
 /** Build the texture streaming component data for a level based on actor's UActorTextureStreamingBuildDataComponent. */
 ENGINE_API bool BuildLevelTextureStreamingComponentDataFromActors(ULevel* InLevel);
-
-/** Resolves whether a texture is streamable independent of compression speed, and avoids unnecessary build settings generation work. */
-ENGINE_API bool GetTextureIsStreamable(const UTexture& Texture);
-
-/** Resolves whether a texture is streamable on the given platform, independent of compression speed, and avoids unnecessary build settings generation work. */
-ENGINE_API bool GetTextureIsStreamableOnPlatform(const UTexture& Texture, const ITargetPlatform& TargetPlatform);
 #endif
 
 /** Build the shaders required for the texture streaming build. Returns whether or not the action was successful. */
diff --git a/Engine/Source/Runtime/Engine/Private/Streaming/TextureStreamingBuild.cpp b/Engine/Source/Runtime/Engine/Private/Streaming/TextureStreamingBuild.cpp
--- a/Engine/Source/Runtime/Engine/Private/Streaming/TextureStreamingBuild.cpp
+++ b/Engine/Source/Runtime/Engine/Private/Streaming/TextureStreamingBuild.cpp
@@ -25,11 +25,6 @@ TextureStreamingBuild.cpp : Contains definitions to build texture streaming data
 #include "UnrealEngine.h"
 #if WITH_EDITOR
 #include "Misc/Crc.h"
-#include "Engine/Texture2D.h"
-#include "Engine/Texture2DArray.h"
-#include "Engine/VolumeTexture.h"
-#include "Interfaces/ITargetPlatform.h"
-#include "Interfaces/ITargetPlatformManagerModule.h"
 #endif
 
 
@@ -72,7 +67,7 @@ static uint32 ComputeHashTextureStreamingDataForActor(AActor* InActor)
 	return Hash;
 }
 
-void BuildActorTextureStreamingData(AActor* InActor, EMaterialQualityLevel::Type InQualityLevel, ERHIFeatureLevel::Type InFeatureLevel)
+ENGINE_API void BuildActorTextureStreamingData(AActor* InActor, EMaterialQualityLevel::Type InQualityLevel, ERHIFeatureLevel::Type InFeatureLevel)
 {
 	if (!InActor || InActor->IsTemplate())
 	{
@@ -107,7 +102,7 @@ void BuildActorTextureStreamingData(AActor* InActor, EMaterialQualityLevel::Type
 	}
 }
 
-bool BuildLevelTextureStreamingComponentDataFromActors(ULevel* InLevel)
+ENGINE_API bool BuildLevelTextureStreamingComponentDataFromActors(ULevel* InLevel)
 {
 	// If Level already contains texture streaming built data, keep it as it was built in the editor.
 	if (!InLevel->StreamingTextureGuids.IsEmpty())
@@ -154,36 +149,6 @@ bool BuildLevelTextureStreamingComponentDataFromActors(ULevel* InLevel)
 	return true;
 }
 
-bool GetTextureIsStreamable(const UTexture& Texture)
-{
-	const ITargetPlatform& CurrentPlatform = *GetTargetPlatformManagerRef().GetRunningTargetPlatform();
-
-	bool bStreamable = GetTextureIsStreamableOnPlatform(Texture, CurrentPlatform);
-	return bStreamable;
-}
-
-bool GetTextureIsStreamableOnPlatform(const UTexture& Texture, const ITargetPlatform& TargetPlatform)
-{
-	const bool bPlatformSupportsTextureStreaming = TargetPlatform.SupportsFeature(ETargetPlatformFeatures::TextureStreaming);
-	
-	bool bStreamable = false;
-	if (Texture.IsA(UTexture2DArray::StaticClass()))
-	{
-		bStreamable = GSupportsTexture2DArrayStreaming;
-	}
-	else if (Texture.IsA(UVolumeTexture::StaticClass()))
-	{
-		bStreamable = GSupportsVolumeTextureStreaming;
-	}
-	else if (Texture.IsA(UTexture2D::StaticClass()))
-	{
-		bStreamable = true;
-	}
-
-	bStreamable &= Texture.IsCandidateForTextureStreaming(&TargetPlatform);
-	return bStreamable;
-}
-
 #else
 
 static bool GetUnpackedTextureStreamingQualityLevelFeatureLevel(uint32 InPackedValue, EMaterialQualityLevel::Type& OutQualityLevel, ERHIFeatureLevel::Type& OutFeatureLevel)
@@ -669,7 +634,7 @@ void FStreamingTextureLevelContext::ProcessMaterial(const FBoxSphereBounds& Comp
 
 	for (UTexture* Texture : Textures)
 	{
-		if (!GetTextureIsStreamable(*Texture))
+		if (!Texture->IsStreamable())
 		{
 			continue;
 		}
diff --git a/Engine/Source/Runtime/Engine/Private/TextureDerivedData.cpp b/Engine/Source/Runtime/Engine/Private/TextureDerivedData.cpp
--- a/Engine/Source/Runtime/Engine/Private/TextureDerivedData.cpp
+++ b/Engine/Source/Runtime/Engine/Private/TextureDerivedData.cpp
@@ -626,6 +626,32 @@ static ETextureEncodeSpeed GetDesiredEncodeSpeed()
 	return FThreadSafeInitializer.CachedEncodeSpeedOption;
 }
 
+//
+// Resolves whether the texture is streamable independent of compression speed, and avoids
+// unnecessary build settings generation work.
+//
+static bool GetTextureIsStreamable(const UTexture& Texture, const ITargetPlatform& CurrentPlatform)
+{
+	const bool bPlatformSupportsTextureStreaming = CurrentPlatform.SupportsFeature(ETargetPlatformFeatures::TextureStreaming);
+
+	bool bStreamable = false;
+	if (Texture.IsA(UTexture2DArray::StaticClass()))
+	{
+		bStreamable = GSupportsTexture2DArrayStreaming;
+	}
+	else if (Texture.IsA(UVolumeTexture::StaticClass()))
+	{
+		bStreamable = GSupportsVolumeTextureStreaming;
+	}
+	else if (Texture.IsA(UTexture2D::StaticClass()))
+	{
+		bStreamable = true;
+	}
+
+	bStreamable &= bPlatformSupportsTextureStreaming && !Texture.NeverStream && (Texture.LODGroup != TEXTUREGROUP_UI);
+	return bStreamable;
+}
+
 /**
  * Sets texture build settings.
  * @param Texture - The texture for which to build compressor settings.
@@ -736,7 +762,7 @@ static void GetTextureBuildSettings(
 	OutBuildSettings.CompositePower = Texture.CompositePower;
 	OutBuildSettings.LODBias = TextureLODSettings.CalculateLODBias(SourceSize.X, SourceSize.Y, Texture.MaxTextureSize, Texture.LODGroup, Texture.LODBias, Texture.NumCinematicMipLevels, Texture.MipGenSettings, bVirtualTextureStreaming);
 	OutBuildSettings.LODBiasWithCinematicMips = TextureLODSettings.CalculateLODBias(SourceSize.X, SourceSize.Y, Texture.MaxTextureSize, Texture.LODGroup, Texture.LODBias, 0, Texture.MipGenSettings, bVirtualTextureStreaming);
-	OutBuildSettings.bStreamable = GetTextureIsStreamableOnPlatform(Texture, CurrentPlatform);
+	OutBuildSettings.bStreamable = GetTextureIsStreamable(Texture, CurrentPlatform);
 	OutBuildSettings.bVirtualStreamable = bVirtualTextureStreaming;
 	OutBuildSettings.PowerOfTwoMode = Texture.PowerOfTwoMode;
 	OutBuildSettings.PaddingColor = Texture.PaddingColor;
@@ -2869,7 +2895,7 @@ void UTexture::SerializeCookedPlatformData(FArchive& Ar)
 				}
 
 				// Pass streamable flag for inlining mips
-				bool bTextureIsStreamable = GetTextureIsStreamableOnPlatform(*this, *Ar.CookingTarget());
+				bool bTextureIsStreamable = GetTextureIsStreamable(*this, *Ar.CookingTarget());
 				PlatformDataToSave->SerializeCooked(Ar, this, bTextureIsStreamable);
 
 				SkipOffset = Ar.Tell() - SkipOffsetLoc;
