World Partition Grid Placement
  - Deprecated enum EActorGridPlacement and replaced AActor::GridPlacement by bIsSpatiallyLoaded.
  - It makes more sense from a user perspective to set an actor as "non-spatially loaded" instead of "always loaded", especially with data layers.
  - If we ever need it, having the possibility to set actors to use their location or bounds to go in the grid will be a per-grid setting.

diff --git a/Engine/Plugins/Editor/DataValidation/Source/DataValidation/Private/WorldPartitionChangelistValidator.cpp b/Engine/Plugins/Editor/DataValidation/Source/DataValidation/Private/WorldPartitionChangelistValidator.cpp
--- a/Engine/Plugins/Editor/DataValidation/Source/DataValidation/Private/WorldPartitionChangelistValidator.cpp
+++ b/Engine/Plugins/Editor/DataValidation/Source/DataValidation/Private/WorldPartitionChangelistValidator.cpp
@@ -195,16 +195,13 @@ void UWorldPartitionChangelistValidator::OnInvalidReferenceGridPlacement(const F
 {
 	if (Filter(ActorDescView) || Filter(ReferenceActorDescView))
 	{
-		const FString StreamedActor(TEXT("Streamed actor"));
-		const FString AlwaysLoadedActor(TEXT("Always loaded actor"));
-
-		const bool bIsActorDescAlwaysLoaded = ActorDescView.GetGridPlacement() == EActorGridPlacement::AlwaysLoaded;
-		const bool bIsActorDescRefAlwaysLoaded = ReferenceActorDescView.GetGridPlacement() == EActorGridPlacement::AlwaysLoaded;
+		const FString SpatiallyLoadedActor(TEXT("Spatially loaded actor"));
+		const FString NonSpatiallyLoadedActor(TEXT("Non-spatially loaded loaded actor"));
 
 		FText CurrentError = FText::Format(LOCTEXT("DataValidation.Changelist.WorldPartition.InvalidReferenceGridPlacement", "{0} {1} is referencing {2} {3}."),
-											FText::FromString(bIsActorDescAlwaysLoaded ? *AlwaysLoadedActor : *StreamedActor),
+											FText::FromString(ActorDescView.GetIsSpatiallyLoaded() ? *SpatiallyLoadedActor : *NonSpatiallyLoadedActor),
 											FText::FromString(GetPrettyPackageName(ActorDescView)),
-											FText::FromString(bIsActorDescRefAlwaysLoaded ? *AlwaysLoadedActor : *StreamedActor),
+											FText::FromString(ReferenceActorDescView.GetIsSpatiallyLoaded() ? *SpatiallyLoadedActor : *NonSpatiallyLoadedActor),
 											FText::FromString(GetPrettyPackageName(ReferenceActorDescView)));
 
 		Errors->Add(CurrentError);
diff --git a/Engine/Plugins/Runtime/ZoneGraph/Source/ZoneGraph/Private/ZoneGraphData.cpp b/Engine/Plugins/Runtime/ZoneGraph/Source/ZoneGraph/Private/ZoneGraphData.cpp
--- a/Engine/Plugins/Runtime/ZoneGraph/Source/ZoneGraph/Private/ZoneGraphData.cpp
+++ b/Engine/Plugins/Runtime/ZoneGraph/Source/ZoneGraph/Private/ZoneGraphData.cpp
@@ -21,6 +21,10 @@ AZoneGraphData::AZoneGraphData(const FObjectInitializer& ObjectInitializer)
 	RenderingComp->SetVisibility(bEnableDrawing);
 	RootComponent = RenderingComp;
 #endif
+
+#if WITH_EDITORONLY_DATA
+	bIsSpatiallyLoaded = false;
+#endif
 }
 
 
diff --git a/Engine/Plugins/Runtime/ZoneGraph/Source/ZoneGraph/Private/ZoneShapeActor.cpp b/Engine/Plugins/Runtime/ZoneGraph/Source/ZoneGraph/Private/ZoneShapeActor.cpp
--- a/Engine/Plugins/Runtime/ZoneGraph/Source/ZoneGraph/Private/ZoneShapeActor.cpp
+++ b/Engine/Plugins/Runtime/ZoneGraph/Source/ZoneGraph/Private/ZoneShapeActor.cpp
@@ -13,4 +13,8 @@ AZoneShape::AZoneShape(const FObjectInitializer& ObjectInitializer)
 
 	SetHidden(true);
 	SetCanBeDamaged(false);
+
+#if WITH_EDITORONLY_DATA
+	bIsSpatiallyLoaded = false;
+#endif
 }
diff --git a/Engine/Plugins/Runtime/ZoneGraph/Source/ZoneGraph/Public/ZoneGraphData.h b/Engine/Plugins/Runtime/ZoneGraph/Source/ZoneGraph/Public/ZoneGraphData.h
--- a/Engine/Plugins/Runtime/ZoneGraph/Source/ZoneGraph/Public/ZoneGraphData.h
+++ b/Engine/Plugins/Runtime/ZoneGraph/Source/ZoneGraph/Public/ZoneGraphData.h
@@ -26,7 +26,7 @@ public:
 #if WITH_EDITOR
 	virtual void PostEditUndo() override;
 	virtual void PostEditChangeProperty(FPropertyChangedEvent& PropertyChangedEvent) override;
-	virtual EActorGridPlacement GetDefaultGridPlacement() const override { return EActorGridPlacement::AlwaysLoaded; }
+	virtual bool CanChangeIsSpatiallyLoadedFlag() const override { return false; }
 #endif // WITH_EDITOR
 	//~ End UObject/AActor Interface
 
diff --git a/Engine/Plugins/Runtime/ZoneGraph/Source/ZoneGraph/Public/ZoneShapeActor.h b/Engine/Plugins/Runtime/ZoneGraph/Source/ZoneGraph/Public/ZoneShapeActor.h
--- a/Engine/Plugins/Runtime/ZoneGraph/Source/ZoneGraph/Public/ZoneShapeActor.h
+++ b/Engine/Plugins/Runtime/ZoneGraph/Source/ZoneGraph/Public/ZoneShapeActor.h
@@ -20,7 +20,7 @@ public:
 	const UZoneShapeComponent* GetShape() const { return ShapeComponent; }
 
 #if WITH_EDITOR
-	virtual EActorGridPlacement GetDefaultGridPlacement() const override { return EActorGridPlacement::AlwaysLoaded; }
+	virtual bool CanChangeIsSpatiallyLoadedFlag() const override { return false; }
 #endif
 
 protected:
diff --git a/Engine/Plugins/VirtualProduction/LevelSnapshots/Source/LevelSnapshots/Private/Restorability/GridPlacementRestoration.cpp b/Engine/Plugins/VirtualProduction/LevelSnapshots/Source/LevelSnapshots/Private/Restorability/GridPlacementRestoration.cpp
--- a/Engine/Plugins/VirtualProduction/LevelSnapshots/Source/LevelSnapshots/Private/Restorability/GridPlacementRestoration.cpp
+++ b/Engine/Plugins/VirtualProduction/LevelSnapshots/Source/LevelSnapshots/Private/Restorability/GridPlacementRestoration.cpp
@@ -37,8 +37,7 @@ namespace UE::LevelSnapshots::Private::Internal
 		{
 			if (Params.LeafProperty == GridPlacementProperty)
 			{
-				const bool bIsEditable = Params.InspectedClass->GetDefaultObject<AActor>()->GetDefaultGridPlacement() == EActorGridPlacement::None;
-				return bIsEditable ? EPropertyComparison::CheckNormally : EPropertyComparison::TreatEqual;
+				return EPropertyComparison::CheckNormally;
 			}
 
 			return EPropertyComparison::CheckNormally;
diff --git a/Engine/Source/Developer/WorldPartitionHLODUtilities/Private/WorldPartitionHLODUtilities.cpp b/Engine/Source/Developer/WorldPartitionHLODUtilities/Private/WorldPartitionHLODUtilities.cpp
--- a/Engine/Source/Developer/WorldPartitionHLODUtilities/Private/WorldPartitionHLODUtilities.cpp
+++ b/Engine/Source/Developer/WorldPartitionHLODUtilities/Private/WorldPartitionHLODUtilities.cpp
@@ -164,7 +164,7 @@ TArray<AWorldPartitionHLOD*> FWorldPartitionHLODUtilities::CreateHLODActors(FHLO
 	struct FSubActorsInfo
 	{
 		TArray<FGuid>			SubActors;
-		EActorGridPlacement		GridPlacement = EActorGridPlacement::Location;
+		bool					bIsSpatiallyLoaded;
 	};
 	TMap<UHLODLayer*, FSubActorsInfo> SubActorsInfos;
 
@@ -179,9 +179,9 @@ TArray<AWorldPartitionHLOD*> FWorldPartitionHLODUtilities::CreateHLODActors(FHLO
 				FSubActorsInfo& SubActorsInfo = SubActorsInfos.FindOrAdd(HLODLayer);
 
 				SubActorsInfo.SubActors.Add(ActorInstance.Actor);
-				if (ActorDescView.GetGridPlacement() == EActorGridPlacement::Bounds)
+				if (ActorDescView.GetIsSpatiallyLoaded())
 				{
-					SubActorsInfo.GridPlacement = EActorGridPlacement::Bounds;
+					SubActorsInfo.bIsSpatiallyLoaded = true;
 				}
 			}
 		}
@@ -267,12 +267,12 @@ TArray<AWorldPartitionHLOD*> FWorldPartitionHLODUtilities::CreateHLODActors(FHLO
 			bIsDirty = true;
 		}
 
-		// Runtime grid placement
+		// Spatially loaded
 		// HLOD that are always loaded will not take the SubActorsInfo.GridPlacement into account
-		EActorGridPlacement ExpectedGridPlacement = HLODLayer->IsAlwaysLoaded() ? EActorGridPlacement::AlwaysLoaded : SubActorsInfo.GridPlacement;
-		if (HLODActor->GetGridPlacement() != ExpectedGridPlacement)
+		bool bExpectedIsSpatiallyLoaded = HLODLayer->IsAlwaysLoaded() ? false : SubActorsInfo.bIsSpatiallyLoaded;
+		if (HLODActor->GetIsSpatiallyLoaded() != bExpectedIsSpatiallyLoaded)
 		{
-			HLODActor->SetGridPlacement(ExpectedGridPlacement);
+			HLODActor->SetIsSpatiallyLoaded(bExpectedIsSpatiallyLoaded);
 			bIsDirty = true;
 		}
 
diff --git a/Engine/Source/Editor/UnrealEd/Classes/Commandlets/WorldPartitionConvertCommandlet.h b/Engine/Source/Editor/UnrealEd/Classes/Commandlets/WorldPartitionConvertCommandlet.h
--- a/Engine/Source/Editor/UnrealEd/Classes/Commandlets/WorldPartitionConvertCommandlet.h
+++ b/Engine/Source/Editor/UnrealEd/Classes/Commandlets/WorldPartitionConvertCommandlet.h
@@ -44,7 +44,6 @@ public:
 
 private:
 	void GatherAndPrepareSubLevelsToConvert(ULevel* Level, TArray<ULevel*>& SubLevels);
-	EActorGridPlacement GetLevelGridPlacement(ULevel* Level, EActorGridPlacement DefaultGridPlacement);
 
 protected:
 	virtual bool GetAdditionalLevelsToConvert(ULevel* Level, TArray<ULevel*>& SubLevels);
@@ -103,9 +102,6 @@ protected:
 	UPROPERTY(Config)
 	TSubclassOf<UWorldPartitionRuntimeHash> RuntimeHashClass;
 
-	UPROPERTY(Config)
-	TMap<FName, EActorGridPlacement> LevelsGridPlacement;
-
 	// Levels excluded from conversion.
 	UPROPERTY(Config)
 	TArray<FString> ExcludedLevels;
diff --git a/Engine/Source/Editor/UnrealEd/Classes/Editor/GroupActor.h b/Engine/Source/Editor/UnrealEd/Classes/Editor/GroupActor.h
--- a/Engine/Source/Editor/UnrealEd/Classes/Editor/GroupActor.h
+++ b/Engine/Source/Editor/UnrealEd/Classes/Editor/GroupActor.h
@@ -45,7 +45,6 @@ class AGroupActor : public AActor
 	virtual void PostActorCreated() override;
 	virtual void SetIsTemporarilyHiddenInEditor( bool bIsHidden ) override;
 	virtual void GetActorBounds(bool bOnlyCollidingComponents, FVector& Origin, FVector& BoxExtent, bool bIncludeFromChildActors) const override;
-	virtual EActorGridPlacement GetDefaultGridPlacement() const override { return EActorGridPlacement::None; }
 #if WITH_EDITOR
 	virtual FBox GetStreamingBounds() const override;
 #endif
diff --git a/Engine/Source/Editor/UnrealEd/Private/Commandlets/WorldPartitionConvertCommandlet.cpp b/Engine/Source/Editor/UnrealEd/Private/Commandlets/WorldPartitionConvertCommandlet.cpp
--- a/Engine/Source/Editor/UnrealEd/Private/Commandlets/WorldPartitionConvertCommandlet.cpp
+++ b/Engine/Source/Editor/UnrealEd/Private/Commandlets/WorldPartitionConvertCommandlet.cpp
@@ -272,16 +272,6 @@ void UWorldPartitionConvertCommandlet::GatherAndPrepareSubLevelsToConvert(ULevel
 	}
 }
 
-EActorGridPlacement UWorldPartitionConvertCommandlet::GetLevelGridPlacement(ULevel* Level, EActorGridPlacement DefaultGridPlacement)
-{
-	FString WorldPath = Level->GetPackage()->GetName();
-	if (EActorGridPlacement* CustomLevelGridPlacement = LevelsGridPlacement.Find(*WorldPath))
-	{
-		return *CustomLevelGridPlacement;
-	}
-	return DefaultGridPlacement;
-}
-
 bool UWorldPartitionConvertCommandlet::PrepareStreamingLevelForConversion(ULevelStreaming* StreamingLevel)
 {
 	TRACE_CPUPROFILER_EVENT_SCOPE(UWorldPartitionConvertCommandlet::PrepareStreamingLevelForConversion);
@@ -292,16 +282,13 @@ bool UWorldPartitionConvertCommandlet::PrepareStreamingLevelForConversion(ULevel
 	if (bOnlyMergeSubLevels || StreamingLevel->ShouldBeAlwaysLoaded() || StreamingLevel->bDisableDistanceStreaming)
 	{
 		FString WorldPath = SubLevel->GetPackage()->GetName();
-		if (!LevelsGridPlacement.Contains(*WorldPath))
-		{
-			UE_LOG(LogWorldPartitionConvertCommandlet, Log, TEXT("Converting %s streaming level %s"), StreamingLevel->bDisableDistanceStreaming ? TEXT("non distance-based") : TEXT("always loaded"), *StreamingLevel->GetWorldAssetPackageName());
+		UE_LOG(LogWorldPartitionConvertCommandlet, Log, TEXT("Converting %s streaming level %s"), StreamingLevel->bDisableDistanceStreaming ? TEXT("non distance-based") : TEXT("always loaded"), *StreamingLevel->GetWorldAssetPackageName());
 
-			for (AActor* Actor: SubLevel->Actors)
+		for (AActor* Actor: SubLevel->Actors)
+		{
+			if (Actor)
 			{
-				if (Actor)
-				{
-					Actor->SetGridPlacement(EActorGridPlacement::AlwaysLoaded);
-				}
+				Actor->SetIsSpatiallyLoaded(false);
 			}
 		}
 	}
@@ -960,7 +947,7 @@ int32 UWorldPartitionConvertCommandlet::Main(const FString& Params)
 		}
 	};
 
-	auto PrepareLevelActors = [this, PartitionFoliage, PartitionLandscape, MainWorldDataLayers](ULevel* Level, TArray<AActor*>& Actors, bool bMainLevel, EActorGridPlacement DefaultGridPlacement)
+	auto PrepareLevelActors = [this, PartitionFoliage, PartitionLandscape, MainWorldDataLayers](ULevel* Level, TArray<AActor*>& Actors, bool bMainLevel)
 	{
 		TRACE_CPUPROFILER_EVENT_SCOPE(PrepareLevelActors);
 
@@ -990,18 +977,14 @@ int32 UWorldPartitionConvertCommandlet::Main(const FString& Params)
 						check(LandscapeInfo);
 						LandscapeInfos.Add(LandscapeInfo);
 					}
-					// Only override default grid placement on actors that are not marked as always loaded
-					else if (Actor->GetGridPlacement() != EActorGridPlacement::AlwaysLoaded)
+					// Only override default grid placement on actors that are spatially loaded
+					else if (Actor->GetIsSpatiallyLoaded())
 					{
 						const FBox ActorBounds = Actor->GetStreamingBounds();
 
 						if (!WorldBounds.IsInside(ActorBounds))
 						{
-							Actor->SetGridPlacement(EActorGridPlacement::AlwaysLoaded);
-						}
-						else
-						{
-							Actor->SetGridPlacement(DefaultGridPlacement);
+							Actor->SetIsSpatiallyLoaded(false);
 						}
 					}
 
@@ -1083,7 +1066,7 @@ int32 UWorldPartitionConvertCommandlet::Main(const FString& Params)
 
 	// Prepare levels for conversion
 	DetachDependantLevelPackages(MainLevel);
-	PrepareLevelActors(MainLevel, MainLevel->Actors, true, GetLevelGridPlacement(MainLevel, SubLevelsToConvert.Num() ? EActorGridPlacement::AlwaysLoaded : EActorGridPlacement::Bounds));
+	PrepareLevelActors(MainLevel, MainLevel->Actors, true);
 	PackagesToSave.Add(MainLevel->GetPackage());
 
 	if (bConversionSuffix)
@@ -1227,7 +1210,7 @@ int32 UWorldPartitionConvertCommandlet::Main(const FString& Params)
 
 		UE_LOG(LogWorldPartitionConvertCommandlet, Log, TEXT("Converting %s"), *SubWorld->GetName());
 
-		PrepareLevelActors(SubLevel, ActorsToConvert, false, GetLevelGridPlacement(SubLevel, EActorGridPlacement::Bounds));
+		PrepareLevelActors(SubLevel, ActorsToConvert, false);
 
 		for(AActor* Actor: ActorsToConvert)
 		{
diff --git a/Engine/Source/Editor/UnrealEd/Private/WorldPartition/WorldPartitionHLODsBuilder.cpp b/Engine/Source/Editor/UnrealEd/Private/WorldPartition/WorldPartitionHLODsBuilder.cpp
--- a/Engine/Source/Editor/UnrealEd/Private/WorldPartition/WorldPartitionHLODsBuilder.cpp
+++ b/Engine/Source/Editor/UnrealEd/Private/WorldPartition/WorldPartitionHLODsBuilder.cpp
@@ -586,7 +586,7 @@ bool UWorldPartitionHLODsBuilder::DumpStats()
 	{
 		FString				Name;
 		FName				RuntimeGrid;
-		EActorGridPlacement GridPlacement = EActorGridPlacement::None;
+		bool				bIsSpatiallyLoaded = false;
 		uint64				GridLocationZ = 0;
 		uint64				GridLocationX = 0;
 		uint64				GridLocationY = 0;
@@ -612,7 +612,7 @@ bool UWorldPartitionHLODsBuilder::DumpStats()
 
 		static FString GetHeaderCSVString()
 		{
-			return TEXT("Name, RuntimeGrid, GridPlacement, GridLocationZ, GridLocationX, GridLocationY, DataLayers, HLODType, InstanceCount, NaniteTriangleCount, NaniteVertexCount, TriangleCount, VertexCount, UVChannelCount, BaseColorTextureSize, NormalTextureSize, MRSTextureSize, EmissiveTextureSize, MeshResourceSize, TexturesResourceSize, DiskSize");
+			return TEXT("Name, RuntimeGrid, SpatiallyLoaded, GridLocationZ, GridLocationX, GridLocationY, DataLayers, HLODType, InstanceCount, NaniteTriangleCount, NaniteVertexCount, TriangleCount, VertexCount, UVChannelCount, BaseColorTextureSize, NormalTextureSize, MRSTextureSize, EmissiveTextureSize, MeshResourceSize, TexturesResourceSize, DiskSize");
 		}
 
 		FString ToCSVString() const
@@ -620,7 +620,7 @@ bool UWorldPartitionHLODsBuilder::DumpStats()
 			return FString::Printf(TEXT("%s, %s, %s, %d, %d, %d, %s, %s, %d, %d, %d, %d, %d, %d, %d, %d, %d, %d, %d, %d, %d"), 
 				*Name, 
 				*RuntimeGrid.ToString(),
-				*StaticEnum<EActorGridPlacement>()->GetNameStringByValue((int64)GridPlacement),
+				bIsSpatiallyLoaded ? TEXT("true") : TEXT("false"),
 				GridLocationZ, GridLocationX, GridLocationY,
 				*DataLayers,
 				*StaticEnum<EHLODLayerType>()->GetNameStringByValue((int64)HLODType),
@@ -644,7 +644,7 @@ bool UWorldPartitionHLODsBuilder::DumpStats()
 
 		HLODStat.Name = HLODActor->GetActorLabel();
 		HLODStat.RuntimeGrid = HLODActor->GetRuntimeGrid();
-		HLODStat.GridPlacement = HLODActor->GetGridPlacement();
+		HLODStat.bIsSpatiallyLoaded = HLODActor->GetIsSpatiallyLoaded();
 		HLODActor->GetGridIndices(HLODStat.GridLocationX, HLODStat.GridLocationY, HLODStat.GridLocationZ);
 		HLODStat.DataLayers = *FString::JoinBy(HLODActor->GetDataLayerObjects(), TEXT("| "), [](const UDataLayer* DataLayer) { return DataLayer->GetDataLayerLabel().ToString(); });
 		HLODStat.HLODType = HLODActor->GetSubActorsHLODLayer()->GetLayerType();
diff --git a/Engine/Source/Editor/WorldPartitionEditor/Private/WorldPartition/SWorldPartitionEditorGrid2D.cpp b/Engine/Source/Editor/WorldPartitionEditor/Private/WorldPartition/SWorldPartitionEditorGrid2D.cpp
--- a/Engine/Source/Editor/WorldPartitionEditor/Private/WorldPartition/SWorldPartitionEditorGrid2D.cpp
+++ b/Engine/Source/Editor/WorldPartitionEditor/Private/WorldPartition/SWorldPartitionEditorGrid2D.cpp
@@ -473,7 +473,7 @@ uint32 SWorldPartitionEditorGrid2D::PaintActors(const FGeometry& AllottedGeometr
 			{
 				FPaintGeometry ActorGeometry = AllottedGeometry.ToPaintGeometry(TopLeft, BottomRight - TopLeft);
 				float ActorColorGradient = FMath::Min((ActorViewBox.GetArea() - MinimumAreaCull) / AreaFadeDistance, 1.0f);
-				float ActorBrightness = (ActorDescView.GetGridPlacement() == EActorGridPlacement::AlwaysLoaded) ? 0.3f : 1.0f;
+				float ActorBrightness = ActorDescView.GetIsSpatiallyLoaded() ? 1.0f : 0.3f;
 				FLinearColor ActorColor(ActorBrightness, ActorBrightness, ActorBrightness, ActorColorGradient);
 
 				UClass* ActorClass = ActorDescView.GetActorClass();
diff --git a/Engine/Source/Runtime/Core/Public/UObject/UE5ReleaseStreamObjectVersion.h b/Engine/Source/Runtime/Core/Public/UObject/UE5ReleaseStreamObjectVersion.h
--- a/Engine/Source/Runtime/Core/Public/UObject/UE5ReleaseStreamObjectVersion.h
+++ b/Engine/Source/Runtime/Core/Public/UObject/UE5ReleaseStreamObjectVersion.h
@@ -63,6 +63,9 @@ struct CORE_API FUE5ReleaseStreamObjectVersion
 		// Add support for linking to the attached parent actor in WorldPartitionActorDesc
 		WorldPartitionActorDescSerializeAttachParent,
 
+		// Converted AActor GridPlacement to bIsSpatiallyLoaded flag
+		ConvertedActorGridPlacementToSpatiallyLoadedFlag,
+
 		// -----<new versions can be added above this line>-------------------------------------------------
 		VersionPlusOne,
 		LatestVersion = VersionPlusOne - 1
diff --git a/Engine/Source/Runtime/Engine/Classes/Engine/DirectionalLight.h b/Engine/Source/Runtime/Engine/Classes/Engine/DirectionalLight.h
--- a/Engine/Source/Runtime/Engine/Classes/Engine/DirectionalLight.h
+++ b/Engine/Source/Runtime/Engine/Classes/Engine/DirectionalLight.h
@@ -37,7 +37,7 @@ public:
 #if WITH_EDITOR
 	virtual void LoadedFromAnotherClass(const FName& OldClassName) override;
 	virtual void PostEditChangeProperty(FPropertyChangedEvent& PropertyChangedEvent) override;
-	virtual EActorGridPlacement GetDefaultGridPlacement() const override { return EActorGridPlacement::AlwaysLoaded; }
+	virtual bool CanChangeIsSpatiallyLoadedFlag() const override { return false; }
 #endif
 
 public:
diff --git a/Engine/Source/Runtime/Engine/Classes/Engine/LevelBounds.h b/Engine/Source/Runtime/Engine/Classes/Engine/LevelBounds.h
--- a/Engine/Source/Runtime/Engine/Classes/Engine/LevelBounds.h
+++ b/Engine/Source/Runtime/Engine/Classes/Engine/LevelBounds.h
@@ -67,8 +67,7 @@ class ENGINE_API ALevelBounds
 	/** Update level bounds immediately so the bounds are accurate when returning. Use only when needed because updating the bounds is slow */
 	void UpdateLevelBoundsImmediately();
 
-	/** Make the level bounds always loaded for partitioned worlds */
-	virtual EActorGridPlacement GetDefaultGridPlacement() const { return EActorGridPlacement::AlwaysLoaded; }
+	virtual bool CanChangeIsSpatiallyLoadedFlag() const override { return false; }
 	
 protected:
 	/** FTickableGameObject interface */
diff --git a/Engine/Source/Runtime/Engine/Classes/Engine/PostProcessVolume.h b/Engine/Source/Runtime/Engine/Classes/Engine/PostProcessVolume.h
--- a/Engine/Source/Runtime/Engine/Classes/Engine/PostProcessVolume.h
+++ b/Engine/Source/Runtime/Engine/Classes/Engine/PostProcessVolume.h
@@ -76,7 +76,7 @@ class ENGINE_API APostProcessVolume : public AVolume, public IInterface_PostProc
 	virtual void PostUnregisterAllComponents( void ) override;
 
 #if WITH_EDITOR
-	virtual EActorGridPlacement GetGridPlacement() const override;
+	virtual bool CanChangeIsSpatiallyLoadedFlag() const override { return false; }
 #endif
 
 protected:
diff --git a/Engine/Source/Runtime/Engine/Classes/GameFramework/Actor.h b/Engine/Source/Runtime/Engine/Classes/GameFramework/Actor.h
--- a/Engine/Source/Runtime/Engine/Classes/GameFramework/Actor.h
+++ b/Engine/Source/Runtime/Engine/Classes/GameFramework/Actor.h
@@ -62,7 +62,7 @@ enum class EActorUpdateOverlapsMethod : uint8
 #if WITH_EDITORONLY_DATA
 /** Enum defining how actor will be placed in the partition */
 UENUM()
-enum class EActorGridPlacement : uint8
+enum class UE_DEPRECATED(5.0, "EActorGridPlacement is deprecated.") EActorGridPlacement : uint8
 {
 	// Actor uses its bounds to determine in which runtime cells it's going to be placed.
 	Bounds,
@@ -72,18 +72,6 @@ enum class EActorGridPlacement : uint8
 	AlwaysLoaded,
 	None UMETA(Hidden)
 };
-
-inline const TCHAR* GetActorGridPlacementName(EActorGridPlacement ActorGridPlacement)
-{
-	switch(ActorGridPlacement)
-	{
-	case EActorGridPlacement::Bounds: return TEXT("Bounds");
-	case EActorGridPlacement::Location: return TEXT("Location");
-	case EActorGridPlacement::AlwaysLoaded: return TEXT("AlwaysLoaded");
-	default: check(0);
-	}
-	return TEXT("Invalid");
-}
 #endif
 
 ENGINE_API DECLARE_LOG_CATEGORY_EXTERN(LogActor, Log, Warning);
@@ -591,12 +579,11 @@ public:
 
 protected:
 #if WITH_EDITORONLY_DATA
-	/** 
-	 * Determine how this actor will be placed in the partition (if the world is partitioned). This value is not taken into account
-	 * if the function ActorClass->GetDefaultGridPlacement returns other than EActorGridPlacement::None.
-	 */
-	UPROPERTY(EditAnywhere, Category=WorldPartition)
-	EActorGridPlacement GridPlacement;
+	/** @deprecated Use bIsSpatiallyLoaded instead */
+	PRAGMA_DISABLE_DEPRECATION_WARNINGS
+	UPROPERTY()
+	EActorGridPlacement GridPlacement_DEPRECATED;
+	PRAGMA_ENABLE_DEPRECATION_WARNINGS
 
 	/** 
 	 * Determine in which partition grid this actor will be placed in the partition (if the world is partitioned).
@@ -866,17 +853,6 @@ public:
 
 	FActorOnPackagingModeChanged OnPackagingModeChanged;
 
-	/**
-	 * Determine how this actor should be placed in the partition (meant to be called on the default object).
-	 / @return EActorGridPlacement::None if the placement should be on a per-instance basis, otherwise the default value wins.
-	 */
-	virtual EActorGridPlacement GetDefaultGridPlacement() const;
-	
-	/** Returns this actor's current grid placement. */
-	virtual EActorGridPlacement GetGridPlacement() const { return GridPlacement; }
-
-	/** Sets this actor's current grid placement. */
-	void SetGridPlacement(EActorGridPlacement InGridPLacement) { GridPlacement = InGridPLacement; }
 
 	/** Returns this actor's current target runtime grid. */
 	virtual FName GetRuntimeGrid() const { return RuntimeGrid; }
@@ -1001,6 +977,15 @@ protected:
 	/** Whether the actor can be used as a PlayFromHere origin (OnPlayFromHere() will be called on that actor) */
 	UPROPERTY()
 	uint8 bCanPlayFromHere : 1;
+
+	/** 
+	 * Determine if this actor is spatially loaded when placed in a partitioned world.
+	 *	If true, this actor will be loaded when in the range of any streaming sources and if (1) in no data layers, or (2) one or more of its data layers are enabled.
+	 *	If false, this actor will be loaded if (1) in no data layers, or (2) one or more of its data layers are enabled.
+	 */
+	UPROPERTY(EditAnywhere, Category=WorldPartition)
+	uint8 bIsSpatiallyLoaded : 1;
+
 private:
 	/** Whether this actor is temporarily hidden within the editor; used for show/hide/etc functionality w/o dirtying the actor. */
 	UPROPERTY(Transient)
@@ -2226,6 +2211,19 @@ public:
 		}
 	}
 
+	/** Returns true if this actor is spatially loaded. */
+	bool GetIsSpatiallyLoaded() const { return bIsSpatiallyLoaded; }
+	
+	/** Set if this actor should be spatially loaded or not. */
+	void SetIsSpatiallyLoaded(bool bInIsSpatiallyLoaded)
+	{
+		check(CanChangeIsSpatiallyLoadedFlag());
+		bIsSpatiallyLoaded = bInIsSpatiallyLoaded;
+	}
+
+	/** Returns true if this actor allows changing the spatially loaded flag.  */
+	virtual bool CanChangeIsSpatiallyLoadedFlag() const { return true; }
+
 	/**
 	 * Returns whether or not this actor was explicitly hidden in the editor for the duration of the current editor session
 	 * @param bIncludeParent - Whether to recurse up child actor hierarchy or not
diff --git a/Engine/Source/Runtime/Engine/Classes/GameFramework/Info.h b/Engine/Source/Runtime/Engine/Classes/GameFramework/Info.h
--- a/Engine/Source/Runtime/Engine/Classes/GameFramework/Info.h
+++ b/Engine/Source/Runtime/Engine/Classes/GameFramework/Info.h
@@ -33,7 +33,7 @@ public:
 	virtual bool IsLevelBoundsRelevant() const override { return false; }
 
 #if WITH_EDITOR
-	virtual EActorGridPlacement GetDefaultGridPlacement() const override { return EActorGridPlacement::AlwaysLoaded; }
+	virtual bool CanChangeIsSpatiallyLoadedFlag() const override { return false; }
 	virtual bool SupportsDataLayer() const override { return false; }
 #endif
 
diff --git a/Engine/Source/Runtime/Engine/Classes/VT/RuntimeVirtualTextureVolume.h b/Engine/Source/Runtime/Engine/Classes/VT/RuntimeVirtualTextureVolume.h
--- a/Engine/Source/Runtime/Engine/Classes/VT/RuntimeVirtualTextureVolume.h
+++ b/Engine/Source/Runtime/Engine/Classes/VT/RuntimeVirtualTextureVolume.h
@@ -24,7 +24,7 @@ public:
 #endif // WITH_EDITORONLY_DATA
 
 #if WITH_EDITOR
-	virtual EActorGridPlacement GetDefaultGridPlacement() const override { return EActorGridPlacement::AlwaysLoaded; }
+	virtual bool CanChangeIsSpatiallyLoadedFlag() const override { return false; }
 #endif
 
 protected:
diff --git a/Engine/Source/Runtime/Engine/Private/Actor.cpp b/Engine/Source/Runtime/Engine/Private/Actor.cpp
--- a/Engine/Source/Runtime/Engine/Private/Actor.cpp
+++ b/Engine/Source/Runtime/Engine/Private/Actor.cpp
@@ -175,7 +175,7 @@ void AActor::InitializeDefaults()
 #endif // (CSV_PROFILER && !UE_BUILD_SHIPPING)
 
 #if WITH_EDITORONLY_DATA
-	GridPlacement = EActorGridPlacement::Location;
+	bIsSpatiallyLoaded = true;
 	CopyPasteId = INDEX_NONE;
 #endif
 }
@@ -889,6 +889,13 @@ void AActor::PostLoad()
 	// Fixup DataLayers
 	FixupDataLayers();
 
+	// Deprecate GridPlacement
+	PRAGMA_DISABLE_DEPRECATION_WARNINGS
+	if (GridPlacement_DEPRECATED != EActorGridPlacement::None)
+	{
+		bIsSpatiallyLoaded = GridPlacement_DEPRECATED != EActorGridPlacement::AlwaysLoaded;
+	}
+	PRAGMA_ENABLE_DEPRECATION_WARNINGS
 #endif // WITH_EDITORONLY_DATA
 
 	// Since the actor is being loading, it finished spawning by definition when it was originally spawned, so set to true now
diff --git a/Engine/Source/Runtime/Engine/Private/ActorEditor.cpp b/Engine/Source/Runtime/Engine/Private/ActorEditor.cpp
--- a/Engine/Source/Runtime/Engine/Private/ActorEditor.cpp
+++ b/Engine/Source/Runtime/Engine/Private/ActorEditor.cpp
@@ -68,12 +68,12 @@ bool AActor::CanEditChange(const FProperty* PropertyThatWillChange) const
 		return false;
 	}
 
-	const bool bIsGridPlacement = PropertyThatWillChange->GetFName() == GET_MEMBER_NAME_CHECKED(AActor, GridPlacement);
-	const bool bIsRuntimeGrid = PropertyThatWillChange->GetFName() == GET_MEMBER_NAME_CHECKED(AActor, RuntimeGrid);
-	const bool bIsDataLayers = PropertyThatWillChange->GetFName() == GET_MEMBER_NAME_CHECKED(AActor, DataLayers);
-	const bool bIsHLODLayer = PropertyThatWillChange->GetFName() == GET_MEMBER_NAME_CHECKED(AActor, HLODLayer);
+	const bool bIsSpatiallyLoadedProperty = PropertyThatWillChange->GetFName() == GET_MEMBER_NAME_CHECKED(AActor, bIsSpatiallyLoaded);
+	const bool bIsRuntimeGridProperty = PropertyThatWillChange->GetFName() == GET_MEMBER_NAME_CHECKED(AActor, RuntimeGrid);
+	const bool bIsDataLayersProperty = PropertyThatWillChange->GetFName() == GET_MEMBER_NAME_CHECKED(AActor, DataLayers);
+	const bool bIsHLODLayerProperty = PropertyThatWillChange->GetFName() == GET_MEMBER_NAME_CHECKED(AActor, HLODLayer);
 
-	if (bIsGridPlacement || bIsRuntimeGrid || bIsDataLayers || bIsHLODLayer)
+	if (bIsSpatiallyLoadedProperty || bIsRuntimeGridProperty || bIsDataLayersProperty || bIsHLODLayerProperty)
 	{
 		if (!IsTemplate())
 		{
@@ -88,12 +88,12 @@ bool AActor::CanEditChange(const FProperty* PropertyThatWillChange) const
 		}
 	}
 
-	if (bIsGridPlacement && (GetDefaultGridPlacement() != EActorGridPlacement::None))
+	if (bIsSpatiallyLoadedProperty && !CanChangeIsSpatiallyLoadedFlag())
 	{
 		return false;
 	}
 
-	if (bIsDataLayers && !SupportsDataLayer())
+	if (bIsDataLayersProperty && !SupportsDataLayer())
 	{
 		return false;
 	}
@@ -845,15 +845,6 @@ void AActor::SetPackageExternal(bool bExternal, bool bShouldDirty)
 	MarkPackageDirty();
 }
 
-EActorGridPlacement AActor::GetDefaultGridPlacement() const
-{
-	if (GetClass()->GetClassFlags() & CLASS_NotPlaceable)
-	{
-		return EActorGridPlacement::AlwaysLoaded;
-	}
-	return EActorGridPlacement::None;
-}
-
 void AActor::OnPlayFromHere()
 {
 	check(bCanPlayFromHere);
diff --git a/Engine/Source/Runtime/Engine/Private/Info.cpp b/Engine/Source/Runtime/Engine/Private/Info.cpp
--- a/Engine/Source/Runtime/Engine/Private/Info.cpp
+++ b/Engine/Source/Runtime/Engine/Private/Info.cpp
@@ -31,9 +31,10 @@ AInfo::AInfo(const FObjectInitializer& ObjectInitializer)
 		SpriteComponent->Sprite = ConstructorStatics.SpriteTexture.Get();
 		SpriteComponent->SpriteInfo.Category = ConstructorStatics.ID_Info;
 		SpriteComponent->SpriteInfo.DisplayName = ConstructorStatics.NAME_Info;
-		SpriteComponent->bIsScreenSizeScaled = true;	
-
+		SpriteComponent->bIsScreenSizeScaled = true;
 	}
+
+	bIsSpatiallyLoaded = false;
 #endif // WITH_EDITORONLY_DATA
 
 	PrimaryActorTick.bCanEverTick = false;
diff --git a/Engine/Source/Runtime/Engine/Private/LevelActor.cpp b/Engine/Source/Runtime/Engine/Private/LevelActor.cpp
--- a/Engine/Source/Runtime/Engine/Private/LevelActor.cpp
+++ b/Engine/Source/Runtime/Engine/Private/LevelActor.cpp
@@ -683,15 +683,6 @@ AActor* UWorld::SpawnActor( UClass* Class, FTransform const* UserTransformPtr, c
 	LevelToSpawnIn->Actors.Add( Actor );
 	LevelToSpawnIn->ActorsForGC.Add(Actor);
 
-#if WITH_EDITOR
-	// Set actor's partition placement
-	EActorGridPlacement DefaultGridPlacement = Class->GetDefaultObject<AActor>()->GetDefaultGridPlacement();
-	if (DefaultGridPlacement != EActorGridPlacement::None)
-	{
-		Actor->SetGridPlacement(DefaultGridPlacement);
-	}
-#endif
-
 #if PERF_SHOW_MULTI_PAWN_SPAWN_FRAMES
 	if( Cast<APawn>(Actor) )
 	{
diff --git a/Engine/Source/Runtime/Engine/Private/LevelBounds.cpp b/Engine/Source/Runtime/Engine/Private/LevelBounds.cpp
--- a/Engine/Source/Runtime/Engine/Private/LevelBounds.cpp
+++ b/Engine/Source/Runtime/Engine/Private/LevelBounds.cpp
@@ -32,6 +32,11 @@ ALevelBounds::ALevelBounds(const FObjectInitializer& ObjectInitializer)
 	bLevelBoundsDirty = true;
 	bUsingDefaultBounds = false;
 #endif
+
+#if WITH_EDITORONLY_DATA
+	bIsSpatiallyLoaded = false;
+#endif
+
 }
 
 void ALevelBounds::PostLoad()
diff --git a/Engine/Source/Runtime/Engine/Private/Light.cpp b/Engine/Source/Runtime/Engine/Private/Light.cpp
--- a/Engine/Source/Runtime/Engine/Private/Light.cpp
+++ b/Engine/Source/Runtime/Engine/Private/Light.cpp
@@ -259,6 +259,8 @@ ADirectionalLight::ADirectionalLight(const FObjectInitializer& ObjectInitializer
 		ArrowComponent->bLightAttachment = true;
 		ArrowComponent->bIsScreenSizeScaled = true;
 	}
+
+	bIsSpatiallyLoaded = false;
 #endif // WITH_EDITORONLY_DATA
 
 }
diff --git a/Engine/Source/Runtime/Engine/Private/PlayerStart.cpp b/Engine/Source/Runtime/Engine/Private/PlayerStart.cpp
--- a/Engine/Source/Runtime/Engine/Private/PlayerStart.cpp
+++ b/Engine/Source/Runtime/Engine/Private/PlayerStart.cpp
@@ -62,7 +62,7 @@ APlayerStart::APlayerStart(const FObjectInitializer& ObjectInitializer)
 		}
 	}
 
-	GridPlacement = EActorGridPlacement::AlwaysLoaded;
+	bIsSpatiallyLoaded = false;
 #endif // WITH_EDITORONLY_DATA
 }
 
diff --git a/Engine/Source/Runtime/Engine/Private/PostProcessVolume.cpp b/Engine/Source/Runtime/Engine/Private/PostProcessVolume.cpp
--- a/Engine/Source/Runtime/Engine/Private/PostProcessVolume.cpp
+++ b/Engine/Source/Runtime/Engine/Private/PostProcessVolume.cpp
@@ -50,11 +50,6 @@ void APostProcessVolume::Serialize(FArchive& Ar)
 }
 
 #if WITH_EDITOR
-EActorGridPlacement APostProcessVolume::GetGridPlacement() const
-{
-	return bUnbound ? EActorGridPlacement::AlwaysLoaded : AActor::GetGridPlacement();
-}
-
 void APostProcessVolume::PostEditChangeProperty(FPropertyChangedEvent& PropertyChangedEvent)
 {
 	Super::PostEditChangeProperty(PropertyChangedEvent);
@@ -76,6 +71,7 @@ void APostProcessVolume::PostEditChangeProperty(FPropertyChangedEvent& PropertyC
 		}
 	}
 
+	bIsSpatiallyLoaded = !bUnbound;
 	
 	if (PropertyChangedEvent.Property)
 	{
@@ -220,11 +216,6 @@ bool APostProcessVolume::CanEditChange(const FProperty* InProperty) const
 				return false;
 			}
 		}
-
-		if (PropertyName == GET_MEMBER_NAME_STRING_CHECKED(APostProcessVolume, GridPlacement))
-		{
-			return !bUnbound;
-		}
 	}
 
 	return Super::CanEditChange(InProperty);
diff --git a/Engine/Source/Runtime/Engine/Private/VT/RuntimeVirtualTextureVolume.cpp b/Engine/Source/Runtime/Engine/Private/VT/RuntimeVirtualTextureVolume.cpp
--- a/Engine/Source/Runtime/Engine/Private/VT/RuntimeVirtualTextureVolume.cpp
+++ b/Engine/Source/Runtime/Engine/Private/VT/RuntimeVirtualTextureVolume.cpp
@@ -22,6 +22,8 @@ ARuntimeVirtualTextureVolume::ARuntimeVirtualTextureVolume(const FObjectInitiali
 	Box->CanCharacterStepUpOn = ECanBeCharacterBase::ECB_No;
 	Box->SetGenerateOverlapEvents(false);
 	Box->SetupAttachment(VirtualTextureComponent);
+
+	bIsSpatiallyLoaded = false;
 #endif
 }
 
diff --git a/Engine/Source/Runtime/Engine/Private/WorldPartition/ErrorHandling/WorldPartitionStreamingGenerationLogErrorHandler.cpp b/Engine/Source/Runtime/Engine/Private/WorldPartition/ErrorHandling/WorldPartitionStreamingGenerationLogErrorHandler.cpp
--- a/Engine/Source/Runtime/Engine/Private/WorldPartition/ErrorHandling/WorldPartitionStreamingGenerationLogErrorHandler.cpp
+++ b/Engine/Source/Runtime/Engine/Private/WorldPartition/ErrorHandling/WorldPartitionStreamingGenerationLogErrorHandler.cpp
@@ -12,11 +12,10 @@ void FStreamingGenerationLogErrorHandler::OnInvalidReference(const FWorldPartiti
 
 void FStreamingGenerationLogErrorHandler::OnInvalidReferenceGridPlacement(const FWorldPartitionActorDescView& ActorDescView, const FWorldPartitionActorDescView& ReferenceActorDescView)
 {
-	const FString StreamedActor(TEXT("Streamed actor"));
-	const FString AlwaysLoadedActor(TEXT("Always loaded actor"));
-	const bool bIsActorDescAlwaysLoaded = ActorDescView.GetGridPlacement() == EActorGridPlacement::AlwaysLoaded;
-	const bool bIsActorDescRefAlwaysLoaded = ReferenceActorDescView.GetGridPlacement() == EActorGridPlacement::AlwaysLoaded;
-	UE_LOG(LogWorldPartition, Log, TEXT("%s %s reference %s %s"), bIsActorDescAlwaysLoaded ? *AlwaysLoadedActor : *StreamedActor, *ActorDescView.GetActorLabelOrName().ToString(), bIsActorDescRefAlwaysLoaded ? *AlwaysLoadedActor : *StreamedActor, *ReferenceActorDescView.GetActorLabelOrName().ToString());
+	const FString SpatiallyLoadedActor(TEXT("Spatially loaded actor"));
+	const FString NonSpatiallyLoadedActor(TEXT("Non-spatially loaded loaded actor"));
+
+	UE_LOG(LogWorldPartition, Log, TEXT("%s %s reference %s %s"), ActorDescView.GetIsSpatiallyLoaded() ? *SpatiallyLoadedActor : *NonSpatiallyLoadedActor, *ActorDescView.GetActorLabelOrName().ToString(), ReferenceActorDescView.GetIsSpatiallyLoaded() ? *SpatiallyLoadedActor : *NonSpatiallyLoadedActor, *ReferenceActorDescView.GetActorLabelOrName().ToString());
 }
 
 void FStreamingGenerationLogErrorHandler::OnInvalidReferenceDataLayers(const FWorldPartitionActorDescView& ActorDescView, const FWorldPartitionActorDescView& ReferenceActorDescView)
diff --git a/Engine/Source/Runtime/Engine/Private/WorldPartition/ErrorHandling/WorldPartitionStreamingGenerationMapCheckErrorHandler.cpp b/Engine/Source/Runtime/Engine/Private/WorldPartition/ErrorHandling/WorldPartitionStreamingGenerationMapCheckErrorHandler.cpp
--- a/Engine/Source/Runtime/Engine/Private/WorldPartition/ErrorHandling/WorldPartitionStreamingGenerationMapCheckErrorHandler.cpp
+++ b/Engine/Source/Runtime/Engine/Private/WorldPartition/ErrorHandling/WorldPartitionStreamingGenerationMapCheckErrorHandler.cpp
@@ -21,17 +21,14 @@ void FStreamingGenerationMapCheckErrorHandler::OnInvalidReference(const FWorldPa
 
 void FStreamingGenerationMapCheckErrorHandler::OnInvalidReferenceGridPlacement(const FWorldPartitionActorDescView& ActorDescView, const FWorldPartitionActorDescView& ReferenceActorDescView)
 {
-	const bool bIsActorDescAlwaysLoaded = ActorDescView.GetGridPlacement() == EActorGridPlacement::AlwaysLoaded;
-	const bool bIsActorDescRefAlwaysLoaded = ReferenceActorDescView.GetGridPlacement() == EActorGridPlacement::AlwaysLoaded;
-
-	const FText StreamedActor(LOCTEXT("MapCheck_WorldPartition_StreamedActor", "Streamed actor"));
-	const FText AlwaysLoadedActor(LOCTEXT("MapCheck_WorldPartition_AlwaysLoadedActor", "Always loaded actor"));
+	const FText SpatiallyLoadedActor(LOCTEXT("MapCheck_WorldPartition_SpatiallyLoadedActor", "Spatially loaded actor"));
+	const FText NonSpatiallyLoadedActor(LOCTEXT("MapCheck_WorldPartition_NonSpatiallyLoadedActor", "Non-spatially loaded actor"));
 
 	TSharedRef<FTokenizedMessage> Error = FMessageLog("MapCheck").Error()
-		->AddToken(FTextToken::Create(bIsActorDescAlwaysLoaded ? AlwaysLoadedActor : StreamedActor))
+		->AddToken(FTextToken::Create(ActorDescView.GetIsSpatiallyLoaded() ? SpatiallyLoadedActor : NonSpatiallyLoadedActor))
 		->AddToken(FAssetNameToken::Create(ActorDescView.GetActorLabelOrName().ToString()))
 		->AddToken(FTextToken::Create(LOCTEXT("MapCheck_WorldPartition_References", "references")))
-		->AddToken(FTextToken::Create(bIsActorDescRefAlwaysLoaded ? AlwaysLoadedActor : StreamedActor)								)
+		->AddToken(FTextToken::Create(ReferenceActorDescView.GetIsSpatiallyLoaded() ? SpatiallyLoadedActor : NonSpatiallyLoadedActor))
 		->AddToken(FAssetNameToken::Create(ReferenceActorDescView.GetActorLabelOrName().ToString()))
 		->AddToken(FMapErrorToken::Create(FName(TEXT("WorldPartition_StreamedActorReferenceAlwaysLoadedActor_CheckForErrors"))));
 }
diff --git a/Engine/Source/Runtime/Engine/Private/WorldPartition/HLOD/HLODActor.cpp b/Engine/Source/Runtime/Engine/Private/WorldPartition/HLOD/HLODActor.cpp
--- a/Engine/Source/Runtime/Engine/Private/WorldPartition/HLOD/HLODActor.cpp
+++ b/Engine/Source/Runtime/Engine/Private/WorldPartition/HLOD/HLODActor.cpp
@@ -92,17 +92,6 @@ void AWorldPartitionHLOD::RerunConstructionScripts()
 
 #if WITH_EDITOR
 
-EActorGridPlacement AWorldPartitionHLOD::GetGridPlacement() const
-{
-	return SubActorsHLODLayer && SubActorsHLODLayer->IsAlwaysLoaded() ? EActorGridPlacement::AlwaysLoaded : Super::GetGridPlacement();
-}
-
-EActorGridPlacement AWorldPartitionHLOD::GetDefaultGridPlacement() const
-{
-	// Overriden as AActor::GetDefaultGridPlacement() will mark all actors that are not placeable as AlwaysLoaded...
-	return EActorGridPlacement::None;
-}
-
 TUniquePtr<FWorldPartitionActorDesc> AWorldPartitionHLOD::CreateClassActorDesc() const
 {
 	return TUniquePtr<FWorldPartitionActorDesc>(new FHLODActorDesc());
@@ -208,4 +197,4 @@ void AWorldPartitionHLOD::BuildHLOD(bool bForceBuild)
 	BuildActorTextureStreamingData(this, EMaterialQualityLevel::High, GMaxRHIFeatureLevel);
 }
 
-#endif // WITH_EDITOR
+#endif // WITH_EDITOR
\ No newline at end of file
diff --git a/Engine/Source/Runtime/Engine/Private/WorldPartition/NavigationData/NavigationDataChunkActor.cpp b/Engine/Source/Runtime/Engine/Private/WorldPartition/NavigationData/NavigationDataChunkActor.cpp
--- a/Engine/Source/Runtime/Engine/Private/WorldPartition/NavigationData/NavigationDataChunkActor.cpp
+++ b/Engine/Source/Runtime/Engine/Private/WorldPartition/NavigationData/NavigationDataChunkActor.cpp
@@ -151,11 +151,6 @@ void ANavigationDataChunkActor::GetActorBounds(bool bOnlyCollidingComponents, FV
 }
 
 #if WITH_EDITOR
-EActorGridPlacement ANavigationDataChunkActor::GetDefaultGridPlacement() const
-{
-	return EActorGridPlacement::Bounds;
-}
-
 FBox ANavigationDataChunkActor::GetStreamingBounds() const
 {
 	return DataChunkActorBounds;
diff --git a/Engine/Source/Runtime/Engine/Private/WorldPartition/RuntimeSpatialHash/RuntimeSpatialHashGridHelper.cpp b/Engine/Source/Runtime/Engine/Private/WorldPartition/RuntimeSpatialHash/RuntimeSpatialHashGridHelper.cpp
--- a/Engine/Source/Runtime/Engine/Private/WorldPartition/RuntimeSpatialHash/RuntimeSpatialHashGridHelper.cpp
+++ b/Engine/Source/Runtime/Engine/Private/WorldPartition/RuntimeSpatialHash/RuntimeSpatialHashGridHelper.cpp
@@ -160,56 +160,52 @@ FSquare2DGridHelper GetPartitionedActors(const UWorldPartition* WorldPartition,
 		}
 	}
 
+	const float CellArea = PartitionedActors.GetLowestLevel().CellSize * PartitionedActors.GetLowestLevel().CellSize;
 	for (const FActorClusterInstance* ClusterInstance : GridActors)
 	{
 		const FActorCluster* ActorCluster = ClusterInstance->Cluster;
 		const FBox2D ClusterBounds(FVector2D(ClusterInstance->Bounds.Min), FVector2D(ClusterInstance->Bounds.Max));
 		check(ActorCluster && ActorCluster->Actors.Num() > 0);
-		EActorGridPlacement GridPlacement = ActorCluster->GridPlacement;		
 		FSquare2DGridHelper::FGridLevel::FGridCell* GridCell = nullptr;
 
-		switch (GridPlacement)
+		if (ActorCluster->bIsSpatiallyLoaded)
 		{
-		case EActorGridPlacement::Location:
-		{
-			FIntVector2 CellCoords;
-			if (PartitionedActors.GetLowestLevel().GetCellCoords(ClusterBounds.GetCenter(), CellCoords))
+			if (ClusterBounds.GetArea() <= CellArea)
 			{
-				GridCell = &PartitionedActors.GetLowestLevel().GetCell(CellCoords);
+				// Find grid level cell that copntains the actor cluster pivot and put actors in it.
+				FIntVector2 CellCoords;
+				if (PartitionedActors.GetLowestLevel().GetCellCoords(ClusterBounds.GetCenter(), CellCoords))
+				{
+					GridCell = &PartitionedActors.GetLowestLevel().GetCell(CellCoords);
+				}
 			}
-			break;
-		}
-		case EActorGridPlacement::Bounds:
-		{
-			// Find grid level cell that encompasses the actor cluster and put actors in it.
-			const FVector2D ClusterSize = ClusterBounds.GetSize();
-			const float MinRequiredCellExtent = FMath::Max(ClusterSize.X, ClusterSize.Y);
-			const int32 FirstPotentialGridLevel = FMath::Max(FMath::CeilToFloat(FMath::Log2(MinRequiredCellExtent / (float)PartitionedActors.CellSize)), 0);
-
-			for (int32 GridLevelIndex = FirstPotentialGridLevel; GridLevelIndex < PartitionedActors.Levels.Num(); GridLevelIndex++)
+			else
 			{
-				FSquare2DGridHelper::FGridLevel& GridLevel = PartitionedActors.Levels[GridLevelIndex];
+				// Find grid level cell that encompasses the actor cluster bounding box and put actors in it.
+				const FVector2D ClusterSize = ClusterBounds.GetSize();
+				const float MinRequiredCellExtent = FMath::Max(ClusterSize.X, ClusterSize.Y);
+				const int32 FirstPotentialGridLevel = FMath::Max(FMath::CeilToFloat(FMath::Log2(MinRequiredCellExtent / (float)PartitionedActors.CellSize)), 0);
 
-				if (GridLevel.GetNumIntersectingCells(ClusterInstance->Bounds) == 1)
+				for (int32 GridLevelIndex = FirstPotentialGridLevel; GridLevelIndex < PartitionedActors.Levels.Num(); GridLevelIndex++)
 				{
-					GridLevel.ForEachIntersectingCells(ClusterInstance->Bounds, [&GridLevel, ActorCluster, ClusterInstance, &GridCell](const FIntVector2& Coords)
-					{
-						check(!GridCell);
-						GridCell = &GridLevel.GetCell(Coords);
-					});
+					FSquare2DGridHelper::FGridLevel& GridLevel = PartitionedActors.Levels[GridLevelIndex];
 
-					break;
+					if (GridLevel.GetNumIntersectingCells(ClusterInstance->Bounds) == 1)
+					{
+						GridLevel.ForEachIntersectingCells(ClusterInstance->Bounds, [&GridLevel, ActorCluster, ClusterInstance, &GridCell](const FIntVector2& Coords)
+						{
+							check(!GridCell);
+							GridCell = &GridLevel.GetCell(Coords);
+						});
+
+						break;
+					}
 				}
 			}
-			break;
 		}
-		case EActorGridPlacement::AlwaysLoaded:
+		else
 		{
 			GridCell = &PartitionedActors.GetAlwaysLoadedCell();
-			break;
-		}
-		default:
-			check(0);
 		}
 
 		if (!GridCell)
@@ -226,7 +222,6 @@ FSquare2DGridHelper GetPartitionedActors(const UWorldPartition* WorldPartition,
 				UE_LOG(LogWorldPartition, Warning, TEXT("         Container (%s): %s"), *ClusterInstance->ContainerInstance->ID.ToString(), *ClusterInstance->ContainerInstance->Container->GetContainerPackage().ToString())
 			}
 
-			GridPlacement = EActorGridPlacement::AlwaysLoaded;
 			GridCell = &PartitionedActors.GetAlwaysLoadedCell();
 		}
 
@@ -237,9 +232,8 @@ FSquare2DGridHelper GetPartitionedActors(const UWorldPartition* WorldPartition,
 		{
 			if (ActorCluster->Actors.Num() > 1)
 			{
-				UE_LOG(LogWorldPartition, Verbose, TEXT("Clustered %d actors (%s), generated shared BV of [%d x %d] (meters)"),
+				UE_LOG(LogWorldPartition, Verbose, TEXT("Clustered %d actors, generated shared BV of [%d x %d] (meters)"),
 					ActorCluster->Actors.Num(),
-					*StaticEnum<EActorGridPlacement>()->GetNameStringByValue((int64)GridPlacement),
 					(int)(0.01f * ClusterInstance->Bounds.GetSize().X),
 					(int)(0.01f * ClusterInstance->Bounds.GetSize().Y));
 
diff --git a/Engine/Source/Runtime/Engine/Private/WorldPartition/WorldPartition.cpp b/Engine/Source/Runtime/Engine/Private/WorldPartition/WorldPartition.cpp
--- a/Engine/Source/Runtime/Engine/Private/WorldPartition/WorldPartition.cpp
+++ b/Engine/Source/Runtime/Engine/Private/WorldPartition/WorldPartition.cpp
@@ -1349,19 +1349,9 @@ FBox UWorldPartition::GetWorldBounds() const
 	FBox WorldBounds(ForceInit);
 	for (UActorDescContainer::TConstIterator<> ActorDescIterator(this); ActorDescIterator; ++ActorDescIterator)
 	{
-		switch (ActorDescIterator->GetGridPlacement())
+		if (ActorDescIterator->GetIsSpatiallyLoaded())
 		{
-			case EActorGridPlacement::Location:
-			{
-				FVector Location = ActorDescIterator->GetOrigin();
-				WorldBounds += FBox(Location, Location);
-			}
-			break;
-			case EActorGridPlacement::Bounds:
-			{
-				WorldBounds += ActorDescIterator->GetBounds();
-			}
-			break;
+			WorldBounds += ActorDescIterator->GetBounds();
 		}
 	}
 	return WorldBounds;
diff --git a/Engine/Source/Runtime/Engine/Private/WorldPartition/WorldPartitionActorCluster.cpp b/Engine/Source/Runtime/Engine/Private/WorldPartition/WorldPartitionActorCluster.cpp
--- a/Engine/Source/Runtime/Engine/Private/WorldPartition/WorldPartitionActorCluster.cpp
+++ b/Engine/Source/Runtime/Engine/Private/WorldPartition/WorldPartitionActorCluster.cpp
@@ -32,12 +32,10 @@ TSet<const UDataLayer*> GetDataLayers(UWorld* InWorld, const LayerNameContainer&
 }
 
 FActorCluster::FActorCluster(UWorld* InWorld, const FWorldPartitionActorDescView& InActorDescView)
-	: GridPlacement(InActorDescView.GetGridPlacement())
+	: bIsSpatiallyLoaded(InActorDescView.GetIsSpatiallyLoaded())
 	, RuntimeGrid(InActorDescView.GetRuntimeGrid())
 	, Bounds(InActorDescView.GetBounds())
 {
-	check(GridPlacement != EActorGridPlacement::None);
-
 	Actors.Add(InActorDescView.GetGuid());
 	DataLayers = GetDataLayers(InWorld, InActorDescView.GetDataLayers());
 	DataLayersID = FDataLayersID(DataLayers.Array());
@@ -47,19 +45,13 @@ void FActorCluster::Add(const FActorCluster& InActorCluster, const TMap<FGuid, F
 {
 	check(RuntimeGrid == InActorCluster.RuntimeGrid);
 	check(DataLayersID == InActorCluster.DataLayersID);
-	check((GridPlacement == EActorGridPlacement::AlwaysLoaded) == (InActorCluster.GridPlacement == EActorGridPlacement::AlwaysLoaded));
+	check(bIsSpatiallyLoaded == InActorCluster.bIsSpatiallyLoaded);
 
 	// Merge bounds
 	Bounds += InActorCluster.Bounds;
 
 	// Merge Actors
 	Actors.Append(InActorCluster.Actors);
-
-	// Set grid placement
-	if (GridPlacement != EActorGridPlacement::AlwaysLoaded)
-	{
-		GridPlacement = EActorGridPlacement::Bounds;
-	}
 }
 
 FActorClusterInstance::FActorClusterInstance(const FActorCluster* InCluster, const FActorContainerInstance* InContainerInstance)
diff --git a/Engine/Source/Runtime/Engine/Private/WorldPartition/WorldPartitionActorDesc.cpp b/Engine/Source/Runtime/Engine/Private/WorldPartition/WorldPartitionActorDesc.cpp
--- a/Engine/Source/Runtime/Engine/Private/WorldPartition/WorldPartitionActorDesc.cpp
+++ b/Engine/Source/Runtime/Engine/Private/WorldPartition/WorldPartitionActorDesc.cpp
@@ -45,17 +45,8 @@ void FWorldPartitionActorDesc::Init(const AActor* InActor)
 	const FBox StreamingBounds = InActor->GetStreamingBounds();
 	StreamingBounds.GetCenterAndExtents(BoundsLocation, BoundsExtent);
 
-	const EActorGridPlacement DefaultGridPlacement = InActor->GetDefaultGridPlacement();
-	if (DefaultGridPlacement != EActorGridPlacement::None)
-	{
-		GridPlacement = DefaultGridPlacement;
-	}
-	else
-	{
-		GridPlacement = InActor->GetGridPlacement();
-	}
-
 	RuntimeGrid = InActor->GetRuntimeGrid();
+	bIsSpatiallyLoaded = InActor->GetIsSpatiallyLoaded();
 	bActorIsEditorOnly = InActor->IsEditorOnly();
 	bLevelBoundsRelevant = InActor->IsLevelBoundsRelevant();
 	bActorIsHLODRelevant = InActor->IsHLODRelevant();
@@ -70,10 +61,6 @@ void FWorldPartitionActorDesc::Init(const AActor* InActor)
 	{
 		ParentActor = AttachParentActor->GetActorGuid();
 	}
-	else
-	{
-		ParentActor = FGuid();
-	}
 	
 	TArray<AActor*> ActorReferences = ActorsReferencesUtils::GetExternalActorReferences((AActor*)InActor);
 
@@ -114,13 +101,6 @@ void FWorldPartitionActorDesc::Init(UActorDescContainer* InContainer, const FWor
 	// Serialize metadata payload
 	Serialize(MetadataAr);
 
-	// Override grid placement by default class value
-	const EActorGridPlacement DefaultGridPlacement = ActorClass->GetDefaultObject<AActor>()->GetDefaultGridPlacement();
-	if (DefaultGridPlacement != EActorGridPlacement::None)
-	{
-		GridPlacement = DefaultGridPlacement;
-	}
-
 	// Only set ActorPtr/Container on WorldPartition ActorDescs
 	if (UWorldPartition* WorldPartition = Cast<UWorldPartition>(InContainer))
 	{
@@ -192,18 +172,18 @@ void FWorldPartitionActorDesc::TransformInstance(const FString& From, const FStr
 FString FWorldPartitionActorDesc::ToString() const
 {
 	return FString::Printf(
-		TEXT("Guid:%s Class:%s Name:%s SpatiallyLoaded:%s Bounds:%s GridPlacement:%s RuntimeGrid:%s EditorOnly:%s LevelBoundsRelevant:%s HLODRelevant:%s FolderPath:%s"), 
+		TEXT("Guid:%s Class:%s Name:%s SpatiallyLoaded:%s Bounds:%s RuntimeGrid:%s EditorOnly:%s LevelBoundsRelevant:%s HLODRelevant:%s FolderPath:%s Parent:%s"), 
 		*Guid.ToString(), 
 		*Class.ToString(), 
-		*FPaths::GetExtension(ActorPath.ToString()), 
-		(GridPlacement != EActorGridPlacement::AlwaysLoaded) ? TEXT("True") : TEXT("False"),
+		*GetActorName().ToString(),
+		bIsSpatiallyLoaded ? TEXT("true") : TEXT("false"),
 		*GetBounds().ToString(),
-		GetActorGridPlacementName(GridPlacement),
 		*RuntimeGrid.ToString(),
 		bActorIsEditorOnly ? TEXT("true") : TEXT("false"),
 		bLevelBoundsRelevant ? TEXT("true") : TEXT("false"),
 		bActorIsHLODRelevant ? TEXT("true") : TEXT("false"),
-		*FolderPath.ToString()
+		*FolderPath.ToString(),
+		*ParentActor.ToString()
 	);
 }
 
@@ -214,7 +194,22 @@ void FWorldPartitionActorDesc::Serialize(FArchive& Ar)
 	Ar.UsingCustomVersion(FUE5MainStreamObjectVersion::GUID);
 	Ar.UsingCustomVersion(FUE5ReleaseStreamObjectVersion::GUID);
 
-	Ar << Class << Guid << BoundsLocation << BoundsExtent << GridPlacement << RuntimeGrid << bActorIsEditorOnly << bLevelBoundsRelevant;
+	Ar << Class << Guid << BoundsLocation << BoundsExtent;
+	
+	if (Ar.CustomVer(FUE5ReleaseStreamObjectVersion::GUID) < FUE5ReleaseStreamObjectVersion::ConvertedActorGridPlacementToSpatiallyLoadedFlag)
+	{
+		PRAGMA_DISABLE_DEPRECATION_WARNINGS
+		EActorGridPlacement GridPlacement;
+		Ar << (__underlying_type(EActorGridPlacement)&)GridPlacement;
+		bIsSpatiallyLoaded = GridPlacement != EActorGridPlacement::AlwaysLoaded;
+		PRAGMA_ENABLE_DEPRECATION_WARNINGS
+	}
+	else
+	{
+		Ar << bIsSpatiallyLoaded;
+	}
+		
+	Ar << RuntimeGrid << bActorIsEditorOnly << bLevelBoundsRelevant;
 	
 	if (Ar.CustomVer(FUE5MainStreamObjectVersion::GUID) < FUE5MainStreamObjectVersion::WorldPartitionActorDescSerializeDataLayers)
 	{
@@ -269,10 +264,7 @@ FBox FWorldPartitionActorDesc::GetBounds() const
 
 FName FWorldPartitionActorDesc::GetActorName() const
 {
-	FString ActorName;
-	FString ActorContext;
-	verify(GetActorPath().ToString().Split(TEXT("."), &ActorContext, &ActorName, ESearchCase::CaseSensitive, ESearchDir::FromEnd));
-	return *ActorName;
+	return *FPaths::GetExtension(ActorPath.ToString());
 }
 
 FName FWorldPartitionActorDesc::GetActorLabelOrName() const
diff --git a/Engine/Source/Runtime/Engine/Private/WorldPartition/WorldPartitionActorDescView.cpp b/Engine/Source/Runtime/Engine/Private/WorldPartition/WorldPartitionActorDescView.cpp
--- a/Engine/Source/Runtime/Engine/Private/WorldPartition/WorldPartitionActorDescView.cpp
+++ b/Engine/Source/Runtime/Engine/Private/WorldPartition/WorldPartitionActorDescView.cpp
@@ -8,13 +8,13 @@
 
 FWorldPartitionActorDescView::FWorldPartitionActorDescView()
 	: ActorDesc(nullptr)
-	, GridPlacement(EActorGridPlacement::None)
+	, bIsForcedNonSpatiallyLoaded(false)
 	, bInvalidDataLayers(false)
 {}
 
 FWorldPartitionActorDescView::FWorldPartitionActorDescView(const FWorldPartitionActorDesc* InActorDesc)
 	: ActorDesc(InActorDesc)
-	, GridPlacement(InActorDesc->GetGridPlacement())
+	, bIsForcedNonSpatiallyLoaded(false)
 	, bInvalidDataLayers(false)
 {}
 
@@ -38,11 +38,6 @@ FVector FWorldPartitionActorDescView::GetOrigin() const
 	return ActorDesc->GetOrigin();
 }
 
-EActorGridPlacement FWorldPartitionActorDescView::GetGridPlacement() const
-{
-	return GridPlacement;
-}
-
 FName FWorldPartitionActorDescView::GetRuntimeGrid() const
 {
 	return ActorDesc->GetRuntimeGrid();
@@ -53,6 +48,11 @@ bool FWorldPartitionActorDescView::GetActorIsEditorOnly() const
 	return ActorDesc->GetActorIsEditorOnly();
 }
 
+bool FWorldPartitionActorDescView::GetIsSpatiallyLoaded() const
+{
+	return bIsForcedNonSpatiallyLoaded ? false : ActorDesc->GetIsSpatiallyLoaded();
+}
+
 bool FWorldPartitionActorDescView::GetLevelBoundsRelevant() const
 {
 	return ActorDesc->GetLevelBoundsRelevant();
@@ -130,12 +130,12 @@ FName FWorldPartitionActorDescView::GetActorLabelOrName() const
 	return ActorDesc->GetActorLabelOrName();
 }
 
-void FWorldPartitionActorDescView::SetGridPlacement(EActorGridPlacement InGridPlacement)
+void FWorldPartitionActorDescView::SetForcedNonSpatiallyLoaded()
 {
-	if (GridPlacement != InGridPlacement)
+	if (!bIsForcedNonSpatiallyLoaded)
 	{
-		GridPlacement = InGridPlacement;
-		UE_LOG(LogWorldPartition, Verbose, TEXT("Actor '%s' grid placement changed to %s"), *GetActorLabel().ToString(), *StaticEnum<EActorGridPlacement>()->GetNameStringByValue((int64)InGridPlacement));
+		bIsForcedNonSpatiallyLoaded = true;
+		UE_LOG(LogWorldPartition, Verbose, TEXT("Actor '%s' forced to be non-spatially loaded"), *GetActorLabel().ToString());
 	}
 }
 
diff --git a/Engine/Source/Runtime/Engine/Private/WorldPartition/WorldPartitionEditorSpatialHash.cpp b/Engine/Source/Runtime/Engine/Private/WorldPartition/WorldPartitionEditorSpatialHash.cpp
--- a/Engine/Source/Runtime/Engine/Private/WorldPartition/WorldPartitionEditorSpatialHash.cpp
+++ b/Engine/Source/Runtime/Engine/Private/WorldPartition/WorldPartitionEditorSpatialHash.cpp
@@ -94,19 +94,12 @@ void UWorldPartitionEditorSpatialHash::Tick(float DeltaSeconds)
 	}
 }
 
-// In the editor, actors are always using their bounds for grid placement, which makes more sense from a user standpoint.
 FBox UWorldPartitionEditorSpatialHash::GetActorBounds(const FWorldPartitionHandle& InActorHandle) const
 {
 	FBox ActorBounds;
-
-	switch(InActorHandle->GetGridPlacement())
-	{
-	case EActorGridPlacement::Location:
-	case EActorGridPlacement::Bounds:
+	if (InActorHandle->GetIsSpatiallyLoaded())
 	{
 		ActorBounds = InActorHandle->GetBounds();
-		break;
-	}
 	}
 
 	check(ActorBounds.IsValid);
@@ -115,7 +108,7 @@ FBox UWorldPartitionEditorSpatialHash::GetActorBounds(const FWorldPartitionHandl
 
 bool UWorldPartitionEditorSpatialHash::IsActorAlwaysLoaded(const FWorldPartitionHandle& InActorHandle) const
 {
-	if (InActorHandle->GetGridPlacement() == EActorGridPlacement::AlwaysLoaded)
+	if (!InActorHandle->GetIsSpatiallyLoaded())
 	{
 		return true;
 	}
diff --git a/Engine/Source/Runtime/Engine/Private/WorldPartition/WorldPartitionRuntimeSpatialHash.cpp b/Engine/Source/Runtime/Engine/Private/WorldPartition/WorldPartitionRuntimeSpatialHash.cpp
--- a/Engine/Source/Runtime/Engine/Private/WorldPartition/WorldPartitionRuntimeSpatialHash.cpp
+++ b/Engine/Source/Runtime/Engine/Private/WorldPartition/WorldPartitionRuntimeSpatialHash.cpp
@@ -756,49 +756,6 @@ FName UWorldPartitionRuntimeSpatialHash::GetCellName(FName InGridName, const FIn
 	return UWorldPartitionRuntimeSpatialHash::GetCellName(WorldPartition, InGridName, InCellGlobalCoord, InDataLayerID);
 }
 
-void UWorldPartitionRuntimeSpatialHash::UpdateActorDescViewMap(TMap<FGuid, FWorldPartitionActorDescView>& ActorDescViewMap) const
-{
-	TRACE_CPUPROFILER_EVENT_SCOPE(UpdateActorDescViewMap);
-
-	Super::UpdateActorDescViewMap(ActorDescViewMap);
-
-	TMap<FName, int32> GridsMapping;
-	GridsMapping.Add(NAME_None, 0);
-	for (int32 i = 0; i < Grids.Num(); i++)
-	{
-		const FSpatialHashRuntimeGrid& Grid = Grids[i];
-		check(!GridsMapping.Contains(Grid.GridName));
-		GridsMapping.Add(Grid.GridName, i);
-	}
-
-	for (auto& ActorDescViewPair : ActorDescViewMap)
-	{
-		FWorldPartitionActorDescView& ActorDescView = ActorDescViewPair.Value;
-
-		if (ActorDescView.GetActorClass()->GetDefaultObject<AActor>()->GetDefaultGridPlacement() == EActorGridPlacement::None)
-		{
-			if (ActorDescView.GetGridPlacement() != EActorGridPlacement::AlwaysLoaded)
-			{
-				const int32* GridIndexPtr = GridsMapping.Find(ActorDescView.GetRuntimeGrid());
-				const int32 GridIndex = GridIndexPtr ? *GridIndexPtr : 0;
-				const FSpatialHashRuntimeGrid& RuntimeGrid = Grids[GridIndex];
-				const float CellArea = RuntimeGrid.CellSize * RuntimeGrid.CellSize;
-				const FBox2D ActorBounds2D = FBox2D(FVector2D(ActorDescView.GetBounds().Min), FVector2D(ActorDescView.GetBounds().Max));
-				const float ActorBoundsArea = ActorBounds2D.GetArea();
-
-				if (ActorBoundsArea < 1.0f)
-				{
-					ActorDescView.SetGridPlacement(EActorGridPlacement::Location);
-				}
-				else if (ActorBoundsArea > CellArea)
-				{
-					ActorDescView.SetGridPlacement(EActorGridPlacement::Bounds);
-				}
-			}
-		}
-	}
-}
-
 bool UWorldPartitionRuntimeSpatialHash::CreateStreamingGrid(const FSpatialHashRuntimeGrid& RuntimeGrid, const FSquare2DGridHelper& PartionedActors, UWorldPartitionStreamingPolicy* StreamingPolicy, TArray<FString>* OutPackagesToGenerate)
 {
 	TRACE_CPUPROFILER_EVENT_SCOPE(CreateStreamingGrid);
diff --git a/Engine/Source/Runtime/Engine/Private/WorldPartition/WorldPartitionStreamingGeneration.cpp b/Engine/Source/Runtime/Engine/Private/WorldPartition/WorldPartitionStreamingGeneration.cpp
--- a/Engine/Source/Runtime/Engine/Private/WorldPartition/WorldPartitionStreamingGeneration.cpp
+++ b/Engine/Source/Runtime/Engine/Private/WorldPartition/WorldPartitionStreamingGeneration.cpp
@@ -178,14 +178,14 @@ class FWorldPartitionStreamingGenerator
 					{
 						FWorldPartitionActorDescView& ActorDescView = ContainerDescriptor.ActorDescViewMap.FindChecked(Actor->GetActorGuid());
 
-						if (ActorDescView.GetGridPlacement() != EActorGridPlacement::AlwaysLoaded)
-						{						
+						if (ActorDescView.GetIsSpatiallyLoaded())
+						{
 							if (ErrorHandler)
 							{
 								ErrorHandler->OnInvalidReferenceLevelScriptStreamed(ActorDescView);
 							}
 
-							ActorDescView.SetGridPlacement(EActorGridPlacement::AlwaysLoaded);
+							ActorDescView.SetForcedNonSpatiallyLoaded();
 						}
 
 						if (ActorDescView.GetDataLayers().Num())
@@ -201,12 +201,6 @@ class FWorldPartitionStreamingGenerator
 				}
 			}
 
-			// Give the associated runtime hash the possibility to adjust actor descriptor views based on its internal settings, etc.
-			if (RuntimeHash != nullptr)
-			{
-				RuntimeHash->UpdateActorDescViewMap(ContainerDescriptor.ActorDescViewMap);
-			}
-
 			// Perform various adjustements based on validations and errors
 			for (auto ActorDescIt = ContainerDescriptor.ActorDescViewMap.CreateIterator(); ActorDescIt; ++ActorDescIt)
 			{
@@ -218,18 +212,18 @@ class FWorldPartitionStreamingGenerator
 					if (FWorldPartitionActorDescView* ReferenceActorDescView = ContainerDescriptor.ActorDescViewMap.Find(ReferenceGuid))
 					{
 						// Validate grid placement
-						const bool bIsActorDescAlwaysLoaded = ActorDescView.GetGridPlacement() == EActorGridPlacement::AlwaysLoaded;
-						const bool bIsActorDescRefAlwaysLoaded = ReferenceActorDescView->GetGridPlacement() == EActorGridPlacement::AlwaysLoaded;
+						const bool bIsActorDescSpatiallyLoaded = ActorDescView.GetIsSpatiallyLoaded();
+						const bool bIsActorDescRefSpatiallyLoaded = ReferenceActorDescView->GetIsSpatiallyLoaded();
 
-						if (bIsActorDescAlwaysLoaded != bIsActorDescRefAlwaysLoaded)
+						if (bIsActorDescSpatiallyLoaded != bIsActorDescRefSpatiallyLoaded)
 						{
 							if (ErrorHandler)
 							{
 								ErrorHandler->OnInvalidReferenceGridPlacement(ActorDescView, *ReferenceActorDescView);
 							}
 
-							ActorDescView.SetGridPlacement(EActorGridPlacement::AlwaysLoaded);
-							ReferenceActorDescView->SetGridPlacement(EActorGridPlacement::AlwaysLoaded);
+							ActorDescView.SetForcedNonSpatiallyLoaded();
+							ReferenceActorDescView->SetForcedNonSpatiallyLoaded();
 						}
 
 						// Validate data layers
@@ -281,10 +275,9 @@ class FWorldPartitionStreamingGenerator
 			{
 				const FWorldPartitionActorDescView& ActorDescView = ActorDescIt.Value();
 
-				switch (ActorDescView.GetGridPlacement())
+				if (ActorDescView.GetIsSpatiallyLoaded())
 				{
-					case EActorGridPlacement::Location: ContainerDescriptor.Bounds += ContainerDescriptor.Transform.TransformPosition(ActorDescView.GetOrigin()); break;
-					case EActorGridPlacement::Bounds: ContainerDescriptor.Bounds += ActorDescView.GetBounds().TransformBy(ContainerDescriptor.Transform); break;
+					ContainerDescriptor.Bounds += ActorDescView.GetBounds().TransformBy(ContainerDescriptor.Transform);
 				}
 			}
 		}
@@ -299,9 +292,8 @@ class FWorldPartitionStreamingGenerator
 	}
 
 public:
-	FWorldPartitionStreamingGenerator(UWorldPartitionRuntimeHash* InRuntimeHash, FActorDescList* InModifiedActorsDescList, IStreamingGenerationErrorHandler* InErrorHandler = nullptr)
-	: RuntimeHash(InRuntimeHash)
-	, ModifiedActorsDescList(InModifiedActorsDescList)
+	FWorldPartitionStreamingGenerator(FActorDescList* InModifiedActorsDescList, IStreamingGenerationErrorHandler* InErrorHandler = nullptr)
+	: ModifiedActorsDescList(InModifiedActorsDescList)
 	, ErrorHandler(InErrorHandler)
 	{}
 
@@ -413,7 +405,6 @@ public:
 	}
 
 private:
-	UWorldPartitionRuntimeHash* RuntimeHash;
 	FActorDescList* ModifiedActorsDescList;
 	IStreamingGenerationErrorHandler* ErrorHandler;
 
@@ -461,7 +452,7 @@ bool UWorldPartition::GenerateStreaming(TArray<FString>* OutPackagesToGenerate)
 
 	FActorClusterContext ActorClusterContext;
 	{
-		FWorldPartitionStreamingGenerator StreamingGenerator(RuntimeHash, ModifiedActorsDescList, ErrorHandler);
+		FWorldPartitionStreamingGenerator StreamingGenerator(ModifiedActorsDescList, ErrorHandler);
 
 		// Preparation Phase
 		StreamingGenerator.PreparationPhase(this);
@@ -488,7 +479,7 @@ bool UWorldPartition::GenerateStreaming(TArray<FString>* OutPackagesToGenerate)
 void UWorldPartition::GenerateHLOD(ISourceControlHelper* SourceControlHelper, bool bCreateActorsOnly)
 {
 	FStreamingGenerationLogErrorHandler LogErrorHandler;
-	FWorldPartitionStreamingGenerator StreamingGenerator(RuntimeHash, nullptr, &LogErrorHandler);
+	FWorldPartitionStreamingGenerator StreamingGenerator(nullptr, &LogErrorHandler);
 	StreamingGenerator.PreparationPhase(this);
 
 	// Preparation Phase :: Actor Clusters Creation
@@ -505,7 +496,7 @@ void UWorldPartition::CheckForErrors(IStreamingGenerationErrorHandler* ErrorHand
 	FActorClusterContext ActorClusterContext;
 	{		
 		FActorDescList ModifiedActorDescList;
-		FWorldPartitionStreamingGenerator StreamingGenerator(RuntimeHash, &ModifiedActorDescList, ErrorHandler);
+		FWorldPartitionStreamingGenerator StreamingGenerator(&ModifiedActorDescList, ErrorHandler);
 		StreamingGenerator.PreparationPhase(this);
 	}
 }
@@ -514,7 +505,7 @@ void UWorldPartition::CheckForErrors(IStreamingGenerationErrorHandler* ErrorHand
 {
 	FActorClusterContext ActorClusterContext;
 	{
-		FWorldPartitionStreamingGenerator StreamingGenerator(nullptr, nullptr, ErrorHandler);
+		FWorldPartitionStreamingGenerator StreamingGenerator(nullptr, ErrorHandler);
 		StreamingGenerator.PreparationPhase(ActorDescContainer);
 	}
 }
diff --git a/Engine/Source/Runtime/Engine/Private/WorldPartition/WorldPartitionVolume.cpp b/Engine/Source/Runtime/Engine/Private/WorldPartition/WorldPartitionVolume.cpp
--- a/Engine/Source/Runtime/Engine/Private/WorldPartition/WorldPartitionVolume.cpp
+++ b/Engine/Source/Runtime/Engine/Private/WorldPartition/WorldPartitionVolume.cpp
@@ -11,6 +11,9 @@
 AWorldPartitionVolume::AWorldPartitionVolume(const FObjectInitializer& ObjectInitializer)
 	: Super(ObjectInitializer)
 {
+#if WITH_EDITORONLY_DATA
+	bIsSpatiallyLoaded = false;
+#endif
 }
 
 #if WITH_EDITOR
diff --git a/Engine/Source/Runtime/Engine/Public/ActorPartition/PartitionActor.h b/Engine/Source/Runtime/Engine/Public/ActorPartition/PartitionActor.h
--- a/Engine/Source/Runtime/Engine/Public/ActorPartition/PartitionActor.h
+++ b/Engine/Source/Runtime/Engine/Public/ActorPartition/PartitionActor.h
@@ -20,7 +20,7 @@ public:
 
 	//~ Begin AActor Interface
 #if WITH_EDITOR	
-	virtual EActorGridPlacement GetDefaultGridPlacement() const override { return EActorGridPlacement::Bounds; }
+	virtual bool CanChangeIsSpatiallyLoadedFlag() const override { return false; }
 	virtual TUniquePtr<class FWorldPartitionActorDesc> CreateClassActorDesc() const override;
 	virtual uint32 GetDefaultGridSize(UWorld* InWorld) const PURE_VIRTUAL(APartitionActor, return 0;)
 	virtual bool ShouldIncludeGridSizeInName(UWorld * InWorld) const;
diff --git a/Engine/Source/Runtime/Engine/Public/LevelInstance/LevelInstanceActor.h b/Engine/Source/Runtime/Engine/Public/LevelInstance/LevelInstanceActor.h
--- a/Engine/Source/Runtime/Engine/Public/LevelInstance/LevelInstanceActor.h
+++ b/Engine/Source/Runtime/Engine/Public/LevelInstance/LevelInstanceActor.h
@@ -72,7 +72,6 @@ public:
 
 	// AActor overrides
 	virtual void CheckForErrors() override;
-	virtual EActorGridPlacement GetDefaultGridPlacement() const override { return EActorGridPlacement::Bounds; }
 	virtual TUniquePtr<class FWorldPartitionActorDesc> CreateClassActorDesc() const override;
 	virtual bool CanDeleteSelectedActor(FText& OutReason) const override;
 	virtual void SetIsTemporarilyHiddenInEditor(bool bIsHidden) override;
diff --git a/Engine/Source/Runtime/Engine/Public/LevelInstance/Packed/PackedLevelInstanceActor.h b/Engine/Source/Runtime/Engine/Public/LevelInstance/Packed/PackedLevelInstanceActor.h
--- a/Engine/Source/Runtime/Engine/Public/LevelInstance/Packed/PackedLevelInstanceActor.h
+++ b/Engine/Source/Runtime/Engine/Public/LevelInstance/Packed/PackedLevelInstanceActor.h
@@ -47,9 +47,6 @@ public:
 		
 	void SetPackedVersion(const FGuid& Version) { PackedVersion = Version; }
 
-	// AActor overrides
-	virtual EActorGridPlacement GetDefaultGridPlacement() const override { return EActorGridPlacement::None; }
-
 	virtual bool IsHiddenEd() const override;
 	virtual bool IsHLODRelevant() const override;
 
diff --git a/Engine/Source/Runtime/Engine/Public/WorldPartition/HLOD/HLODActor.h b/Engine/Source/Runtime/Engine/Public/WorldPartition/HLOD/HLODActor.h
--- a/Engine/Source/Runtime/Engine/Public/WorldPartition/HLOD/HLODActor.h
+++ b/Engine/Source/Runtime/Engine/Public/WorldPartition/HLOD/HLODActor.h
@@ -60,8 +60,7 @@ public:
 	void BuildHLOD(bool bForceBuild = false);
 	uint32 GetHLODHash() const;
 
-	virtual EActorGridPlacement GetGridPlacement() const override;
-	virtual EActorGridPlacement GetDefaultGridPlacement() const override;
+	virtual bool CanChangeIsSpatiallyLoadedFlag() const override { return false; }
 #endif // WITH_EDITOR
 
 protected:
diff --git a/Engine/Source/Runtime/Engine/Public/WorldPartition/NavigationData/NavigationDataChunkActor.h b/Engine/Source/Runtime/Engine/Public/WorldPartition/NavigationData/NavigationDataChunkActor.h
--- a/Engine/Source/Runtime/Engine/Public/WorldPartition/NavigationData/NavigationDataChunkActor.h
+++ b/Engine/Source/Runtime/Engine/Public/WorldPartition/NavigationData/NavigationDataChunkActor.h
@@ -37,7 +37,7 @@ public:
 	void SetDataChunkActorBounds(const FBox& InBounds);
 
 	//~ Begin AActor Interface.
-	virtual EActorGridPlacement GetDefaultGridPlacement() const override;
+	virtual bool CanChangeIsSpatiallyLoadedFlag() const override { return false; }
 	virtual FBox GetStreamingBounds() const override;
 	//~ End AActor Interface.
 #endif // WITH_EDITOR
diff --git a/Engine/Source/Runtime/Engine/Public/WorldPartition/WorldPartitionActorCluster.h b/Engine/Source/Runtime/Engine/Public/WorldPartition/WorldPartitionActorCluster.h
--- a/Engine/Source/Runtime/Engine/Public/WorldPartition/WorldPartitionActorCluster.h
+++ b/Engine/Source/Runtime/Engine/Public/WorldPartition/WorldPartitionActorCluster.h
@@ -66,12 +66,12 @@ struct FActorContainerID
 
 #if WITH_EDITOR
 /**
- * List of actors bound together based on clustering rules. (mainly object references)
+ * List of actors bound together based on clustering rules (mainly object references)
  */
 struct FActorCluster
 {
 	TSet<FGuid>					Actors;
-	EActorGridPlacement			GridPlacement;
+	bool						bIsSpatiallyLoaded;
 	FName						RuntimeGrid;
 	FBox						Bounds;
 	TSet<const UDataLayer*>		DataLayers;
@@ -81,8 +81,6 @@ struct FActorCluster
 	void Add(const FActorCluster& InActorCluster, const TMap<FGuid, FWorldPartitionActorDescView>& InActorDescViewMap);
 };
 
-struct FActorClusterInstance;
-
 /**
  * Instance of a container (level) with specific instance properties (transform, id, datalayers) 
  */
diff --git a/Engine/Source/Runtime/Engine/Public/WorldPartition/WorldPartitionActorDesc.h b/Engine/Source/Runtime/Engine/Public/WorldPartition/WorldPartitionActorDesc.h
--- a/Engine/Source/Runtime/Engine/Public/WorldPartition/WorldPartitionActorDesc.h
+++ b/Engine/Source/Runtime/Engine/Public/WorldPartition/WorldPartitionActorDesc.h
@@ -61,8 +61,8 @@ public:
 	inline FName GetClass() const { return Class; }
 	inline UClass* GetActorClass() const { return ActorClass; }
 	inline FVector GetOrigin() const { return GetBounds().GetCenter(); }
-	inline EActorGridPlacement GetGridPlacement() const { return GridPlacement; }
 	inline FName GetRuntimeGrid() const { return RuntimeGrid; }
+	inline bool GetIsSpatiallyLoaded() const { return bIsSpatiallyLoaded; }
 	inline bool GetActorIsEditorOnly() const { return bActorIsEditorOnly; }
 	inline bool GetLevelBoundsRelevant() const { return bLevelBoundsRelevant; }
 	inline bool GetActorIsHLODRelevant() const { return bActorIsHLODRelevant; }
@@ -173,8 +173,8 @@ protected:
 	FName							ActorLabel;
 	FVector							BoundsLocation;
 	FVector							BoundsExtent;
-	EActorGridPlacement				GridPlacement;
 	FName							RuntimeGrid;
+	bool							bIsSpatiallyLoaded;
 	bool							bActorIsEditorOnly;
 	bool							bLevelBoundsRelevant;
 	bool							bActorIsHLODRelevant;
diff --git a/Engine/Source/Runtime/Engine/Public/WorldPartition/WorldPartitionActorDescView.h b/Engine/Source/Runtime/Engine/Public/WorldPartition/WorldPartitionActorDescView.h
--- a/Engine/Source/Runtime/Engine/Public/WorldPartition/WorldPartitionActorDescView.h
+++ b/Engine/Source/Runtime/Engine/Public/WorldPartition/WorldPartitionActorDescView.h
@@ -9,7 +9,6 @@ class FWorldPartitionActorDesc;
 class UActorDescContainer;
 class UWorldPartition;
 enum class EContainerClusterMode : uint8;
-enum class EActorGridPlacement : uint8;
 
 /**
  * A view on top of an actor desc, used to cache information that can be (potentially) different than the actor desc
@@ -27,8 +26,8 @@ public:
 	FName GetClass() const;
 	UClass* GetActorClass() const;
 	FVector GetOrigin() const;
-	EActorGridPlacement GetGridPlacement() const;
 	FName GetRuntimeGrid() const;
+	bool GetIsSpatiallyLoaded() const;
 	bool GetActorIsEditorOnly() const;
 	bool GetLevelBoundsRelevant() const;
 	bool GetActorIsHLODRelevant() const;
@@ -48,7 +47,7 @@ public:
 
 	FName GetActorLabelOrName() const;
 
-	void SetGridPlacement(EActorGridPlacement InGridPlacement);
+	void SetForcedNonSpatiallyLoaded();
 
 	void SetInvalidDataLayers();
 
@@ -64,7 +63,7 @@ public:
 
 protected:
 	const FWorldPartitionActorDesc* ActorDesc;
-	EActorGridPlacement GridPlacement;
+	bool bIsForcedNonSpatiallyLoaded;
 	bool bInvalidDataLayers;
 };
 #endif
\ No newline at end of file
diff --git a/Engine/Source/Runtime/Engine/Public/WorldPartition/WorldPartitionRuntimeHash.h b/Engine/Source/Runtime/Engine/Public/WorldPartition/WorldPartitionRuntimeHash.h
--- a/Engine/Source/Runtime/Engine/Public/WorldPartition/WorldPartitionRuntimeHash.h
+++ b/Engine/Source/Runtime/Engine/Public/WorldPartition/WorldPartitionRuntimeHash.h
@@ -70,11 +70,6 @@ class ENGINE_API UWorldPartitionRuntimeHash : public UObject
 protected:
 	virtual EWorldPartitionStreamingPerformance GetStreamingPerformanceForCell(const UWorldPartitionRuntimeCell* Cell) const { return EWorldPartitionStreamingPerformance::Good; }
 
-#if WITH_EDITOR
-public:
-	virtual void UpdateActorDescViewMap(TMap<FGuid, FWorldPartitionActorDescView>& ActorDescViewMap) const {}
-#endif
-
 private:
 #if WITH_EDITOR
 	void ForceExternalActorLevelReference(bool bForceExternalActorLevelReferenceForPIE);
diff --git a/Engine/Source/Runtime/Engine/Public/WorldPartition/WorldPartitionRuntimeSpatialHash.h b/Engine/Source/Runtime/Engine/Public/WorldPartition/WorldPartitionRuntimeSpatialHash.h
--- a/Engine/Source/Runtime/Engine/Public/WorldPartition/WorldPartitionRuntimeSpatialHash.h
+++ b/Engine/Source/Runtime/Engine/Public/WorldPartition/WorldPartitionRuntimeSpatialHash.h
@@ -236,8 +236,6 @@ public:
 
 	FName GetCellName(FName InGridName, const FIntVector& InCellGlobalCoord, const FDataLayersID& InDataLayerID) const;
 	static FName GetCellName(UWorldPartition* WorldPartition, FName InGridName, const FIntVector& InCellGlobalCoord, const FDataLayersID& InDataLayerID);
-
-	virtual void UpdateActorDescViewMap(TMap<FGuid, FWorldPartitionActorDescView>& ActorDescViewMap) const override;
 #endif
 
 	// streaming interface
diff --git a/Engine/Source/Runtime/Engine/Public/WorldPartition/WorldPartitionVolume.h b/Engine/Source/Runtime/Engine/Public/WorldPartition/WorldPartitionVolume.h
--- a/Engine/Source/Runtime/Engine/Public/WorldPartition/WorldPartitionVolume.h
+++ b/Engine/Source/Runtime/Engine/Public/WorldPartition/WorldPartitionVolume.h
@@ -23,7 +23,7 @@ public:
 #if WITH_EDITOR
 	void LoadIntersectingCells(bool bIsFromUserChange);
 	void UnloadIntersectingCells(bool bIsFromUserChange);
-	virtual EActorGridPlacement GetDefaultGridPlacement() const override { return EActorGridPlacement::AlwaysLoaded; }
+	virtual bool CanChangeIsSpatiallyLoadedFlag() const override { return false; }
 
 private:
 	virtual bool SupportsDataLayer() const override { return false; }
diff --git a/Engine/Source/Runtime/Landscape/Classes/Landscape.h b/Engine/Source/Runtime/Landscape/Classes/Landscape.h
--- a/Engine/Source/Runtime/Landscape/Classes/Landscape.h
+++ b/Engine/Source/Runtime/Landscape/Classes/Landscape.h
@@ -233,7 +233,7 @@ public:
 	virtual void PostEditImport() override;
 	virtual void PostDuplicate(bool bDuplicateForPIE) override;
 	virtual bool CanDeleteSelectedActor(FText& OutReason) const override;
-	virtual EActorGridPlacement GetDefaultGridPlacement() const override { return EActorGridPlacement::AlwaysLoaded; }
+	virtual bool CanChangeIsSpatiallyLoadedFlag() const override { return false; }
 #endif
 	virtual void PostLoad() override;
 	virtual void BeginDestroy() override;
diff --git a/Engine/Source/Runtime/Landscape/Classes/LandscapeProxy.h b/Engine/Source/Runtime/Landscape/Classes/LandscapeProxy.h
--- a/Engine/Source/Runtime/Landscape/Classes/LandscapeProxy.h
+++ b/Engine/Source/Runtime/Landscape/Classes/LandscapeProxy.h
@@ -799,7 +799,6 @@ public:
 	virtual void PostEditMove(bool bFinished) override;
 	virtual bool ShouldImport(FString* ActorPropString, bool IsMovingLevel) override;
 	virtual bool ShouldExport() override;
-	virtual EActorGridPlacement GetDefaultGridPlacement() const override { return EActorGridPlacement::None; }
 	//~ End AActor Interface
 	virtual uint32 GetDefaultGridSize(UWorld* InWorld) const override { return 1; }
 	virtual FGuid GetGridGuid() const override { return LandscapeGuid; }
diff --git a/Engine/Source/Runtime/Landscape/Classes/LandscapeSplineActor.h b/Engine/Source/Runtime/Landscape/Classes/LandscapeSplineActor.h
--- a/Engine/Source/Runtime/Landscape/Classes/LandscapeSplineActor.h
+++ b/Engine/Source/Runtime/Landscape/Classes/LandscapeSplineActor.h
@@ -26,7 +26,7 @@ public:
 	virtual ULandscapeInfo* GetLandscapeInfo() const override;
 		
 #if WITH_EDITOR
-	virtual EActorGridPlacement GetDefaultGridPlacement() const override { return EActorGridPlacement::Bounds; }
+	virtual bool CanChangeIsSpatiallyLoadedFlag() const override { return false; }
 	virtual TUniquePtr<class FWorldPartitionActorDesc> CreateClassActorDesc() const override;
 
 	void GetSharedProperties(ULandscapeInfo* InLandscapeInfo);
diff --git a/Engine/Source/Runtime/Landscape/Private/Landscape.cpp b/Engine/Source/Runtime/Landscape/Private/Landscape.cpp
--- a/Engine/Source/Runtime/Landscape/Private/Landscape.cpp
+++ b/Engine/Source/Runtime/Landscape/Private/Landscape.cpp
@@ -1094,8 +1094,6 @@ ALandscapeProxy::ALandscapeProxy(const FObjectInitializer& ObjectInitializer)
 
 #if WITH_EDITORONLY_DATA
 	LandscapeProxies.Add(this);
-
-	GridPlacement = EActorGridPlacement::Location;
 #endif
 }
 
@@ -1122,6 +1120,7 @@ ALandscape::ALandscape(const FObjectInitializer& ObjectInitializer)
 	bLandscapeLayersAreInitialized = false;
 	LandscapeEdMode = nullptr;
 	bGrassUpdateEnabled = true;
+	bIsSpatiallyLoaded = false;
 #endif // WITH_EDITORONLY_DATA
 }
 
diff --git a/Engine/Source/Runtime/Landscape/Private/LandscapeBlueprintBrushBase.cpp b/Engine/Source/Runtime/Landscape/Private/LandscapeBlueprintBrushBase.cpp
--- a/Engine/Source/Runtime/Landscape/Private/LandscapeBlueprintBrushBase.cpp
+++ b/Engine/Source/Runtime/Landscape/Private/LandscapeBlueprintBrushBase.cpp
@@ -39,6 +39,10 @@ ALandscapeBlueprintBrushBase::ALandscapeBlueprintBrushBase(const FObjectInitiali
 	PrimaryActorTick.SetTickFunctionEnable(true);
 	bIsEditorOnlyActor = true;
 #endif // WITH_EDITOR
+
+#if WITH_EDITORONLY_DATA
+	bIsSpatiallyLoaded = false;
+#endif
 }
 
 UTextureRenderTarget2D* ALandscapeBlueprintBrushBase::Render_Implementation(bool InIsHeightmap, UTextureRenderTarget2D* InCombinedResult, const FName& InWeightmapLayerName)
diff --git a/Engine/Source/Runtime/Landscape/Public/LandscapeBlueprintBrushBase.h b/Engine/Source/Runtime/Landscape/Public/LandscapeBlueprintBrushBase.h
--- a/Engine/Source/Runtime/Landscape/Public/LandscapeBlueprintBrushBase.h
+++ b/Engine/Source/Runtime/Landscape/Public/LandscapeBlueprintBrushBase.h
@@ -77,6 +77,6 @@ public:
 
 	virtual void PushDeferredLayersContentUpdate();
 
-	virtual EActorGridPlacement GetDefaultGridPlacement() const override { return EActorGridPlacement::AlwaysLoaded; }
+	virtual bool CanChangeIsSpatiallyLoadedFlag() const override { return false; }
 #endif
 };
\ No newline at end of file
diff --git a/Engine/Source/Runtime/LevelSequence/Private/LevelSequenceActor.cpp b/Engine/Source/Runtime/LevelSequence/Private/LevelSequenceActor.cpp
--- a/Engine/Source/Runtime/LevelSequence/Private/LevelSequenceActor.cpp
+++ b/Engine/Source/Runtime/LevelSequence/Private/LevelSequenceActor.cpp
@@ -57,6 +57,8 @@ ALevelSequenceActor::ALevelSequenceActor(const FObjectInitializer& Init)
 			SpriteComponent->bHiddenInGame = true;
 		}
 	}
+
+	bIsSpatiallyLoaded = false;
 #endif //WITH_EDITORONLY_DATA
 
 	BindingOverrides = Init.CreateDefaultSubobject<UMovieSceneBindingOverrides>(this, "BindingOverrides");
diff --git a/Engine/Source/Runtime/LevelSequence/Public/LevelSequenceActor.h b/Engine/Source/Runtime/LevelSequence/Public/LevelSequenceActor.h
--- a/Engine/Source/Runtime/LevelSequence/Public/LevelSequenceActor.h
+++ b/Engine/Source/Runtime/LevelSequence/Public/LevelSequenceActor.h
@@ -282,7 +282,7 @@ protected:
 	virtual void EndPlay(const EEndPlayReason::Type EndPlayReason) override;
 	virtual void RewindForReplay() override;
 #if WITH_EDITOR
-	virtual EActorGridPlacement GetDefaultGridPlacement() const override { return EActorGridPlacement::AlwaysLoaded; }
+	virtual bool CanChangeIsSpatiallyLoadedFlag() const override { return false; }
 #endif
 	//~ End AActor interface
 
diff --git a/Engine/Source/Runtime/NavigationSystem/Private/NavigationData.cpp b/Engine/Source/Runtime/NavigationSystem/Private/NavigationData.cpp
--- a/Engine/Source/Runtime/NavigationSystem/Private/NavigationData.cpp
+++ b/Engine/Source/Runtime/NavigationSystem/Private/NavigationData.cpp
@@ -166,6 +166,10 @@ ANavigationData::ANavigationData(const FObjectInitializer& ObjectInitializer)
 	USceneComponent* SceneComponent = CreateDefaultSubobject<USceneComponent>(TEXT("SceneComp"));
 	RootComponent = SceneComponent;
 	RootComponent->Mobility = EComponentMobility::Static;
+
+#if WITH_EDITORONLY_DATA
+	bIsSpatiallyLoaded = false;
+#endif
 }
 
 uint16 ANavigationData::GetNextUniqueID()
diff --git a/Engine/Source/Runtime/NavigationSystem/Public/NavigationData.h b/Engine/Source/Runtime/NavigationSystem/Public/NavigationData.h
--- a/Engine/Source/Runtime/NavigationSystem/Public/NavigationData.h
+++ b/Engine/Source/Runtime/NavigationSystem/Public/NavigationData.h
@@ -938,7 +938,7 @@ public:
 	void SetNavRenderingEnabled(bool bEnable);
 
 #if WITH_EDITOR
-	virtual EActorGridPlacement GetDefaultGridPlacement() const override { return EActorGridPlacement::AlwaysLoaded; }
+	virtual bool CanChangeIsSpatiallyLoadedFlag() const override { return false; }
 #endif
 
 protected:
