Add support for input bindings in components.

diff --git a/Engine/Plugins/Experimental/EnhancedInput/Source/EnhancedInput/Private/EnhancedInputActionDelegateBinding.cpp b/Engine/Plugins/Experimental/EnhancedInput/Source/EnhancedInput/Private/EnhancedInputActionDelegateBinding.cpp
--- a/Engine/Plugins/Experimental/EnhancedInput/Source/EnhancedInput/Private/EnhancedInputActionDelegateBinding.cpp
+++ b/Engine/Plugins/Experimental/EnhancedInput/Source/EnhancedInput/Private/EnhancedInputActionDelegateBinding.cpp
@@ -9,7 +9,7 @@ UEnhancedInputActionDelegateBinding::UEnhancedInputActionDelegateBinding(const F
 {
 }
 
-void UEnhancedInputActionDelegateBinding::BindToInputComponent(UInputComponent* InputComponent) const
+void UEnhancedInputActionDelegateBinding::BindToInputComponent(UInputComponent* InputComponent, UObject* ObjectToBindTo) const
 {
 	UEnhancedInputComponent* Component = Cast<UEnhancedInputComponent>(InputComponent);
 	if (!Component)
@@ -19,8 +19,7 @@ void UEnhancedInputActionDelegateBinding::BindToInputComponent(UInputComponent*
 
 	for(const FBlueprintEnhancedInputActionBinding& Binding : InputActionDelegateBindings)
 	{
-		UObject* Owner = CastChecked<UObject>(Component->GetOwner());
-		Component->BindAction(Binding.InputAction, Binding.TriggerEvent, Owner, Binding.FunctionNameToBind);
+		Component->BindAction(Binding.InputAction, Binding.TriggerEvent, ObjectToBindTo, Binding.FunctionNameToBind);
 	}
 }
 
@@ -29,7 +28,7 @@ UEnhancedInputActionValueBinding::UEnhancedInputActionValueBinding(const FObject
 {
 }
 
-void UEnhancedInputActionValueBinding::BindToInputComponent(UInputComponent* InputComponent) const
+void UEnhancedInputActionValueBinding::BindToInputComponent(UInputComponent* InputComponent, UObject* ObjectToBindTo) const
 {
 	UEnhancedInputComponent* Component = Cast<UEnhancedInputComponent>(InputComponent);
 	if (!Component)
diff --git a/Engine/Plugins/Experimental/EnhancedInput/Source/EnhancedInput/Private/InputDebugKeyDelegateBinding.cpp b/Engine/Plugins/Experimental/EnhancedInput/Source/EnhancedInput/Private/InputDebugKeyDelegateBinding.cpp
--- a/Engine/Plugins/Experimental/EnhancedInput/Source/EnhancedInput/Private/InputDebugKeyDelegateBinding.cpp
+++ b/Engine/Plugins/Experimental/EnhancedInput/Source/EnhancedInput/Private/InputDebugKeyDelegateBinding.cpp
@@ -4,7 +4,7 @@
 #include "EnhancedInputComponent.h"
 #include "GameFramework/Actor.h"
 
-void UInputDebugKeyDelegateBinding::BindToInputComponent(UInputComponent* InputComponent) const
+void UInputDebugKeyDelegateBinding::BindToInputComponent(UInputComponent* InputComponent, UObject* ObjectToBindTo) const
 {
 #if DEV_ONLY_KEY_BINDINGS_AVAILABLE
 	UEnhancedInputComponent* Component = Cast<UEnhancedInputComponent>(InputComponent);
@@ -15,8 +15,7 @@ void UInputDebugKeyDelegateBinding::BindToInputComponent(UInputComponent* InputC
 
 	for(const FBlueprintInputDebugKeyDelegateBinding& Binding : InputDebugKeyDelegateBindings)
 	{
-		UObject* Owner = CastChecked<UObject>(Component->GetOwner());
-		Component->BindDebugKey(Binding.InputChord, Binding.InputKeyEvent, Owner, Binding.FunctionNameToBind, Binding.bExecuteWhenPaused);
+		Component->BindDebugKey(Binding.InputChord, Binding.InputKeyEvent, ObjectToBindTo, Binding.FunctionNameToBind, Binding.bExecuteWhenPaused);
 	}
 #endif
 }
diff --git a/Engine/Plugins/Experimental/EnhancedInput/Source/EnhancedInput/Public/EnhancedInputActionDelegateBinding.h b/Engine/Plugins/Experimental/EnhancedInput/Source/EnhancedInput/Public/EnhancedInputActionDelegateBinding.h
--- a/Engine/Plugins/Experimental/EnhancedInput/Source/EnhancedInput/Public/EnhancedInputActionDelegateBinding.h
+++ b/Engine/Plugins/Experimental/EnhancedInput/Source/EnhancedInput/Public/EnhancedInputActionDelegateBinding.h
@@ -39,7 +39,7 @@ class ENHANCEDINPUT_API UEnhancedInputActionDelegateBinding : public UInputDeleg
 	TArray<FBlueprintEnhancedInputActionBinding> InputActionDelegateBindings;
 
 	//~ Begin UInputDelegateBinding Interface
-	virtual void BindToInputComponent(UInputComponent* InputComponent) const override;
+	virtual void BindToInputComponent(UInputComponent* InputComponent, UObject* ObjectToBindTo) const override;
 	//~ End UInputDelegateBinding Interface
 };
 
@@ -52,6 +52,6 @@ class ENHANCEDINPUT_API UEnhancedInputActionValueBinding : public UInputDelegate
 	TArray<FBlueprintEnhancedInputActionBinding> InputActionValueBindings;
 
 	//~ Begin UInputDelegateBinding Interface
-	virtual void BindToInputComponent(UInputComponent* InputComponent) const override;
+	virtual void BindToInputComponent(UInputComponent* InputComponent, UObject* ObjectToBindTo) const override;
 	//~ End UInputDelegateBinding Interface
 };
\ No newline at end of file
diff --git a/Engine/Plugins/Experimental/EnhancedInput/Source/EnhancedInput/Public/InputDebugKeyDelegateBinding.h b/Engine/Plugins/Experimental/EnhancedInput/Source/EnhancedInput/Public/InputDebugKeyDelegateBinding.h
--- a/Engine/Plugins/Experimental/EnhancedInput/Source/EnhancedInput/Public/InputDebugKeyDelegateBinding.h
+++ b/Engine/Plugins/Experimental/EnhancedInput/Source/EnhancedInput/Public/InputDebugKeyDelegateBinding.h
@@ -41,6 +41,6 @@ public:
 	TArray<FBlueprintInputDebugKeyDelegateBinding> InputDebugKeyDelegateBindings;
 
 	//~ Begin UInputDelegateBinding Interface
-	virtual void BindToInputComponent(UInputComponent* InputComponent) const override;
+	virtual void BindToInputComponent(UInputComponent* InputComponent, UObject* ObjectToBindTo) const override;
 	//~ End UInputDelegateBinding Interface
 };
diff --git a/Engine/Source/Runtime/Engine/Classes/Engine/InputActionDelegateBinding.h b/Engine/Source/Runtime/Engine/Classes/Engine/InputActionDelegateBinding.h
--- a/Engine/Source/Runtime/Engine/Classes/Engine/InputActionDelegateBinding.h
+++ b/Engine/Source/Runtime/Engine/Classes/Engine/InputActionDelegateBinding.h
@@ -42,6 +42,6 @@ class ENGINE_API UInputActionDelegateBinding : public UInputDelegateBinding
 	TArray<FBlueprintInputActionDelegateBinding> InputActionDelegateBindings;
 
 	//~ Begin UInputDelegateBinding Interface
-	virtual void BindToInputComponent(UInputComponent* InputComponent) const override;
+	virtual void BindToInputComponent(UInputComponent* InputComponent, UObject* ObjectToBindTo) const override;
 	//~ End UInputDelegateBinding Interface
 };
diff --git a/Engine/Source/Runtime/Engine/Classes/Engine/InputAxisDelegateBinding.h b/Engine/Source/Runtime/Engine/Classes/Engine/InputAxisDelegateBinding.h
--- a/Engine/Source/Runtime/Engine/Classes/Engine/InputAxisDelegateBinding.h
+++ b/Engine/Source/Runtime/Engine/Classes/Engine/InputAxisDelegateBinding.h
@@ -37,6 +37,6 @@ class ENGINE_API UInputAxisDelegateBinding : public UInputDelegateBinding
 	TArray<FBlueprintInputAxisDelegateBinding> InputAxisDelegateBindings;
 
 	//~ Begin UInputDelegateBinding Interface
-	virtual void BindToInputComponent(UInputComponent* InputComponent) const override;
+	virtual void BindToInputComponent(UInputComponent* InputComponent, UObject* ObjectToBindTo) const override;
 	//~ End UInputDelegateBinding Interface
 };
diff --git a/Engine/Source/Runtime/Engine/Classes/Engine/InputAxisKeyDelegateBinding.h b/Engine/Source/Runtime/Engine/Classes/Engine/InputAxisKeyDelegateBinding.h
--- a/Engine/Source/Runtime/Engine/Classes/Engine/InputAxisKeyDelegateBinding.h
+++ b/Engine/Source/Runtime/Engine/Classes/Engine/InputAxisKeyDelegateBinding.h
@@ -37,6 +37,6 @@ class ENGINE_API UInputAxisKeyDelegateBinding : public UInputDelegateBinding
 	TArray<FBlueprintInputAxisKeyDelegateBinding> InputAxisKeyDelegateBindings;
 
 	//~ Begin UInputDelegateBinding Interface
-	virtual void BindToInputComponent(UInputComponent* InputComponent) const override;
+	virtual void BindToInputComponent(UInputComponent* InputComponent, UObject* ObjectToBindTo) const override;
 	//~ End UInputDelegateBinding Interface
 };
diff --git a/Engine/Source/Runtime/Engine/Classes/Engine/InputDelegateBinding.h b/Engine/Source/Runtime/Engine/Classes/Engine/InputDelegateBinding.h
--- a/Engine/Source/Runtime/Engine/Classes/Engine/InputDelegateBinding.h
+++ b/Engine/Source/Runtime/Engine/Classes/Engine/InputDelegateBinding.h
@@ -36,9 +36,34 @@ class ENGINE_API UInputDelegateBinding : public UDynamicBlueprintBinding
 {
 	GENERATED_UCLASS_BODY()
 
+	UE_DEPRECATED(5.0, "This version of BindToInputComponent is deprecated, please use the version that takes in an ObjectToBindTo instead.")
 	virtual void BindToInputComponent(UInputComponent* InputComponent) const { };
+	
+	/**
+	 * Override this function to bind a delegate to the given input component.
+	 *
+	 * @param InputComponent		The InputComponent to Bind a delegate to
+	 * @param ObjectToBindTo		The UObject that the binding should use.
+	 */
+	virtual void BindToInputComponent(UInputComponent* InputComponent, UObject* ObjectToBindTo) const { };
+
+	/** Returns true if the given class supports input binding delegates (i.e. it is a BP generated class) */
 	static bool SupportsInputDelegate(const UClass* InClass);
-	static void BindInputDelegates(const UClass* InClass, UInputComponent* InputComponent);
+
+	/**
+	 * Calls BindToInputComponent for each dynamic binding object on the given Class if it supports input delegates.
+	 *
+	 * @param InClass				The class that will should be used to determine if Input Delegates are supported
+	 * @param InputComponent		The InputComponent to Bind a delegate to
+	 * @param ObjectToBindTo		The UObject that the binding should use. If this is null, the Owner of the input componet will be used.
+	 */
+	static void BindInputDelegates(const UClass* InClass, UInputComponent* InputComponent, UObject* ObjectToBindTo = nullptr);
+	
+	/**
+	 * Will bind input delegates for the given Actor and traverse it's subobjects attempting to bind
+	 * each of them
+	 */
+	static void BindInputDelegatesWithSubojects(AActor* InActor, UInputComponent* InputComponent);
 
 protected:
 	static TSet<UClass*> InputBindingClasses;
diff --git a/Engine/Source/Runtime/Engine/Classes/Engine/InputKeyDelegateBinding.h b/Engine/Source/Runtime/Engine/Classes/Engine/InputKeyDelegateBinding.h
--- a/Engine/Source/Runtime/Engine/Classes/Engine/InputKeyDelegateBinding.h
+++ b/Engine/Source/Runtime/Engine/Classes/Engine/InputKeyDelegateBinding.h
@@ -42,6 +42,6 @@ class ENGINE_API UInputKeyDelegateBinding : public UInputDelegateBinding
 	TArray<FBlueprintInputKeyDelegateBinding> InputKeyDelegateBindings;
 
 	//~ Begin UInputDelegateBinding Interface
-	virtual void BindToInputComponent(UInputComponent* InputComponent) const override;
+	virtual void BindToInputComponent(UInputComponent* InputComponent, UObject* ObjectToBindTo) const override;
 	//~ End UInputDelegateBinding Interface
 };
diff --git a/Engine/Source/Runtime/Engine/Classes/Engine/InputTouchDelegateBinding.h b/Engine/Source/Runtime/Engine/Classes/Engine/InputTouchDelegateBinding.h
--- a/Engine/Source/Runtime/Engine/Classes/Engine/InputTouchDelegateBinding.h
+++ b/Engine/Source/Runtime/Engine/Classes/Engine/InputTouchDelegateBinding.h
@@ -39,6 +39,6 @@ class ENGINE_API UInputTouchDelegateBinding : public UInputDelegateBinding
 	TArray<FBlueprintInputTouchDelegateBinding> InputTouchDelegateBindings;
 
 	//~ Begin UInputDelegateBinding Interface
-	virtual void BindToInputComponent(UInputComponent* InputComponent) const override;
+	virtual void BindToInputComponent(UInputComponent* InputComponent, UObject* ObjectToBindTo) const override;
 	//~ End UInputDelegateBinding Interface
 };
diff --git a/Engine/Source/Runtime/Engine/Classes/Engine/InputVectorAxisDelegateBinding.h b/Engine/Source/Runtime/Engine/Classes/Engine/InputVectorAxisDelegateBinding.h
--- a/Engine/Source/Runtime/Engine/Classes/Engine/InputVectorAxisDelegateBinding.h
+++ b/Engine/Source/Runtime/Engine/Classes/Engine/InputVectorAxisDelegateBinding.h
@@ -15,6 +15,6 @@ class ENGINE_API UInputVectorAxisDelegateBinding : public UInputAxisKeyDelegateB
 	GENERATED_UCLASS_BODY()
 
 	//~ Begin UInputDelegateBinding Interface
-	virtual void BindToInputComponent(UInputComponent* InputComponent) const override;
+	virtual void BindToInputComponent(UInputComponent* InputComponent, UObject* ObjectToBindTo) const override;
 	//~ End UInputDelegateBinding Interface
 };
diff --git a/Engine/Source/Runtime/Engine/Private/Actor.cpp b/Engine/Source/Runtime/Engine/Private/Actor.cpp
--- a/Engine/Source/Runtime/Engine/Private/Actor.cpp
+++ b/Engine/Source/Runtime/Engine/Private/Actor.cpp
@@ -3877,7 +3877,7 @@ void AActor::EnableInput(APlayerController* PlayerController)
 			InputComponent->bBlockInput = bBlockInput;
 			InputComponent->Priority = InputPriority;
 
-			UInputDelegateBinding::BindInputDelegates(GetClass(), InputComponent);
+			UInputDelegateBinding::BindInputDelegatesWithSubojects(this, InputComponent);
 		}
 		else
 		{
diff --git a/Engine/Source/Runtime/Engine/Private/Blueprint.cpp b/Engine/Source/Runtime/Engine/Private/Blueprint.cpp
--- a/Engine/Source/Runtime/Engine/Private/Blueprint.cpp
+++ b/Engine/Source/Runtime/Engine/Private/Blueprint.cpp
@@ -1021,7 +1021,7 @@ bool UBlueprint::AllowsDynamicBinding() const
 
 bool UBlueprint::SupportsInputEvents() const
 {
-	return FBlueprintEditorUtils::IsActorBased(this);
+	return FBlueprintEditorUtils::IsActorBased(this) || FBlueprintEditorUtils::IsComponentBased(this);
 }
 
 struct FBlueprintInnerHelper
diff --git a/Engine/Source/Runtime/Engine/Private/InputActionDelegateBinding.cpp b/Engine/Source/Runtime/Engine/Private/InputActionDelegateBinding.cpp
--- a/Engine/Source/Runtime/Engine/Private/InputActionDelegateBinding.cpp
+++ b/Engine/Source/Runtime/Engine/Private/InputActionDelegateBinding.cpp
@@ -9,7 +9,7 @@ UInputActionDelegateBinding::UInputActionDelegateBinding(const FObjectInitialize
 {
 }
 
-void UInputActionDelegateBinding::BindToInputComponent(UInputComponent* InputComponent) const
+void UInputActionDelegateBinding::BindToInputComponent(UInputComponent* InputComponent, UObject* ObjectToBindTo) const
 {
 	TArray<FInputActionBinding> BindsToAdd;
 
@@ -20,7 +20,7 @@ void UInputActionDelegateBinding::BindToInputComponent(UInputComponent* InputCom
 		FInputActionBinding AB( Binding.InputActionName, Binding.InputKeyEvent );
 		AB.bConsumeInput = Binding.bConsumeInput;
 		AB.bExecuteWhenPaused = Binding.bExecuteWhenPaused;
-		AB.ActionDelegate.BindDelegate(InputComponent->GetOwner(), Binding.FunctionNameToBind);
+		AB.ActionDelegate.BindDelegate(ObjectToBindTo, Binding.FunctionNameToBind);
 
 		if (Binding.bOverrideParentBinding)
 		{
diff --git a/Engine/Source/Runtime/Engine/Private/InputAxisDelegateBinding.cpp b/Engine/Source/Runtime/Engine/Private/InputAxisDelegateBinding.cpp
--- a/Engine/Source/Runtime/Engine/Private/InputAxisDelegateBinding.cpp
+++ b/Engine/Source/Runtime/Engine/Private/InputAxisDelegateBinding.cpp
@@ -24,7 +24,7 @@ UInputAxisDelegateBinding::UInputAxisDelegateBinding(const FObjectInitializer& O
 {
 }
 
-void UInputAxisDelegateBinding::BindToInputComponent(UInputComponent* InputComponent) const
+void UInputAxisDelegateBinding::BindToInputComponent(UInputComponent* InputComponent, UObject* ObjectToBindTo) const
 {
 	TArray<FInputAxisBinding> BindsToAdd;
 
@@ -51,7 +51,7 @@ void UInputAxisDelegateBinding::BindToInputComponent(UInputComponent* InputCompo
 			FInputAxisBinding AB( Binding.InputAxisName );
 			AB.bConsumeInput = Binding.bConsumeInput;
 			AB.bExecuteWhenPaused = Binding.bExecuteWhenPaused;
-			AB.AxisDelegate.BindDelegate(InputComponent->GetOwner(), Binding.FunctionNameToBind);
+			AB.AxisDelegate.BindDelegate(ObjectToBindTo, Binding.FunctionNameToBind);
 
 			if (Binding.bOverrideParentBinding)
 			{
diff --git a/Engine/Source/Runtime/Engine/Private/InputAxisKeyDelegateBinding.cpp b/Engine/Source/Runtime/Engine/Private/InputAxisKeyDelegateBinding.cpp
--- a/Engine/Source/Runtime/Engine/Private/InputAxisKeyDelegateBinding.cpp
+++ b/Engine/Source/Runtime/Engine/Private/InputAxisKeyDelegateBinding.cpp
@@ -9,7 +9,7 @@ UInputAxisKeyDelegateBinding::UInputAxisKeyDelegateBinding(const FObjectInitiali
 {
 }
 
-void UInputAxisKeyDelegateBinding::BindToInputComponent(UInputComponent* InputComponent) const
+void UInputAxisKeyDelegateBinding::BindToInputComponent(UInputComponent* InputComponent, UObject* ObjectToBindTo) const
 {
 	TArray<FInputAxisKeyBinding> BindsToAdd;
 
@@ -20,7 +20,7 @@ void UInputAxisKeyDelegateBinding::BindToInputComponent(UInputComponent* InputCo
 		FInputAxisKeyBinding AB( Binding.AxisKey );
 		AB.bConsumeInput = Binding.bConsumeInput;
 		AB.bExecuteWhenPaused = Binding.bExecuteWhenPaused;
-		AB.AxisDelegate.BindDelegate(InputComponent->GetOwner(), Binding.FunctionNameToBind);
+		AB.AxisDelegate.BindDelegate(ObjectToBindTo, Binding.FunctionNameToBind);
 
 		if (Binding.bOverrideParentBinding)
 		{
diff --git a/Engine/Source/Runtime/Engine/Private/InputDelegateBinding.cpp b/Engine/Source/Runtime/Engine/Private/InputDelegateBinding.cpp
--- a/Engine/Source/Runtime/Engine/Private/InputDelegateBinding.cpp
+++ b/Engine/Source/Runtime/Engine/Private/InputDelegateBinding.cpp
@@ -21,11 +21,17 @@ bool UInputDelegateBinding::SupportsInputDelegate(const UClass* InClass)
 	return !!Cast<UBlueprintGeneratedClass>(InClass);
 }
 
-void UInputDelegateBinding::BindInputDelegates(const UClass* InClass, UInputComponent* InputComponent)
+void UInputDelegateBinding::BindInputDelegates(const UClass* InClass, UInputComponent* InputComponent, UObject* InObjectToBindTo /* = nullptr */)
 {
 	if (SupportsInputDelegate(InClass))
 	{
-		BindInputDelegates(InClass->GetSuperClass(), InputComponent);
+		ensureMsgf(InputComponent, TEXT("Attempting to bind input delegates to an invalid Input Component!"));
+		
+		// If there was an object given to bind to use that, otherwise fall back to the input component's owner
+		// which will be an AActor.
+		UObject* ObjectToBindTo = InObjectToBindTo ? InObjectToBindTo : InputComponent->GetOwner();
+		
+		BindInputDelegates(InClass->GetSuperClass(), InputComponent, ObjectToBindTo);
 
 		for(UClass* BindingClass : InputBindingClasses)
 		{
@@ -34,8 +40,35 @@ void UInputDelegateBinding::BindInputDelegates(const UClass* InClass, UInputComp
 				, ECastCheckedType::NullAllowed);
 			if (BindingObject)
 			{
-				BindingObject->BindToInputComponent(InputComponent);
+				BindingObject->BindToInputComponent(InputComponent, ObjectToBindTo);
 			}
 		}
 	}
+}
+
+void UInputDelegateBinding::BindInputDelegatesWithSubojects(AActor* InActor, UInputComponent* InputComponent)
+{
+	TRACE_CPUPROFILER_EVENT_SCOPE(UInputDelegateBinding::BindInputDelegatesWithSubojects);
+	
+	ensureMsgf(InActor && InputComponent, TEXT("Attempting to bind input delegates to an invalid actor or input component!"));
+
+	const UClass* ActorClass = InActor->GetClass();
+
+	if (SupportsInputDelegate(ActorClass))
+	{
+		// Bind any input delegates on the base actor class
+		UInputDelegateBinding::BindInputDelegates(ActorClass, InputComponent, InputComponent->GetOwner());
+
+		// Bind any input delegates on the actor's components
+		TInlineComponentArray<UActorComponent*> ComponentArray;
+		InActor->GetComponents(ComponentArray);
+		for(UActorComponent* Comp : ComponentArray)
+		{
+			const UClass* CompClass = Comp ? Comp->GetClass() : nullptr;
+			if(CompClass && SupportsInputDelegate(CompClass) && Comp != InputComponent)
+			{
+				UInputDelegateBinding::BindInputDelegates(CompClass, InputComponent, Comp);	
+			}			
+		}
+	}
 }
\ No newline at end of file
diff --git a/Engine/Source/Runtime/Engine/Private/InputKeyDelegateBinding.cpp b/Engine/Source/Runtime/Engine/Private/InputKeyDelegateBinding.cpp
--- a/Engine/Source/Runtime/Engine/Private/InputKeyDelegateBinding.cpp
+++ b/Engine/Source/Runtime/Engine/Private/InputKeyDelegateBinding.cpp
@@ -9,7 +9,7 @@ UInputKeyDelegateBinding::UInputKeyDelegateBinding(const FObjectInitializer& Obj
 {
 }
 
-void UInputKeyDelegateBinding::BindToInputComponent(UInputComponent* InputComponent) const
+void UInputKeyDelegateBinding::BindToInputComponent(UInputComponent* InputComponent, UObject* ObjectToBindTo) const
 {
 	TArray<FInputKeyBinding> BindsToAdd;
 
@@ -20,7 +20,7 @@ void UInputKeyDelegateBinding::BindToInputComponent(UInputComponent* InputCompon
 		FInputKeyBinding KB( Binding.InputChord, Binding.InputKeyEvent );
 		KB.bConsumeInput = Binding.bConsumeInput;
 		KB.bExecuteWhenPaused = Binding.bExecuteWhenPaused;
-		KB.KeyDelegate.BindDelegate(InputComponent->GetOwner(), Binding.FunctionNameToBind);
+		KB.KeyDelegate.BindDelegate(ObjectToBindTo, Binding.FunctionNameToBind);
 
 		if (Binding.bOverrideParentBinding)
 		{
diff --git a/Engine/Source/Runtime/Engine/Private/InputTouchDelegateBinding.cpp b/Engine/Source/Runtime/Engine/Private/InputTouchDelegateBinding.cpp
--- a/Engine/Source/Runtime/Engine/Private/InputTouchDelegateBinding.cpp
+++ b/Engine/Source/Runtime/Engine/Private/InputTouchDelegateBinding.cpp
@@ -9,7 +9,7 @@ UInputTouchDelegateBinding::UInputTouchDelegateBinding(const FObjectInitializer&
 {
 }
 
-void UInputTouchDelegateBinding::BindToInputComponent(UInputComponent* InputComponent) const
+void UInputTouchDelegateBinding::BindToInputComponent(UInputComponent* InputComponent, UObject* ObjectToBindTo) const
 {
 	TArray<FInputTouchBinding> BindsToAdd;
 
@@ -20,7 +20,7 @@ void UInputTouchDelegateBinding::BindToInputComponent(UInputComponent* InputComp
 		FInputTouchBinding TB( Binding.InputKeyEvent );
 		TB.bConsumeInput = Binding.bConsumeInput;
 		TB.bExecuteWhenPaused = Binding.bExecuteWhenPaused;
-		TB.TouchDelegate.BindDelegate(InputComponent->GetOwner(), Binding.FunctionNameToBind);
+		TB.TouchDelegate.BindDelegate(ObjectToBindTo, Binding.FunctionNameToBind);
 
 		if (Binding.bOverrideParentBinding)
 		{
diff --git a/Engine/Source/Runtime/Engine/Private/InputVectorAxisDelegateBinding.cpp b/Engine/Source/Runtime/Engine/Private/InputVectorAxisDelegateBinding.cpp
--- a/Engine/Source/Runtime/Engine/Private/InputVectorAxisDelegateBinding.cpp
+++ b/Engine/Source/Runtime/Engine/Private/InputVectorAxisDelegateBinding.cpp
@@ -10,7 +10,7 @@ UInputVectorAxisDelegateBinding::UInputVectorAxisDelegateBinding(const FObjectIn
 {
 }
 
-void UInputVectorAxisDelegateBinding::BindToInputComponent(UInputComponent* InputComponent) const
+void UInputVectorAxisDelegateBinding::BindToInputComponent(UInputComponent* InputComponent, UObject* ObjectToBindTo) const
 {
 	TArray<FInputVectorAxisBinding> BindsToAdd;
 
@@ -21,7 +21,7 @@ void UInputVectorAxisDelegateBinding::BindToInputComponent(UInputComponent* Inpu
 		FInputVectorAxisBinding VAB( Binding.AxisKey );
 		VAB.bConsumeInput = Binding.bConsumeInput;
 		VAB.bExecuteWhenPaused = Binding.bExecuteWhenPaused;
-		VAB.AxisDelegate.BindDelegate(InputComponent->GetOwner(), Binding.FunctionNameToBind);
+		VAB.AxisDelegate.BindDelegate(ObjectToBindTo, Binding.FunctionNameToBind);
 
 		if (Binding.bOverrideParentBinding)
 		{
diff --git a/Engine/Source/Runtime/Engine/Private/Pawn.cpp b/Engine/Source/Runtime/Engine/Private/Pawn.cpp
--- a/Engine/Source/Runtime/Engine/Private/Pawn.cpp
+++ b/Engine/Source/Runtime/Engine/Private/Pawn.cpp
@@ -463,9 +463,8 @@ void APawn::PawnClientRestart()
 				if (UInputDelegateBinding::SupportsInputDelegate(GetClass()))
 				{
 					InputComponent->bBlockInput = bBlockInput;
-					UInputDelegateBinding::BindInputDelegates(GetClass(), InputComponent);
+					UInputDelegateBinding::BindInputDelegatesWithSubojects(this, InputComponent);
 				}
-
 			}
 		}
 	}
diff --git a/Engine/Source/Runtime/Engine/Private/PlayerController.cpp b/Engine/Source/Runtime/Engine/Private/PlayerController.cpp
--- a/Engine/Source/Runtime/Engine/Private/PlayerController.cpp
+++ b/Engine/Source/Runtime/Engine/Private/PlayerController.cpp
@@ -2522,8 +2522,8 @@ void APlayerController::SetupInputComponent()
 
 	if (UInputDelegateBinding::SupportsInputDelegate(GetClass()))
 	{
-		InputComponent->bBlockInput = bBlockInput;
-		UInputDelegateBinding::BindInputDelegates(GetClass(), InputComponent);
+		InputComponent->bBlockInput = bBlockInput;		
+		UInputDelegateBinding::BindInputDelegatesWithSubojects(this, InputComponent);
 	}
 }
 
