Rename HLODLayer.bAlwaysLoaded to bIsNonSpatiallyLoaded to be consistent with 18395825.

diff --git a/Engine/Source/Developer/WorldPartitionHLODUtilities/Private/WorldPartitionHLODUtilities.cpp b/Engine/Source/Developer/WorldPartitionHLODUtilities/Private/WorldPartitionHLODUtilities.cpp
--- a/Engine/Source/Developer/WorldPartitionHLODUtilities/Private/WorldPartitionHLODUtilities.cpp
+++ b/Engine/Source/Developer/WorldPartitionHLODUtilities/Private/WorldPartitionHLODUtilities.cpp
@@ -269,7 +269,7 @@ TArray<AWorldPartitionHLOD*> FWorldPartitionHLODUtilities::CreateHLODActors(FHLO
 
 		// Spatially loaded
 		// HLOD that are always loaded will not take the SubActorsInfo.GridPlacement into account
-		bool bExpectedIsSpatiallyLoaded = HLODLayer->IsAlwaysLoaded() ? false : SubActorsInfo.bIsSpatiallyLoaded;
+		bool bExpectedIsSpatiallyLoaded = !HLODLayer->IsSpatiallyLoaded() ? false : SubActorsInfo.bIsSpatiallyLoaded;
 		if (HLODActor->GetIsSpatiallyLoaded() != bExpectedIsSpatiallyLoaded)
 		{
 			HLODActor->SetIsSpatiallyLoaded(bExpectedIsSpatiallyLoaded);
diff --git a/Engine/Source/Runtime/Engine/Private/WorldPartition/HLOD/HLODLayer.cpp b/Engine/Source/Runtime/Engine/Private/WorldPartition/HLOD/HLODLayer.cpp
--- a/Engine/Source/Runtime/Engine/Private/WorldPartition/HLOD/HLODLayer.cpp
+++ b/Engine/Source/Runtime/Engine/Private/WorldPartition/HLOD/HLODLayer.cpp
@@ -26,6 +26,7 @@ DEFINE_LOG_CATEGORY_STATIC(LogHLODLayer, Log, All);
 UHLODLayer::UHLODLayer(const FObjectInitializer& ObjectInitializer)
 	: Super(ObjectInitializer)
 #if WITH_EDITORONLY_DATA
+	, bIsSpatiallyLoaded(true)
 	, CellSize(3200)
 	, LoadingRange(12800)
 #endif
@@ -103,6 +104,11 @@ void UHLODLayer::PostLoad()
 			HLODBuilderSettings = WPHLODUtilities->CreateHLODBuilderSettings(this);
 		}
 	}
+
+	if (bAlwaysLoaded_DEPRECATED)
+	{
+		bIsSpatiallyLoaded = false;
+	}
 }
 
 void UHLODLayer::PostEditChangeProperty(FPropertyChangedEvent& PropertyChangedEvent)
@@ -131,13 +137,13 @@ FName UHLODLayer::GetRuntimeGridName(uint32 InLODLevel, int32 InCellSize, double
 
 FName UHLODLayer::GetRuntimeGrid(uint32 InHLODLevel) const
 {
-	return IsAlwaysLoaded() ? NAME_None : GetRuntimeGridName(InHLODLevel, CellSize, LoadingRange);
+	return !IsSpatiallyLoaded() ? NAME_None : GetRuntimeGridName(InHLODLevel, CellSize, LoadingRange);
 }
 
 const TSoftObjectPtr<UHLODLayer>& UHLODLayer::GetParentLayer() const
 {
 	static const TSoftObjectPtr<UHLODLayer> NullLayer;
-	return IsAlwaysLoaded() ? NullLayer : ParentLayer;
+	return !IsSpatiallyLoaded() ? NullLayer : ParentLayer;
 }
 
 #endif // WITH_EDITOR
diff --git a/Engine/Source/Runtime/Engine/Private/WorldPartition/RuntimeSpatialHash/RuntimeSpatialHashHLOD.cpp b/Engine/Source/Runtime/Engine/Private/WorldPartition/RuntimeSpatialHash/RuntimeSpatialHashHLOD.cpp
--- a/Engine/Source/Runtime/Engine/Private/WorldPartition/RuntimeSpatialHash/RuntimeSpatialHashHLOD.cpp
+++ b/Engine/Source/Runtime/Engine/Private/WorldPartition/RuntimeSpatialHash/RuntimeSpatialHashHLOD.cpp
@@ -299,8 +299,8 @@ static TMap<FName, FSpatialHashRuntimeGrid> CreateHLODGrids(TMap<UHLODLayer*, in
 	{
 		UHLODLayer* HLODLayer = HLODLayerLevel.Key;
 
-		// No need to create a runtime grid if the HLOD layer is set to be always loaded
-		if (!HLODLayer->IsAlwaysLoaded())
+		// No need to create a runtime grid if the HLOD layer is set to be non spatially loaded
+		if (HLODLayer->IsSpatiallyLoaded())
 		{
 			int32 HLODLevel = HLODLayerLevel.Value;
 
diff --git a/Engine/Source/Runtime/Engine/Public/WorldPartition/HLOD/HLODLayer.h b/Engine/Source/Runtime/Engine/Public/WorldPartition/HLOD/HLODLayer.h
--- a/Engine/Source/Runtime/Engine/Public/WorldPartition/HLOD/HLODLayer.h
+++ b/Engine/Source/Runtime/Engine/Public/WorldPartition/HLOD/HLODLayer.h
@@ -48,10 +48,10 @@ public:
 	const TSubclassOf<UHLODBuilder> GetHLODBuilderClass() const { return HLODBuilderClass; }
 	const UHLODBuilderSettings* GetHLODBuilderSettings() const { return HLODBuilderSettings; }
 	FName GetRuntimeGrid(uint32 InHLODLevel) const;
-	int32 GetCellSize() const { return bAlwaysLoaded ? 0 : CellSize; }
-	float GetLoadingRange() const { return bAlwaysLoaded ? WORLD_MAX : LoadingRange; }
+	int32 GetCellSize() const { return !bIsSpatiallyLoaded ? 0 : CellSize; }
+	float GetLoadingRange() const { return !bIsSpatiallyLoaded ? WORLD_MAX : LoadingRange; }
 	const TSoftObjectPtr<UHLODLayer>& GetParentLayer() const;
-	bool IsAlwaysLoaded() const { return bAlwaysLoaded; }
+	bool IsSpatiallyLoaded() const { return bIsSpatiallyLoaded; }
 
 	bool DoesRequireWarmup() const;
 
@@ -79,20 +79,23 @@ private:
 	UPROPERTY(VisibleAnywhere, Export, NoClear, Category=HLOD, meta = (EditInline, NoResetToDefault))
 	TObjectPtr<UHLODBuilderSettings> HLODBuilderSettings;
 
-	/** Whether HLOD actors generated for this layer will be always loaded */
+	/** Whether HLOD actors generated for this layer will be spatially loaded */
 	UPROPERTY(EditAnywhere, Config, Category=HLOD)
-	uint32 bAlwaysLoaded : 1;
+	uint32 bIsSpatiallyLoaded : 1;
+
+	UPROPERTY()
+	uint32 bAlwaysLoaded_DEPRECATED : 1;
 
 	/** Cell size of the runtime grid created to encompass HLOD actors generated for this HLOD Layer */
-	UPROPERTY(EditAnywhere, Config, Category=HLOD, meta = (EditConditionHides, EditCondition = "!bAlwaysLoaded"))
+	UPROPERTY(EditAnywhere, Config, Category=HLOD, meta = (EditConditionHides, EditCondition = "bIsSpatiallyLoaded"))
 	int32 CellSize;
 
 	/** Loading range of the runtime grid created to encompass HLOD actors generated for this HLOD Layer */
-	UPROPERTY(EditAnywhere, Config, Category=HLOD, meta = (EditConditionHides, EditCondition = "!bAlwaysLoaded"))
+	UPROPERTY(EditAnywhere, Config, Category=HLOD, meta = (EditConditionHides, EditCondition = "bIsSpatiallyLoaded"))
 	double LoadingRange;
 
 	/** HLOD Layer to assign to the generated HLOD actors */
-	UPROPERTY(EditAnywhere, Config, Category=HLOD, meta = (EditConditionHides, EditCondition = "!bAlwaysLoaded"))
+	UPROPERTY(EditAnywhere, Config, Category=HLOD, meta = (EditConditionHides, EditCondition = "bIsSpatiallyLoaded"))
 	TSoftObjectPtr<UHLODLayer> ParentLayer;
 
 public:
