Improve skeleton notify display in rewind debugger

diff --git a/Engine/Source/Runtime/Engine/Private/Animation/AnimInstance.cpp b/Engine/Source/Runtime/Engine/Private/Animation/AnimInstance.cpp
--- a/Engine/Source/Runtime/Engine/Private/Animation/AnimInstance.cpp
+++ b/Engine/Source/Runtime/Engine/Private/Animation/AnimInstance.cpp
@@ -1540,13 +1540,14 @@ void UAnimInstance::TriggerSingleAnimNotify(FAnimNotifyEventReference& EventRefe
 
 				if (InAnimInstance == this || InAnimInstance->bReceiveNotifiesFromLinkedInstances)
 				{
+					TRACE_ANIM_NOTIFY(this, *AnimNotifyEvent, Event);
+					
 					UFunction* Function = InAnimInstance->FindFunction(FuncName);
 					if (Function)
 					{
 						// if parameter is none, add event
 						if (Function->NumParms == 0)
 						{
-							TRACE_ANIM_NOTIFY(this, *AnimNotifyEvent, Event);
 							InAnimInstance->ProcessEvent(Function, nullptr);
 						}
 						else if ((Function->NumParms == 1) && (CastField<FObjectProperty>(Function->PropertyLink) != nullptr))
@@ -1558,7 +1559,6 @@ void UAnimInstance::TriggerSingleAnimNotify(FAnimNotifyEventReference& EventRefe
 
 							FAnimNotifierHandler_Parms Parms;
 							Parms.Notify = AnimNotifyEvent->Notify;
-							TRACE_ANIM_NOTIFY(this, *AnimNotifyEvent, Event);
 							InAnimInstance->ProcessEvent(Function, &Parms);
 						}
 						else
diff --git a/Engine/Source/Runtime/Engine/Private/Animation/AnimTrace.cpp b/Engine/Source/Runtime/Engine/Private/Animation/AnimTrace.cpp
--- a/Engine/Source/Runtime/Engine/Private/Animation/AnimTrace.cpp
+++ b/Engine/Source/Runtime/Engine/Private/Animation/AnimTrace.cpp
@@ -1060,6 +1060,15 @@ void FAnimTrace::OutputAnimNotify(UAnimInstance* InAnimInstance, const FAnimNoti
 		NotifyObject = InNotifyEvent.NotifyStateClass;
 		NotifyAsset = NotifyObject->GetOuter();
 	}
+	else
+	{
+        NotifyAsset = InNotifyEvent.GetLinkedSequence();
+		if (NotifyAsset == nullptr)
+		{
+			NotifyAsset = InNotifyEvent.GetLinkedMontage();
+		}
+	}
+	
 
 	TRACE_OBJECT(NotifyAsset);
 	TRACE_OBJECT(NotifyObject);
