Fix non-unity compile errors and SA warnings

#rnx

diff --git a/Engine/Source/Runtime/Engine/Private/InputDelegateBinding.cpp b/Engine/Source/Runtime/Engine/Private/InputDelegateBinding.cpp
--- a/Engine/Source/Runtime/Engine/Private/InputDelegateBinding.cpp
+++ b/Engine/Source/Runtime/Engine/Private/InputDelegateBinding.cpp
@@ -3,6 +3,8 @@
 #include "Engine/InputDelegateBinding.h"
 #include "UObject/Class.h"
 #include "Engine/BlueprintGeneratedClass.h"
+#include "GameFramework/Actor.h"
+#include "Components/InputComponent.h"
 
 TSet<UClass*> UInputDelegateBinding::InputBindingClasses;
 
@@ -23,7 +25,7 @@ bool UInputDelegateBinding::SupportsInputDelegate(const UClass* InClass)
 
 void UInputDelegateBinding::BindInputDelegates(const UClass* InClass, UInputComponent* InputComponent, UObject* InObjectToBindTo /* = nullptr */)
 {
-	if (SupportsInputDelegate(InClass))
+	if (InClass && InputComponent && SupportsInputDelegate(InClass))
 	{
 		ensureMsgf(InputComponent, TEXT("Attempting to bind input delegates to an invalid Input Component!"));
 		
@@ -52,9 +54,9 @@ void UInputDelegateBinding::BindInputDelegatesWithSubojects(AActor* InActor, UIn
 	
 	ensureMsgf(InActor && InputComponent, TEXT("Attempting to bind input delegates to an invalid actor or input component!"));
 
-	const UClass* ActorClass = InActor->GetClass();
+	const UClass* ActorClass = InActor ? InActor->GetClass() : nullptr;
 
-	if (SupportsInputDelegate(ActorClass))
+	if (ActorClass && InputComponent && SupportsInputDelegate(ActorClass))
 	{
 		// Bind any input delegates on the base actor class
 		UInputDelegateBinding::BindInputDelegates(ActorClass, InputComponent, InputComponent->GetOwner());
