Support preshaders that evaluate to structs within uniform buffer (not just for constants)
Correctly handling for int, bool, double uniforms
Some support for LWC values in new HLSL translator

diff --git a/Engine/Shaders/Private/MaterialTemplate.ush b/Engine/Shaders/Private/MaterialTemplate.ush
--- a/Engine/Shaders/Private/MaterialTemplate.ush
+++ b/Engine/Shaders/Private/MaterialTemplate.ush
@@ -530,6 +530,26 @@ FPrimitiveSceneData GetPrimitiveData(FMaterialPixelParameters Parameters)
 	return GetPrimitiveData(Parameters.PrimitiveId);
 }
 
+bool UnpackUniform_bool(uint Packed, uint BitOffset)
+{
+	return (bool)((Packed >> BitOffset) & 0x1);
+}
+
+bool2 UnpackUniform_bool2(uint Packed, uint BitOffset)
+{
+	return bool2(UnpackUniform_bool(Packed, BitOffset), UnpackUniform_bool(Packed, BitOffset + 1));
+}
+
+bool3 UnpackUniform_bool3(uint Packed, uint BitOffset)
+{
+	return bool3(UnpackUniform_bool(Packed, BitOffset), UnpackUniform_bool(Packed, BitOffset + 1), UnpackUniform_bool(Packed, BitOffset + 2));
+}
+
+bool4 UnpackUniform_bool4(uint Packed, uint BitOffset)
+{
+	return bool4(UnpackUniform_bool(Packed, BitOffset), UnpackUniform_bool(Packed, BitOffset + 1), UnpackUniform_bool(Packed, BitOffset + 2), UnpackUniform_bool(Packed, BitOffset + 3));
+}
+
 /**
  * Returns the upper 3x3 portion of the LocalToWorld matrix.
  */
diff --git a/Engine/Source/Runtime/Engine/Classes/Materials/MaterialExpressionDoubleVectorParameter.h b/Engine/Source/Runtime/Engine/Classes/Materials/MaterialExpressionDoubleVectorParameter.h
--- a/Engine/Source/Runtime/Engine/Classes/Materials/MaterialExpressionDoubleVectorParameter.h
+++ b/Engine/Source/Runtime/Engine/Classes/Materials/MaterialExpressionDoubleVectorParameter.h
@@ -19,6 +19,7 @@ class UMaterialExpressionDoubleVectorParameter : public UMaterialExpressionParam
 
 	//~ Begin UMaterialExpression Interface
 #if WITH_EDITOR
+	EMaterialGenerateHLSLStatus GenerateHLSLExpression(FMaterialHLSLGenerator& Generator, UE::HLSLTree::FScope& Scope, int32 OutputIndex, UE::HLSLTree::FExpression*& OutExpression);
 	virtual int32 Compile(class FMaterialCompiler* Compiler, int32 OutputIndex) override;
 	virtual void GetCaption(TArray<FString>& OutCaptions) const override;
 	virtual bool GetParameterValue(FMaterialParameterMetadata& OutMeta) const override
diff --git a/Engine/Source/Runtime/Engine/Classes/Materials/MaterialExpressionWorldPosition.h b/Engine/Source/Runtime/Engine/Classes/Materials/MaterialExpressionWorldPosition.h
--- a/Engine/Source/Runtime/Engine/Classes/Materials/MaterialExpressionWorldPosition.h
+++ b/Engine/Source/Runtime/Engine/Classes/Materials/MaterialExpressionWorldPosition.h
@@ -37,6 +37,7 @@ class UMaterialExpressionWorldPosition : public UMaterialExpression
 	
 	//~ Begin UMaterialExpression Interface
 #if WITH_EDITOR
+	EMaterialGenerateHLSLStatus GenerateHLSLExpression(FMaterialHLSLGenerator& Generator, UE::HLSLTree::FScope& Scope, int32 OutputIndex, UE::HLSLTree::FExpression*& OutExpression);
 	virtual int32 Compile(class FMaterialCompiler* Compiler, int32 OutputIndex) override;
 	virtual void GetCaption(TArray<FString>& OutCaptions) const override;
 	virtual FText GetKeywords() const override {return FText::FromString(TEXT("position"));}
diff --git a/Engine/Source/Runtime/Engine/Private/HLSLTree/HLSLTree.cpp b/Engine/Source/Runtime/Engine/Private/HLSLTree/HLSLTree.cpp
--- a/Engine/Source/Runtime/Engine/Private/HLSLTree/HLSLTree.cpp
+++ b/Engine/Source/Runtime/Engine/Private/HLSLTree/HLSLTree.cpp
@@ -4,9 +4,6 @@
 #include "Misc/MemStackUtility.h"
 #include "MaterialShared.h" // TODO - split preshader out into its own module
 
-// From MaterialUniformExpressions.cpp
-extern void WriteMaterialUniformAccess(UE::Shader::EValueComponentType ComponentType, uint32 NumComponents, uint32 UniformOffset, FStringBuilderBase& OutResult);
-
 static FSHAHash HashString(const TCHAR* String, int32 Length)
 {
 	FSHAHash Hash;
@@ -244,10 +241,7 @@ FEmitShaderValue* FEmitContext::AcquireShader(FScope* Scope, const FShaderValue&
 	{
 		ShaderValue = *PrevShaderValue;
 		check(ShaderValue->Type == Shader.Type);
-		if (!Shader.bInline)
-		{
-			Private::MoveToScope(ShaderValue, Scope);
-		}
+		Private::MoveToScope(ShaderValue, Scope);
 	}
 	else
 	{
@@ -269,78 +263,269 @@ FEmitShaderValue* FEmitContext::AcquireShader(FScope* Scope, const FShaderValue&
 	return ShaderValue;
 }
 
-FEmitShaderValue* FEmitContext::AcquirePreshader(const FRequestedType& RequestedType, FScope* Scope, FExpression* Expression)
+namespace Private
+{
+void WriteMaterialUniformAccess(Shader::EValueComponentType ComponentType, uint32 NumComponents, uint32 UniformOffset, FStringBuilderBase& OutResult)
+{
+	static const TCHAR IndexToMask[] = TEXT("xyzw");
+	uint32 RegisterIndex = UniformOffset / 4;
+	uint32 RegisterOffset = UniformOffset % 4;
+	uint32 NumComponentsToWrite = NumComponents;
+	bool bConstructor = false;
+
+	check(ComponentType == Shader::EValueComponentType::Float || ComponentType == Shader::EValueComponentType::Int);
+	const bool bIsInt = (ComponentType == Shader::EValueComponentType::Int);
+
+	while (NumComponentsToWrite > 0u)
+	{
+		const uint32 NumComponentsInRegister = FMath::Min(NumComponentsToWrite, 4u - RegisterOffset);
+		if (NumComponentsInRegister < NumComponents && !bConstructor)
+		{
+			// Uniform will be split across multiple registers, so add the constructor to concat them together
+			OutResult.Appendf(TEXT("%s%d("), Shader::GetComponentTypeName(ComponentType), NumComponents);
+			bConstructor = true;
+		}
+
+		if (bIsInt)
+		{
+			// PreshaderBuffer is typed as float4, so reinterpret as 'int' if needed
+			OutResult.Append(TEXT("asint("));
+		}
+
+		OutResult.Appendf(TEXT("Material.PreshaderBuffer[%u]"), RegisterIndex);
+		// Can skip writing mask if we're taking all 4 components from the register
+		if (NumComponentsInRegister < 4u)
+		{
+			OutResult.Append(TCHAR('.'));
+			for (uint32 i = 0u; i < NumComponentsInRegister; ++i)
+			{
+				OutResult.Append(IndexToMask[RegisterOffset + i]);
+			}
+		}
+
+		if (bIsInt)
+		{
+			OutResult.Append(TEXT(")"));
+		}
+
+		NumComponentsToWrite -= NumComponentsInRegister;
+		RegisterIndex++;
+		RegisterOffset = 0u;
+		if (NumComponentsToWrite > 0u)
+		{
+			OutResult.Append(TEXT(", "));
+		}
+	}
+	if (bConstructor)
+	{
+		OutResult.Append(TEXT(")"));
+	}
+}
+} // namespace Private
+
+FEmitShaderValue* FEmitContext::AcquirePreshaderOrConstant(const FRequestedType& RequestedType, FScope* Scope, FExpression* Expression)
 {
 	Shader::FPreshaderData LocalPreshader;
 	Expression->EmitValuePreshader(*this, RequestedType, LocalPreshader);
 
 	const Shader::FType Type = RequestedType.GetType();
-	check(!Type.IsStruct()); // TODO
 
 	FSHA1 Hasher;
 	Hasher.Update((uint8*)&Type, sizeof(Type));
 	LocalPreshader.AppendHash(Hasher);
 	const FSHAHash Hash = Hasher.Finalize();
-	FEmitShaderValue* const* PrevPreshader = Preshaders.Find(Hash);
-	if (PrevPreshader)
+	FEmitShaderValue* const* PrevShaderValue = PreshaderValueMap.Find(Hash);
+	if (PrevShaderValue)
 	{
-		return *PrevPreshader;
+		FEmitShaderValue* ShaderValue = *PrevShaderValue;
+		check(ShaderValue->Type == Type);
+		Private::MoveToScope(ShaderValue, Scope);
+		return ShaderValue;
 	}
 
-	const Shader::FValueTypeDescription TypeDesc = Shader::GetValueTypeDescription(Type);
-	Shader::EValueComponentType ComponentType = TypeDesc.ComponentType;
-	if (ComponentType == Shader::EValueComponentType::Bool || ComponentType == Shader::EValueComponentType::Int)
+	Shader::FPreshaderStack Stack;
+	const Shader::FPreshaderValue ConstantValue = LocalPreshader.EvaluateConstant(*Material, Stack);
+
+	TStringBuilder<1024> FormattedCode;
+	if (Type.IsStruct())
 	{
-		ComponentType = Shader::EValueComponentType::Float; // TODO
+		FormattedCode.Append(TEXT("{ "));
 	}
 
-	//ensure(TypeDesc.ComponentType == Shader::EValueComponentType::Float || TypeDesc.ComponentType == Shader::EValueComponentType::Double);
+	FMaterialUniformPreshaderHeader* PreshaderHeader = nullptr;
+	uint32 CurrentBoolUniformOffset = ~0u;
+	uint32 CurrentNumBoolComponents = 32u;
 
-	FUniformExpressionSet& UniformExpressionSet = MaterialCompilationOutput->UniformExpressionSet;
-	FMaterialUniformPreshaderHeader& PreshaderHeader = UniformExpressionSet.UniformPreshaders.AddDefaulted_GetRef();
-
-	const uint32 RegisterOffset = UniformPreshaderOffset % 4;
-	if (ComponentType == Shader::EValueComponentType::Float && RegisterOffset + TypeDesc.NumComponents > 4u)
+	int32 ComponentIndex = 0;
+	for (int32 FieldIndex = 0; FieldIndex < Type.GetNumFlatFields(); ++FieldIndex)
 	{
-		// If this uniform would span multiple registers, align offset to the next register to avoid this
-		UniformPreshaderOffset = Align(UniformPreshaderOffset, 4u);
-	}
+		if (FieldIndex > 0)
+		{
+			FormattedCode.Append(TEXT(", "));
+		}
 
-	PreshaderHeader.OpcodeOffset = UniformExpressionSet.UniformPreshaderData.Num();
-	Expression->EmitValuePreshader(*this, RequestedType, UniformExpressionSet.UniformPreshaderData);
-	PreshaderHeader.OpcodeSize = UniformExpressionSet.UniformPreshaderData.Num() - PreshaderHeader.OpcodeOffset;
-	PreshaderHeader.BufferOffset = UniformPreshaderOffset;
-	PreshaderHeader.ComponentType = ComponentType;
-	PreshaderHeader.NumComponents = TypeDesc.NumComponents;
+		const Shader::EValueType FieldType = Type.GetFlatFieldType(FieldIndex);
+		const Shader::FValueTypeDescription TypeDesc = Shader::GetValueTypeDescription(FieldType);
+		const int32 NumFieldComponents = TypeDesc.NumComponents;
+		const EExpressionEvaluationType FieldEvaluationType = Expression->GetPreparedType().GetFieldEvaluationType(ComponentIndex, NumFieldComponents);
 
-	TStringBuilder<1024> FormattedCode;
-	if (ComponentType == Shader::EValueComponentType::Double)
-	{
-		if (TypeDesc.NumComponents == 1)
+		if (FieldEvaluationType == EExpressionEvaluationType::Preshader)
 		{
-			FormattedCode.Append(TEXT("MakeLWCScalar("));
+			// Only need to allocate uniform buffer for non-constant components
+			// Constant components can have their value inlined into the shader directly
+			FUniformExpressionSet& UniformExpressionSet = MaterialCompilationOutput->UniformExpressionSet;
+			if (!PreshaderHeader)
+			{
+				// Allocate a preshader header the first time we hit a non-constant field
+				PreshaderHeader = &UniformExpressionSet.UniformPreshaders.AddDefaulted_GetRef();
+				PreshaderHeader->FieldIndex = UniformExpressionSet.UniformPreshaderFields.Num();
+				PreshaderHeader->NumFields = 0u;
+				PreshaderHeader->OpcodeOffset = UniformExpressionSet.UniformPreshaderData.Num();
+				Expression->EmitValuePreshader(*this, RequestedType, UniformExpressionSet.UniformPreshaderData);
+				PreshaderHeader->OpcodeSize = UniformExpressionSet.UniformPreshaderData.Num() - PreshaderHeader->OpcodeOffset;
+			}
+
+			FMaterialUniformPreshaderField& PreshaderField = UniformExpressionSet.UniformPreshaderFields.AddDefaulted_GetRef();
+			PreshaderField.ComponentIndex = ComponentIndex;
+			PreshaderField.Type = FieldType;
+			PreshaderHeader->NumFields++;
+
+			if (TypeDesc.ComponentType == Shader::EValueComponentType::Bool)
+			{
+				// 'Bool' uniforms are packed into bits
+				if (CurrentNumBoolComponents + NumFieldComponents > 32u)
+				{
+					CurrentBoolUniformOffset = UniformPreshaderOffset++;
+					CurrentNumBoolComponents = 0u;
+				}
+
+				const uint32 RegisterIndex = CurrentBoolUniformOffset / 4;
+				const uint32 RegisterOffset = CurrentBoolUniformOffset % 4;
+				FormattedCode.Appendf(TEXT("UnpackUniform_%s(asuint(Material.PreshaderBuffer[%u][%u]), %u)"),
+					TypeDesc.Name,
+					RegisterIndex,
+					RegisterOffset,
+					CurrentNumBoolComponents);
+
+				PreshaderField.BufferOffset = CurrentBoolUniformOffset * 32u + CurrentNumBoolComponents;
+				CurrentNumBoolComponents += NumFieldComponents;
+			}
+			else if (TypeDesc.ComponentType == Shader::EValueComponentType::Double)
+			{
+				// Double uniforms are split into Tile/Offset components to make FLWCScalar/FLWCVectors
+				PreshaderField.BufferOffset = UniformPreshaderOffset;
+
+				if (NumFieldComponents > 1)
+				{
+					FormattedCode.Appendf(TEXT("MakeLWCVector%d("), NumFieldComponents);
+				}
+				else
+				{
+					FormattedCode.Append(TEXT("MakeLWCScalar("));
+				}
+
+				// Write the tile uniform
+				Private::WriteMaterialUniformAccess(Shader::EValueComponentType::Float, NumFieldComponents, UniformPreshaderOffset, FormattedCode);
+				UniformPreshaderOffset += NumFieldComponents;
+				FormattedCode.Append(TEXT(", "));
+
+				// Write the offset uniform
+				Private::WriteMaterialUniformAccess(Shader::EValueComponentType::Float, NumFieldComponents, UniformPreshaderOffset, FormattedCode);
+				UniformPreshaderOffset += NumFieldComponents;
+				FormattedCode.Append(TEXT(")"));
+			}
+			else
+			{
+				// Float/Int uniforms are written directly to the uniform buffer
+				const uint32 RegisterOffset = UniformPreshaderOffset % 4;
+				if (RegisterOffset + NumFieldComponents > 4u)
+				{
+					// If this uniform would span multiple registers, align offset to the next register to avoid this
+					// TODO - we could keep track of this empty padding space, and pack other smaller uniform types here
+					UniformPreshaderOffset = Align(UniformPreshaderOffset, 4u);
+				}
+
+				PreshaderField.BufferOffset = UniformPreshaderOffset;
+				Private::WriteMaterialUniformAccess(TypeDesc.ComponentType, NumFieldComponents, UniformPreshaderOffset, FormattedCode);
+				UniformPreshaderOffset += NumFieldComponents;
+			}
 		}
 		else
 		{
-			FormattedCode.Appendf(TEXT("MakeLWCVector%d("), TypeDesc.NumComponents);
-		}
+			// We allow FieldEvaluationType to be 'None', since in that case we still need to fill in a value for the HLSL initializer
+			check(FieldEvaluationType == EExpressionEvaluationType::Constant || FieldEvaluationType == EExpressionEvaluationType::None);
+
+			// The type generated by the preshader might not match the expected type
+			// In the future, with new HLSLTree, preshader could potentially include explicit cast opcodes, and avoid implicit conversions
+			Shader::FValue FieldConstantValue(ConstantValue.Type.GetComponentType(ComponentIndex), NumFieldComponents);
+			for (int32 i = 0; i < NumFieldComponents; ++i)
+			{
+				FieldConstantValue.Component[i] = ConstantValue.Component[ComponentIndex + i];
+			}
+
+			if (TypeDesc.ComponentType == Shader::EValueComponentType::Double)
+			{
+				const Shader::FDoubleValue DoubleValue = FieldConstantValue.AsDouble();
+				TStringBuilder<256> TileValue;
+				TStringBuilder<256> OffsetValue;
+				for (int32 Index = 0; Index < NumFieldComponents; ++Index)
+				{
+					if (Index > 0)
+					{
+						TileValue.Append(TEXT(", "));
+						OffsetValue.Append(TEXT(", "));
+					}
 
-		WriteMaterialUniformAccess(Shader::EValueComponentType::Float, TypeDesc.NumComponents, UniformPreshaderOffset, FormattedCode); // Tile
-		UniformPreshaderOffset += TypeDesc.NumComponents;
-		FormattedCode.Append(TEXT(","));
-		WriteMaterialUniformAccess(Shader::EValueComponentType::Float, TypeDesc.NumComponents, UniformPreshaderOffset, FormattedCode); // Offset
-		UniformPreshaderOffset += TypeDesc.NumComponents;
-		FormattedCode.Append(TEXT(")"));
+					const FLargeWorldRenderScalar Value(DoubleValue[Index]);
+					TileValue.Appendf(TEXT("%#.9gf"), Value.GetTile());
+					OffsetValue.Appendf(TEXT("%#.9gf"), Value.GetOffset());
+				}
+
+				if (NumFieldComponents > 1)
+				{
+					FormattedCode.Appendf(TEXT("MakeLWCVector%d(float%d(%s), float%d(%s))"), NumFieldComponents, NumFieldComponents, TileValue.ToString(), NumFieldComponents, OffsetValue.ToString());
+				}
+				else
+				{
+					FormattedCode.Appendf(TEXT("MakeLWCScalar(%s, %s)"), TileValue.ToString(), OffsetValue.ToString());
+				}
+			}
+			else
+			{
+				const Shader::FValue CastFieldConstantValue = Shader::Cast(FieldConstantValue, FieldType);
+				if (NumFieldComponents > 1)
+				{
+					FormattedCode.Appendf(TEXT("%s("), TypeDesc.Name);
+				}
+				for (int32 Index = 0; Index < NumFieldComponents; ++Index)
+				{
+					if (Index > 0)
+					{
+						FormattedCode.Append(TEXT(", "));
+					}
+					CastFieldConstantValue.Component[Index].ToString(TypeDesc.ComponentType, FormattedCode);
+				}
+				if (NumFieldComponents > 1)
+				{
+					FormattedCode.Append(TEXT(")"));
+				}
+			}
+		}
+		ComponentIndex += NumFieldComponents;
 	}
-	else
+	check(ComponentIndex == Type.GetNumComponents());
+
+	if (Type.IsStruct())
 	{
-		WriteMaterialUniformAccess(ComponentType, TypeDesc.NumComponents, UniformPreshaderOffset, FormattedCode);
-		UniformPreshaderOffset += TypeDesc.NumComponents;
+		FormattedCode.Append(TEXT(" }"));
 	}
 
-	FEmitShaderValue* ShaderValue = new(*Allocator) FEmitShaderValue(Scope, Shader::MakeValueType(ComponentType, TypeDesc.NumComponents));
-	ShaderValue->Reference = MemStack::AllocateString(*Allocator, FormattedCode);
-	Preshaders.Add(Hash, ShaderValue);
+	FShaderValue Shader(FormattedCode);
+	Shader.Type = Type;
+	Shader.bInline = !Type.IsStruct(); // struct declarations can't be inline, due to HLSL syntax
+	FEmitShaderValue* ShaderValue = AcquireShader(Scope, Shader, TArrayView<FEmitShaderValue*>());
+	PreshaderValueMap.Add(Hash, ShaderValue);
+
 	return ShaderValue;
 }
 
@@ -379,6 +564,7 @@ void FEmitContext::Finalize()
 	}
 
 	ShaderValueMap.Reset();
+	PreshaderValueMap.Reset();
 	LocalPHIs.Reset();
 
 	MaterialCompilationOutput->UniformExpressionSet.UniformPreshaderBufferSize = (UniformPreshaderOffset + 3u) / 4u;
@@ -745,6 +931,15 @@ int32 FPreparedType::GetNumComponents() const
 	return 0;
 }
 
+FRequestedType MakeRequestedType(Shader::EValueComponentType ComponentType, const FRequestedType& RequestedComponents)
+{
+	check(!RequestedComponents.IsStruct());
+	FRequestedType Result;
+	Result.ValueComponentType = ComponentType;
+	Result.RequestedComponents = RequestedComponents.RequestedComponents;
+	return Result;
+}
+
 bool FPreparedType::IsVoid() const
 {
 	return GetNumComponents() == 0;
@@ -797,6 +992,16 @@ EExpressionEvaluationType FPreparedType::GetEvaluationType(const FRequestedType&
 	return Result;
 }
 
+EExpressionEvaluationType FPreparedType::GetFieldEvaluationType(int32 ComponentIndex, int32 NumComponents) const
+{
+	EExpressionEvaluationType Result = EExpressionEvaluationType::None;
+	for (int32 Index = 0; Index < NumComponents; ++Index)
+	{
+		Result = CombineEvaluationTypes(Result, GetComponentEvaluationType(ComponentIndex + Index));
+	}
+	return Result;
+}
+
 void FPreparedType::SetEvaluationType(EExpressionEvaluationType EvaluationType)
 {
 	const int32 NumComponents = GetNumComponents();
@@ -879,10 +1084,13 @@ FPreparedType MergePreparedTypes(const FPreparedType& Lhs, const FPreparedType&
 
 bool FPrepareValueResult::TryMergePreparedType(FEmitContext& Context, const Shader::FStructType* StructType, Shader::EValueComponentType ComponentType)
 {
-	if (!PreparedType.IsInitialized())
+	// If we previously had a forwarded value set, reset that and start over
+	if (ForwardValue || !PreparedType.IsInitialized())
 	{
+		PreparedType.ComponentEvaluationType.Reset();
 		PreparedType.ValueComponentType = ComponentType;
 		PreparedType.StructType = StructType;
+		ForwardValue = nullptr;
 		return true;
 	}
 
@@ -989,17 +1197,17 @@ const TCHAR* FExpression::GetValueShader(FEmitContext& Context, const FRequested
 	check(IsScopeLive(CurrentScope));
 
 	FEmitShaderValue* ShaderValue = nullptr;
-	if (EvaluationType == EExpressionEvaluationType::Preshader)
+	if (EvaluationType == EExpressionEvaluationType::Constant || EvaluationType == EExpressionEvaluationType::Preshader)
 	{
-		ShaderValue = Context.AcquirePreshader(RequestedType, CurrentScope, this);
+		ShaderValue = Context.AcquirePreshaderOrConstant(RequestedType, CurrentScope, this);
 	}
 	else
 	{
+		check(EvaluationType == EExpressionEvaluationType::Shader);
+
 		TStringBuilder<1024> FormattedCode;
 		FShaderValue ShaderResult(FormattedCode);
 		FEmitShaderValueDependencies Dependencies;
-
-		if (EvaluationType == EExpressionEvaluationType::Shader)
 		{
 			FEmitShaderValueContext& ValueContext = Context.ShaderValueStack.AddDefaulted_GetRef();
 
@@ -1009,15 +1217,6 @@ const TCHAR* FExpression::GetValueShader(FEmitContext& Context, const FRequested
 			Dependencies = MoveTemp(ValueContext.Dependencies);
 			Context.ShaderValueStack.Pop(false);
 		}
-		else
-		{
-			check(EvaluationType == EExpressionEvaluationType::Constant);
-			const Shader::FValue ConstantValue = GetValueConstant(Context, RequestedType);
-			ShaderResult.Type = ConstantValue.Type;
-			ShaderResult.bInline = !ConstantValue.Type.IsStruct(); // 'struct' constants can't be inline, due to HLSL syntax limitations
-			ConstantValue.ToString(Shader::EValueStringFormat::HLSL, ShaderResult.Code);
-		}
-
 		ShaderValue = Context.AcquireShader(CurrentScope, ShaderResult, Dependencies);
 	}
 
@@ -1077,9 +1276,7 @@ Shader::FValue FExpression::GetValueConstant(FEmitContext& Context, const FReque
 
 	// Evaluate the constant preshader and store its value
 	Shader::FPreshaderStack Stack;
-	FMaterialRenderContext RenderContext(nullptr, *Context.Material, nullptr);
-
-	const Shader::FPreshaderValue PreshaderValue = ConstantPreshader.Evaluate(nullptr, RenderContext, Stack);
+	const Shader::FPreshaderValue PreshaderValue = ConstantPreshader.EvaluateConstant(*Context.Material, Stack);
 	Shader::FValue Result = PreshaderValue.AsShaderValue(Context.TypeRegistry);
 
 	const Shader::FType RequestedConstantType = RequestedType.GetType();
diff --git a/Engine/Source/Runtime/Engine/Private/HLSLTree/HLSLTreeCommon.cpp b/Engine/Source/Runtime/Engine/Private/HLSLTree/HLSLTreeCommon.cpp
--- a/Engine/Source/Runtime/Engine/Private/HLSLTree/HLSLTreeCommon.cpp
+++ b/Engine/Source/Runtime/Engine/Private/HLSLTree/HLSLTreeCommon.cpp
@@ -155,19 +155,38 @@ void FExpressionMaterialParameter::EmitValuePreshader(FEmitContext& Context, con
 
 void FExpressionExternalInput::EmitValueShader(FEmitContext& Context, const FRequestedType& RequestedType, FShaderValue& OutShader) const
 {
-	const int32 TypeIndex = (int32)InputType;
-	const int32 TexCoordIndex = TypeIndex - (int32)EExternalInputType::TexCoord0;
-
 	OutShader.Type = GetInputExpressionType(InputType);
-	if (TexCoordIndex >= 0 && TexCoordIndex < 8)
+	OutShader.bInline = true;
+
+	if (IsTexCoord(InputType))
 	{
+		const int32 TypeIndex = (int32)InputType;
+		const int32 TexCoordIndex = TypeIndex - (int32)EExternalInputType::TexCoord0;
+
 		Context.NumTexCoords = FMath::Max(Context.NumTexCoords, TexCoordIndex + 1);
-		OutShader.bInline = true;
 		OutShader.Code.Appendf(TEXT("Parameters.TexCoords[%u].xy"), TexCoordIndex);
 	}
 	else
 	{
-		return Context.Errors.AddError(this, TEXT("Invalid texcoord"));
+		// TODO - handle PrevFrame position
+		switch (InputType)
+		{
+		case EExternalInputType::WorldPosition:
+			OutShader.Code.Append(TEXT("GetWorldPosition(Parameters)"));
+			break;
+		case EExternalInputType::WorldPosition_NoOffsets:
+			OutShader.Code.Append(TEXT("GetWorldPosition_NoMaterialOffsets(Parameters)"));
+			break;
+		case EExternalInputType::TranslatedWorldPosition:
+			OutShader.Code.Append(TEXT("GetTranslatedWorldPosition(Parameters)"));
+			break;
+		case EExternalInputType::TranslatedWorldPosition_NoOffsets:
+			OutShader.Code.Append(TEXT("GetTranslatedWorldPosition_NoMaterialOffsets(Parameters)"));
+			break;
+		default:
+			checkNoEntry();
+			break;
+		}
 	}
 }
 
@@ -501,23 +520,27 @@ struct FBinaryOpTypes
 	Shader::FType ResultType;
 	FRequestedType LhsType;
 	FRequestedType RhsType;
+	bool bIsLWC;
 };
 FBinaryOpTypes GetBinaryOpTypes(const FRequestedType& RequestedType, EBinaryOp Op, const FRequestedType& LhsType, const FRequestedType& RhsType)
 {
 	const int32 NumRequestedComponents = RequestedType.GetNumComponents();
 	const Shader::EValueComponentType InputComponentType = Shader::CombineComponentTypes(LhsType.ValueComponentType, RhsType.ValueComponentType);
+	const FRequestedType InputType = MakeRequestedType(InputComponentType, RequestedType);
+
 	FBinaryOpTypes Result;
+	Result.bIsLWC = (InputComponentType == Shader::EValueComponentType::Double);
 	switch (Op)
 	{
 	case EBinaryOp::Less:
-		Result.LhsType = Shader::MakeValueType(InputComponentType, NumRequestedComponents);
-		Result.RhsType = Shader::MakeValueType(InputComponentType, NumRequestedComponents);
+		Result.LhsType = InputType;
+		Result.RhsType = InputType;
 		Result.ResultType = Shader::MakeValueType(Shader::EValueComponentType::Bool, NumRequestedComponents);
 		break;
 	default:
-		Result.LhsType = RequestedType;
-		Result.RhsType = RequestedType;
-		Result.ResultType = RequestedType.GetType();
+		Result.LhsType = InputType;
+		Result.RhsType = InputType;
+		Result.ResultType = Shader::MakeValueType(InputComponentType, NumRequestedComponents);
 		break;
 	}
 	return Result;
@@ -533,12 +556,59 @@ void FExpressionBinaryOp::EmitValueShader(FEmitContext& Context, const FRequeste
 	OutShader.Type = Types.ResultType;
 	switch (Op)
 	{
-	case EBinaryOp::Add: OutShader.Code.Appendf(TEXT("(%s + %s)"), LhsCode, RhsCode); break;
-	case EBinaryOp::Sub: OutShader.Code.Appendf(TEXT("(%s - %s)"), LhsCode, RhsCode); break;
-	case EBinaryOp::Mul: OutShader.Code.Appendf(TEXT("(%s * %s)"), LhsCode, RhsCode); break;
-	case EBinaryOp::Div: OutShader.Code.Appendf(TEXT("(%s / %s)"), LhsCode, RhsCode); break;
-	case EBinaryOp::Less: OutShader.Code.Appendf(TEXT("(%s < %s)"), LhsCode, RhsCode); break;
-	default: checkNoEntry(); break;
+	case EBinaryOp::Add:
+		if (Types.bIsLWC)
+		{
+			OutShader.Code.Appendf(TEXT("LWCAdd(%s, %s)"), LhsCode, RhsCode);
+		}
+		else
+		{
+			OutShader.Code.Appendf(TEXT("(%s + %s)"), LhsCode, RhsCode);
+		}
+		break;
+	case EBinaryOp::Sub:
+		if (Types.bIsLWC)
+		{
+			OutShader.Code.Appendf(TEXT("LWCSub(%s, %s)"), LhsCode, RhsCode);
+		}
+		else
+		{
+			OutShader.Code.Appendf(TEXT("(%s - %s)"), LhsCode, RhsCode);
+		}
+		break;
+	case EBinaryOp::Mul:
+		if (Types.bIsLWC)
+		{
+			OutShader.Code.Appendf(TEXT("LWCMultiply(%s, %s)"), LhsCode, RhsCode);
+		}
+		else
+		{
+			OutShader.Code.Appendf(TEXT("(%s * %s)"), LhsCode, RhsCode);
+		}
+		break;
+	case EBinaryOp::Div:
+		if (Types.bIsLWC)
+		{
+			OutShader.Code.Appendf(TEXT("LWCDivide(%s, %s)"), LhsCode, RhsCode);
+		}
+		else
+		{
+			OutShader.Code.Appendf(TEXT("(%s / %s)"), LhsCode, RhsCode);
+		}
+		break;
+	case EBinaryOp::Less:
+		if (Types.bIsLWC)
+		{
+			OutShader.Code.Appendf(TEXT("LWCLess(%s, %s)"), LhsCode, RhsCode);
+		}
+		else
+		{
+			OutShader.Code.Appendf(TEXT("(%s < %s)"), LhsCode, RhsCode);
+		}
+		break;
+	default:
+		checkNoEntry();
+		break;
 	}
 }
 
diff --git a/Engine/Source/Runtime/Engine/Private/Materials/HLSLMaterialTranslator.cpp b/Engine/Source/Runtime/Engine/Private/Materials/HLSLMaterialTranslator.cpp
--- a/Engine/Source/Runtime/Engine/Private/Materials/HLSLMaterialTranslator.cpp
+++ b/Engine/Source/Runtime/Engine/Private/Materials/HLSLMaterialTranslator.cpp
@@ -3146,13 +3146,19 @@ int32 FHLSLMaterialTranslator::AccessUniformExpression(int32 Index)
 				UniformPreshaderOffset = Align(UniformPreshaderOffset, 4u);
 			}
 
+			const UE::Shader::EValueComponentType ComponentType = bIsLWC ? UE::Shader::EValueComponentType::Double : UE::Shader::EValueComponentType::Float;
+
 			FMaterialUniformPreshaderHeader& Preshader = MaterialCompilationOutput.UniformExpressionSet.UniformPreshaders.AddDefaulted_GetRef();
 			Preshader.OpcodeOffset = MaterialCompilationOutput.UniformExpressionSet.UniformPreshaderData.Num();
 			CodeChunk.UniformExpression->WriteNumberOpcodes(MaterialCompilationOutput.UniformExpressionSet.UniformPreshaderData);
-			Preshader.BufferOffset = UniformPreshaderOffset;
 			Preshader.OpcodeSize = MaterialCompilationOutput.UniformExpressionSet.UniformPreshaderData.Num() - Preshader.OpcodeOffset;
-			Preshader.ComponentType = bIsLWC ? UE::Shader::EValueComponentType::Double : UE::Shader::EValueComponentType::Float;
-			Preshader.NumComponents = NumComponents;
+			Preshader.FieldIndex = MaterialCompilationOutput.UniformExpressionSet.UniformPreshaderFields.Num();
+			Preshader.NumFields = 1;
+
+			FMaterialUniformPreshaderField& PreshaderField = MaterialCompilationOutput.UniformExpressionSet.UniformPreshaderFields.AddDefaulted_GetRef();
+			PreshaderField.BufferOffset = UniformPreshaderOffset;
+			PreshaderField.Type = UE::Shader::MakeValueType(ComponentType, NumComponents);
+			PreshaderField.ComponentIndex = 0;
 
 			CodeChunk.UniformExpression->UniformOffset = UniformPreshaderOffset;
 			UniformPreshaderOffset += bIsLWC ? NumComponents * 2u : NumComponents;
diff --git a/Engine/Source/Runtime/Engine/Private/Materials/MaterialExpressionHLSL.cpp b/Engine/Source/Runtime/Engine/Private/Materials/MaterialExpressionHLSL.cpp
--- a/Engine/Source/Runtime/Engine/Private/Materials/MaterialExpressionHLSL.cpp
+++ b/Engine/Source/Runtime/Engine/Private/Materials/MaterialExpressionHLSL.cpp
@@ -17,8 +17,10 @@
 #include "Materials/MaterialExpressionConstant4Vector.h"
 #include "Materials/MaterialExpressionStaticBool.h"
 #include "Materials/MaterialExpressionVectorParameter.h"
+#include "Materials/MaterialExpressionDoubleVectorParameter.h"
 #include "Materials/MaterialExpressionScalarParameter.h"
 #include "Materials/MaterialExpressionStaticBoolParameter.h"
+#include "Materials/MaterialExpressionWorldPosition.h"
 #include "Materials/MaterialExpressionTextureCoordinate.h"
 #include "Materials/MaterialExpressionTextureSample.h"
 #include "Materials/MaterialExpressionTextureSampleParameter.h"
@@ -119,6 +121,12 @@ EMaterialGenerateHLSLStatus UMaterialExpressionVectorParameter::GenerateHLSLExpr
 	return EMaterialGenerateHLSLStatus::Success;
 }
 
+EMaterialGenerateHLSLStatus UMaterialExpressionDoubleVectorParameter::GenerateHLSLExpression(FMaterialHLSLGenerator& Generator, UE::HLSLTree::FScope& Scope, int32 OutputIndex, UE::HLSLTree::FExpression*& OutExpression)
+{
+	OutExpression = Generator.GetTree().NewExpression<UE::HLSLTree::FExpressionMaterialParameter>(EMaterialParameterType::DoubleVector, ParameterName, DefaultValue);
+	return EMaterialGenerateHLSLStatus::Success;
+}
+
 EMaterialGenerateHLSLStatus UMaterialExpressionScalarParameter::GenerateHLSLExpression(FMaterialHLSLGenerator& Generator, UE::HLSLTree::FScope& Scope, int32 OutputIndex, UE::HLSLTree::FExpression*& OutExpression)
 {
 	OutExpression = Generator.GetTree().NewExpression<UE::HLSLTree::FExpressionMaterialParameter>(EMaterialParameterType::Scalar, ParameterName, DefaultValue);
@@ -130,6 +138,22 @@ EMaterialGenerateHLSLStatus UMaterialExpressionStaticBoolParameter::GenerateHLSL
 	OutExpression = Generator.GetTree().NewExpression<UE::HLSLTree::FExpressionMaterialParameter>(EMaterialParameterType::StaticSwitch, ParameterName, (bool)DefaultValue);
 	return EMaterialGenerateHLSLStatus::Success;
 }
+EMaterialGenerateHLSLStatus UMaterialExpressionWorldPosition::GenerateHLSLExpression(FMaterialHLSLGenerator& Generator, UE::HLSLTree::FScope& Scope, int32 OutputIndex, UE::HLSLTree::FExpression*& OutExpression)
+{
+	UE::HLSLTree::EExternalInputType InputType = UE::HLSLTree::EExternalInputType::None;
+
+	switch (WorldPositionShaderOffset)
+	{
+	case WPT_Default: InputType = UE::HLSLTree::EExternalInputType::WorldPosition; break;
+	case WPT_ExcludeAllShaderOffsets: InputType = UE::HLSLTree::EExternalInputType::WorldPosition_NoOffsets; break;
+	case WPT_CameraRelative: InputType = UE::HLSLTree::EExternalInputType::TranslatedWorldPosition; break;
+	case WPT_CameraRelativeNoOffsets: InputType = UE::HLSLTree::EExternalInputType::TranslatedWorldPosition_NoOffsets; break;
+	default: checkNoEntry(); break;
+	}
+
+	OutExpression = Generator.GetTree().NewExpression<UE::HLSLTree::FExpressionExternalInput>(InputType);
+	return EMaterialGenerateHLSLStatus::Success;
+}
 
 EMaterialGenerateHLSLStatus UMaterialExpressionTextureCoordinate::GenerateHLSLExpression(FMaterialHLSLGenerator& Generator, UE::HLSLTree::FScope& Scope, int32 OutputIndex, UE::HLSLTree::FExpression*& OutExpression)
 {
diff --git a/Engine/Source/Runtime/Engine/Private/Materials/MaterialShared.cpp b/Engine/Source/Runtime/Engine/Private/Materials/MaterialShared.cpp
--- a/Engine/Source/Runtime/Engine/Private/Materials/MaterialShared.cpp
+++ b/Engine/Source/Runtime/Engine/Private/Materials/MaterialShared.cpp
@@ -69,6 +69,7 @@ IMPLEMENT_TYPE_LAYOUT(FMeshMaterialShaderMap);
 IMPLEMENT_TYPE_LAYOUT(FMaterialProcessedSource);
 IMPLEMENT_TYPE_LAYOUT(FMaterialShaderMapContent);
 IMPLEMENT_TYPE_LAYOUT(FMaterialUniformPreshaderHeader);
+IMPLEMENT_TYPE_LAYOUT(FMaterialUniformPreshaderField);
 IMPLEMENT_TYPE_LAYOUT(FMaterialNumericParameterInfo);
 IMPLEMENT_TYPE_LAYOUT(FMaterialTextureParameterInfo);
 IMPLEMENT_TYPE_LAYOUT(FMaterialExternalTextureParameterInfo);
diff --git a/Engine/Source/Runtime/Engine/Private/Materials/MaterialUniformExpressions.cpp b/Engine/Source/Runtime/Engine/Private/Materials/MaterialUniformExpressions.cpp
--- a/Engine/Source/Runtime/Engine/Private/Materials/MaterialUniformExpressions.cpp
+++ b/Engine/Source/Runtime/Engine/Private/Materials/MaterialUniformExpressions.cpp
@@ -716,7 +716,7 @@ void FUniformExpressionSet::FillUniformBuffer(const FMaterialRenderContext& Mate
 		}
 
 		// Dump preshader results into buffer.
-		float* PreshaderBuffer = (float*)BufferCursor;
+		float* const PreshaderBuffer = (float*)BufferCursor;
 		FPreshaderStack PreshaderStack;
 		FPreshaderDataContext PreshaderBaseContext(UniformPreshaderData);
 		for (const FMaterialUniformPreshaderHeader& Preshader : UniformPreshaders)
@@ -724,32 +724,79 @@ void FUniformExpressionSet::FillUniformBuffer(const FMaterialRenderContext& Mate
 			FPreshaderDataContext PreshaderContext(PreshaderBaseContext, Preshader.OpcodeOffset, Preshader.OpcodeSize);
 			const FPreshaderValue Result = EvaluatePreshader(this, MaterialRenderContext, PreshaderStack, PreshaderContext);
 
-			float* DestAddress = PreshaderBuffer + Preshader.BufferOffset;
-			if(Preshader.ComponentType == EValueComponentType::Float)
+			for (uint32 FieldIndex = 0u; FieldIndex < Preshader.NumFields; ++FieldIndex)
 			{
-				const FFloatValue FloatValue = Result.AsShaderValue().AsFloat();
-				for (uint32 i = 0u; i < Preshader.NumComponents; ++i)
+				const FMaterialUniformPreshaderField& PreshaderField = UniformPreshaderFields[Preshader.FieldIndex + FieldIndex];
+				const FValueTypeDescription UniformTypeDesc = GetValueTypeDescription(PreshaderField.Type);
+				const int32 NumFieldComponents = UniformTypeDesc.NumComponents;
+
+				// The type generated by the preshader might not match the expected type
+				// In the future, with new HLSLTree, preshader could potentially include explicit cast opcodes, and avoid implicit conversions
+				// Making this change is difficult while also maintaining backwards compatibility however
+				FValue FieldValue(Result.Type.GetComponentType(PreshaderField.ComponentIndex), NumFieldComponents);
+				for (int32 i = 0; i < NumFieldComponents; ++i)
 				{
-					*DestAddress++ = FloatValue[i];
+					FieldValue.Component[i] = Result.Component[PreshaderField.ComponentIndex + i];
 				}
-			}
-			else if (Preshader.ComponentType == EValueComponentType::Double)
-			{
-				const FDoubleValue DoubleValue = Result.AsShaderValue().AsDouble();
-				float TileValue[4];
-				float OffsetValue[4];
-				for (uint32 i = 0u; i < Preshader.NumComponents; ++i)
+
+				if (UniformTypeDesc.ComponentType == EValueComponentType::Float)
 				{
-					const FLargeWorldRenderScalar Value(DoubleValue[i]);
-					TileValue[i] = Value.GetTile();
-					OffsetValue[i] = Value.GetOffset();
+					const FFloatValue FloatValue = FieldValue.AsFloat();
+					float* DestAddress = PreshaderBuffer + PreshaderField.BufferOffset;
+					for (int32 i = 0; i < NumFieldComponents; ++i)
+					{
+						*DestAddress++ = FloatValue[i];
+					}
+				}
+				else if (UniformTypeDesc.ComponentType == EValueComponentType::Bool)
+				{
+					const FBoolValue BoolValue = FieldValue.AsBool();
+					uint32 Mask = 0u;
+					for (int32 i = 0; i < NumFieldComponents; ++i)
+					{
+						if (BoolValue[i])
+						{
+							Mask |= (1u << i);
+						}
+					}
+
+					const uint32 BufferOffset = PreshaderField.BufferOffset / 32u;
+					const uint32 BufferBitOffset = PreshaderField.BufferOffset % 32u;
+					check(BufferBitOffset + NumFieldComponents <= 32u);
+
+					uint32* DestAddress = (uint32*)PreshaderBuffer + BufferOffset;
+					if (BufferBitOffset == 0u)
+					{
+						// First update to a group of bits needs to initialize memory
+						*DestAddress = Mask;
+					}
+					else
+					{
+						// Combine with any previous bits
+						*DestAddress |= (Mask << BufferBitOffset);
+					}
+				}
+				else if (UniformTypeDesc.ComponentType == EValueComponentType::Double)
+				{
+					const FDoubleValue DoubleValue = FieldValue.AsDouble();
+
+					float TileValue[4];
+					float OffsetValue[4];
+					for (int32 i = 0; i < NumFieldComponents; ++i)
+					{
+						const FLargeWorldRenderScalar Value(DoubleValue[i]);
+						TileValue[i] = Value.GetTile();
+						OffsetValue[i] = Value.GetOffset();
+					}
+
+					float* DestAddress = PreshaderBuffer + PreshaderField.BufferOffset;
+					for (int32 i = 0; i < NumFieldComponents; ++i) *DestAddress++ = TileValue[i];
+					for (int32 i = 0; i < NumFieldComponents; ++i) *DestAddress++ = OffsetValue[i];
+				}
+				else
+				{
+					ensure(false);
 				}
-				for (uint32 i = 0u; i < Preshader.NumComponents; ++i) *DestAddress++ = TileValue[i];
-				for (uint32 i = 0u; i < Preshader.NumComponents; ++i) *DestAddress++ = OffsetValue[i];
-			}
-			else
-			{
-				ensure(false);
 			}
 		}
 
diff --git a/Engine/Source/Runtime/Engine/Private/Shader/Preshader.cpp b/Engine/Source/Runtime/Engine/Private/Shader/Preshader.cpp
--- a/Engine/Source/Runtime/Engine/Private/Shader/Preshader.cpp
+++ b/Engine/Source/Runtime/Engine/Private/Shader/Preshader.cpp
@@ -24,6 +24,10 @@ FPreshaderType::FPreshaderType(const FType& InType) : ValueType(InType.ValueType
 	}
 }
 
+FPreshaderType::FPreshaderType(EValueType InType) : ValueType(InType)
+{
+}
+
 EValueComponentType FPreshaderType::GetComponentType(int32 Index) const
 {
 	if (IsStruct())
@@ -292,9 +296,16 @@ static void EvaluateGetField(FPreshaderStack& Stack, FPreshaderDataContext& REST
 	Stack.PushValue(FieldType, FieldComponents);
 }
 
-static void EvaluateParameter(FPreshaderStack& Stack, const FUniformExpressionSet& UniformExpressionSet, uint32 ParameterIndex, const FMaterialRenderContext& Context)
+static void EvaluateParameter(FPreshaderStack& Stack, const FUniformExpressionSet* UniformExpressionSet, uint32 ParameterIndex, const FMaterialRenderContext& Context)
 {
-	const FMaterialNumericParameterInfo& Parameter = UniformExpressionSet.GetNumericParameter(ParameterIndex);
+	if (!UniformExpressionSet)
+	{
+		// return 0 for parameters if we don't have UniformExpressionSet
+		Stack.PushEmptyValue(EValueType::Float1);
+		return;
+	}
+
+	const FMaterialNumericParameterInfo& Parameter = UniformExpressionSet->GetNumericParameter(ParameterIndex);
 	bool bFoundParameter = false;
 
 	// First allow proxy the chance to override parameter
@@ -324,7 +335,7 @@ static void EvaluateParameter(FPreshaderStack& Stack, const FUniformExpressionSe
 	// Default value
 	if (!bFoundParameter)
 	{
-		Stack.PushValue(UniformExpressionSet.GetDefaultParameterValue(Parameter.ParameterType, Parameter.DefaultValueOffset));
+		Stack.PushValue(UniformExpressionSet->GetDefaultParameterValue(Parameter.ParameterType, Parameter.DefaultValueOffset));
 	}
 }
 
@@ -365,16 +376,16 @@ static void EvaluateComponentSwizzle(FPreshaderStack& Stack, FPreshaderDataConte
 	switch (NumElements)
 	{
 	case 4:
-		Result.Component[3] = Value.GetComponent(IndexA);
+		Result.Component[3] = Value.TryGetComponent(IndexA);
 		// Fallthrough...
 	case 3:
-		Result.Component[2] = Value.GetComponent(IndexB);
+		Result.Component[2] = Value.TryGetComponent(IndexB);
 		// Fallthrough...
 	case 2:
-		Result.Component[1] = Value.GetComponent(IndexG);
+		Result.Component[1] = Value.TryGetComponent(IndexG);
 		// Fallthrough...
 	case 1:
-		Result.Component[0] = Value.GetComponent(IndexR);
+		Result.Component[0] = Value.TryGetComponent(IndexR);
 		break;
 	default:
 		UE_LOG(LogMaterial, Fatal, TEXT("Invalid number of swizzle elements: %d"), NumElements);
@@ -491,8 +502,7 @@ FPreshaderValue EvaluatePreshader(const FUniformExpressionSet* UniformExpression
 		case EPreshaderOpcode::GetField: EvaluateGetField(Stack, Data); break;
 		case EPreshaderOpcode::SetField: EvaluateSetField(Stack, Data); break;
 		case EPreshaderOpcode::Parameter:
-			check(UniformExpressionSet);
-			EvaluateParameter(Stack, *UniformExpressionSet, ReadPreshaderValue<uint16>(Data), Context);
+			EvaluateParameter(Stack, UniformExpressionSet, ReadPreshaderValue<uint16>(Data), Context);
 			break;
 		case EPreshaderOpcode::Add: EvaluateBinaryOp(Stack, Add); break;
 		case EPreshaderOpcode::Sub: EvaluateBinaryOp(Stack, Sub); break;
@@ -549,15 +559,23 @@ FPreshaderValue EvaluatePreshader(const FUniformExpressionSet* UniformExpression
 	return Result;
 }
 
-FPreshaderValue FPreshaderData::Evaluate(FUniformExpressionSet* UniformExpressionSet, const struct FMaterialRenderContext& Context, FPreshaderStack& Stack)
+FPreshaderValue FPreshaderData::Evaluate(FUniformExpressionSet* UniformExpressionSet, const struct FMaterialRenderContext& Context, FPreshaderStack& Stack) const
 {
 	FPreshaderDataContext PreshaderContext(*this);
 	return EvaluatePreshader(UniformExpressionSet, Context, Stack, PreshaderContext);
 }
 
+FPreshaderValue FPreshaderData::EvaluateConstant(const FMaterial& Material, FPreshaderStack& Stack) const
+{
+	FPreshaderDataContext PreshaderContext(*this);
+	return EvaluatePreshader(nullptr, FMaterialRenderContext(nullptr, Material, nullptr), Stack, PreshaderContext);
+}
+
 void FPreshaderData::AppendHash(FSHA1& OutHasher) const
 {
 	OutHasher.Update((uint8*)Names.GetData(), Names.Num() * Names.GetTypeSize());
+	OutHasher.Update((uint8*)StructTypes.GetData(), StructTypes.Num() * StructTypes.GetTypeSize());
+	OutHasher.Update((uint8*)StructComponentTypes.GetData(), StructComponentTypes.Num() * StructComponentTypes.GetTypeSize());
 	OutHasher.Update((uint8*)Data.GetData(), Data.Num() * Data.GetTypeSize());
 }
 
diff --git a/Engine/Source/Runtime/Engine/Private/Shader/ShaderValue.cpp b/Engine/Source/Runtime/Engine/Private/Shader/ShaderValue.cpp
--- a/Engine/Source/Runtime/Engine/Private/Shader/ShaderValue.cpp
+++ b/Engine/Source/Runtime/Engine/Private/Shader/ShaderValue.cpp
@@ -36,6 +36,18 @@ int32 FType::GetNumComponents() const
 	return TypeDesc.NumComponents;
 }
 
+int32 FType::GetNumFlatFields() const
+{
+	if (StructType)
+	{
+		check(ValueType == EValueType::Struct);
+		return StructType->FlatFieldTypes.Num();
+	}
+
+	check(ValueType != EValueType::Struct);
+	return 1;
+}
+
 EValueComponentType FType::GetComponentType(int32 Index) const
 {
 	if (StructType)
@@ -47,6 +59,16 @@ EValueComponentType FType::GetComponentType(int32 Index) const
 	return (Index >= 0 && Index < TypeDesc.NumComponents) ? TypeDesc.ComponentType : EValueComponentType::Void;
 }
 
+EValueType FType::GetFlatFieldType(int32 Index) const
+{
+	if (StructType)
+	{
+		check(ValueType == EValueType::Struct);
+		return StructType->FlatFieldTypes.IsValidIndex(Index) ? StructType->FlatFieldTypes[Index] : EValueType::Void;
+	}
+	return (Index == 0) ? ValueType : EValueType::Void;
+}
+
 bool FType::Merge(const FType& OtherType)
 {
 	if (ValueType == EValueType::Void)
@@ -323,6 +345,19 @@ bool FValue::AsBoolScalar() const
 	return false;
 }
 
+const TCHAR* GetComponentTypeName(EValueComponentType Type)
+{
+	switch (Type)
+	{
+	case EValueComponentType::Void: return TEXT("void");
+	case EValueComponentType::Float: return TEXT("float");
+	case EValueComponentType::Double: return TEXT("double");
+	case EValueComponentType::Int: return TEXT("int");
+	case EValueComponentType::Bool: return TEXT("bool");
+	default: checkNoEntry() return TEXT("void");
+	}
+}
+
 uint32 GetComponentTypeSizeInBytes(EValueComponentType Type)
 {
 	switch (Type)
@@ -697,18 +732,19 @@ void FStructTypeRegistry::EmitDeclarationsCode(FStringBuilderBase& OutCode) cons
 
 namespace Private
 {
-void SetComponentTypes(EValueComponentType* ComponentTypes, int32 ComponentIndex, const FType& Type)
+void SetFieldType(EValueType* FieldTypes, EValueComponentType* ComponentTypes, int32 FieldIndex, int32 ComponentIndex, const FType& Type)
 {
 	if (Type.IsStruct())
 	{
 		for (const FStructField& Field : Type.StructType->Fields)
 		{
-			SetComponentTypes(ComponentTypes, ComponentIndex + Field.ComponentIndex, Field.Type);
+			SetFieldType(FieldTypes, ComponentTypes, FieldIndex + Field.FlatFieldIndex, ComponentIndex + Field.ComponentIndex, Field.Type);
 		}
 	}
 	else
 	{
-		const FValueTypeDescription TypeDesc = GetValueTypeDescription(Type);
+		FieldTypes[FieldIndex] = Type.ValueType;
+		const FValueTypeDescription TypeDesc = GetValueTypeDescription(Type.ValueType);
 		for (int32 i = 0; i < TypeDesc.NumComponents; ++i)
 		{
 			ComponentTypes[ComponentIndex + i] = TypeDesc.ComponentType;
@@ -724,6 +760,7 @@ const FStructType* FStructTypeRegistry::NewType(const FStructTypeInitializer& In
 
 	const int32 NumFields = Initializer.Fields.Num();
 	int32 ComponentIndex = 0;
+	int32 FlatFieldIndex = 0;
 	FStructField* Fields = new(*Allocator) FStructField[NumFields];
 	for (int32 FieldIndex = 0; FieldIndex < NumFields; ++FieldIndex)
 	{
@@ -744,14 +781,17 @@ const FStructType* FStructTypeRegistry::NewType(const FStructTypeInitializer& In
 		Field.Name = MemStack::AllocateString(*Allocator, FieldInitializer.Name);
 		Field.Type = FieldType;
 		Field.ComponentIndex = ComponentIndex;
+		Field.FlatFieldIndex = FlatFieldIndex;
 		ComponentIndex += FieldType.GetNumComponents();
+		FlatFieldIndex += FieldType.GetNumFlatFields();
 	}
 
 	EValueComponentType* ComponentTypes = new(*Allocator) EValueComponentType[ComponentIndex];
+	EValueType* FlatFieldTypes = new(*Allocator) EValueType[FlatFieldIndex];
 	for (int32 FieldIndex = 0; FieldIndex < NumFields; ++FieldIndex)
 	{
 		const FStructField& Field = Fields[FieldIndex];
-		Private::SetComponentTypes(ComponentTypes, Field.ComponentIndex, Field.Type);
+		Private::SetFieldType(FlatFieldTypes, ComponentTypes, Field.FlatFieldIndex, Field.ComponentIndex, Field.Type);
 	}
 
 	const FSHAHash Hash = Hasher.Finalize();
@@ -761,6 +801,7 @@ const FStructType* FStructTypeRegistry::NewType(const FStructTypeInitializer& In
 	StructType->Hash = *reinterpret_cast<const uint64*>(Hash.Hash);
 	StructType->Fields = MakeArrayView(Fields, NumFields);
 	StructType->ComponentTypes = MakeArrayView(ComponentTypes, ComponentIndex);
+	StructType->FlatFieldTypes = MakeArrayView(FlatFieldTypes, FlatFieldIndex);
 
 	Types.Add(StructType->Hash, StructType);
 
diff --git a/Engine/Source/Runtime/Engine/Private/ShaderDerivedDataVersion.h b/Engine/Source/Runtime/Engine/Private/ShaderDerivedDataVersion.h
--- a/Engine/Source/Runtime/Engine/Private/ShaderDerivedDataVersion.h
+++ b/Engine/Source/Runtime/Engine/Private/ShaderDerivedDataVersion.h
@@ -11,4 +11,4 @@
 // In case of merge conflicts with DDC versions, you MUST generate a new GUID and set this new
 // guid as version
 #define GLOBALSHADERMAP_DERIVEDDATA_VER			TEXT("7429FF5E4C57472FA538321D6777E8BC")
-#define MATERIALSHADERMAP_DERIVEDDATA_VER		TEXT("C2D6564DC8B14777B5D6747464213D7E")
\ No newline at end of file
+#define MATERIALSHADERMAP_DERIVEDDATA_VER		TEXT("DB3996BDDA7D410FA8F3F4B74B4A4404")
\ No newline at end of file
diff --git a/Engine/Source/Runtime/Engine/Public/HLSLTree/HLSLTree.h b/Engine/Source/Runtime/Engine/Public/HLSLTree/HLSLTree.h
--- a/Engine/Source/Runtime/Engine/Public/HLSLTree/HLSLTree.h
+++ b/Engine/Source/Runtime/Engine/Public/HLSLTree/HLSLTree.h
@@ -102,13 +102,14 @@ public:
 	const TCHAR* AcquireLocalDeclarationCode();
 
 	FEmitShaderValue* AcquireShader(FScope* Scope, const FShaderValue& Shader, TArrayView<FEmitShaderValue*> Dependencies);
-	FEmitShaderValue* AcquirePreshader(const FRequestedType& RequestedType, FScope* Scope, FExpression* Expression);
+	FEmitShaderValue* AcquirePreshaderOrConstant(const FRequestedType& RequestedType, FScope* Scope, FExpression* Expression);
 
 	FEmitShaderValue* CastShaderValue(FNode* Node, FScope* Scope, FEmitShaderValue* ShaderValue, const Shader::FType& DestType);
 
 	FMemStackBase* Allocator = nullptr;
 	const Shader::FStructTypeRegistry* TypeRegistry = nullptr;
 	TMap<FSHAHash, FEmitShaderValue*> ShaderValueMap;
+	TMap<FSHAHash, FEmitShaderValue*> PreshaderValueMap;
 	TArray<const FExpressionLocalPHI*> LocalPHIs;
 	FErrors Errors;
 
@@ -117,7 +118,6 @@ public:
 	const FStaticParameterSet* StaticParameters = nullptr;
 	FMaterialCompilationOutput* MaterialCompilationOutput = nullptr;
 	TMap<Shader::FValue, uint32> DefaultUniformValues;
-	TMap<FSHAHash, FEmitShaderValue*> Preshaders;
 	TArray<FScope*, TInlineAllocator<16>> ScopeStack;
 	TArray<FEmitShaderValueContext, TInlineAllocator<16>> ShaderValueStack;
 	uint32 UniformPreshaderOffset = 0u;
@@ -230,6 +230,8 @@ public:
 	TBitArray<> RequestedComponents;
 };
 
+FRequestedType MakeRequestedType(Shader::EValueComponentType ComponentType, const FRequestedType& RequestedComponents);
+
 /**
  * Like FRequestedType, but tracks an EExpressionEvaluationType per component, rather than a simple requested flag
  */
@@ -254,6 +256,7 @@ public:
 	bool IsInitialized() const { return StructType != nullptr || ValueComponentType != Shader::EValueComponentType::Void; }
 	bool IsVoid() const;
 	EExpressionEvaluationType GetEvaluationType(const FRequestedType& RequestedType) const;
+	EExpressionEvaluationType GetFieldEvaluationType(int32 ComponentIndex, int32 NumComponents) const;
 
 	EExpressionEvaluationType GetComponentEvaluationType(int32 Index) const
 	{
@@ -323,6 +326,7 @@ struct FEmitShaderValue
 class FExpression : public FNode
 {
 public:
+	const FPreparedType& GetPreparedType() const { return PrepareValueResult.PreparedType; }
 	FRequestedType GetRequestedType() const { return PrepareValueResult.PreparedType.GetRequestedType(); }
 	Shader::FType GetType() const { return PrepareValueResult.PreparedType.GetType(); }
 	EExpressionEvaluationType GetEvaluationType(const FRequestedType& RequestedType) const { return PrepareValueResult.PreparedType.GetEvaluationType(RequestedType); }
diff --git a/Engine/Source/Runtime/Engine/Public/HLSLTree/HLSLTreeCommon.h b/Engine/Source/Runtime/Engine/Public/HLSLTree/HLSLTreeCommon.h
--- a/Engine/Source/Runtime/Engine/Public/HLSLTree/HLSLTreeCommon.h
+++ b/Engine/Source/Runtime/Engine/Public/HLSLTree/HLSLTreeCommon.h
@@ -66,8 +66,9 @@ public:
 	virtual void EmitValuePreshader(FEmitContext& Context, const FRequestedType& RequestedType, Shader::FPreshaderData& OutPreshader) const override;
 };
 
-enum class EExternalInputType
+enum class EExternalInputType : uint8
 {
+	None,
 	TexCoord0,
 	TexCoord1,
 	TexCoord2,
@@ -76,14 +77,43 @@ enum class EExternalInputType
 	TexCoord5,
 	TexCoord6,
 	TexCoord7,
+	WorldPosition,
+	WorldPosition_NoOffsets,
+	TranslatedWorldPosition,
+	TranslatedWorldPosition_NoOffsets,
 };
+static constexpr int32 NumTexCoords = 8;
+
+inline bool IsTexCoord(EExternalInputType Type)
+{
+	return FMath::IsWithin((int32)Type, (int32)EExternalInputType::TexCoord0, (int32)EExternalInputType::TexCoord0 + NumTexCoords);
+}
+
 inline Shader::EValueType GetInputExpressionType(EExternalInputType Type)
 {
-	return Shader::EValueType::Float2;
+	if (IsTexCoord(Type))
+	{
+		return Shader::EValueType::Float2;
+	}
+
+	switch (Type)
+	{
+	case EExternalInputType::None:
+		return Shader::EValueType::Void;
+	case EExternalInputType::WorldPosition:
+	case EExternalInputType::WorldPosition_NoOffsets:
+		return Shader::EValueType::Double3;
+	case EExternalInputType::TranslatedWorldPosition:
+	case EExternalInputType::TranslatedWorldPosition_NoOffsets:
+		return Shader::EValueType::Float3;
+	default:
+		checkNoEntry();
+		return Shader::EValueType::Void;
+	}
 }
 inline EExternalInputType MakeInputTexCoord(int32 Index)
 {
-	check(Index >= 0 && Index < 8);
+	check(Index >= 0 && Index < NumTexCoords);
 	return (EExternalInputType)((int32)EExternalInputType::TexCoord0 + Index);
 }
 
diff --git a/Engine/Source/Runtime/Engine/Public/MaterialShared.h b/Engine/Source/Runtime/Engine/Public/MaterialShared.h
--- a/Engine/Source/Runtime/Engine/Public/MaterialShared.h
+++ b/Engine/Source/Runtime/Engine/Public/MaterialShared.h
@@ -382,6 +382,15 @@ inline bool operator!=(const FMaterialVirtualTextureStack& Lhs, const FMaterialV
 	return !operator==(Lhs, Rhs);
 }
 
+class FMaterialUniformPreshaderField
+{
+	DECLARE_TYPE_LAYOUT(FMaterialUniformPreshaderField, NonVirtual);
+public:
+	LAYOUT_FIELD(uint32, BufferOffset); /** Offset in the uniform buffer where result will be written */
+	LAYOUT_FIELD(uint32, ComponentIndex);
+	LAYOUT_FIELD(UE::Shader::EValueType, Type);
+};
+
 class FMaterialUniformPreshaderHeader
 {
 	DECLARE_TYPE_LAYOUT(FMaterialUniformPreshaderHeader, NonVirtual);
@@ -390,9 +399,8 @@ public:
 	{
 		return Lhs.OpcodeOffset == Rhs.OpcodeOffset &&
 			Lhs.OpcodeSize == Rhs.OpcodeSize &&
-			Lhs.BufferOffset == Rhs.BufferOffset &&
-			Lhs.ComponentType == Rhs.ComponentType &&
-			Lhs.NumComponents == Rhs.NumComponents;
+			Lhs.FieldIndex == Rhs.FieldIndex &&
+			Lhs.NumFields == Rhs.NumFields;
 	}
 	friend inline bool operator!=(const FMaterialUniformPreshaderHeader& Lhs, const FMaterialUniformPreshaderHeader& Rhs)
 	{
@@ -401,9 +409,8 @@ public:
 
 	LAYOUT_FIELD(uint32, OpcodeOffset); /** Offset of the preshader opcodes, within the material's buffer */
 	LAYOUT_FIELD(uint32, OpcodeSize); /** Size of the preshader opcodes */
-	LAYOUT_FIELD(uint32, BufferOffset); /** Offset in the uniform buffer where result will be written */
-	LAYOUT_FIELD(UE::Shader::EValueComponentType, ComponentType);
-	LAYOUT_FIELD(uint8, NumComponents);
+	LAYOUT_FIELD(uint32, FieldIndex);
+	LAYOUT_FIELD(uint32, NumFields);
 };
 
 class FMaterialNumericParameterInfo
@@ -590,6 +597,7 @@ protected:
 	friend class UE::HLSLTree::FEmitContext;
 
 	LAYOUT_FIELD(TMemoryImageArray<FMaterialUniformPreshaderHeader>, UniformPreshaders);
+	LAYOUT_FIELD(TMemoryImageArray<FMaterialUniformPreshaderField>, UniformPreshaderFields);
 	LAYOUT_FIELD(TMemoryImageArray<FMaterialNumericParameterInfo>, UniformNumericParameters);
 	LAYOUT_ARRAY(TMemoryImageArray<FMaterialTextureParameterInfo>, UniformTextureParameters, NumMaterialTextureParameterTypes);
 	LAYOUT_FIELD(TMemoryImageArray<FMaterialExternalTextureParameterInfo>, UniformExternalTextureParameters);
diff --git a/Engine/Source/Runtime/Engine/Public/Shader/Preshader.h b/Engine/Source/Runtime/Engine/Public/Shader/Preshader.h
--- a/Engine/Source/Runtime/Engine/Public/Shader/Preshader.h
+++ b/Engine/Source/Runtime/Engine/Public/Shader/Preshader.h
@@ -88,7 +88,8 @@ public:
 	FSHAHash GetHash() const;
 	void AppendHash(FSHA1& OutHasher) const;
 
-	FPreshaderValue Evaluate(FUniformExpressionSet* UniformExpressionSet, const struct FMaterialRenderContext& Context, FPreshaderStack& Stack);
+	FPreshaderValue Evaluate(FUniformExpressionSet* UniformExpressionSet, const struct FMaterialRenderContext& Context, FPreshaderStack& Stack) const;
+	FPreshaderValue EvaluateConstant(const FMaterial& Material, FPreshaderStack& Stack) const;
 
 	const int32 Num() const { return Data.Num(); }
 
diff --git a/Engine/Source/Runtime/Engine/Public/Shader/PreshaderTypes.h b/Engine/Source/Runtime/Engine/Public/Shader/PreshaderTypes.h
--- a/Engine/Source/Runtime/Engine/Public/Shader/PreshaderTypes.h
+++ b/Engine/Source/Runtime/Engine/Public/Shader/PreshaderTypes.h
@@ -17,6 +17,7 @@ struct FPreshaderType
 {
 	FPreshaderType() = default;
 	FPreshaderType(const FType& InType);
+	FPreshaderType(EValueType InType);
 
 	bool IsStruct() const { return StructTypeHash != 0u; }
 	int32 GetNumComponents() const { return StructTypeHash != 0u ? StructComponentTypes.Num() : GetValueTypeDescription(ValueType).NumComponents; }
diff --git a/Engine/Source/Runtime/Engine/Public/Shader/ShaderTypes.h b/Engine/Source/Runtime/Engine/Public/Shader/ShaderTypes.h
--- a/Engine/Source/Runtime/Engine/Public/Shader/ShaderTypes.h
+++ b/Engine/Source/Runtime/Engine/Public/Shader/ShaderTypes.h
@@ -24,6 +24,7 @@ enum class EValueComponentType : uint8
 	Bool,
 };
 
+const TCHAR* GetComponentTypeName(EValueComponentType Type);
 uint32 GetComponentTypeSizeInBytes(EValueComponentType Type);
 EValueComponentType CombineComponentTypes(EValueComponentType Lhs, EValueComponentType Rhs);
 
@@ -82,7 +83,9 @@ struct FType
 	bool IsStruct() const { return ValueType == EValueType::Struct; }
 	bool IsNumeric() const { return !IsVoid() && !IsStruct(); }
 	int32 GetNumComponents() const;
+	int32 GetNumFlatFields() const;
 	EValueComponentType GetComponentType(int32 Index) const;
+	EValueType GetFlatFieldType(int32 Index) const;
 	bool Merge(const FType& OtherType);
 
 	inline operator EValueType() const { return ValueType; }
@@ -127,6 +130,7 @@ struct FStructField
 	const TCHAR* Name;
 	FType Type;
 	int32 ComponentIndex;
+	int32 FlatFieldIndex;
 
 	int32 GetNumComponents() const { return Type.GetNumComponents(); }
 };
@@ -143,6 +147,11 @@ struct FStructType
 	 */
 	TArrayView<const EValueComponentType> ComponentTypes;
 
+	/**
+	 * Type may be viewed as a flat list of fields, rather than of individual components
+	 */
+	TArrayView<const EValueType> FlatFieldTypes;
+
 	const FStructField* FindFieldByName(const TCHAR* InName) const;
 };
 
@@ -358,6 +367,12 @@ struct FValue
 		return Component[i];
 	}
 
+	/** returns 0 for invalid components */
+	inline FValueComponent TryGetComponent(int32 i) const
+	{
+		return Component.IsValidIndex(i) ? Component[i] : FValueComponent();
+	}
+
 	static FValue FromMemoryImage(EValueType Type, const void* Data, uint32* OutSizeInBytes = nullptr);
 	FMemoryImageValue AsMemoryImage() const;
 
