Fix indeterminism when applying transformation in shaders by providing a unique guid that represents that transformation when possible

diff --git a/Engine/Source/Runtime/Engine/Classes/Materials/Material.h b/Engine/Source/Runtime/Engine/Classes/Materials/Material.h
--- a/Engine/Source/Runtime/Engine/Classes/Materials/Material.h
+++ b/Engine/Source/Runtime/Engine/Classes/Materials/Material.h
@@ -1456,9 +1456,11 @@ public:
 
 private:
 	/**
-	 * Flush existing resource shader maps and resets the material resource's Ids.
+	 * Flush existing resource shader maps and combines the material resource's Ids with another unique guid that represent the transformation applied so we can stay deterministic.
+	 * If the guid given is non-valid (default value) we generate an entirely new guid for the resource.
+	 * @param	TransformationId	If valid, combine this guid with the current material resource's Ids yielding a new but deterministic Ids.
 	 */
-	ENGINE_API virtual void ReleaseResourcesAndMutateDDCKey();
+	ENGINE_API virtual void ReleaseResourcesAndMutateDDCKey(const FGuid& TransformationId = FGuid());
 	
 	/** 
 	 * Cache resource shaders for rendering. 
diff --git a/Engine/Source/Runtime/Engine/Private/Materials/Material.cpp b/Engine/Source/Runtime/Engine/Private/Materials/Material.cpp
--- a/Engine/Source/Runtime/Engine/Private/Materials/Material.cpp
+++ b/Engine/Source/Runtime/Engine/Private/Materials/Material.cpp
@@ -2101,6 +2101,8 @@ void UMaterial::CacheResourceShadersForRendering(bool bRegenerateId, EMaterialSh
 	if (bRegenerateId)
 	{
 		// Regenerate this material's Id if requested
+		// Since we can't provide an explanation for why we've been asked to change the guid.
+		// We can't give this function a unique transformation id, let it generate a new one.
 		ReleaseResourcesAndMutateDDCKey();
 	}
 
@@ -2267,9 +2269,21 @@ void UMaterial::CacheShadersForResources(EShaderPlatform ShaderPlatform, const T
 	}
 }
 
-void UMaterial::ReleaseResourcesAndMutateDDCKey()
+void UMaterial::ReleaseResourcesAndMutateDDCKey(const FGuid& TransformationId)
 {
-	FPlatformMisc::CreateGuid(StateId);
+	if (TransformationId.IsValid())
+	{
+		// Combine current guid with the transformation applied.
+		StateId.A ^= TransformationId.A;
+		StateId.B ^= TransformationId.B;
+		StateId.C ^= TransformationId.C;
+		StateId.D ^= TransformationId.D;
+		
+	}
+	else
+	{
+		FPlatformMisc::CreateGuid(StateId);
+	}
 
 	if(FApp::CanEverRender())
 	{
@@ -2533,6 +2547,9 @@ void UMaterial::BackwardsCompatibilityVirtualTextureOutputConversion()
 	// Remove MD_RuntimeVirtualTexture support and replace with an explicit UMaterialExpressionRuntimeVirtualTextureOutput.
 	if (MaterialDomain == MD_RuntimeVirtualTexture)
 	{
+		// Change this guid if you change the conversion code below
+		static FGuid BackwardsCompatibilityVirtualTextureOutputConversionGuid(TEXT("BABD7074-001F-4FC2-BDE5-3A0C436F4414"));
+
 		MaterialDomain = MD_Surface;
 
 		if (!bUseMaterialAttributes)
@@ -2595,19 +2612,8 @@ void UMaterial::BackwardsCompatibilityVirtualTextureOutputConversion()
 			}
 		}
 
-		// Recompile after changes.
-		ReleaseResourcesAndMutateDDCKey();
-
-		// Note we can't mark the package dirty during post load so we will recompile on load until this material is manually resaved. 
-		// Also we can't fully deprecate until all the relevant materials are resaved. 
-		// So add a map load error here to encourage a save.
-		FFormatNamedArguments Arguments;
-		Arguments.Add(TEXT("Material"), FText::FromString(*GetPathName()));
-		FMessageLog("MapCheck").Warning()
-			->AddToken(FUObjectToken::Create(this))
-			->AddToken(FTextToken::Create(FText::Format(LOCTEXT("MapCheck_RuntimeVirtualTextureMaterialDomain", "Material {Material} has been updated to use a RuntimeVirtualTextureOutput node instead of a RuntimeVirtualTexture Domain. If the material asset is not re-saved it may not render correctly when run outside of the editor."), Arguments)))
-			->AddToken(FActionToken::Create(LOCTEXT("MapCheck_RuntimeVirtualTextureMaterialDomain_Fix", "Fix"), LOCTEXT("MapCheck_RuntimeVirtualTextureMaterialDomain_Action", "Click to mark the material as needing to be saved."), FOnActionTokenExecuted::CreateUObject(this, &UMaterial::FixupMaterialUsageAfterLoad), true));
-		FMessageLog("MapCheck").Open(EMessageSeverity::Warning);
+		// Recompile after changes with a guid representing the conversion applied here.
+		ReleaseResourcesAndMutateDDCKey(BackwardsCompatibilityVirtualTextureOutputConversionGuid);
 	}
 #endif // WITH_EDITOR
 }
@@ -2617,6 +2623,9 @@ void UMaterial::BackwardsCompatibilityDecalConversion()
 #if WITH_EDITOR
 	if (GMaterialsThatNeedDecalFix.Get(this))
 	{
+		// Change this guid if you change the conversion code below
+		static FGuid BackwardsCompatibilityDecalConversionGuid(TEXT("352069F8-1B8C-406A-9B88-6946BCDF2C10"));
+
 		GMaterialsThatNeedDecalFix.Clear(this);
 
 		// Move stain and alpha composite setting into material blend mode.
@@ -2690,18 +2699,8 @@ void UMaterial::BackwardsCompatibilityDecalConversion()
 			}
 		}
 
-		// Force the material to recompile.
-		ReleaseResourcesAndMutateDDCKey();
-
-		// Note we can't mark the package dirty during post load so we will recompile on load until this material is manually resaved. 
-		// Add an asset check here to encourage a save.
-		FFormatNamedArguments Arguments;
-		Arguments.Add(TEXT("Material"), FText::FromString(*GetPathName()));
-		FMessageLog("AssetCheck").Info()
-			->AddToken(FUObjectToken::Create(this))
-			->AddToken(FTextToken::Create(FText::Format(LOCTEXT("AssetCheck_Decal", "Material {Material} has been updated to fix up decal settings. Resave material to avoid it recompiling on each load."), Arguments)))
-			->AddToken(FActionToken::Create(LOCTEXT("AssetCheck_Decal_Fix", "Fix"), LOCTEXT("AssetCheck_Decal_Action", "Click to mark the material as needing to be saved."), FOnActionTokenExecuted::CreateUObject(this, &UMaterial::FixupMaterialUsageAfterLoad), true));
-		FMessageLog("AssetCheck").Open(EMessageSeverity::Info);
+		// Recompile after changes with a guid representing the conversion applied here.
+		ReleaseResourcesAndMutateDDCKey(BackwardsCompatibilityDecalConversionGuid);
 	}
 #endif // WITH_EDITOR
 }
@@ -3095,8 +3094,8 @@ void UMaterial::PostLoad()
 		if (Expressions.Remove(NULL) != 0)
 		{
 			// Force this material to recompile because its expressions have changed
-			// Warning: any content taking this path will recompile every load until saved!
-			// Which means removing an expression class will cause the need for a resave of all materials
+			// We're not providing a deterministic transformation guid because there could be many different ways expression
+			// could change. Each conversion code removing such expression would need its own guid.
 			ReleaseResourcesAndMutateDDCKey();
 		}
 	}
