Chaos Cloth - Add environmental collision support.

diff --git a/Engine/Config/BaseEngine.ini b/Engine/Config/BaseEngine.ini
--- a/Engine/Config/BaseEngine.ini
+++ b/Engine/Config/BaseEngine.ini
@@ -2911,3 +2911,11 @@ GlobalLambdaMultiplier=1.000000
 ; Specifies how many buffers Oodle should keep around for 
 ; reuse as opposed to heap allocating every decode.
 PreallocatedBufferCount=2
+
+[ClothSettings]
+;  Extra padding added to the bounds so that the collision can still be detected after being resolved
+EnvironmentCollisionPadding = 2.f
+; Extra thickness added to edgy collision shapes (cubes & convexes)
+EnvironmentCollisionThickness = 2.f
+; Limit the number of extracted shapes per component as these collisions are very expensive
+EnvironmentCollisionMaxShapes = 32
diff --git a/Engine/Source/Runtime/ClothingSystemRuntimeInterface/Private/ClothCollisionPrim.cpp b/Engine/Source/Runtime/ClothingSystemRuntimeInterface/Private/ClothCollisionPrim.cpp
deleted file mode 100644
diff --git a/Engine/Source/Runtime/ClothingSystemRuntimeInterface/Public/ClothCollisionPrim.h b/Engine/Source/Runtime/ClothingSystemRuntimeInterface/Public/ClothCollisionPrim.h
--- a/Engine/Source/Runtime/ClothingSystemRuntimeInterface/Public/ClothCollisionPrim.h
+++ b/Engine/Source/Runtime/ClothingSystemRuntimeInterface/Public/ClothCollisionPrim.h
@@ -14,10 +14,10 @@ struct FClothCollisionPrim_Sphere
 {
 	GENERATED_BODY()
 
-	FClothCollisionPrim_Sphere()
-		: BoneIndex(INDEX_NONE)
-		, Radius(0.0f)
-		, LocalPosition(0, 0, 0)
+	FClothCollisionPrim_Sphere(float InRadius = 0.0f, const FVector& InLocalPosition = FVector::ZeroVector, int32 InBoneIndex = INDEX_NONE)
+		: BoneIndex(InBoneIndex)
+		, Radius(InRadius)
+		, LocalPosition(InLocalPosition)
 	{}
 
 	UPROPERTY()
@@ -38,10 +38,10 @@ struct FClothCollisionPrim_SphereConnection
 {
 	GENERATED_BODY()
 
-	FClothCollisionPrim_SphereConnection()
+	FClothCollisionPrim_SphereConnection(int32 SphereIndex0 = INDEX_NONE, int32 SphereIndex1 = INDEX_NONE)
 	{
-		SphereIndices[0] = INDEX_NONE;
-		SphereIndices[1] = INDEX_NONE;
+		SphereIndices[0] = SphereIndex0;
+		SphereIndices[1] = SphereIndex1;
 	}
 
 	UPROPERTY()
@@ -77,11 +77,14 @@ struct CLOTHINGSYSTEMRUNTIMEINTERFACE_API FClothCollisionPrim_Convex
 		: BoneIndex(INDEX_NONE)
 	{}
 
-	/** Rebuild the surface point array from the existing planes.
-	 *  This is an expensive function (O(n^4) per number of planes).
-	 */
-	UE_DEPRECATED(4.27, "RebuildSurfacePoints is now deprecated as it doesn't provide enough data to regenerate the indices required by FKConvexElem and FConvex.")
-	void RebuildSurfacePoints();
+	FClothCollisionPrim_Convex(
+		TArray<FClothCollisionPrim_ConvexFace>&& InFaces,
+		TArray<FVector>&& InSurfacePoints,
+		int32 InBoneIndex = INDEX_NONE)
+		: Faces(InFaces)
+		, SurfacePoints(InSurfacePoints)
+		, BoneIndex(InBoneIndex)
+	{}
 
 #if WITH_EDITORONLY_DATA
 	UPROPERTY()
@@ -104,11 +107,15 @@ struct FClothCollisionPrim_Box
 {
 	GENERATED_BODY()
 
-	FClothCollisionPrim_Box()
-		: LocalPosition(FVector::ZeroVector)
-		, LocalRotation(FQuat::Identity)
-		, HalfExtents(FVector::ZeroVector)
-		, BoneIndex(INDEX_NONE)
+	FClothCollisionPrim_Box(
+		const FVector& InLocalPosition = FVector::ZeroVector,
+		const FQuat& InLocalRotation = FQuat::Identity,
+		const FVector& InHalfExtents = FVector::ZeroVector,
+		int32 InBoneIndex = INDEX_NONE)
+		: LocalPosition(InLocalPosition)
+		, LocalRotation(InLocalRotation)
+		, HalfExtents(InHalfExtents)
+		, BoneIndex(InBoneIndex)
 	{}
 
 	UPROPERTY()
diff --git a/Engine/Source/Runtime/Engine/Classes/Components/SkeletalMeshComponent.h b/Engine/Source/Runtime/Engine/Classes/Components/SkeletalMeshComponent.h
--- a/Engine/Source/Runtime/Engine/Classes/Components/SkeletalMeshComponent.h
+++ b/Engine/Source/Runtime/Engine/Classes/Components/SkeletalMeshComponent.h
@@ -638,6 +638,14 @@ public:
 	/** can't collide with part of attached children if total collision volumes exceed 16 capsules or 32 planes per convex */
 	UPROPERTY(EditAnywhere, BlueprintReadWrite, Category=Clothing)
 	uint8 bCollideWithAttachedChildren:1;
+	/**
+	 * Forces the cloth simulation to constantly update its external collisions, at the expense of performance.
+	 * This will help to prevent missed collisions if the cloth's skeletal mesh component isn't moving,
+	 * and when instead, wind or other moving objects are causing new collision shapes to come into the cloth's vicinity.
+	 */
+	UPROPERTY(EditAnywhere, BlueprintReadWrite, Category=Clothing)
+	uint8 bForceCollisionUpdate : 1;
+
 	/**
 	 * It's worth trying this option when you feel that the current cloth simulation is unstable.
 	 * The scale of the actor is maintained during the simulation. 
diff --git a/Engine/Source/Runtime/Engine/Private/Components/SkeletalMeshComponent.cpp b/Engine/Source/Runtime/Engine/Private/Components/SkeletalMeshComponent.cpp
--- a/Engine/Source/Runtime/Engine/Private/Components/SkeletalMeshComponent.cpp
+++ b/Engine/Source/Runtime/Engine/Private/Components/SkeletalMeshComponent.cpp
@@ -631,16 +631,6 @@ void USkeletalMeshComponent::OnRegister()
 
 	RecreateClothingActors();
 #endif  // #if WITH_APEX_CLOTHING || WITH_CHAOS_CLOTHING
-
-#if WITH_CHAOS_CLOTHING
-	// TODO: Add missing Chaos Cloth collision with environment
-	if (bCollideWithEnvironment && ClothingSimulationFactory && ClothingSimulationFactory->GetName() == TEXT("ChaosClothingSimulationFactory"))
-	{
-		UE_LOG(LogSkeletalMesh, Display, TEXT("OnRegister[%s]: "
-			"Chaos Cloth does not currently support bCollideWithEnvironment."),
-			*GetPathNameSafe(SkeletalMesh));
-	}
-#endif  // #if WITH_CHAOS_CLOTHING
 }
 
 void USkeletalMeshComponent::OnUnregister()
@@ -2117,18 +2107,23 @@ void USkeletalMeshComponent::UpdateClothSimulationContext(float InDeltaTime)
 	//Do the teleport cloth test here on the game thread
 	CheckClothTeleport();
 
-	if(bPendingClothTransformUpdate)	//it's possible we want to update cloth collision based on a pending transform
+	bool bMustUpdateClothTransform = bForceCollisionUpdate;
+
+	if (bPendingClothTransformUpdate)	//it's possible we want to update cloth collision based on a pending transform
 	{
 		bPendingClothTransformUpdate = false;
-		if(PendingTeleportType == ETeleportType::TeleportPhysics)	//If the pending transform came from a teleport, make sure to teleport the cloth in this upcoming simulation
+		if (PendingTeleportType == ETeleportType::TeleportPhysics)	//If the pending transform came from a teleport, make sure to teleport the cloth in this upcoming simulation
 		{
 			ClothTeleportMode = (ClothTeleportMode == EClothingTeleportMode::TeleportAndReset) ? ClothTeleportMode : EClothingTeleportMode::Teleport;
 		}
-		else if(PendingTeleportType == ETeleportType::ResetPhysics)
+		else if (PendingTeleportType == ETeleportType::ResetPhysics)
 		{
 			ClothTeleportMode = EClothingTeleportMode::TeleportAndReset;
 		}
-
+		bMustUpdateClothTransform = true;
+	}
+	if (bMustUpdateClothTransform)
+	{
 		UpdateClothTransformImp();
 	}
 
diff --git a/Engine/Source/Runtime/Engine/Private/SkeletalMeshComponentPhysics.cpp b/Engine/Source/Runtime/Engine/Private/SkeletalMeshComponentPhysics.cpp
--- a/Engine/Source/Runtime/Engine/Private/SkeletalMeshComponentPhysics.cpp
+++ b/Engine/Source/Runtime/Engine/Private/SkeletalMeshComponentPhysics.cpp
@@ -56,6 +56,19 @@
 #include "SimulationEditorExtender.h"
 #endif
 
+#if !PHYSICS_INTERFACE_PHYSX && WITH_CHAOS  // For cloth environmental collision
+#include "Chaos/Capsule.h"
+#include "Chaos/Sphere.h"
+#include "Chaos/Box.h"
+#include "Chaos/Convex.h"
+#include "Chaos/ImplicitFwd.h"
+#include "Chaos/ImplicitObject.h"
+#include "Chaos/ImplicitObjectScaled.h"
+#include "Chaos/ImplicitObjectTransformed.h"
+
+using namespace Chaos;
+#endif  // #if !PHYSICS_INTERFACE_PHYSX && WITH_CHAOS
+
 #define LOCTEXT_NAMESPACE "SkeletalMeshComponentPhysics"
 
 DECLARE_CYCLE_STAT(TEXT("CreateClothing"), STAT_CreateClothing, STATGROUP_Physics);
@@ -1545,10 +1558,8 @@ void USkeletalMeshComponent::OnUpdateTransform(EUpdateTransformFlags UpdateTrans
 		ResetAnimInstanceDynamics(Teleport);
 	}
 
-	if(ClothingSimulation && ClothingSimulation->ShouldSimulate())
-	{
-		UpdateClothTransform(Teleport);
-	}
+	// Mark the cloth simulation transform update as pending (the actual update will happen in the cloth thread if the simulation is running)
+	UpdateClothTransform(Teleport);
 }
 
 bool USkeletalMeshComponent::UpdateOverlapsImpl(const TOverlapArrayView* PendingOverlaps, bool bDoNotifies, const TOverlapArrayView* OverlapsAtEndLocation)
@@ -3012,8 +3023,24 @@ void USkeletalMeshComponent::CopyChildrenClothCollisionsToParent()
 
 void USkeletalMeshComponent::ProcessClothCollisionWithEnvironment()
 {
-	// Limiting the number of extracted shapes per component as these collisions are very expensive
-	static const int32 MaxSyncShapesToConsider = 32;
+	// Read config on first call
+	struct FEnvironmentCollisionConfig
+	{
+		float Padding = 2.f;  // Extra padding added to the bounds so that the collision can still be detected after being resolved
+		float Thickness = 2.f;  // Extra thickness added to edgy collision shapes (cubes & convexes)
+		int32 MaxShapes = 32;  // Limit the number of extracted shapes per component as these collisions are very expensive
+
+		FEnvironmentCollisionConfig()
+		{
+			if (GConfig)
+			{
+				GConfig->GetFloat(TEXT("ClothSettings"), TEXT("EnvironmentCollisionPadding"), Padding, GEngineIni);
+				GConfig->GetFloat(TEXT("ClothSettings"), TEXT("EnvironmentCollisionThickness"), Thickness, GEngineIni);
+				GConfig->GetInt(TEXT("ClothSettings"), TEXT("EnvironmentCollisionMaxShapes"), MaxShapes, GEngineIni);
+			}
+		}
+	};
+	static FEnvironmentCollisionConfig EnvironmentCollisionConfig;
 
 	// don't handle collision detection if this component is in editor
 	if(!GetWorld()->IsGameWorld() || !ClothingSimulation)
@@ -3032,8 +3059,9 @@ void USkeletalMeshComponent::ProcessClothCollisionWithEnvironment()
 	ObjectParams.AddObjectTypesToQuery(ECollisionChannel::ECC_PhysicsBody);
 
 	FCollisionQueryParams Params(SCENE_QUERY_STAT(ClothOverlapComponents), false);
+	const FBoxSphereBounds PaddedBounds = Bounds.ExpandBy(EnvironmentCollisionConfig.Padding);
 
-	GetWorld()->OverlapMultiByObjectType(Overlaps, Bounds.Origin, FQuat::Identity, ObjectParams, FCollisionShape::MakeBox(Bounds.BoxExtent), Params);
+	GetWorld()->OverlapMultiByObjectType(Overlaps, PaddedBounds.Origin, FQuat::Identity, ObjectParams, FCollisionShape::MakeBox(PaddedBounds.BoxExtent), Params);
 
 	for (int32 OverlapIdx=0; OverlapIdx<Overlaps.Num(); ++OverlapIdx)
 	{
@@ -3059,12 +3087,13 @@ void USkeletalMeshComponent::ProcessClothCollisionWithEnvironment()
 					TArray<FPhysicsShapeHandle> AllShapes;
 					const int32 NumSyncShapes = Component->BodyInstance.GetAllShapes_AssumesLocked(AllShapes);
 
-					if(NumSyncShapes == 0 || NumSyncShapes > MaxSyncShapesToConsider)
+					if (NumSyncShapes == 0 || NumSyncShapes > EnvironmentCollisionConfig.MaxShapes)
 					{
 						// Either no shapes or too complicated to consider
 						return;
 					}
 
+#if PHYSICS_INTERFACE_PHYSX
 					// Matrices required to transform shapes into sim space (component space)
 					// Transform of external component and matrix describing external component -> this component
 					FTransform Transform = Component->GetComponentTransform();
@@ -3079,9 +3108,6 @@ void USkeletalMeshComponent::ProcessClothCollisionWithEnvironment()
 						// Pose of the shape in actor space
 						FMatrix ShapeLocalPose = FPhysicsInterface::GetLocalTransform(Shape).ToMatrixWithScale();
 
-#if WITH_CHAOS
-						// TODO: Add missing Chaos Cloth collision with environment
-#else
 						switch(GeoType)
 						{
 							default:
@@ -3141,8 +3167,7 @@ void USkeletalMeshComponent::ProcessClothCollisionWithEnvironment()
 								Convex.Faces.SetNum(6);
 
 								// we need to inflate the hull to get nicer collisions (only particles collide)
-								const static float Inflate = 2.0f;
-								BoxGeo.halfExtents += PxVec3(Inflate);
+								BoxGeo.halfExtents += PxVec3(EnvironmentCollisionConfig.Thickness);
 				
 								FPlane UPlane1(1, 0, 0, BoxGeo.halfExtents.x);
 								UPlane1 = UPlane1.TransformBy(FullTransformMatrix);
@@ -3176,9 +3201,6 @@ void USkeletalMeshComponent::ProcessClothCollisionWithEnvironment()
 							{
 								PxConvexMeshGeometry& MeshGeo = GeoCollection.GetConvexGeometry();
 
-								// we need to inflate the hull to get nicer collisions (only particles collide)
-								const static float Inflate = 2.0f;
-
 								if(MeshGeo.convexMesh)
 								{
 									NewCollisionData.Convexes.AddDefaulted();
@@ -3197,7 +3219,7 @@ void USkeletalMeshComponent::ProcessClothCollisionWithEnvironment()
 										FPlane UPlane = P2UPlane(PPlane);
 										UPlane = UPlane.TransformBy(FullTransformMatrix);
 										
-										UPlane.W += Inflate;
+										UPlane.W += EnvironmentCollisionConfig.Thickness;
 
 										NewConvex.Faces[PolyIndex].Plane = UPlane;
 									}
@@ -3205,8 +3227,255 @@ void USkeletalMeshComponent::ProcessClothCollisionWithEnvironment()
 							}
 							break;
 						}
-#endif
 					}
+#elif WITH_CHAOS
+					const FReal Thickness = (FReal)EnvironmentCollisionConfig.Thickness;
+
+					auto AddSphere = [&NewCollisionData](const FImplicitSphere3& ImplicitSphere, const FTransform& ComponentToClothTransform, const FVec3& Scale = FVec3::OneVector)
+					{
+						NewCollisionData.Spheres.Emplace(
+							ImplicitSphere.GetRadius() * Scale.X,  // Assumes uniform scale
+							ComponentToClothTransform.TransformPosition(ImplicitSphere.GetCenter()));
+					};
+
+					auto AddBox = [&NewCollisionData, Thickness](const FImplicitBox3& ImplicitBox, const FTransform& ComponentToClothTransform, const FVec3& Scale = FVec3::OneVector)
+						{
+							NewCollisionData.Boxes.Emplace(
+								ComponentToClothTransform.TransformPosition(ImplicitBox.GetCenter()),
+								ComponentToClothTransform.GetRotation(),
+								ImplicitBox.Extents() * Scale * (FReal)0.5 + Thickness);
+						};
+
+					auto AddCapsule = [&NewCollisionData](const FImplicitCapsule3& ImplicitCapsule, const FTransform& ComponentToClothTransform, const FVec3& Scale = FVec3::OneVector)
+						{
+							const int32 BaseSphereIndex = NewCollisionData.Spheres.Num();
+
+							const FReal Radius = ImplicitCapsule.GetRadius() * Scale.X;  // Assumes uniform scale
+							const FReal HalfHeight = ImplicitCapsule.GetHeight() * (FReal)0.5;
+							const FVector HalfSegment = ComponentToClothTransform.GetUnitAxis(EAxis::X) * HalfHeight * Scale.X;  // Assumes uniform scale
+							const FVector TransformedCenter = ComponentToClothTransform.TransformPosition(ImplicitCapsule.GetCenter());
+
+							NewCollisionData.Spheres.Emplace(Radius, TransformedCenter + HalfSegment);
+							NewCollisionData.Spheres.Emplace(Radius, TransformedCenter - HalfSegment);
+
+							NewCollisionData.SphereConnections.Emplace(BaseSphereIndex, BaseSphereIndex + 1);
+						};
+
+					auto AddConvex = [&NewCollisionData, Thickness](const FImplicitConvex3& ImplicitConvex, const FMatrix& ComponentToClothMatrix, const FVec3* const Scale = nullptr)
+						{
+							TArray<FClothCollisionPrim_ConvexFace> Faces;
+							const int32 NumPlanes = ImplicitConvex.NumPlanes();
+							Faces.SetNum(NumPlanes);
+
+							TArray<FVector> SurfacePoints;
+							const int32 NumSurfacePoints = ImplicitConvex.NumVertices();
+							SurfacePoints.SetNumUninitialized(NumSurfacePoints);
+
+							TArray<TArray<int32, TInlineAllocator<4>>, TInlineAllocator<16>> PointFaces;
+							PointFaces.SetNum(NumSurfacePoints);
+
+							for (int32 FaceIndex = 0; FaceIndex < NumPlanes; ++FaceIndex)
+							{
+								const TPlaneConcrete<FReal, 3>& Plane = ImplicitConvex.GetPlane(FaceIndex);
+								const FVec3& Normal = Plane.Normal();
+								const FVec3 Base = Plane.X() + Normal * (FReal)Thickness;
+
+								Faces[FaceIndex].Plane = FPlane(FVector(Base), FVector(Normal)).TransformBy(ComponentToClothMatrix);
+
+								const int32 NumFaceIndices = ImplicitConvex.NumPlaneVertices(FaceIndex);
+								Faces[FaceIndex].Indices.SetNumUninitialized(NumFaceIndices);
+
+								for (int32 Index = 0; Index < NumFaceIndices; ++Index)
+								{
+									const int32 PointIndex = ImplicitConvex.GetPlaneVertex(FaceIndex, Index);
+									Faces[FaceIndex].Indices[Index] = PointIndex;
+
+									PointFaces[PointIndex].Add(FaceIndex);
+								}
+							}
+
+							for (int32 PointIndex = 0; PointIndex < NumSurfacePoints; ++PointIndex)
+							{
+								check(PointFaces[PointIndex].Num() >= 3);
+								const int32 Index0 = PointFaces[PointIndex][0];
+								const int32 Index1 = PointFaces[PointIndex][1];
+								const int32 Index2 = PointFaces[PointIndex][2];
+
+								if (!FMath::IntersectPlanes3(SurfacePoints[PointIndex], Faces[Index0].Plane, Faces[Index1].Plane, Faces[Index2].Plane))
+								{
+									SurfacePoints[PointIndex] = ComponentToClothMatrix.TransformPosition(FVector(ImplicitConvex.GetVertex(PointIndex)));
+								}
+							}
+
+							NewCollisionData.Convexes.Emplace(MoveTemp(Faces), MoveTemp(SurfacePoints));
+						};
+
+					FTransform ClothComponentTransform = GetComponentTransform();
+					ClothComponentTransform.RemoveScaling();  // The environment collision shape doesn't need the scale of the cloth skeletal mesh applied to it (but it does need the source scale from its component transform)
+					const FTransform ComponentToClothBaseTransform = Component->GetComponentTransform() * ClothComponentTransform.Inverse();
+
+					bool bHasSimpleCollision = false;
+
+					for (FPhysicsShapeHandle& ShapeHandle : AllShapes)
+					{
+						FTransform ComponentToClothTransform = ComponentToClothBaseTransform;
+
+						const FImplicitObject* ImplicitObject = &ShapeHandle.GetGeometry();
+						EImplicitObjectType ImplicitType = ImplicitObject->GetType();
+
+						// Transformed implicits
+						if (ImplicitType == ImplicitObjectType::Transformed)
+						{
+							const TImplicitObjectTransformed<FReal, 3>& ImplicitTransformed = ImplicitObject->GetObjectChecked<TImplicitObjectTransformed<FReal, 3>>();
+							ImplicitObject = ImplicitTransformed.GetTransformedObject();
+							ImplicitType = ImplicitObject->GetType();
+
+							ComponentToClothTransform = ImplicitTransformed.GetTransform() * ComponentToClothTransform;
+							UE_LOG(LogSkeletalMesh, Verbose, TEXT("Found transformed environmental collision"));
+						}
+
+						switch (ImplicitType)
+						{
+						// Base implicits
+						case ImplicitObjectType::Sphere:
+							{
+								const FImplicitSphere3& ImplicitSphere = ImplicitObject->GetObjectChecked<FImplicitSphere3>();;
+								AddSphere(ImplicitSphere, ComponentToClothTransform);
+							}
+							UE_LOG(LogSkeletalMesh, Verbose, TEXT("Found Sphere cloth environmental collision in [%s]"), *Component->GetOwner()->GetFName().ToString());
+							bHasSimpleCollision = true;
+							break;
+						case ImplicitObjectType::Box:
+							{
+								const FImplicitBox3& ImplicitBox = ImplicitObject->GetObjectChecked<FImplicitBox3>();
+								AddBox(ImplicitBox, ComponentToClothTransform);
+							}
+							UE_LOG(LogSkeletalMesh, Verbose, TEXT("Found Box cloth environmental collision in [%s]"), *Component->GetOwner()->GetFName().ToString());
+							bHasSimpleCollision = true;
+							break;
+						case ImplicitObjectType::Capsule:
+							{
+								const FImplicitCapsule3& ImplicitCapsule = ImplicitObject->GetObjectChecked<FImplicitCapsule3>();
+								AddCapsule(ImplicitCapsule, ComponentToClothTransform);
+								UE_LOG(LogSkeletalMesh, Verbose, TEXT("Found Capsule cloth environmental collision in [%s]"), *Component->GetOwner()->GetFName().ToString());
+								bHasSimpleCollision = true;
+							}
+							break;
+						case ImplicitObjectType::Convex:
+							{
+								const FImplicitConvex3& ImplicitConvex = ImplicitObject->GetObjectChecked<FImplicitConvex3>();
+								AddConvex(ImplicitConvex, ComponentToClothTransform.ToMatrixNoScale());
+								UE_LOG(LogSkeletalMesh, Verbose, TEXT("Found Convex cloth environmental collision in [%s]"), *Component->GetOwner()->GetFName().ToString());
+								bHasSimpleCollision = true;
+							}
+							break;
+
+						// Instanced implicits
+						case ImplicitObjectType::IsInstanced | ImplicitObjectType::Sphere:
+							{
+								const TImplicitObjectInstanced<FImplicitSphere3>& ImplicitInstanced = ImplicitObject->GetObjectChecked<TImplicitObjectInstanced<FImplicitSphere3>>();
+								check(ImplicitInstanced.Object());
+								const FImplicitSphere3& ImplicitSphere = *ImplicitInstanced.GetInstancedObject();
+								AddSphere(ImplicitSphere, ComponentToClothTransform);
+								UE_LOG(LogSkeletalMesh, Verbose, TEXT("Found Instanced Sphere cloth environmental collision in [%s]"), *Component->GetOwner()->GetFName().ToString());
+								bHasSimpleCollision = true;
+							}
+							break;
+						case ImplicitObjectType::IsInstanced | ImplicitObjectType::Box:
+							{
+								const TImplicitObjectInstanced<FImplicitBox3>& ImplicitInstanced = ImplicitObject->GetObjectChecked<TImplicitObjectInstanced<FImplicitBox3>>();
+								check(ImplicitInstanced.Object());
+								const FImplicitBox3& ImplicitBox = *ImplicitInstanced.GetInstancedObject();
+								AddBox(ImplicitBox, ComponentToClothTransform);
+								UE_LOG(LogSkeletalMesh, Verbose, TEXT("Found Instanced Box cloth environmental collision in [%s]"), *Component->GetOwner()->GetFName().ToString());
+								bHasSimpleCollision = true;
+							}
+							break;
+						case ImplicitObjectType::IsInstanced | ImplicitObjectType::Capsule:
+							{
+								const TImplicitObjectInstanced<FImplicitCapsule3>& ImplicitInstanced = ImplicitObject->GetObjectChecked<TImplicitObjectInstanced<FImplicitCapsule3>>();
+								check(ImplicitInstanced.Object());
+								const FImplicitCapsule3& ImplicitCapsule = *ImplicitInstanced.GetInstancedObject();
+								AddCapsule(ImplicitCapsule, ComponentToClothTransform);
+								UE_LOG(LogSkeletalMesh, Verbose, TEXT("Found Instanced Capsule cloth environmental collision in [%s]"), *Component->GetOwner()->GetFName().ToString());
+								bHasSimpleCollision = true;
+							}
+							break;
+						case ImplicitObjectType::IsInstanced | ImplicitObjectType::Convex:
+							{
+								const TImplicitObjectInstanced<FImplicitConvex3>& ImplicitInstanced = ImplicitObject->GetObjectChecked<TImplicitObjectInstanced<FImplicitConvex3>>();
+								check(ImplicitInstanced.Object());
+								const FImplicitConvex3& ImplicitConvex = *ImplicitInstanced.GetInstancedObject();
+								AddConvex(ImplicitConvex, ComponentToClothTransform.ToMatrixNoScale());
+								UE_LOG(LogSkeletalMesh, Verbose, TEXT("Found Instanced Convex cloth environmental collision in [%s]"), *Component->GetOwner()->GetFName().ToString());
+								bHasSimpleCollision = true;
+							}
+							break;
+
+						// Scaled implicits
+						case ImplicitObjectType::IsScaled | ImplicitObjectType::Sphere:
+							{
+								const TImplicitObjectScaled<FImplicitSphere3>& ImplicitScaled = ImplicitObject->GetObjectChecked<TImplicitObjectScaled<FImplicitSphere3>>();
+								check(ImplicitScaled.Object());
+								const FImplicitSphere3& ImplicitSphere = *ImplicitScaled.GetUnscaledObject();
+								ensure(FVector::DistSquared(ComponentToClothTransform.GetScale3D(), FVector(ImplicitScaled.GetScale())) < KINDA_SMALL_NUMBER);
+								AddSphere(ImplicitSphere, ComponentToClothTransform, ImplicitScaled.GetScale());
+								UE_LOG(LogSkeletalMesh, Verbose, TEXT("Found Scaled Sphere cloth environmental collision in [%s]"), *Component->GetOwner()->GetFName().ToString());
+								bHasSimpleCollision = true;
+							}
+							break;
+						case ImplicitObjectType::IsScaled | ImplicitObjectType::Box:
+							{
+								const TImplicitObjectScaled<FImplicitBox3>& ImplicitScaled = ImplicitObject->GetObjectChecked<TImplicitObjectScaled<FImplicitBox3>>();
+								check(ImplicitScaled.Object());
+								const FImplicitBox3& ImplicitBox = *ImplicitScaled.GetUnscaledObject();
+								ensure(FVector::DistSquared(ComponentToClothTransform.GetScale3D(), FVector(ImplicitScaled.GetScale())) < KINDA_SMALL_NUMBER);
+								AddBox(ImplicitBox, ComponentToClothTransform, ImplicitScaled.GetScale());
+								UE_LOG(LogSkeletalMesh, Verbose, TEXT("Found Scaled Box cloth environmental collision in [%s]"), *Component->GetOwner()->GetFName().ToString());
+								bHasSimpleCollision = true;
+							}
+							break;
+						case ImplicitObjectType::IsScaled | ImplicitObjectType::Capsule:
+							{
+								const TImplicitObjectScaled<FImplicitCapsule3>& ImplicitScaled = ImplicitObject->GetObjectChecked<TImplicitObjectScaled<FImplicitCapsule3>>();
+								check(ImplicitScaled.Object());
+								const FImplicitCapsule3& ImplicitCapsule = *ImplicitScaled.GetUnscaledObject();
+								ensure(FVector::DistSquared(ComponentToClothTransform.GetScale3D(), FVector(ImplicitScaled.GetScale())) < KINDA_SMALL_NUMBER);
+								AddCapsule(ImplicitCapsule, ComponentToClothTransform, ImplicitScaled.GetScale());
+								UE_LOG(LogSkeletalMesh, Verbose, TEXT("Found Scaled Capsule cloth environmental collision in [%s]"), *Component->GetOwner()->GetFName().ToString());
+								bHasSimpleCollision = true;
+							}
+							break;
+						case ImplicitObjectType::IsScaled | ImplicitObjectType::Convex:
+							{
+								const TImplicitObjectScaled<FImplicitConvex3>& ImplicitScaled = ImplicitObject->GetObjectChecked<TImplicitObjectScaled<FImplicitConvex3>>();
+								check(ImplicitScaled.Object());
+								const FImplicitConvex3& ImplicitConvex = *ImplicitScaled.GetUnscaledObject();
+								ensure(FVector::DistSquared(ComponentToClothTransform.GetScale3D(), FVector(ImplicitScaled.GetScale())) < KINDA_SMALL_NUMBER);
+								AddConvex(ImplicitConvex, ComponentToClothTransform.ToMatrixWithScale());
+								UE_LOG(LogSkeletalMesh, Verbose, TEXT("Found Scaled Convex cloth environmental collision in [%s]"), *Component->GetOwner()->GetFName().ToString());
+								bHasSimpleCollision = true;
+							}
+							break;
+
+						// Triangle mesh
+						case ImplicitObjectType::TriangleMesh:
+						case ImplicitObjectType::IsInstanced | ImplicitObjectType::TriangleMesh:
+						case ImplicitObjectType::IsScaled | ImplicitObjectType::TriangleMesh:
+							// TODO: We could eventually want to collide cloth against triangle meshes,
+							//       however the concept of simple vs complex shape might need to be clarified
+							//       as it currently iterates over all shape to discard the triangle mesh ones.
+							UE_LOG(LogSkeletalMesh, Verbose, TEXT("Found unusable Triangle Mesh cloth environmental collision in [%s]"), !Component->GetOwner() ? TEXT("Unknown") : *Component->GetOwner()->GetFName().ToString());
+							break;
+
+						default: 
+							ensure(false);  // Is there a missing collision type?
+							break;
+						}
+					}
+#else
+					static_assert(false, "A physics engine interface must be defined to build");
+#endif
 					bSuccessfulRead = true;
 				});
 			}
@@ -3337,29 +3606,26 @@ void USkeletalMeshComponent::EndPhysicsTickComponent(FSkeletalMeshComponentEndPh
 
 void USkeletalMeshComponent::UpdateClothTransformImp()
 {
-	const bool bActiveClothing = ClothingSimulation && ClothingSimulation->ShouldSimulate();
-
 #if WITH_CLOTH_COLLISION_DETECTION
-
-	if(ClothingSimulation)
+	if (ClothingSimulation && ClothingSimulation->ShouldSimulate())
 	{
+		// Component has moved, and there is something to simulate, update all external cloth collisions
 		ClothingSimulation->ClearExternalCollisions();
-	}
 
-	if(bCollideWithAttachedChildren)
-	{
-		CopyClothCollisionsToChildren();
-	}
+		if (bCollideWithAttachedChildren)
+		{
+			CopyClothCollisionsToChildren();
+		}
 
-	if(ClothCollisionSources.Num() > 0)
-	{
-		CopyClothCollisionSources();
-	}
+		if (ClothCollisionSources.Num() > 0)
+		{
+			CopyClothCollisionSources();
+		}
 
-	//check the environment when only transform is updated
-	if(bCollideWithEnvironment && bActiveClothing)
-	{
-		ProcessClothCollisionWithEnvironment();
+		if (bCollideWithEnvironment)
+		{
+			ProcessClothCollisionWithEnvironment();
+		}
 	}
 #endif // WITH_CLOTH_COLLISION_DETECTION
 
@@ -3593,14 +3859,14 @@ const TMap<int32, FClothSimulData>& USkeletalMeshComponent::GetCurrentClothingDa
 
 void USkeletalMeshComponent::UpdateClothStateAndSimulate(float DeltaTime, FTickFunction& ThisTickFunction)
 {
+	check(IsInGameThread());
+
 	// If disabled or no simulation
 	if (CVarEnableClothPhysics.GetValueOnGameThread() == 0 || !ClothingSimulation || bDisableClothSimulation)
 	{
 		return;
 	}
 
-	check(IsInGameThread());
-
 	// If we simulate a clothing actor at 0s it will fill simulated positions and normals with NaNs.
 	// we can skip all the work it is still doing, and get the desired result (frozen sim) by not
 	// updating and simulating.
@@ -3615,10 +3881,7 @@ void USkeletalMeshComponent::UpdateClothStateAndSimulate(float DeltaTime, FTickF
 #if WITH_CLOTH_COLLISION_DETECTION
 	if (bCollideWithAttachedChildren)
 	{
-		if(ClothingSimulation)
-		{
-			ClothingSimulation->ClearExternalCollisions();
-		}
+		ClothingSimulation->ClearExternalCollisions();
 
 		CopyClothCollisionsToChildren();
 		CopyChildrenClothCollisionsToParent();
@@ -3627,17 +3890,14 @@ void USkeletalMeshComponent::UpdateClothStateAndSimulate(float DeltaTime, FTickF
 
 	UpdateClothSimulationContext(DeltaTime);
 
-	if(ClothingSimulation)
+	ParallelClothTask = TGraphTask<FParallelClothTask>::CreateTask(nullptr, ENamedThreads::GameThread).ConstructAndDispatchWhenReady(*this, DeltaTime);
+
+	if (ShouldWaitForClothInTickFunction())
 	{
-		ParallelClothTask = TGraphTask<FParallelClothTask>::CreateTask(nullptr, ENamedThreads::GameThread).ConstructAndDispatchWhenReady(*this, DeltaTime);
-		
-		if (ShouldWaitForClothInTickFunction())
-		{
-			FGraphEventArray Prerequisites;
-			Prerequisites.Add(ParallelClothTask);
-			FGraphEventRef ClothCompletionEvent = TGraphTask<FParallelClothCompletionTask>::CreateTask(&Prerequisites, ENamedThreads::GameThread).ConstructAndDispatchWhenReady(this);
-			ThisTickFunction.GetCompletionHandle()->DontCompleteUntil(ClothCompletionEvent);
-		}
+		FGraphEventArray Prerequisites;
+		Prerequisites.Add(ParallelClothTask);
+		FGraphEventRef ClothCompletionEvent = TGraphTask<FParallelClothCompletionTask>::CreateTask(&Prerequisites, ENamedThreads::GameThread).ConstructAndDispatchWhenReady(this);
+		ThisTickFunction.GetCompletionHandle()->DontCompleteUntil(ClothCompletionEvent);
 	}
 }
 
