Non-Unity Fix

diff --git a/Engine/Source/Runtime/Engine/Private/HLSLTree/HLSLTree.cpp b/Engine/Source/Runtime/Engine/Private/HLSLTree/HLSLTree.cpp
--- a/Engine/Source/Runtime/Engine/Private/HLSLTree/HLSLTree.cpp
+++ b/Engine/Source/Runtime/Engine/Private/HLSLTree/HLSLTree.cpp
@@ -1,7 +1,8 @@
 // Copyright Epic Games, Inc. All Rights Reserved.
 #include "HLSLTree/HLSLTree.h"
 #include "Misc/StringBuilder.h"
-#include "Misc/MemStackUtility.h"
+#include "Misc/MemStack.h"
+#include "Shader/ShaderTypes.h"
 #include "MaterialShared.h" // TODO - split preshader out into its own module
 
 static FSHAHash HashString(const TCHAR* String, int32 Length)
