Non-Unity Fix

diff --git a/Engine/Source/Runtime/Engine/Private/HLSLTree/HLSLTree.cpp b/Engine/Source/Runtime/Engine/Private/HLSLTree/HLSLTree.cpp
--- a/Engine/Source/Runtime/Engine/Private/HLSLTree/HLSLTree.cpp
+++ b/Engine/Source/Runtime/Engine/Private/HLSLTree/HLSLTree.cpp
@@ -2,6 +2,7 @@
 #include "HLSLTree/HLSLTree.h"
 #include "Misc/StringBuilder.h"
 #include "Misc/MemStack.h"
+#include "Shader/ShaderTypes.h"
 #include "MaterialShared.h" // TODO - split preshader out into its own module
 
 UE::HLSLTree::EExpressionEvaluationType UE::HLSLTree::CombineEvaluationTypes(EExpressionEvaluationType Lhs, EExpressionEvaluationType Rhs)
