Added anim instance proxy debug draws for circles and cones
#rnx

diff --git a/Engine/Source/Runtime/Engine/Private/Animation/AnimInstanceProxy.cpp b/Engine/Source/Runtime/Engine/Private/Animation/AnimInstanceProxy.cpp
--- a/Engine/Source/Runtime/Engine/Private/Animation/AnimInstanceProxy.cpp
+++ b/Engine/Source/Runtime/Engine/Private/Animation/AnimInstanceProxy.cpp
@@ -522,11 +522,20 @@ void FAnimInstanceProxy::PostUpdate(UAnimInstance* InAnimInstance) const
 		switch (DebugItem.ItemType)
 		{
 			case EDrawDebugItemType::OnScreenMessage: GEngine->AddOnScreenDebugMessage(INDEX_NONE, 0.f, DebugItem.Color, DebugItem.Message, false, DebugItem.TextScale); break;
-			case EDrawDebugItemType::DirectionalArrow: DrawDebugDirectionalArrow(InAnimInstance->GetSkelMeshComponent()->GetWorld(), DebugItem.StartLoc, DebugItem.EndLoc, DebugItem.Size, DebugItem.Color, DebugItem.bPersistentLines, DebugItem.LifeTime, DebugItem.DepthPriority, DebugItem.Thickness); break;
+			case EDrawDebugItemType::DirectionalArrow: DrawDebugDirectionalArrow(InAnimInstance->GetSkelMeshComponent()->GetWorld(), DebugItem.StartLoc, DebugItem.EndLoc.Value, DebugItem.Size, DebugItem.Color, DebugItem.bPersistentLines, DebugItem.LifeTime, DebugItem.DepthPriority, DebugItem.Thickness); break;
 			case EDrawDebugItemType::Sphere: DrawDebugSphere(InAnimInstance->GetSkelMeshComponent()->GetWorld(), DebugItem.Center, DebugItem.Radius, DebugItem.Segments, DebugItem.Color, DebugItem.bPersistentLines, DebugItem.LifeTime, DebugItem.DepthPriority, DebugItem.Thickness); break;
-			case EDrawDebugItemType::Line: DrawDebugLine(InAnimInstance->GetSkelMeshComponent()->GetWorld(), DebugItem.StartLoc, DebugItem.EndLoc, DebugItem.Color, DebugItem.bPersistentLines, DebugItem.LifeTime, DebugItem.DepthPriority, DebugItem.Thickness); break;
+			case EDrawDebugItemType::Line: DrawDebugLine(InAnimInstance->GetSkelMeshComponent()->GetWorld(), DebugItem.StartLoc, DebugItem.EndLoc.Value, DebugItem.Color, DebugItem.bPersistentLines, DebugItem.LifeTime, DebugItem.DepthPriority, DebugItem.Thickness); break;
 			case EDrawDebugItemType::CoordinateSystem: DrawDebugCoordinateSystem(InAnimInstance->GetSkelMeshComponent()->GetWorld(), DebugItem.StartLoc, DebugItem.Rotation, DebugItem.Size, DebugItem.bPersistentLines, DebugItem.LifeTime, DebugItem.DepthPriority, DebugItem.Thickness); break;
 			case EDrawDebugItemType::Point: DrawDebugPoint(InAnimInstance->GetSkelMeshComponent()->GetWorld(), DebugItem.StartLoc, DebugItem.Size, DebugItem.Color, DebugItem.bPersistentLines, DebugItem.LifeTime, DebugItem.DepthPriority); break;
+			case EDrawDebugItemType::Circle:
+			{
+				const FMatrix RotationMatrix = FRotationMatrix::MakeFromZ(DebugItem.Direction.Value);
+				const FVector ForwardVector = RotationMatrix.GetScaledAxis(EAxis::X);
+				const FVector RightVector = RotationMatrix.GetScaledAxis(EAxis::Y);
+				DrawDebugCircle(InAnimInstance->GetSkelMeshComponent()->GetWorld(), DebugItem.Center, DebugItem.Radius, DebugItem.Segments, DebugItem.Color, DebugItem.bPersistentLines, DebugItem.LifeTime, DebugItem.DepthPriority, DebugItem.Thickness, ForwardVector, RightVector, false);
+				break;
+			}
+			case EDrawDebugItemType::Cone: DrawDebugCone(InAnimInstance->GetSkelMeshComponent()->GetWorld(), DebugItem.Center, DebugItem.Direction.Value, DebugItem.Length, DebugItem.AngleWidth, DebugItem.AngleHeight, DebugItem.Segments, DebugItem.Color, DebugItem.bPersistentLines, DebugItem.LifeTime, DebugItem.DepthPriority, DebugItem.Thickness); break;
 		}
 	}
 #endif
@@ -2067,6 +2076,48 @@ void FAnimInstanceProxy::AnimDrawDebugPoint(const FVector& Loc, float Size, cons
 
 	QueuedDrawDebugItems.Add(DrawDebugItem);
 }
+
+void FAnimInstanceProxy::AnimDrawDebugCircle(const FVector& Center, float Radius, int32 Segments, const FColor& Color, const FVector& UpVector, bool bPersistentLines, float LifeTime, ESceneDepthPriorityGroup DepthPriority, float Thickness)
+{
+	FQueuedDrawDebugItem DrawDebugItem;
+
+	DrawDebugItem.ItemType = EDrawDebugItemType::Circle;
+	DrawDebugItem.Center = Center;
+	DrawDebugItem.Radius = Radius;
+	DrawDebugItem.Segments = Segments;
+	DrawDebugItem.Color = Color;
+	// We're using EndLoc as our direction unit vector.
+	DrawDebugItem.EndLoc = UpVector.GetSafeNormal();
+	DrawDebugItem.bPersistentLines = bPersistentLines;
+	DrawDebugItem.LifeTime = LifeTime;
+	DrawDebugItem.DepthPriority = DepthPriority;
+	DrawDebugItem.Thickness = Thickness;
+
+	QueuedDrawDebugItems.Add(DrawDebugItem);
+}
+
+void FAnimInstanceProxy::AnimDrawDebugCone(const FVector& Center, float Length, const FVector& Direction, float AngleWidth, float AngleHeight, int32 Segments, const FColor & Color, bool bPersistentLines, float LifeTime, ESceneDepthPriorityGroup DepthPriority, float Thickness)
+{
+	FQueuedDrawDebugItem DrawDebugItem;
+
+	DrawDebugItem.ItemType = EDrawDebugItemType::Cone;
+	DrawDebugItem.Center = Center;
+	DrawDebugItem.Length = Length;
+	DrawDebugItem.AngleWidth = AngleWidth;
+	DrawDebugItem.AngleHeight = AngleHeight;
+	// We're using EndLoc as our direction unit vector.
+	DrawDebugItem.Direction = Direction.GetSafeNormal();
+	DrawDebugItem.Segments = Segments;
+	DrawDebugItem.Color = Color;
+	DrawDebugItem.bPersistentLines = bPersistentLines;
+	DrawDebugItem.LifeTime = LifeTime;
+	DrawDebugItem.DepthPriority = DepthPriority;
+	DrawDebugItem.Thickness = Thickness;
+
+	QueuedDrawDebugItems.Add(DrawDebugItem);
+}
+
+
 #endif // ENABLE_ANIM_DRAW_DEBUG
 
 float FAnimInstanceProxy::GetInstanceAssetPlayerLength(int32 AssetPlayerIndex) const
diff --git a/Engine/Source/Runtime/Engine/Public/Animation/AnimInstanceProxy.h b/Engine/Source/Runtime/Engine/Public/Animation/AnimInstanceProxy.h
--- a/Engine/Source/Runtime/Engine/Public/Animation/AnimInstanceProxy.h
+++ b/Engine/Source/Runtime/Engine/Public/Animation/AnimInstanceProxy.h
@@ -54,58 +54,57 @@ namespace EDrawDebugItemType
 		OnScreenMessage,
 		CoordinateSystem,
 		Point,
+		Circle,
+		Cone,
 	};
 }
 
-USTRUCT()
 struct FQueuedDrawDebugItem 
 {
-	GENERATED_BODY()
+	struct FVectorEntry
+	{
+		FVectorEntry(const FVector& InVector)
+			: Value(InVector)
+		{}
 
-	UPROPERTY(Transient)
-	TEnumAsByte<EDrawDebugItemType::Type> ItemType = EDrawDebugItemType::DirectionalArrow;
+		FVector Value;
+	};
 
-	UPROPERTY(Transient)
-	FVector StartLoc = FVector(0.f);
+	TEnumAsByte<EDrawDebugItemType::Type> ItemType = EDrawDebugItemType::DirectionalArrow;
 
-	UPROPERTY(Transient)
-	FVector EndLoc = FVector(0.f);
+	union
+	{
+		FVector StartLoc;
+		struct
+		{
+			float Length;
+			float AngleWidth;
+			float AngleHeight;
+		};
+	};
+	
+	union
+	{
+		FVectorEntry EndLoc;
+		FVectorEntry Direction;
+	};
 
-	UPROPERTY(Transient)
 	FVector Center = FVector(0.f);
-
-	UPROPERTY(Transient)
 	FRotator Rotation = FRotator(0.f);
-
-	UPROPERTY(Transient)
 	float Radius = 0.f;
-
-	UPROPERTY(Transient)
 	float Size = 0.f;
-
-	UPROPERTY(Transient)
 	int32 Segments = 0;
-
-	UPROPERTY(Transient)
 	FColor Color = FColor(0);
-
-	UPROPERTY(Transient)
 	bool bPersistentLines = false;
-
-	UPROPERTY(Transient)
 	float LifeTime = 0.f;
-
-	UPROPERTY(Transient)
 	float Thickness = 0.f;
-
-	UPROPERTY(Transient)
 	FString Message;
-
-	UPROPERTY(Transient)
 	FVector2D TextScale = FVector2D(0.f);
-
-	UPROPERTY(Transient)
 	TEnumAsByte<ESceneDepthPriorityGroup> DepthPriority = SDPG_World;
+	FQueuedDrawDebugItem() 
+		: StartLoc(FVector(0.f))
+		, EndLoc(FVector(0.f))
+	{};
 };
 
 /** Proxy object passed around during animation tree update in lieu of a UAnimInstance */
@@ -491,6 +490,8 @@ public:
 	void AnimDrawDebugCoordinateSystem(FVector const& AxisLoc, FRotator const& AxisRot, float Scale = 1.f, bool bPersistentLines = false, float LifeTime = -1.f, float Thickness = 0.f, ESceneDepthPriorityGroup DepthPriority = SDPG_World);
 	void AnimDrawDebugPlane(const FTransform& BaseTransform, float Radii, const FColor& Color, bool bPersistentLines = false, float LifeTime = -1.f, float Thickness = 0.f, ESceneDepthPriorityGroup DepthPriority = SDPG_World);
 	void AnimDrawDebugPoint(const FVector& Loc, float Size, const FColor& Color, bool bPersistentLines = false, float LifeTime = -1.f, ESceneDepthPriorityGroup DepthPriority = SDPG_World);
+	void AnimDrawDebugCircle(const FVector& Center, float Radius, int32 Segments, const FColor& Color, const FVector& UpVector = FVector::UpVector, bool bPersistentLines = false, float LifeTime = -1.f, ESceneDepthPriorityGroup DepthPriority = SDPG_World, float Thickness = 0.f);
+	void AnimDrawDebugCone(const FVector& Center, float Radius, const FVector& Direction, float AngleWidth, float AngleHeight, int32 Segments, const FColor& Color, bool bPersistentLines = false, float LifeTime = -1.f, ESceneDepthPriorityGroup DepthPriority = SDPG_World, float Thickness = 0.f);
 #else
 	void AnimDrawDebugOnScreenMessage(const FString& DebugMessage, const FColor& Color, const FVector2D& TextScale = FVector2D::UnitVector, ESceneDepthPriorityGroup DepthPriority = SDPG_World) {}
 	void AnimDrawDebugLine(const FVector& StartLoc, const FVector& EndLoc, const FColor& Color, bool bPersistentLines = false, float LifeTime = -1.f, float Thickness = 0.f, ESceneDepthPriorityGroup DepthPriority = SDPG_World) {}
@@ -499,6 +500,8 @@ public:
 	void AnimDrawDebugCoordinateSystem(FVector const& AxisLoc, FRotator const& AxisRot, float Scale = 1.f, bool bPersistentLines = false, float LifeTime = -1.f, float Thickness = 0.f, ESceneDepthPriorityGroup DepthPriority = SDPG_World) {}
 	void AnimDrawDebugPlane(const FTransform& BaseTransform, float Radii, const FColor& Color, bool bPersistentLines = false, float LifeTime = -1.f, float Thickness = 0.f, ESceneDepthPriorityGroup DepthPriority = SDPG_World) {}
 	void AnimDrawDebugPoint(const FVector& Loc, float Size, const FColor& Color, bool bPersistentLines = false, float LifeTime = -1.f, ESceneDepthPriorityGroup DepthPriority = SDPG_World) {}
+	void AnimDrawDebugCircle(const FVector& Center, float Radius, int32 Segments, const FColor& Color, const FVector& UpVector = FVector::UpVector, bool bPersistentLines = false, float LifeTime=-1.f, ESceneDepthPriorityGroup DepthPriority = SDPG_World, float Thickness = 0.f) {}
+	void AnimDrawDebugCone(const FVector& Center, float Radius, const FVector& Direction, float AngleWidth, float AngleHeight, int32 Segments, const FColor& Color, bool bPersistentLines = false, float LifeTime = -1.f, ESceneDepthPriorityGroup DepthPriority = SDPG_World, float Thickness = 0.f) {}
 #endif // ENABLE_ANIM_DRAW_DEBUG
 
 #if ENABLE_ANIM_LOGGING
