HLSL Tree Fix for CIS

diff --git a/Engine/Source/Runtime/Engine/Private/HLSLTree/HLSLTree.cpp b/Engine/Source/Runtime/Engine/Private/HLSLTree/HLSLTree.cpp
--- a/Engine/Source/Runtime/Engine/Private/HLSLTree/HLSLTree.cpp
+++ b/Engine/Source/Runtime/Engine/Private/HLSLTree/HLSLTree.cpp
@@ -2,6 +2,7 @@
 #include "HLSLTree/HLSLTree.h"
 #include "Misc/StringBuilder.h"
 #include "Misc/MemStack.h"
+#include "Misc/MemStackUtility.h"
 #include "Shader/ShaderTypes.h"
 #include "MaterialShared.h" // TODO - split preshader out into its own module
 
