Fix FPoly::OptimizeIntoConvexPolys() creating polygons that fail coplanar check in FPoly::IsCoplanar()

diff --git a/Engine/Plugins/Runtime/MeshModelingToolset/Source/ModelingComponentsEditorOnly/Private/ConversionUtils/DynamicMeshToVolume.cpp b/Engine/Plugins/Runtime/MeshModelingToolset/Source/ModelingComponentsEditorOnly/Private/ConversionUtils/DynamicMeshToVolume.cpp
--- a/Engine/Plugins/Runtime/MeshModelingToolset/Source/ModelingComponentsEditorOnly/Private/ConversionUtils/DynamicMeshToVolume.cpp
+++ b/Engine/Plugins/Runtime/MeshModelingToolset/Source/ModelingComponentsEditorOnly/Private/ConversionUtils/DynamicMeshToVolume.cpp
@@ -133,10 +133,8 @@ static void OptimizeIntoConvexPolys(ABrush* InOwnerBrush, TArray<FPoly>& InPolyg
 						PolyMerged.Vertices[PolyMainIndex1 + i] = PolyNeighbor->Vertices[NeighborIndex];
 					}
 
-					PolyMerged.CalcNormal(true);
-
-					// Check if the result is valid
-					if (PolyMerged.IsConvex() && PolyMerged.Finalize(InOwnerBrush, 1) == 0)
+					// Check if the result is both coplanar and convex, and then finalize it.
+					if (PolyMerged.IsCoplanar() && PolyMerged.IsConvex() && PolyMerged.Finalize(InOwnerBrush, 1) == 0)
 					{
 						// Remove the original polygons from the list. Uses knowledge that p2 > PolyIndex
 						InPolygons.RemoveAtSwap(p2);
diff --git a/Engine/Source/Runtime/Engine/Private/Polygon.cpp b/Engine/Source/Runtime/Engine/Private/Polygon.cpp
--- a/Engine/Source/Runtime/Engine/Private/Polygon.cpp
+++ b/Engine/Source/Runtime/Engine/Private/Polygon.cpp
@@ -678,7 +678,8 @@ bool FPoly::IsCoplanar()
 
 	CalcNormal(1);
 
-	for( int32 x = 0 ; x < Vertices.Num() ; ++x )
+	// Note that we do not need to test the first vertex since the OnPlane() check uses it as the "reference" vertex, and always return true for it. 
+	for (int32 x = 1; x < Vertices.Num(); ++x)
 	{
 		if( !OnPlane( Vertices[x] ) )
 		{
@@ -1027,7 +1028,7 @@ void FPoly::OptimizeIntoConvexPolys(ABrush* InOwnerBrush, ArrayType& InPolygons)
 
 					if (bSharedEdge)
 					{
-						// Found a shared edge. Merge both polygons around the shared edge and check if the result is convex.
+						// Found a shared edge. Merge both polygons around the shared edge and check if the result is coplanar and convex.
 
 						// Note that adding the vertices and polygons in this very order keeps the repeated execution of "Triangulate" and "Optimize"
 						// more stable, i.e. when you have a polygon and click "Triangulate" followed by "Optimize" it will more likely give you the
@@ -1050,8 +1051,8 @@ void FPoly::OptimizeIntoConvexPolys(ABrush* InOwnerBrush, ArrayType& InPolygons)
 						PolyMerged.Vertices[PolyMergedIndex] = EdgeVtxA;
 						checkSlow(++PolyMergedIndex == PolyMerged.Vertices.Num());
 
-						// Computing the normal is not cheap, but we need it for the convexity test. Thus, we just use the normals we already have since it
-						// is good enough for that. The proper normal will be computed during the later call to `FPoly::Finalize()` anyway.
+						// Computing the normal is not cheap, but we need it for the convexity test. Thus, we just use the normals we already have since they
+						// are close enough anyway. The proper normal will be computed during the later call to `FPoly::Finalize()` anyway.
 						PolyMerged.Normal = (PolyMain->Normal + PolyNeighbor->Normal) * 0.5f;
 
 						// Check convexity equivalent to FPoly::IsConvex(), but also remove collinear edges in the process.
@@ -1104,8 +1105,10 @@ void FPoly::OptimizeIntoConvexPolys(ABrush* InOwnerBrush, ArrayType& InPolygons)
 							return true;
 						}();
 
-						// If the merged polygon is convex, finalize it, and enter it into the list.
-						if (bIsConvex && PolyMerged.Finalize(InOwnerBrush, 1) == 0)
+						// If the merged polygon is convex and coplanar, finalize it, and enter it into the list.
+						// We explicitly call FPoly::IsCoplanar() here even though our own test already established that to make sure that we don't end up
+						// in infinite loops of failing coplanar check, triangulation, and optimization due to precision issues in the coplanar check.
+						if (bIsConvex && PolyMerged.IsCoplanar() && PolyMerged.Finalize(InOwnerBrush, 1) == 0)
 						{
 							// If the main polygon is not last one or the neighbor polygon is not the second to last one, we might have missed neighbors.
 							bMightHaveMissedSomeNeighbors |= PolyMainIndex != InPolygons.Num() - 1 || PolyNeighborIndex != PolyMainIndex - 1;
