Allow disabling hardware occlusion queries on a per-scene basis in editor.

#rnx

diff --git a/Engine/Plugins/Experimental/UVEditor/Source/UVEditor/Private/UVEditorToolkit.cpp b/Engine/Plugins/Experimental/UVEditor/Source/UVEditor/Private/UVEditorToolkit.cpp
--- a/Engine/Plugins/Experimental/UVEditor/Source/UVEditor/Private/UVEditorToolkit.cpp
+++ b/Engine/Plugins/Experimental/UVEditor/Source/UVEditor/Private/UVEditorToolkit.cpp
@@ -457,6 +457,11 @@ void FUVEditorToolkit::PostInitAssetEditor()
 
 		// Disable the dithering of occluded portions of gizmos.
 		Client->EngineShowFlags.SetOpaqueCompositeEditorPrimitives(true);
+
+		// Disable hardware occlusion queries, which make it harder to use vertex shaders to pull materials
+		// toward camera for z ordering because non-translucent materials start occluding themselves (once
+		// the component bounds are behind the displaced geometry).
+		Client->EngineShowFlags.SetDisableOcclusionQueries(true);
 	};
 	SetCommonViewportClientOptions(ViewportClient.Get());
 
diff --git a/Engine/Source/Runtime/Engine/Public/ShowFlags.h b/Engine/Source/Runtime/Engine/Public/ShowFlags.h
--- a/Engine/Source/Runtime/Engine/Public/ShowFlags.h
+++ b/Engine/Source/Runtime/Engine/Public/ShowFlags.h
@@ -459,6 +459,7 @@ private:
 		SetScreenPercentage(InitMode != ESFIM_Editor && InitMode != ESFIM_VREditing);
 		SetVREditing(InitMode == ESFIM_VREditing);
 		SetOcclusionMeshes(false);
+		SetDisableOcclusionQueries(false);
 		SetPathTracing(false);
 		SetRayTracingDebug(false);
 		SetVisualizeSkyAtmosphere(false);
diff --git a/Engine/Source/Runtime/Engine/Public/ShowFlagsValues.inl b/Engine/Source/Runtime/Engine/Public/ShowFlagsValues.inl
--- a/Engine/Source/Runtime/Engine/Public/ShowFlagsValues.inl
+++ b/Engine/Source/Runtime/Engine/Public/ShowFlagsValues.inl
@@ -367,6 +367,8 @@ SHOWFLAG_ALWAYS_ACCESSIBLE(MediaPlanes, SFG_Normal, NSLOCTEXT("UnrealEd", "Media
 SHOWFLAG_FIXED_IN_SHIPPING(0, VREditing, SFG_Hidden, NSLOCTEXT("UnrealEd", "VREditSF", "VR Editing"))
 /** Visualize Occlusion Query bounding meshes */
 SHOWFLAG_FIXED_IN_SHIPPING(0, OcclusionMeshes, SFG_Visualize, NSLOCTEXT("UnrealEd", "VisualizeOcclusionQueries", "Visualize Occlusion Queries"))
+/** Disable hardware occlusion queries, similar to setting r.AllowOcclusionQueries=0, but just for this scene. */
+SHOWFLAG_FIXED_IN_SHIPPING(0, DisableOcclusionQueries, SFG_Developer, NSLOCTEXT("UnrealEd", "DisableOcclusionQueries", "Disable Hardware Occlusion Queries"))
 
 // RHI_RAYTRACING begin
 SHOWFLAG_ALWAYS_ACCESSIBLE(PathTracing, SFG_Hidden, NSLOCTEXT("UnrealEd", "PathTracing", "Path tracing"))
diff --git a/Engine/Source/Runtime/Renderer/Private/DeferredShadingRenderer.cpp b/Engine/Source/Runtime/Renderer/Private/DeferredShadingRenderer.cpp
--- a/Engine/Source/Runtime/Renderer/Private/DeferredShadingRenderer.cpp
+++ b/Engine/Source/Runtime/Renderer/Private/DeferredShadingRenderer.cpp
@@ -2235,7 +2235,8 @@ void FDeferredShadingSceneRenderer::Render(FRDGBuilder& GraphBuilder)
 	#endif
 
 	
-	const bool bIsOcclusionTesting = DoOcclusionQueries() && (!ViewFamily.EngineShowFlags.Wireframe || bIsViewFrozen || bHasViewParent);
+	const bool bIsOcclusionTesting = DoOcclusionQueries() && !ViewFamily.EngineShowFlags.DisableOcclusionQueries
+		&& (!ViewFamily.EngineShowFlags.Wireframe || bIsViewFrozen || bHasViewParent);
 	const bool bNeedsPrePass = ShouldRenderPrePass();
 
 	GEngine->GetPreRenderDelegateEx().Broadcast(GraphBuilder);
diff --git a/Engine/Source/Runtime/Renderer/Private/SceneVisibility.cpp b/Engine/Source/Runtime/Renderer/Private/SceneVisibility.cpp
--- a/Engine/Source/Runtime/Renderer/Private/SceneVisibility.cpp
+++ b/Engine/Source/Runtime/Renderer/Private/SceneVisibility.cpp
@@ -3573,7 +3573,7 @@ void FSceneRenderer::PreVisibilityFrameSetup(FRDGBuilder& GraphBuilder, const FS
 		const bool bIsHitTesting = ViewFamily.EngineShowFlags.HitProxies;
 		// Don't test occlusion queries in collision viewmode as they can be bigger then the rendering bounds.
 		const bool bCollisionView = ViewFamily.EngineShowFlags.CollisionVisibility || ViewFamily.EngineShowFlags.CollisionPawn;
-		if (GIsHighResScreenshot || !DoOcclusionQueries() || bIsHitTesting || bCollisionView)
+		if (GIsHighResScreenshot || !DoOcclusionQueries() || bIsHitTesting || bCollisionView || ViewFamily.EngineShowFlags.DisableOcclusionQueries)
 		{
 			View.bDisableQuerySubmissions = true;
 			View.bIgnoreExistingQueries = true;
