Auto-upgrading from Separable to Burley.

 'Tint' is the same to the original `Subsurface Color` in the profile to enable direct separable to burley upgrade. It also introduces a capability to add Non-PBR into SSS profile. Stay PBR with the default value (white).

diff --git a/Engine/Source/Editor/UnrealEd/Private/Factories/EditorFactories.cpp b/Engine/Source/Editor/UnrealEd/Private/Factories/EditorFactories.cpp
--- a/Engine/Source/Editor/UnrealEd/Private/Factories/EditorFactories.cpp
+++ b/Engine/Source/Editor/UnrealEd/Private/Factories/EditorFactories.cpp
@@ -7675,7 +7675,14 @@ USubsurfaceProfileFactory::USubsurfaceProfileFactory(const FObjectInitializer& O
 
 UObject* USubsurfaceProfileFactory::FactoryCreateNew(UClass* InClass, UObject* InParent, FName InName, EObjectFlags Flags, UObject* Context, FFeedbackContext* Warn)
 {
-	return NewObject<USubsurfaceProfile>(InParent, InName, Flags);
+	USubsurfaceProfile* Object = NewObject<USubsurfaceProfile>(InParent, InName, Flags);
+	
+	// Enable smooth upgrading from Separable to Burley. This flag change
+	// allows all newly created subsurface profile in Burley, while allowing
+	// the profiles loaded from files to be converted into burley if they
+	// are separable.
+	Object->Settings.bEnableBurley = true;
+	return Object;
 }
 
 /*-----------------------------------------------------------------------------
diff --git a/Engine/Source/Editor/UnrealEd/Private/ThumbnailRendering/SubsurfaceProfileRenderer.cpp b/Engine/Source/Editor/UnrealEd/Private/ThumbnailRendering/SubsurfaceProfileRenderer.cpp
--- a/Engine/Source/Editor/UnrealEd/Private/ThumbnailRendering/SubsurfaceProfileRenderer.cpp
+++ b/Engine/Source/Editor/UnrealEd/Private/ThumbnailRendering/SubsurfaceProfileRenderer.cpp
@@ -25,13 +25,13 @@ void USubsurfaceProfileRenderer::Draw(UObject* Object, int32 X, int32 Y, uint32
 	{
 		FLinearColor Col;
 
-		Col = LocalSubsurfaceProfile->Settings.SubsurfaceColor; Col.A = 1;
+		Col = LocalSubsurfaceProfile->Settings.Tint; Col.A = 1;
 		Canvas->DrawTile(0,          0, Width, Height / 2, 0, 0, 1, 1, Col);
-		Col = LocalSubsurfaceProfile->Settings.FalloffColor; Col.A = 1;
+		Col = LocalSubsurfaceProfile->Settings.SurfaceAlbedo; Col.A = 1;
 		Canvas->DrawTile(0, Height / 2, Width, Height / 2, 0, 0, 1, 1, Col);
 
-		FText ScatterRadiusText = FText::AsNumber(LocalSubsurfaceProfile->Settings.ScatterRadius);
-		FCanvasTextItem TextItem(FVector2D(5.0f, 5.0f), ScatterRadiusText, GEngine->GetLargeFont(), FLinearColor::White);
+		FText MeanFreePathDistanceText = FText::AsNumber(LocalSubsurfaceProfile->Settings.MeanFreePathDistance);
+		FCanvasTextItem TextItem(FVector2D(5.0f, 5.0f), MeanFreePathDistanceText, GEngine->GetLargeFont(), FLinearColor::White);
 		TextItem.EnableShadow(FLinearColor::Black);
 		TextItem.Scale = FVector2D(Width / 128.0f, Height / 128.0f);
 		TextItem.Draw(Canvas);
diff --git a/Engine/Source/Runtime/Engine/Classes/Engine/SubsurfaceProfile.h b/Engine/Source/Runtime/Engine/Classes/Engine/SubsurfaceProfile.h
--- a/Engine/Source/Runtime/Engine/Classes/Engine/SubsurfaceProfile.h
+++ b/Engine/Source/Runtime/Engine/Classes/Engine/SubsurfaceProfile.h
@@ -40,11 +40,19 @@ struct FSubsurfaceProfileStruct
 	/**
 	* Effective only when Burley subsurface scattering is enabled in cmd.
 	*/
-	UPROPERTY(Category = "Burley Normalized", EditAnywhere, BlueprintReadOnly)
+	UPROPERTY(Category = "Burley Normalized", EditAnywhere, BlueprintReadOnly, meta = (editcondition = "false", EditConditionHides))
 	bool  bEnableBurley;
 
+	/**
+	* Specifies the how much of the diffuse light gets into the material,
+	* can be seen as a per-channel mix factor between the original image,
+	* and the SSS-filtered image. It introduces Non-PBR looks.
+	*/
+	UPROPERTY(Category = "Burley Normalized", EditAnywhere, BlueprintReadOnly, meta = (HideAlphaChannel, editcondition = "bEnableBurley"))
+	FLinearColor Tint;
+
 	/** in world/unreal units (cm) */
-	UPROPERTY(Category = "SubsurfaceProfileStruct", EditAnywhere, BlueprintReadOnly, meta = (ClampMin = "0.1", UIMax = "50.0", ClampMax = "1000.0", editcondition = "!bEnableBurley"))
+	UPROPERTY(Category = "SubsurfaceProfileStruct", EditAnywhere, BlueprintReadOnly, meta = (ClampMin = "0.1", UIMax = "50.0", ClampMax = "1000.0", DeprecatedProperty, editcondition = "!bEnableBurley", EditConditionHides))
 	float ScatterRadius;
 
 	/**
@@ -52,7 +60,7 @@ struct FSubsurfaceProfileStruct
 	* can be seen as a per-channel mix factor between the original image,
 	* and the SSS-filtered image (called "strength" in SeparableSSS, default there: 0.48, 0.41, 0.28)
 	*/
-	UPROPERTY(Category = "SubsurfaceProfileStruct", EditAnywhere, BlueprintReadOnly, meta = (HideAlphaChannel, editcondition = "!bEnableBurley"))
+	UPROPERTY(Category = "SubsurfaceProfileStruct", EditAnywhere, BlueprintReadOnly, meta = (HideAlphaChannel, DeprecatedProperty, editcondition = "!bEnableBurley", EditConditionHides))
 	FLinearColor SubsurfaceColor;
 
 	/**
@@ -60,10 +68,10 @@ struct FSubsurfaceProfileStruct
 	* produced by the subsurface scattering events, can be used to fine tune the color of the gradients
 	* (called "falloff" in SeparableSSS, default there: 1, 0.37, 0.3)
 	*/
-	UPROPERTY(Category = "SubsurfaceProfileStruct", EditAnywhere, BlueprintReadOnly, meta = (HideAlphaChannel, editcondition = "!bEnableBurley"))
+	UPROPERTY(Category = "SubsurfaceProfileStruct", EditAnywhere, BlueprintReadOnly, meta = (HideAlphaChannel, DeprecatedProperty, editcondition = "!bEnableBurley", EditConditionHides))
 	FLinearColor FalloffColor;
 
-	UPROPERTY(Category = "SubsurfaceProfileStruct", EditAnywhere, BlueprintReadOnly, meta = (HideAlphaChannel))
+	UPROPERTY(Category = "Burley Normalized", EditAnywhere, BlueprintReadOnly, meta = (HideAlphaChannel))
 	FLinearColor BoundaryColorBleed;
 
 	UPROPERTY(Category = "Transmission", EditAnywhere, BlueprintReadOnly, meta = (ClampMin = "0.01", UIMax = "1.0", ClampMax = "1.0"))
@@ -115,6 +123,7 @@ struct FSubsurfaceProfileStruct
 		MeanFreePathDistance = 1.2f*2.229f;
 		WorldUnitScale = 0.1f;
 		TransmissionTintColor = FLinearColor(1.0f, 1.0f, 1.0f);
+		Tint = FLinearColor(1.0f, 1.0f, 1.0f);
 	}
 
 	void Invalidate()
diff --git a/Engine/Source/Runtime/Engine/Private/Rendering/SubsurfaceProfile.cpp b/Engine/Source/Runtime/Engine/Private/Rendering/SubsurfaceProfile.cpp
--- a/Engine/Source/Runtime/Engine/Private/Rendering/SubsurfaceProfile.cpp
+++ b/Engine/Source/Runtime/Engine/Private/Rendering/SubsurfaceProfile.cpp
@@ -36,12 +36,38 @@ static TRefCountPtr<IPooledRenderTarget> GSSProfiles;
 // Texture with one or more pre-integrated textures or 0 if there is no user
 static TRefCountPtr<IPooledRenderTarget> GSSProfilesPreIntegratedTexture;
 
+void ConvertSubsurfaceParametersFromSeparableToBurley(const FSubsurfaceProfileStruct& Settings, FLinearColor& SurfaceAlbedo, FLinearColor& MeanFreePathColor, float& MeanFreePathDistance)
+{
+	MapFallOffColor2SurfaceAlbedoAndDiffuseMeanFreePath(Settings.FalloffColor.R, SurfaceAlbedo.R, MeanFreePathColor.R);
+	MapFallOffColor2SurfaceAlbedoAndDiffuseMeanFreePath(Settings.FalloffColor.G, SurfaceAlbedo.G, MeanFreePathColor.G);
+	MapFallOffColor2SurfaceAlbedoAndDiffuseMeanFreePath(Settings.FalloffColor.B, SurfaceAlbedo.B, MeanFreePathColor.B);
+
+	//Normalize mean free path color and set the corresponding dfmp
+	float MaxMeanFreePathColor = FMath::Max3(MeanFreePathColor.R, MeanFreePathColor.G, MeanFreePathColor.B);
+	if (MaxMeanFreePathColor > 1)
+	{
+		MeanFreePathColor /= MaxMeanFreePathColor;
+		MeanFreePathDistance = FMath::Clamp(Settings.ScatterRadius * MaxMeanFreePathColor, 0.1f, 50.0f);	// 50.0f is the ClampMax of MeanFreePathDistance.
+	}
+}
+
+void UpgradeSeparableToBurley(FSubsurfaceProfileStruct& Settings)
+{
+	ConvertSubsurfaceParametersFromSeparableToBurley(Settings, Settings.SurfaceAlbedo, Settings.MeanFreePathColor, Settings.MeanFreePathDistance);
+	Settings.MeanFreePathDistance /= (2.229f * 2.229f);
+	Settings.Tint = Settings.SubsurfaceColor;
+	Settings.bEnableBurley = true;
+}
+
 FSubsurfaceProfileTexture::FSubsurfaceProfileTexture()
 {
 	check(IsInGameThread());
 
 	FSubsurfaceProfileStruct DefaultSkin;
 
+	//The default burley in slot 0 behaves the same as Separable previously
+	UpgradeSeparableToBurley(DefaultSkin);
+
 	// add element 0, it is used as default profile
 	SubsurfaceProfileEntries.Add(FSubsurfaceProfileEntry(DefaultSkin, 0));
 }
@@ -317,21 +343,6 @@ float Pow4(float X)
 	return X * X * X * X;
 }
 
-void ConvertSubsurfaceParametersFromSeparableToBurley(const FSubsurfaceProfileStruct& Settings, FLinearColor& SurfaceAlbedo, FLinearColor& MeanFreePathColor, float& MeanFreePathDistance)
-{
-	MapFallOffColor2SurfaceAlbedoAndDiffuseMeanFreePath(Settings.FalloffColor.R, SurfaceAlbedo.R, MeanFreePathColor.R);
-	MapFallOffColor2SurfaceAlbedoAndDiffuseMeanFreePath(Settings.FalloffColor.G, SurfaceAlbedo.G, MeanFreePathColor.G);
-	MapFallOffColor2SurfaceAlbedoAndDiffuseMeanFreePath(Settings.FalloffColor.B, SurfaceAlbedo.B, MeanFreePathColor.B);
-
-	//Normalize mean free path color and set the corresponding dfmp
-	float MaxMeanFreePathColor = FMath::Max3(MeanFreePathColor.R, MeanFreePathColor.G, MeanFreePathColor.B);
-	if (MaxMeanFreePathColor > 1)
-	{
-		MeanFreePathColor /= MaxMeanFreePathColor;
-		MeanFreePathDistance = FMath::Clamp(Settings.ScatterRadius * MaxMeanFreePathColor, 0.1f, 50.0f);	// 50.0f is the ClampMax of MeanFreePathDistance.
-	}
-}
-
 void FSubsurfaceProfileTexture::CreateTexture(FRHICommandListImmediate& RHICmdList)
 {
 	uint32 Height = SubsurfaceProfileEntries.Num();
@@ -456,11 +467,10 @@ void FSubsurfaceProfileTexture::CreateTexture(FRHICommandListImmediate& RHICmdLi
 			const float SamplingCountCompensation = 0.707f;
 			Data.ScatterRadius *= (Data.WorldUnitScale * 10.0f)*2.229f* SamplingCountCompensation;
 
-			// Reset subsurface rgb component to 1. So in the combine pass, we will directly use the subsurface scattering result
-			// in Burley fallback and Burley mode.
-			TextureRow[SSSS_SUBSURFACE_COLOR_OFFSET].R = 1; 
-			TextureRow[SSSS_SUBSURFACE_COLOR_OFFSET].G = 1; 
-			TextureRow[SSSS_SUBSURFACE_COLOR_OFFSET].B = 1; 
+			// To depricate {scatter radius, falloff} in favor of albedo/MFP, need the subsurface color as tint.
+			TextureRow[SSSS_SUBSURFACE_COLOR_OFFSET].R = Data.Tint.R; 
+			TextureRow[SSSS_SUBSURFACE_COLOR_OFFSET].G = Data.Tint.G;
+			TextureRow[SSSS_SUBSURFACE_COLOR_OFFSET].B = Data.Tint.B;
 		}
 		else
 		{
@@ -575,18 +585,16 @@ bool FSubsurfaceProfileTexture::GetEntryString(uint32 Index, FString& Out) const
 	const FSubsurfaceProfileStruct& ref = SubsurfaceProfileEntries[Index].Settings;
 
 
-	Out = FString::Printf(TEXT(" %c. %p ScatterRadius=%.1f, SubsurfaceColor=%.1f %.1f %.1f, FalloffColor=%.1f %.1f %.1f, \
-								SurfaceAlbedo=%.1f %.1f %.1f, MeanFreePathColor=%.1f %.1f %.1f, MeanFreePathDistance=%.1f, WorldUnitScale=%.1f,\
-								TransmissionTintColor=%.1f %.1f %.1f, EnableBurley=%.1f"), 
+	Out = FString::Printf(TEXT(" %c. %p SurfaceAlbedo=%.1f %.1f %.1f, MeanFreePathColor=%.1f %.1f %.1f, MeanFreePathDistance=%.1f, WorldUnitScale=%.1f,\
+								Tint=%.1f %.1f %.1f, BoundaryColorBleeding=%.1f %.1f %.1f, TransmissionTintColor=%.1f %.1f %.1f, EnableBurley=%.1f"), 
 		MiniFontCharFromIndex(Index), 
 		SubsurfaceProfileEntries[Index].Profile,
-		ref.ScatterRadius,
-		ref.SubsurfaceColor.R, ref.SubsurfaceColor.G, ref.SubsurfaceColor.B,
-		ref.FalloffColor.R, ref.FalloffColor.G, ref.FalloffColor.B,
 		ref.SurfaceAlbedo.R, ref.SurfaceAlbedo.G, ref.SurfaceAlbedo.B,
 		ref.MeanFreePathColor.R, ref.MeanFreePathColor.G, ref.MeanFreePathColor.B,
 		ref.MeanFreePathDistance,
 		ref.WorldUnitScale,
+		ref.Tint.R, ref.Tint.G, ref.Tint.B,
+		ref.BoundaryColorBleed.R, ref.BoundaryColorBleed.G, ref.BoundaryColorBleed.B, 
 		ref.TransmissionTintColor.R, ref.TransmissionTintColor.G, ref.TransmissionTintColor.B,
 		ref.bEnableBurley?1.0f:0.0f);
 
@@ -618,13 +626,6 @@ void FSubsurfaceProfileTexture::Dump()
 		UE_LOG(LogSubsurfaceProfile, Log, TEXT("  %d. AllocationId=%d, Pointer=%p"), i, i + 1, SubsurfaceProfileEntries[i].Profile);
 
 		{
-			UE_LOG(LogSubsurfaceProfile, Log, TEXT("     ScatterRadius = %f"),
-				SubsurfaceProfileEntries[i].Settings.ScatterRadius);
-			UE_LOG(LogSubsurfaceProfile, Log, TEXT("     SubsurfaceColor=%f %f %f"),
-				SubsurfaceProfileEntries[i].Settings.SubsurfaceColor.R, SubsurfaceProfileEntries[i].Settings.SubsurfaceColor.G, SubsurfaceProfileEntries[i].Settings.SubsurfaceColor.B);
-			UE_LOG(LogSubsurfaceProfile, Log, TEXT("     FalloffColor=%f %f %f"),
-				SubsurfaceProfileEntries[i].Settings.FalloffColor.R, SubsurfaceProfileEntries[i].Settings.FalloffColor.G, SubsurfaceProfileEntries[i].Settings.FalloffColor.B);
-
 			UE_LOG(LogSubsurfaceProfile, Log, TEXT("     SurfaceAlbedo=%f %f %f"),
 				SubsurfaceProfileEntries[i].Settings.SurfaceAlbedo.R, SubsurfaceProfileEntries[i].Settings.SurfaceAlbedo.G, SubsurfaceProfileEntries[i].Settings.SurfaceAlbedo.B);
 			UE_LOG(LogSubsurfaceProfile, Log, TEXT("     MeanFreePathColor=%f %f %f"),
@@ -633,6 +634,10 @@ void FSubsurfaceProfileTexture::Dump()
 				SubsurfaceProfileEntries[i].Settings.MeanFreePathDistance);
 			UE_LOG(LogSubsurfaceProfile, Log, TEXT("     WorldUnitScale=%f"),
 				SubsurfaceProfileEntries[i].Settings.WorldUnitScale);
+			UE_LOG(LogSubsurfaceProfile, Log, TEXT("     Tint=%f %f %f"),
+				SubsurfaceProfileEntries[i].Settings.Tint.R, SubsurfaceProfileEntries[i].Settings.Tint.G, SubsurfaceProfileEntries[i].Settings.Tint.B);
+			UE_LOG(LogSubsurfaceProfile, Log, TEXT("     Boundary Color Bleed=%f %f %f"),
+				SubsurfaceProfileEntries[i].Settings.BoundaryColorBleed.R, SubsurfaceProfileEntries[i].Settings.BoundaryColorBleed.G, SubsurfaceProfileEntries[i].Settings.BoundaryColorBleed.B);
 			UE_LOG(LogSubsurfaceProfile, Log, TEXT("     TransmissionTintColor=%f %f %f"),
 				SubsurfaceProfileEntries[i].Settings.TransmissionTintColor.R, SubsurfaceProfileEntries[i].Settings.TransmissionTintColor.G, SubsurfaceProfileEntries[i].Settings.TransmissionTintColor.B);
 			UE_LOG(LogSubsurfaceProfile, Log, TEXT("     EnableBurley=%.1f"),
@@ -728,9 +733,13 @@ void USubsurfaceProfile::PostLoad()
 	{
 		const bool bUpdateBurleyParametersFromSeparable = CVar->GetValueOnAnyThread() == 1;
 
-		if (bUpdateBurleyParametersFromSeparable)
+		if (!Settings.bEnableBurley
+#if !(UE_BUILD_SHIPPING || UE_BUILD_TEST)
+			|| bUpdateBurleyParametersFromSeparable
+#endif
+			)
 		{
-			ConvertSubsurfaceParametersFromSeparableToBurley(Settings, Settings.SurfaceAlbedo, Settings.MeanFreePathColor, Settings.MeanFreePathDistance);
+			UpgradeSeparableToBurley(Settings);
 		}
 	}
 }
