Tooltip fixes for Audio Volume

diff --git a/Engine/Source/Runtime/Engine/Classes/Sound/AudioVolume.h b/Engine/Source/Runtime/Engine/Classes/Sound/AudioVolume.h
--- a/Engine/Source/Runtime/Engine/Classes/Sound/AudioVolume.h
+++ b/Engine/Source/Runtime/Engine/Classes/Sound/AudioVolume.h
@@ -60,7 +60,7 @@ struct ENGINE_API FAudioVolumeSubmixOverrideSettings
 	UPROPERTY(EditAnywhere, BlueprintReadWrite, Category = AudioVolumeSubmixSends)
 	TObjectPtr<USoundSubmix> Submix = nullptr;
 
-	// The submix effect chain to overrideac
+	// The submix effect chain to override
 	UPROPERTY(EditAnywhere, BlueprintReadOnly, Category = SoundSubmix)
 	TArray<TObjectPtr<USoundEffectSubmixPreset>> SubmixEffectChain;
 
@@ -84,15 +84,15 @@ struct ENGINE_API FInteriorSettings
 	UPROPERTY(EditAnywhere, BlueprintReadWrite, Category=InteriorSettings)
 	float ExteriorVolume;
 
-	// The time over which to interpolate from the current volume to the desired volume of sounds outside the volume when the player enters the volume
+	// The time over which to interpolate from the current volume to the desired volume of sounds outside the volume when the player enters or exits the volume
 	UPROPERTY(EditAnywhere, BlueprintReadWrite, Category=InteriorSettings)
 	float ExteriorTime;
 
-	// The desired LPF frequency cutoff in hertz of sounds inside the volume when the player is outside the volume
+	// The desired LPF frequency cutoff in hertz of sounds outside the volume when the player is inside the volume
 	UPROPERTY(EditAnywhere, BlueprintReadWrite, Category = InteriorSettings)
 	float ExteriorLPF;
 
-	// The time over which to interpolate from the current LPF to the desired LPF of sounds outside the volume when the player enters the volume
+	// The time over which to interpolate from the current LPF to the desired LPF of sounds outside the volume when the player enters or exits the volume
 	UPROPERTY(EditAnywhere, BlueprintReadWrite, Category=InteriorSettings)
 	float ExteriorLPFTime;
 
@@ -100,15 +100,15 @@ struct ENGINE_API FInteriorSettings
 	UPROPERTY(EditAnywhere, BlueprintReadWrite, Category=InteriorSettings)
 	float InteriorVolume;
 
-	// The time over which to interpolate from the current volume to the desired volume of sounds inside the volume when the player enters the volume
+	// The time over which to interpolate from the current volume to the desired volume of sounds inside the volume when the player enters or exits the volume
 	UPROPERTY(EditAnywhere, BlueprintReadWrite, Category=InteriorSettings)
 	float InteriorTime;
 
-	// The desired LPF frequency cutoff in hertz of sounds outside the volume when the player is inside the volume
+	// The desired LPF frequency cutoff in hertz of sounds inside the volume when the player is outside the volume
 	UPROPERTY(EditAnywhere, BlueprintReadWrite, Category = InteriorSettings)
 	float InteriorLPF;
 
-	// The time over which to interpolate from the current LPF to the desired LPF of sounds inside the volume when the player enters the volume
+	// The time over which to interpolate from the current LPF to the desired LPF of sounds inside the volume when the player enters or exits the volume
 	UPROPERTY(EditAnywhere, BlueprintReadWrite, Category=InteriorSettings)
 	float InteriorLPFTime;
 
@@ -155,13 +155,13 @@ class ENGINE_API AAudioVolume : public AVolume
 
 private:
 	/**
-	 * Priority of this volume. In the case of overlapping volumes the one with the highest priority
+	 * Priority of this volume. In the case of overlapping volumes, the one with the highest priority
 	 * is chosen. The order is undefined if two or more overlapping volumes have the same priority.
 	 */
 	UPROPERTY(EditAnywhere, BlueprintReadOnly, Category=AudioVolume, meta=(AllowPrivateAccess="true"))
 	float Priority;
 
-	/** whether this volume is currently enabled and able to affect sounds */
+	/** Whether this volume is currently enabled and able to affect sounds */
 	UPROPERTY(EditAnywhere, BlueprintReadOnly, ReplicatedUsing=OnRep_bEnabled, Category=AudioVolume, meta=(AllowPrivateAccess="true"))
 	uint32 bEnabled:1;
 
@@ -169,7 +169,7 @@ private:
 	UPROPERTY(EditAnywhere, BlueprintReadOnly, Category=Reverb, meta=(AllowPrivateAccess="true"))
 	FReverbSettings Settings;
 
-	/** Interior settings used for this volume */
+	/** Impacts sounds that have "Apply Ambient Volumes" set to true in their Sound Class, based on whether the sound sources and the player are inside or outside the audio volume */
 	UPROPERTY(EditAnywhere, BlueprintReadOnly, Category=AmbientZone, meta=(AllowPrivateAccess="true"))
 	FInteriorSettings AmbientZoneSettings;
 
