World Partition: assert in AActor::SetIsSpatiallyLoaded only when the actor doesn't allow changing the spatial flag and the flag actually changes.
Also fixup bad default value for deprecated field AActor::GridPlacement.

diff --git a/Engine/Source/Runtime/Core/Public/UObject/UE5ReleaseStreamObjectVersion.h b/Engine/Source/Runtime/Core/Public/UObject/UE5ReleaseStreamObjectVersion.h
--- a/Engine/Source/Runtime/Core/Public/UObject/UE5ReleaseStreamObjectVersion.h
+++ b/Engine/Source/Runtime/Core/Public/UObject/UE5ReleaseStreamObjectVersion.h
@@ -66,6 +66,9 @@ struct CORE_API FUE5ReleaseStreamObjectVersion
 		// Converted AActor GridPlacement to bIsSpatiallyLoaded flag
 		ConvertedActorGridPlacementToSpatiallyLoadedFlag,
 
+		// Fixup for bad default value for GridPlacement_DEPRECATED
+		ActorGridPlacementDeprecateDefaultValueFixup,
+
 		// -----<new versions can be added above this line>-------------------------------------------------
 		VersionPlusOne,
 		LatestVersion = VersionPlusOne - 1
diff --git a/Engine/Source/Runtime/Engine/Classes/GameFramework/Actor.h b/Engine/Source/Runtime/Engine/Classes/GameFramework/Actor.h
--- a/Engine/Source/Runtime/Engine/Classes/GameFramework/Actor.h
+++ b/Engine/Source/Runtime/Engine/Classes/GameFramework/Actor.h
@@ -2217,8 +2217,11 @@ public:
 	/** Set if this actor should be spatially loaded or not. */
 	void SetIsSpatiallyLoaded(bool bInIsSpatiallyLoaded)
 	{
-		check(CanChangeIsSpatiallyLoadedFlag());
-		bIsSpatiallyLoaded = bInIsSpatiallyLoaded;
+		if (bIsSpatiallyLoaded != bInIsSpatiallyLoaded)
+		{
+			check(CanChangeIsSpatiallyLoadedFlag());
+			bIsSpatiallyLoaded = bInIsSpatiallyLoaded;
+		}
 	}
 
 	/** Returns true if this actor allows changing the spatially loaded flag.  */
diff --git a/Engine/Source/Runtime/Engine/Private/Actor.cpp b/Engine/Source/Runtime/Engine/Private/Actor.cpp
--- a/Engine/Source/Runtime/Engine/Private/Actor.cpp
+++ b/Engine/Source/Runtime/Engine/Private/Actor.cpp
@@ -17,6 +17,7 @@
 #include "UObject/UObjectHash.h"
 #include "UObject/PropertyPortFlags.h"
 #include "UObject/UE5MainStreamObjectVersion.h"
+#include "UObject/UE5ReleaseStreamObjectVersion.h"
 #include "UObject/UE5PrivateFrostyStreamObjectVersion.h"
 #include "Kismet/GameplayStatics.h"
 #include "Engine/BlueprintGeneratedClass.h"
@@ -177,6 +178,10 @@ void AActor::InitializeDefaults()
 #if WITH_EDITORONLY_DATA
 	bIsSpatiallyLoaded = true;
 	CopyPasteId = INDEX_NONE;
+
+	PRAGMA_DISABLE_DEPRECATION_WARNINGS
+	GridPlacement_DEPRECATED = EActorGridPlacement::None;
+	PRAGMA_ENABLE_DEPRECATION_WARNINGS
 #endif
 }
 
@@ -726,6 +731,7 @@ void AActor::Serialize(FArchive& Ar)
 {
 	Ar.UsingCustomVersion(FUE5MainStreamObjectVersion::GUID);
 	Ar.UsingCustomVersion(FUE5PrivateFrostyStreamObjectVersion::GUID);
+	Ar.UsingCustomVersion(FUE5ReleaseStreamObjectVersion::GUID);
 
 #if WITH_EDITOR
 	// Prior to load, map natively-constructed component instances for Blueprint-generated class types to any serialized properties that might reference them.
@@ -824,6 +830,11 @@ void AActor::Serialize(FArchive& Ar)
 		{
 			ActorGuid = FGuid::NewGuid();
 		}
+
+		if (!CanChangeIsSpatiallyLoadedFlag() && (Ar.CustomVer(FUE5ReleaseStreamObjectVersion::GUID) < FUE5ReleaseStreamObjectVersion::ActorGridPlacementDeprecateDefaultValueFixup))
+		{
+			bIsSpatiallyLoaded = GetClass()->GetDefaultObject<AActor>()->bIsSpatiallyLoaded;
+		}
 	}
 #endif
 }
