-Fix changelist validation in case of unloaded map , and correctly handle BluePrint Actor classes
-Add validation for same Runtime Grid in references
-Run validation iteratively until we've propagated all settings fixes (in case of long cycles, which shouldn't be a common case) but only report errors in the initial check (before propagating any fix) so that users only see errors for data they can view/fix.
-Changed validation for Attached actors, Attached actors validate their settings from the topmost parent so that we don't get multiple confusing errors and fixups

diff --git a/Engine/Plugins/Editor/DataValidation/Source/DataValidation/Private/WorldPartitionChangelistValidator.cpp b/Engine/Plugins/Editor/DataValidation/Source/DataValidation/Private/WorldPartitionChangelistValidator.cpp
--- a/Engine/Plugins/Editor/DataValidation/Source/DataValidation/Private/WorldPartitionChangelistValidator.cpp
+++ b/Engine/Plugins/Editor/DataValidation/Source/DataValidation/Private/WorldPartitionChangelistValidator.cpp
@@ -92,10 +92,11 @@ EDataValidationResult UWorldPartitionChangelistValidator::ValidateActorsListFrom
 			for (FAssetData& AssetData : PackageAssetsData)
 			{
 				// Check that the asset is an actor
-				if (AssetData.GetClass()->IsChildOf<AActor>())
+				static FName NAME_ActorMetaDataClass(TEXT("ActorMetaDataClass"));
+				if (AssetData.TagsAndValues.Contains(NAME_ActorMetaDataClass)) // Could check AssetData.GetClass()->IsChildOf<AActor>() but this doesn't handle blueprints, all OFPA/WP Actors will have this tag
 				{
 					// WorldPartition actors are all in OFPA mode so they're external
-					// Extract the MapNMame from the ObjectPath (<PathToPackage>.<mapName>:<level>.<actorName>)
+					// Extract the MapName from the ObjectPath (<PathToPackage>.<mapName>:<level>.<actorName>)
 					FSoftObjectPath ActorPath = FSoftObjectPath(AssetData.ObjectPath);
 					FName MapAssetName = ActorPath.GetAssetPathName();
 
@@ -136,9 +137,11 @@ EDataValidationResult UWorldPartitionChangelistValidator::ValidateActorsListFrom
 		}
 		else
 		{
-			// Find in memory 
+			// Find in memory failed, load the ActorDescContainer
+			FSoftObjectPath MapPath = FSoftObjectPath(MapName);
+
 			ActorDescContainer = NewObject<UActorDescContainer>();
-			ActorDescContainer->Initialize(nullptr, MapName);
+			ActorDescContainer->Initialize(nullptr, MapPath.GetLongPackageFName());
 		}
 
 		// Build a set of Relevant Actor Guids to scope error messages to what's contained in the CL 
@@ -220,6 +223,18 @@ void UWorldPartitionChangelistValidator::OnInvalidReferenceDataLayers(const FWor
 	}
 }
 
+void UWorldPartitionChangelistValidator::OnInvalidReferenceRuntimeGrid(const FWorldPartitionActorDescView& ActorDescView, const FWorldPartitionActorDescView& ReferenceActorDescView)
+{
+	if (Filter(ActorDescView) || Filter(ReferenceActorDescView))
+	{
+		FText CurrentError = FText::Format(LOCTEXT("DataValidation.Changelist.WorldPartition.RuntimeGridError", "{0} is referencing {1} but both Actors are using a different Runtime Grid."),
+			FText::FromString(GetPrettyPackageName(ActorDescView)),
+			FText::FromString(GetPrettyPackageName(ReferenceActorDescView)));
+
+		Errors->Add(CurrentError);
+	}
+}
+
 void UWorldPartitionChangelistValidator::OnInvalidReferenceLevelScriptStreamed(const FWorldPartitionActorDescView& ActorDescView)
 {
 	if (Filter(ActorDescView))
diff --git a/Engine/Plugins/Editor/DataValidation/Source/DataValidation/Private/WorldPartitionChangelistValidator.h b/Engine/Plugins/Editor/DataValidation/Source/DataValidation/Private/WorldPartitionChangelistValidator.h
--- a/Engine/Plugins/Editor/DataValidation/Source/DataValidation/Private/WorldPartitionChangelistValidator.h
+++ b/Engine/Plugins/Editor/DataValidation/Source/DataValidation/Private/WorldPartitionChangelistValidator.h
@@ -34,5 +34,6 @@ protected:
 	virtual void OnInvalidReferenceDataLayers(const FWorldPartitionActorDescView& ActorDescView, const FWorldPartitionActorDescView& ReferenceActorDescView) override;
 	virtual void OnInvalidReferenceLevelScriptStreamed(const FWorldPartitionActorDescView& ActorDescView) override;
 	virtual void OnInvalidReferenceLevelScriptDataLayers(const FWorldPartitionActorDescView& ActorDescView) override;
+	virtual void OnInvalidReferenceRuntimeGrid(const FWorldPartitionActorDescView& ActorDescView, const FWorldPartitionActorDescView& ReferenceActorDescView) override;
 };
 
diff --git a/Engine/Source/Runtime/Engine/Private/WorldPartition/ErrorHandling/WorldPartitionStreamingGenerationLogErrorHandler.cpp b/Engine/Source/Runtime/Engine/Private/WorldPartition/ErrorHandling/WorldPartitionStreamingGenerationLogErrorHandler.cpp
--- a/Engine/Source/Runtime/Engine/Private/WorldPartition/ErrorHandling/WorldPartitionStreamingGenerationLogErrorHandler.cpp
+++ b/Engine/Source/Runtime/Engine/Private/WorldPartition/ErrorHandling/WorldPartitionStreamingGenerationLogErrorHandler.cpp
@@ -23,6 +23,11 @@ void FStreamingGenerationLogErrorHandler::OnInvalidReferenceDataLayers(const FWo
 	UE_LOG(LogWorldPartition, Log, TEXT("Actor %s references an actor in a different set of data layers %s"), *ActorDescView.GetActorLabelOrName().ToString(), *ReferenceActorDescView.GetActorLabelOrName().ToString());
 }
 
+void FStreamingGenerationLogErrorHandler::OnInvalidReferenceRuntimeGrid(const FWorldPartitionActorDescView& ActorDescView, const FWorldPartitionActorDescView& ReferenceActorDescView)
+{
+	UE_LOG(LogWorldPartition, Log, TEXT("Actor %s references an actor in a different runtime grid %s"), *ActorDescView.GetActorLabelOrName().ToString(), *ReferenceActorDescView.GetActorLabelOrName().ToString());
+}
+
 void FStreamingGenerationLogErrorHandler::OnInvalidReferenceLevelScriptStreamed(const FWorldPartitionActorDescView& ActorDescView)
 {
 	UE_LOG(LogWorldPartition, Log, TEXT("Level Script Blueprint references streamed actor %s"), *ActorDescView.GetActorLabelOrName().ToString());
diff --git a/Engine/Source/Runtime/Engine/Private/WorldPartition/ErrorHandling/WorldPartitionStreamingGenerationMapCheckErrorHandler.cpp b/Engine/Source/Runtime/Engine/Private/WorldPartition/ErrorHandling/WorldPartitionStreamingGenerationMapCheckErrorHandler.cpp
--- a/Engine/Source/Runtime/Engine/Private/WorldPartition/ErrorHandling/WorldPartitionStreamingGenerationMapCheckErrorHandler.cpp
+++ b/Engine/Source/Runtime/Engine/Private/WorldPartition/ErrorHandling/WorldPartitionStreamingGenerationMapCheckErrorHandler.cpp
@@ -43,6 +43,16 @@ void FStreamingGenerationMapCheckErrorHandler::OnInvalidReferenceDataLayers(cons
 		->AddToken(FMapErrorToken::Create(FName(TEXT("WorldPartition_ActorReferenceActorInAnotherDataLayer_CheckForErrors"))));
 }
 
+void FStreamingGenerationMapCheckErrorHandler::OnInvalidReferenceRuntimeGrid(const FWorldPartitionActorDescView& ActorDescView, const FWorldPartitionActorDescView& ReferenceActorDescView)
+{
+	TSharedRef<FTokenizedMessage> Error = FMessageLog("MapCheck").Error()
+		->AddToken(FTextToken::Create(LOCTEXT("MapCheck_WorldPartition_Actor", "Actor")))
+		->AddToken(FAssetNameToken::Create(ActorDescView.GetActorLabelOrName().ToString()))
+		->AddToken(FTextToken::Create(LOCTEXT("MapCheck_WorldPartition_ReferenceActorInOtherRuntimeGrid", "references an actor in a different runtime grid")))
+		->AddToken(FAssetNameToken::Create(ReferenceActorDescView.GetActorLabelOrName().ToString()))
+		->AddToken(FMapErrorToken::Create(FName(TEXT("WorldPartition_ActorReferenceActorInAnotherRuntimeGrid_CheckForErrors"))));
+}
+
 void FStreamingGenerationMapCheckErrorHandler::OnInvalidReferenceLevelScriptStreamed(const FWorldPartitionActorDescView& ActorDescView)
 {
 	TSharedRef<FTokenizedMessage> Error = FMessageLog("MapCheck").Error()
diff --git a/Engine/Source/Runtime/Engine/Private/WorldPartition/WorldPartitionActorDescView.cpp b/Engine/Source/Runtime/Engine/Private/WorldPartition/WorldPartitionActorDescView.cpp
--- a/Engine/Source/Runtime/Engine/Private/WorldPartition/WorldPartitionActorDescView.cpp
+++ b/Engine/Source/Runtime/Engine/Private/WorldPartition/WorldPartitionActorDescView.cpp
@@ -40,6 +40,11 @@ FVector FWorldPartitionActorDescView::GetOrigin() const
 
 FName FWorldPartitionActorDescView::GetRuntimeGrid() const
 {
+	if (bInvalidRuntimeGrid)
+	{
+		return FName();
+	}
+
 	return ActorDesc->GetRuntimeGrid();
 }
 
@@ -139,6 +144,11 @@ void FWorldPartitionActorDescView::SetForcedNonSpatiallyLoaded()
 	}
 }
 
+void FWorldPartitionActorDescView::SetInvalidRuntimeGrid()
+{
+	bInvalidRuntimeGrid = true;	
+}
+
 void FWorldPartitionActorDescView::SetInvalidDataLayers()
 {
 	if (!bInvalidDataLayers)
diff --git a/Engine/Source/Runtime/Engine/Private/WorldPartition/WorldPartitionStreamingGeneration.cpp b/Engine/Source/Runtime/Engine/Private/WorldPartition/WorldPartitionStreamingGeneration.cpp
--- a/Engine/Source/Runtime/Engine/Private/WorldPartition/WorldPartitionStreamingGeneration.cpp
+++ b/Engine/Source/Runtime/Engine/Private/WorldPartition/WorldPartitionStreamingGeneration.cpp
@@ -32,6 +32,7 @@
 		HLOD Generation
 */
 
+
 class FWorldPartitionStreamingGenerator
 {
 	void CreateActorDescViewMap(const UActorDescContainer* InContainer, TMap<FGuid, FWorldPartitionActorDescView>& OutActorDescViewMap, const FActorContainerID& InContainerID)
@@ -162,6 +163,21 @@ class FWorldPartitionStreamingGenerator
 	 */
 	void ValidateActorDescriptorViews()
 	{
+
+		// Validate data layers
+		auto IsReferenceDataLayersValid = [](const FWorldPartitionActorDescView& ActorDescView, const FWorldPartitionActorDescView& ReferenceActorDescView)
+		{
+			if (ActorDescView.GetDataLayers().Num() == ReferenceActorDescView.GetDataLayers().Num())
+			{
+				const TSet<FName> ActorDescDataLayers(ActorDescView.GetDataLayers());
+				const TSet<FName> ReferenceActorDescDataLayers(ReferenceActorDescView.GetDataLayers());
+
+				return ActorDescDataLayers.Includes(ReferenceActorDescDataLayers);
+			}
+
+			return false;
+		};
+
 		for (auto It = ContainerDescriptorsMap.CreateIterator(); It; ++It)
 		{
 			const FActorContainerID& ContainerID = It.Key();
@@ -201,62 +217,152 @@ class FWorldPartitionStreamingGenerator
 				}
 			}
 
-			// Perform various adjustements based on validations and errors
-			for (auto ActorDescIt = ContainerDescriptor.ActorDescViewMap.CreateIterator(); ActorDescIt; ++ActorDescIt)
+			// Perform various adjustements based on validations and report errors
+			//
+			// The first validation pass is used to report errors, subsequent passes are used to make corrections to the FWorldPartitionActorDescView
+			// Since the references can form cycles/long chains in the data fixes might need to be propagated in multiple passes.
+			// 
+			// This works because fixes are deterministic and always apply the same way to both Actors being modified, so there's no ordering issues possible
+			uint32 NbValidationPasses = 0;
+
+			while(true)
 			{
-				FWorldPartitionActorDescView& ActorDescView = ActorDescIt.Value();
+				uint32 NbErrorsDetected = 0;
 
-				// Validate references
-				for (const FGuid& ReferenceGuid : ActorDescView.GetReferences())
+				for (auto ActorDescIt = ContainerDescriptor.ActorDescViewMap.CreateIterator(); ActorDescIt; ++ActorDescIt)
 				{
-					if (FWorldPartitionActorDescView* ReferenceActorDescView = ContainerDescriptor.ActorDescViewMap.Find(ReferenceGuid))
-					{
-						// Validate grid placement
-						const bool bIsActorDescSpatiallyLoaded = ActorDescView.GetIsSpatiallyLoaded();
-						const bool bIsActorDescRefSpatiallyLoaded = ReferenceActorDescView->GetIsSpatiallyLoaded();
+					FWorldPartitionActorDescView& ActorDescView = ActorDescIt.Value();
 
-						if (bIsActorDescSpatiallyLoaded != bIsActorDescRefSpatiallyLoaded)
+					struct ActorReferenceInfo
+					{
+						FGuid								ActorGuid;
+						FWorldPartitionActorDescView*		ActorDesc;
+						FGuid								ReferenceGuid;
+						FWorldPartitionActorDescView*		ReferenceActorDesc;
+					};
+
+					TArray<ActorReferenceInfo, TInlineAllocator<8>> References;
+					
+					// Build References List
+
+					// Add normal Actore references
+					for (const FGuid& ReferenceGuid : ActorDescView.GetReferences())
+					{
+						if (ReferenceGuid != ActorDescView.GetParentActor()) // References to the parent are inversed in their handling 
 						{
-							if (ErrorHandler)
-							{
-								ErrorHandler->OnInvalidReferenceGridPlacement(ActorDescView, *ReferenceActorDescView);
-							}
+							// Filter out parent back references
+							FWorldPartitionActorDescView* ReferenceActorDesc = ContainerDescriptor.ActorDescViewMap.Find(ReferenceGuid);
+							if (ReferenceActorDesc && ReferenceActorDesc->GetParentActor() == ActorDescView.GetGuid())
+								continue;
 
-							ActorDescView.SetForcedNonSpatiallyLoaded();
-							ReferenceActorDescView->SetForcedNonSpatiallyLoaded();
+							References.Emplace(ActorReferenceInfo{ ActorDescView.GetGuid(), &ActorDescView, ReferenceGuid, ReferenceActorDesc });
+						}
+					}
+
+					// Add attach reference for the topmost parent, this reference is inverted since we consider the top most existing 
+					// parent to be refering to us, not the child to be referering the parent
+					FGuid ParentGuid = ActorDescView.GetParentActor();
+					FWorldPartitionActorDescView* TopParentDescView = nullptr;
+
+					while (ParentGuid.IsValid())
+					{
+						FWorldPartitionActorDescView* ParentDescView = ContainerDescriptor.ActorDescViewMap.Find(ParentGuid);
+					
+						if (ParentDescView)
+						{
+							TopParentDescView = ParentDescView;
+							ParentGuid = ParentDescView->GetParentActor();
+						}
+						else
+						{
+							// we had a guid but parent cannot be found, this will be a missing reference
+							break; 
 						}
+					}
+
+					if (TopParentDescView)
+					{
+						References.Emplace(ActorReferenceInfo{ TopParentDescView->GetGuid(), TopParentDescView, ActorDescView.GetGuid(), &ActorDescView });
+					}
+
+					if (ParentGuid.IsValid())
+					{
+						// In case of missing parent add a missing reference 
+						References.Emplace(ActorReferenceInfo{ ActorDescView.GetGuid(), &ActorDescView, ParentGuid, nullptr });
+					}
+
+					for (ActorReferenceInfo& Info : References)
+					{
+						FWorldPartitionActorDescView* RefererActorDescView = Info.ActorDesc;
+						FWorldPartitionActorDescView* ReferenceActorDescView = Info.ReferenceActorDesc;
 
-						// Validate data layers
-						auto IsReferenceDataLayersValid = [](const FWorldPartitionActorDescView& ActorDescView, const FWorldPartitionActorDescView& ReferenceActorDescView)
+						if (ReferenceActorDescView)
 						{
-							if (ActorDescView.GetDataLayers().Num() == ReferenceActorDescView.GetDataLayers().Num())
+							// Validate grid placement
+							const bool bIsActorDescSpatiallyLoaded = RefererActorDescView->GetIsSpatiallyLoaded();
+							const bool bIsActorDescRefSpatiallyLoaded = ReferenceActorDescView->GetIsSpatiallyLoaded();
+
+							if (bIsActorDescSpatiallyLoaded != bIsActorDescRefSpatiallyLoaded)
 							{
-								const TSet<FName> ActorDescDataLayers(ActorDescView.GetDataLayers());
-								const TSet<FName> ReferenceActorDescDataLayers(ReferenceActorDescView.GetDataLayers());
+								if (ErrorHandler && !NbValidationPasses)
+								{
+									ErrorHandler->OnInvalidReferenceGridPlacement(*RefererActorDescView, *ReferenceActorDescView);									
+								}
+								NbErrorsDetected++;
+
+								if (NbValidationPasses)
+								{
+									RefererActorDescView->SetForcedNonSpatiallyLoaded();
+									ReferenceActorDescView->SetForcedNonSpatiallyLoaded();
+								}
+							}
 
-								return ActorDescDataLayers.Includes(ReferenceActorDescDataLayers);
+							if (!IsReferenceDataLayersValid(*RefererActorDescView, *ReferenceActorDescView))
+							{
+								if (ErrorHandler && !NbValidationPasses)
+								{
+									ErrorHandler->OnInvalidReferenceDataLayers(*RefererActorDescView, *ReferenceActorDescView);									
+								}
+								NbErrorsDetected++;
+
+								if (NbValidationPasses)
+								{
+									RefererActorDescView->SetInvalidDataLayers();
+									ReferenceActorDescView->SetInvalidDataLayers();
+								}
 							}
 
-							return false;
-						};
+							if (ReferenceActorDescView->GetRuntimeGrid() != RefererActorDescView->GetRuntimeGrid())
+							{
+								if (ErrorHandler && !NbValidationPasses)
+								{
+									ErrorHandler->OnInvalidReferenceRuntimeGrid(*RefererActorDescView, *ReferenceActorDescView);
+								}
+								NbErrorsDetected++;
+
+								if  (NbValidationPasses)
+								{
+									RefererActorDescView->SetInvalidRuntimeGrid();
+									ReferenceActorDescView->SetInvalidRuntimeGrid();
+								}
+							}
 
-						if (!IsReferenceDataLayersValid(ActorDescView, *ReferenceActorDescView))
+						}
+						else if (ErrorHandler)
 						{
-							if (ErrorHandler)
+							if (!NbValidationPasses)
 							{
-								ErrorHandler->OnInvalidReferenceDataLayers(ActorDescView, *ReferenceActorDescView);
+								ErrorHandler->OnInvalidReference(*RefererActorDescView, Info.ReferenceGuid);
 							}
-
-							ActorDescView.SetInvalidDataLayers();
-							ReferenceActorDescView->SetInvalidDataLayers();
+							// Do not increment NbErrorsDetected since it won't be fixed and thus will always occur
 						}
-					}
-					else if (ErrorHandler)
-					{
-						ErrorHandler->OnInvalidReference(ActorDescView, ReferenceGuid);
-					}
-				}
-			}
+					}				
+				}		
+
+				NbValidationPasses++;
+				if (NbErrorsDetected == 0)
+					break;
+			} 
 		}
 	}
 
diff --git a/Engine/Source/Runtime/Engine/Public/WorldPartition/ErrorHandling/WorldPartitionStreamingGenerationErrorHandler.h b/Engine/Source/Runtime/Engine/Public/WorldPartition/ErrorHandling/WorldPartitionStreamingGenerationErrorHandler.h
--- a/Engine/Source/Runtime/Engine/Public/WorldPartition/ErrorHandling/WorldPartitionStreamingGenerationErrorHandler.h
+++ b/Engine/Source/Runtime/Engine/Public/WorldPartition/ErrorHandling/WorldPartitionStreamingGenerationErrorHandler.h
@@ -25,6 +25,11 @@ public:
 	 */
 	virtual void OnInvalidReferenceDataLayers(const FWorldPartitionActorDescView& ActorDescView, const FWorldPartitionActorDescView& ReferenceActorDescView) = 0;
 
+	/**
+	 * Called when an actor references an actor using a different RuntimeGrid.
+	 */
+	virtual void OnInvalidReferenceRuntimeGrid(const FWorldPartitionActorDescView& ActorDescView, const FWorldPartitionActorDescView& ReferenceActorDescView) = 0;
+
 	/** 
 	 * Called when the level script references a streamed actor.
 	 */
diff --git a/Engine/Source/Runtime/Engine/Public/WorldPartition/ErrorHandling/WorldPartitionStreamingGenerationLogErrorHandler.h b/Engine/Source/Runtime/Engine/Public/WorldPartition/ErrorHandling/WorldPartitionStreamingGenerationLogErrorHandler.h
--- a/Engine/Source/Runtime/Engine/Public/WorldPartition/ErrorHandling/WorldPartitionStreamingGenerationLogErrorHandler.h
+++ b/Engine/Source/Runtime/Engine/Public/WorldPartition/ErrorHandling/WorldPartitionStreamingGenerationLogErrorHandler.h
@@ -12,5 +12,6 @@ public:
 	virtual void OnInvalidReferenceDataLayers(const FWorldPartitionActorDescView& ActorDescView, const FWorldPartitionActorDescView& ReferenceActorDescView) override;
 	virtual void OnInvalidReferenceLevelScriptStreamed(const FWorldPartitionActorDescView& ActorDescView) override;
 	virtual void OnInvalidReferenceLevelScriptDataLayers(const FWorldPartitionActorDescView& ActorDescView) override;
+	virtual void OnInvalidReferenceRuntimeGrid(const FWorldPartitionActorDescView& ActorDescView, const FWorldPartitionActorDescView& ReferenceActorDescView) override;
 };
 #endif
diff --git a/Engine/Source/Runtime/Engine/Public/WorldPartition/ErrorHandling/WorldPartitionStreamingGenerationMapCheckErrorHandler.h b/Engine/Source/Runtime/Engine/Public/WorldPartition/ErrorHandling/WorldPartitionStreamingGenerationMapCheckErrorHandler.h
--- a/Engine/Source/Runtime/Engine/Public/WorldPartition/ErrorHandling/WorldPartitionStreamingGenerationMapCheckErrorHandler.h
+++ b/Engine/Source/Runtime/Engine/Public/WorldPartition/ErrorHandling/WorldPartitionStreamingGenerationMapCheckErrorHandler.h
@@ -12,5 +12,6 @@ public:
 	virtual void OnInvalidReferenceDataLayers(const FWorldPartitionActorDescView& ActorDescView, const FWorldPartitionActorDescView& ReferenceActorDescView) override;
 	virtual void OnInvalidReferenceLevelScriptStreamed(const FWorldPartitionActorDescView& ActorDescView) override;
 	virtual void OnInvalidReferenceLevelScriptDataLayers(const FWorldPartitionActorDescView& ActorDescView) override;
+	virtual void OnInvalidReferenceRuntimeGrid(const FWorldPartitionActorDescView& ActorDescView, const FWorldPartitionActorDescView& ReferenceActorDescView) override;
 };
 #endif
diff --git a/Engine/Source/Runtime/Engine/Public/WorldPartition/WorldPartitionActorDescView.h b/Engine/Source/Runtime/Engine/Public/WorldPartition/WorldPartitionActorDescView.h
--- a/Engine/Source/Runtime/Engine/Public/WorldPartition/WorldPartitionActorDescView.h
+++ b/Engine/Source/Runtime/Engine/Public/WorldPartition/WorldPartitionActorDescView.h
@@ -49,6 +49,8 @@ public:
 
 	void SetForcedNonSpatiallyLoaded();
 
+	void SetInvalidRuntimeGrid();
+
 	void SetInvalidDataLayers();
 
 	bool operator==(const FWorldPartitionActorDescView& Other) const
@@ -65,5 +67,6 @@ protected:
 	const FWorldPartitionActorDesc* ActorDesc;
 	bool bIsForcedNonSpatiallyLoaded;
 	bool bInvalidDataLayers;
+	bool bInvalidRuntimeGrid;
 };
 #endif
\ No newline at end of file
