Minor fix to put more worker timing logs into the LogShaderCompilers category.

diff --git a/Engine/Source/Runtime/Engine/Private/ShaderCompiler/ShaderCompiler.cpp b/Engine/Source/Runtime/Engine/Private/ShaderCompiler/ShaderCompiler.cpp
--- a/Engine/Source/Runtime/Engine/Private/ShaderCompiler/ShaderCompiler.cpp
+++ b/Engine/Source/Runtime/Engine/Private/ShaderCompiler/ShaderCompiler.cpp
@@ -2059,7 +2059,7 @@ int32 FShaderCompileThreadRunnable::PullTasksFromQueue()
 							GShaderCompilerStats->RegisterLocalWorkerIdleTime(WorkerIdleTime);
 							if (Manager->bLogJobCompletionTimes)
 							{
-								UE_LOG(LogShaders, Display, TEXT("  Worker (%d/%d) starts working after being idle for %fs"), WorkerIndex + 1, WorkerInfos.Num(), WorkerIdleTime);
+								UE_LOG(LogShaderCompilers, Display, TEXT("  Worker (%d/%d) started working after being idle for %fs"), WorkerIndex + 1, WorkerInfos.Num(), WorkerIdleTime);
 							}
 						}
 					}
@@ -2153,7 +2153,7 @@ void FShaderCompileThreadRunnable::PushCompletedJobsToManager()
 					}
 				}
 
-				UE_LOG(LogShaders, Display, TEXT("Worker (%d/%d) finished batch of %u jobs in %.3fs, %s"), WorkerIndex + 1, WorkerInfos.Num(), SortedJobs.Num(), ElapsedTime, *JobNames);
+				UE_LOG(LogShaderCompilers, Display, TEXT("Worker (%d/%d) finished batch of %u jobs in %.3fs, %s"), WorkerIndex + 1, WorkerInfos.Num(), SortedJobs.Num(), ElapsedTime, *JobNames);
 			}
 
 			CurrentWorkerInfo.FinishTime = FPlatformTime::Seconds();
