Fix multiple indeterminism related to texture generation during cook by using content based hash for texture by default

diff --git a/Engine/Source/Runtime/Engine/Private/Texture.cpp b/Engine/Source/Runtime/Engine/Private/Texture.cpp
--- a/Engine/Source/Runtime/Engine/Private/Texture.cpp
+++ b/Engine/Source/Runtime/Engine/Private/Texture.cpp
@@ -1542,7 +1542,9 @@ void FTextureSource::UnlockMip(int32 BlockIndex, int32 LayerIndex, int32 MipInde
 			bPNGCompressed = false;
 			CompressionFormat = TSCF_None;
 
-			ForceGenerateGuid();
+			// Need to unlock before calling UseHashAsGuid
+			LockState = ELockState::None;
+			UseHashAsGuid();
 		}
 
 		LockState = ELockState::None;
