LevelInstance:
- Rename APackedLevelInstance to APackedLevelActor
- Change Context menu structure (moved under Level)
- Make creation of LevelInstance and PackedLevelActor distinct entries
- Remove possibility to create non-BP PackedLevelActors

diff --git a/Engine/Config/BaseEngine.ini b/Engine/Config/BaseEngine.ini
--- a/Engine/Config/BaseEngine.ini
+++ b/Engine/Config/BaseEngine.ini
@@ -998,6 +998,8 @@ NearClipPlane=10.0
 +ClassRedirects=(OldName="ComputeKernelSource", NewName="/Script/ComputeFramework.ComputeKernelSource")
 +ClassRedirects=(OldName="ComputeKernel", NewName="/Script/ComputeFramework.ComputeKernel")
 
++ClassRedirects=(OldName="PackedLevelInstance", NewName="/Script/Engine.PackedLevelActor")
+
 [CoreUObject.Metadata]
 ; Note: UnrealHeaderTool should be rerun after making changes to MetadataRedirects to catch any keys specified in class headers
 +MetadataRedirects=(OldKey="K2Protected", NewKey="BlueprintProtected")
diff --git a/Engine/Source/Editor/EditorStyle/Private/StarshipStyle.cpp b/Engine/Source/Editor/EditorStyle/Private/StarshipStyle.cpp
--- a/Engine/Source/Editor/EditorStyle/Private/StarshipStyle.cpp
+++ b/Engine/Source/Editor/EditorStyle/Private/StarshipStyle.cpp
@@ -6081,6 +6081,7 @@ void FStarshipEditorStyle::FStyle::SetupClassIconsAndThumbnails()
 			{TEXT("Landscape")},
 			{TEXT("LevelBounds")},
 			{TEXT("LevelInstance")},
+			{TEXT("PackedLevelActor")},
 			{TEXT("LevelScriptActor")},
 			{TEXT("LevelSequenceActor")},
 			{TEXT("LevelStreamingVolume")},
diff --git a/Engine/Source/Editor/LevelInstanceEditor/Private/LevelInstanceEditorModule.cpp b/Engine/Source/Editor/LevelInstanceEditor/Private/LevelInstanceEditorModule.cpp
--- a/Engine/Source/Editor/LevelInstanceEditor/Private/LevelInstanceEditorModule.cpp
+++ b/Engine/Source/Editor/LevelInstanceEditor/Private/LevelInstanceEditorModule.cpp
@@ -4,8 +4,7 @@
 #include "LevelInstancePivotDetails.h"
 #include "LevelInstance/LevelInstanceSubsystem.h"
 #include "LevelInstance/LevelInstanceActor.h"
-#include "LevelInstance/Packed/PackedLevelInstanceBuilder.h"
-#include "LevelInstance/Packed/PackedLevelInstanceActor.h"
+#include "PackedLevelActor/PackedLevelActor.h"
 #include "LevelInstanceEditorSettings.h"
 #include "ToolMenus.h"
 #include "Editor.h"
@@ -45,14 +44,13 @@ IMPLEMENT_MODULE( FLevelInstanceEditorModule, LevelInstanceEditor );
 
 namespace LevelInstanceMenuUtils
 {
-	FToolMenuSection& CreateLevelInstanceSection(UToolMenu* Menu)
+	FToolMenuSection& CreateLevelSection(UToolMenu* Menu)
 	{
-		const FName LevelInstanceSectionName = TEXT("LevelInstance");
-		FToolMenuSection* SectionPtr = Menu->FindSection(LevelInstanceSectionName);
+		const FName LevelSectionName = TEXT("Level");
+		FToolMenuSection* SectionPtr = Menu->FindSection(LevelSectionName);
 		if (!SectionPtr)
 		{
-			FToolMenuInsert InsertPosition("ActorControl", EToolMenuInsertType::After);
-			SectionPtr = &(Menu->AddSection(LevelInstanceSectionName, LOCTEXT("LevelInstance", "Level Instance"), InsertPosition));
+			SectionPtr = &(Menu->AddSection(LevelSectionName, LOCTEXT("LevelSectionLabel", "Level")));
 		}
 		FToolMenuSection& Section = *SectionPtr;
 		return Section;
@@ -122,12 +120,12 @@ namespace LevelInstanceMenuUtils
 			// Don't create sub menu if only one Level Instance is available to edit
 			if (LevelInstanceHierarchy.Num() == 1)
 			{
-				FToolMenuSection& Section = CreateLevelInstanceSection(Menu);
+				FToolMenuSection& Section = CreateLevelSection(Menu);
 				CreateEditMenuEntry(Section, LevelInstanceHierarchy[0], ContextActor, true);
 			}
 			else if(LevelInstanceHierarchy.Num() > 1)
 			{
-				FToolMenuSection& Section = CreateLevelInstanceSection(Menu);
+				FToolMenuSection& Section = CreateLevelSection(Menu);
 				Section.AddSubMenu(
 					"EditLevelInstances",
 					LOCTEXT("EditLevelInstances", "Edit"),
@@ -182,7 +180,7 @@ namespace LevelInstanceMenuUtils
 					return bCanCommit;
 				});
 
-				FToolMenuSection& Section = CreateLevelInstanceSection(Menu);
+				FToolMenuSection& Section = CreateLevelSection(Menu);
 				Section.AddMenuEntry(NAME_None, bDiscard ? LOCTEXT("LevelInstanceDiscardLabel", "Discard") : LOCTEXT("LevelInstanceCommitLabel", "Commit"), Tooltip, FSlateIcon(), MenuAction);
 			};
 			
@@ -214,7 +212,7 @@ namespace LevelInstanceMenuUtils
 						LevelInstanceEdit->SetCurrent();
 					});
 
-				FToolMenuSection& Section = CreateLevelInstanceSection(Menu);
+				FToolMenuSection& Section = CreateLevelSection(Menu);
 				Section.AddMenuEntry(NAME_None, LOCTEXT("LevelInstanceSetCurrent", "Set Current Level"), TAttribute<FText>(), FSlateIcon(), LevelInstanceSetCurrentAction);
 			}
 		}
@@ -256,14 +254,14 @@ namespace LevelInstanceMenuUtils
 						MoveSelectionToLevelInstance(LevelInstanceEdit);
 					});
 
-				FToolMenuSection& Section = CreateLevelInstanceSection(Menu);
+				FToolMenuSection& Section = CreateLevelSection(Menu);
 				Section.AddMenuEntry(NAME_None, LOCTEXT("LevelInstanceMoveSelectionTo", "Move Selection to"), TAttribute<FText>(), FSlateIcon(), LevelInstanceMoveSelectionAction);
 
 			}
 		}
 	}
 
-	void CreateLevelInstanceFromSelection(ULevelInstanceSubsystem* LevelInstanceSubsystem)
+	void CreateLevelInstanceFromSelection(ULevelInstanceSubsystem* LevelInstanceSubsystem, ELevelInstanceCreationType CreationType)
 	{
 		TArray<AActor*> ActorsToMove;
 		ActorsToMove.Reserve(GEditor->GetSelectedActorCount());
@@ -279,12 +277,10 @@ namespace LevelInstanceMenuUtils
 
 		TSharedPtr<SWindow> NewLevelInstanceWindow =
 			SNew(SWindow)
-			.Title(LOCTEXT("NewLevelInstanceWindowTitle", "New Level Instance"))
-			.ClientSize(SNewLevelInstanceDialog::DEFAULT_WINDOW_SIZE)
-			.SizingRule(ESizingRule::UserSized)
+			.Title(FText::Format(LOCTEXT("NewLevelInstanceWindowTitle", "New {0}"), StaticEnum<ELevelInstanceCreationType>()->GetDisplayNameTextByValue((int64)CreationType)))
 			.SupportsMinimize(false)
 			.SupportsMaximize(false)
-			.SizingRule(ESizingRule::FixedSize);
+			.SizingRule(ESizingRule::Autosized);
 
 		TSharedRef<SNewLevelInstanceDialog> NewLevelInstanceDialog =
 			SNew(SNewLevelInstanceDialog)
@@ -293,6 +289,8 @@ namespace LevelInstanceMenuUtils
 
 		const bool bForceExternalActors = LevelInstanceSubsystem->GetWorld()->IsPartitionedWorld();
 		FNewLevelInstanceParams& DialogParams = NewLevelInstanceDialog->GetCreationParams();
+		DialogParams.Type = CreationType;
+		DialogParams.HideCreationType();
 		DialogParams.SetForceExternalActors(bForceExternalActors);
 		NewLevelInstanceWindow->SetContent(NewLevelInstanceDialog);
 
@@ -315,37 +313,38 @@ namespace LevelInstanceMenuUtils
 
 				if (!LevelInstanceSubsystem->CreateLevelInstanceFrom(ActorsToMove, CreationParams))
 				{
-					FText Title = LOCTEXT("CreateFromSelectionFailTitle", "Create from selection failed");
-					FMessageDialog::Open(EAppMsgType::Ok, LOCTEXT("CreateFromSelectionFailMsg", "Failed to create LevelInstance from selection. Check log for details."), &Title);
+					FText FailedTitle = LOCTEXT("CreateFromSelectionFailTitle", "Create from selection failed");
+					FMessageDialog::Open(EAppMsgType::Ok, LOCTEXT("CreateFromSelectionFailMsg", "Failed to create from selection. Check log for details."), &FailedTitle);
 				}
 			}
 		}
 	}
-	
-	void CreateCreateMenu(UToolMenu* Menu)
+
+	void CreateCreateMenu(UToolMenu* ToolMenu)
 	{
 		if (ULevelInstanceSubsystem* LevelInstanceSubsystem = GEditor->GetEditorWorldContext().World()->GetSubsystem<ULevelInstanceSubsystem>())
 		{
-			FToolMenuSection& Section = CreateLevelInstanceSection(Menu);
-			FToolUIAction UIAction;
-			UIAction.ExecuteAction.BindLambda([LevelInstanceSubsystem](const FToolMenuContext& MenuContext)
-			{
-				CreateLevelInstanceFromSelection(LevelInstanceSubsystem);
-			});
-			UIAction.CanExecuteAction.BindLambda([](const FToolMenuContext& MenuContext)
+			if (GEditor->GetSelectedActorCount() > 0)
 			{
-				return GEditor->GetSelectedActorCount() > 0;
-			});
+				FToolMenuSection& Section = ToolMenu->AddSection("ActorSelectionSectionName", LOCTEXT("ActorSelectionSectionLabel", "Actor Selection"));
 
-			Section.AddMenuEntry(
-				"CreateLevelInstanceFromSelection",
-				LOCTEXT("CreateLevelInstanceFromSelection", "Create from selection..."),
-				TAttribute<FText>(),
-				TAttribute<FSlateIcon>(),
-				UIAction);
+				Section.AddMenuEntry(
+					NAME_None,
+					FText::Format(LOCTEXT("CreateFromSelectionLabel", "Create {0}..."), StaticEnum<ELevelInstanceCreationType>()->GetDisplayNameTextByValue((int64)ELevelInstanceCreationType::LevelInstance)),
+					TAttribute<FText>(),
+					FSlateIcon(FAppStyle::GetAppStyleSetName(), "ClassIcon.LevelInstance"),
+					FExecuteAction::CreateStatic(&LevelInstanceMenuUtils::CreateLevelInstanceFromSelection, LevelInstanceSubsystem, ELevelInstanceCreationType::LevelInstance));
+
+				Section.AddMenuEntry(
+					NAME_None,
+					FText::Format(LOCTEXT("CreateFromSelectionLabel", "Create {0}..."), StaticEnum<ELevelInstanceCreationType>()->GetDisplayNameTextByValue((int64)ELevelInstanceCreationType::PackedLevelActor)),
+					TAttribute<FText>(),
+					FSlateIcon(FAppStyle::GetAppStyleSetName(), "ClassIcon.PackedLevelActor"),
+					FExecuteAction::CreateStatic(&LevelInstanceMenuUtils::CreateLevelInstanceFromSelection, LevelInstanceSubsystem, ELevelInstanceCreationType::PackedLevelActor));
+			}
 		}
 	}
-
+		
 	void CreateSaveAsMenu(UToolMenu* Menu, AActor* ContextActor)
 	{
 		ALevelInstance* ContextLevelInstance = nullptr;
@@ -367,7 +366,7 @@ namespace LevelInstanceMenuUtils
 
 		if (ContextLevelInstance)
 		{
-			FToolMenuSection& Section = CreateLevelInstanceSection(Menu);
+			FToolMenuSection& Section = CreateLevelSection(Menu);
 			FText EntryDesc = LOCTEXT("LevelInstanceEditSubMenuEntry", "");
 			const bool bCanCommit = ContextLevelInstance->CanCommit(&EntryDesc);
 
@@ -464,7 +463,7 @@ namespace LevelInstanceMenuUtils
 
 			if (ContextLevelInstance && !ContextLevelInstance->IsEditing() && !LevelInstanceSubsystem->LevelInstanceHasLevelScriptBlueprint(ContextLevelInstance))
 			{
-				FToolMenuSection& Section = CreateLevelInstanceSection(Menu);
+				FToolMenuSection& Section = CreateLevelSection(Menu);
 
 				Section.AddSubMenu(
 					"BreakLevelInstances",
@@ -495,52 +494,29 @@ namespace LevelInstanceMenuUtils
 						
 			if (ContextLevelInstance && !ContextLevelInstance->IsEditing())
 			{
-				FToolMenuSection& Section = CreateLevelInstanceSection(Menu);
-				TSoftObjectPtr<UBlueprint> BlueprintAsset;
-				if (APackedLevelInstance* PackedLevelInstance = Cast<APackedLevelInstance>(ContextLevelInstance))
+				FToolMenuSection& Section = CreateLevelSection(Menu);
+				;
+				if (APackedLevelActor* PackedLevelActor = Cast<APackedLevelActor>(ContextLevelInstance))
 				{
-					BlueprintAsset = PackedLevelInstance->BlueprintAsset;
-				}
-
-				if (BlueprintAsset.IsNull())
-				{
-					FToolUIAction UIAction;
-					UIAction.ExecuteAction.BindLambda([ContextLevelInstance](const FToolMenuContext& MenuContext)
-					{
-						TSharedPtr<FPackedLevelInstanceBuilder> Builder = FPackedLevelInstanceBuilder::CreateDefaultBuilder();
-						Builder->CreateOrUpdateBlueprint(ContextLevelInstance, nullptr);
-					});
-					UIAction.CanExecuteAction.BindLambda([](const FToolMenuContext& MenuContext)
-					{
-						return GEditor->GetSelectedActorCount() > 0;
-					});
-
-					Section.AddMenuEntry(
-						"CreatePackedBlueprint",
-						LOCTEXT("CreatePackedBlueprint", "Create Packed Blueprint"),
-						TAttribute<FText>(),
-						TAttribute<FSlateIcon>(),
-						UIAction);
-				}
-				else
-				{
-					FToolUIAction UIAction;
-					UIAction.ExecuteAction.BindLambda([ContextLevelInstance, BlueprintAsset](const FToolMenuContext& MenuContext)
-					{
-						TSharedPtr<FPackedLevelInstanceBuilder> Builder = FPackedLevelInstanceBuilder::CreateDefaultBuilder();
-						Builder->CreateOrUpdateBlueprint(ContextLevelInstance->GetWorldAsset(), BlueprintAsset);
-					});
-					UIAction.CanExecuteAction.BindLambda([](const FToolMenuContext& MenuContext)
+					if (TSoftObjectPtr<UBlueprint> BlueprintAsset = PackedLevelActor->BlueprintAsset)
 					{
-						return GEditor->GetSelectedActorCount() > 0;
-					});
-
-					Section.AddMenuEntry(
-						"UpdatePackedBlueprint",
-						LOCTEXT("UpdatePackedBlueprint", "Update Packed Blueprint"),
-						TAttribute<FText>(),
-						TAttribute<FSlateIcon>(),
-						UIAction);
+						FToolUIAction UIAction;
+						UIAction.ExecuteAction.BindLambda([ContextLevelInstance, BlueprintAsset](const FToolMenuContext& MenuContext)
+							{
+								APackedLevelActor::CreateOrUpdateBlueprint(ContextLevelInstance->GetWorldAsset(), BlueprintAsset);
+							});
+						UIAction.CanExecuteAction.BindLambda([](const FToolMenuContext& MenuContext)
+							{
+								return GEditor->GetSelectedActorCount() > 0;
+							});
+
+						Section.AddMenuEntry(
+							"UpdatePackedBlueprint",
+							LOCTEXT("UpdatePackedBlueprint", "Update Packed Blueprint"),
+							TAttribute<FText>(),
+							TAttribute<FSlateIcon>(),
+							UIAction);
+					}
 				}
 			}
 		}
@@ -604,7 +580,7 @@ namespace LevelInstanceMenuUtils
 
 	void CreateBlueprintFromMenu(UToolMenu* Menu, UWorld* WorldAsset)
 	{
-		FToolMenuSection& Section = CreateLevelInstanceSection(Menu);
+		FToolMenuSection& Section = CreateLevelSection(Menu);
 		FToolUIAction UIAction;
 		UIAction.ExecuteAction.BindLambda([WorldAsset](const FToolMenuContext& MenuContext)
 		{
@@ -715,14 +691,14 @@ void FLevelInstanceEditorModule::ExtendContextMenu()
 {
 	if (UToolMenu* BuildMenu = UToolMenus::Get()->ExtendMenu("LevelEditor.MainMenu.Build"))
 	{
-		FToolMenuSection& Section = BuildMenu->AddSection("LevelEditorLevelInstance", LOCTEXT("LevelInstanceHeading", "Level Instance"));
+		FToolMenuSection& Section = BuildMenu->AddSection("LevelEditorLevelInstance", LOCTEXT("PackedLevelActorsHeading", "Packed Level Actor"));
 		FUIAction PackAction(
 			FExecuteAction::CreateLambda([]() 
 			{
 				UWorld* World = GEditor->GetEditorWorldContext().World();
 				if (ULevelInstanceSubsystem* LevelInstanceSubsystem = World->GetSubsystem<ULevelInstanceSubsystem>())
 				{
-					LevelInstanceSubsystem->PackLevelInstances();
+					LevelInstanceSubsystem->PackAllLoadedActors();
 				}
 			}), 
 			FCanExecuteAction::CreateLambda([]()
@@ -730,49 +706,53 @@ void FLevelInstanceEditorModule::ExtendContextMenu()
 				UWorld* World = GEditor->GetEditorWorldContext().World();
 				if (ULevelInstanceSubsystem* LevelInstanceSubsystem = World->GetSubsystem<ULevelInstanceSubsystem>())
 				{
-					return LevelInstanceSubsystem->CanPackLevelInstances();
+					return LevelInstanceSubsystem->CanPackAllLoadedActors();
 				}
 				return false;
 			}));
 
-		FToolMenuEntry& Entry = Section.AddMenuEntry(NAME_None, LOCTEXT("PackLevelInstancesTitle", "Pack Level Instances"),
-			LOCTEXT("PackLevelInstancesTooltip", "Update packed level instances and blueprints"), FSlateIcon(), PackAction, EUserInterfaceActionType::Button);
+		FToolMenuEntry& Entry = Section.AddMenuEntry(NAME_None, LOCTEXT("PackLevelActorsTitle", "Pack Level Actors"),
+			LOCTEXT("PackLevelActorsTooltip", "Update packed level actor blueprints"), FSlateIcon(), PackAction, EUserInterfaceActionType::Button);
 	}
 
-	if (UToolMenu* Menu = UToolMenus::Get()->ExtendMenu("LevelEditor.ActorContextMenu"))
-	{
-		FToolMenuSection& Section = Menu->AddSection("ActorLevelInstance", TAttribute<FText>(), FToolMenuInsert("ActorOptions", EToolMenuInsertType::After));
-		Section.AddSubMenu(NAME_None, LOCTEXT("LevelInstanceSubMenuLabel", "Level Instance"), TAttribute<FText>(), FNewToolMenuDelegate::CreateLambda([this](UToolMenu* ToolMenu)
+	auto AddDynamicSection = [](UToolMenu* ToolMenu)
+	{				
+		if (ULevelEditorContextMenuContext* LevelEditorMenuContext = ToolMenu->Context.FindContext<ULevelEditorContextMenuContext>())
 		{
-			if (ToolMenu)
+			// Use the actor under the cursor if available (e.g. right-click menu).
+			// Otherwise use the first selected actor if there's one (e.g. Actor pulldown menu or outliner).
+			AActor* ContextActor = LevelEditorMenuContext->HitProxyActor;
+			if (!ContextActor && GEditor->GetSelectedActorCount() != 0)
 			{
-				if (ULevelEditorContextMenuContext* LevelEditorMenuContext = ToolMenu->Context.FindContext<ULevelEditorContextMenuContext>())
-				{
-					// Use the actor under the cursor if available (e.g. right-click menu).
-					// Otherwise use the first selected actor if there's one (e.g. Actor pulldown menu or outliner).
-					AActor* ContextActor = LevelEditorMenuContext->HitProxyActor;
-					if (!ContextActor && GEditor->GetSelectedActorCount() != 0)
-					{
-						ContextActor = Cast<AActor>(GEditor->GetSelectedActors()->GetSelectedObject(0));
-					}
-
-					if (ContextActor)
-					{
-						LevelInstanceMenuUtils::CreateEditMenu(ToolMenu, ContextActor);
-						LevelInstanceMenuUtils::CreateCommitMenu(ToolMenu, ContextActor);
-						LevelInstanceMenuUtils::CreateSaveAsMenu(ToolMenu, ContextActor);
-						LevelInstanceMenuUtils::CreateBreakMenu(ToolMenu, ContextActor);
-						LevelInstanceMenuUtils::CreatePackedBlueprintMenu(ToolMenu, ContextActor);
-					}
+				ContextActor = Cast<AActor>(GEditor->GetSelectedActors()->GetSelectedObject(0));
+			}
 
-					LevelInstanceMenuUtils::CreateSetCurrentMenu(ToolMenu, ContextActor);
-					LevelInstanceMenuUtils::CreateMoveSelectionToMenu(ToolMenu);
-					LevelInstanceMenuUtils::CreateCreateMenu(ToolMenu);
-				}
+			if (ContextActor)
+			{
+				LevelInstanceMenuUtils::CreateEditMenu(ToolMenu, ContextActor);
+				LevelInstanceMenuUtils::CreateCommitMenu(ToolMenu, ContextActor);
+				LevelInstanceMenuUtils::CreateSaveAsMenu(ToolMenu, ContextActor);
+				LevelInstanceMenuUtils::CreateBreakMenu(ToolMenu, ContextActor);
+				LevelInstanceMenuUtils::CreatePackedBlueprintMenu(ToolMenu, ContextActor);
 			}
-		}), false, FSlateIcon(FAppStyle::GetAppStyleSetName(), "ClassIcon.LevelInstance"));
+
+			LevelInstanceMenuUtils::CreateSetCurrentMenu(ToolMenu, ContextActor);
+			LevelInstanceMenuUtils::CreateMoveSelectionToMenu(ToolMenu);
+		}
+
+		LevelInstanceMenuUtils::CreateCreateMenu(ToolMenu);
+	};
+
+	if (UToolMenu* ToolMenu = UToolMenus::Get()->ExtendMenu("LevelEditor.ActorContextMenu.LevelSubMenu"))
+	{
+		ToolMenu->AddDynamicSection("LevelInstanceEditorModuleDynamicSection", FNewToolMenuDelegate::CreateLambda(AddDynamicSection));
 	}
 
+	if (UToolMenu* ToolMenu = UToolMenus::Get()->ExtendMenu("LevelEditor.LevelEditorSceneOutliner.ContextMenu.LevelSubMenu"))
+	{
+		ToolMenu->AddDynamicSection("LevelInstanceEditorModuleDynamicSection", FNewToolMenuDelegate::CreateLambda(AddDynamicSection));
+	}
+		
 	if (UToolMenu* WorldAssetMenu = UToolMenus::Get()->ExtendMenu("ContentBrowser.AssetContextMenu.World"))
 	{
 		FToolMenuSection& Section = WorldAssetMenu->AddDynamicSection("ActorLevelInstance", FNewToolMenuDelegate::CreateLambda([this](UToolMenu* ToolMenu)
diff --git a/Engine/Source/Editor/UnrealEd/Public/WorldPartition/WorldPartitionResaveActorsBuilder.h b/Engine/Source/Editor/UnrealEd/Public/WorldPartition/WorldPartitionResaveActorsBuilder.h
--- a/Engine/Source/Editor/UnrealEd/Public/WorldPartition/WorldPartitionResaveActorsBuilder.h
+++ b/Engine/Source/Editor/UnrealEd/Public/WorldPartition/WorldPartitionResaveActorsBuilder.h
@@ -5,7 +5,7 @@
 #include "WorldPartition/WorldPartitionBuilder.h"
 #include "WorldPartitionResaveActorsBuilder.generated.h"
 
-// Example Command Line: ProjectName MapName -run=WorldPartitionBuilderCommandlet -SCCProvider=Perforce -Builder=WorldPartitionResaveActorsBuilder [-ActorClass=PackedLevelInstance] [-SwitchActorPackagingSchemeToReduced]
+// Example Command Line: ProjectName MapName -run=WorldPartitionBuilderCommandlet -SCCProvider=Perforce -Builder=WorldPartitionResaveActorsBuilder [-ActorClass=StaticMeshActor] [-SwitchActorPackagingSchemeToReduced]
 
 UCLASS()
 class UWorldPartitionResaveActorsBuilder : public UWorldPartitionBuilder
diff --git a/Engine/Source/Runtime/Engine/Private/LevelInstance/LevelInstanceActor.cpp b/Engine/Source/Runtime/Engine/Private/LevelInstance/LevelInstanceActor.cpp
--- a/Engine/Source/Runtime/Engine/Private/LevelInstance/LevelInstanceActor.cpp
+++ b/Engine/Source/Runtime/Engine/Private/LevelInstance/LevelInstanceActor.cpp
@@ -527,7 +527,7 @@ void ALevelInstance::PostEditImport()
 		TGuardValue<bool> LoadUnloadGuard(bGuardLoadUnload, true);
 		Super::PostEditImport();
 	}
-	UpdateLevelInstance();
+	UpdateFromLevel();
 }
 
 bool ALevelInstance::CanDeleteSelectedActor(FText& OutReason) const
@@ -600,7 +600,7 @@ void ALevelInstance::EditorGetUnderlyingActors(TSet<AActor*>& OutUnderlyingActor
 	}
 }
 
-void ALevelInstance::UpdateLevelInstance()
+void ALevelInstance::UpdateFromLevel()
 {
 	if (HasValidLevelInstanceID())
 	{
diff --git a/Engine/Source/Runtime/Engine/Private/LevelInstance/LevelInstanceSubsystem.cpp b/Engine/Source/Runtime/Engine/Private/LevelInstance/LevelInstanceSubsystem.cpp
--- a/Engine/Source/Runtime/Engine/Private/LevelInstance/LevelInstanceSubsystem.cpp
+++ b/Engine/Source/Runtime/Engine/Private/LevelInstance/LevelInstanceSubsystem.cpp
@@ -33,8 +33,8 @@
 #include "LevelInstance/LevelInstanceEditorInstanceActor.h"
 #include "Modules/ModuleManager.h"
 #include "Engine/Blueprint.h"
-#include "LevelInstance/Packed/PackedLevelInstanceActor.h"
-#include "LevelInstance/Packed/PackedLevelInstanceBuilder.h"
+#include "PackedLevelActor/PackedLevelActor.h"
+#include "PackedLevelActor/PackedLevelActorBuilder.h"
 #include "Engine/LevelScriptBlueprint.h"
 #include "EdGraph/EdGraph.h"
 #include "UObject/ObjectSaveContext.h"
@@ -444,34 +444,34 @@ void ULevelInstanceSubsystem::OnExitEditorMode()
 	}
 }
 
-bool ULevelInstanceSubsystem::CanPackLevelInstances() const
+bool ULevelInstanceSubsystem::CanPackAllLoadedActors() const
 {
 	return !LevelInstanceEdit;
 }
 
-void ULevelInstanceSubsystem::PackLevelInstances()
+void ULevelInstanceSubsystem::PackAllLoadedActors()
 {
-	if (!CanPackLevelInstances())
+	if (!CanPackAllLoadedActors())
 	{
 		return;
 	}
 
-	// Add Dependencies first so that we pack the LevelInstances in the proper order (depth first)
-	TFunction<void(APackedLevelInstance*, TArray<UBlueprint*>&, TArray<APackedLevelInstance*>&)> GatherDepencenciesRecursive = [&GatherDepencenciesRecursive](APackedLevelInstance* PackedLevelInstance, TArray<UBlueprint*>& BPsToPack, TArray<APackedLevelInstance*>& ToPack)
+	// Add Dependencies first so that we pack in the proper order (depth first)
+	TFunction<void(APackedLevelActor*, TArray<UBlueprint*>&, TArray<APackedLevelActor*>&)> GatherDepencenciesRecursive = [&GatherDepencenciesRecursive](APackedLevelActor* PackedLevelActor, TArray<UBlueprint*>& BPsToPack, TArray<APackedLevelActor*>& ToPack)
 	{
 		// Early out on already processed BPs or non BP Packed LIs.
-		UBlueprint* Blueprint = Cast<UBlueprint>(PackedLevelInstance->GetClass()->ClassGeneratedBy);
-		if ((Blueprint && BPsToPack.Contains(Blueprint)) || ToPack.Contains(PackedLevelInstance))
+		UBlueprint* Blueprint = Cast<UBlueprint>(PackedLevelActor->GetClass()->ClassGeneratedBy);
+		if ((Blueprint && BPsToPack.Contains(Blueprint)) || ToPack.Contains(PackedLevelActor))
 		{
 			return;
 		}
 		
 		// Recursive deps
-		for (const TSoftObjectPtr<UBlueprint>& Dependency : PackedLevelInstance->PackedBPDependencies)
+		for (const TSoftObjectPtr<UBlueprint>& Dependency : PackedLevelActor->PackedBPDependencies)
 		{
 			if (UBlueprint* LoadedDependency = Dependency.LoadSynchronous())
 			{
-				if (APackedLevelInstance* CDO = Cast<APackedLevelInstance>(LoadedDependency->GeneratedClass ? LoadedDependency->GeneratedClass->GetDefaultObject() : nullptr))
+				if (APackedLevelActor* CDO = Cast<APackedLevelActor>(LoadedDependency->GeneratedClass ? LoadedDependency->GeneratedClass->GetDefaultObject() : nullptr))
 				{
 					GatherDepencenciesRecursive(CDO, BPsToPack, ToPack);
 				}
@@ -485,25 +485,25 @@ void ULevelInstanceSubsystem::PackLevelInstances()
 		}
 		else
 		{
-			ToPack.Add(PackedLevelInstance);
+			ToPack.Add(PackedLevelActor);
 		}
 	};
 
-	TArray<APackedLevelInstance*> PackedLevelInstancesToUpdate;
+	TArray<APackedLevelActor*> PackedLevelActorsToUpdate;
 	TArray<UBlueprint*> BlueprintsToUpdate;
 	for (TObjectIterator<UWorld> It(RF_ClassDefaultObject | RF_ArchetypeObject, true); It; ++It)
 	{
 		UWorld* CurrentWorld = *It;
 		if (IsValid(CurrentWorld) && CurrentWorld->GetSubsystem<ULevelInstanceSubsystem>() != nullptr)
 		{
-			for (TActorIterator<APackedLevelInstance> LevelInstanceIt(CurrentWorld); LevelInstanceIt; ++LevelInstanceIt)
+			for (TActorIterator<APackedLevelActor> PackedLevelActorIt(CurrentWorld); PackedLevelActorIt; ++PackedLevelActorIt)
 			{
-				GatherDepencenciesRecursive(*LevelInstanceIt, BlueprintsToUpdate, PackedLevelInstancesToUpdate);
+				GatherDepencenciesRecursive(*PackedLevelActorIt, BlueprintsToUpdate, PackedLevelActorsToUpdate);
 			}
 		}
 	}
 
-	int32 Count = BlueprintsToUpdate.Num() + PackedLevelInstancesToUpdate.Num();
+	int32 Count = BlueprintsToUpdate.Num() + PackedLevelActorsToUpdate.Num();
 	if (!Count)
 	{
 		return;
@@ -511,18 +511,18 @@ void ULevelInstanceSubsystem::PackLevelInstances()
 	
 	GEditor->SelectNone(true, true);
 
-	FScopedSlowTask SlowTask(Count, (LOCTEXT("LevelInstance_PackLevelInstances", "Packing Level Instances")));
+	FScopedSlowTask SlowTask(Count, (LOCTEXT("TaskPackLevels", "Packing Levels")));
 	SlowTask.MakeDialog();
 		
 	auto UpdateProgress = [&SlowTask]()
 	{
 		if (SlowTask.CompletedWork < SlowTask.TotalAmountOfWork)
 		{
-			SlowTask.EnterProgressFrame(1, FText::Format(LOCTEXT("LevelInstance_PackLevelInstancesProgress", "Packing Level Instance {0} of {1}"), FText::AsNumber(SlowTask.CompletedWork), FText::AsNumber(SlowTask.TotalAmountOfWork)));
+			SlowTask.EnterProgressFrame(1, FText::Format(LOCTEXT("TaskPackLevelProgress", "Packing Level {0} of {1}"), FText::AsNumber(SlowTask.CompletedWork), FText::AsNumber(SlowTask.TotalAmountOfWork)));
 		}
 	};
 
-	TSharedPtr<FPackedLevelInstanceBuilder> Builder = FPackedLevelInstanceBuilder::CreateDefaultBuilder();
+	TSharedPtr<FPackedLevelActorBuilder> Builder = FPackedLevelActorBuilder::CreateDefaultBuilder();
 	const bool bCheckoutAndSave = false;
 	for (UBlueprint* Blueprint : BlueprintsToUpdate)
 	{
@@ -530,9 +530,9 @@ void ULevelInstanceSubsystem::PackLevelInstances()
 		UpdateProgress();
 	}
 
-	for (APackedLevelInstance* PackedLevelInstance : PackedLevelInstancesToUpdate)
+	for (APackedLevelActor* PackedLevelActor : PackedLevelActorsToUpdate)
 	{
-		PackedLevelInstance->UpdateLevelInstance();
+		PackedLevelActor->UpdateFromLevel();
 		UpdateProgress();
 	}
 }
@@ -869,36 +869,33 @@ ALevelInstance* ULevelInstanceSubsystem::CreateLevelInstanceFrom(const TArray<AA
 	{
 		NewLevelInstanceActor = GetWorld()->SpawnActor<ALevelInstance>(ALevelInstance::StaticClass(), SpawnParams);
 	}
-	else if (CreationParams.Type == ELevelInstanceCreationType::PackedLevelInstance)
-	{
-		NewLevelInstanceActor = GetWorld()->SpawnActor<APackedLevelInstance>(APackedLevelInstance::StaticClass(), SpawnParams);
-	}
-	else if (CreationParams.Type == ELevelInstanceCreationType::PackedLevelInstanceBlueprint)
+	else
 	{
+		check(CreationParams.Type == ELevelInstanceCreationType::PackedLevelActor);
 		FString PackageDir = FPaths::GetPath(WorldPtr.GetLongPackageName());
-		FString AssetName = FPackedLevelInstanceBuilder::GetPackedBPPrefix() + WorldPtr.GetAssetName();
+		FString AssetName = FPackedLevelActorBuilder::GetPackedBPPrefix() + WorldPtr.GetAssetName();
 		FString BPAssetPath = FString::Format(TEXT("{0}/{1}.{1}"), { PackageDir , AssetName });
 		const bool bCompile = true;
 
 		UBlueprint* NewBP = nullptr;
 		if (CreationParams.LevelPackageName.IsEmpty())
 		{
-			NewBP = FPackedLevelInstanceBuilder::CreatePackedLevelInstanceBlueprintWithDialog(TSoftObjectPtr<UBlueprint>(BPAssetPath), WorldPtr, bCompile);
+			NewBP = FPackedLevelActorBuilder::CreatePackedLevelActorBlueprintWithDialog(TSoftObjectPtr<UBlueprint>(BPAssetPath), WorldPtr, bCompile);
 		}
 		else
 		{
-			NewBP = FPackedLevelInstanceBuilder::CreatePackedLevelInstanceBlueprint(TSoftObjectPtr<UBlueprint>(BPAssetPath), WorldPtr, bCompile);
+			NewBP = FPackedLevelActorBuilder::CreatePackedLevelActorBlueprint(TSoftObjectPtr<UBlueprint>(BPAssetPath), WorldPtr, bCompile);
 		}
 				
 		if (NewBP)
 		{
-			NewLevelInstanceActor = GetWorld()->SpawnActor<APackedLevelInstance>(NewBP->GeneratedClass, SpawnParams);
+			NewLevelInstanceActor = GetWorld()->SpawnActor<APackedLevelActor>(NewBP->GeneratedClass, SpawnParams);
 		}
 
 		if (!NewLevelInstanceActor)
 		{
 			UE_LOG(LogLevelInstance, Warning, TEXT("Failed to create packed level blueprint. Creating non blueprint packed level instance instead."));
-			NewLevelInstanceActor = GetWorld()->SpawnActor<APackedLevelInstance>(APackedLevelInstance::StaticClass(), SpawnParams);
+			NewLevelInstanceActor = GetWorld()->SpawnActor<APackedLevelActor>(APackedLevelActor::StaticClass(), SpawnParams);
 		}
 	}
 	
@@ -995,8 +992,8 @@ void ULevelInstanceSubsystem::BreakLevelInstance_Impl(ALevelInstance* LevelInsta
 		// Can only break the top level LevelInstance
 		check(LevelInstanceActor->GetLevel() == GetWorld()->GetCurrentLevel());
 
-		// Actors in a packed level instance will not be streamed in unless they are editing. Must force this before moving.
-		if (LevelInstanceActor->IsA<APackedLevelInstance>())
+		// Actors in a packed level actor will not be streamed in unless they are editing. Must force this before moving.
+		if (LevelInstanceActor->IsA<APackedLevelActor>())
 		{
 			BlockLoadLevelInstance(LevelInstanceActor);
 		}
@@ -1009,7 +1006,7 @@ void ULevelInstanceSubsystem::BreakLevelInstance_Impl(ALevelInstance* LevelInsta
 		{
 			UE_LOG(LogLevelInstance, Warning, TEXT("Failed to completely break Level Instance because some children have Level Scripts."));
 
-			if (LevelInstanceActor->IsA<APackedLevelInstance>())
+			if (LevelInstanceActor->IsA<APackedLevelActor>())
 			{
 				BlockUnloadLevelInstance(LevelInstanceActor);
 			}
@@ -1077,7 +1074,7 @@ void ULevelInstanceSubsystem::BreakLevelInstance_Impl(ALevelInstance* LevelInsta
 			return;
 		}
 
-		if (LevelInstanceActor->IsA<APackedLevelInstance>())
+		if (LevelInstanceActor->IsA<APackedLevelActor>())
 		{
 			BlockUnloadLevelInstance(LevelInstanceActor);
 		}
@@ -1263,7 +1260,7 @@ void ULevelInstanceSubsystem::OnActorDeleted(AActor* Actor)
 		}
 
 		const bool bIsEditingLevelInstance = IsEditingLevelInstance(LevelInstanceActor);
-		if (!bIsEditingLevelInstance && LevelInstanceActor->IsA<APackedLevelInstance>())
+		if (!bIsEditingLevelInstance && LevelInstanceActor->IsA<APackedLevelActor>())
 		{
 			return;
 		}
@@ -1646,7 +1643,7 @@ ALevelInstance* ULevelInstanceSubsystem::CommitLevelInstanceInternal(TUniquePtr<
 		}
 	}
 
-	FScopedSlowTask SlowTask(0, LOCTEXT("EndEditLevelInstance", "Unloading edit Level Instance..."), !GetWorld()->IsGameWorld());
+	FScopedSlowTask SlowTask(0, LOCTEXT("EndEditLevelInstance", "Unloading Level..."), !GetWorld()->IsGameWorld());
 	SlowTask.MakeDialog();
 
 	GEditor->SelectNone(false, true);
@@ -1696,7 +1693,7 @@ ALevelInstance* ULevelInstanceSubsystem::CommitLevelInstanceInternal(TUniquePtr<
 	{
 		if (ALevelInstance* LevelInstance = GetLevelInstance(LevelInstanceToUpdateID))
 		{
-			LevelInstance->UpdateLevelInstance();
+			LevelInstance->UpdateFromLevel();
 		}
 	}
 
@@ -1747,7 +1744,7 @@ void ULevelInstanceSubsystem::SaveLevelInstanceAs(ALevelInstance* LevelInstanceA
 
 	if (OutObjects.Num() == 0 || OutObjects[0] == EditingWorld)
 	{
-		UE_LOG(LogLevelInstance, Warning, TEXT("Failed to save Level Instance as new asset"));
+		UE_LOG(LogLevelInstance, Warning, TEXT("Failed to save Level as new asset"));
 		return;
 	}
 
diff --git a/Engine/Source/Runtime/Engine/Private/LevelInstance/Packed/LevelInstanceISMPacker.cpp b/Engine/Source/Runtime/Engine/Private/LevelInstance/Packed/LevelInstanceISMPacker.cpp
deleted file mode 100644
diff --git a/Engine/Source/Runtime/Engine/Private/LevelInstance/Packed/LevelInstanceRecursivePacker.cpp b/Engine/Source/Runtime/Engine/Private/LevelInstance/Packed/LevelInstanceRecursivePacker.cpp
deleted file mode 100644
diff --git a/Engine/Source/Runtime/Engine/Private/LevelInstance/Packed/PackedLevelInstanceActor.cpp b/Engine/Source/Runtime/Engine/Private/LevelInstance/Packed/PackedLevelInstanceActor.cpp
deleted file mode 100644
diff --git a/Engine/Source/Runtime/Engine/Private/LevelInstance/Packed/PackedLevelInstanceBuilder.cpp b/Engine/Source/Runtime/Engine/Private/LevelInstance/Packed/PackedLevelInstanceBuilder.cpp
deleted file mode 100644
diff --git a/Engine/Source/Runtime/Engine/Private/PackedLevelActor/IPackedLevelActorBuilder.h b/Engine/Source/Runtime/Engine/Private/PackedLevelActor/IPackedLevelActorBuilder.h
new file mode 100644
--- /dev/null
+++ b/Engine/Source/Runtime/Engine/Private/PackedLevelActor/IPackedLevelActorBuilder.h
@@ -0,0 +1,26 @@
+// Copyright Epic Games, Inc. All Rights Reserved.
+
+#pragma once
+
+#if WITH_EDITOR
+
+#include "CoreMinimal.h"
+#include "UObject/ObjectMacros.h"
+#include "PackedLevelActor/PackedLevelActorTypes.h"
+
+class APackedLevelActor;
+class AActor;
+class UActorComponent;
+class FPackedLevelActorBuilderContext;
+
+class ENGINE_API IPackedLevelActorBuilder
+{
+public:
+	IPackedLevelActorBuilder() {}
+	virtual ~IPackedLevelActorBuilder() {}
+	virtual FPackedLevelActorBuilderID GetID() const = 0;
+	virtual void GetPackClusters(FPackedLevelActorBuilderContext& InContext, AActor* InActor) const = 0;
+	virtual void PackActors(FPackedLevelActorBuilderContext& InBuilder, APackedLevelActor* InPackingActor, const FPackedLevelActorBuilderClusterID& InClusterID, const TArray<UActorComponent*>& InComponents) const = 0;
+};
+
+#endif
\ No newline at end of file
diff --git a/Engine/Source/Runtime/Engine/Private/PackedLevelActor/PackedLevelActor.cpp b/Engine/Source/Runtime/Engine/Private/PackedLevelActor/PackedLevelActor.cpp
new file mode 100644
--- /dev/null
+++ b/Engine/Source/Runtime/Engine/Private/PackedLevelActor/PackedLevelActor.cpp
@@ -0,0 +1,238 @@
+// Copyright Epic Games, Inc. All Rights Reserved.
+
+#include "PackedLevelActor/PackedLevelActor.h"
+#include "PackedLevelActor/PackedLevelActorBuilder.h"
+#include "LevelInstance/LevelInstanceSubsystem.h"
+#include "Components/InstancedStaticMeshComponent.h"
+#include "LevelInstance/LevelInstancePrivate.h"
+#include "UObject/UE5ReleaseStreamObjectVersion.h"
+
+APackedLevelActor::APackedLevelActor(const FObjectInitializer& ObjectInitializer)
+	: Super(ObjectInitializer)
+#if WITH_EDITORONLY_DATA
+	, ISMComponentClass(UInstancedStaticMeshComponent::StaticClass())
+	, bChildChanged(false)
+#endif
+{
+#if WITH_EDITORONLY_DATA
+	// Packed Level Instances don't support level streaming or sub actors
+	DesiredRuntimeBehavior = ELevelInstanceRuntimeBehavior::None;
+#endif
+}
+
+void APackedLevelActor::Serialize(FArchive& Ar)
+{
+	Ar.UsingCustomVersion(FUE5ReleaseStreamObjectVersion::GUID);
+	Super::Serialize(Ar);
+
+#if WITH_EDITORONLY_DATA
+	// We want to make sure we serialize that property so we can compare to the CDO
+	if (!Ar.IsFilterEditorOnly() && Ar.CustomVer(FUE5ReleaseStreamObjectVersion::GUID) >= FUE5ReleaseStreamObjectVersion::PackedLevelInstanceVersion)
+	{
+		Ar << PackedVersion;
+	}
+#endif
+}
+
+bool APackedLevelActor::SupportsLoading() const
+{
+#if WITH_EDITOR
+	return HasChildEdit() || IsLoaded();
+#else
+	return false;
+#endif
+}
+
+#if WITH_EDITOR
+
+void APackedLevelActor::PostLoad()
+{
+	Super::PostLoad();
+
+	// Non CDO: Set the Guid to something different then the default value so that we actually run the construction script on actors that haven't been resaved against their latest BP
+	if (!HasAnyFlags(RF_ClassDefaultObject | RF_ArchetypeObject) && GetLinkerCustomVersion(FUE5ReleaseStreamObjectVersion::GUID) < FUE5ReleaseStreamObjectVersion::PackedLevelInstanceVersion)
+	{
+		const static FGuid NoVersionGUID(0x50817615, 0x74A547A3, 0x9295D655, 0x8A852C0F);
+		PackedVersion = NoVersionGUID;
+	}
+}
+
+void APackedLevelActor::RerunConstructionScripts()
+{
+	bool bShouldRerunConstructionScript = true;
+
+	if (GetWorld() && GetWorld()->IsGameWorld())
+	{
+		bShouldRerunConstructionScript = false;
+
+		// Only rerun if version mismatchs
+		if (PackedVersion != GetClass()->GetDefaultObject<APackedLevelActor>()->PackedVersion)
+		{
+			bShouldRerunConstructionScript = true;
+			UE_LOG(LogLevelInstance, Verbose, TEXT("RerunConstructionScript was executed on %s (%s) because its version (%s) doesn't match latest version (%s). Resaving this actor will fix this"),
+				*GetPathName(),
+				*GetPackage()->GetPathName(),
+				*PackedVersion.ToString(),
+				*GetClass()->GetDefaultObject<APackedLevelActor>()->PackedVersion.ToString());
+		}
+	}
+	
+	if(bShouldRerunConstructionScript)
+	{
+		Super::RerunConstructionScripts();
+		PackedVersion = GetClass()->GetDefaultObject<APackedLevelActor>()->PackedVersion;
+	}
+}
+
+bool APackedLevelActor::CreateOrUpdateBlueprint(ALevelInstance* InLevelInstance, TSoftObjectPtr<UBlueprint> InBlueprintAsset, bool bCheckoutAndSave, bool bPromptForSave)
+{
+	return FPackedLevelActorBuilder::CreateDefaultBuilder()->CreateOrUpdateBlueprint(InLevelInstance, InBlueprintAsset, bCheckoutAndSave, bPromptForSave);
+}
+
+bool APackedLevelActor::CreateOrUpdateBlueprint(TSoftObjectPtr<UWorld> InWorldAsset, TSoftObjectPtr<UBlueprint> InBlueprintAsset, bool bCheckoutAndSave, bool bPromptForSave)
+{
+	return FPackedLevelActorBuilder::CreateDefaultBuilder()->CreateOrUpdateBlueprint(InWorldAsset, InBlueprintAsset, bCheckoutAndSave, bPromptForSave);
+}
+
+FName APackedLevelActor::GetPackedComponentTag()
+{
+	static FName PackedComponentTag("PackedComponent");
+	return PackedComponentTag;
+}
+
+void APackedLevelActor::UpdateFromLevel()
+{
+	Super::UpdateFromLevel();
+
+	if (!Cast<UBlueprint>(GetClass()->ClassGeneratedBy))
+	{
+		if (IsLevelInstancePathValid())
+		{
+			TSharedPtr<FPackedLevelActorBuilder> Builder = FPackedLevelActorBuilder::CreateDefaultBuilder();
+			Builder->PackActor(this, GetWorldAsset());
+		}
+		else
+		{
+			DestroyPackedComponents();
+		}
+	}
+}
+
+void APackedLevelActor::OnEditChild()
+{
+	Super::OnEditChild();
+
+	check(GetLevelInstanceSubsystem()->GetLevelInstanceLevel(this) != nullptr);
+	MarkComponentsRenderStateDirty();
+}
+
+void APackedLevelActor::OnCommitChild(bool bChanged)
+{
+	Super::OnCommitChild(bChanged);
+
+	ULevelInstanceSubsystem* LevelInstanceSubsystem = GetLevelInstanceSubsystem();
+	check(GetLevelInstanceSubsystem()->GetLevelInstanceLevel(this));
+	bChildChanged |= bChanged;
+	if (!HasChildEdit())
+	{
+		UnloadLevelInstance();
+
+		if (bChildChanged)
+		{
+			// Reflect child changes
+			TSharedPtr<FPackedLevelActorBuilder> Builder = FPackedLevelActorBuilder::CreateDefaultBuilder();
+
+			if (UBlueprint* GeneratedBy = Cast<UBlueprint>(GetClass()->ClassGeneratedBy))
+			{
+				check(GeneratedBy == BlueprintAsset.Get());
+				Builder->UpdateBlueprint(GeneratedBy);
+			}
+			else
+			{
+				Builder->PackActor(this, GetWorldAsset());
+			}
+			bChildChanged = false;
+		}
+
+		MarkComponentsRenderStateDirty();
+	}
+}
+
+void APackedLevelActor::OnEdit()
+{
+	Super::OnEdit();
+	MarkComponentsRenderStateDirty();
+}
+
+void APackedLevelActor::OnCommit(bool bChanged, bool bPromptForSave)
+{
+	Super::OnCommit(bChanged, bPromptForSave);
+
+	if (bChanged)
+	{
+		if (UBlueprint* GeneratedBy = Cast<UBlueprint>(GetClass()->ClassGeneratedBy))
+		{
+			check(GeneratedBy == BlueprintAsset.Get());
+			const bool bCheckoutAndSave = true;
+			TSharedPtr<FPackedLevelActorBuilder> Builder = FPackedLevelActorBuilder::CreateDefaultBuilder();
+			Builder->UpdateBlueprint(GeneratedBy, bCheckoutAndSave, bPromptForSave);
+		}
+	}
+	
+	// bEditing flag changed so dirty render state
+	MarkComponentsRenderStateDirty();
+}
+
+bool APackedLevelActor::IsHiddenEd() const
+{
+	return Super::IsHiddenEd() || IsEditing() || HasChildEdit();
+}
+
+bool APackedLevelActor::IsHLODRelevant() const
+{
+	// Bypass base class ALevelInstance (because it always returns true). We want the same implementation as AActor.
+	return AActor::IsHLODRelevant();
+}
+
+bool APackedLevelActor::CanEditChange(const FProperty* InProperty) const
+{
+	if (!Super::CanEditChange(InProperty))
+	{
+		return false;
+	}
+
+	// Disallow editing of the World if we are a BP instance
+	if (InProperty->GetFName() == GET_MEMBER_NAME_CHECKED(APackedLevelActor, WorldAsset))
+	{
+		return GetClass()->IsNative();
+	}
+
+	return true;
+}
+
+void APackedLevelActor::GetPackedComponents(TArray<UActorComponent*>& OutPackedComponents) const
+{
+	const TSet<UActorComponent*>& Components = GetComponents();
+	OutPackedComponents.Reserve(Components.Num());
+		
+	for (UActorComponent* Component : Components)
+	{
+		if (Component && Component->ComponentHasTag(GetPackedComponentTag()))
+		{
+			OutPackedComponents.Add(Component);
+		}
+	}
+}
+
+void APackedLevelActor::DestroyPackedComponents()
+{
+	Modify();
+	TArray<UActorComponent*> PackedComponents;
+	GetPackedComponents(PackedComponents);
+	for (UActorComponent* PackedComponent : PackedComponents)
+	{
+		PackedComponent->Modify();
+		PackedComponent->DestroyComponent();
+	}
+}
+#endif
\ No newline at end of file
diff --git a/Engine/Source/Runtime/Engine/Private/PackedLevelActor/PackedLevelActorBuilder.cpp b/Engine/Source/Runtime/Engine/Private/PackedLevelActor/PackedLevelActorBuilder.cpp
new file mode 100644
--- /dev/null
+++ b/Engine/Source/Runtime/Engine/Private/PackedLevelActor/PackedLevelActorBuilder.cpp
@@ -0,0 +1,494 @@
+// Copyright Epic Games, Inc. All Rights Reserved.
+
+#include "PackedLevelActor/PackedLevelActorBuilder.h"
+
+#if WITH_EDITOR
+
+#include "PackedLevelActor/PackedLevelActor.h"
+#include "PackedLevelActor/IPackedLevelActorBuilder.h"
+#include "PackedLevelActor/PackedLevelActorISMBuilder.h"
+#include "PackedLevelActor/PackedLevelActorRecursiveBuilder.h"
+
+#include "LevelInstance/LevelInstanceSubsystem.h"
+
+#include "Kismet2/KismetEditorUtilities.h"
+#include "Engine/SimpleConstructionScript.h"
+#include "Engine/SCS_Node.h"
+
+#include "Misc/ScopeExit.h"
+#include "Misc/Paths.h"
+#include "Logging/LogMacros.h"
+#include "Logging/MessageLog.h"
+
+#include "Components/SceneComponent.h"
+#include "GameFramework/WorldSettings.h"
+#include "Engine/LevelBounds.h"
+#include "Engine/LevelStreaming.h"
+#include "Engine/Brush.h"
+#include "FileHelpers.h"
+#include "Editor.h"
+#include "AssetToolsModule.h"
+#include "Factories/BlueprintFactory.h"
+
+#include "IContentBrowserSingleton.h"
+#include "ContentBrowserModule.h"
+#include "EditorDirectories.h"
+
+#define LOCTEXT_NAMESPACE "FPackedLevelActorBuilder"
+
+void FPackedLevelActorBuilderContext::ClusterLevelActor(AActor* InActor)
+{
+	if (!ActorDiscards.Contains(InActor))
+	{
+		PerActorClusteredComponents.FindOrAdd(InActor);
+
+		for (const auto& Pair : Builders)
+		{
+			Pair.Value->GetPackClusters(*this, InActor);
+		}
+	}
+}
+
+void FPackedLevelActorBuilderContext::FindOrAddCluster(FPackedLevelActorBuilderClusterID&& InClusterID, UActorComponent* InComponent)
+{
+	TArray<UActorComponent*>& ClusterComponents = Clusters.FindOrAdd(MoveTemp(InClusterID));
+	if (InComponent)
+	{
+		ClusterComponents.Add(InComponent);
+		PerActorClusteredComponents.FindChecked(InComponent->GetOwner()).Add(InComponent);
+	}
+}
+
+void FPackedLevelActorBuilderContext::DiscardActor(AActor* InActor)
+{
+	ActorDiscards.Add(InActor);
+}
+
+FPackedLevelActorBuilder::FPackedLevelActorBuilder()
+{
+}
+
+const FString& FPackedLevelActorBuilder::GetPackedBPPrefix()
+{
+	static FString BPPrefix = "BPP_";
+	return BPPrefix;
+}
+
+UBlueprint* FPackedLevelActorBuilder::CreatePackedLevelActorBlueprintWithDialog(TSoftObjectPtr<UBlueprint> InBlueprintAsset, TSoftObjectPtr<UWorld> InWorldAsset, bool bInCompile)
+{
+	FSaveAssetDialogConfig SaveAssetDialogConfig;
+	SaveAssetDialogConfig.DialogTitleOverride = LOCTEXT("SaveAssetDialogTitle", "Save Asset As");
+	SaveAssetDialogConfig.DefaultPath = FPaths::GetPath(InBlueprintAsset.GetLongPackageName());
+	SaveAssetDialogConfig.DefaultAssetName = InBlueprintAsset.GetAssetName();
+	SaveAssetDialogConfig.ExistingAssetPolicy = ESaveAssetDialogExistingAssetPolicy::AllowButWarn;
+
+	FContentBrowserModule& ContentBrowserModule = FModuleManager::LoadModuleChecked<FContentBrowserModule>("ContentBrowser");
+	FString SaveObjectPath = ContentBrowserModule.Get().CreateModalSaveAssetDialog(SaveAssetDialogConfig);
+
+	if (!SaveObjectPath.IsEmpty())
+	{
+		TSoftObjectPtr<UBlueprint> ExistingBPAsset(SaveObjectPath);
+
+		if (UBlueprint* BP = ExistingBPAsset.LoadSynchronous())
+		{
+			return BP;
+		}
+				
+		return CreatePackedLevelActorBlueprint(ExistingBPAsset, InWorldAsset, bInCompile);
+	}
+
+	return nullptr;
+}
+
+UBlueprint* FPackedLevelActorBuilder::CreatePackedLevelActorBlueprint(TSoftObjectPtr<UBlueprint> InBlueprintAsset, TSoftObjectPtr<UWorld> InWorldAsset, bool bInCompile)
+{
+	IAssetTools& AssetTools = FAssetToolsModule::GetModule().Get();
+
+	UBlueprintFactory* BlueprintFactory = NewObject<UBlueprintFactory>();
+	BlueprintFactory->ParentClass = APackedLevelActor::StaticClass();
+	BlueprintFactory->bSkipClassPicker = true;
+
+	FEditorDelegates::OnConfigureNewAssetProperties.Broadcast(BlueprintFactory);
+	if (BlueprintFactory->ConfigureProperties())
+	{
+		FString PackageDir = FPaths::GetPath(InBlueprintAsset.GetLongPackageName());
+		FEditorDirectories::Get().SetLastDirectory(ELastDirectory::NEW_ASSET, PackageDir);
+
+		if (UBlueprint* NewBP = Cast<UBlueprint>(AssetTools.CreateAsset(InBlueprintAsset.GetAssetName(), PackageDir, UBlueprint::StaticClass(), BlueprintFactory, FName("Create LevelInstance Blueprint"))))
+		{
+			APackedLevelActor* CDO = CastChecked<APackedLevelActor>(NewBP->GeneratedClass->GetDefaultObject());
+			CDO->BlueprintAsset = NewBP;
+			CDO->SetWorldAsset(InWorldAsset);
+
+			if (bInCompile)
+			{
+				FKismetEditorUtilities::CompileBlueprint(NewBP, EBlueprintCompileOptions::None);
+			}
+
+			AssetTools.SyncBrowserToAssets(TArray<UObject*>{ NewBP });
+
+			return NewBP;
+		}
+	}
+
+	return nullptr;
+}
+
+TSharedPtr<FPackedLevelActorBuilder> FPackedLevelActorBuilder::CreateDefaultBuilder()
+{
+	TSharedPtr<FPackedLevelActorBuilder> Builder = MakeShared<FPackedLevelActorBuilder>();
+
+	// Class Discards are used to validate the packing result.
+	// Components or Actor classes in this set will not generate warnings
+	Builder->ClassDiscards.Add(ALevelBounds::StaticClass());
+	// Avoid dependency (find class)
+	UClass* ChaosDebugClass = FindObject<UClass>(ANY_PACKAGE, TEXT("ChaosDebugDrawComponent"));
+	if (ChaosDebugClass)
+	{
+		Builder->ClassDiscards.Add(ChaosDebugClass);
+	}
+
+	// Root Components that are SceneComponents (not child class of)
+	Builder->ClassDiscards.Add(USceneComponent::StaticClass());
+	
+	Builder->Builders.Add(FPackedLevelActorRecursiveBuilder::BuilderID, MakeUnique<FPackedLevelActorRecursiveBuilder>());
+	Builder->Builders.Add(FPackedLevelActorISMBuilder::BuilderID, MakeUnique<FPackedLevelActorISMBuilder>());
+
+	return Builder;
+}
+
+bool FPackedLevelActorBuilder::PackActor(APackedLevelActor* InPackedLevelActor)
+{
+	return PackActor(InPackedLevelActor, InPackedLevelActor);
+}
+
+bool FPackedLevelActorBuilder::PackActor(APackedLevelActor* InPackedLevelActor, ALevelInstance* InLevelInstanceToPack)
+{
+	FMessageLog LevelInstanceLog("LevelInstance");
+	LevelInstanceLog.Info(FText::Format(LOCTEXT("PackingStarted", "Packing of '{0}' started..."), FText::FromString(InPackedLevelActor->GetWorldAssetPackage())));
+	
+	FPackedLevelActorBuilderContext Context(*this, InPackedLevelActor);
+
+	InPackedLevelActor->DestroyPackedComponents();
+
+	ULevelInstanceSubsystem* LevelInstanceSubystem = InPackedLevelActor->GetLevelInstanceSubsystem();
+	check(LevelInstanceSubystem);
+	
+	ULevel* SourceLevel = LevelInstanceSubystem->GetLevelInstanceLevel(InLevelInstanceToPack);
+	if (!SourceLevel)
+	{
+		LevelInstanceLog.Error(FText::Format(LOCTEXT("FailedPackingNoLevel", "Packing of '{0}' failed"), FText::FromString(InPackedLevelActor->GetWorldAssetPackage())));
+		return false;
+	}
+
+	ULevelStreaming* SourceLevelStreaming = ULevelStreaming::FindStreamingLevel(SourceLevel);
+	AWorldSettings* WorldSettings = SourceLevel->GetWorldSettings();
+	Context.DiscardActor(WorldSettings);
+
+	// Build relative transform without rotation because pivots don't support rotation
+	FTransform CurrentPivotTransform(SourceLevelStreaming->LevelTransform.GetRelativeTransform(InPackedLevelActor->GetActorTransform()).GetTranslation());
+	FTransform NewPivotTransform(WorldSettings->LevelInstancePivotOffset);
+	FTransform RelativePivotTransform(NewPivotTransform.GetRelativeTransform(CurrentPivotTransform));
+		
+	Context.SetRelativePivotTransform(RelativePivotTransform);
+	
+	if (AActor* DefaultBrush = SourceLevel->GetDefaultBrush())
+	{
+		Context.DiscardActor(DefaultBrush);
+	}
+		
+	for (AActor* LevelActor : SourceLevel->Actors)
+	{
+		if (LevelActor)
+		{
+			Context.ClusterLevelActor(LevelActor);
+		}
+	}
+
+	for (const auto& Pair : Context.GetClusters())
+	{
+		TUniquePtr<IPackedLevelActorBuilder>& Builder = Builders.FindChecked(Pair.Key.GetBuilderID());
+		Builder->PackActors(Context, InPackedLevelActor, Pair.Key, Pair.Value);
+	}
+			
+	Context.Report(LevelInstanceLog);
+	return true;
+}
+
+bool FPackedLevelActorBuilderContext::ShouldPackComponent(UActorComponent* ActorComponent) const
+{
+	return ActorComponent && !ActorComponent->IsVisualizationComponent();
+}
+
+void FPackedLevelActorBuilderContext::Report(FMessageLog& LevelInstanceLog) const
+{
+	TSet<UActorComponent*> NotClusteredComponents;
+	uint32 TotalWarningCount = 0;
+
+	for (const auto& Pair : PerActorClusteredComponents)
+	{
+		AActor* Actor = Pair.Key;
+		const TSet<UActorComponent*>& ClusteredComponents = Pair.Value;
+
+		if (ActorDiscards.Contains(Actor))
+		{
+			LevelInstanceLog.Info(FText::Format(LOCTEXT("ActorDiscard", "Actor '{0}' ignored (Actor Discard)"), FText::FromString(Actor->GetPathName())));
+			continue;
+		}
+
+		if (Actor->GetClass()->HasAnyClassFlags(CLASS_Transient))
+		{
+			LevelInstanceLog.Info(FText::Format(LOCTEXT("ActorTransientClassDiscard", "Actor '{0}' of type '{1}' ignored (Transient Class Discard)"), FText::FromString(Actor->GetPathName()), FText::FromString(Actor->GetClass()->GetPathName())));
+			continue;
+		}
+
+		// Class must match (not a child)
+		if (ClassDiscards.Contains(Actor->GetClass()))
+		{
+			LevelInstanceLog.Info(FText::Format(LOCTEXT("ActorClassDiscard", "Actor '{0}' of type '{1}' ignored (Class Discard)"), FText::FromString(Actor->GetPathName()), FText::FromString(Actor->GetClass()->GetPathName())));
+			continue;
+		}
+
+		NotClusteredComponents = Actor->GetComponents().Difference(ClusteredComponents);
+		uint32 WarningCount = 0;
+		for (UActorComponent* Component : NotClusteredComponents)
+		{
+			if (!ShouldPackComponent(Component))
+			{
+				continue;
+			}
+
+			if (ClassDiscards.Contains(Component->GetClass()))
+			{
+				LevelInstanceLog.Info(FText::Format(LOCTEXT("ComponentClassDiscard", "Component '{0}' of type '{1}' ignored (Class Discard)"), FText::FromString(Component->GetPathName()), FText::FromString(Component->GetClass()->GetPathName())));
+				continue;
+			}
+
+			if (Actor->GetClass()->HasAnyClassFlags(CLASS_Transient))
+			{
+				LevelInstanceLog.Info(FText::Format(LOCTEXT("ComopnentTransientClassDiscard", "Component '{0}' of type '{1}' ignored (Transient Class Discard)"), FText::FromString(Component->GetPathName()), FText::FromString(Component->GetClass()->GetPathName())));
+				continue;
+			}
+
+			WarningCount++;
+
+			LevelInstanceLog.Warning(FText::Format(LOCTEXT("ComponentNotPacked", "Component '{0}' was not packed"), FText::FromString(Component->GetPathName())));
+		}
+
+		if (WarningCount)
+		{
+			LevelInstanceLog.Warning(FText::Format(LOCTEXT("ActorNotPacked", "Actor '{0}' was not packed completely ({1} warning(s))"), FText::FromString(Actor->GetPathName()), FText::AsNumber(WarningCount)));
+		}
+		else
+		{
+			LevelInstanceLog.Info(FText::Format(LOCTEXT("ActorPacked", "Actor '{0}' packed successfully"), FText::FromString(Actor->GetPathName())));
+		}
+		TotalWarningCount += WarningCount;
+	}
+
+	if (TotalWarningCount)
+	{
+		LevelInstanceLog.Warning(LOCTEXT("WarningsReported", "Warnings have been reported. Consider using a regular ALevelInstance instead."));
+		LevelInstanceLog.Open();
+	}
+	LevelInstanceLog.Info(FText::Format(LOCTEXT("PackCompleted", "Packing '{0}' completed with {1} warning(s)"), FText::FromString(PackedLevelActor->GetWorldAssetPackage()), FText::AsNumber(TotalWarningCount)));
+}
+
+ALevelInstance* FPackedLevelActorBuilder::CreateTransientLevelInstanceForPacking(TSoftObjectPtr<UWorld> InWorldAsset, const FVector& InLocation, const FRotator& InRotator)
+{
+	// Create Temp Actor for Packing
+	FActorSpawnParameters SpawnParams;
+	SpawnParams.bCreateActorPackage = false;
+	SpawnParams.bHideFromSceneOutliner = true;
+	SpawnParams.bNoFail = true;
+	SpawnParams.ObjectFlags |= RF_Transient;
+
+	UWorld* World = GEditor->GetEditorWorldContext().World();
+	check(World);
+	SpawnParams.OverrideLevel = World->PersistentLevel;
+	ALevelInstance* LevelInstance = World->SpawnActor<ALevelInstance>(InLocation, InRotator, SpawnParams);
+	LevelInstance->SetWorldAsset(InWorldAsset);
+
+	// Wait for load
+	LevelInstance->GetLevelInstanceSubsystem()->BlockLoadLevelInstance(LevelInstance);
+
+	return LevelInstance;
+}
+
+bool FPackedLevelActorBuilder::PackActor(APackedLevelActor* InActor, TSoftObjectPtr<UWorld> InWorldAsset)
+{
+	ALevelInstance* TransientLevelInstance = CreateTransientLevelInstanceForPacking(InWorldAsset, InActor->GetActorLocation(), InActor->GetActorRotation());
+	ON_SCOPE_EXIT
+	{
+		TransientLevelInstance->GetWorld()->DestroyActor(TransientLevelInstance);
+	};
+	
+	return PackActor(InActor, TransientLevelInstance);
+}
+
+void FPackedLevelActorBuilder::UpdateBlueprint(UBlueprint* Blueprint, bool bCheckoutAndSave, bool bPromptForSave)
+{
+	APackedLevelActor* CDO = CastChecked<APackedLevelActor>(Blueprint->GeneratedClass->GetDefaultObject());
+	check(CDO);
+
+	CreateOrUpdateBlueprint(CDO->GetWorldAsset(), Blueprint, bCheckoutAndSave, bPromptForSave);
+}
+
+bool FPackedLevelActorBuilder::CreateOrUpdateBlueprint(TSoftObjectPtr<UWorld> InWorldAsset, TSoftObjectPtr<UBlueprint> InBlueprintAsset, bool bCheckoutAndSave, bool bPromptForSave)
+{
+	bool bResult = true;
+	
+	ALevelInstance* TransientLevelInstance = CreateTransientLevelInstanceForPacking(InWorldAsset, FVector::ZeroVector, FRotator::ZeroRotator);
+	
+	bResult = CreateOrUpdateBlueprintFromUnpacked(TransientLevelInstance, InBlueprintAsset, bCheckoutAndSave, bPromptForSave);
+
+	TransientLevelInstance->GetWorld()->DestroyActor(TransientLevelInstance);
+
+	return bResult;
+}
+
+bool FPackedLevelActorBuilder::CreateOrUpdateBlueprint(ALevelInstance* InLevelInstance, TSoftObjectPtr<UBlueprint> InBlueprintAsset, bool bCheckoutAndSave, bool bPromptForSave)
+{
+	if (APackedLevelActor* PackedLevelActor = Cast<APackedLevelActor>(InLevelInstance))
+	{
+		return CreateOrUpdateBlueprintFromPacked(PackedLevelActor, InBlueprintAsset, bCheckoutAndSave, bPromptForSave);
+	}
+	
+	return CreateOrUpdateBlueprintFromUnpacked(InLevelInstance, InBlueprintAsset, bCheckoutAndSave, bPromptForSave);
+}
+
+bool FPackedLevelActorBuilder::CreateOrUpdateBlueprintFromUnpacked(ALevelInstance* InActor, TSoftObjectPtr<UBlueprint> InBlueprintAsset, bool bCheckoutAndSave, bool bPromptForSave)
+{
+	bool bResult = true;
+	
+	// Create Temp Actor for Packing
+	FActorSpawnParameters SpawnParams;
+	SpawnParams.bCreateActorPackage = false;
+	SpawnParams.bHideFromSceneOutliner = true;
+	SpawnParams.bNoFail = true;
+	SpawnParams.ObjectFlags |= RF_Transient;
+
+	UWorld* World = InActor->GetWorld();
+	SpawnParams.OverrideLevel = World->PersistentLevel;
+
+	APackedLevelActor* PackedLevelActor = World->SpawnActor<APackedLevelActor>(InActor->GetActorLocation(), InActor->GetActorRotation(), SpawnParams);	
+	PackedLevelActor->SetWorldAsset(InActor->GetWorldAsset());
+	ON_SCOPE_EXIT
+	{
+		InActor->GetWorld()->DestroyActor(PackedLevelActor);
+	};
+
+	if (!PackActor(PackedLevelActor, InActor))
+	{
+		return false;
+	}
+
+	PackedLevelActor->BlueprintAsset = InBlueprintAsset;
+	bResult &= CreateOrUpdateBlueprintFromPacked(PackedLevelActor, PackedLevelActor->BlueprintAsset, bCheckoutAndSave, bPromptForSave);
+
+	return bResult;
+}
+
+bool FPackedLevelActorBuilder::CreateOrUpdateBlueprintFromPacked(APackedLevelActor* InActor, TSoftObjectPtr<UBlueprint> InBlueprintAsset, bool bCheckoutAndSave, bool bPromptToSave)
+{
+	UBlueprint* BP = nullptr;
+	if (!InBlueprintAsset.IsNull())
+	{
+		BP = InBlueprintAsset.LoadSynchronous();
+	}
+
+	if (!BP)
+	{
+		int32 LastSlashIndex = 0;
+		FString LongPackageName = InActor->GetWorldAsset().GetLongPackageName();
+		LongPackageName.FindLastChar('/', LastSlashIndex);
+
+		FString PackagePath = LongPackageName.Mid(0, LastSlashIndex == INDEX_NONE ? MAX_int32 : LastSlashIndex);
+		FString AssetName = GetPackedBPPrefix() + InActor->GetWorldAsset().GetAssetName();
+		const bool bCompile = false;
+
+		FString AssetPath = PackagePath + AssetName + "." + AssetName;
+		BP = CreatePackedLevelActorBlueprintWithDialog(TSoftObjectPtr<UBlueprint>(AssetPath), InActor->GetWorldAsset(), bCompile);
+	}
+
+	if (BP && BP->SimpleConstructionScript)
+	{
+		TArray<USCS_Node*> AllNodes = BP->SimpleConstructionScript->GetAllNodes();
+		for (USCS_Node* Node : AllNodes)
+		{
+			BP->SimpleConstructionScript->RemoveNodeAndPromoteChildren(Node);
+		}
+	}
+	else
+	{
+		return false;
+	}
+
+	BP->Modify();
+	// Avoid running construction script while dragging an instance of that BP for performance reasons
+	BP->bRunConstructionScriptOnDrag = false;
+	FGuid NewVersion = FGuid::NewGuid();
+	APackedLevelActor* CDO = CastChecked<APackedLevelActor>(BP->GeneratedClass->GetDefaultObject());
+	auto PropagatePropertiesToActor = [InActor, NewVersion](APackedLevelActor* TargetActor)
+	{
+		TargetActor->Modify(false);
+		TargetActor->SetWorldAsset(InActor->GetWorldAsset());
+		TargetActor->PackedBPDependencies = InActor->PackedBPDependencies;
+		TargetActor->SetPackedVersion(NewVersion);
+		// match root component mobility to source actor
+		USceneComponent* Root = TargetActor->GetRootComponent();
+		Root->SetMobility(InActor->GetRootComponent()->Mobility);
+	};
+	PropagatePropertiesToActor(CDO);
+		
+	// Prep AddComponentsToBlueprintParam
+	FKismetEditorUtilities::FAddComponentsToBlueprintParams AddCompToBPParams;
+	AddCompToBPParams.HarvestMode = FKismetEditorUtilities::EAddComponentToBPHarvestMode::None;
+	AddCompToBPParams.bKeepMobility = true;
+
+	// Add Components
+	TArray<UActorComponent*> PackedComponents;
+	InActor->GetPackedComponents(PackedComponents);
+	
+	// To avoid any delta serialization happening on those generated components, we make them non editable.
+	for (UActorComponent* PackedComponent : PackedComponents)
+	{
+		PackedComponent->bEditableWhenInherited = false;
+	}
+
+	FKismetEditorUtilities::AddComponentsToBlueprint(BP, PackedComponents, AddCompToBPParams);
+	// If we are packing the actors BP then destroy packed components as they are now part of the BPs construction script
+	if (UBlueprint* GeneratedBy = Cast<UBlueprint>(InActor->GetClass()->ClassGeneratedBy))
+	{
+		InActor->DestroyPackedComponents();
+	}
+
+	// Propagate properties before BP Compilation so that they are considered default (no delta)
+	if (BP->GeneratedClass)
+	{
+		TArray<UObject*> ObjectsOfClass;
+		GetObjectsOfClass(BP->GeneratedClass, ObjectsOfClass, true, RF_ClassDefaultObject | RF_ArchetypeObject, EInternalObjectFlags::PendingKill);
+		for (UObject* ObjectOfClass : ObjectsOfClass)
+		{
+			APackedLevelActor* PackedLevelActor = CastChecked<APackedLevelActor>(ObjectOfClass);
+			PropagatePropertiesToActor(PackedLevelActor);
+		}
+	}
+			
+	// Synchronous compile
+	FKismetEditorUtilities::CompileBlueprint(BP, EBlueprintCompileOptions::SkipGarbageCollection);
+
+	if (bCheckoutAndSave)
+	{
+		const bool bCheckDirty = false;
+		TArray<UPackage*> OutFailedPackages;
+		FEditorFileUtils::PromptForCheckoutAndSave({ BP->GetPackage() }, bCheckDirty, bPromptToSave, &OutFailedPackages);
+
+		return !OutFailedPackages.Num();
+	}
+
+	return true;
+}
+#endif
+
+#undef LOCTEXT_NAMESPACE
\ No newline at end of file
diff --git a/Engine/Source/Runtime/Engine/Private/PackedLevelActor/PackedLevelActorBuilder.h b/Engine/Source/Runtime/Engine/Private/PackedLevelActor/PackedLevelActorBuilder.h
new file mode 100644
--- /dev/null
+++ b/Engine/Source/Runtime/Engine/Private/PackedLevelActor/PackedLevelActorBuilder.h
@@ -0,0 +1,104 @@
+// Copyright Epic Games, Inc. All Rights Reserved.
+
+#pragma once
+
+#if WITH_EDITOR
+
+#include "CoreMinimal.h"
+#include "UObject/ObjectMacros.h"
+#include "PackedLevelActor/PackedLevelActorTypes.h"
+#include "PackedLevelActor/IPackedLevelActorBuilder.h"
+#include "UObject/SoftObjectPtr.h"
+#include "Containers/Set.h"
+
+class APackedLevelActor;
+class ALevelInstance;
+class AActor;
+class UActorComponent;
+class UBlueprint;
+class FMessageLog;
+
+/**
+ * FPackedLevelActorBuilder handles packing of ALevelInstance actors into APackedLevelActor actors and Blueprints.
+ */
+class ENGINE_API FPackedLevelActorBuilder
+{
+public:
+	FPackedLevelActorBuilder();
+	
+	static TSharedPtr<FPackedLevelActorBuilder> CreateDefaultBuilder();
+	
+	/* Packs InPackedLevelActor using InLevelInstanceToPack as its source level actor */
+	bool PackActor(APackedLevelActor* InPackedLevelActor, ALevelInstance* InLevelInstanceToPack);
+	/* Packs InPackedLevelActor using itself as the source level actor */
+	bool PackActor(APackedLevelActor* InPackedLevelActor);
+	/* Packs InPackedLevelActor using InWorldAsset as the source level */
+	bool PackActor(APackedLevelActor* InPackedLevelActor, TSoftObjectPtr<UWorld> InWorldAsset);
+
+	/* Creates/Updates a APackedLevelActor Blueprint from InLevelInstance (will overwrite existing asset or show a dialog if InBlueprintAsset is null) */
+	bool CreateOrUpdateBlueprint(ALevelInstance* InLevelInstance, TSoftObjectPtr<UBlueprint> InBlueprintAsset, bool bCheckoutAndSave = true, bool bPromptForSave = true);
+	/* Creates/Updates a APackedLeveInstance Blueprint from InWorldAsset (will overwrite existing asset or show a dialog if InBlueprintAsset is null) */
+	bool CreateOrUpdateBlueprint(TSoftObjectPtr<UWorld> InWorldAsset, TSoftObjectPtr<UBlueprint> InBlueprintAsset, bool bCheckoutAndSave = true, bool bPromptForSave = true);
+	/* Update existing Blueprint */
+	void UpdateBlueprint(UBlueprint* Blueprint, bool bCheckoutAndSave = true, bool bPromptForSave = true);
+
+	static const FString& GetPackedBPPrefix();
+	/* Creates a new APackedLevelActor Blueprint using InPackagePath/InAssetName as hint for path. Prompts the user to input the final asset name. */
+	static UBlueprint* CreatePackedLevelActorBlueprintWithDialog(TSoftObjectPtr<UBlueprint> InBlueprintAsset, TSoftObjectPtr<UWorld> InWorldAsset, bool bInCompile);
+	/* Creates a new APackedLevelActor Blueprint using InPackagePath/InAssetName */
+	static UBlueprint* CreatePackedLevelActorBlueprint(TSoftObjectPtr<UBlueprint> InBlueprintAsset, TSoftObjectPtr<UWorld> InWorldAsset, bool bInCompile);
+	
+private:
+	/* Create/Updates a APackedLevelActor Blueprint from InPackedActor (will overwrite existing asset or show a dialog if InBlueprintAsset is null) */
+	bool CreateOrUpdateBlueprintFromPacked(APackedLevelActor* InPackedActor, TSoftObjectPtr<UBlueprint> InBlueprintAsset, bool bCheckoutAndSave, bool bPromptForSave);
+	/* Creates/Updates a APackedLevelActor Blueprint from InLevelInstance (will overwrite existing asset or show a dialog if InBlueprintAsset is null) */
+	bool CreateOrUpdateBlueprintFromUnpacked(ALevelInstance* InPackedActor, TSoftObjectPtr<UBlueprint> InBlueprintAsset, bool bCheckoutAndSave, bool bPromptForSave);
+	/* Creates and loads a ALevelInstance so it can be used for packing */
+	ALevelInstance* CreateTransientLevelInstanceForPacking(TSoftObjectPtr<UWorld> InWorldAsset, const FVector& InLocation, const FRotator& InRotator);
+
+	FPackedLevelActorBuilder(const FPackedLevelActorBuilder&) = delete;
+	FPackedLevelActorBuilder& operator=(const FPackedLevelActorBuilder&) = delete;
+
+public:
+	friend class FPackedLevelActorBuilderContext;
+		
+private:
+	
+	TSet<UClass*> ClassDiscards;
+	TMap<FPackedLevelActorBuilderID, TUniquePtr<IPackedLevelActorBuilder>> Builders;
+};
+
+
+class FPackedLevelActorBuilderContext
+{
+public:
+	FPackedLevelActorBuilderContext(const FPackedLevelActorBuilder& InBuilder, APackedLevelActor* InPackedLevelActor) 
+		: Builders(InBuilder.Builders), ClassDiscards(InBuilder.ClassDiscards), PackedLevelActor(InPackedLevelActor), RelativePivotTransform(FTransform::Identity) {}
+
+	/* Interface for IPackedLevelActorBuilder's to use */
+	void ClusterLevelActor(AActor* InLevelActor);
+	void FindOrAddCluster(FPackedLevelActorBuilderClusterID&& InClusterID, UActorComponent* InComponent = nullptr);
+	void DiscardActor(AActor* InActor);
+	void Report(FMessageLog& LevelInstanceLog) const;
+
+	const TMap<FPackedLevelActorBuilderClusterID, TArray<UActorComponent*>>& GetClusters() const { return Clusters; }
+	
+	void SetRelativePivotTransform(const FTransform& InRelativePivotTransform) { RelativePivotTransform = InRelativePivotTransform; }
+	const FTransform GetRelativePivotTransform() const { return RelativePivotTransform; }
+
+	bool ShouldPackComponent(UActorComponent* InActorComponent) const;
+private:
+	const TMap<FPackedLevelActorBuilderID, TUniquePtr<IPackedLevelActorBuilder>>& Builders;
+	const TSet<UClass*>& ClassDiscards;
+
+	APackedLevelActor* PackedLevelActor;
+	
+	TMap<FPackedLevelActorBuilderClusterID, TArray<UActorComponent*>> Clusters;
+
+	TMap<AActor*, TSet<UActorComponent*>> PerActorClusteredComponents;
+	TSet<AActor*> ActorDiscards;
+
+	FTransform RelativePivotTransform;
+};
+
+#endif
\ No newline at end of file
diff --git a/Engine/Source/Runtime/Engine/Private/PackedLevelActor/PackedLevelActorISMBuilder.cpp b/Engine/Source/Runtime/Engine/Private/PackedLevelActor/PackedLevelActorISMBuilder.cpp
new file mode 100644
--- /dev/null
+++ b/Engine/Source/Runtime/Engine/Private/PackedLevelActor/PackedLevelActorISMBuilder.cpp
@@ -0,0 +1,106 @@
+// Copyright Epic Games, Inc. All Rights Reserved.
+
+#include "PackedLevelActor/PackedLevelActorISMBuilder.h"
+
+#if WITH_EDITOR
+
+#include "PackedLevelActor/PackedLevelActorBuilder.h"
+#include "PackedLevelActor/PackedLevelActor.h"
+
+#include "Templates/TypeHash.h"
+
+#include "Components/StaticMeshComponent.h"
+#include "Components/InstancedStaticMeshComponent.h"
+
+FPackedLevelActorBuilderID FPackedLevelActorISMBuilder::BuilderID = 'ISMP';
+
+FPackedLevelActorBuilderID FPackedLevelActorISMBuilder::GetID() const
+{
+	return BuilderID;
+}
+
+void FPackedLevelActorISMBuilder::GetPackClusters(FPackedLevelActorBuilderContext& InContext, AActor* InActor) const
+{
+	TArray<UStaticMeshComponent*> StaticMeshComponents;
+	InActor->GetComponents<UStaticMeshComponent>(StaticMeshComponents);
+
+	for(UStaticMeshComponent* StaticMeshComponent : StaticMeshComponents)
+	{
+		if (InContext.ShouldPackComponent(StaticMeshComponent))
+		{
+			FPackedLevelActorBuilderClusterID ClusterID(MakeUnique<FPackedLevelActorISMBuilderCluster>(GetID(), StaticMeshComponent));
+
+			InContext.FindOrAddCluster(MoveTemp(ClusterID), StaticMeshComponent);
+		}
+	}
+}
+
+void FPackedLevelActorISMBuilder::PackActors(FPackedLevelActorBuilderContext& InContext, APackedLevelActor* InPackingActor, const FPackedLevelActorBuilderClusterID& InClusterID, const TArray<UActorComponent*>& InComponents) const
+{
+	check(InClusterID.GetBuilderID() == GetID());
+	UInstancedStaticMeshComponent* PackComponent = InPackingActor->AddPackedComponent<UInstancedStaticMeshComponent>();
+	
+	FTransform ActorTransform = InPackingActor->GetActorTransform();
+	PackComponent->AttachToComponent(InPackingActor->GetRootComponent(), FAttachmentTransformRules::SnapToTargetIncludingScale);
+	
+	FPackedLevelActorISMBuilderCluster* ISMCluster = (FPackedLevelActorISMBuilderCluster*)InClusterID.GetData();
+	check(ISMCluster);
+
+	ISMCluster->ISMDescriptor.InitComponent(PackComponent);
+
+	TArray<FTransform> InstanceTransforms;
+	for (UActorComponent* Component : InComponents)
+	{
+		// If we have a ISM we need to add all instances
+		if (UInstancedStaticMeshComponent* ISMComponent = Cast<UInstancedStaticMeshComponent>(Component))
+		{
+			for (int32 InstanceIndex = 0; InstanceIndex < ISMComponent->GetInstanceCount(); ++InstanceIndex)
+			{
+				FTransform InstanceTransform;
+
+				if (ensure(ISMComponent->GetInstanceTransform(InstanceIndex, InstanceTransform, /*bWorldSpace=*/ true)))
+				{
+					// WorldSpace -> ActorSpace -> Apply pivot change
+					InstanceTransforms.Add(InstanceTransform.GetRelativeTransform(ActorTransform) * InContext.GetRelativePivotTransform());
+				}
+			}
+		}
+		else // other subclasses are processed like regular UStaticMeshComponent
+		{
+			UStaticMeshComponent* StaticMeshComponent = CastChecked<UStaticMeshComponent>(Component);
+
+			// WorldSpace -> ActorSpace -> Apply pivot change
+			InstanceTransforms.Add(StaticMeshComponent->GetComponentTransform().GetRelativeTransform(ActorTransform) * InContext.GetRelativePivotTransform());
+		}
+	}
+	PackComponent->AddInstances(InstanceTransforms, /*bShouldReturnIndices*/false, /*bWorldSpace*/false);
+	PackComponent->RegisterComponent();
+}
+
+FPackedLevelActorISMBuilderCluster::FPackedLevelActorISMBuilderCluster(FPackedLevelActorBuilderID InBuilderID, UStaticMeshComponent* InComponent)
+	: FPackedLevelActorBuilderCluster(InBuilderID)
+{
+	ISMDescriptor.InitFrom(InComponent, /** bInitBodyInstance= */ false);
+	// Component descriptor should be considered hidden if original actor owner was.
+	ISMDescriptor.bHiddenInGame |= InComponent->GetOwner()->IsHidden();
+	ISMDescriptor.BodyInstance.CopyRuntimeBodyInstancePropertiesFrom(&InComponent->BodyInstance);
+	ISMDescriptor.ComputeHash();
+}
+
+uint32 FPackedLevelActorISMBuilderCluster::ComputeHash() const
+{
+	return HashCombine(FPackedLevelActorBuilderCluster::ComputeHash(), ISMDescriptor.Hash);
+}
+
+bool FPackedLevelActorISMBuilderCluster::Equals(const FPackedLevelActorBuilderCluster& InOther) const
+{
+	if (!FPackedLevelActorBuilderCluster::Equals(InOther))
+	{
+		return false;
+	}
+
+	const FPackedLevelActorISMBuilderCluster& ISMOther = (const FPackedLevelActorISMBuilderCluster&)InOther;
+	return ISMDescriptor == ISMOther.ISMDescriptor;
+}
+
+#endif
\ No newline at end of file
diff --git a/Engine/Source/Runtime/Engine/Private/PackedLevelActor/PackedLevelActorISMBuilder.h b/Engine/Source/Runtime/Engine/Private/PackedLevelActor/PackedLevelActorISMBuilder.h
new file mode 100644
--- /dev/null
+++ b/Engine/Source/Runtime/Engine/Private/PackedLevelActor/PackedLevelActorISMBuilder.h
@@ -0,0 +1,34 @@
+// Copyright Epic Games, Inc. All Rights Reserved.
+
+#pragma once
+
+#if WITH_EDITOR
+
+#include "PackedLevelActor/PackedLevelActorTypes.h"
+#include "PackedLevelActor/IPackedLevelActorBuilder.h"
+#include "ISMPartition/ISMComponentDescriptor.h"
+
+class UStaticMeshComponent;
+
+class FPackedLevelActorISMBuilder : public IPackedLevelActorBuilder
+{
+public:
+	static FPackedLevelActorBuilderID BuilderID;
+
+	virtual FPackedLevelActorBuilderID GetID() const override;
+	virtual void GetPackClusters(FPackedLevelActorBuilderContext& InContext, AActor* InActor) const override;
+	virtual void PackActors(FPackedLevelActorBuilderContext& InContext, APackedLevelActor* InPackingActor, const FPackedLevelActorBuilderClusterID& InClusterID, const TArray<UActorComponent*>& InComponents) const override;
+};
+
+class FPackedLevelActorISMBuilderCluster : public FPackedLevelActorBuilderCluster
+{
+public:
+	FPackedLevelActorISMBuilderCluster(FPackedLevelActorBuilderID InBuilderID, UStaticMeshComponent* InComponent);
+
+	virtual bool Equals(const FPackedLevelActorBuilderCluster& InOther) const override;
+	virtual uint32 ComputeHash() const override;
+
+	FISMComponentDescriptor ISMDescriptor;
+};
+
+#endif
\ No newline at end of file
diff --git a/Engine/Source/Runtime/Engine/Private/PackedLevelActor/PackedLevelActorRecursiveBuilder.cpp b/Engine/Source/Runtime/Engine/Private/PackedLevelActor/PackedLevelActorRecursiveBuilder.cpp
new file mode 100644
--- /dev/null
+++ b/Engine/Source/Runtime/Engine/Private/PackedLevelActor/PackedLevelActorRecursiveBuilder.cpp
@@ -0,0 +1,85 @@
+// Copyright Epic Games, Inc. All Rights Reserved.
+
+#include "PackedLevelActor/PackedLevelActorRecursiveBuilder.h"
+
+#if WITH_EDITOR
+
+#include "PackedLevelActor/PackedLevelActorBuilder.h"
+#include "PackedLevelActor/PackedLevelActor.h"
+
+#include "LevelInstance/LevelInstanceActor.h"
+#include "LevelInstance/LevelInstanceSubsystem.h"
+
+#include "Engine/Blueprint.h"
+
+#include "Misc/Crc.h"
+
+FPackedLevelActorBuilderID FPackedLevelActorRecursiveBuilder::BuilderID = 'RECP';
+
+FPackedLevelActorBuilderID FPackedLevelActorRecursiveBuilder::GetID() const
+{
+	return BuilderID;
+}
+
+void FPackedLevelActorRecursiveBuilder::GetPackClusters(FPackedLevelActorBuilderContext& InContext, AActor* InActor) const
+{
+	if (ALevelInstance* LevelInstance = Cast<ALevelInstance>(InActor))
+	{
+		FPackedLevelActorBuilderClusterID ClusterID(MakeUnique<FPackedLevelActorRecursiveBuilderCluster>(GetID(), LevelInstance));
+		InContext.FindOrAddCluster(MoveTemp(ClusterID));
+
+		// This Actor can be safely discarded without warning because it is a container
+		InContext.DiscardActor(LevelInstance);
+
+		ULevelInstanceSubsystem* LevelInstanceSubsystem = LevelInstance->GetLevelInstanceSubsystem();
+		check(LevelInstanceSubsystem);
+		if (ULevel* Level = LevelInstanceSubsystem->GetLevelInstanceLevel(LevelInstance))
+		{
+			for (AActor* LevelActor : Level->Actors)
+			{
+				if (LevelActor)
+				{
+					InContext.ClusterLevelActor(LevelActor);
+				}
+			}
+		}
+	}
+}
+
+void FPackedLevelActorRecursiveBuilder::PackActors(FPackedLevelActorBuilderContext& InContext, APackedLevelActor* InPackingActor, const FPackedLevelActorBuilderClusterID& InClusterID, const TArray<UActorComponent*>& InComponents) const
+{
+	check(InClusterID.GetBuilderID() == GetID());
+	FPackedLevelActorRecursiveBuilderCluster* LevelInstanceCluster = (FPackedLevelActorRecursiveBuilderCluster*)InClusterID.GetData();
+	check(LevelInstanceCluster);
+
+	if (LevelInstanceCluster->LevelInstance->IsA<APackedLevelActor>())
+	{
+		if (UBlueprint* GeneratedBy = Cast<UBlueprint>(LevelInstanceCluster->LevelInstance->GetClass()->ClassGeneratedBy))
+		{
+			InPackingActor->PackedBPDependencies.AddUnique(GeneratedBy);
+		}
+	}
+}
+
+FPackedLevelActorRecursiveBuilderCluster::FPackedLevelActorRecursiveBuilderCluster(FPackedLevelActorBuilderID InBuilderID, ALevelInstance* InLevelInstance)
+	: FPackedLevelActorBuilderCluster(InBuilderID), LevelInstance(InLevelInstance)
+{
+}
+
+bool FPackedLevelActorRecursiveBuilderCluster::Equals(const FPackedLevelActorBuilderCluster& InOther) const
+{
+	if (!FPackedLevelActorBuilderCluster::Equals(InOther))
+	{
+		return false;
+	}
+
+	const FPackedLevelActorRecursiveBuilderCluster& LevelInstanceCluster = (const FPackedLevelActorRecursiveBuilderCluster&)InOther;
+	return LevelInstance == LevelInstanceCluster.LevelInstance;
+}
+	
+uint32 FPackedLevelActorRecursiveBuilderCluster::ComputeHash() const
+{
+	return FCrc::TypeCrc32(LevelInstance->GetLevelInstanceID(), FPackedLevelActorBuilderCluster::ComputeHash());
+}
+
+#endif
diff --git a/Engine/Source/Runtime/Engine/Private/PackedLevelActor/PackedLevelActorRecursiveBuilder.h b/Engine/Source/Runtime/Engine/Private/PackedLevelActor/PackedLevelActorRecursiveBuilder.h
new file mode 100644
--- /dev/null
+++ b/Engine/Source/Runtime/Engine/Private/PackedLevelActor/PackedLevelActorRecursiveBuilder.h
@@ -0,0 +1,33 @@
+// Copyright Epic Games, Inc. All Rights Reserved.
+
+#pragma once
+
+#if WITH_EDITOR
+
+#include "PackedLevelActor/PackedLevelActorTypes.h"
+#include "PackedLevelActor/IPackedLevelActorBuilder.h"
+
+class ALevelInstance;
+
+class FPackedLevelActorRecursiveBuilder : public IPackedLevelActorBuilder
+{
+public:
+	static FPackedLevelActorBuilderID BuilderID;
+
+	virtual FPackedLevelActorBuilderID GetID() const override;
+	virtual void GetPackClusters(FPackedLevelActorBuilderContext& InContext, AActor* InActor) const override;
+	virtual void PackActors(FPackedLevelActorBuilderContext& InContext, APackedLevelActor* InPackingActor, const FPackedLevelActorBuilderClusterID& InClusterID, const TArray<UActorComponent*>& InComponents) const override;
+};
+
+class FPackedLevelActorRecursiveBuilderCluster : public FPackedLevelActorBuilderCluster
+{
+public:
+	FPackedLevelActorRecursiveBuilderCluster(FPackedLevelActorBuilderID InBuilderID, ALevelInstance* InLevelInstance);
+
+	virtual bool Equals(const FPackedLevelActorBuilderCluster& InOther) const override;
+	virtual uint32 ComputeHash() const override;
+
+	ALevelInstance* LevelInstance = nullptr;
+};
+
+#endif
\ No newline at end of file
diff --git a/Engine/Source/Runtime/Engine/Private/PackedLevelActor/PackedLevelActorTypes.h b/Engine/Source/Runtime/Engine/Private/PackedLevelActor/PackedLevelActorTypes.h
new file mode 100644
--- /dev/null
+++ b/Engine/Source/Runtime/Engine/Private/PackedLevelActor/PackedLevelActorTypes.h
@@ -0,0 +1,92 @@
+// Copyright Epic Games, Inc. All Rights Reserved.
+
+#pragma once
+
+#if WITH_EDITOR
+
+#include "CoreMinimal.h"
+#include "UObject/ObjectMacros.h"
+#include "Misc/Crc.h"
+#include "Logging/LogMacros.h"
+
+using FPackedLevelActorBuilderID = uint32;
+
+class ENGINE_API FPackedLevelActorBuilderCluster
+{
+	FPackedLevelActorBuilderID BuilderID;
+
+public:
+	FPackedLevelActorBuilderCluster(FPackedLevelActorBuilderID InBuilderID)
+		: BuilderID(InBuilderID) {}
+	virtual ~FPackedLevelActorBuilderCluster() {}
+
+	FPackedLevelActorBuilderID GetBuilderID() const { return BuilderID; }
+
+	virtual uint32 ComputeHash() const
+	{
+		return FCrc::TypeCrc32(BuilderID);
+	}
+
+	virtual bool Equals(const FPackedLevelActorBuilderCluster& Other) const
+	{
+		return BuilderID == Other.BuilderID;
+	}
+
+private:
+	FPackedLevelActorBuilderCluster(const FPackedLevelActorBuilderCluster&) = delete;
+	FPackedLevelActorBuilderCluster& operator=(const FPackedLevelActorBuilderCluster&) = delete;
+};
+
+class ENGINE_API FPackedLevelActorBuilderClusterID
+{
+public:
+	static FPackedLevelActorBuilderClusterID Invalid;
+
+	FPackedLevelActorBuilderClusterID() : Hash(0)
+	{
+	}
+
+	FPackedLevelActorBuilderClusterID(FPackedLevelActorBuilderClusterID&& Other)
+		: Hash(Other.Hash), Data(MoveTemp(Other.Data))
+	{
+	}
+
+	FPackedLevelActorBuilderClusterID(TUniquePtr<FPackedLevelActorBuilderCluster>&& InData)
+		: Data(MoveTemp(InData))
+	{
+		Hash = Data->ComputeHash();
+	}
+		
+	bool operator==(const FPackedLevelActorBuilderClusterID& Other) const
+	{
+		return (Hash == Other.Hash) && Data->Equals(*Other.Data);
+	}
+
+	bool operator!=(const FPackedLevelActorBuilderClusterID& Other) const
+	{
+		return !(*this == Other);
+	}
+
+	uint32 GetHash() const { return Hash; }
+
+	FPackedLevelActorBuilderID GetBuilderID() const { return Data->GetBuilderID(); }
+		
+	friend uint32 GetTypeHash(const FPackedLevelActorBuilderClusterID& ID)
+	{
+		return ID.GetHash();
+	}
+
+	FPackedLevelActorBuilderCluster* GetData() const
+	{
+		return Data.Get();
+	}
+
+private:
+	FPackedLevelActorBuilderClusterID(const FPackedLevelActorBuilderClusterID&) = delete;
+	FPackedLevelActorBuilderClusterID& operator=(const FPackedLevelActorBuilderClusterID&) = delete;
+
+	uint32 Hash;
+	TUniquePtr<FPackedLevelActorBuilderCluster> Data;
+};
+
+#endif
\ No newline at end of file
diff --git a/Engine/Source/Runtime/Engine/Public/LevelInstance/LevelInstanceActor.h b/Engine/Source/Runtime/Engine/Public/LevelInstance/LevelInstanceActor.h
--- a/Engine/Source/Runtime/Engine/Public/LevelInstance/LevelInstanceActor.h
+++ b/Engine/Source/Runtime/Engine/Public/LevelInstance/LevelInstanceActor.h
@@ -30,7 +30,7 @@ class ENGINE_API ALevelInstance : public AActor
 protected:
 #if WITH_EDITORONLY_DATA
 	/** Level LevelInstance */
-	UPROPERTY(EditAnywhere, Category = LevelInstance, meta = (NoCreate, DisplayName="Level"))
+	UPROPERTY(EditAnywhere, Category = Level, meta = (NoCreate, DisplayName="Level"))
 	TSoftObjectPtr<UWorld> WorldAsset;
 #endif
 
@@ -105,12 +105,12 @@ public:
 	bool CheckForLoop(TSoftObjectPtr<UWorld> WorldAsset, TArray<TPair<FText, TSoftObjectPtr<UWorld>>>* LoopInfo = nullptr, const ALevelInstance** LoopStart = nullptr) const;
 	AActor* FindEditorInstanceActor() const;
 
-	virtual void OnWorldAssetChanged() { UpdateLevelInstance(); }
+	virtual void OnWorldAssetChanged() { UpdateFromLevel(); }
 	virtual void OnEdit() {}
 	virtual void OnEditChild() {}
 	virtual void OnCommit(bool bChanged, bool bPromptForSave) {}
 	virtual void OnCommitChild(bool bChanged) {}
-	virtual void UpdateLevelInstance();
+	virtual void UpdateFromLevel();
 		
 	DECLARE_MULTICAST_DELEGATE_OneParam(FOnLevelInstanceActorPostLoad, ALevelInstance*);
 	static FOnLevelInstanceActorPostLoad OnLevelInstanceActorPostLoad;
diff --git a/Engine/Source/Runtime/Engine/Public/LevelInstance/LevelInstanceSubsystem.h b/Engine/Source/Runtime/Engine/Public/LevelInstance/LevelInstanceSubsystem.h
--- a/Engine/Source/Runtime/Engine/Public/LevelInstance/LevelInstanceSubsystem.h
+++ b/Engine/Source/Runtime/Engine/Public/LevelInstance/LevelInstanceSubsystem.h
@@ -51,8 +51,8 @@ public:
 	void OnObjectPreSave(UObject* Object, FObjectPreSaveContext SaveContext);
 	void OnPackageDeleted(UPackage* Package);
 	void OnExitEditorMode();
-	void PackLevelInstances();
-	bool CanPackLevelInstances() const;
+	void PackAllLoadedActors();
+	bool CanPackAllLoadedActors() const;
 
 	ALevelInstance* GetEditingLevelInstance() const;
 	bool CanEditLevelInstance(const ALevelInstance* LevelInstanceActor, FText* OutReason = nullptr) const;
diff --git a/Engine/Source/Runtime/Engine/Public/LevelInstance/LevelInstanceTypes.h b/Engine/Source/Runtime/Engine/Public/LevelInstance/LevelInstanceTypes.h
--- a/Engine/Source/Runtime/Engine/Public/LevelInstance/LevelInstanceTypes.h
+++ b/Engine/Source/Runtime/Engine/Public/LevelInstance/LevelInstanceTypes.h
@@ -41,8 +41,7 @@ UENUM()
 enum class ELevelInstanceCreationType : uint8
 {
 	LevelInstance,
-	PackedLevelInstance,
-	PackedLevelInstanceBlueprint
+	PackedLevelActor
 };
 
 UENUM()
@@ -59,7 +58,7 @@ struct FNewLevelInstanceParams
 {
 	GENERATED_USTRUCT_BODY()
 			
-	UPROPERTY(EditAnywhere, Category = Default)
+	UPROPERTY(EditAnywhere, Category = Default, meta = (EditCondition = "!bHideCreationType", EditConditionHides, HideEditConditionToggle))
 	ELevelInstanceCreationType Type = ELevelInstanceCreationType::LevelInstance;
 
 	UPROPERTY(EditAnywhere, Category = Pivot)
@@ -84,7 +83,11 @@ private:
 	UPROPERTY()
 	bool bForceExternalActors = false;
 
+	UPROPERTY()
+	bool bHideCreationType = false;
+
 public:
+	void HideCreationType() { bHideCreationType = true; }
 	void SetForceExternalActors(bool bInForceExternalActors) { bForceExternalActors = bInForceExternalActors; }
 	void SetExternalActors(bool bInExternalActors) { bExternalActors = bInExternalActors; }
 	bool UseExternalActors() const { return bForceExternalActors || bExternalActors; }
diff --git a/Engine/Source/Runtime/Engine/Public/LevelInstance/Packed/ILevelInstancePacker.h b/Engine/Source/Runtime/Engine/Public/LevelInstance/Packed/ILevelInstancePacker.h
deleted file mode 100644
diff --git a/Engine/Source/Runtime/Engine/Public/LevelInstance/Packed/LevelInstanceISMPacker.h b/Engine/Source/Runtime/Engine/Public/LevelInstance/Packed/LevelInstanceISMPacker.h
deleted file mode 100644
diff --git a/Engine/Source/Runtime/Engine/Public/LevelInstance/Packed/LevelInstanceRecursivePacker.h b/Engine/Source/Runtime/Engine/Public/LevelInstance/Packed/LevelInstanceRecursivePacker.h
deleted file mode 100644
diff --git a/Engine/Source/Runtime/Engine/Public/LevelInstance/Packed/PackedLevelInstanceActor.h b/Engine/Source/Runtime/Engine/Public/LevelInstance/Packed/PackedLevelInstanceActor.h
deleted file mode 100644
diff --git a/Engine/Source/Runtime/Engine/Public/LevelInstance/Packed/PackedLevelInstanceBuilder.h b/Engine/Source/Runtime/Engine/Public/LevelInstance/Packed/PackedLevelInstanceBuilder.h
deleted file mode 100644
diff --git a/Engine/Source/Runtime/Engine/Public/LevelInstance/Packed/PackedLevelInstanceTypes.h b/Engine/Source/Runtime/Engine/Public/LevelInstance/Packed/PackedLevelInstanceTypes.h
deleted file mode 100644
diff --git a/Engine/Source/Runtime/Engine/Public/PackedLevelActor/PackedLevelActor.h b/Engine/Source/Runtime/Engine/Public/PackedLevelActor/PackedLevelActor.h
new file mode 100644
--- /dev/null
+++ b/Engine/Source/Runtime/Engine/Public/PackedLevelActor/PackedLevelActor.h
@@ -0,0 +1,88 @@
+// Copyright Epic Games, Inc. All Rights Reserved.
+
+#pragma once
+
+#include "CoreMinimal.h"
+#include "UObject/ObjectMacros.h"
+#include "LevelInstance/LevelInstanceActor.h"
+
+#if WITH_EDITOR
+#include "Kismet2/ComponentEditorUtils.h"
+#endif
+
+#include "PackedLevelActor.generated.h"
+
+class UInstancedStaticMeshComponent;
+class UBlueprint;
+
+/**
+ * APackedLevelActor is the result of packing the source level (WorldAsset base class property) into a single actor. See FPackedLevelActorBuilder.
+ * 
+ * 
+ * Other components are unsupported and will result in an incomplete APackedLevelActor. In this case using a regular ALevelInstance is recommended.
+ */
+UCLASS()
+class ENGINE_API APackedLevelActor : public ALevelInstance
+{
+	GENERATED_UCLASS_BODY()
+
+public:
+	virtual bool SupportsLoading() const override;
+
+	virtual void Serialize(FArchive& Ar) override;
+#if WITH_EDITOR
+	static bool CreateOrUpdateBlueprint(ALevelInstance* InLevelInstance, TSoftObjectPtr<UBlueprint> InBlueprintAsset, bool bCheckoutAndSave = true, bool bPromptForSave = true);
+	static bool CreateOrUpdateBlueprint(TSoftObjectPtr<UWorld> InWorldAsset, TSoftObjectPtr<UBlueprint> InBlueprintAsset, bool bCheckoutAndSave = true, bool bPromptForSave = true);
+
+	static FName GetPackedComponentTag();
+
+	virtual void PostLoad() override;
+	virtual void UpdateFromLevel() override;
+	virtual void OnCommit(bool bChanged, bool bPromptForSave) override;
+	virtual void OnCommitChild(bool bChanged) override;
+	virtual void OnEdit() override;
+	virtual void OnEditChild() override;
+
+	virtual bool CanEditChange(const FProperty* InProperty) const override;
+		
+	void SetPackedVersion(const FGuid& Version) { PackedVersion = Version; }
+
+	virtual bool IsHiddenEd() const override;
+	virtual bool IsHLODRelevant() const override;
+
+	void DestroyPackedComponents();
+	void GetPackedComponents(TArray<UActorComponent*>& OutPackedComponents) const;
+
+	virtual ELevelInstanceRuntimeBehavior GetDefaultRuntimeBehavior() const override { return ELevelInstanceRuntimeBehavior::None; }
+
+	virtual void RerunConstructionScripts() override;
+
+	template<class T>
+	T* AddPackedComponent()
+	{
+		Modify();
+		FName NewComponentName = *FComponentEditorUtils::GenerateValidVariableName(T::StaticClass(), this);
+		T* NewComponent = NewObject<T>(this, NewComponentName, RF_Transactional);
+		AddInstanceComponent(NewComponent);
+		NewComponent->ComponentTags.Add(GetPackedComponentTag());
+		return NewComponent;
+	}
+#endif
+
+#if WITH_EDITORONLY_DATA
+	UPROPERTY(EditAnywhere, Category = Packed, meta = (DisplayName = "ISM Component Class"))
+	TSubclassOf<UInstancedStaticMeshComponent> ISMComponentClass;
+
+	UPROPERTY(VisibleAnywhere, Category = Packed)
+	TSoftObjectPtr<UBlueprint> BlueprintAsset;
+
+	UPROPERTY()
+	TArray<TSoftObjectPtr<UBlueprint>> PackedBPDependencies;
+
+private:
+	bool bChildChanged;
+
+	UPROPERTY()
+	FGuid PackedVersion;
+#endif
+};
\ No newline at end of file
