Engine support for optimus mesh deformers providing their own buffers.
Most work is in the passthrough vertex factory to allow the use of RDG buffers.
Still need to correctly implement deformer support for previous frame postion and ray tracing.

diff --git a/Engine/Source/Editor/DetailCustomizations/Private/DetailCustomizations.cpp b/Engine/Source/Editor/DetailCustomizations/Private/DetailCustomizations.cpp
--- a/Engine/Source/Editor/DetailCustomizations/Private/DetailCustomizations.cpp
+++ b/Engine/Source/Editor/DetailCustomizations/Private/DetailCustomizations.cpp
@@ -646,6 +646,7 @@ void FDetailCustomizationsModule::RegisterSectionMappings()
 		{
 			TSharedRef<FPropertySection> Section = PropertyModule.FindOrCreateSection("SkinnedMeshComponent", "Animation", LOCTEXT("Animation", "Animation"));
 			Section->AddCategory("Skeletal Mesh");
+			Section->AddCategory("Deformer");
 		}
 
 		{
diff --git a/Engine/Source/Runtime/Engine/Classes/Animation/MeshDeformer.h b/Engine/Source/Runtime/Engine/Classes/Animation/MeshDeformer.h
new file mode 100644
--- /dev/null
+++ b/Engine/Source/Runtime/Engine/Classes/Animation/MeshDeformer.h
@@ -0,0 +1,24 @@
+// Copyright Epic Games, Inc. All Rights Reserved.
+
+#pragma once
+
+#include "CoreMinimal.h"
+#include "MeshDeformer.generated.h"
+
+class UMeshComponent;
+class UMeshDeformerInstance;
+
+/** 
+ * Base class for mesh deformer assets.
+ * Mesh deformers can be added to mesh components to implement flexible deformation systems.
+ * A UMeshDeformer needs to implement creation of a UMeshDeformerInstance which will apply deformer actions and store deformer state.
+ */
+UCLASS(Abstract)
+class ENGINE_API UMeshDeformer : public UObject
+{
+	GENERATED_BODY()
+
+public:
+	/** Create an instance of the mesh deformer. */
+	virtual UMeshDeformerInstance* CreateInstance(UMeshComponent* InMeshComponent) PURE_VIRTUAL(, return nullptr;);
+};
diff --git a/Engine/Source/Runtime/Engine/Classes/Animation/MeshDeformerInstance.h b/Engine/Source/Runtime/Engine/Classes/Animation/MeshDeformerInstance.h
new file mode 100644
--- /dev/null
+++ b/Engine/Source/Runtime/Engine/Classes/Animation/MeshDeformerInstance.h
@@ -0,0 +1,32 @@
+// Copyright Epic Games, Inc. All Rights Reserved.
+
+#pragma once
+
+#include "CoreMinimal.h"
+#include "MeshDeformerInstance.generated.h"
+
+class FSceneInterface;
+
+/** 
+ * Base class for mesh deformers instances.
+ * This contains the per instance state for a UMeshDeformer.
+ */
+UCLASS(Abstract)
+class ENGINE_API UMeshDeformerInstance : public UObject
+{
+	GENERATED_BODY()
+
+public:
+	/** Enumeration for workloads to enqueue. */
+	enum EWorkLoad
+	{
+		WorkLoad_Setup,
+		WorkLoad_Trigger,
+		WorkLoad_Update,
+	};
+
+	/** Get if mesh deformer is active (compiled and valid). */
+	virtual bool IsActive() const PURE_VIRTUAL(, return false;);
+	/** Enqueue the mesh deformer workload on a scene. */
+	virtual void EnqueueWork(FSceneInterface* InScene, EWorkLoad WorkLoadType) PURE_VIRTUAL(, );
+};
diff --git a/Engine/Source/Runtime/Engine/Classes/Components/SkinnedMeshComponent.h b/Engine/Source/Runtime/Engine/Classes/Components/SkinnedMeshComponent.h
--- a/Engine/Source/Runtime/Engine/Classes/Components/SkinnedMeshComponent.h
+++ b/Engine/Source/Runtime/Engine/Classes/Components/SkinnedMeshComponent.h
@@ -26,6 +26,8 @@ class FSkeletalMeshRenderData;
 class FSkeletalMeshLODRenderData;
 struct FSkelMeshRenderSection;
 class FPositionVertexBuffer;
+class UMeshDeformer;
+class UMeshDeformerInstance;
 
 DECLARE_DELEGATE_OneParam(FOnAnimUpdateRateParamsCreated, FAnimUpdateRateParameters*)
 DECLARE_MULTICAST_DELEGATE_ThreeParams(FOnTickPose, USkinnedMeshComponent* /*SkinnedMeshComponent*/, float /*DeltaTime*/, bool /*bNeedsValidRootMotion*/)
@@ -232,6 +234,14 @@ class ENGINE_API USkinnedMeshComponent : public UMeshComponent, public ILODSyncI
 	UPROPERTY(EditAnywhere, BlueprintReadOnly, Category = "Mesh")
 	TArray<ESkinCacheUsage> SkinCacheUsage;
 
+	/** If set then the MeshDeformer will be used instead of the fixed animation pipeline. */
+	UPROPERTY(EditAnywhere, BlueprintReadOnly, Category = "Deformer")
+	TObjectPtr<UMeshDeformer> MeshDeformer;
+
+	/** Object containing state for the bound MeshDeformer. */
+	UPROPERTY(Transient)
+	TObjectPtr<UMeshDeformerInstance> MeshDeformerInstance;
+
 	/** const getters for previous transform idea */
 	const TArray<uint8>& GetPreviousBoneVisibilityStates() const
 	{
@@ -605,6 +615,9 @@ protected:
 	/** Whether we have updated bone visibility this tick */
 	uint8 bBoneVisibilityDirty:1;
 
+	/** Whether mesh deformer state is dirty and will need updating at the next tick. */
+	uint8 bUpdateDeformerAtNextTick : 1;
+
 private:
 	/** If true, UpdateTransform will always result in a call to MeshObject->Update. */
 	UPROPERTY(transient)
@@ -779,6 +792,14 @@ public:
 	UFUNCTION(BlueprintCallable, Category="Components|SkinnedMesh")
 	virtual void SetSkeletalMesh(class USkeletalMesh* NewMesh, bool bReinitPose = true);
 
+	/**
+	 * Change the MeshDeformer that is used for this Component.
+	 *
+	 * @param InMeshDeformer New mesh deformer to set for this component
+	 */
+	UFUNCTION(BlueprintCallable, Category = "Components|SkinnedMesh")
+	void SetMeshDeformer(UMeshDeformer* InMeshDeformer);
+
 	/** 
 	 * Get Parent Bone of the input bone
 	 * 
@@ -817,6 +838,9 @@ public:
 	bool GetTwistAndSwingAngleOfDeltaRotationFromRefPose(FName BoneName, float& OutTwistAngle, float& OutSwingAngle) const;
 
 	bool IsSkinCacheAllowed(int32 LodIdx) const;
+
+	bool HasMeshDeformer() const { return MeshDeformer != nullptr; }
+
 	/**
 	 *	Compute SkeletalMesh MinLOD that will be used by this component
 	 */
diff --git a/Engine/Source/Runtime/Engine/Private/Components/SkinnedMeshComponent.cpp b/Engine/Source/Runtime/Engine/Private/Components/SkinnedMeshComponent.cpp
--- a/Engine/Source/Runtime/Engine/Private/Components/SkinnedMeshComponent.cpp
+++ b/Engine/Source/Runtime/Engine/Private/Components/SkinnedMeshComponent.cpp
@@ -17,6 +17,7 @@
 #include "SkeletalRenderStatic.h"
 #include "Animation/AnimStats.h"
 #include "SkeletalMeshCompiler.h"
+#include "SkeletalMeshDeformerHelpers.h"
 #include "Engine/SkeletalMeshSocket.h"
 #include "PhysicsEngine/PhysicsAsset.h"
 #include "EngineGlobals.h"
@@ -24,6 +25,8 @@
 #include "Engine/CollisionProfile.h"
 #include "Rendering/SkinWeightVertexBuffer.h"
 #include "SkeletalMeshTypes.h"
+#include "Animation/MeshDeformer.h"
+#include "Animation/MeshDeformerInstance.h"
 #include "Animation/MorphTarget.h"
 #include "AnimationRuntime.h"
 #include "Animation/SkinWeightProfileManager.h"
@@ -354,6 +357,8 @@ USkinnedMeshComponent::USkinnedMeshComponent(const FObjectInitializer& ObjectIni
 	bNeedToFlipSpaceBaseBuffers = false;
 	bBoneVisibilityDirty = false;
 
+	bUpdateDeformerAtNextTick = false;
+
 	bCanEverAffectNavigation = false;
 	MasterBoneMapCacheCount = 0;
 	bSyncAttachParentLOD = true;
@@ -508,6 +513,8 @@ void USkinnedMeshComponent::OnRegister()
 
 	UpdateLODStatus();
 	InvalidateCachedBounds();
+
+	MeshDeformerInstance = (MeshDeformer != nullptr) ? MeshDeformer->CreateInstance(this) : nullptr;
 }
 
 void USkinnedMeshComponent::OnUnregister()
@@ -520,6 +527,8 @@ void USkinnedMeshComponent::OnUnregister()
 		FAnimUpdateRateManager::CleanupUpdateRateParametersRef(this);
 		AnimUpdateRateParams = nullptr;
 	}
+
+	MeshDeformerInstance = nullptr;
 }
 
 void USkinnedMeshComponent::CreateRenderState_Concurrent(FRegisterComponentContext* Context)
@@ -587,6 +596,9 @@ void USkinnedMeshComponent::CreateRenderState_Concurrent(FRegisterComponentConte
 				}
 			}
 
+			// After creating the MeshObject we need to run the MeshDeformer at least once to set up the vertex data.
+			bUpdateDeformerAtNextTick |= (MeshDeformerInstance != nullptr);
+
 			//Allow the editor a chance to manipulate it before its added to the scene
 			PostInitMeshObject(MeshObject);
 		}
@@ -725,6 +737,20 @@ void USkinnedMeshComponent::SendRenderDynamicData_Concurrent()
 			// scene proxy update of material usage based on active morphs
 			UpdateMorphMaterialUsageOnProxy();
 		}
+
+		if (MeshDeformerInstance != nullptr)
+		{
+			if (MeshDeformerInstance->IsActive())
+			{
+				MeshDeformerInstance->EnqueueWork(GetScene(), UMeshDeformerInstance::WorkLoad_Update);
+			}
+			else
+			{
+				// Reset so mesh appears in bind pose.
+				// We could use a fallback mesh deformer instead?
+				FSkeletalMeshDeformerHelpers::ResetVertexFactoryBufferOverrides_GameThread(MeshObject, UseLOD);
+			}
+		}
 	}
 }
 
@@ -905,6 +931,12 @@ void USkinnedMeshComponent::TickComponent(float DeltaTime, enum ELevelTick TickT
 		}
 	}
 #endif // WITH_EDITOR
+
+	if (bUpdateDeformerAtNextTick)
+	{
+		MarkRenderDynamicDataDirty();
+		bUpdateDeformerAtNextTick = false;
+	}
 }
 
 UObject const* USkinnedMeshComponent::AdditionalStatObject() const
@@ -1480,17 +1512,15 @@ bool USkinnedMeshComponent::GetTwistAndSwingAngleOfDeltaRotationFromRefPose(FNam
 bool USkinnedMeshComponent::IsSkinCacheAllowed(int32 LodIdx) const
 {
 	static const IConsoleVariable* CVarDefaultGPUSkinCacheBehavior = IConsoleManager::Get().FindConsoleVariable(TEXT("r.SkinCache.DefaultBehavior"));
+	const bool bGlobalDefault = CVarDefaultGPUSkinCacheBehavior && ESkinCacheDefaultBehavior(CVarDefaultGPUSkinCacheBehavior->GetInt()) == ESkinCacheDefaultBehavior::Inclusive;
 
-	bool bIsRayTracing = FGPUSkinCache::IsGPUSkinCacheRayTracingSupported();
-	if (bIsRayTracing && SkeletalMesh)
+	if (MeshDeformer)
 	{
-PRAGMA_DISABLE_DEPRECATION_WARNINGS
-		bIsRayTracing = SkeletalMesh->bSupportRayTracing;
-PRAGMA_ENABLE_DEPRECATION_WARNINGS
+		// Disable skin cache if a mesh deformer is in use.
+		// Any animation buffers are expected to be owned by the MeshDeformer.
+		return false;
 	}
 
-	bool bGlobalDefault = CVarDefaultGPUSkinCacheBehavior && ESkinCacheDefaultBehavior(CVarDefaultGPUSkinCacheBehavior->GetInt()) == ESkinCacheDefaultBehavior::Inclusive;
-
 	if (!SkeletalMesh)
 	{
 		return GEnableGPUSkinCache && bGlobalDefault;
@@ -1642,6 +1672,13 @@ void USkinnedMeshComponent::SetSkeletalMesh(USkeletalMesh* InSkelMesh, bool bRei
 	}
 }
 
+void USkinnedMeshComponent::SetMeshDeformer(UMeshDeformer* InMeshDeformer)
+{
+	MeshDeformer = InMeshDeformer;
+	MeshDeformerInstance = (MeshDeformer != nullptr) ? MeshDeformer->CreateInstance(this) : nullptr;
+	MarkRenderDynamicDataDirty();
+}
+
 FSkeletalMeshRenderData* USkinnedMeshComponent::GetSkeletalMeshRenderData() const
 {
 	if (MeshObject)
diff --git a/Engine/Source/Runtime/Engine/Private/GPUSkinCache.cpp b/Engine/Source/Runtime/Engine/Private/GPUSkinCache.cpp
--- a/Engine/Source/Runtime/Engine/Private/GPUSkinCache.cpp
+++ b/Engine/Source/Runtime/Engine/Private/GPUSkinCache.cpp
@@ -328,14 +328,6 @@ public:
 			return PositionBuffer;
 		}
 
-		inline FRHIShaderResourceView* GetPreSkinPositionSRV()
-		{
-			check(SourceVertexFactory);
-			check(SourceVertexFactory->GetPositionsSRV());
-
-			return SourceVertexFactory->GetPositionsSRV().GetReference();
-		}
-
 		inline FGPUSkinCache::FSkinCacheRWBuffer* GetTangentRWBuffer()
 		{
 			return TangentBuffer;
@@ -355,7 +347,7 @@ public:
 
 		void UpdateVertexFactoryDeclaration()
 		{
-			TargetVertexFactory->UpdateVertexDeclaration(SourceVertexFactory, &GetPositionRWBuffer()->Buffer, GetPreSkinPositionSRV(), &GetTangentRWBuffer()->Buffer);
+			TargetVertexFactory->UpdateVertexDeclaration(SourceVertexFactory, &GetPositionRWBuffer()->Buffer, &GetTangentRWBuffer()->Buffer);
 		}
 	};
 
diff --git a/Engine/Source/Runtime/Engine/Private/GPUSkinVertexFactory.cpp b/Engine/Source/Runtime/Engine/Private/GPUSkinVertexFactory.cpp
--- a/Engine/Source/Runtime/Engine/Private/GPUSkinVertexFactory.cpp
+++ b/Engine/Source/Runtime/Engine/Private/GPUSkinVertexFactory.cpp
@@ -16,6 +16,7 @@
 #include "Interfaces/ITargetPlatformManagerModule.h"
 #include "Logging/LogMacros.h"
 #include "Misc/CoreMisc.h"
+#include "RenderGraphResources.h"
 
 #include "Engine/RendererSettings.h"
 #if INTEL_ISPC
@@ -777,17 +778,66 @@ public:
 		class FMeshDrawSingleShaderBindings& ShaderBindings,
 		FVertexInputStreamArray& VertexStreams) const
 	{
-		check(VertexFactory->GetType() == &FGPUSkinPassthroughVertexFactory::StaticType);
+		// todo: Add more context into VertexFactoryUserData about whether this is skin cache/mesh deformer/ray tracing etc.
+		// For now it just holds a skin cache pointer which is null if using other mesh deformers.
 		FGPUSkinBatchElementUserData* BatchUserData = (FGPUSkinBatchElementUserData*)BatchElement.VertexFactoryUserData;
-		check(BatchUserData);
 
-		const auto* LocalVertexFactory = static_cast<const FGPUSkinPassthroughVertexFactory*>(VertexFactory);
+		check(VertexFactory->GetType() == &FGPUSkinPassthroughVertexFactory::StaticType);
+		FGPUSkinPassthroughVertexFactory const* LocalVertexFactory = static_cast<FGPUSkinPassthroughVertexFactory const*>(VertexFactory);
+		
 		FRHIUniformBuffer* VertexFactoryUniformBuffer = nullptr;
 		VertexFactoryUniformBuffer = LocalVertexFactory->GetUniformBuffer();
 
 		// #dxr_todo do we need this call to the base?
 		FLocalVertexFactoryShaderParametersBase::GetElementShaderBindingsBase(Scene, View, Shader, InputStreamType, FeatureLevel, VertexFactory, BatchElement, VertexFactoryUniformBuffer, ShaderBindings, VertexStreams);
-		FGPUSkinCache::GetShaderBindings(BatchUserData->Entry, BatchUserData->Section, Shader, (const FGPUSkinPassthroughVertexFactory*)VertexFactory, BatchElement.MinVertexIndex, GPUSkinCachePositionBuffer, GPUSkinCachePreviousPositionBuffer, ShaderBindings, VertexStreams, View);
+
+		const bool bUsesSkinCache = BatchUserData != nullptr;
+		if (bUsesSkinCache)
+		{
+			GetElementShaderBindingsSkinCache(View, Shader, LocalVertexFactory, BatchElement, BatchUserData, ShaderBindings, VertexStreams);
+		}
+		else
+		{
+			GetElementShaderBindingsMeshDeformer(View, Shader, LocalVertexFactory, BatchElement, ShaderBindings, VertexStreams);
+		}
+	}
+
+	void GetElementShaderBindingsSkinCache(
+		FSceneView const* View,
+		FMeshMaterialShader const* Shader,
+		FGPUSkinPassthroughVertexFactory const* VertexFactory, 
+		FMeshBatchElement const& BatchElement,
+		FGPUSkinBatchElementUserData* BatchUserData,
+		class FMeshDrawSingleShaderBindings& ShaderBindings,
+		FVertexInputStreamArray& VertexStreams) const
+	{
+		FGPUSkinCache::GetShaderBindings(
+			BatchUserData->Entry, BatchUserData->Section, 
+			Shader, VertexFactory, BatchElement.MinVertexIndex, 
+			GPUSkinCachePositionBuffer, GPUSkinCachePreviousPositionBuffer, 
+			ShaderBindings, VertexStreams, View);
+	}
+
+	void GetElementShaderBindingsMeshDeformer(
+		FSceneView const* View,
+		FMeshMaterialShader const* Shader,
+		FGPUSkinPassthroughVertexFactory const* VertexFactory,
+		FMeshBatchElement const& BatchElement,
+		class FMeshDrawSingleShaderBindings& ShaderBindings,
+		FVertexInputStreamArray& VertexStreams) const
+	{
+		if (VertexFactory->PositionRDG.IsValid())
+		{
+			VertexStreams.Add(FVertexInputStream(VertexFactory->GetPositionStreamIndex(), 0, VertexFactory->PositionRDG->GetRHI()));
+			ShaderBindings.Add(GPUSkinCachePositionBuffer, VertexFactory->PositionRDG->GetOrCreateSRV(FRHIBufferSRVCreateInfo(PF_R32_FLOAT)));
+		}
+		if (VertexFactory->TangentRDG.IsValid() && VertexFactory->GetTangentStreamIndex() > -1)
+		{
+			VertexStreams.Add(FVertexInputStream(VertexFactory->GetTangentStreamIndex(), 0, VertexFactory->TangentRDG->GetRHI()));
+		}
+
+		// todo: Need to track previous position buffer for mesh deformers
+		ShaderBindings.Add(GPUSkinCachePreviousPositionBuffer, VertexFactory->GetPositionsSRV());
 	}
 
 private:
@@ -800,6 +850,16 @@ IMPLEMENT_TYPE_LAYOUT(FGPUSkinVertexPassthroughFactoryShaderParameters);
 /*-----------------------------------------------------------------------------
 FGPUSkinPassthroughVertexFactory
 -----------------------------------------------------------------------------*/
+FGPUSkinPassthroughVertexFactory::FGPUSkinPassthroughVertexFactory(ERHIFeatureLevel::Type InFeatureLevel)
+	: FLocalVertexFactory(InFeatureLevel, "FGPUSkinPassthroughVertexFactory"), PositionStreamIndex(-1), TangentStreamIndex(-1)
+{
+	bSupportsManualVertexFetch = true;
+}
+
+FGPUSkinPassthroughVertexFactory::~FGPUSkinPassthroughVertexFactory()
+{
+}
+
 void FGPUSkinPassthroughVertexFactory::ModifyCompilationEnvironment(const FVertexFactoryShaderPermutationParameters& Parameters, FShaderCompilerEnvironment& OutEnvironment )
 {
 	const bool ContainsManualVertexFetch = OutEnvironment.GetDefinitions().Contains("MANUAL_VERTEX_FETCH");
@@ -820,30 +880,34 @@ bool FGPUSkinPassthroughVertexFactory::ShouldCompilePermutation(const FVertexFac
 		(Parameters.MaterialParameters.bIsUsedWithSkeletalMesh || Parameters.MaterialParameters.bIsSpecialEngineMaterial);
 }
 
-void FGPUSkinPassthroughVertexFactory::InternalUpdateVertexDeclaration(FGPUBaseSkinVertexFactory* SourceVertexFactory, struct FRWBuffer* PositionRWBuffer, class FRHIShaderResourceView* PreSkinPositionSRV, struct FRWBuffer* TangentRWBuffer)
+void FGPUSkinPassthroughVertexFactory::ReleaseRHI()
 {
-	// Point this vertex buffer to the RWBuffer
-	PositionVBAlias.VertexBufferRHI = PositionRWBuffer->Buffer;
-
-	TangentVBAlias.VertexBufferRHI = TangentRWBuffer ? TangentRWBuffer->Buffer : nullptr;
+	FLocalVertexFactory::ReleaseRHI();
 
-	// Modify the vertex declaration using the RWBuffer for the position & tangent information
-	Data.PositionComponent.VertexBuffer = &PositionVBAlias;
-	Data.PositionComponent.Offset = 0;
-	Data.PositionComponent.VertexStreamUsage = EVertexStreamUsage::Overridden;
-	Data.PositionComponent.Stride = 3 * sizeof(float);
+	// When adding anything else to this function be aware of the bypassing code in InternalUpdateVertexDeclaration.
+	PositionRDG = nullptr;
+	TangentRDG = nullptr;
+	ColorRDG = nullptr;
+	PositionVBAlias.ReleaseRHI();
+	TangentVBAlias.ReleaseRHI();
+	ColorVBAlias.ReleaseRHI();
+}
 
-	
+void FGPUSkinPassthroughVertexFactory::InternalUpdateVertexDeclaration(FGPUBaseSkinVertexFactory* SourceVertexFactory)
+{
+	if (PositionVBAlias.VertexBufferRHI != nullptr)
 	{
-		Data.TangentsSRV = TangentRWBuffer ? TangentRWBuffer->SRV : SourceVertexFactory->GetTangentsSRV();
-		Data.PositionComponentSRV = PositionRWBuffer->SRV;
-		Data.PreSkinPositionComponentSRV = PreSkinPositionSRV;
+		Data.PositionComponent.VertexBuffer = &PositionVBAlias;
+		Data.PositionComponent.Offset = 0;
+		Data.PositionComponent.VertexStreamUsage = EVertexStreamUsage::Overridden;
+		Data.PositionComponent.Stride = 3 * sizeof(float);
+	}
+	else
+	{
+		Data.PositionComponent = SourceVertexFactory->GetPositionStreamComponent();
 	}
 
-	Data.TangentBasisComponents[0] = SourceVertexFactory->GetTangentStreamComponent(0);
-	Data.TangentBasisComponents[1] = SourceVertexFactory->GetTangentStreamComponent(1);
-
-	if (TangentRWBuffer)
+	if (TangentVBAlias.VertexBufferRHI != nullptr)
 	{
 		Data.TangentBasisComponents[0].VertexBuffer = &TangentVBAlias;
 		Data.TangentBasisComponents[0].Offset = 0;
@@ -857,35 +921,116 @@ void FGPUSkinPassthroughVertexFactory::InternalUpdateVertexDeclaration(FGPUBaseS
 		Data.TangentBasisComponents[1].Stride = 16;
 		Data.TangentBasisComponents[1].VertexStreamUsage = EVertexStreamUsage::Overridden | EVertexStreamUsage::ManualFetch;
 	}
+	else
+	{
+		Data.TangentBasisComponents[0] = SourceVertexFactory->GetTangentStreamComponent(0);
+		Data.TangentBasisComponents[1] = SourceVertexFactory->GetTangentStreamComponent(1);
+	}
+
+	if (ColorVBAlias.VertexBufferRHI)
+	{
+		Data.ColorComponent.VertexBuffer = &ColorVBAlias;
+		Data.ColorComponent.Offset = 0;
+		Data.ColorComponent.Type = VET_Color;
+		Data.ColorComponent.Stride = 4;
+		Data.ColorComponent.VertexStreamUsage = EVertexStreamUsage::Overridden | EVertexStreamUsage::ManualFetch;
+		
+		Data.ColorIndexMask = ~0u;
+	}
+
+	Data.PositionComponentSRV = PositionSRVAlias ? PositionSRVAlias : SourceVertexFactory->GetPositionsSRV().GetReference();
+	Data.PreSkinPositionComponentSRV = SourceVertexFactory->GetPositionsSRV().GetReference();
+	Data.TangentsSRV = TangentSRVAlias ? TangentSRVAlias : SourceVertexFactory->GetTangentsSRV().GetReference();
+	Data.ColorComponentsSRV = ColorSRVAlias ? ColorSRVAlias : SourceVertexFactory->GetColorComponentsSRV().GetReference();
 
-	//hack to allow us to release the alias pointers properly in ReleaseRHI.
-	//To be cleaned up in UE-68826
+	// hack to allow us to release the alias pointers properly in ReleaseRHI.
+	// To be cleaned up in UE-68826
 	FLocalVertexFactory::ReleaseRHI();
 	FLocalVertexFactory::ReleaseDynamicRHI();
 	FLocalVertexFactory::InitDynamicRHI();
 	FLocalVertexFactory::InitRHI();
 
-	// Find the added stream (usually at 0)
+	// Find added streams
 	PositionStreamIndex = -1;
 	TangentStreamIndex = -1;
+	ColorStreamIndex = -1;
+
 	for (int32 Index = 0; Index < Streams.Num(); ++Index)
 	{
-		if (Streams[Index].VertexBuffer->VertexBufferRHI.GetReference() == PositionRWBuffer->Buffer.GetReference())
+		if (Streams[Index].VertexBuffer->VertexBufferRHI.GetReference() == Data.PositionComponent.VertexBuffer->VertexBufferRHI.GetReference())
 		{
 			PositionStreamIndex = Index;
 		}
-
-		if (TangentRWBuffer)
+		if (TangentVBAlias.VertexBufferRHI != nullptr)
 		{
-			if (Streams[Index].VertexBuffer->VertexBufferRHI.GetReference() == TangentRWBuffer->Buffer.GetReference())
+			if (Streams[Index].VertexBuffer->VertexBufferRHI.GetReference() == Data.TangentBasisComponents[0].VertexBuffer->VertexBufferRHI.GetReference())
 			{
 				TangentStreamIndex = Index;
 			}
 		}
+		if (ColorVBAlias.VertexBufferRHI != nullptr)
+		{
+			if (Streams[Index].VertexBuffer->VertexBufferRHI.GetReference() == Data.ColorComponent.VertexBuffer->VertexBufferRHI.GetReference())
+			{
+				ColorStreamIndex = Index;
+			}
+		}
 	}
 	checkf(PositionStreamIndex != -1, TEXT("Unable to find stream for RWBuffer Vertex buffer!"));
 }
 
+void FGPUSkinPassthroughVertexFactory::InternalUpdateVertexDeclaration(
+	FGPUBaseSkinVertexFactory* SourceVertexFactory, 
+	struct FRWBuffer* PositionRWBuffer, 
+	struct FRWBuffer* TangentRWBuffer)
+{
+	PositionRDG = nullptr;
+	TangentRDG = nullptr;
+	ColorRDG = nullptr;
+
+	PositionVBAlias.VertexBufferRHI = PositionRWBuffer ? PositionRWBuffer->Buffer : nullptr;
+	TangentVBAlias.VertexBufferRHI = TangentRWBuffer ? TangentRWBuffer->Buffer : nullptr;
+	ColorVBAlias.VertexBufferRHI = nullptr;
+
+	PositionSRVAlias = PositionRWBuffer ? PositionRWBuffer->SRV : nullptr;
+	TangentSRVAlias = TangentRWBuffer ? TangentRWBuffer->SRV : nullptr;
+	ColorSRVAlias = nullptr;
+
+	InternalUpdateVertexDeclaration(SourceVertexFactory);
+}
+
+void FGPUSkinPassthroughVertexFactory::InternalUpdateVertexDeclaration(
+	EOverrideFlags OverrideFlags,
+	FGPUBaseSkinVertexFactory* SourceVertexFactory, 
+	TRefCountPtr<FRDGPooledBuffer> const& PositionBuffer, 
+	TRefCountPtr<FRDGPooledBuffer> const& TangentBuffer,
+	TRefCountPtr<FRDGPooledBuffer> const& ColorBuffer)
+{
+	if (EnumHasAnyFlags(OverrideFlags, EOverrideFlags::Position))
+	{
+		PositionRDG = PositionBuffer;
+	}
+	if (EnumHasAnyFlags(OverrideFlags, EOverrideFlags::Tangent))
+	{
+		TangentRDG = TangentBuffer;
+	}
+	if (EnumHasAnyFlags(OverrideFlags, EOverrideFlags::Color))
+	{
+		ColorRDG = ColorBuffer;
+	}
+
+	PositionVBAlias.VertexBufferRHI = PositionRDG.IsValid() ? PositionRDG->GetRHI() : nullptr;
+	TangentVBAlias.VertexBufferRHI = TangentRDG.IsValid() ? TangentRDG->GetRHI() : nullptr;
+	ColorVBAlias.VertexBufferRHI = ColorRDG.IsValid() ? ColorRDG->GetRHI() : nullptr;
+
+	PositionSRVAlias = PositionRDG.IsValid() ? PositionRDG->GetOrCreateSRV(FRHIBufferSRVCreateInfo(PF_R32_FLOAT)) : nullptr;
+	const EPixelFormat TangentsFormat = IsOpenGLPlatform(GMaxRHIShaderPlatform) ? PF_R16G16B16A16_SINT : PF_R16G16B16A16_SNORM;
+	TangentSRVAlias = TangentRDG.IsValid() ? TangentRDG->GetOrCreateSRV(FRHIBufferSRVCreateInfo(TangentsFormat)) : nullptr;
+	ColorSRVAlias = ColorRDG.IsValid() ? ColorRDG->GetOrCreateSRV(FRHIBufferSRVCreateInfo(PF_R8G8B8A8)) : nullptr;
+
+	InternalUpdateVertexDeclaration(SourceVertexFactory);
+}
+
 IMPLEMENT_VERTEX_FACTORY_PARAMETER_TYPE(FGPUSkinPassthroughVertexFactory, SF_Vertex, FGPUSkinVertexPassthroughFactoryShaderParameters);
 #if RHI_RAYTRACING
 IMPLEMENT_VERTEX_FACTORY_PARAMETER_TYPE(FGPUSkinPassthroughVertexFactory, SF_RayHitGroup, FGPUSkinVertexPassthroughFactoryShaderParameters);
diff --git a/Engine/Source/Runtime/Engine/Private/SkeletalMeshDeformerHelpers.cpp b/Engine/Source/Runtime/Engine/Private/SkeletalMeshDeformerHelpers.cpp
new file mode 100644
--- /dev/null
+++ b/Engine/Source/Runtime/Engine/Private/SkeletalMeshDeformerHelpers.cpp
@@ -0,0 +1,75 @@
+// Copyright Epic Games, Inc. All Rights Reserved.
+
+#pragma once
+
+#include "SkeletalMeshDeformerHelpers.h"
+
+#include "GPUSkinVertexFactory.h"
+#include "RenderGraphResources.h"
+#include "SkeletalRenderGPUSkin.h"
+
+FRHIShaderResourceView* FSkeletalMeshDeformerHelpers::GetBoneBufferForReading(
+	FSkeletalMeshObject* MeshObject,
+	int32 LODIndex,
+	int32 SectionIndex,
+	bool bPreviousFrame)
+{
+	if (MeshObject->IsCPUSkinned())
+	{
+		return nullptr;
+	}
+
+	FSkeletalMeshObjectGPUSkin* MeshObjectGPU = static_cast<FSkeletalMeshObjectGPUSkin*>(MeshObject);
+	const FSkeletalMeshObjectGPUSkin::FSkeletalMeshObjectLOD& LOD = MeshObjectGPU->LODs[LODIndex];
+	FGPUBaseSkinVertexFactory* BaseVertexFactory = LOD.GPUSkinVertexFactories.VertexFactories[SectionIndex].Get();
+	FShaderResourceViewRHIRef SRV = BaseVertexFactory->GetShaderData().GetBoneBufferForReading(bPreviousFrame).VertexBufferSRV;
+	return SRV.IsValid() ? SRV : nullptr;
+}
+
+void FSkeletalMeshDeformerHelpers::SetVertexFactoryBufferOverrides(
+	FSkeletalMeshObject* MeshObject,
+	int32 LODIndex,
+	EOverrideType OverrideType,
+	TRefCountPtr<FRDGPooledBuffer> const& PositionBuffer,
+	TRefCountPtr<FRDGPooledBuffer> const& TangentBuffer,
+	TRefCountPtr<FRDGPooledBuffer> const& ColorBuffer)
+{
+	if (MeshObject->IsCPUSkinned())
+	{
+		return;
+	}
+
+	FGPUSkinPassthroughVertexFactory::EOverrideFlags OverrideFlags = FGPUSkinPassthroughVertexFactory::EOverrideFlags::All;
+	if (OverrideType == EOverrideType::Partial)
+	{
+		OverrideFlags = FGPUSkinPassthroughVertexFactory::EOverrideFlags::None;
+ 		OverrideFlags |= (PositionBuffer.IsValid() ? FGPUSkinPassthroughVertexFactory::EOverrideFlags::Position : FGPUSkinPassthroughVertexFactory::EOverrideFlags::None);
+ 		OverrideFlags |= (TangentBuffer.IsValid() ? FGPUSkinPassthroughVertexFactory::EOverrideFlags::Tangent : FGPUSkinPassthroughVertexFactory::EOverrideFlags::None);
+ 		OverrideFlags |= (ColorBuffer.IsValid() ? FGPUSkinPassthroughVertexFactory::EOverrideFlags::Color : FGPUSkinPassthroughVertexFactory::EOverrideFlags::None);
+	}
+
+	FSkeletalMeshObjectGPUSkin* MeshObjectGPU = static_cast<FSkeletalMeshObjectGPUSkin*>(MeshObject);
+	const FSkeletalMeshObjectGPUSkin::FSkeletalMeshObjectLOD& LOD = MeshObjectGPU->LODs[LODIndex];
+	
+	const int32 NumSections = MeshObject->GetRenderSections(LODIndex).Num();
+	for (int32 SectionIndex = 0; SectionIndex < NumSections; ++SectionIndex)
+	{
+		FGPUBaseSkinVertexFactory* BaseVertexFactory = LOD.GPUSkinVertexFactories.VertexFactories[SectionIndex].Get();
+		FGPUSkinPassthroughVertexFactory* TargetVertexFactory = LOD.GPUSkinVertexFactories.PassthroughVertexFactories[SectionIndex].Get();
+		TargetVertexFactory->InvalidateStreams();
+		TargetVertexFactory->UpdateVertexDeclaration(OverrideFlags, BaseVertexFactory, PositionBuffer, TangentBuffer, ColorBuffer);
+	}
+}
+
+void FSkeletalMeshDeformerHelpers::ResetVertexFactoryBufferOverrides_GameThread(FSkeletalMeshObject* MeshObject, int32 LODIndex)
+{
+	if (MeshObject->IsCPUSkinned())
+	{
+		return;
+	}
+
+	ENQUEUE_RENDER_COMMAND(ResetSkinPassthroughVertexFactory)([MeshObject, LODIndex](FRHICommandList& CmdList)
+	{
+		SetVertexFactoryBufferOverrides(MeshObject, LODIndex, EOverrideType::All, nullptr, nullptr, nullptr);
+	});
+}
diff --git a/Engine/Source/Runtime/Engine/Private/SkeletalRenderGPUSkin.cpp b/Engine/Source/Runtime/Engine/Private/SkeletalRenderGPUSkin.cpp
--- a/Engine/Source/Runtime/Engine/Private/SkeletalRenderGPUSkin.cpp
+++ b/Engine/Source/Runtime/Engine/Private/SkeletalRenderGPUSkin.cpp
@@ -626,7 +626,7 @@ void FSkeletalMeshObjectGPUSkin::ProcessUpdatedDynamicData(EGPUSkinCacheEntryMod
 
 	bool bDataPresent = false;
 
-	bool bGPUSkinCacheEnabled = GPUSkinCache && GEnableGPUSkinCache && (DynamicData->bIsSkinCacheAllowed || Mode == EGPUSkinCacheEntryMode::RayTracing);
+	bool bGPUSkinCacheEnabled = GPUSkinCache && GEnableGPUSkinCache && (DynamicData->bIsSkinCacheAllowed || Mode == EGPUSkinCacheEntryMode::RayTracing) && !DynamicData->bHasMeshDeformer;
 
 	// Immediately release any stale entry if we've recently switched to a LOD level that disallows skin cache
 	// This saves memory and avoids confusing ShouldUseSeparateSkinCacheEntryForRayTracing() which checks SkinCacheEntry == nullptr
@@ -1226,6 +1226,12 @@ const FVertexFactory* FSkeletalMeshObjectGPUSkin::GetSkinVertexFactory(const FSc
 	}
 #endif
 
+	// If a mesh deformer cache was used, return the passthrough vertex factory
+	if (DynamicData->bHasMeshDeformer)
+	{
+		return LOD.GPUSkinVertexFactories.PassthroughVertexFactories[ChunkIdx].Get();
+	}
+
 	// If the GPU skinning cache was used, return the passthrough vertex factory
 	if (SkinCacheEntry && FGPUSkinCache::IsEntryValid(SkinCacheEntry, ChunkIdx) && DynamicData->bIsSkinCacheAllowed)
 	{
@@ -1923,6 +1929,7 @@ void FDynamicSkelMeshObjectDataGPUSkin::Clear()
 	ClothingSimData.Reset();
 	ClothBlendWeight = 0.0f;
 	bIsSkinCacheAllowed = false;
+	bHasMeshDeformer = false;
 #if RHI_RAYTRACING
 	bAnySegmentUsesWorldPositionOffset = false;
 #endif
@@ -2114,6 +2121,7 @@ void FDynamicSkelMeshObjectDataGPUSkin::InitDynamicSkelMeshObjectDataGPUSkin(
 	UpdateClothSimulationData(InMeshComponent);
 
 	bIsSkinCacheAllowed = InMeshComponent ? InMeshComponent->IsSkinCacheAllowed(InLODIndex) : false;
+	bHasMeshDeformer = InMeshComponent ? InMeshComponent->HasMeshDeformer() : false;
 
 	if (bIsSkinCacheAllowed && InMeshObject->FeatureLevel == ERHIFeatureLevel::ES3_1)
 	{
diff --git a/Engine/Source/Runtime/Engine/Private/SkeletalRenderGPUSkin.h b/Engine/Source/Runtime/Engine/Private/SkeletalRenderGPUSkin.h
--- a/Engine/Source/Runtime/Engine/Private/SkeletalRenderGPUSkin.h
+++ b/Engine/Source/Runtime/Engine/Private/SkeletalRenderGPUSkin.h
@@ -126,6 +126,9 @@ public:
 
 	// Whether this LOD is allowed to use the skin cache feature
 	uint8 bIsSkinCacheAllowed : 1;
+	
+	// Whether animation is done with a mesh deformer.
+	uint8 bHasMeshDeformer : 1;
 
 #if RHI_RAYTRACING
 	uint8 bAnySegmentUsesWorldPositionOffset : 1;
@@ -418,7 +421,9 @@ public:
 	virtual void SetTransform(const FMatrix& InNewLocalToWorld, uint32 FrameNumber) override;
 	virtual void RefreshClothingTransforms(const FMatrix& InNewLocalToWorld, uint32 FrameNumber) override;
 	virtual void UpdateSkinWeightBuffer(USkinnedMeshComponent* InMeshComponent) override;
+
 protected:
+	friend class FSkeletalMeshDeformerHelpers;
 
 	/**
 	 * Vertex factories and their matrix arrays
diff --git a/Engine/Source/Runtime/Engine/Public/GPUSkinVertexFactory.h b/Engine/Source/Runtime/Engine/Public/GPUSkinVertexFactory.h
--- a/Engine/Source/Runtime/Engine/Public/GPUSkinVertexFactory.h
+++ b/Engine/Source/Runtime/Engine/Public/GPUSkinVertexFactory.h
@@ -19,6 +19,8 @@
 #include "Matrix3x4.h"
 #include "SkeletalMeshTypes.h"
 
+class FRDGPooledBuffer;
+
 template <class T> class TConsoleVariableData;
 
 // Uniform buffer for APEX cloth
@@ -366,6 +368,12 @@ public:
 	virtual uint32 GetNumTexCoords() const										{ return Data.IsValid() ? Data->NumTexCoords : 0; }
 	virtual const uint32 GetColorIndexMask() const								{ return Data.IsValid() ? Data->ColorIndexMask : 0; }
 
+	inline const FVertexStreamComponent& GetPositionStreamComponent()
+	{
+		check(Data.IsValid() && Data->PositionComponent.VertexBuffer != nullptr);
+		return Data->PositionComponent;
+	}
+	
 	inline const FVertexStreamComponent& GetTangentStreamComponent(int Index)
 	{
 		check(Data.IsValid() && Data->TangentBasisComponents[Index].VertexBuffer != nullptr);
@@ -440,23 +448,12 @@ class FGPUSkinPassthroughVertexFactory : public FLocalVertexFactory
 	typedef FLocalVertexFactory Super;
 
 public:
-	FGPUSkinPassthroughVertexFactory(ERHIFeatureLevel::Type InFeatureLevel)
-		: FLocalVertexFactory(InFeatureLevel, "FGPUSkinPassthroughVertexFactory"), PositionStreamIndex(-1), TangentStreamIndex(-1)
-	{
-		bSupportsManualVertexFetch = true;
-	}
+	FGPUSkinPassthroughVertexFactory(ERHIFeatureLevel::Type InFeatureLevel);
+	~FGPUSkinPassthroughVertexFactory();
 
 	static void ModifyCompilationEnvironment(const FVertexFactoryShaderPermutationParameters& Parameters, FShaderCompilerEnvironment& OutEnvironment);
 	static bool ShouldCompilePermutation(const FVertexFactoryShaderPermutationParameters& Parameters);
 
-	inline void UpdateVertexDeclaration(FGPUBaseSkinVertexFactory* SourceVertexFactory, struct FRWBuffer* PositionRWBuffer, class FRHIShaderResourceView* PreSkinPositionSRV, struct FRWBuffer* TangentRWBuffer)
-	{
-		if (PositionStreamIndex == -1)
-		{
-			InternalUpdateVertexDeclaration(SourceVertexFactory, PositionRWBuffer, PreSkinPositionSRV, TangentRWBuffer);
-		}
-	}
-
 	inline int32 GetPositionStreamIndex() const
 	{
 		check(PositionStreamIndex > -1);
@@ -478,29 +475,85 @@ public:
 	{
 		PositionStreamIndex = -1;
 		TangentStreamIndex = -1;
+		ColorStreamIndex = -1;
 		PositionVBAlias.ReleaseRHI();
 		TangentVBAlias.ReleaseRHI();
+		ColorVBAlias.ReleaseRHI();
+	}
+
+	virtual void ReleaseRHI() override;
+
+	/** Poke values into the vertex factory. */
+	inline void UpdateVertexDeclaration(
+		FGPUBaseSkinVertexFactory* SourceVertexFactory, 
+		struct FRWBuffer* PositionRWBuffer, 
+		struct FRWBuffer* TangentRWBuffer)
+	{
+		if (PositionStreamIndex == -1)
+		{
+			InternalUpdateVertexDeclaration(SourceVertexFactory, PositionRWBuffer, TangentRWBuffer);
+		}
 	}
 
-	virtual void ReleaseRHI() override
+	/** Flags for override bitmask passed to UpdateVertexDeclaration(). */
+	enum class EOverrideFlags
 	{
-		FLocalVertexFactory::ReleaseRHI();
+		None		= 0,
+		Position	= 1 << 0,
+		Tangent		= 1 << 1,
+		Color		= 1 << 2,
+		All			= 0xff,
+	};
 
-		//when adding anything else to this function be aware of the bypassing code in InternalUpdateVertexDeclaration
-		PositionVBAlias.ReleaseRHI();
-		TangentVBAlias.ReleaseRHI();
+	/** Poke values into the vertex factory. */
+	inline void UpdateVertexDeclaration(
+		EOverrideFlags OverrideFlags,
+		FGPUBaseSkinVertexFactory* SourceVertexFactory,
+		TRefCountPtr<FRDGPooledBuffer> const& PositionBuffer,
+		TRefCountPtr<FRDGPooledBuffer> const& TangentBuffer,
+		TRefCountPtr<FRDGPooledBuffer> const& ColorBuffer)
+	{
+		if (PositionStreamIndex == -1)
+		{
+			InternalUpdateVertexDeclaration(OverrideFlags, SourceVertexFactory, PositionBuffer, TangentBuffer, ColorBuffer);
+		}
 	}
 
 protected:
+	friend class FGPUSkinVertexPassthroughFactoryShaderParameters;
+
+	// Reference holders for RDG buffers
+	TRefCountPtr<FRDGPooledBuffer> PositionRDG;
+	TRefCountPtr<FRDGPooledBuffer> TangentRDG;
+	TRefCountPtr<FRDGPooledBuffer> ColorRDG;
 	// Vertex buffer required for creating the Vertex Declaration
 	FVertexBuffer PositionVBAlias;
 	FVertexBuffer TangentVBAlias;
+	FVertexBuffer ColorVBAlias;
+	// SRVs required for binding
+	FRHIShaderResourceView* PositionSRVAlias;
+	FRHIShaderResourceView* TangentSRVAlias;
+	FRHIShaderResourceView* ColorSRVAlias;
+	// Cached stream indices
 	int32 PositionStreamIndex;
 	int32 TangentStreamIndex;
 
-	void InternalUpdateVertexDeclaration(FGPUBaseSkinVertexFactory* SourceVertexFactory, struct FRWBuffer* PositionRWBuffer, class FRHIShaderResourceView* PreSkinPositionSRV, struct FRWBuffer* TangentRWBuffer);
+	void InternalUpdateVertexDeclaration(
+		FGPUBaseSkinVertexFactory* SourceVertexFactory);
+	void InternalUpdateVertexDeclaration(
+		FGPUBaseSkinVertexFactory* SourceVertexFactory, 
+		struct FRWBuffer* PositionRWBuffer, 
+		struct FRWBuffer* TangentRWBuffer);
+	void InternalUpdateVertexDeclaration(
+		EOverrideFlags OverrideFlags,
+		FGPUBaseSkinVertexFactory* SourceVertexFactory,
+		TRefCountPtr<FRDGPooledBuffer> const& PositionBuffer,
+		TRefCountPtr<FRDGPooledBuffer> const& TangentBuffer,
+		TRefCountPtr<FRDGPooledBuffer> const& ColorBuffer);
 };
 
+ENUM_CLASS_FLAGS(FGPUSkinPassthroughVertexFactory::EOverrideFlags)
+
 /** Stream component data bound to morph vertex factory */
 struct FGPUSkinMorphDataType : public FGPUSkinDataType
 {
diff --git a/Engine/Source/Runtime/Engine/Public/SkeletalMeshDeformerHelpers.h b/Engine/Source/Runtime/Engine/Public/SkeletalMeshDeformerHelpers.h
new file mode 100644
--- /dev/null
+++ b/Engine/Source/Runtime/Engine/Public/SkeletalMeshDeformerHelpers.h
@@ -0,0 +1,39 @@
+// Copyright Epic Games, Inc. All Rights Reserved.
+
+#pragma once
+
+#include "CoreMinimal.h"
+
+class FSkeletalMeshObject;
+class FRDGPooledBuffer;
+
+/** Functions that expose some functionality of FSkeletalMeshObject required by MeshDeformer systems. */
+class FSkeletalMeshDeformerHelpers
+{
+public:
+	/** Get direct access to bone matrix buffer SRV. */
+	ENGINE_API static FRHIShaderResourceView* GetBoneBufferForReading(
+		FSkeletalMeshObject* MeshObject,
+		int32 LODIndex,
+		int32 SectionIndex,
+		bool bPreviousFrame);
+
+	/** Buffer override behavior for SetVertexFactoryBufferOverrides. */
+	enum class EOverrideType
+	{
+		All,		// Clear overrides for each buffer input that is null.
+		Partial,	// Leave existing overrides for each buffer input that is null.
+	};
+
+	/** Apply buffer overrides to the pass through vertex factory. */
+	ENGINE_API static void SetVertexFactoryBufferOverrides(
+		FSkeletalMeshObject* MeshObject,
+		int32 LODIndex, 
+		EOverrideType OverrideType,
+		TRefCountPtr<FRDGPooledBuffer> const& PositionBuffer, 
+		TRefCountPtr<FRDGPooledBuffer> const& TangentBuffer, 
+		TRefCountPtr<FRDGPooledBuffer> const& ColorBuffer);
+
+	/** Reset all buffer overrides that were applied through SetVertexFactoryBufferOverrides. */
+	ENGINE_API static void ResetVertexFactoryBufferOverrides_GameThread(FSkeletalMeshObject* MeshObject, int32 LODIndex);
+};
