PackedLevelActor: Fix low repro crash, avoid calling MarkComponentsRenderStateDirty on reinstanced BP Actor

diff --git a/Engine/Source/Runtime/Engine/Private/PackedLevelActor/PackedLevelActor.cpp b/Engine/Source/Runtime/Engine/Private/PackedLevelActor/PackedLevelActor.cpp
--- a/Engine/Source/Runtime/Engine/Private/PackedLevelActor/PackedLevelActor.cpp
+++ b/Engine/Source/Runtime/Engine/Private/PackedLevelActor/PackedLevelActor.cpp
@@ -139,6 +139,8 @@ void APackedLevelActor::OnCommitChild(bool bChanged)
 
 		if (bChildChanged)
 		{
+			bChildChanged = false;
+
 			// Reflect child changes
 			TSharedPtr<FPackedLevelActorBuilder> Builder = FPackedLevelActorBuilder::CreateDefaultBuilder();
 
@@ -146,14 +148,14 @@ void APackedLevelActor::OnCommitChild(bool bChanged)
 			{
 				check(GeneratedBy == BlueprintAsset.Get());
 				Builder->UpdateBlueprint(GeneratedBy);
+				return; // return here because Actor might have been reinstanced
 			}
 			else
 			{
 				Builder->PackActor(this, GetWorldAsset());
 			}
-			bChildChanged = false;
 		}
-
+		// When child edit state changes we need to dirty render state so that actor is no longer hidden
 		MarkComponentsRenderStateDirty();
 	}
 }
@@ -176,10 +178,11 @@ void APackedLevelActor::OnCommit(bool bChanged, bool bPromptForSave)
 			const bool bCheckoutAndSave = true;
 			TSharedPtr<FPackedLevelActorBuilder> Builder = FPackedLevelActorBuilder::CreateDefaultBuilder();
 			Builder->UpdateBlueprint(GeneratedBy, bCheckoutAndSave, bPromptForSave);
+			return; // return here because Actor might have been reinstanced
 		}
 	}
 	
-	// bEditing flag changed so dirty render state
+	// When edit state changes we need to dirty render state so that actor is no longer hidden
 	MarkComponentsRenderStateDirty();
 }
 
