Moving a couple hardcoded uses of IsVulkanMobileSM5Platform to DDSPI. Fixing discontinuity in how SINGLE_LAYER_WATER_SIMPLE_FORWARD is set.

diff --git a/Engine/Config/Android/DataDrivenPlatformInfo.ini b/Engine/Config/Android/DataDrivenPlatformInfo.ini
--- a/Engine/Config/Android/DataDrivenPlatformInfo.ini
+++ b/Engine/Config/Android/DataDrivenPlatformInfo.ini
@@ -80,3 +80,5 @@ bTargetsTiledGPU=true
 bNeedsOfflineCompiler=true
 bSupportsIndexBufferUAVs=true
 bSupportsManualVertexFetch=false
+bRequiresDisableForwardLocalLights=true
+bWaterUsesSimpleForwardShading=true
\ No newline at end of file
diff --git a/Engine/Source/Runtime/Engine/Private/Materials/HLSLMaterialTranslator.cpp b/Engine/Source/Runtime/Engine/Private/Materials/HLSLMaterialTranslator.cpp
--- a/Engine/Source/Runtime/Engine/Private/Materials/HLSLMaterialTranslator.cpp
+++ b/Engine/Source/Runtime/Engine/Private/Materials/HLSLMaterialTranslator.cpp
@@ -1628,16 +1628,13 @@ void FHLSLMaterialTranslator::GetMaterialEnvironment(EShaderPlatform InPlatform,
 			bMaterialRequestsDualSourceBlending = true;
 		}
 
-		if (ShadingModels.HasShadingModel(MSM_SingleLayerWater) &&
-			(IsVulkanMobileSM5Platform(Platform) || FDataDrivenShaderPlatformInfo::GetRequiresDisableForwardLocalLights(Platform)))
+		if (ShadingModels.HasShadingModel(MSM_SingleLayerWater) && FDataDrivenShaderPlatformInfo::GetRequiresDisableForwardLocalLights(Platform))
 		{
 			OutEnvironment.SetDefine(TEXT("DISABLE_FORWARD_LOCAL_LIGHTS"), TEXT("1"));
 		}
 
 		// This is to have platforms use the simple single layer water shading similar to mobile: no dynamic lights, only sun and sky, no distortion, no colored transmittance on background, no custom depth read.
-		bool bSingleLayerWaterUsesSimpleShading = IsVulkanMobileSM5Platform(InPlatform);
-		bSingleLayerWaterUsesSimpleShading = bSingleLayerWaterUsesSimpleShading || FDataDrivenShaderPlatformInfo::GetWaterUsesSimpleForwardShading(InPlatform);
-		bSingleLayerWaterUsesSimpleShading = bSingleLayerWaterUsesSimpleShading && IsForwardShadingEnabled(InPlatform);
+		const bool bSingleLayerWaterUsesSimpleShading = FDataDrivenShaderPlatformInfo::GetWaterUsesSimpleForwardShading(InPlatform) && IsForwardShadingEnabled(InPlatform);
 
 		if (ShadingModels.HasShadingModel(MSM_SingleLayerWater) && bSingleLayerWaterUsesSimpleShading)
 		{
diff --git a/Engine/Source/Runtime/Engine/Private/Materials/MaterialHLSLEmitter.cpp b/Engine/Source/Runtime/Engine/Private/Materials/MaterialHLSLEmitter.cpp
--- a/Engine/Source/Runtime/Engine/Private/Materials/MaterialHLSLEmitter.cpp
+++ b/Engine/Source/Runtime/Engine/Private/Materials/MaterialHLSLEmitter.cpp
@@ -556,14 +556,14 @@ static void GetMaterialEnvironment(EShaderPlatform InPlatform,
 			bMaterialRequestsDualSourceBlending = true;
 		}
 
-		if (ShadingModels.HasShadingModel(MSM_SingleLayerWater) &&
-			(IsVulkanMobileSM5Platform(InPlatform) || FDataDrivenShaderPlatformInfo::GetRequiresDisableForwardLocalLights(InPlatform)))
+		if (ShadingModels.HasShadingModel(MSM_SingleLayerWater) && FDataDrivenShaderPlatformInfo::GetRequiresDisableForwardLocalLights(InPlatform))
 		{
 			OutEnvironment.SetDefine(TEXT("DISABLE_FORWARD_LOCAL_LIGHTS"), TEXT("1"));
 		}
 
 		// This is to have platforms use the simple single layer water shading similar to mobile: no dynamic lights, only sun and sky, no distortion, no colored transmittance on background, no custom depth read.
-		const bool bSingleLayerWaterUsesSimpleShading = IsVulkanMobileSM5Platform(InPlatform) && IsForwardShadingEnabled(InPlatform);
+		const bool bSingleLayerWaterUsesSimpleShading = FDataDrivenShaderPlatformInfo::GetWaterUsesSimpleForwardShading(InPlatform) && IsForwardShadingEnabled(InPlatform);
+
 		if (ShadingModels.HasShadingModel(MSM_SingleLayerWater) && bSingleLayerWaterUsesSimpleShading)
 		{
 			OutEnvironment.SetDefine(TEXT("SINGLE_LAYER_WATER_SIMPLE_FORWARD"), TEXT("1"));
diff --git a/Engine/Source/Runtime/Renderer/Private/SingleLayerWaterRendering.cpp b/Engine/Source/Runtime/Renderer/Private/SingleLayerWaterRendering.cpp
--- a/Engine/Source/Runtime/Renderer/Private/SingleLayerWaterRendering.cpp
+++ b/Engine/Source/Runtime/Renderer/Private/SingleLayerWaterRendering.cpp
@@ -75,15 +75,11 @@ static TAutoConsoleVariable<int32> CVarRHICmdFlushRenderThreadTasksSingleLayerWa
 	TEXT("r.RHICmdFlushRenderThreadTasksSingleLayerWater"),
 	0,
 	TEXT("Wait for completion of parallel render thread tasks at the end of Single layer water. A more granular version of r.RHICmdFlushRenderThreadTasks. If either r.RHICmdFlushRenderThreadTasks or r.RHICmdFlushRenderThreadTasksSingleLayerWater is > 0 we will flush."));
+
 // This is to have platforms use the simple single layer water shading similar to mobile: no dynamic lights, only sun and sky, no distortion, no colored transmittance on background, no custom depth read.
 bool SingleLayerWaterUsesSimpleShading(EShaderPlatform ShaderPlatform)
 {
-	bool bUsesSimpleShading;
-	bUsesSimpleShading = IsVulkanMobileSM5Platform(ShaderPlatform);
-	bUsesSimpleShading = bUsesSimpleShading || FDataDrivenShaderPlatformInfo::GetWaterUsesSimpleForwardShading(ShaderPlatform);
-	bUsesSimpleShading = bUsesSimpleShading && IsForwardShadingEnabled(ShaderPlatform);
-
-	return bUsesSimpleShading;
+	return FDataDrivenShaderPlatformInfo::GetWaterUsesSimpleForwardShading(ShaderPlatform) && IsForwardShadingEnabled(ShaderPlatform);
 }
 
 bool ShouldRenderSingleLayerWater(TArrayView<const FViewInfo> Views)
