World Partition: merge code for different versions of CheckForErrors.

diff --git a/Engine/Source/Runtime/Engine/Private/WorldPartition/WorldPartitionStreamingGeneration.cpp b/Engine/Source/Runtime/Engine/Private/WorldPartition/WorldPartitionStreamingGeneration.cpp
--- a/Engine/Source/Runtime/Engine/Private/WorldPartition/WorldPartitionStreamingGeneration.cpp
+++ b/Engine/Source/Runtime/Engine/Private/WorldPartition/WorldPartitionStreamingGeneration.cpp
@@ -32,7 +32,6 @@
 		HLOD Generation
 */
 
-
 class FWorldPartitionStreamingGenerator
 {
 	void CreateActorDescViewMap(const UActorDescContainer* InContainer, TMap<FGuid, FWorldPartitionActorDescView>& OutActorDescViewMap, const FActorContainerID& InContainerID)
@@ -163,7 +162,6 @@ class FWorldPartitionStreamingGenerator
 	 */
 	void ValidateActorDescriptorViews()
 	{
-
 		// Validate data layers
 		auto IsReferenceDataLayersValid = [](const FWorldPartitionActorDescView& ActorDescView, const FWorldPartitionActorDescView& ReferenceActorDescView)
 		{
@@ -243,9 +241,9 @@ class FWorldPartitionStreamingGenerator
 
 					TArray<ActorReferenceInfo, TInlineAllocator<8>> References;
 					
-					// Build References List
+					// Build references List
 
-					// Add normal Actore references
+					// Add normal actor references
 					for (const FGuid& ReferenceGuid : ActorDescView.GetReferences())
 					{
 						if (ReferenceGuid != ActorDescView.GetParentActor()) // References to the parent are inversed in their handling 
@@ -253,7 +251,9 @@ class FWorldPartitionStreamingGenerator
 							// Filter out parent back references
 							FWorldPartitionActorDescView* ReferenceActorDesc = ContainerDescriptor.ActorDescViewMap.Find(ReferenceGuid);
 							if (ReferenceActorDesc && ReferenceActorDesc->GetParentActor() == ActorDescView.GetGuid())
+							{
 								continue;
+							}
 
 							References.Emplace(ActorReferenceInfo{ ActorDescView.GetGuid(), &ActorDescView, ReferenceGuid, ReferenceActorDesc });
 						}
@@ -340,7 +340,7 @@ class FWorldPartitionStreamingGenerator
 								}
 								NbErrorsDetected++;
 
-								if  (NbValidationPasses)
+								if (NbValidationPasses)
 								{
 									RefererActorDescView->SetInvalidRuntimeGrid();
 									ReferenceActorDescView->SetInvalidRuntimeGrid();
@@ -361,7 +361,9 @@ class FWorldPartitionStreamingGenerator
 
 				NbValidationPasses++;
 				if (NbErrorsDetected == 0)
+				{
 					break;
+				}
 			} 
 		}
 	}
@@ -599,23 +601,18 @@ void UWorldPartition::GenerateHLOD(ISourceControlHelper* SourceControlHelper, bo
 
 void UWorldPartition::CheckForErrors(IStreamingGenerationErrorHandler* ErrorHandler) const
 {
-	FActorClusterContext ActorClusterContext;
-	{		
-		FActorDescList ModifiedActorDescList;
-		FWorldPartitionStreamingGenerator StreamingGenerator(&ModifiedActorDescList, ErrorHandler);
-		StreamingGenerator.PreparationPhase(this);
-	}
+	CheckForErrors(ErrorHandler, this);
 }
 
-void UWorldPartition::CheckForErrors(IStreamingGenerationErrorHandler* ErrorHandler, const UActorDescContainer* ActorDescContainer) 
+void UWorldPartition::CheckForErrors(IStreamingGenerationErrorHandler* ErrorHandler, const UActorDescContainer* ActorDescContainer)
 {
 	FActorClusterContext ActorClusterContext;
 	{
-		FWorldPartitionStreamingGenerator StreamingGenerator(nullptr, ErrorHandler);
+		FActorDescList ModifiedActorDescList;
+		FWorldPartitionStreamingGenerator StreamingGenerator(ActorDescContainer->GetWorld() ? &ModifiedActorDescList : nullptr, ErrorHandler);
 		StreamingGenerator.PreparationPhase(ActorDescContainer);
 	}
 }
-
 #endif // WITH_EDITOR
 
 #undef LOCTEXT_NAMESPACE
