Prioritize montage over sequence in Notifies view in rewind debugger when linking to source asset.

diff --git a/Engine/Source/Runtime/Engine/Private/Animation/AnimTrace.cpp b/Engine/Source/Runtime/Engine/Private/Animation/AnimTrace.cpp
--- a/Engine/Source/Runtime/Engine/Private/Animation/AnimTrace.cpp
+++ b/Engine/Source/Runtime/Engine/Private/Animation/AnimTrace.cpp
@@ -1062,10 +1062,10 @@ void FAnimTrace::OutputAnimNotify(UAnimInstance* InAnimInstance, const FAnimNoti
 	}
 	else
 	{
-        NotifyAsset = InNotifyEvent.GetLinkedSequence();
+		NotifyAsset = InNotifyEvent.GetLinkedMontage();
 		if (NotifyAsset == nullptr)
 		{
-			NotifyAsset = InNotifyEvent.GetLinkedMontage();
+			NotifyAsset = InNotifyEvent.GetLinkedSequence();
 		}
 	}
 	
