Fix multiple issues with emulating real game networking flows in PIE
Fix multiplayer PIE so it uses the same URL rules as normal PIE, it now saves the URL as if it was a normal map load and reads that to determine net mode once actors are initialized
Add a new game instance function to properly enable listen server at runtime
Delay calling InitWorld for PIE worlds until after the net mode and world context are set, this is the path that normal world loading uses and fixes issues where the net mode is wrong during initialization callbacks

diff --git a/Engine/Source/Editor/UnrealEd/Classes/Editor/EditorEngine.h b/Engine/Source/Editor/UnrealEd/Classes/Editor/EditorEngine.h
--- a/Engine/Source/Editor/UnrealEd/Classes/Editor/EditorEngine.h
+++ b/Engine/Source/Editor/UnrealEd/Classes/Editor/EditorEngine.h
@@ -753,6 +753,7 @@ public:
 	virtual void ResetPIEAudioSetting(UWorld *CurrentPieWorld) override;
 	virtual class UGameViewportClient* GetNextPIEViewport(UGameViewportClient* CurrentViewport) override;
 	virtual UWorld* CreatePIEWorldByDuplication(FWorldContext &WorldContext, UWorld* InWorld, FString &PlayWorldMapName) override;
+	virtual void PostCreatePIEWorld(UWorld* InWorld) override;
 	virtual bool GetMapBuildCancelled() const override { return false; }
 	virtual void SetMapBuildCancelled(bool InCancelled) override { /* Intentionally empty. */ }
 	virtual void HandleNetworkFailure(UWorld *World, UNetDriver *NetDriver, ENetworkFailure::Type FailureType, const FString& ErrorString) override;
@@ -3022,9 +3023,6 @@ protected:
 	/** Gets the scene viewport for a viewport client */
 	FSceneViewport* GetGameSceneViewport(UGameViewportClient* ViewportClient) const;
 
-	/** Common init shared by CreatePIEWorldByDuplication and CreatePIEWorldBySavingToTemp */
-	void PostCreatePIEWorld(UWorld *InWorld);
-
 	/**
 	 * Toggles PIE to SIE or vice-versa
 	 */
diff --git a/Engine/Source/Editor/UnrealEd/Private/PlayLevel.cpp b/Engine/Source/Editor/UnrealEd/Private/PlayLevel.cpp
--- a/Engine/Source/Editor/UnrealEd/Private/PlayLevel.cpp
+++ b/Engine/Source/Editor/UnrealEd/Private/PlayLevel.cpp
@@ -2308,8 +2308,7 @@ UWorld* UEditorEngine::CreatePIEWorldByDuplication(FWorldContext &WorldContext,
 
 	check( NewPIEWorld );
 	NewPIEWorld->FeatureLevel = InWorld->FeatureLevel;
-	PostCreatePIEWorld(NewPIEWorld);
-
+	NewPIEWorld->WorldType = EWorldType::PIE;
 	
 	UE_LOG(LogPlayLevel, Log, TEXT("PIE: Created PIE world by copying editor world from %s to %s (%fs)"), *InWorld->GetPathName(), *NewPIEWorld->GetPathName(), float(FPlatformTime::Seconds() - StartTime));
 	return NewPIEWorld;
@@ -2319,8 +2318,12 @@ void UEditorEngine::PostCreatePIEWorld(UWorld *NewPIEWorld)
 {
 	double WorldInitStart = FPlatformTime::Seconds();
 	
+	ensure(!NewPIEWorld->bIsWorldInitialized);
+	
+	// make sure we can clean up this world!
+	NewPIEWorld->ClearFlags(RF_Standalone);
+
 	// Init the PIE world
-	NewPIEWorld->WorldType = EWorldType::PIE;
 	NewPIEWorld->InitWorld();
 	UE_LOG(LogPlayLevel, Log, TEXT("PIE: World Init took: (%fs)"),  float(FPlatformTime::Seconds() - WorldInitStart));
 
@@ -2335,8 +2338,8 @@ UWorld* UEditorEngine::CreatePIEWorldFromEntry(FWorldContext &WorldContext, UWor
 {
 	double StartTime = FPlatformTime::Seconds();
 
-	// Create the world
-	UWorld *LoadedWorld = UWorld::CreateWorld( EWorldType::PIE, false );
+	// Create the world but do not initialize yet
+	UWorld* LoadedWorld = UWorld::CreateWorld(EWorldType::PIE, false, NAME_None, nullptr, false, ERHIFeatureLevel::Num, nullptr, true);
 	check(LoadedWorld);
 	if (LoadedWorld->GetOutermost() != GetTransientPackage())
 	{
diff --git a/Engine/Source/Runtime/Engine/Classes/Engine/Engine.h b/Engine/Source/Runtime/Engine/Classes/Engine/Engine.h
--- a/Engine/Source/Runtime/Engine/Classes/Engine/Engine.h
+++ b/Engine/Source/Runtime/Engine/Classes/Engine/Engine.h
@@ -2703,6 +2703,7 @@ public:
 
 	// This should only ever be called for a EditorEngine
 	virtual UWorld* CreatePIEWorldByDuplication(FWorldContext &Context, UWorld* InWorld, FString &PlayWorldMapName) { check(false); return nullptr; }
+	virtual void PostCreatePIEWorld(UWorld* InWorld) { check(false); }
 
 	/** 
 	 *	If this function returns true, the DynamicSourceLevels collection will be duplicated for the given map.
diff --git a/Engine/Source/Runtime/Engine/Classes/Engine/GameInstance.h b/Engine/Source/Runtime/Engine/Classes/Engine/GameInstance.h
--- a/Engine/Source/Runtime/Engine/Classes/Engine/GameInstance.h
+++ b/Engine/Source/Runtime/Engine/Classes/Engine/GameInstance.h
@@ -449,6 +449,16 @@ public:
 	 */
 	virtual void AddUserToReplay(const FString& UserString);
 
+	/** 
+	 * Turns on/off listen server capability for a game instance
+	 * By default this will set up the persistent URL state so it persists across server travels and spawn the appropriate network listener
+	 *
+	 * @param bEnable turn on or off the listen server
+	 * @param PortOverride will use this specific port, or if 0 will use the URL default port
+	 * @return true if desired settings were applied, games can override to deny changes in certain states
+	 */
+	virtual bool EnableListenServer(bool bEnable, int32 PortOverride = 0);
+
 	/** handle a game specific net control message (NMT_GameSpecific)
 	 * this allows games to insert their own logic into the control channel
 	 * the meaning of both data parameters is game-specific
@@ -468,6 +478,9 @@ public:
 	/** Call to create the game mode for a given map URL */
 	virtual class AGameModeBase* CreateGameModeForURL(FURL InURL, UWorld* InWorld);
 
+	/** Call to modify the saved url that will be used as a base for future map travels */
+	virtual void SetPersistentTravelURL(FURL InURL);
+
 	/** Return the game mode subclass to use for a given map, options, and portal. By default return passed in one */
 	virtual TSubclassOf<AGameModeBase> OverrideGameModeClass(TSubclassOf<AGameModeBase> GameModeClass, const FString& MapName, const FString& Options, const FString& Portal) const;
 
diff --git a/Engine/Source/Runtime/Engine/Private/GameInstance.cpp b/Engine/Source/Runtime/Engine/Private/GameInstance.cpp
--- a/Engine/Source/Runtime/Engine/Private/GameInstance.cpp
+++ b/Engine/Source/Runtime/Engine/Private/GameInstance.cpp
@@ -303,11 +303,11 @@ FGameInstancePIEResult UGameInstance::InitializeForPlayInEditor(int32 PIEInstanc
 	NewWorld->SetGameInstance(this);
 	WorldContext->SetCurrentWorld(NewWorld);
 	WorldContext->AddRef(static_cast<UWorld*&>(EditorEngine->PlayWorld));	// Tie this context to this UEngine::PlayWorld*		// @fixme, needed still?
-
-	// make sure we can clean up this world!
-	NewWorld->ClearFlags(RF_Standalone);
 	NewWorld->bKismetScriptError = Params.bAnyBlueprintErrors;
 
+	// Initialize the world after setting world context to be consistent with normal loads
+	EditorEngine->PostCreatePIEWorld(NewWorld);
+
 	// Do a GC pass if necessary to remove any potentially unreferenced objects
 	if(bNeedsGarbageCollection)
 	{
@@ -403,6 +403,9 @@ FGameInstancePIEResult UGameInstance::StartPlayInEditorGameInstance(ULocalPlayer
 			URL = FURL(NULL, *EditorEngine->BuildPlayWorldURL(*PIEMapName, Params.bStartInSpectatorMode, ExtraURLOptions), TRAVEL_Absolute);
 		}
 
+		// Save our URL for later map travels
+		SetPersistentTravelURL(URL);
+
 		// If a start location is specified, spawn a temporary PlayerStartPIE actor at the start location and use it as the portal.
 		AActor* PlayerStart = NULL;
 		if (!EditorEngine->SpawnPlayFromHereStart(PlayWorld, PlayerStart))
@@ -452,14 +455,15 @@ FGameInstancePIEResult UGameInstance::StartPlayInEditorGameInstance(ULocalPlayer
 		if (Params.NetMode == PIE_ListenServer)
 		{
 			// Add port
+			uint32 ListenPort = 0;
 			uint16 ServerPort = 0;
 			if (Params.EditorPlaySettings->GetServerPort(ServerPort))
 			{
-				URL.Port = ServerPort;
+				ListenPort = ServerPort;
 			}
 
-			// start listen server with the built URL
-			PlayWorld->Listen(URL);
+			// Start a listen server
+			ensureMsgf(EnableListenServer(true, ListenPort), TEXT("Starting Listen Server for Play in Editor failed!"));
 		}
 
 		PlayWorld->BeginPlay();
@@ -1077,6 +1081,66 @@ void UGameInstance::AddUserToReplay(const FString& UserString)
 	}
 }
 
+void UGameInstance::SetPersistentTravelURL(FURL InURL)
+{
+	check(WorldContext);
+	WorldContext->LastURL = InURL;
+}
+
+bool UGameInstance::EnableListenServer(bool bEnable, int32 PortOverride /*= 0*/)
+{
+	UWorld* World = GetWorld();
+
+	if (!World || !World->IsGameWorld())
+	{
+		return false;
+	}
+
+	ENetMode ExistingMode = World->GetNetMode();
+
+	if (ExistingMode == NM_Client || ExistingMode == NM_DedicatedServer)
+	{
+		// Clients and dedicated servers cannot change to listen!
+		return false;
+	}
+
+	int32 DefaultListenPort = FURL::UrlConfig.DefaultPort;
+	if (bEnable)
+	{
+		// Modify the persistent url
+		if (PortOverride != 0)
+		{
+			WorldContext->LastURL.Port = PortOverride;
+		}
+		WorldContext->LastURL.AddOption(TEXT("Listen"));
+
+		if (ExistingMode == NM_Standalone)
+		{
+			// This actually opens the port
+			FURL ListenURL = WorldContext->LastURL;
+			return World->Listen(ListenURL);
+		}
+		else
+		{
+			// Already listening
+			return true;
+		}
+	}
+	else
+	{
+		WorldContext->LastURL.RemoveOption(TEXT("Listen"));
+		WorldContext->LastURL.Port = FURL::UrlConfig.DefaultPort;
+
+		if (ExistingMode == NM_ListenServer)
+		{
+			// What to do in this case is very game-specific
+			UE_LOG(LogGameSession, Warning, TEXT("Disabling a listen server with active connections does not disconnect existing players by default"));
+		}
+
+		return true;
+	}
+}
+
 void UGameInstance::ReceivedNetworkEncryptionToken(const FString& EncryptionToken, const FOnEncryptionKeyResponse& Delegate)
 {
 	FEncryptionKeyResponse Response(EEncryptionResponse::Failure, TEXT("ReceivedNetworkEncryptionToken not implemented"));
diff --git a/Engine/Source/Runtime/Engine/Private/UnrealEngine.cpp b/Engine/Source/Runtime/Engine/Private/UnrealEngine.cpp
--- a/Engine/Source/Runtime/Engine/Private/UnrealEngine.cpp
+++ b/Engine/Source/Runtime/Engine/Private/UnrealEngine.cpp
@@ -13991,22 +13991,23 @@ bool UEngine::LoadMap( FWorldContext& WorldContext, FURL URL, class UPendingNetG
 	GWorld->bCreateRenderStateForHiddenComponentsWithCollsion = bOldWorldWasShowingCollisionForHiddenComponents;
 #endif
 
-	// Fixme: hacky but we need to set PackageFlags here if we are in a PIE Context.
-	// Also, don't add to root when in PIE, since PIE doesn't remove world from root
+	// PIE worlds are not added to root and are initialized differently
 	if (WorldContext.WorldType == EWorldType::PIE)
 	{
 		check(WorldContext.World()->GetOutermost()->HasAnyPackageFlags(PKG_PlayInEditor));
 		WorldContext.World()->ClearFlags(RF_Standalone);
+
+		PostCreatePIEWorld(WorldContext.World());
 	}
 	else
 	{
 		WorldContext.World()->AddToRoot();
-	}
 
-	// In the PIE case the world will already have been initialized as part of CreatePIEWorldByDuplication
-	if (!WorldContext.World()->bIsWorldInitialized)
-	{
-		WorldContext.World()->InitWorld();
+		// The world should not have been initialized before this
+		if (ensure(!WorldContext.World()->bIsWorldInitialized))
+		{
+			WorldContext.World()->InitWorld();
+		}
 	}
 
 	// Handle pending level.
diff --git a/Engine/Source/Runtime/Engine/Private/World.cpp b/Engine/Source/Runtime/Engine/Private/World.cpp
--- a/Engine/Source/Runtime/Engine/Private/World.cpp
+++ b/Engine/Source/Runtime/Engine/Private/World.cpp
@@ -1836,7 +1836,7 @@ void UWorld::InitWorld(const InitializationValues IVS)
 	bIsWorldInitialized = true;
 	bInitializedAndNeedsCleanup = true;
 
-	//@TODO: Should this happen here, or below here?
+	// Call the general post initialization delegates
 	FWorldDelegates::OnPostWorldInitialization.Broadcast(this, IVS);
 
 	PersistentLevel->PrecomputedVisibilityHandler.UpdateScene(Scene);
@@ -7721,16 +7721,17 @@ ENetMode UWorld::InternalGetNetMode() const
 	}
 	PRAGMA_ENABLE_DEPRECATION_WARNINGS
 
+	ENetMode URLNetMode = AttemptDeriveFromURL();
 #if WITH_EDITOR
-	if (WorldType == EWorldType::PIE)
+	if (WorldType == EWorldType::PIE && URLNetMode == NM_Standalone && !AreActorsInitialized())
 	{
-		// Return the cached NetMode that we were started with. We can't derive this from the Play Settings,
-		// because those can be modified during PIE startup, plus it is not easy for us to tell if we are
-		// a listen server or a client.
+		// If we're too early in startup and it defaults to Standalone, use the mode we were created with
+		// Once the world has been initialized the URL will be correct so we will use that
+		// This is required for dedicated server worlds so it is correct for InitWorld
 		return PlayInEditorNetMode;
 	}
 #endif
-	return AttemptDeriveFromURL();
+	return URLNetMode;
 }
 
 PRAGMA_DISABLE_DEPRECATION_WARNINGS
