Fix for tooltip  on source buses

diff --git a/Engine/Source/Runtime/Engine/Classes/Sound/SoundSourceBus.h b/Engine/Source/Runtime/Engine/Classes/Sound/SoundSourceBus.h
--- a/Engine/Source/Runtime/Engine/Classes/Sound/SoundSourceBus.h
+++ b/Engine/Source/Runtime/Engine/Classes/Sound/SoundSourceBus.h
@@ -37,7 +37,7 @@ public:
 	UPROPERTY(EditAnywhere, Category = BusProperties, meta = (UIMin = 0.0, ClampMin = 0.0))
 	float SourceBusDuration;
 
-	/** Audio bus to use to use as audio for this source bus. This source bus will sonify the audio from the audio bus. */
+	/** Audio bus to use as audio for this source bus. This source bus will sonify the audio from the audio bus. */
 	UPROPERTY(EditAnywhere, Category = BusProperties)
 	TObjectPtr<UAudioBus> AudioBus;
 
