Fixing tooltips on sound submixes

diff --git a/Engine/Source/Runtime/Engine/Classes/Sound/SoundSubmix.h b/Engine/Source/Runtime/Engine/Classes/Sound/SoundSubmix.h
--- a/Engine/Source/Runtime/Engine/Classes/Sound/SoundSubmix.h
+++ b/Engine/Source/Runtime/Engine/Classes/Sound/SoundSubmix.h
@@ -483,8 +483,6 @@ public:
 	UPROPERTY()
 	TSubclassOf<UAudioEndpointSettingsBase> EndpointSettingsClass;
 
-	//TODO: Make this editable only if EndpointType is non-default,
-	// and filter classes based on ISoundfieldFactory::GetCustomSettingsClass().
 	UPROPERTY(EditAnywhere, Category = Endpoint)
 	TObjectPtr<UAudioEndpointSettingsBase> EndpointSettings;
 
@@ -519,8 +517,6 @@ public:
 	// This function goes through every child submix and the parent submix to ensure that they have a compatible format.
 	void SanitizeLinks();
 
-	//TODO: Make this editable only if EndpointType is non-default,
-	// and filter classes based on ISoundfieldFactory::GetCustomSettingsClass().
 	UPROPERTY(EditAnywhere, Category = Endpoint)
 	TObjectPtr<USoundfieldEndpointSettingsBase> EndpointSettings;
 
