Fix GetAllParametersOfType(), so that less-derived instances don't override parameter values in more-derived instances.

diff --git a/Engine/Source/Runtime/Engine/Private/Materials/MaterialInstance.cpp b/Engine/Source/Runtime/Engine/Private/Materials/MaterialInstance.cpp
--- a/Engine/Source/Runtime/Engine/Private/Materials/MaterialInstance.cpp
+++ b/Engine/Source/Runtime/Engine/Private/Materials/MaterialInstance.cpp
@@ -1735,6 +1735,10 @@ void UMaterialInstance::GetAllParametersOfType(EMaterialParameterType Type, TMap
 		LayerIndexRemap.Add(LayerIndex);
 	}
 
+	// We walk the inheritance hierarchy backwards to the root, so we keep track of overrides that are set, to avoid setting them again from less-derived instances
+	// Alternately could walk the hierarchy starting from the root, but then we'd need an alternate way to track layer index remapping
+	TSet<FMaterialParameterInfo, DefaultKeyFuncs<FMaterialParameterInfo>, TInlineSetAllocator<32>> OverridenParameters;
+
 	for (int32 Index = 0; Index < InstanceChain.MaterialInstances.Num(); ++Index)
 	{
 		const UMaterialInstance* Instance = InstanceChain.MaterialInstances[Index];
@@ -1742,15 +1746,15 @@ void UMaterialInstance::GetAllParametersOfType(EMaterialParameterType Type, TMap
 		const bool bSetOverride = (Index == 0); // Only set the override flag for parameters overriden by the current material (always at slot0)
 		switch (Type)
 		{
-		case EMaterialParameterType::Scalar: GameThread_ApplyParameterOverrides(Instance->ScalarParameterValues, LayerIndexRemap, bSetOverride, OutParameters); break;
-		case EMaterialParameterType::Vector: GameThread_ApplyParameterOverrides(Instance->VectorParameterValues, LayerIndexRemap, bSetOverride, OutParameters); break;
-		case EMaterialParameterType::DoubleVector: GameThread_ApplyParameterOverrides(Instance->DoubleVectorParameterValues, LayerIndexRemap, bSetOverride, OutParameters); break;
-		case EMaterialParameterType::Texture: GameThread_ApplyParameterOverrides(Instance->TextureParameterValues, LayerIndexRemap, bSetOverride, OutParameters); break;
-		case EMaterialParameterType::RuntimeVirtualTexture: GameThread_ApplyParameterOverrides(Instance->RuntimeVirtualTextureParameterValues, LayerIndexRemap, bSetOverride, OutParameters); break;
-		case EMaterialParameterType::Font: GameThread_ApplyParameterOverrides(Instance->FontParameterValues, LayerIndexRemap, bSetOverride, OutParameters); break;
+		case EMaterialParameterType::Scalar: GameThread_ApplyParameterOverrides(Instance->ScalarParameterValues, LayerIndexRemap, bSetOverride, OverridenParameters, OutParameters); break;
+		case EMaterialParameterType::Vector: GameThread_ApplyParameterOverrides(Instance->VectorParameterValues, LayerIndexRemap, bSetOverride, OverridenParameters, OutParameters); break;
+		case EMaterialParameterType::DoubleVector: GameThread_ApplyParameterOverrides(Instance->DoubleVectorParameterValues, LayerIndexRemap, bSetOverride, OverridenParameters, OutParameters); break;
+		case EMaterialParameterType::Texture: GameThread_ApplyParameterOverrides(Instance->TextureParameterValues, LayerIndexRemap, bSetOverride, OverridenParameters, OutParameters); break;
+		case EMaterialParameterType::RuntimeVirtualTexture: GameThread_ApplyParameterOverrides(Instance->RuntimeVirtualTextureParameterValues, LayerIndexRemap, bSetOverride, OverridenParameters, OutParameters); break;
+		case EMaterialParameterType::Font: GameThread_ApplyParameterOverrides(Instance->FontParameterValues, LayerIndexRemap, bSetOverride, OverridenParameters, OutParameters); break;
 #if WITH_EDITORONLY_DATA
-		case EMaterialParameterType::StaticSwitch: GameThread_ApplyParameterOverrides(Instance->StaticParameters.StaticSwitchParameters, LayerIndexRemap, bSetOverride, OutParameters); break;
-		case EMaterialParameterType::StaticComponentMask: GameThread_ApplyParameterOverrides(Instance->StaticParameters.StaticComponentMaskParameters, LayerIndexRemap, bSetOverride, OutParameters); break;
+		case EMaterialParameterType::StaticSwitch: GameThread_ApplyParameterOverrides(Instance->StaticParameters.StaticSwitchParameters, LayerIndexRemap, bSetOverride, OverridenParameters, OutParameters); break;
+		case EMaterialParameterType::StaticComponentMask: GameThread_ApplyParameterOverrides(Instance->StaticParameters.StaticComponentMaskParameters, LayerIndexRemap, bSetOverride, OverridenParameters, OutParameters); break;
 #endif // WITH_EDITORONLY_DATA
 		default: checkNoEntry();
 		}
diff --git a/Engine/Source/Runtime/Engine/Private/Materials/MaterialInstanceSupport.h b/Engine/Source/Runtime/Engine/Private/Materials/MaterialInstanceSupport.h
--- a/Engine/Source/Runtime/Engine/Private/Materials/MaterialInstanceSupport.h
+++ b/Engine/Source/Runtime/Engine/Private/Materials/MaterialInstanceSupport.h
@@ -332,8 +332,12 @@ inline bool GameThread_GetParameterValue(const TArray<ParameterType>& Parameters
 	return false;
 }
 
-template <typename ParameterType>
-inline void GameThread_ApplyParameterOverrides(const TArray<ParameterType>& Parameters, TArrayView<const int32> LayerIndexRemap, bool bSetOverride, TMap<FMaterialParameterInfo, FMaterialParameterMetadata>& OutParameters)
+template <typename ParameterType, typename OverridenParametersType>
+inline void GameThread_ApplyParameterOverrides(const TArray<ParameterType>& Parameters,
+	TArrayView<const int32> LayerIndexRemap,
+	bool bSetOverride,
+	OverridenParametersType& OverridenParameters, // TSet<FMaterialParameterInfo, ...>
+	TMap<FMaterialParameterInfo, FMaterialParameterMetadata>& OutParameters)
 {
 	for (int32 ParameterIndex = 0; ParameterIndex < Parameters.Num(); ParameterIndex++)
 	{
@@ -343,16 +347,21 @@ inline void GameThread_ApplyParameterOverrides(const TArray<ParameterType>& Para
 			FMaterialParameterInfo ParameterInfo;
 			if (Parameter->ParameterInfo.RemapLayerIndex(LayerIndexRemap, ParameterInfo))
 			{
-				FMaterialParameterMetadata* Result = OutParameters.Find(ParameterInfo);
-				if (Result)
+				bool bPreviouslyOverriden = false;
+				OverridenParameters.Add(ParameterInfo, &bPreviouslyOverriden);
+				if (!bPreviouslyOverriden)
 				{
-					Parameter->GetValue(*Result);
-#if WITH_EDITORONLY_DATA
-					if (bSetOverride)
+					FMaterialParameterMetadata* Result = OutParameters.Find(ParameterInfo);
+					if (Result)
 					{
-						Result->bOverride = true;
-					}
+						Parameter->GetValue(*Result);
+#if WITH_EDITORONLY_DATA
+						if (bSetOverride)
+						{
+							Result->bOverride = true;
+						}
 #endif // WITH_EDITORONLY_DATA
+					}
 				}
 			}
 		}
