Remove duplicate transforms in FAnimInstanceProxy

We were holding two different copies of the component & actor transforms, now we just hold the one
Minor: also removed some pointer chasing in PreUpdate

diff --git a/Engine/Source/Runtime/AnimGraphRuntime/Private/AnimNodes/AnimNode_AimOffsetLookAt.cpp b/Engine/Source/Runtime/AnimGraphRuntime/Private/AnimNodes/AnimNode_AimOffsetLookAt.cpp
--- a/Engine/Source/Runtime/AnimGraphRuntime/Private/AnimNodes/AnimNode_AimOffsetLookAt.cpp
+++ b/Engine/Source/Runtime/AnimGraphRuntime/Private/AnimNodes/AnimNode_AimOffsetLookAt.cpp
@@ -125,8 +125,8 @@ void FAnimNode_AimOffsetLookAt::UpdateFromLookAtTarget(FPoseContext& LocalPoseCo
 
 		FAnimInstanceProxy* AnimProxy = LocalPoseContext.AnimInstanceProxy;
 		check(AnimProxy);
-		const FTransform SourceWorldTransform = SourceComponentTransform * AnimProxy->GetSkelMeshCompLocalToWorld();
-		const FTransform ActorTransform = AnimProxy->GetSkelMeshCompOwnerTransform();
+		const FTransform SourceWorldTransform = SourceComponentTransform * AnimProxy->GetComponentTransform();
+		const FTransform ActorTransform = AnimProxy->GetActorTransform();
 
 		// Convert Target to Actor Space
 		const FTransform TargetWorldTransform(LookAtLocation);
diff --git a/Engine/Source/Runtime/AnimGraphRuntime/Private/BoneControllers/AnimNode_LookAt.cpp b/Engine/Source/Runtime/AnimGraphRuntime/Private/BoneControllers/AnimNode_LookAt.cpp
--- a/Engine/Source/Runtime/AnimGraphRuntime/Private/BoneControllers/AnimNode_LookAt.cpp
+++ b/Engine/Source/Runtime/AnimGraphRuntime/Private/BoneControllers/AnimNode_LookAt.cpp
@@ -160,7 +160,7 @@ void FAnimNode_LookAt::EvaluateComponentSpaceInternal(FComponentSpacePoseContext
 #if !(UE_BUILD_SHIPPING || UE_BUILD_TEST)
 	/*if (bEnableDebug)
 	{
-		const FTransform LocalToWorld = Context.AnimInstanceProxy->GetSkelMeshCompLocalToWorld();
+		const FTransform LocalToWorld = Context.AnimInstanceProxy->GetComponentTransform();
 		FVector TargetWorldLoc = LocalToWorld.TransformPosition(CachedCurrentTargetLocation);
 		FVector SourceWorldLoc = LocalToWorld.TransformPosition(CachedComponentBoneLocation);
 
diff --git a/Engine/Source/Runtime/AnimGraphRuntime/Private/BoneControllers/AnimNode_ScaleChainLength.cpp b/Engine/Source/Runtime/AnimGraphRuntime/Private/BoneControllers/AnimNode_ScaleChainLength.cpp
--- a/Engine/Source/Runtime/AnimGraphRuntime/Private/BoneControllers/AnimNode_ScaleChainLength.cpp
+++ b/Engine/Source/Runtime/AnimGraphRuntime/Private/BoneControllers/AnimNode_ScaleChainLength.cpp
@@ -94,7 +94,7 @@ void FAnimNode_ScaleChainLength::Evaluate_AnyThread(FPoseContext& Output)
 		return;
 	}
 
-	const FVector TargetLocationCompSpace = Output.AnimInstanceProxy->GetSkelMeshCompLocalToWorld().InverseTransformPosition(TargetLocation);
+	const FVector TargetLocationCompSpace = Output.AnimInstanceProxy->GetComponentTransform().InverseTransformPosition(TargetLocation);
 
 	// Allocate transforms to get component space transform of chain start bone.
 	FCSPose<FCompactPose> CSPose;
diff --git a/Engine/Source/Runtime/Engine/Private/Animation/AnimInstanceProxy.cpp b/Engine/Source/Runtime/Engine/Private/Animation/AnimInstanceProxy.cpp
--- a/Engine/Source/Runtime/Engine/Private/Animation/AnimInstanceProxy.cpp
+++ b/Engine/Source/Runtime/Engine/Private/Animation/AnimInstanceProxy.cpp
@@ -366,8 +366,9 @@ void FAnimInstanceProxy::PreUpdate(UAnimInstance* InAnimInstance, float DeltaSec
 {
 	DECLARE_SCOPE_HIERARCHICAL_COUNTER_FUNC()
 
-	USkeletalMeshComponent* SkelMeshComp = InAnimInstance->GetSkelMeshComponent();
-	UWorld* World = SkelMeshComp ? SkelMeshComp->GetWorld() : nullptr;
+	InitializeObjects(InAnimInstance);
+
+	UWorld* World = SkeletalMeshComponent ? SkeletalMeshComponent->GetWorld() : nullptr;
 	AWorldSettings* WorldSettings = World ? World->GetWorldSettings() : nullptr;
 
 	CurrentDeltaSeconds = DeltaSeconds;
@@ -375,8 +376,6 @@ void FAnimInstanceProxy::PreUpdate(UAnimInstance* InAnimInstance, float DeltaSec
 	RootMotionMode = InAnimInstance->RootMotionMode;
 	bShouldExtractRootMotion = InAnimInstance->ShouldExtractRootMotion();
 
-	InitializeObjects(InAnimInstance);
-
 #if WITH_EDITORONLY_DATA
 	if (FAnimBlueprintDebugData* DebugData = GetAnimBlueprintDebugData())
 	{
@@ -384,7 +383,7 @@ void FAnimInstanceProxy::PreUpdate(UAnimInstance* InAnimInstance, float DeltaSec
 	}
 #endif
 
-	if (SkelMeshComp)
+	if (SkeletalMeshComponent)
 	{
 		// Save off LOD level that we're currently using.
 		const int32 PreviousLODLevel = LODLevel;
@@ -393,16 +392,9 @@ void FAnimInstanceProxy::PreUpdate(UAnimInstance* InAnimInstance, float DeltaSec
 		{
 			OnPreUpdateLODChanged(PreviousLODLevel, LODLevel);
 		}
-
-		// Cache these transforms, so nodes don't have to pull it off the gamethread manually.
-		SkelMeshCompLocalToWorld = SkelMeshComp->GetComponentTransform();
-		if (const AActor* Owner = SkelMeshComp->GetOwner())
-		{
-			SkelMeshCompOwnerTransform = Owner->GetTransform();
-		}
 	}
 
-	NotifyQueue.Reset(InAnimInstance->GetSkelMeshComponent());
+	NotifyQueue.Reset(SkeletalMeshComponent);
 
 #if ENABLE_ANIM_DRAW_DEBUG
 	QueuedDrawDebugItems.Reset();
diff --git a/Engine/Source/Runtime/Engine/Public/Animation/AnimInstanceProxy.h b/Engine/Source/Runtime/Engine/Public/Animation/AnimInstanceProxy.h
--- a/Engine/Source/Runtime/Engine/Public/Animation/AnimInstanceProxy.h
+++ b/Engine/Source/Runtime/Engine/Public/Animation/AnimInstanceProxy.h
@@ -374,16 +374,16 @@ public:
 		return LODLevel;
 	}
 
-	// Cached SkeletalMeshComponent LocalToWorld Transform.
+	UE_DEPRECATED(5.0, "Please use GetComponentTransform")
 	const FTransform& GetSkelMeshCompLocalToWorld() const
 	{
-		return SkelMeshCompLocalToWorld;
+		return ComponentTransform;
 	}
 
-	// Cached SkeletalMeshComponent Owner Transform.
+	UE_DEPRECATED(5.0, "Please use GetActorTransform")
 	const FTransform& GetSkelMeshCompOwnerTransform() const
 	{
-		return SkelMeshCompOwnerTransform;
+		return ActorTransform;
 	}
 
 	/** Get the current skeleton we are using. Note that this will return nullptr outside of pre/post update */
@@ -1089,12 +1089,6 @@ private:
 	/** Counter used to control CacheBones recursion behavior - makes sure we cache bones correctly when recursing into different subgraphs */
 	int32 CacheBonesRecursionCounter;
 
-	/** Cached SkeletalMeshComponent LocalToWorld transform. */
-	FTransform SkelMeshCompLocalToWorld;
-
-	/** Cached SkeletalMeshComponent Owner Transform */
-	FTransform SkelMeshCompOwnerTransform;
-
 	/** During animation update and eval, records the number of frames we will skip due to URO */
 	int16 NumUroSkippedFrames_Update;
 	int16 NumUroSkippedFrames_Eval;
