Support generating analytic derivatives with new HLSL translator

diff --git a/Engine/Source/Runtime/Engine/Private/HLSLTree/HLSLTree.cpp b/Engine/Source/Runtime/Engine/Private/HLSLTree/HLSLTree.cpp
--- a/Engine/Source/Runtime/Engine/Private/HLSLTree/HLSLTree.cpp
+++ b/Engine/Source/Runtime/Engine/Private/HLSLTree/HLSLTree.cpp
@@ -5,6 +5,7 @@
 #include "Misc/MemStackUtility.h"
 #include "Shader/ShaderTypes.h"
 #include "MaterialShared.h" // TODO - split preshader out into its own module
+#include "Misc/LargeWorldRenderPosition.h"
 
 static FSHAHash HashString(const TCHAR* String, int32 Length)
 {
@@ -18,34 +19,60 @@ static FSHAHash HashString(const FStringBuilderBase& StringBuilder)
 	return HashString(StringBuilder.GetData(), StringBuilder.Len());
 }
 
+static FSHAHash HashString(FStringView String)
+{
+	return HashString(String.GetData(), String.Len());
+}
+
 namespace UE
 {
 namespace HLSLTree
 {
 
-EExpressionEvaluationType CombineEvaluationTypes(EExpressionEvaluationType Lhs, EExpressionEvaluationType Rhs)
+EExpressionEvaluation CombineEvaluations(EExpressionEvaluation Lhs, EExpressionEvaluation Rhs)
 {
-	if (Lhs == EExpressionEvaluationType::None)
+	if (Lhs == EExpressionEvaluation::None)
 	{
 		// If either is 'None', return the other
 		return Rhs;
 	}
-	else if (Rhs == EExpressionEvaluationType::None)
+	else if (Rhs == EExpressionEvaluation::None)
 	{
 		return Lhs;
 	}
-	else if (Lhs == EExpressionEvaluationType::Constant && Rhs == EExpressionEvaluationType::Constant)
+	else if (Lhs == EExpressionEvaluation::Constant && Rhs == EExpressionEvaluation::Constant)
 	{
 		// 2 constants make a constant
-		return EExpressionEvaluationType::Constant;
+		return EExpressionEvaluation::Constant;
 	}
-	else if (Lhs == EExpressionEvaluationType::Shader || Rhs == EExpressionEvaluationType::Shader)
+	else if (Lhs == EExpressionEvaluation::Shader || Rhs == EExpressionEvaluation::Shader)
 	{
 		// If either requires shader, shader is required
-		return EExpressionEvaluationType::Shader;
+		return EExpressionEvaluation::Shader;
 	}
 	// Any combination of constants/preshader can make a preshader
-	return EExpressionEvaluationType::Preshader;
+	return EExpressionEvaluation::Preshader;
+}
+
+EExpressionDerivative CombineDerivatives(EExpressionDerivative Lhs, EExpressionDerivative Rhs)
+{
+	if (Lhs == EExpressionDerivative::None)
+	{
+		return Rhs;
+	}
+	else if (Rhs == EExpressionDerivative::None)
+	{
+		return Lhs;
+	}
+	else if (Lhs == EExpressionDerivative::Invalid || Rhs == EExpressionDerivative::Invalid)
+	{
+		return EExpressionDerivative::Invalid;
+	}
+	else if (Lhs == EExpressionDerivative::Zero && Rhs == EExpressionDerivative::Zero)
+	{
+		return EExpressionDerivative::Zero;
+	}
+	return EExpressionDerivative::Valid;
 }
 
 FErrors::FErrors(FMemStackBase& InAllocator)
@@ -96,9 +123,7 @@ FEmitShaderValue* FEmitContext::CastShaderValue(FNode* Node, FScope* Scope, FEmi
 	const Shader::FValueTypeDescription DestTypeDesc = Shader::GetValueTypeDescription(DestType);
 
 	TStringBuilder<1024> FormattedCode;
-	FShaderValue Shader(FormattedCode);
-	Shader.Type = DestType;
-	Shader.bInline = true;
+	Shader::FType IntermediateType = DestType;
 
 	if (SourceTypeDesc.NumComponents > 0 && DestTypeDesc.NumComponents > 0)
 	{
@@ -111,13 +136,13 @@ FEmitShaderValue* FEmitContext::CastShaderValue(FNode* Node, FScope* Scope, FEmi
 			{
 				// float->LWC
 				ShaderValue = CastShaderValue(Node, Scope, ShaderValue, Shader::MakeValueType(Shader::EValueComponentType::Float, DestTypeDesc.NumComponents));
-				Shader.Code.Appendf(TEXT("LWCPromote(%s)"), ShaderValue->Reference);
+				FormattedCode.Appendf(TEXT("LWCPromote(%s)"), ShaderValue->Reference);
 			}
 			else
 			{
 				//LWC->float
-				Shader.Code.Appendf(TEXT("LWCToFloat(%s)"), ShaderValue->Reference);
-				Shader.Type = Shader::MakeValueType(Shader::EValueComponentType::Float, SourceTypeDesc.NumComponents);
+				FormattedCode.Appendf(TEXT("LWCToFloat(%s)"), ShaderValue->Reference);
+				IntermediateType = Shader::MakeValueType(Shader::EValueComponentType::Float, SourceTypeDesc.NumComponents);
 			}
 		}
 		else
@@ -128,7 +153,7 @@ FEmitShaderValue* FEmitContext::CastShaderValue(FNode* Node, FScope* Scope, FEmi
 			bool bNeedClosingParen = false;
 			if (bIsLWC)
 			{
-				Shader.Code.Append(TEXT("MakeLWCVector("));
+				FormattedCode.Append(TEXT("MakeLWCVector("));
 				bNeedClosingParen = true;
 			}
 			else
@@ -137,27 +162,26 @@ FEmitShaderValue* FEmitContext::CastShaderValue(FNode* Node, FScope* Scope, FEmi
 				{
 					NumComponents = DestTypeDesc.NumComponents;
 					// Cast the scalar to the correct type, HLSL language will replicate the scalar if needed when performing this cast
-					Shader.Code.Appendf(TEXT("((%s)%s)"), DestTypeDesc.Name, ShaderValue->Reference);
+					FormattedCode.Appendf(TEXT("((%s)%s)"), DestTypeDesc.Name, ShaderValue->Reference);
 				}
 				else
 				{
 					NumComponents = FMath::Min(SourceTypeDesc.NumComponents, DestTypeDesc.NumComponents);
 					if (NumComponents < DestTypeDesc.NumComponents)
 					{
-						Shader.Code.Appendf(TEXT("%s("), DestTypeDesc.Name);
+						FormattedCode.Appendf(TEXT("%s("), DestTypeDesc.Name);
 						bNeedClosingParen = true;
 					}
 					if (NumComponents == SourceTypeDesc.NumComponents && SourceTypeDesc.ComponentType == DestTypeDesc.ComponentType)
 					{
 						// If we're taking all the components from the source, can avoid adding a swizzle
-						Shader.Code.Append(ShaderValue->Reference);
+						FormattedCode.Append(ShaderValue->Reference);
 					}
 					else
 					{
 						// Use a cast to truncate the source to the correct number of types
-						const Shader::EValueType IntermediateType = Shader::MakeValueType(DestTypeDesc.ComponentType, NumComponents);
-						const Shader::FValueTypeDescription IntermediateTypeDesc = Shader::GetValueTypeDescription(IntermediateType);
-						Shader.Code.Appendf(TEXT("((%s)%s)"), IntermediateTypeDesc.Name, ShaderValue->Reference);
+						const Shader::EValueType LocalType = Shader::MakeValueType(DestTypeDesc.ComponentType, NumComponents);
+						FormattedCode.Appendf(TEXT("((%s)%s)"), Shader::GetValueTypeDescription(LocalType).Name, ShaderValue->Reference);
 					}
 				}
 			}
@@ -169,17 +193,17 @@ FEmitShaderValue* FEmitContext::CastShaderValue(FNode* Node, FScope* Scope, FEmi
 				{
 					if (ComponentIndex > 0u)
 					{
-						Shader.Code.Append(TEXT(","));
+						FormattedCode.Append(TEXT(","));
 					}
 					if (bIsLWC)
 					{
 						if (!bReplicateScalar && ComponentIndex >= SourceTypeDesc.NumComponents)
 						{
-							Shader.Code.Append(TEXT("LWCPromote(0.0f)"));
+							FormattedCode.Append(TEXT("LWCPromote(0.0f)"));
 						}
 						else
 						{
-							Shader.Code.Appendf(TEXT("LWCGetComponent(%s, %d)"), ShaderValue->Reference, bReplicateScalar ? 0 : ComponentIndex);
+							FormattedCode.Appendf(TEXT("LWCGetComponent(%s, %d)"), ShaderValue->Reference, bReplicateScalar ? 0 : ComponentIndex);
 						}
 					}
 					else
@@ -187,11 +211,11 @@ FEmitShaderValue* FEmitContext::CastShaderValue(FNode* Node, FScope* Scope, FEmi
 						// Non-LWC case should only be zero-filling here, other cases should have already been handled
 						check(!bReplicateScalar);
 						check(ComponentIndex >= SourceTypeDesc.NumComponents);
-						ZeroValue.ToString(Shader::EValueStringFormat::HLSL, Shader.Code);
+						ZeroValue.ToString(Shader::EValueStringFormat::HLSL, FormattedCode);
 					}
 				}
 				NumComponents = DestTypeDesc.NumComponents;
-				Shader.Code.Append(TEXT(")"));
+				FormattedCode.Append(TEXT(")"));
 			}
 
 			check(NumComponents == DestTypeDesc.NumComponents);
@@ -200,12 +224,12 @@ FEmitShaderValue* FEmitContext::CastShaderValue(FNode* Node, FScope* Scope, FEmi
 	else
 	{
 		Errors.AddErrorf(Node, TEXT("Cannot cast between non-numeric types %s to %s."), SourceTypeDesc.Name, DestTypeDesc.Name);
-		Shader.Code.Appendf(TEXT("((%s)0)"), DestType.GetName());
-		Shader.bInline = false;
+		FormattedCode.Appendf(TEXT("((%s)0)"), DestType.GetName());
 	}
 
-	check(Shader.Type != ShaderValue->Type);
-	ShaderValue = AcquireShader(Scope, Shader, MakeArrayView(&ShaderValue, 1));
+	check(IntermediateType != ShaderValue->Type);
+	const bool bInline = true;
+	ShaderValue = AcquireShader(Scope, IntermediateType, FormattedCode.ToView(), bInline, MakeArrayView(&ShaderValue, 1));
 	if (ShaderValue->Type != DestType)
 	{
 		// May need to cast through multiple intermediate types to reach our destination type
@@ -232,32 +256,32 @@ void MoveToScope(FEmitShaderValue* ShaderValue, FScope* Scope)
 }
 }
 
-FEmitShaderValue* FEmitContext::AcquireShader(FScope* Scope, const FShaderValue& Shader, TArrayView<FEmitShaderValue*> Dependencies)
+FEmitShaderValue* FEmitContext::AcquireShader(FScope* Scope, const Shader::FType& Type, FStringView Code, bool bInline, TArrayView<FEmitShaderValue*> Dependencies)
 {
 	FEmitShaderValue* ShaderValue = nullptr;
 
 	// Check to see if we've already generated code for an equivalent expression
-	const FSHAHash ShaderHash = HashString(Shader.Code);
+	const FSHAHash ShaderHash = HashString(Code);
 	FEmitShaderValue** const PrevShaderValue = ShaderValueMap.Find(ShaderHash);
 	if (PrevShaderValue)
 	{
 		ShaderValue = *PrevShaderValue;
-		check(ShaderValue->Type == Shader.Type);
+		check(ShaderValue->Type == Type);
 		Private::MoveToScope(ShaderValue, Scope);
 	}
 	else
 	{
-		ShaderValue = new(*Allocator) FEmitShaderValue(Scope, Shader.Type);
+		ShaderValue = new(*Allocator) FEmitShaderValue(Scope, Type);
 		ShaderValue->Hash = ShaderHash;
 		ShaderValue->Dependencies = MemStack::AllocateArrayView(*Allocator, Dependencies);
-		if (Shader.bInline)
+		if (bInline)
 		{
-			ShaderValue->Reference = MemStack::AllocateString(*Allocator, Shader.Code);
+			ShaderValue->Reference = MemStack::AllocateString(*Allocator, Code);
 		}
 		else
 		{
 			ShaderValue->Reference = AcquireLocalDeclarationCode();
-			ShaderValue->Value = MemStack::AllocateString(*Allocator, Shader.Code);
+			ShaderValue->Value = MemStack::AllocateString(*Allocator, Code);
 			ShaderValueMap.Add(ShaderHash, ShaderValue);
 		}
 	}
@@ -369,9 +393,9 @@ FEmitShaderValue* FEmitContext::AcquirePreshaderOrConstant(const FRequestedType&
 		const Shader::EValueType FieldType = Type.GetFlatFieldType(FieldIndex);
 		const Shader::FValueTypeDescription TypeDesc = Shader::GetValueTypeDescription(FieldType);
 		const int32 NumFieldComponents = TypeDesc.NumComponents;
-		const EExpressionEvaluationType FieldEvaluationType = Expression->GetPreparedType().GetFieldEvaluationType(ComponentIndex, NumFieldComponents);
+		const EExpressionEvaluation FieldEvaluation = Expression->GetPreparedType().GetFieldData(ComponentIndex, NumFieldComponents).Evaluation;
 
-		if (FieldEvaluationType == EExpressionEvaluationType::Preshader)
+		if (FieldEvaluation == EExpressionEvaluation::Preshader)
 		{
 			// Only need to allocate uniform buffer for non-constant components
 			// Constant components can have their value inlined into the shader directly
@@ -454,15 +478,16 @@ FEmitShaderValue* FEmitContext::AcquirePreshaderOrConstant(const FRequestedType&
 		}
 		else
 		{
-			// We allow FieldEvaluationType to be 'None', since in that case we still need to fill in a value for the HLSL initializer
-			check(FieldEvaluationType == EExpressionEvaluationType::Constant || FieldEvaluationType == EExpressionEvaluationType::None);
+			// We allow FieldEvaluation to be 'None', since in that case we still need to fill in a value for the HLSL initializer
+			check(FieldEvaluation == EExpressionEvaluation::Constant || FieldEvaluation == EExpressionEvaluation::None);
 
 			// The type generated by the preshader might not match the expected type
 			// In the future, with new HLSLTree, preshader could potentially include explicit cast opcodes, and avoid implicit conversions
 			Shader::FValue FieldConstantValue(ConstantValue.Type.GetComponentType(ComponentIndex), NumFieldComponents);
 			for (int32 i = 0; i < NumFieldComponents; ++i)
 			{
-				FieldConstantValue.Component[i] = ConstantValue.Component[ComponentIndex + i];
+				// Allow replicating scalar values
+				FieldConstantValue.Component[i] = ConstantValue.Component.Num() == 1 ? ConstantValue.Component[0] : ConstantValue.Component[ComponentIndex + i];
 			}
 
 			if (TypeDesc.ComponentType == Shader::EValueComponentType::Double)
@@ -522,15 +547,20 @@ FEmitShaderValue* FEmitContext::AcquirePreshaderOrConstant(const FRequestedType&
 		FormattedCode.Append(TEXT(" }"));
 	}
 
-	FShaderValue Shader(FormattedCode);
-	Shader.Type = Type;
-	Shader.bInline = !Type.IsStruct(); // struct declarations can't be inline, due to HLSL syntax
-	FEmitShaderValue* ShaderValue = AcquireShader(Scope, Shader, TArrayView<FEmitShaderValue*>());
+	const bool bInline = !Type.IsStruct(); // struct declarations can't be inline, due to HLSL syntax
+	FEmitShaderValue* ShaderValue = AcquireShader(Scope, Type, FormattedCode.ToView(), bInline, TArrayView<FEmitShaderValue*>());
 	PreshaderValueMap.Add(Hash, ShaderValue);
 
 	return ShaderValue;
 }
 
+FEmitShaderValue* FEmitContext::AcquireConstantZero(FScope* Scope, const Shader::FType& Type)
+{
+	TStringBuilder<256> FormattedCode;
+	FormattedCode.Appendf(TEXT("((%s)0)"), Type.GetName());
+	return AcquireShader(Scope, Type, FormattedCode.ToView(), true, TArrayView<FEmitShaderValue*>());
+}
+
 namespace Private
 {
 void EmitShaderValue(FEmitContext& Context, FEmitShaderValue* ShaderValue)
@@ -683,7 +713,7 @@ void FExpressionLocalPHI::PrepareValue(FEmitContext& Context, const FRequestedTy
 	UpdateValueTypes();
 
 	// Assuming we have at least one value with a valid type, we use that to initialize our type
-	CurrentType.SetEvaluationType(EExpressionEvaluationType::Shader); // TODO - No support for preshader flow control
+	CurrentType.SetEvaluation(EExpressionEvaluation::Shader); // TODO - No support for preshader flow control
 	OutResult.SetType(Context, RequestedType, CurrentType);
 
 	if (NumValidTypes < NumValues)
@@ -697,7 +727,7 @@ void FExpressionLocalPHI::PrepareValue(FEmitContext& Context, const FRequestedTy
 	}
 }
 
-void FExpressionLocalPHI::EmitValueShader(FEmitContext& Context, const FRequestedType& RequestedType, FShaderValue& OutShader) const
+void FExpressionLocalPHI::EmitValueShader(FEmitContext& Context, const FRequestedType& RequestedType, FEmitShaderResult& OutShader) const
 {
 	int32 LocalPHIIndex = Context.LocalPHIs.Find(this);
 	if (LocalPHIIndex == INDEX_NONE)
@@ -717,7 +747,9 @@ void FExpressionLocalPHI::EmitValueShader(FEmitContext& Context, const FRequeste
 			}
 		}
 
+		const EExpressionDerivative Derivative = GetPreparedType().GetData().Derivative;
 		const FRequestedType LocalType = GetRequestedType();
+		const Shader::FType LocalDerivativeType = LocalType.GetType().GetDerivativeType();
 
 		bool bNeedToAddDeclaration = true;
 		for (int32 i = 0; i < NumValues; ++i)
@@ -726,7 +758,7 @@ void FExpressionLocalPHI::EmitValueShader(FEmitContext& Context, const FRequeste
 			check(IsScopeLive(ValueScope));
 
 			Context.ScopeStack.Add(ValueScope);
-			const TCHAR* ValueCode = Values[i]->GetValueShader(Context, LocalType);
+			const FShaderValue ShaderValue = Values[i]->GetValueShader(Context, LocalType);
 			Context.ScopeStack.Pop();
 
 			if (ValueScope == DeclarationScope)
@@ -734,14 +766,35 @@ void FExpressionLocalPHI::EmitValueShader(FEmitContext& Context, const FRequeste
 				ValueScope->EmitDeclarationf(Context, TEXT("%s LocalPHI%d = %s;"),
 					LocalType.GetName(),
 					LocalPHIIndex,
-					ValueCode);
+					ShaderValue.Code);
+				if (Derivative == EExpressionDerivative::Valid)
+				{
+					ValueScope->EmitDeclarationf(Context, TEXT("%s LocalPHI%dDdx = %s;"),
+						LocalDerivativeType.GetName(),
+						LocalPHIIndex,
+						ShaderValue.CodeDdx);
+					ValueScope->EmitDeclarationf(Context, TEXT("%s LocalPHI%dDdy = %s;"),
+						LocalDerivativeType.GetName(),
+						LocalPHIIndex,
+						ShaderValue.CodeDdy);
+				}
+
 				bNeedToAddDeclaration = false;
 			}
 			else
 			{
 				ValueScope->EmitStatementf(Context, TEXT("LocalPHI%d = %s;"),
 					LocalPHIIndex,
-					ValueCode);
+					ShaderValue.Code);
+				if (Derivative == EExpressionDerivative::Valid)
+				{
+					ValueScope->EmitStatementf(Context, TEXT("LocalPHI%dDdx = %s;"),
+						LocalPHIIndex,
+						ShaderValue.CodeDdx);
+					ValueScope->EmitStatementf(Context, TEXT("LocalPHI%dDdy = %s;"),
+						LocalPHIIndex,
+						ShaderValue.CodeDdy);
+				}
 			}
 		}
 
@@ -751,10 +804,21 @@ void FExpressionLocalPHI::EmitValueShader(FEmitContext& Context, const FRequeste
 			DeclarationScope->EmitDeclarationf(Context, TEXT("%s LocalPHI%d;"),
 				LocalType.GetName(),
 				LocalPHIIndex);
+			if (Derivative == EExpressionDerivative::Valid)
+			{
+				DeclarationScope->EmitDeclarationf(Context, TEXT("%s LocalPHI%dDdx;"),
+					LocalDerivativeType.GetName(),
+					LocalPHIIndex);
+				DeclarationScope->EmitDeclarationf(Context, TEXT("%s LocalPHI%dDdy;"),
+					LocalDerivativeType.GetName(),
+					LocalPHIIndex);
+			}
 		}
 	}
 
 	OutShader.Code.Appendf(TEXT("LocalPHI%d"), LocalPHIIndex);
+	OutShader.CodeDdx.Appendf(TEXT("LocalPHI%dDdx"), LocalPHIIndex);
+	OutShader.CodeDdy.Appendf(TEXT("LocalPHI%dDdy"), LocalPHIIndex);
 	OutShader.Type = GetType();
 	OutShader.bInline = true;
 }
@@ -792,7 +856,7 @@ const FPreparedType& PrepareExpressionValue(FEmitContext& Context, FExpression*
 		InExpression->CurrentRequestedType = RequestedType;
 		bNeedToUpdateType = !RequestedType.IsVoid();
 	}
-	else if(InExpression->CurrentRequestedType.GetStructType() != RequestedType.GetStructType())
+	else if (InExpression->CurrentRequestedType.GetStructType() != RequestedType.GetStructType())
 	{
 		Context.Errors.AddError(InExpression, TEXT("Type mismatch"));
 		return VoidType;
@@ -801,11 +865,13 @@ const FPreparedType& PrepareExpressionValue(FEmitContext& Context, FExpression*
 	{
 		const int32 NumComponents = RequestedType.GetNumComponents();
 		InExpression->CurrentRequestedType.RequestedComponents.PadToNum(NumComponents, false);
-		for (int32 ComponentIndex = 0; ComponentIndex < NumComponents; ++ComponentIndex)
+		for (int32 Index = 0; Index < NumComponents; ++Index)
 		{
-			if (!InExpression->CurrentRequestedType.IsComponentRequested(ComponentIndex) && RequestedType.IsComponentRequested(ComponentIndex))
+			const EComponentRequest PrevRequest = InExpression->CurrentRequestedType.GetComponentRequest(Index);
+			const EComponentRequest Request = RequestedType.GetComponentRequest(Index);
+			if ((uint32)Request > (uint32)PrevRequest)
 			{
-				InExpression->CurrentRequestedType.RequestedComponents[ComponentIndex] = true;
+				InExpression->CurrentRequestedType.SetComponentRequest(Index, Request);
 				bNeedToUpdateType = true;
 			}
 		}
@@ -830,7 +896,13 @@ const FPreparedType& PrepareExpressionValue(FEmitContext& Context, FExpression*
 	return InExpression->PrepareValueResult.GetPreparedType();
 }
 
-FRequestedType::FRequestedType(const Shader::FType& InType, bool bDefaultRequest)
+FRequestedType::FRequestedType(int32 NumComponents, EComponentRequest DefaultRequest)
+{
+	RequestedComponents.Init(IsRequested(DefaultRequest), NumComponents);
+	RequestedComponentDerivatives.Init(IsDerivativeRequested(DefaultRequest), NumComponents);
+}
+
+FRequestedType::FRequestedType(const Shader::FType& InType, EComponentRequest DefaultRequest)
 {
 	int32 NumComponents = 0;
 	if (InType.IsStruct())
@@ -844,14 +916,17 @@ FRequestedType::FRequestedType(const Shader::FType& InType, bool bDefaultRequest
 		ValueComponentType = TypeDesc.ComponentType;
 		NumComponents = TypeDesc.NumComponents;
 	}
-	RequestedComponents.Init(bDefaultRequest, NumComponents);
+
+	RequestedComponents.Init(IsRequested(DefaultRequest), NumComponents);
+	RequestedComponentDerivatives.Init(IsDerivativeRequested(DefaultRequest), NumComponents);
 }
 
-FRequestedType::FRequestedType(const Shader::EValueType& InType, bool bDefaultRequest)
+FRequestedType::FRequestedType(const Shader::EValueType& InType, EComponentRequest DefaultRequest)
 {
 	const Shader::FValueTypeDescription TypeDesc = Shader::GetValueTypeDescription(InType);
 	ValueComponentType = TypeDesc.ComponentType;
-	RequestedComponents.Init(bDefaultRequest, TypeDesc.NumComponents);
+	RequestedComponents.Init(IsRequested(DefaultRequest), TypeDesc.NumComponents);
+	RequestedComponentDerivatives.Init(IsDerivativeRequested(DefaultRequest), TypeDesc.NumComponents);
 }
 
 const Shader::FType FRequestedType::GetType() const
@@ -880,28 +955,72 @@ int32 FRequestedType::GetNumComponents() const
 	return 0;
 }
 
-bool FRequestedType::Merge(const FRequestedType& OtherType)
+EComponentRequest FRequestedType::GetComponentRequest(int32 Index) const
 {
-	if ((StructType || OtherType.StructType) && StructType != OtherType.StructType)
+	if (RequestedComponents.IsValidIndex(Index))
 	{
-		return false;
+		if (RequestedComponentDerivatives.IsValidIndex(Index) && RequestedComponentDerivatives[Index])
+		{
+			return EComponentRequest::RequestedWithDerivative;
+		}
+		if (RequestedComponents[Index])
+		{
+			return EComponentRequest::Requested;
+		}
 	}
+	return EComponentRequest::None;
+}
 
-	StructType = OtherType.StructType;
-	if (!StructType)
+void FRequestedType::SetComponentRequest(int32 Index, EComponentRequest Request)
+{
+	const bool bRequested = IsRequested(Request);
+	const bool bRequestedDerivative = IsDerivativeRequested(Request);
+	if (bRequested)
 	{
-		ValueComponentType = Shader::CombineComponentTypes(ValueComponentType, OtherType.ValueComponentType);
+		RequestedComponents.PadToNum(Index + 1, false);
+	}
+	if (bRequestedDerivative)
+	{
+		RequestedComponentDerivatives.PadToNum(Index + 1, false);
 	}
 
-	for (int32 Index = 0; Index < OtherType.RequestedComponents.Num(); ++Index)
+	if (RequestedComponents.IsValidIndex(Index))
 	{
-		if (OtherType.IsComponentRequested(Index))
-		{
-			SetComponentRequested(Index);
-		}
+		RequestedComponents[Index] = bRequested;
+	}
+	if (RequestedComponentDerivatives.IsValidIndex(Index))
+	{
+		RequestedComponentDerivatives[Index] = bRequestedDerivative;
+	}
+}
+
+void FRequestedType::SetFieldRequested(const Shader::FStructField* Field, EComponentRequest Request)
+{
+	const int32 NumComponents = Field->GetNumComponents();
+	for (int32 Index = 0; Index < NumComponents; ++Index)
+	{
+		SetComponentRequest(Field->ComponentIndex + Index, Request);
 	}
+}
 
-	return true;
+void FRequestedType::SetField(const Shader::FStructField* Field, const FRequestedType& InRequest)
+{
+	const int32 NumComponents = Field->GetNumComponents();
+	for (int32 Index = 0; Index < NumComponents; ++Index)
+	{
+		SetComponentRequest(Field->ComponentIndex + Index, InRequest.GetComponentRequest(Index));
+	}
+}
+
+FRequestedType FRequestedType::GetField(const Shader::FStructField* Field) const
+{
+	FRequestedType Result(Field->Type);
+	const int32 NumComponents = Field->GetNumComponents();
+	for (int32 Index = 0; Index < NumComponents; ++Index)
+	{
+		Result.SetComponentRequest(Index, GetComponentRequest(Field->ComponentIndex + Index));
+	}
+	return Result;
 }
 
 FPreparedType::FPreparedType(const Shader::FType& InType)
@@ -924,7 +1043,7 @@ int32 FPreparedType::GetNumComponents() const
 	}
 	else if (ValueComponentType != Shader::EValueComponentType::Void)
 	{
-		const int32 MaxComponentIndex = ComponentEvaluationType.FindLastByPredicate([](EExpressionEvaluationType InEvaluationType) { return InEvaluationType != EExpressionEvaluationType::None; });
+		const int32 MaxComponentIndex = PreparedComponents.FindLastByPredicate([](const FPreparedData& InData) { return InData.IsValid(); });
 		if (MaxComponentIndex != INDEX_NONE)
 		{
 			return MaxComponentIndex + 1;
@@ -939,6 +1058,7 @@ FRequestedType MakeRequestedType(Shader::EValueComponentType ComponentType, cons
 	FRequestedType Result;
 	Result.ValueComponentType = ComponentType;
 	Result.RequestedComponents = RequestedComponents.RequestedComponents;
+	Result.RequestedComponentDerivatives = RequestedComponents.RequestedComponentDerivatives;
 	return Result;
 }
 
@@ -972,76 +1092,127 @@ FRequestedType FPreparedType::GetRequestedType() const
 		}
 		for (int32 Index = 0; Index < NumComponents; ++Index)
 		{
-			if (GetComponentEvaluationType(Index) != EExpressionEvaluationType::None)
-			{
-				Result.SetComponentRequested(Index);
-			}
+			Result.SetComponentRequest(Index, GetComponentData(Index).GetRequest());
 		}
 	}
 	return Result;
 }
 
-EExpressionEvaluationType FPreparedType::GetEvaluationType(const FRequestedType& RequestedType) const
+FPreparedData FPreparedType::GetData() const
 {
-	EExpressionEvaluationType Result = EExpressionEvaluationType::None;
-	for(int32 Index = 0; Index < ComponentEvaluationType.Num(); ++Index)
+	FPreparedData Result;
+	for (int32 Index = 0; Index < PreparedComponents.Num(); ++Index)
 	{
-		if (RequestedType.IsComponentRequested(Index))
+		const FPreparedData& Component = PreparedComponents[Index];
+		if (Component.IsValid())
 		{
-			Result = CombineEvaluationTypes(Result, ComponentEvaluationType[Index]);
+			Result = CombinePreparedData(Result, Component);
 		}
 	}
 	return Result;
 }
 
-EExpressionEvaluationType FPreparedType::GetFieldEvaluationType(int32 ComponentIndex, int32 NumComponents) const
+FPreparedData FPreparedType::GetData(const FRequestedType& RequestedType) const
 {
-	EExpressionEvaluationType Result = EExpressionEvaluationType::None;
-	for (int32 Index = 0; Index < NumComponents; ++Index)
+	FPreparedData Result;
+	for (int32 Index = 0; Index < PreparedComponents.Num(); ++Index)
 	{
-		Result = CombineEvaluationTypes(Result, GetComponentEvaluationType(ComponentIndex + Index));
+		const FPreparedData& Component = PreparedComponents[Index];
+		if (Component.IsValid())
+		{
+			const EComponentRequest ComponentRequest = RequestedType.GetComponentRequest(Index);
+			if (IsRequested(ComponentRequest))
+			{
+				Result.Evaluation = CombineEvaluations(Result.Evaluation, Component.Evaluation);
+				if (IsDerivativeRequested(ComponentRequest))
+				{
+					Result.Derivative = CombineDerivatives(Result.Derivative, Component.Derivative);
+				}
+			}
+		}
 	}
 	return Result;
 }
 
-void FPreparedType::SetEvaluationType(EExpressionEvaluationType EvaluationType)
+FPreparedData FPreparedType::GetFieldData(int32 ComponentIndex, int32 NumComponents) const
 {
-	const int32 NumComponents = GetNumComponents();
+	FPreparedData Result;
 	for (int32 Index = 0; Index < NumComponents; ++Index)
 	{
-		if (ComponentEvaluationType[Index] != EExpressionEvaluationType::None)
+		Result = CombinePreparedData(Result, GetComponentData(ComponentIndex + Index));
+	}
+	return Result;
+}
+
+FPreparedData FPreparedType::GetComponentData(int32 Index) const
+{
+	FPreparedData Result;
+	if (PreparedComponents.IsValidIndex(Index))
+	{
+		Result = PreparedComponents[Index];
+	}
+	return Result;
+}
+
+void FPreparedType::SetComponentData(int32 Index, const FPreparedData& Data)
+{
+	if (Data.IsValid() && Index >= PreparedComponents.Num())
+	{
+		PreparedComponents.AddDefaulted(Index + 1 - PreparedComponents.Num());
+	}
+	if (PreparedComponents.IsValidIndex(Index))
+	{
+		PreparedComponents[Index] = Data;
+	}
+}
+
+void FPreparedType::MergeComponentData(int32 Index, EComponentRequest Request, const FPreparedData& Data)
+{
+	if (Request != EComponentRequest::None)
+	{
+		if (Data.IsValid() && Index >= PreparedComponents.Num())
 		{
-			ComponentEvaluationType[Index] = EvaluationType;
+			PreparedComponents.AddDefaulted(Index + 1 - PreparedComponents.Num());
+		}
+		if (PreparedComponents.IsValidIndex(Index))
+		{
+			FPreparedData& Component = PreparedComponents[Index];
+			Component.Evaluation = CombineEvaluations(Component.Evaluation, Data.Evaluation);
+			if (Request == EComponentRequest::RequestedWithDerivative)
+			{
+				Component.Derivative = CombineDerivatives(Component.Derivative, Data.Derivative);
+			}
 		}
 	}
 }
 
-void FPreparedType::SetField(const Shader::FStructField* Field, const FPreparedType& FieldType)
+void FPreparedType::SetEvaluation(EExpressionEvaluation Evaluation)
 {
-	for (int32 Index = 0; Index < Field->GetNumComponents(); ++Index)
+	for (int32 Index = 0; Index < PreparedComponents.Num(); ++Index)
 	{
-		SetComponentEvaluationType(Field->ComponentIndex + Index, FieldType.GetComponentEvaluationType(Index));
+		if (PreparedComponents[Index].IsValid())
+		{
+			PreparedComponents[Index].Evaluation = Evaluation;
+		}
 	}
 }
 
-FPreparedType FPreparedType::GetFieldType(const Shader::FStructField* Field) const
+void FPreparedType::SetField(const Shader::FStructField* Field, const FPreparedType& FieldType)
 {
-	FPreparedType Result(Field->Type);
 	for (int32 Index = 0; Index < Field->GetNumComponents(); ++Index)
 	{
-		Result.SetComponentEvaluationType(Index, GetComponentEvaluationType(Field->ComponentIndex + Index));
+		SetComponentData(Field->ComponentIndex + Index, FieldType.GetComponentData(Index));
 	}
-	return Result;
 }
 
-void FPreparedType::SetComponentEvaluationType(int32 Index, EExpressionEvaluationType EvaluationType)
+FPreparedType FPreparedType::GetFieldType(const Shader::FStructField* Field) const
 {
-	if (Index >= ComponentEvaluationType.Num())
+	FPreparedType Result(Field->Type);
+	for (int32 Index = 0; Index < Field->GetNumComponents(); ++Index)
 	{
-		static_assert((uint8)EExpressionEvaluationType::None == 0u, "Assume AddZeroed() will initialize to None");
-		ComponentEvaluationType.AddZeroed(Index + 1 - ComponentEvaluationType.Num());
+		Result.SetComponentData(Index, GetComponentData(Field->ComponentIndex + Index));
 	}
-	ComponentEvaluationType[Index] = EvaluationType;
+	return Result;
 }
 
 FPreparedType MergePreparedTypes(const FPreparedType& Lhs, const FPreparedType& Rhs)
@@ -1076,9 +1247,9 @@ FPreparedType MergePreparedTypes(const FPreparedType& Lhs, const FPreparedType&
 
 	for (int32 Index = 0; Index < NumComponents; ++Index)
 	{
-		const EExpressionEvaluationType LhsEvaluation = Lhs.GetComponentEvaluationType(Index);
-		const EExpressionEvaluationType RhsEvaluation = Rhs.GetComponentEvaluationType(Index);
-		Result.SetComponentEvaluationType(Index, CombineEvaluationTypes(LhsEvaluation, RhsEvaluation));
+		const FPreparedData LhsData = Lhs.GetComponentData(Index);
+		const FPreparedData RhsData = Rhs.GetComponentData(Index);
+		Result.SetComponentData(Index, CombinePreparedData(LhsData, RhsData));
 	}
 
 	return Result;
@@ -1089,7 +1260,7 @@ bool FPrepareValueResult::TryMergePreparedType(FEmitContext& Context, const Shad
 	// If we previously had a forwarded value set, reset that and start over
 	if (ForwardValue || !PreparedType.IsInitialized())
 	{
-		PreparedType.ComponentEvaluationType.Reset();
+		PreparedType.PreparedComponents.Reset();
 		PreparedType.ValueComponentType = ComponentType;
 		PreparedType.StructType = StructType;
 		ForwardValue = nullptr;
@@ -1114,37 +1285,17 @@ bool FPrepareValueResult::TryMergePreparedType(FEmitContext& Context, const Shad
 	return true;
 }
 
-void FPrepareValueResult::SetType(FEmitContext& Context, const FRequestedType& RequestedType, EExpressionEvaluationType EvaluationType, const Shader::FType& Type)
+void FPrepareValueResult::SetType(FEmitContext& Context, const FRequestedType& RequestedType, const FPreparedData& Data, const Shader::FType& Type)
 {
 	if (TryMergePreparedType(Context, Type.StructType, Shader::GetValueTypeDescription(Type.ValueType).ComponentType))
 	{
-		if (EvaluationType != EExpressionEvaluationType::None)
+		if (Data.IsValid())
 		{
 			const int32 NumComponents = Type.GetNumComponents();
 			for (int32 Index = 0; Index < NumComponents; ++Index)
 			{
-				if (RequestedType.IsComponentRequested(Index))
-				{
-					PreparedType.SetComponentEvaluationType(Index, EvaluationType);
-				}
-			}
-		}
-	}
-}
-
-void FPrepareValueResult::SetType(FEmitContext& Context, const FRequestedType& RequestedType, EExpressionEvaluationType EvaluationType, Shader::EValueComponentType ComponentType)
-{
-	if (TryMergePreparedType(Context, nullptr, ComponentType))
-	{
-		if (EvaluationType != EExpressionEvaluationType::None)
-		{
-			const int32 NumComponents = RequestedType.GetNumComponents();
-			for (int32 Index = 0; Index < NumComponents; ++Index)
-			{
-				if (RequestedType.IsComponentRequested(Index))
-				{
-					PreparedType.SetComponentEvaluationType(Index, EvaluationType);
-				}
+				const EComponentRequest ComponentRequest = RequestedType.GetComponentRequest(Index);
+				PreparedType.MergeComponentData(Index, ComponentRequest, Data);
 			}
 		}
 	}
@@ -1157,10 +1308,8 @@ void FPrepareValueResult::SetType(FEmitContext& Context, const FRequestedType& R
 		const int32 NumComponents = RequestedType.GetNumComponents();
 		for (int32 Index = 0; Index < NumComponents; ++Index)
 		{
-			if (RequestedType.IsComponentRequested(Index))
-			{
-				PreparedType.SetComponentEvaluationType(Index, Type.GetComponentEvaluationType(Index));
-			}
+			const EComponentRequest ComponentRequest = RequestedType.GetComponentRequest(Index);
+			PreparedType.MergeComponentData(Index, ComponentRequest, Type.GetComponentData(Index));
 		}
 	}
 }
@@ -1175,7 +1324,7 @@ void FPrepareValueResult::SetForwardValue(FEmitContext& Context, const FRequeste
 	}
 }
 
-void FExpression::EmitValueShader(FEmitContext& Context, const FRequestedType& RequestedType, FShaderValue& OutShader) const
+void FExpression::EmitValueShader(FEmitContext& Context, const FRequestedType& RequestedType, FEmitShaderResult& OutShader) const
 {
 	check(false);
 }
@@ -1185,30 +1334,36 @@ void FExpression::EmitValuePreshader(FEmitContext& Context, const FRequestedType
 	check(false);
 }
 
-const TCHAR* FExpression::GetValueShader(FEmitContext& Context, const FRequestedType& RequestedType)
+FShaderValue FExpression::GetValueShader(FEmitContext& Context, const FRequestedType& RequestedType, const Shader::FType& ResultType)
 {
 	if (PrepareValueResult.ForwardValue)
 	{
-		return PrepareValueResult.ForwardValue->GetValueShader(Context, RequestedType);
+		return PrepareValueResult.ForwardValue->GetValueShader(Context, RequestedType, ResultType);
 	}
 
-	const EExpressionEvaluationType EvaluationType = PrepareValueResult.PreparedType.GetEvaluationType(RequestedType);
-	check(EvaluationType != EExpressionEvaluationType::None);
+	const FPreparedData Data = PrepareValueResult.PreparedType.GetData(RequestedType);
+	check(Data.IsValid());
 
 	FScope* CurrentScope = Context.ScopeStack.Last();
 	check(IsScopeLive(CurrentScope));
 
+	const EExpressionDerivative Derivative = Data.Derivative;
 	FEmitShaderValue* ShaderValue = nullptr;
-	if (EvaluationType == EExpressionEvaluationType::Constant || EvaluationType == EExpressionEvaluationType::Preshader)
+	FEmitShaderValue* ShaderValueDdx = nullptr;
+	FEmitShaderValue* ShaderValueDdy = nullptr;
+	if (Data.Evaluation == EExpressionEvaluation::Constant || Data.Evaluation == EExpressionEvaluation::Preshader)
 	{
 		ShaderValue = Context.AcquirePreshaderOrConstant(RequestedType, CurrentScope, this);
+		check(Derivative != EExpressionDerivative::Valid); // If a constant has a valid derivative, it should be 'Zero'
 	}
 	else
 	{
-		check(EvaluationType == EExpressionEvaluationType::Shader);
+		check(Data.Evaluation == EExpressionEvaluation::Shader);
 
-		TStringBuilder<1024> FormattedCode;
-		FShaderValue ShaderResult(FormattedCode);
+		TStringBuilder<512> FormattedCode;
+		TStringBuilder<512> FormattedCodeDdx;
+		TStringBuilder<512> FormattedCodeDdy;
+		FEmitShaderResult ShaderResult(FormattedCode, FormattedCodeDdx, FormattedCodeDdy);
 		FEmitShaderValueDependencies Dependencies;
 		{
 			FEmitShaderValueContext& ValueContext = Context.ShaderValueStack.AddDefaulted_GetRef();
@@ -1217,23 +1372,80 @@ const TCHAR* FExpression::GetValueShader(FEmitContext& Context, const FRequested
 			check(!ShaderResult.Type.IsVoid());
 
 			Dependencies = MoveTemp(ValueContext.Dependencies);
+			if (Derivative == EExpressionDerivative::Valid)
+			{
+				Dependencies.Append(MoveTemp(ValueContext.DerivativeDependencies));
+			}
 			Context.ShaderValueStack.Pop(false);
 		}
-		ShaderValue = Context.AcquireShader(CurrentScope, ShaderResult, Dependencies);
+		
+		ShaderValue = Context.AcquireShader(CurrentScope, ShaderResult.Type, ShaderResult.Code, ShaderResult.bInline, Dependencies);
+		if (Derivative == EExpressionDerivative::Valid)
+		{
+			check(ShaderResult.CodeDdx.Len() > 0);
+			check(ShaderResult.CodeDdy.Len() > 0);
+			const Shader::FType DerivativeResultType = ShaderResult.Type.GetDerivativeType();
+			check(!DerivativeResultType.IsVoid());
+			ShaderValueDdx = Context.AcquireShader(CurrentScope, DerivativeResultType, ShaderResult.CodeDdx, ShaderResult.bInline, Dependencies);
+			ShaderValueDdy = Context.AcquireShader(CurrentScope, DerivativeResultType, ShaderResult.CodeDdy, ShaderResult.bInline, Dependencies);
+		}
 	}
 
+	const Shader::FType DerivativeResultType = ResultType.GetDerivativeType();
 	check(ShaderValue);
-	ShaderValue = Context.CastShaderValue(this, CurrentScope, ShaderValue, RequestedType.GetType());
+	ShaderValue = Context.CastShaderValue(this, CurrentScope, ShaderValue, ResultType);
+	bool bHasDerivatives = false;
+	if (Derivative == EExpressionDerivative::Valid)
+	{
+		check(!DerivativeResultType.IsVoid());
+		ShaderValueDdx = Context.CastShaderValue(this, CurrentScope, ShaderValueDdx, DerivativeResultType);
+		ShaderValueDdy = Context.CastShaderValue(this, CurrentScope, ShaderValueDdy, DerivativeResultType);
+		bHasDerivatives = true;
+	}
+	else if (Derivative == EExpressionDerivative::Zero)
+	{
+		check(!DerivativeResultType.IsVoid());
+		ShaderValueDdx = Context.AcquireConstantZero(CurrentScope, DerivativeResultType);
+		ShaderValueDdy = ShaderValueDdx;
+		bHasDerivatives = true;
+	}
 
 	if (Context.ShaderValueStack.Num() > 0)
 	{
-		Context.ShaderValueStack.Last().Dependencies.Add(ShaderValue);
+		FEmitShaderValueContext& PrevValueContext = Context.ShaderValueStack.Last();
+		PrevValueContext.Dependencies.Add(ShaderValue);
+		if (bHasDerivatives)
+		{
+			check(ShaderValueDdx && ShaderValueDdy); // must have both
+			if (ShaderValueDdx != ShaderValue)
+			{
+				PrevValueContext.DerivativeDependencies.Add(ShaderValueDdx);
+				if (ShaderValueDdy != ShaderValueDdx)
+				{
+					PrevValueContext.DerivativeDependencies.Add(ShaderValueDdy);
+				}
+			}
+		}
+	}
+
+	FShaderValue Result;
+	Result.Type = ShaderValue->Type;
+	Result.Code = ShaderValue->Reference;
+	Result.bHasDerivatives = bHasDerivatives;
+	if (bHasDerivatives)
+	{
+		Result.CodeDdx = ShaderValueDdx->Reference;
+		Result.CodeDdy = ShaderValueDdy->Reference;
 	}
+	return Result;
+}
 
-	return ShaderValue->Reference;
+FShaderValue FExpression::GetValueShader(FEmitContext& Context, const FRequestedType& RequestedType)
+{
+	return GetValueShader(Context, RequestedType, RequestedType.GetType());
 }
 
-const TCHAR* FExpression::GetValueShader(FEmitContext& Context)
+FShaderValue FExpression::GetValueShader(FEmitContext& Context)
 {
 	return GetValueShader(Context, GetRequestedType());
 }
@@ -1246,8 +1458,8 @@ void FExpression::GetValuePreshader(FEmitContext& Context, const FRequestedType&
 	}
 
 	check(!bReentryFlag);
-	const EExpressionEvaluationType EvaluationType = PrepareValueResult.PreparedType.GetEvaluationType(RequestedType);
-	if (EvaluationType == EExpressionEvaluationType::Preshader)
+	const EExpressionEvaluation Evaluation = PrepareValueResult.PreparedType.GetData(RequestedType).Evaluation;
+	if (Evaluation == EExpressionEvaluation::Preshader)
 	{
 		bReentryFlag = true;
 		EmitValuePreshader(Context, RequestedType, OutPreshader);
@@ -1255,7 +1467,7 @@ void FExpression::GetValuePreshader(FEmitContext& Context, const FRequestedType&
 	}
 	else
 	{
-		check(EvaluationType == EExpressionEvaluationType::Constant);
+		check(Evaluation == EExpressionEvaluation::Constant);
 		const Shader::FValue ConstantValue = GetValueConstant(Context, RequestedType);
 		OutPreshader.WriteOpcode(Shader::EPreshaderOpcode::Constant).Write(ConstantValue);
 	}
@@ -1269,7 +1481,7 @@ Shader::FValue FExpression::GetValueConstant(FEmitContext& Context, const FReque
 	}
 
 	check(!bReentryFlag);
-	check(PrepareValueResult.PreparedType.GetEvaluationType(RequestedType) == EExpressionEvaluationType::Constant);
+	check(PrepareValueResult.PreparedType.GetData(RequestedType).Evaluation == EExpressionEvaluation::Constant);
 	
 	Shader::FPreshaderData ConstantPreshader;
 	bReentryFlag = true;
diff --git a/Engine/Source/Runtime/Engine/Private/HLSLTree/HLSLTreeCommon.cpp b/Engine/Source/Runtime/Engine/Private/HLSLTree/HLSLTreeCommon.cpp
--- a/Engine/Source/Runtime/Engine/Private/HLSLTree/HLSLTreeCommon.cpp
+++ b/Engine/Source/Runtime/Engine/Private/HLSLTree/HLSLTreeCommon.cpp
@@ -69,10 +69,11 @@ FRequestedType FSwizzleParameters::GetRequestedInputType(const FRequestedType& R
 
 	for (int32 Index = 0; Index < NumComponents; ++Index)
 	{
-		if (RequestedType.IsComponentRequested(Index))
+		const EComponentRequest ComponentRequest = RequestedType.GetComponentRequest(Index);
+		if (ComponentRequest != EComponentRequest::None)
 		{
 			const int32 SwizzledComponentIndex = ComponentIndex[Index];
-			RequestedInputType.SetComponentRequested(SwizzledComponentIndex);
+			RequestedInputType.SetComponentRequest(SwizzledComponentIndex, ComponentRequest);
 		}
 	}
 	
@@ -104,7 +105,7 @@ FSwizzleParameters MakeSwizzleMask(bool bInR, bool bInG, bool bInB, bool bInA)
 
 void FExpressionConstant::PrepareValue(FEmitContext& Context, const FRequestedType& RequestedType, FPrepareValueResult& OutResult) const
 {
-	OutResult.SetType(Context, RequestedType, EExpressionEvaluationType::Constant, Value.Type);
+	OutResult.SetType(Context, RequestedType, EExpressionEvaluation::Constant, Value.Type);
 }
 
 void FExpressionConstant::EmitValuePreshader(FEmitContext& Context, const FRequestedType& RequestedType, Shader::FPreshaderData& OutPreshader) const
@@ -114,8 +115,8 @@ void FExpressionConstant::EmitValuePreshader(FEmitContext& Context, const FReque
 
 void FExpressionMaterialParameter::PrepareValue(FEmitContext& Context, const FRequestedType& RequestedType, FPrepareValueResult& OutResult) const
 {
-	const EExpressionEvaluationType EvaluationType = IsStaticMaterialParameter(ParameterType) ? EExpressionEvaluationType::Constant : EExpressionEvaluationType::Preshader;
-	OutResult.SetType(Context, RequestedType, EvaluationType, GetShaderValueType(ParameterType));
+	const EExpressionEvaluation Evaluation = IsStaticMaterialParameter(ParameterType) ? EExpressionEvaluation::Constant : EExpressionEvaluation::Preshader;
+	OutResult.SetType(Context, RequestedType, Evaluation, GetShaderValueType(ParameterType));
 }
 
 void FExpressionMaterialParameter::EmitValuePreshader(FEmitContext& Context, const FRequestedType& RequestedType, Shader::FPreshaderData& OutPreshader) const
@@ -153,7 +154,13 @@ void FExpressionMaterialParameter::EmitValuePreshader(FEmitContext& Context, con
 	}
 }
 
-void FExpressionExternalInput::EmitValueShader(FEmitContext& Context, const FRequestedType& RequestedType, FShaderValue& OutShader) const
+void FExpressionExternalInput::PrepareValue(FEmitContext& Context, const FRequestedType& RequestedType, FPrepareValueResult& OutResult) const
+{
+	// TODO - derivative type
+	OutResult.SetTypeWithDerivative(Context, RequestedType, EExpressionEvaluation::Shader, GetInputExpressionType(InputType));
+}
+
+void FExpressionExternalInput::EmitValueShader(FEmitContext& Context, const FRequestedType& RequestedType, FEmitShaderResult& OutShader) const
 {
 	OutShader.Type = GetInputExpressionType(InputType);
 	OutShader.bInline = true;
@@ -165,6 +172,8 @@ void FExpressionExternalInput::EmitValueShader(FEmitContext& Context, const FReq
 
 		Context.NumTexCoords = FMath::Max(Context.NumTexCoords, TexCoordIndex + 1);
 		OutShader.Code.Appendf(TEXT("Parameters.TexCoords[%u].xy"), TexCoordIndex);
+		OutShader.CodeDdx.Appendf(TEXT("Parameters.TexCoords_DDX[%u].xy"), TexCoordIndex);
+		OutShader.CodeDdy.Appendf(TEXT("Parameters.TexCoords_DDY[%u].xy"), TexCoordIndex);
 	}
 	else
 	{
@@ -187,16 +196,20 @@ void FExpressionExternalInput::EmitValueShader(FEmitContext& Context, const FReq
 			checkNoEntry();
 			break;
 		}
+
+		OutShader.CodeDdx.Append(TEXT("Parameters.WorldPosition_DDX"));
+		OutShader.CodeDdy.Append(TEXT("Parameters.WorldPosition_DDY"));
 	}
 }
 
 void FExpressionTextureSample::PrepareValue(FEmitContext& Context, const FRequestedType& RequestedType, FPrepareValueResult& OutResult) const
 {
-	PrepareExpressionValue(Context, TexCoordExpression, 2);
-	OutResult.SetType(Context, RequestedType, EExpressionEvaluationType::Shader, Shader::EValueType::Float4);
+	const FRequestedType TexCoordRequestedType(Shader::EValueType::Float2, EComponentRequest::RequestedWithDerivative);
+	PrepareExpressionValue(Context, TexCoordExpression, TexCoordRequestedType);
+	OutResult.SetType(Context, RequestedType, EExpressionEvaluation::Shader, Shader::EValueType::Float4);
 }
 
-void FExpressionTextureSample::EmitValueShader(FEmitContext& Context, const FRequestedType& RequestedType, FShaderValue& OutShader) const
+void FExpressionTextureSample::EmitValueShader(FEmitContext& Context, const FRequestedType& RequestedType, FEmitShaderResult& OutShader) const
 {
 	const FTextureDescription& Desc = Declaration->Description;
 	const EMaterialValueType MaterialValueType = Desc.Texture->GetMaterialType();
@@ -322,18 +335,22 @@ void FExpressionTextureSample::EmitValueShader(FEmitContext& Context, const FReq
 		}
 	}
 
+	const FRequestedType TexCoordRequestedType(Shader::EValueType::Float2, EComponentRequest::RequestedWithDerivative);
+	const FShaderValue TexCoordValue = TexCoordExpression->GetValueShader(Context, TexCoordRequestedType);
 	OutShader.Type = Shader::EValueType::Float4;
-	OutShader.Code.Appendf(TEXT("%s(%s(%s, %s, %s))"),
+	OutShader.Code.Appendf(TEXT("%s(%sGrad(%s, %s, %s, %s, %s))"),
 		SamplerTypeFunction,
 		SampleFunctionName,
 		*TextureName,
 		*SamplerStateCode,
-		TexCoordExpression->GetValueShader(Context, Shader::EValueType::Float2));
+		TexCoordValue.Code,
+		TexCoordValue.CodeDdx,
+		TexCoordValue.CodeDdy);
 }
 
 void FExpressionGetStructField::PrepareValue(FEmitContext& Context, const FRequestedType& RequestedType, FPrepareValueResult& OutResult) const
 {
-	FRequestedType RequestedStructType(StructType, false);
+	FRequestedType RequestedStructType(StructType, EComponentRequest::None);
 	RequestedStructType.SetField(Field, RequestedType);
 
 	const FPreparedType& StructPreparedType = PrepareExpressionValue(Context, StructExpression, RequestedStructType);
@@ -345,21 +362,32 @@ void FExpressionGetStructField::PrepareValue(FEmitContext& Context, const FReque
 	OutResult.SetType(Context, RequestedType, StructPreparedType.GetFieldType(Field));
 }
 
-void FExpressionGetStructField::EmitValueShader(FEmitContext& Context, const FRequestedType& RequestedType, FShaderValue& OutShader) const
+void FExpressionGetStructField::EmitValueShader(FEmitContext& Context, const FRequestedType& RequestedType, FEmitShaderResult& OutShader) const
 {
-	FRequestedType RequestedStructType(StructType, false);
+	FRequestedType RequestedStructType(StructType, EComponentRequest::None);
 	RequestedStructType.SetField(Field, RequestedType);
 
+	const FShaderValue StructValue = StructExpression->GetValueShader(Context, RequestedStructType);
+
 	OutShader.bInline = true;
 	OutShader.Type = Field->Type;
 	OutShader.Code.Appendf(TEXT("%s.%s"),
-		StructExpression->GetValueShader(Context, RequestedStructType),
+		StructValue.Code,
 		Field->Name);
+	if (StructValue.bHasDerivatives)
+	{
+		OutShader.CodeDdx.Appendf(TEXT("%s.%s"),
+			StructValue.CodeDdx,
+			Field->Name);
+		OutShader.CodeDdy.Appendf(TEXT("%s.%s"),
+			StructValue.CodeDdy,
+			Field->Name);
+	}
 }
 
 void FExpressionGetStructField::EmitValuePreshader(FEmitContext& Context, const FRequestedType& RequestedType, Shader::FPreshaderData& OutPreshader) const
 {
-	FRequestedType RequestedStructType(StructType, false);
+	FRequestedType RequestedStructType(StructType, EComponentRequest::None);
 	RequestedStructType.SetField(Field, RequestedType);
 
 	StructExpression->GetValuePreshader(Context, RequestedStructType, OutPreshader);
@@ -394,19 +422,23 @@ void FExpressionSetStructField::PrepareValue(FEmitContext& Context, const FReque
 	}
 }
 
-void FExpressionSetStructField::EmitValueShader(FEmitContext& Context, const FRequestedType& RequestedType, FShaderValue& OutShader) const
+void FExpressionSetStructField::EmitValueShader(FEmitContext& Context, const FRequestedType& RequestedType, FEmitShaderResult& OutShader) const
 {
 	FRequestedType RequestedStructType(RequestedType);
 	RequestedStructType.ClearFieldRequested(Field);
-	const EExpressionEvaluationType StructEvaluationType = StructExpression->GetEvaluationType(RequestedStructType);
+	const EExpressionEvaluation StructEvaluation = StructExpression->GetEvaluation(RequestedStructType);
 
 	const FRequestedType RequestedFieldType = RequestedType.GetField(Field);
-	const EExpressionEvaluationType FieldEvaluationType = FieldExpression->GetEvaluationType(RequestedFieldType);
+	const EExpressionEvaluation FieldEvaluation = FieldExpression->GetEvaluation(RequestedFieldType);
+
+	const FShaderValue FieldValue = FieldExpression->GetValueShader(Context, RequestedFieldType, Field->Type);
+
+	const Shader::FStructType* DerivativeStructType = StructType->DerivativeType;
 
 	OutShader.Type = StructType;
-	if (StructEvaluationType == EExpressionEvaluationType::None)
+	if (StructEvaluation == EExpressionEvaluation::None)
 	{
-		check(FieldEvaluationType != EExpressionEvaluationType::None);
+		check(FieldEvaluation != EExpressionEvaluation::None);
 
 		// StructExpression is not used, so default to a zero-initialized struct
 		// This will happen if all the accessed struct fields are explicitly defined
@@ -414,15 +446,42 @@ void FExpressionSetStructField::EmitValueShader(FEmitContext& Context, const FRe
 			StructType->Name,
 			Field->Name,
 			StructType->Name,
-			FieldExpression->GetValueShader(Context, RequestedFieldType));
+			FieldValue.Code);
+		if (DerivativeStructType && FieldValue.bHasDerivatives)
+		{
+			OutShader.CodeDdx.Appendf(TEXT("%s_Set%s((%s)0, %s)"),
+				DerivativeStructType->Name,
+				Field->Name,
+				DerivativeStructType->Name,
+				FieldValue.CodeDdx);
+			OutShader.CodeDdy.Appendf(TEXT("%s_Set%s((%s)0, %s)"),
+				DerivativeStructType->Name,
+				Field->Name,
+				DerivativeStructType->Name,
+				FieldValue.CodeDdy);
+		}
 	}
 	else
 	{
+		const FShaderValue StructValue = StructExpression->GetValueShader(Context, RequestedStructType);
 		OutShader.Code.Appendf(TEXT("%s_Set%s(%s, %s)"),
 			StructType->Name,
 			Field->Name,
-			StructExpression->GetValueShader(Context, RequestedStructType),
-			FieldExpression->GetValueShader(Context, RequestedFieldType));
+			StructValue.Code,
+			FieldValue.Code);
+		if (DerivativeStructType && StructValue.bHasDerivatives && FieldValue.bHasDerivatives)
+		{
+			OutShader.CodeDdx.Appendf(TEXT("%s_Set%s(%s, %s)"),
+				DerivativeStructType->Name,
+				Field->Name,
+				StructValue.CodeDdx,
+				FieldValue.CodeDdx);
+			OutShader.CodeDdy.Appendf(TEXT("%s_Set%s(%s, %s)"),
+				DerivativeStructType->Name,
+				Field->Name,
+				StructValue.CodeDdy,
+				FieldValue.CodeDdy);
+		}
 	}
 }
 
@@ -430,12 +489,12 @@ void FExpressionSetStructField::EmitValuePreshader(FEmitContext& Context, const
 {
 	FRequestedType RequestedStructType(RequestedType);
 	RequestedStructType.ClearFieldRequested(Field);
-	const EExpressionEvaluationType StructEvaluationType = StructExpression->GetEvaluationType(RequestedStructType);
+	const EExpressionEvaluation StructEvaluation = StructExpression->GetEvaluation(RequestedStructType);
 
 	const FRequestedType RequestedFieldType = RequestedType.GetField(Field);
-	const EExpressionEvaluationType FieldEvaluationType = FieldExpression->GetEvaluationType(RequestedFieldType);
+	const EExpressionEvaluation FieldEvaluation = FieldExpression->GetEvaluation(RequestedFieldType);
 
-	if (StructEvaluationType != EExpressionEvaluationType::None)
+	if (StructEvaluation != EExpressionEvaluation::None)
 	{
 		StructExpression->GetValuePreshader(Context, RequestedStructType, OutPreshader);
 	}
@@ -451,7 +510,7 @@ void FExpressionSetStructField::EmitValuePreshader(FEmitContext& Context, const
 void FExpressionSelect::PrepareValue(FEmitContext& Context, const FRequestedType& RequestedType, FPrepareValueResult& OutResult) const
 {
 	const FPreparedType& ConditionType = PrepareExpressionValue(Context, ConditionExpression, 1);
-	if (ConditionType.GetEvaluationType(Shader::EValueType::Bool1) == EExpressionEvaluationType::Constant)
+	if (ConditionType.GetEvaluation(Shader::EValueType::Bool1) == EExpressionEvaluation::Constant)
 	{
 		const bool bCondition = ConditionExpression->GetValueConstant(Context, Shader::EValueType::Bool1).AsBoolScalar();
 		return OutResult.SetForwardValue(Context, RequestedType, bCondition ? TrueExpression : FalseExpression);
@@ -469,13 +528,16 @@ void FExpressionSelect::PrepareValue(FEmitContext& Context, const FRequestedType
 	OutResult.SetType(Context, RequestedType, MergePreparedTypes(LhsType, RhsType));
 }
 
-void FExpressionSelect::EmitValueShader(FEmitContext& Context, const FRequestedType& RequestedType, FShaderValue& OutShader) const
+void FExpressionSelect::EmitValueShader(FEmitContext& Context, const FRequestedType& RequestedType, FEmitShaderResult& OutShader) const
 {
+	const FShaderValue TrueValue = TrueExpression->GetValueShader(Context, RequestedType);
+	const FShaderValue FalseValue = FalseExpression->GetValueShader(Context, RequestedType);
+
 	OutShader.Type = RequestedType.GetType();
 	OutShader.Code.Appendf(TEXT("(%s ? %s : %s)"),
-		ConditionExpression->GetValueShader(Context, Shader::EValueType::Bool1),
-		TrueExpression->GetValueShader(Context, RequestedType),
-		FalseExpression->GetValueShader(Context, RequestedType));
+		ConditionExpression->GetValueShader(Context, Shader::EValueType::Bool1).Code,
+		TrueValue.Code,
+		FalseValue.Code);
 }
 
 void FExpressionSelect::EmitValuePreshader(FEmitContext& Context, const FRequestedType& RequestedType, Shader::FPreshaderData& OutPreshader) const
@@ -507,7 +569,7 @@ void FExpressionBinaryOp::PrepareValue(FEmitContext& Context, const FRequestedTy
 	if (OpDesc.PreshaderOpcode == Shader::EPreshaderOpcode::Nop)
 	{
 		// No preshader support
-		ResultType.SetEvaluationType(EExpressionEvaluationType::Shader);
+		ResultType.SetEvaluation(EExpressionEvaluation::Shader);
 	}
 
 	OutResult.SetType(Context, RequestedType, ResultType);
@@ -547,11 +609,12 @@ FBinaryOpTypes GetBinaryOpTypes(const FRequestedType& RequestedType, EBinaryOp O
 }
 } // namespace Private
 
-void FExpressionBinaryOp::EmitValueShader(FEmitContext& Context, const FRequestedType& RequestedType, FShaderValue& OutShader) const
+void FExpressionBinaryOp::EmitValueShader(FEmitContext& Context, const FRequestedType& RequestedType, FEmitShaderResult& OutShader) const
 {
 	const Private::FBinaryOpTypes Types = Private::GetBinaryOpTypes(RequestedType, Op, Lhs->GetRequestedType(), Rhs->GetRequestedType());
-	const TCHAR* LhsCode = Lhs->GetValueShader(Context, Types.LhsType);
-	const TCHAR* RhsCode = Rhs->GetValueShader(Context, Types.RhsType);
+	const FShaderValue LhsValue = Lhs->GetValueShader(Context, Types.LhsType);
+	const FShaderValue RhsValue = Rhs->GetValueShader(Context, Types.RhsType);
+	const bool bHasDerivatives = LhsValue.bHasDerivatives && RhsValue.bHasDerivatives;
 
 	OutShader.Type = Types.ResultType;
 	switch (Op)
@@ -559,51 +622,61 @@ void FExpressionBinaryOp::EmitValueShader(FEmitContext& Context, const FRequeste
 	case EBinaryOp::Add:
 		if (Types.bIsLWC)
 		{
-			OutShader.Code.Appendf(TEXT("LWCAdd(%s, %s)"), LhsCode, RhsCode);
+			OutShader.Code.Appendf(TEXT("LWCAdd(%s, %s)"), LhsValue.Code, RhsValue.Code);
 		}
 		else
 		{
-			OutShader.Code.Appendf(TEXT("(%s + %s)"), LhsCode, RhsCode);
+			OutShader.Code.Appendf(TEXT("(%s + %s)"), LhsValue.Code, RhsValue.Code);
+		}
+		if (bHasDerivatives)
+		{
+			OutShader.CodeDdx.Appendf(TEXT("(%s + %s)"), LhsValue.CodeDdx, RhsValue.CodeDdx);
+			OutShader.CodeDdy.Appendf(TEXT("(%s + %s)"), LhsValue.CodeDdy, RhsValue.CodeDdy);
 		}
 		break;
 	case EBinaryOp::Sub:
 		if (Types.bIsLWC)
 		{
-			OutShader.Code.Appendf(TEXT("LWCSub(%s, %s)"), LhsCode, RhsCode);
+			OutShader.Code.Appendf(TEXT("LWCSub(%s, %s)"), LhsValue.Code, RhsValue.Code);
 		}
 		else
 		{
-			OutShader.Code.Appendf(TEXT("(%s - %s)"), LhsCode, RhsCode);
+			OutShader.Code.Appendf(TEXT("(%s - %s)"), LhsValue.Code, RhsValue.Code);
 		}
 		break;
 	case EBinaryOp::Mul:
 		if (Types.bIsLWC)
 		{
-			OutShader.Code.Appendf(TEXT("LWCMultiply(%s, %s)"), LhsCode, RhsCode);
+			OutShader.Code.Appendf(TEXT("LWCMultiply(%s, %s)"), LhsValue.Code, RhsValue.Code);
 		}
 		else
 		{
-			OutShader.Code.Appendf(TEXT("(%s * %s)"), LhsCode, RhsCode);
+			OutShader.Code.Appendf(TEXT("(%s * %s)"), LhsValue.Code, RhsValue.Code);
+		}
+		if (bHasDerivatives)
+		{
+			OutShader.CodeDdx.Appendf(TEXT("(%s * %s + %s * %s)"), LhsValue.CodeDdx, RhsValue.Code, RhsValue.CodeDdx, LhsValue.Code);
+			OutShader.CodeDdy.Appendf(TEXT("(%s * %s + %s * %s)"), LhsValue.CodeDdy, RhsValue.Code, RhsValue.CodeDdy, LhsValue.Code);
 		}
 		break;
 	case EBinaryOp::Div:
 		if (Types.bIsLWC)
 		{
-			OutShader.Code.Appendf(TEXT("LWCDivide(%s, %s)"), LhsCode, RhsCode);
+			OutShader.Code.Appendf(TEXT("LWCDivide(%s, %s)"), LhsValue.Code, RhsValue.Code);
 		}
 		else
 		{
-			OutShader.Code.Appendf(TEXT("(%s / %s)"), LhsCode, RhsCode);
+			OutShader.Code.Appendf(TEXT("(%s / %s)"), LhsValue.Code, RhsValue.Code);
 		}
 		break;
 	case EBinaryOp::Less:
 		if (Types.bIsLWC)
 		{
-			OutShader.Code.Appendf(TEXT("LWCLess(%s, %s)"), LhsCode, RhsCode);
+			OutShader.Code.Appendf(TEXT("LWCLess(%s, %s)"), LhsValue.Code, RhsValue.Code);
 		}
 		else
 		{
-			OutShader.Code.Appendf(TEXT("(%s < %s)"), LhsCode, RhsCode);
+			OutShader.Code.Appendf(TEXT("(%s < %s)"), LhsValue.Code, RhsValue.Code);
 		}
 		break;
 	default:
@@ -632,13 +705,13 @@ void FExpressionSwizzle::PrepareValue(FEmitContext& Context, const FRequestedTyp
 	for (int32 ComponentIndex = 0; ComponentIndex < Parameters.NumComponents; ++ComponentIndex)
 	{
 		const int32 SwizzledComponentIndex = Parameters.ComponentIndex[ComponentIndex];
-		ResultType.SetComponentEvaluationType(ComponentIndex, InputType.GetComponentEvaluationType(SwizzledComponentIndex));
+		ResultType.SetComponentData(ComponentIndex, InputType.GetComponentData(SwizzledComponentIndex));
 	}
 
 	OutResult.SetType(Context, RequestedType, ResultType);
 }
 
-void FExpressionSwizzle::EmitValueShader(FEmitContext& Context, const FRequestedType& RequestedType, FShaderValue& OutShader) const
+void FExpressionSwizzle::EmitValueShader(FEmitContext& Context, const FRequestedType& RequestedType, FEmitShaderResult& OutShader) const
 {
 	const FRequestedType RequestedInputType = Parameters.GetRequestedInputType(RequestedType);
 
@@ -656,17 +729,19 @@ void FExpressionSwizzle::EmitValueShader(FEmitContext& Context, const FRequested
 		}
 	}
 
+	const FShaderValue InputValue = Input->GetValueShader(Context, RequestedInputType);
+
 	OutShader.Type = RequestedType.GetType();
 	OutShader.bInline = true;
 	if (bHasSwizzleReorder)
 	{
 		OutShader.Code.Appendf(TEXT("%s.%s"),
-			Input->GetValueShader(Context, RequestedInputType),
+			InputValue.Code,
 			Swizzle);
 	}
 	else
 	{
-		OutShader.Code.Appendf(TEXT("%s"), Input->GetValueShader(Context, RequestedInputType));
+		OutShader.Code.Appendf(TEXT("%s"), InputValue.Code);
 	}
 }
 
@@ -685,12 +760,17 @@ void FExpressionSwizzle::EmitValuePreshader(FEmitContext& Context, const FReques
 
 void FExpressionAppend::PrepareValue(FEmitContext& Context, const FRequestedType& RequestedType, FPrepareValueResult& OutResult) const
 {
-	// TODO - better handling of requested types for inputs
+	const FPreparedType& LhsType = PrepareExpressionValue(Context, Lhs, RequestedType);
 	const int32 NumRequestedComponents = RequestedType.GetNumComponents();
-	const FPreparedType& LhsType = PrepareExpressionValue(Context, Lhs, NumRequestedComponents);
 	const int32 NumLhsComponents = FMath::Min(LhsType.GetNumComponents(), NumRequestedComponents);
 
-	const FPreparedType& RhsType = PrepareExpressionValue(Context, Rhs, NumRequestedComponents - NumLhsComponents);
+	FRequestedType RhsRequestedType;
+	for (int32 Index = NumLhsComponents; Index < NumRequestedComponents; ++Index)
+	{
+		RhsRequestedType.SetComponentRequest(Index - NumLhsComponents, RequestedType.GetComponentRequest(Index));
+	}
+
+	const FPreparedType& RhsType = PrepareExpressionValue(Context, Rhs, RhsRequestedType);
 	const int32 NumRhsComponents = FMath::Min(RhsType.GetNumComponents(), NumRequestedComponents - NumLhsComponents);
 
 	if (LhsType.ValueComponentType != RhsType.ValueComponentType)
@@ -701,11 +781,11 @@ void FExpressionAppend::PrepareValue(FEmitContext& Context, const FRequestedType
 	FPreparedType ResultType(LhsType.ValueComponentType);
 	for (int32 Index = 0; Index < NumLhsComponents; ++Index)
 	{
-		ResultType.SetComponentEvaluationType(Index, LhsType.GetComponentEvaluationType(Index));
+		ResultType.SetComponentData(Index, LhsType.GetComponentData(Index));
 	}
 	for (int32 Index = 0; Index < NumRhsComponents; ++Index)
 	{
-		ResultType.SetComponentEvaluationType(NumLhsComponents + Index, LhsType.GetComponentEvaluationType(Index));
+		ResultType.SetComponentData(NumLhsComponents + Index, LhsType.GetComponentData(Index));
 	}
 
 	OutResult.SetType(Context, RequestedType, ResultType);
@@ -731,32 +811,34 @@ FAppendTypes GetAppendTypes(const FRequestedType& RequestedType, const FRequeste
 	Types.RhsType.ValueComponentType = RequestedType.ValueComponentType;
 	for (int32 Index = 0; Index < NumLhsComponents; ++Index)
 	{
-		Types.LhsType.SetComponentRequested(Index, RequestedType.IsComponentRequested(Index));
+		Types.LhsType.SetComponentRequest(Index, RequestedType.GetComponentRequest(Index));
 	}
 	for (int32 Index = 0; Index < NumRhsComponents; ++Index)
 	{
-		Types.RhsType.SetComponentRequested(Index, RequestedType.IsComponentRequested(NumLhsComponents + Index));
+		Types.RhsType.SetComponentRequest(Index, RequestedType.GetComponentRequest(NumLhsComponents + Index));
 	}
 	return Types;
 }
 } // namespace Private
 
-void FExpressionAppend::EmitValueShader(FEmitContext& Context, const FRequestedType& RequestedType, FShaderValue& OutShader) const
+void FExpressionAppend::EmitValueShader(FEmitContext& Context, const FRequestedType& RequestedType, FEmitShaderResult& OutShader) const
 {
 	const Private::FAppendTypes Types = Private::GetAppendTypes(RequestedType, Lhs->GetRequestedType(), Rhs->GetRequestedType());
+	const FShaderValue LhsValue = Lhs->GetValueShader(Context, Types.LhsType);
 
 	OutShader.bInline = true;
 	OutShader.Type = Types.ResultType;
 	if (Types.RhsType.IsVoid())
 	{
-		OutShader.Code.Appendf(TEXT("%s"), Lhs->GetValueShader(Context, Types.LhsType));
+		OutShader.Code.Appendf(TEXT("%s"), LhsValue.Code);
 	}
 	else
 	{
+		const FShaderValue RhsValue = Rhs->GetValueShader(Context, Types.RhsType);
 		OutShader.Code.Appendf(TEXT("%s(%s, %s)"),
 			RequestedType.GetName(),
-			Lhs->GetValueShader(Context, Types.LhsType),
-			Rhs->GetValueShader(Context, Types.RhsType));
+			LhsValue.Code,
+			RhsValue.Code);
 	}
 }
 
@@ -774,10 +856,10 @@ void FExpressionAppend::EmitValuePreshader(FEmitContext& Context, const FRequest
 void FExpressionReflectionVector::PrepareValue(FEmitContext& Context, const FRequestedType& RequestedType, FPrepareValueResult& OutResult) const
 {
 	Context.bReadMaterialNormal = true;
-	OutResult.SetType(Context, RequestedType, EExpressionEvaluationType::Shader, Shader::EValueType::Float3);
+	OutResult.SetType(Context, RequestedType, EExpressionEvaluation::Shader, Shader::EValueType::Float3);
 }
 
-void FExpressionReflectionVector::EmitValueShader(FEmitContext& Context, const FRequestedType& RequestedType, FShaderValue& OutShader) const
+void FExpressionReflectionVector::EmitValueShader(FEmitContext& Context, const FRequestedType& RequestedType, FEmitShaderResult& OutShader) const
 {
 	OutShader.bInline = true;
 	OutShader.Type = Shader::EValueType::Float3;
@@ -799,13 +881,13 @@ void FStatementReturn::Prepare(FEmitContext& Context) const
 
 void FStatementReturn::EmitShader(FEmitContext& Context) const
 {
-	ParentScope->EmitStatementf(Context, TEXT("return %s;"), Expression->GetValueShader(Context));
+	ParentScope->EmitStatementf(Context, TEXT("return %s;"), Expression->GetValueShader(Context).Code);
 }
 
 void FStatementIf::Prepare(FEmitContext& Context) const
 {
 	const FPreparedType& ConditionType = PrepareExpressionValue(Context, ConditionExpression, 1);
-	if (ConditionType.GetEvaluationType(Shader::EValueType::Bool1) == EExpressionEvaluationType::Constant)
+	if (ConditionType.GetEvaluation(Shader::EValueType::Bool1) == EExpressionEvaluation::Constant)
 	{
 		const bool bCondition = ConditionExpression->GetValueConstant(Context, Shader::EValueType::Bool1).AsBoolScalar();
 		if (bCondition)
@@ -828,7 +910,7 @@ void FStatementIf::Prepare(FEmitContext& Context) const
 
 void FStatementIf::EmitShader(FEmitContext& Context) const
 {
-	if (ConditionExpression->GetEvaluationType(Shader::EValueType::Bool1) == EExpressionEvaluationType::Constant)
+	if (ConditionExpression->GetEvaluation(Shader::EValueType::Bool1) == EExpressionEvaluation::Constant)
 	{
 		const bool bCondition = ConditionExpression->GetValueConstant(Context, Shader::EValueType::Bool1).AsBoolScalar();
 		if (bCondition && IsScopeLive(ThenScope))
@@ -844,8 +926,8 @@ void FStatementIf::EmitShader(FEmitContext& Context) const
 	{
 		if (IsScopeLive(ThenScope) || IsScopeLive(ElseScope))
 		{
-			ParentScope->EmitNestedScopef(Context, ThenScope, TEXT("if (%s)"),
-				ConditionExpression->GetValueShader(Context, Shader::EValueType::Bool1));
+			const FShaderValue ConditionValue = ConditionExpression->GetValueShader(Context, Shader::EValueType::Bool1);
+			ParentScope->EmitNestedScopef(Context, ThenScope, TEXT("if (%s)"), ConditionValue.Code);
 			if (IsScopeLive(ElseScope))
 			{
 				ParentScope->EmitNestedScopef(Context, ElseScope, TEXT("else"));
diff --git a/Engine/Source/Runtime/Engine/Private/Materials/MaterialHLSLEmitter.cpp b/Engine/Source/Runtime/Engine/Private/Materials/MaterialHLSLEmitter.cpp
--- a/Engine/Source/Runtime/Engine/Private/Materials/MaterialHLSLEmitter.cpp
+++ b/Engine/Source/Runtime/Engine/Private/Materials/MaterialHLSLEmitter.cpp
@@ -786,7 +786,7 @@ bool MaterialEmitHLSL(const FMaterialCompileTargetParameters& InCompilerTarget,
 		}
 
 		// Prepare code for just the normal
-		FRequestedType RequestedMaterialNormal(Generator.GetMaterialAttributesType(), false);
+		FRequestedType RequestedMaterialNormal(Generator.GetMaterialAttributesType(), EComponentRequest::None);
 		RequestedMaterialNormal.SetFieldRequested(NormalField);
 		PrepareExpressionValue(EmitContext, Generator.GetResultExpression(), RequestedMaterialNormal);
 
diff --git a/Engine/Source/Runtime/Engine/Private/Shader/ShaderValue.cpp b/Engine/Source/Runtime/Engine/Private/Shader/ShaderValue.cpp
--- a/Engine/Source/Runtime/Engine/Private/Shader/ShaderValue.cpp
+++ b/Engine/Source/Runtime/Engine/Private/Shader/ShaderValue.cpp
@@ -23,6 +23,19 @@ const TCHAR* FType::GetName() const
 	return TypeDesc.Name;
 }
 
+FType FType::GetDerivativeType() const
+{
+	if (StructType)
+	{
+		check(ValueType == EValueType::Struct);
+		return StructType->DerivativeType; // will convert to 'Void' if DerivativeType is nullptr
+	}
+
+	check(ValueType != EValueType::Struct);
+	const FValueTypeDescription TypeDesc = GetValueTypeDescription(ValueType);
+	return MakeValueType(EValueComponentType::Float, TypeDesc.NumComponents);
+}
+
 int32 FType::GetNumComponents() const
 {
 	if (StructType)
@@ -563,7 +576,7 @@ FValueTypeDescription GetValueTypeDescription(EValueType Type)
 
 EValueType MakeValueType(EValueComponentType ComponentType, int32 NumComponents)
 {
-	if (NumComponents == 0)
+	if (ComponentType == EValueComponentType::Void || NumComponents == 0)
 	{
 		return EValueType::Void;
 	}
@@ -758,6 +771,8 @@ const FStructType* FStructTypeRegistry::NewType(const FStructTypeInitializer& In
 	FSHA1 Hasher;
 	Hasher.UpdateWithString(Initializer.Name.GetData(), Initializer.Name.Len());
 
+	TArray<FStructFieldInitializer, TInlineAllocator<16>> DerivativeFields;
+
 	const int32 NumFields = Initializer.Fields.Num();
 	int32 ComponentIndex = 0;
 	int32 FlatFieldIndex = 0;
@@ -784,6 +799,15 @@ const FStructType* FStructTypeRegistry::NewType(const FStructTypeInitializer& In
 		Field.FlatFieldIndex = FlatFieldIndex;
 		ComponentIndex += FieldType.GetNumComponents();
 		FlatFieldIndex += FieldType.GetNumFlatFields();
+
+		if (!Initializer.bIsDerivativeType)
+		{
+			const FType FieldDerivativeType = FieldType.GetDerivativeType();
+			if (!FieldDerivativeType.IsVoid())
+			{
+				DerivativeFields.Emplace(FieldInitializer.Name, FieldDerivativeType);
+			}
+		}
 	}
 
 	EValueComponentType* ComponentTypes = new(*Allocator) EValueComponentType[ComponentIndex];
@@ -805,6 +829,16 @@ const FStructType* FStructTypeRegistry::NewType(const FStructTypeInitializer& In
 
 	Types.Add(StructType->Hash, StructType);
 
+	if (DerivativeFields.Num() > 0)
+	{
+		const FString DerivativeTypeName = FString(Initializer.Name) + TEXT("_Deriv");
+		FStructTypeInitializer DerivativeTypeInitializer;
+		DerivativeTypeInitializer.Name = DerivativeTypeName;
+		DerivativeTypeInitializer.Fields = DerivativeFields;
+		DerivativeTypeInitializer.bIsDerivativeType = true;
+		StructType->DerivativeType = NewType(DerivativeTypeInitializer);
+	}
+
 	return StructType;
 }
 
diff --git a/Engine/Source/Runtime/Engine/Public/HLSLTree/HLSLTree.h b/Engine/Source/Runtime/Engine/Public/HLSLTree/HLSLTree.h
--- a/Engine/Source/Runtime/Engine/Public/HLSLTree/HLSLTree.h
+++ b/Engine/Source/Runtime/Engine/Public/HLSLTree/HLSLTree.h
@@ -36,7 +36,7 @@ class FScope;
 class FExpressionLocalPHI;
 class FRequestedType;
 struct FEmitShaderValue;
-struct FShaderValue;
+struct FEmitShaderResult;
 
 static constexpr int32 MaxNumPreviousScopes = 2;
 
@@ -87,6 +87,7 @@ using FEmitShaderValueDependencies = TArray<FEmitShaderValue*, TInlineAllocator<
 struct FEmitShaderValueContext
 {
 	FEmitShaderValueDependencies Dependencies;
+	FEmitShaderValueDependencies DerivativeDependencies;
 };
 
 /** Tracks shared state while emitting HLSL code */
@@ -101,8 +102,9 @@ public:
 	/** Get a unique local variable name */
 	const TCHAR* AcquireLocalDeclarationCode();
 
-	FEmitShaderValue* AcquireShader(FScope* Scope, const FShaderValue& Shader, TArrayView<FEmitShaderValue*> Dependencies);
+	FEmitShaderValue* AcquireShader(FScope* Scope, const Shader::FType& Type, FStringView Code, bool bInline, TArrayView<FEmitShaderValue*> Dependencies);
 	FEmitShaderValue* AcquirePreshaderOrConstant(const FRequestedType& RequestedType, FScope* Scope, FExpression* Expression);
+	FEmitShaderValue* AcquireConstantZero(FScope* Scope, const Shader::FType& Type);
 
 	FEmitShaderValue* CastShaderValue(FNode* Node, FScope* Scope, FEmitShaderValue* ShaderValue, const Shader::FType& DestType);
 
@@ -152,6 +154,15 @@ public:
 	bool bEmitShader = false;
 };
 
+enum class EComponentRequest : uint8
+{
+	None,
+	Requested,
+	RequestedWithDerivative,
+};
+inline bool IsRequested(EComponentRequest InRequest) { return InRequest != EComponentRequest::None; }
+inline bool IsDerivativeRequested(EComponentRequest InRequest) { return InRequest == EComponentRequest::RequestedWithDerivative; }
+
 /**
  * Like Shader::FType, but tracks which individual components are needed
  */
@@ -159,9 +170,9 @@ class FRequestedType
 {
 public:
 	FRequestedType() = default;
-	FRequestedType(int32 NumComponents, bool bDefaultRequest = true) : RequestedComponents(bDefaultRequest, NumComponents) {}
-	FRequestedType(const Shader::FType& InType, bool bDefaultRequest = true);
-	FRequestedType(const Shader::EValueType& InType, bool bDefaultRequest = true);
+	FRequestedType(int32 NumComponents, EComponentRequest DefaultRequest = EComponentRequest::Requested);
+	FRequestedType(const Shader::FType& InType, EComponentRequest DefaultRequest = EComponentRequest::Requested);
+	FRequestedType(const Shader::EValueType& InType, EComponentRequest DefaultRequest = EComponentRequest::Requested);
 	
 	const TCHAR* GetName() const { return GetType().GetName(); }
 	bool IsStruct() const { return StructType != nullptr; }
@@ -170,54 +181,32 @@ public:
 
 	int32 GetNumComponents() const;
 	bool IsComponentRequested(int32 Index) const { return RequestedComponents.IsValidIndex(Index) ? (bool)RequestedComponents[Index] : false; }
+	bool IsComponentDerivativeRequested(int32 Index) const { return RequestedComponentDerivatives.IsValidIndex(Index) ? (bool)RequestedComponentDerivatives[Index] : false; }
 	bool IsVoid() const { return RequestedComponents.Find(true) == INDEX_NONE; }
 
-	bool Merge(const FRequestedType& OtherType);
+	EComponentRequest GetComponentRequest(int32 Index) const;
+	void SetComponentRequest(int32 Index, EComponentRequest Request = EComponentRequest::Requested);
 
 	void Reset()
 	{
 		StructType = nullptr;
 		ValueComponentType = Shader::EValueComponentType::Void;
 		RequestedComponents.Reset();
-	}
-
-	void SetComponentRequested(int32 Index, bool bRequested = true)
-	{
-		if (bRequested)
-		{
-			RequestedComponents.PadToNum(Index + 1, false);
-		}
-		if (RequestedComponents.IsValidIndex(Index))
-		{
-			RequestedComponents[Index] = bRequested;
-		}
+		RequestedComponentDerivatives.Reset();
 	}
 
 	/** Marks the given field as requested (or not) */
-	void SetFieldRequested(const Shader::FStructField* Field, bool bRequested = true)
-	{
-		RequestedComponents.SetRange(Field->ComponentIndex, Field->GetNumComponents(), bRequested);
-	}
-
+	void SetFieldRequested(const Shader::FStructField* Field, EComponentRequest Request = EComponentRequest::Requested);
 	void ClearFieldRequested(const Shader::FStructField* Field)
 	{
-		SetFieldRequested(Field, false);
+		SetFieldRequested(Field, EComponentRequest::None);
 	}
 
 	/** Marks the given field as requested, based on the input request type (which should match the field type) */
-	void SetField(const Shader::FStructField* Field, const FRequestedType& InRequest)
-	{
-		ensure(InRequest.GetNumComponents() == Field->GetNumComponents());
-		RequestedComponents.SetRangeFromRange(Field->ComponentIndex, InRequest.GetNumComponents(), InRequest.RequestedComponents, 0);
-	}
+	void SetField(const Shader::FStructField* Field, const FRequestedType& InRequest);
 
 	/** Returns the requested type of the given field */
-	FRequestedType GetField(const Shader::FStructField* Field) const
-	{
-		FRequestedType Result(Field->Type);
-		Result.RequestedComponents.SetRangeFromRange(0, Field->GetNumComponents(), RequestedComponents, Field->ComponentIndex);
-		return Result;
-	}
+	FRequestedType GetField(const Shader::FStructField* Field) const;
 
 	/**
 	 * If either StructType or ValueComponentType are set, then the request is for an explicit type
@@ -228,12 +217,59 @@ public:
 
 	/** 1 bit per component, a value of 'true' means the specified component is requsted */
 	TBitArray<> RequestedComponents;
+	TBitArray<> RequestedComponentDerivatives;
 };
 
 FRequestedType MakeRequestedType(Shader::EValueComponentType ComponentType, const FRequestedType& RequestedComponents);
 
+struct FPreparedData
+{
+	FPreparedData() = default;
+
+	FPreparedData(EExpressionEvaluation InEvaluation, EExpressionDerivative InDerivative = EExpressionDerivative::Invalid)
+		: Evaluation(InEvaluation)
+	{
+		switch (InEvaluation)
+		{
+		case EExpressionEvaluation::None:
+			Derivative = EExpressionDerivative::None;
+			break;
+		case EExpressionEvaluation::Constant:
+		case EExpressionEvaluation::Preshader:
+			Derivative = EExpressionDerivative::Zero;
+			break;
+		default:
+			Derivative = InDerivative;
+			break;
+		}
+	}
+
+	bool IsValid() const { return Evaluation != EExpressionEvaluation::None; }
+
+	EComponentRequest GetRequest() const
+	{
+		if (Evaluation != EExpressionEvaluation::None)
+		{
+			if (Derivative == EExpressionDerivative::Valid || Derivative == EExpressionDerivative::Zero)
+			{
+				return  EComponentRequest::RequestedWithDerivative;
+			}
+			return  EComponentRequest::Requested;
+		}
+		return EComponentRequest::None;
+	}
+
+	EExpressionEvaluation Evaluation = EExpressionEvaluation::None;
+	EExpressionDerivative Derivative = EExpressionDerivative::None;
+};
+
+inline FPreparedData CombinePreparedData(const FPreparedData& Lhs, const FPreparedData& Rhs)
+{
+	return FPreparedData(CombineEvaluations(Lhs.Evaluation, Rhs.Evaluation), CombineDerivatives(Lhs.Derivative, Rhs.Derivative));
+}
+
 /**
- * Like FRequestedType, but tracks an EExpressionEvaluationType per component, rather than a simple requested flag
+ * Like FRequestedType, but tracks an EExpressionEvaluation per component, rather than a simple requested flag
  */
 class FPreparedType
 {
@@ -243,7 +279,7 @@ public:
 	FPreparedType(const Shader::FStructType* InStructType) : StructType(InStructType) {}
 	FPreparedType(const Shader::FType& InType);
 
-	void SetEvaluationType(EExpressionEvaluationType EvaluationType);
+	void SetEvaluation(EExpressionEvaluation Evaluation);
 
 	void SetField(const Shader::FStructField* Field, const FPreparedType& FieldType);
 	FPreparedType GetFieldType(const Shader::FStructField* Field) const;
@@ -255,31 +291,38 @@ public:
 	bool IsNumeric() const { return !IsVoid() && ValueComponentType != Shader::EValueComponentType::Void; }
 	bool IsInitialized() const { return StructType != nullptr || ValueComponentType != Shader::EValueComponentType::Void; }
 	bool IsVoid() const;
-	EExpressionEvaluationType GetEvaluationType(const FRequestedType& RequestedType) const;
-	EExpressionEvaluationType GetFieldEvaluationType(int32 ComponentIndex, int32 NumComponents) const;
 
-	EExpressionEvaluationType GetComponentEvaluationType(int32 Index) const
-	{
-		return ComponentEvaluationType.IsValidIndex(Index) ? ComponentEvaluationType[Index] : EExpressionEvaluationType::None;
-	}
+	FPreparedData GetData() const;
+	FPreparedData GetData(const FRequestedType& RequestedType) const;
+	FPreparedData GetFieldData(int32 ComponentIndex, int32 NumComponents) const;
+	FPreparedData GetComponentData(int32 Index) const;
+
+	inline EExpressionEvaluation GetEvaluation(const FRequestedType& RequestedType) const { return GetData(RequestedType).Evaluation; }
 
-	void SetComponentEvaluationType(int32 Index, EExpressionEvaluationType EvaluationType);
+	void SetComponentData(int32 Index, const FPreparedData& Data);
+	void MergeComponentData(int32 Index, EComponentRequest Request, const FPreparedData& Data);
 
 	/** Unlike FRequestedType, one of these should be set */
 	const Shader::FStructType* StructType = nullptr;
 	Shader::EValueComponentType ValueComponentType = Shader::EValueComponentType::Void;
 
 	/** Evaluation type for each component, may be 'None' for components that are unused */
-	TArray<EExpressionEvaluationType, TInlineAllocator<16>> ComponentEvaluationType;
+	TArray<FPreparedData, TInlineAllocator<16>> PreparedComponents;
 };
 
 FPreparedType MergePreparedTypes(const FPreparedType& Lhs, const FPreparedType& Rhs);
 
-struct FShaderValue
+struct FEmitShaderResult
 {
-	explicit FShaderValue(FStringBuilderBase& InCode) : Code(InCode) {}
+	FEmitShaderResult(FStringBuilderBase& InCode, FStringBuilderBase& InCodeDdx, FStringBuilderBase& InCodeDdy)
+		: Code(InCode)
+		, CodeDdx(InCodeDdx)
+		, CodeDdy(InCodeDdy)
+	{}
 
 	FStringBuilderBase& Code;
+	FStringBuilderBase& CodeDdx;
+	FStringBuilderBase& CodeDdy;
 	Shader::FType Type;
 	bool bInline = false;
 };
@@ -289,9 +332,14 @@ class FPrepareValueResult
 public:
 	const FPreparedType& GetPreparedType() const { return PreparedType; }
 
-	void SetType(FEmitContext& Context, const FRequestedType& RequestedType, EExpressionEvaluationType EvaluationType, const Shader::FType& Type);
-	void SetType(FEmitContext& Context, const FRequestedType& RequestedType, EExpressionEvaluationType EvaluationType, Shader::EValueComponentType ComponentType);
+	void SetType(FEmitContext& Context, const FRequestedType& RequestedType, const FPreparedData& Data, const Shader::FType& Type);
 	void SetType(FEmitContext& Context, const FRequestedType& RequestedType, const FPreparedType& Type);
+
+	void SetTypeWithDerivative(FEmitContext& Context, const FRequestedType& RequestedType, EExpressionEvaluation Evaluation, const Shader::FType& Type)
+	{
+		SetType(Context, RequestedType, FPreparedData(Evaluation, EExpressionDerivative::Valid), Type);
+	}
+
 	void SetForwardValue(FEmitContext& Context, const FRequestedType& RequestedType, FExpression* InValue);
 
 private:
@@ -317,6 +365,15 @@ struct FEmitShaderValue
 	FSHAHash Hash;
 };
 
+struct FShaderValue
+{
+	const TCHAR* Code = nullptr;
+	const TCHAR* CodeDdx = nullptr;
+	const TCHAR* CodeDdy = nullptr;
+	Shader::FType Type;
+	bool bHasDerivatives = false;
+};
+
 /**
  * Represents an HLSL expression.  This is a piece of code that evaluates to a value, but has no side effects.
  * Unlike statements, expressions are not expected to execute in any particular order.  They may be cached (or not) in generated code, without the underlying implementation needing to care.
@@ -329,21 +386,22 @@ public:
 	const FPreparedType& GetPreparedType() const { return PrepareValueResult.PreparedType; }
 	FRequestedType GetRequestedType() const { return PrepareValueResult.PreparedType.GetRequestedType(); }
 	Shader::FType GetType() const { return PrepareValueResult.PreparedType.GetType(); }
-	EExpressionEvaluationType GetEvaluationType(const FRequestedType& RequestedType) const { return PrepareValueResult.PreparedType.GetEvaluationType(RequestedType); }
+	EExpressionEvaluation GetEvaluation(const FRequestedType& RequestedType) const { return PrepareValueResult.PreparedType.GetEvaluation(RequestedType); }
 
 	virtual void Reset() override;
 
 	friend const FPreparedType& PrepareExpressionValue(FEmitContext& Context, FExpression* InExpression, const FRequestedType& RequestedType);
 
-	const TCHAR* GetValueShader(FEmitContext& Context, const FRequestedType& RequestedType);
+	FShaderValue GetValueShader(FEmitContext& Context, const FRequestedType& RequestedType, const Shader::FType& ResultType);
 	void GetValuePreshader(FEmitContext& Context, const FRequestedType& RequestedType, Shader::FPreshaderData& OutPreshader);
 	Shader::FValue GetValueConstant(FEmitContext& Context, const FRequestedType& RequestedType);
 
-	const TCHAR* GetValueShader(FEmitContext& Context);
+	FShaderValue GetValueShader(FEmitContext& Context, const FRequestedType& RequestedType);
+	FShaderValue GetValueShader(FEmitContext& Context);
 
 protected:
 	virtual void PrepareValue(FEmitContext& Context, const FRequestedType& RequestedType, FPrepareValueResult& OutResult) const = 0;
-	virtual void EmitValueShader(FEmitContext& Context, const FRequestedType& RequestedType, FShaderValue& OutShader) const;
+	virtual void EmitValueShader(FEmitContext& Context, const FRequestedType& RequestedType, FEmitShaderResult& OutShader) const;
 	virtual void EmitValuePreshader(FEmitContext& Context, const FRequestedType& RequestedType, Shader::FPreshaderData& OutPreshader) const;
 
 private:
@@ -364,7 +422,7 @@ class FExpressionLocalPHI final : public FExpression
 {
 public:
 	virtual void PrepareValue(FEmitContext& Context, const FRequestedType& RequestedType, FPrepareValueResult& OutResult) const override;
-	virtual void EmitValueShader(FEmitContext& Context, const FRequestedType& RequestedType, FShaderValue& OutShader) const override;
+	virtual void EmitValueShader(FEmitContext& Context, const FRequestedType& RequestedType, FEmitShaderResult& OutShader) const override;
 
 	FName LocalName;
 	FScope* Scopes[MaxNumPreviousScopes];
@@ -572,7 +630,7 @@ private:
 };
 
 //Shader::EValueType RequestExpressionType(FEmitContext& Context, FExpression* InExpression, int8 InRequestedNumComponents); // friend of FExpression
-//EExpressionEvaluationType PrepareExpressionValue(FEmitContext& Context, FExpression* InExpression); // friend of FExpression
+//EExpressionEvaluation PrepareExpressionValue(FEmitContext& Context, FExpression* InExpression); // friend of FExpression
 //void PrepareScopeValues(FEmitContext& Context, const FScope* InScope); // friend of FScope
 
 } // namespace HLSLTree
diff --git a/Engine/Source/Runtime/Engine/Public/HLSLTree/HLSLTreeCommon.h b/Engine/Source/Runtime/Engine/Public/HLSLTree/HLSLTreeCommon.h
--- a/Engine/Source/Runtime/Engine/Public/HLSLTree/HLSLTreeCommon.h
+++ b/Engine/Source/Runtime/Engine/Public/HLSLTree/HLSLTreeCommon.h
@@ -124,8 +124,8 @@ public:
 
 	EExternalInputType InputType;
 
-	virtual void PrepareValue(FEmitContext& Context, const FRequestedType& RequestedType, FPrepareValueResult& OutResult) const override { OutResult.SetType(Context, RequestedType, EExpressionEvaluationType::Shader, GetInputExpressionType(InputType)); }
-	virtual void EmitValueShader(FEmitContext& Context, const FRequestedType& RequestedType, FShaderValue& OutShader) const override;
+	virtual void PrepareValue(FEmitContext& Context, const FRequestedType& RequestedType, FPrepareValueResult& OutResult) const override;
+	virtual void EmitValueShader(FEmitContext& Context, const FRequestedType& RequestedType, FEmitShaderResult& OutShader) const override;
 };
 
 class FExpressionTextureSample : public FExpression
@@ -144,7 +144,7 @@ public:
 	ETextureMipValueMode MipValueMode;
 
 	virtual void PrepareValue(FEmitContext& Context, const FRequestedType& RequestedType, FPrepareValueResult& OutResult) const override;
-	virtual void EmitValueShader(FEmitContext& Context, const FRequestedType& RequestedType, FShaderValue& OutShader) const override;
+	virtual void EmitValueShader(FEmitContext& Context, const FRequestedType& RequestedType, FEmitShaderResult& OutShader) const override;
 };
 
 class FExpressionGetStructField : public FExpression
@@ -163,7 +163,7 @@ public:
 	FExpression* StructExpression;
 
 	virtual void PrepareValue(FEmitContext& Context, const FRequestedType& RequestedType, FPrepareValueResult& OutResult) const override;
-	virtual void EmitValueShader(FEmitContext& Context, const FRequestedType& RequestedType, FShaderValue& OutShader) const override;
+	virtual void EmitValueShader(FEmitContext& Context, const FRequestedType& RequestedType, FEmitShaderResult& OutShader) const override;
 	virtual void EmitValuePreshader(FEmitContext& Context, const FRequestedType& RequestedType, Shader::FPreshaderData& OutPreshader) const override;
 };
 
@@ -185,7 +185,7 @@ public:
 	FExpression* FieldExpression;
 
 	virtual void PrepareValue(FEmitContext& Context, const FRequestedType& RequestedType, FPrepareValueResult& OutResult) const override;
-	virtual void EmitValueShader(FEmitContext& Context, const FRequestedType& RequestedType, FShaderValue& OutShader) const override;
+	virtual void EmitValueShader(FEmitContext& Context, const FRequestedType& RequestedType, FEmitShaderResult& OutShader) const override;
 	virtual void EmitValuePreshader(FEmitContext& Context, const FRequestedType& RequestedType, Shader::FPreshaderData& OutPreshader) const override;
 };
 
@@ -203,7 +203,7 @@ public:
 	FExpression* FalseExpression;
 
 	virtual void PrepareValue(FEmitContext& Context, const FRequestedType& RequestedType, FPrepareValueResult& OutResult) const override;
-	virtual void EmitValueShader(FEmitContext& Context, const FRequestedType& RequestedType, FShaderValue& OutShader) const override;
+	virtual void EmitValueShader(FEmitContext& Context, const FRequestedType& RequestedType, FEmitShaderResult& OutShader) const override;
 	virtual void EmitValuePreshader(FEmitContext& Context, const FRequestedType& RequestedType, Shader::FPreshaderData& OutPreshader) const override;
 };
 
@@ -221,7 +221,7 @@ public:
 	FExpression* Rhs;
 
 	virtual void PrepareValue(FEmitContext& Context, const FRequestedType& RequestedType, FPrepareValueResult& OutResult) const override;
-	virtual void EmitValueShader(FEmitContext& Context, const FRequestedType& RequestedType, FShaderValue& OutShader) const override;
+	virtual void EmitValueShader(FEmitContext& Context, const FRequestedType& RequestedType, FEmitShaderResult& OutShader) const override;
 	virtual void EmitValuePreshader(FEmitContext& Context, const FRequestedType& RequestedType, Shader::FPreshaderData& OutPreshader) const override;
 };
 
@@ -250,7 +250,7 @@ public:
 	FExpression* Input;
 
 	virtual void PrepareValue(FEmitContext& Context, const FRequestedType& RequestedType, FPrepareValueResult& OutResult) const override;
-	virtual void EmitValueShader(FEmitContext& Context, const FRequestedType& RequestedType, FShaderValue& OutShader) const override;
+	virtual void EmitValueShader(FEmitContext& Context, const FRequestedType& RequestedType, FEmitShaderResult& OutShader) const override;
 	virtual void EmitValuePreshader(FEmitContext& Context, const FRequestedType& RequestedType, Shader::FPreshaderData& OutPreshader) const override;
 };
 
@@ -266,7 +266,7 @@ public:
 	FExpression* Rhs;
 
 	virtual void PrepareValue(FEmitContext& Context, const FRequestedType& RequestedType, FPrepareValueResult& OutResult) const override;
-	virtual void EmitValueShader(FEmitContext& Context, const FRequestedType& RequestedType, FShaderValue& OutShader) const override;
+	virtual void EmitValueShader(FEmitContext& Context, const FRequestedType& RequestedType, FEmitShaderResult& OutShader) const override;
 	virtual void EmitValuePreshader(FEmitContext& Context, const FRequestedType& RequestedType, Shader::FPreshaderData& OutPreshader) const override;
 };
 
@@ -274,7 +274,7 @@ class FExpressionReflectionVector : public FExpression
 {
 public:
 	virtual void PrepareValue(FEmitContext& Context, const FRequestedType& RequestedType, FPrepareValueResult& OutResult) const override;
-	virtual void EmitValueShader(FEmitContext& Context, const FRequestedType& RequestedType, FShaderValue& OutShader) const override;
+	virtual void EmitValueShader(FEmitContext& Context, const FRequestedType& RequestedType, FEmitShaderResult& OutShader) const override;
 };
 
 class FStatementReturn : public FStatement
diff --git a/Engine/Source/Runtime/Engine/Public/HLSLTree/HLSLTreeTypes.h b/Engine/Source/Runtime/Engine/Public/HLSLTree/HLSLTreeTypes.h
--- a/Engine/Source/Runtime/Engine/Public/HLSLTree/HLSLTreeTypes.h
+++ b/Engine/Source/Runtime/Engine/Public/HLSLTree/HLSLTreeTypes.h
@@ -22,7 +22,7 @@ class FTree;
 
 /**
  * Describes how a given expression needs to be evaluated */
-enum class EExpressionEvaluationType : uint8
+enum class EExpressionEvaluation : uint8
 {
 	/** Invalid/uninitialized */
 	None,
@@ -37,7 +37,24 @@ enum class EExpressionEvaluationType : uint8
 	Constant,
 };
 
-EExpressionEvaluationType CombineEvaluationTypes(EExpressionEvaluationType Lhs, EExpressionEvaluationType Rhs);
+EExpressionEvaluation CombineEvaluations(EExpressionEvaluation Lhs, EExpressionEvaluation Rhs);
+
+enum class EExpressionDerivative : uint8
+{
+	/** Uninitialized/unknown */
+	None,
+
+	/** Derivative is not valid */
+	Invalid,
+
+	/** Derivative is known to be 0 */
+	Zero,
+
+	/** Derivative is valid */
+	Valid,
+};
+
+EExpressionDerivative CombineDerivatives(EExpressionDerivative Lhs, EExpressionDerivative Rhs);
 
 struct FTextureDescription
 {
diff --git a/Engine/Source/Runtime/Engine/Public/Shader/ShaderTypes.h b/Engine/Source/Runtime/Engine/Public/Shader/ShaderTypes.h
--- a/Engine/Source/Runtime/Engine/Public/Shader/ShaderTypes.h
+++ b/Engine/Source/Runtime/Engine/Public/Shader/ShaderTypes.h
@@ -79,6 +79,7 @@ struct FType
 	FType(const FStructType* InStruct) : StructType(InStruct), ValueType(InStruct ? EValueType::Struct : EValueType::Void) {}
 
 	const TCHAR* GetName() const;
+	FType GetDerivativeType() const;
 	bool IsVoid() const { return ValueType == EValueType::Void; }
 	bool IsStruct() const { return ValueType == EValueType::Struct; }
 	bool IsNumeric() const { return !IsVoid() && !IsStruct(); }
@@ -139,6 +140,7 @@ struct FStructType
 {
 	uint64 Hash;
 	const TCHAR* Name;
+	const FStructType* DerivativeType;
 	TArrayView<const FStructField> Fields;
 
 	/**
@@ -168,6 +170,7 @@ struct FStructTypeInitializer
 {
 	FStringView Name;
 	TArrayView<const FStructFieldInitializer> Fields;
+	bool bIsDerivativeType = false;
 };
 
 class FStructTypeRegistry
