PackedLevelActor: use base FWorldPartitionActorDesc

#tests AncientGame (PIE and Cook)

diff --git a/Engine/Source/Runtime/Core/Public/UObject/UE5ReleaseStreamObjectVersion.h b/Engine/Source/Runtime/Core/Public/UObject/UE5ReleaseStreamObjectVersion.h
--- a/Engine/Source/Runtime/Core/Public/UObject/UE5ReleaseStreamObjectVersion.h
+++ b/Engine/Source/Runtime/Core/Public/UObject/UE5ReleaseStreamObjectVersion.h
@@ -69,6 +69,9 @@ struct CORE_API FUE5ReleaseStreamObjectVersion
 		// Fixup for bad default value for GridPlacement_DEPRECATED
 		ActorGridPlacementDeprecateDefaultValueFixup,
 
+		// PackedLevelActor started using FWorldPartitionActorDesc (not currently checked against but added as a security)
+		PackedLevelActorUseWorldPartitionActorDesc,
+
 		// -----<new versions can be added above this line>-------------------------------------------------
 		VersionPlusOne,
 		LatestVersion = VersionPlusOne - 1
diff --git a/Engine/Source/Runtime/Engine/Private/PackedLevelActor/PackedLevelActor.cpp b/Engine/Source/Runtime/Engine/Private/PackedLevelActor/PackedLevelActor.cpp
--- a/Engine/Source/Runtime/Engine/Private/PackedLevelActor/PackedLevelActor.cpp
+++ b/Engine/Source/Runtime/Engine/Private/PackedLevelActor/PackedLevelActor.cpp
@@ -84,6 +84,11 @@ void APackedLevelActor::RerunConstructionScripts()
 	}
 }
 
+TUniquePtr<class FWorldPartitionActorDesc> APackedLevelActor::CreateClassActorDesc() const
+{
+	return AActor::CreateClassActorDesc();
+}
+
 bool APackedLevelActor::CreateOrUpdateBlueprint(ALevelInstance* InLevelInstance, TSoftObjectPtr<UBlueprint> InBlueprintAsset, bool bCheckoutAndSave, bool bPromptForSave)
 {
 	return FPackedLevelActorBuilder::CreateDefaultBuilder()->CreateOrUpdateBlueprint(InLevelInstance, InBlueprintAsset, bCheckoutAndSave, bPromptForSave);
diff --git a/Engine/Source/Runtime/Engine/Public/PackedLevelActor/PackedLevelActor.h b/Engine/Source/Runtime/Engine/Public/PackedLevelActor/PackedLevelActor.h
--- a/Engine/Source/Runtime/Engine/Public/PackedLevelActor/PackedLevelActor.h
+++ b/Engine/Source/Runtime/Engine/Public/PackedLevelActor/PackedLevelActor.h
@@ -31,6 +31,8 @@ public:
 
 	virtual void Serialize(FArchive& Ar) override;
 #if WITH_EDITOR
+	virtual TUniquePtr<class FWorldPartitionActorDesc> CreateClassActorDesc() const override;
+	
 	static bool CreateOrUpdateBlueprint(ALevelInstance* InLevelInstance, TSoftObjectPtr<UBlueprint> InBlueprintAsset, bool bCheckoutAndSave = true, bool bPromptForSave = true);
 	static bool CreateOrUpdateBlueprint(TSoftObjectPtr<UWorld> InWorldAsset, TSoftObjectPtr<UBlueprint> InBlueprintAsset, bool bCheckoutAndSave = true, bool bPromptForSave = true);
 
@@ -85,4 +87,7 @@ private:
 	UPROPERTY()
 	FGuid PackedVersion;
 #endif
-};
\ No newline at end of file
+};
+
+
+DEFINE_ACTORDESC_TYPE(APackedLevelActor, FWorldPartitionActorDesc);
\ No newline at end of file
