Optimize tick record context data overheads

When resetting tick records each frame, a fairly large shared ptr destruction chain occured due to the final references to the various bits of context data (e.g. state machine scopes, mirror scopes) being removed.
This change switches individual shared ptrs to context data to unique ptrs, while allowing shared ownership (e.g. via the FAnimNotifyEventReference) via holding a shared ptr to the entire array.

diff --git a/Engine/Source/Runtime/AnimGraphRuntime/Private/AnimNodes/AnimNode_BlendSpacePlayer.cpp b/Engine/Source/Runtime/AnimGraphRuntime/Private/AnimNodes/AnimNode_BlendSpacePlayer.cpp
--- a/Engine/Source/Runtime/AnimGraphRuntime/Private/AnimNodes/AnimNode_BlendSpacePlayer.cpp
+++ b/Engine/Source/Runtime/AnimGraphRuntime/Private/AnimNodes/AnimNode_BlendSpacePlayer.cpp
@@ -92,10 +92,7 @@ void FAnimNode_BlendSpacePlayer::UpdateInternal(const FAnimationUpdateContext& C
 		TickRecord.DeltaTimeRecord = &DeltaTimeRecord;
 
 		UE::Anim::FAnimSyncParams SyncParams(GetGroupName(), GetGroupRole(), GetGroupMethod());
-		if(Context.GetSharedContext())
-		{
-			Context.GetSharedContext()->MessageStack.MakeEventContextData(TickRecord.ContextData);
-		}
+		TickRecord.GatherContextData(Context);
 
 		SyncScope.AddTickRecord(TickRecord, SyncParams, UE::Anim::FAnimSyncDebugInfo(Context));
 
diff --git a/Engine/Source/Runtime/AnimGraphRuntime/Private/AnimNodes/AnimNode_RandomPlayer.cpp b/Engine/Source/Runtime/AnimGraphRuntime/Private/AnimNodes/AnimNode_RandomPlayer.cpp
--- a/Engine/Source/Runtime/AnimGraphRuntime/Private/AnimNodes/AnimNode_RandomPlayer.cpp
+++ b/Engine/Source/Runtime/AnimGraphRuntime/Private/AnimNodes/AnimNode_RandomPlayer.cpp
@@ -233,11 +233,8 @@ void FAnimNode_RandomPlayer::Update_AnyThread(const FAnimationUpdateContext& Con
 
 	FAnimTickRecord TickRecord(CurrentData->Entry->Sequence, true, CurrentData->PlayRate, CurrentData->BlendWeight, CurrentData->CurrentPlayTime, CurrentData->MarkerTickRecord);
 	TickRecord.DeltaTimeRecord = &CurrentData->DeltaTimeRecord;
+	TickRecord.GatherContextData(Context);
 
-	if(Context.GetSharedContext())
-	{
-		Context.GetSharedContext()->MessageStack.MakeEventContextData(TickRecord.ContextData);
-	}
 	UE::Anim::FAnimSyncGroupScope& SyncScope = Context.GetMessageChecked<UE::Anim::FAnimSyncGroupScope>();
 	SyncScope.AddTickRecord(TickRecord, UE::Anim::FAnimSyncParams(), UE::Anim::FAnimSyncDebugInfo(Context));
 
@@ -247,11 +244,8 @@ void FAnimNode_RandomPlayer::Update_AnyThread(const FAnimationUpdateContext& Con
 	{
 		FAnimTickRecord NextTickRecord(NextData->Entry->Sequence, true, NextData->PlayRate, NextData->BlendWeight, NextData->CurrentPlayTime, NextData->MarkerTickRecord);
 		NextTickRecord.DeltaTimeRecord = &NextData->DeltaTimeRecord;
+		NextTickRecord.GatherContextData(Context);
 
-		if(Context.GetSharedContext())
-		{
-			Context.GetSharedContext()->MessageStack.MakeEventContextData(NextTickRecord.ContextData);
-		}
 		SyncScope.AddTickRecord(NextTickRecord, UE::Anim::FAnimSyncParams(), UE::Anim::FAnimSyncDebugInfo(Context));
 
 		TRACE_ANIM_TICK_RECORD(Context, NextTickRecord);
diff --git a/Engine/Source/Runtime/Engine/Classes/Animation/AnimationAsset.h b/Engine/Source/Runtime/Engine/Classes/Animation/AnimationAsset.h
--- a/Engine/Source/Runtime/Engine/Classes/Animation/AnimationAsset.h
+++ b/Engine/Source/Runtime/Engine/Classes/Animation/AnimationAsset.h
@@ -27,6 +27,7 @@ class UAnimSequenceBase;
 class UBlendSpace;
 class UPoseAsset;
 class UMirrorDataTable;
+struct FAnimationUpdateContext;
 
 namespace UE { namespace Anim {
 class IAnimNotifyEventContextDataInterface;
@@ -354,7 +355,8 @@ struct FAnimTickRecord
 	bool bLooping = false;
 	const UMirrorDataTable* MirrorDataTable = nullptr;
 
-	TArray<TSharedPtr<const UE::Anim::IAnimNotifyEventContextDataInterface>> ContextData;	
+	TSharedPtr<TArray<TUniquePtr<const UE::Anim::IAnimNotifyEventContextDataInterface>>> ContextData;
+
 	union
 	{
 		struct
@@ -420,6 +422,9 @@ public:
 	// Create a tick record for a pose asset
 	ENGINE_API FAnimTickRecord(UPoseAsset* InPoseAsset, float InFinalBlendWeight);
 
+	// Gather any data from the current update context
+	ENGINE_API void GatherContextData(const FAnimationUpdateContext& InContext);
+	
 	/** This can be used with the Sort() function on a TArray of FAnimTickRecord to sort from higher leader score */
 	ENGINE_API bool operator <(const FAnimTickRecord& Other) const { return LeaderScore > Other.LeaderScore; }
 };
diff --git a/Engine/Source/Runtime/Engine/Private/Animation/ActiveStateMachineScope.cpp b/Engine/Source/Runtime/Engine/Private/Animation/ActiveStateMachineScope.cpp
--- a/Engine/Source/Runtime/Engine/Private/Animation/ActiveStateMachineScope.cpp
+++ b/Engine/Source/Runtime/Engine/Private/Animation/ActiveStateMachineScope.cpp
@@ -15,14 +15,14 @@ FEncounteredStateMachineStack::FEncounteredStateMachineStack(int32 InStateMachin
 	StateStack.Emplace(InStateMachineIndex, InStateIndex);
 }
 
-TSharedPtr<const FEncounteredStateMachineStack> FEncounteredStateMachineStack::InitStack(const FEncounteredStateMachineStack& ParentStack, int32 InStateMachineIndex, int32 InStateIndex)
+FEncounteredStateMachineStack FEncounteredStateMachineStack::InitStack(const FEncounteredStateMachineStack& ParentStack, int32 InStateMachineIndex, int32 InStateIndex)
 {
-	return MakeShared<const FEncounteredStateMachineStack>(ParentStack, InStateMachineIndex, InStateIndex); 
+	return FEncounteredStateMachineStack(ParentStack, InStateMachineIndex, InStateIndex); 
 }
 
-TSharedPtr<const FEncounteredStateMachineStack> FEncounteredStateMachineStack::InitStack(int32 InStateMachineIndex, int32 InStateIndex)
+FEncounteredStateMachineStack FEncounteredStateMachineStack::InitStack(int32 InStateMachineIndex, int32 InStateIndex)
 {
-	return MakeShared<const FEncounteredStateMachineStack>(InStateMachineIndex, InStateIndex);
+	return FEncounteredStateMachineStack(InStateMachineIndex, InStateIndex);
 }
 
 
@@ -31,19 +31,14 @@ namespace UE { namespace Anim {
 
 IMPLEMENT_NOTIFY_CONTEXT_INTERFACE(FAnimNotifyStateMachineContext)
 
-FAnimNotifyStateMachineContext::FAnimNotifyStateMachineContext(const TSharedPtr<const FEncounteredStateMachineStack>& InEncounteredStateMachines)
+FAnimNotifyStateMachineContext::FAnimNotifyStateMachineContext(const FEncounteredStateMachineStack& InEncounteredStateMachines)
+	: EncounteredStateMachines(InEncounteredStateMachines)
 {
-	EncounteredStateMachines = InEncounteredStateMachines;
 }
 
 bool FAnimNotifyStateMachineContext::IsStateMachineInContext(int32 StateMachineIndex) const
 {
-	if (!EncounteredStateMachines.IsValid())
-	{
-		return false; 
-	}
-
-	for (const FEncounteredStateMachineStack::FStateMachineEntry& Entry : EncounteredStateMachines->StateStack)
+	for (const FEncounteredStateMachineStack::FStateMachineEntry& Entry : EncounteredStateMachines.StateStack)
 	{
 		if (Entry.StateMachineIndex == StateMachineIndex)
 		{
@@ -55,12 +50,7 @@ bool FAnimNotifyStateMachineContext::IsStateMachineInContext(int32 StateMachineI
 
 bool FAnimNotifyStateMachineContext::IsStateInStateMachineInContext(int32 StateMachineIndex, int32 StateIndex) const
 {
-	if (!EncounteredStateMachines.IsValid())
-	{
-		return false;
-	}
-
-	for (const FEncounteredStateMachineStack::FStateMachineEntry& Entry : EncounteredStateMachines->StateStack)
+	for (const FEncounteredStateMachineStack::FStateMachineEntry& Entry : EncounteredStateMachines.StateStack)
 	{
 		if (Entry.StateMachineIndex == StateMachineIndex && Entry.StateIndex == StateIndex)
 		{
@@ -73,19 +63,18 @@ bool FAnimNotifyStateMachineContext::IsStateInStateMachineInContext(int32 StateM
 IMPLEMENT_ANIMGRAPH_MESSAGE(FActiveStateMachineScope);
 
 FActiveStateMachineScope::FActiveStateMachineScope(const FAnimationBaseContext& InContext, FAnimNode_StateMachine* StateMachine, int32 InStateIndex)
-{
-	int32 StateMachineIndex = GetStateMachineIndex(StateMachine, InContext);
-	if (FActiveStateMachineScope* ParentStateMachineScope = InContext.GetMessage<FActiveStateMachineScope>())
-	{
-		ActiveStateMachines = FEncounteredStateMachineStack::InitStack(*(ParentStateMachineScope->ActiveStateMachines.Get()), StateMachineIndex, InStateIndex); 
-	}
-	else
+	: ActiveStateMachines([&InContext, StateMachine, InStateIndex]()
 	{
-		ActiveStateMachines = FEncounteredStateMachineStack::InitStack(StateMachineIndex, InStateIndex);
-	}
-}
-
-FActiveStateMachineScope::~FActiveStateMachineScope()
+		const int32 StateMachineIndex = GetStateMachineIndex(StateMachine, InContext);
+		if (FActiveStateMachineScope* ParentStateMachineScope = InContext.GetMessage<FActiveStateMachineScope>())
+		{
+			return FEncounteredStateMachineStack::InitStack(ParentStateMachineScope->ActiveStateMachines, StateMachineIndex, InStateIndex);
+		}
+		else
+		{
+			return FEncounteredStateMachineStack::InitStack(StateMachineIndex, InStateIndex);
+		}
+	}())
 {
 }
 
@@ -100,9 +89,9 @@ int32 FActiveStateMachineScope::GetStateMachineIndex(FAnimNode_StateMachine* Sta
 	return INDEX_NONE;
 }
 
-TSharedPtr<const IAnimNotifyEventContextDataInterface> FActiveStateMachineScope::MakeEventContextData() const
+TUniquePtr<const IAnimNotifyEventContextDataInterface> FActiveStateMachineScope::MakeUniqueEventContextData() const
 {
-	return MakeShared<const FAnimNotifyStateMachineContext>(ActiveStateMachines);
+	return MakeUnique<const FAnimNotifyStateMachineContext>(ActiveStateMachines);
 }
 	
 }}	// namespace UE::Anim
\ No newline at end of file
diff --git a/Engine/Source/Runtime/Engine/Private/Animation/AnimNodeMessages.cpp b/Engine/Source/Runtime/Engine/Private/Animation/AnimNodeMessages.cpp
--- a/Engine/Source/Runtime/Engine/Private/Animation/AnimNodeMessages.cpp
+++ b/Engine/Source/Runtime/Engine/Private/Animation/AnimNodeMessages.cpp
@@ -189,7 +189,7 @@ void FMessageStack::CopyForCachedUpdate(const FMessageStack& InStack)
 	}
 }
 
-void FMessageStack::MakeEventContextData(TArray<TSharedPtr<const IAnimNotifyEventContextDataInterface>>& ContextData) const
+void FMessageStack::MakeEventContextData(TArray<TUniquePtr<const IAnimNotifyEventContextDataInterface>>& ContextData) const
 {
 	ContextData.Empty();
 	for (const TPair<FName, FMessageStackEntry>& MessageStackPair : MessageStacks)
@@ -198,13 +198,12 @@ void FMessageStack::MakeEventContextData(TArray<TSharedPtr<const IAnimNotifyEven
 		if(StackEntry.Num())
 		{
 			const FMessageEntry& MessageEntry = StackEntry.Top();
-			TSharedPtr<const IAnimNotifyEventContextDataInterface> NotifyData = MessageEntry.Message->MakeEventContextData();
-			if(NotifyData)
+			if(TUniquePtr<const IAnimNotifyEventContextDataInterface> NotifyData = MessageEntry.Message->MakeUniqueEventContextData())
 			{
-				ContextData.Emplace(NotifyData);
+				ContextData.Add(MoveTemp(NotifyData));
 			}
 		}
 	}
 }
-	
+
 }}	// namespace UE::Anim
\ No newline at end of file
diff --git a/Engine/Source/Runtime/Engine/Private/Animation/AnimNode_AssetPlayerBase.cpp b/Engine/Source/Runtime/Engine/Private/Animation/AnimNode_AssetPlayerBase.cpp
--- a/Engine/Source/Runtime/Engine/Private/Animation/AnimNode_AssetPlayerBase.cpp
+++ b/Engine/Source/Runtime/Engine/Private/Animation/AnimNode_AssetPlayerBase.cpp
@@ -40,10 +40,8 @@ void FAnimNode_AssetPlayerBase::CreateTickRecordForNode(const FAnimationUpdateCo
 
 	const UE::Anim::FAnimSyncParams SyncParams(GroupNameToUse, SyncGroupRole, MethodToUse);
 	FAnimTickRecord TickRecord(Sequence, bLooping, PlayRate, FinalBlendWeight, /*inout*/ InternalTimeAccumulator, MarkerTickRecord);
-	if(Context.GetSharedContext())
-	{
-		Context.GetSharedContext()->MessageStack.MakeEventContextData(TickRecord.ContextData);
-	}
+	TickRecord.GatherContextData(Context);
+
 	TickRecord.RootMotionWeightModifier = Context.GetRootMotionWeightModifier();
 	TickRecord.DeltaTimeRecord = &DeltaTimeRecord;
 
diff --git a/Engine/Source/Runtime/Engine/Private/Animation/AnimNotifyQueue.cpp b/Engine/Source/Runtime/Engine/Private/Animation/AnimNotifyQueue.cpp
--- a/Engine/Source/Runtime/Engine/Private/Animation/AnimNotifyQueue.cpp
+++ b/Engine/Source/Runtime/Engine/Private/Animation/AnimNotifyQueue.cpp
@@ -12,9 +12,9 @@ bool operator==(const FAnimNotifyEventReference& Lhs, const FAnimNotifyEvent& Rh
 	return false;
 }
 
-void FAnimNotifyEventReference::SetContextDataInterfaces(const TArray<TSharedPtr<const UE::Anim::IAnimNotifyEventContextDataInterface>>& InContextData)
+void FAnimNotifyEventReference::GatherTickRecordData(const FAnimTickRecord& InTickRecord)
 {
-	ContextData = InContextData; 
+	ContextData = InTickRecord.ContextData;
 }
 
 bool FAnimNotifyQueue::PassesFiltering(const FAnimNotifyEvent* Notify) const
diff --git a/Engine/Source/Runtime/Engine/Private/Animation/AnimNotifyStateMachineInspectionLibrary.cpp b/Engine/Source/Runtime/Engine/Private/Animation/AnimNotifyStateMachineInspectionLibrary.cpp
--- a/Engine/Source/Runtime/Engine/Private/Animation/AnimNotifyStateMachineInspectionLibrary.cpp
+++ b/Engine/Source/Runtime/Engine/Private/Animation/AnimNotifyStateMachineInspectionLibrary.cpp
@@ -72,7 +72,7 @@ bool  UAnimNotifyStateMachineInspectionLibrary::IsTriggeredByState(const FAnimNo
 		const UE::Anim::FAnimNotifyStateMachineContext * StateMachineContext = EventReference.GetContextData<UE::Anim::FAnimNotifyStateMachineContext>();
 		if (StateMachineContext)
 		{
-			for (const FEncounteredStateMachineStack::FStateMachineEntry& Entry: StateMachineContext->EncounteredStateMachines->StateStack)
+			for (const FEncounteredStateMachineStack::FStateMachineEntry& Entry: StateMachineContext->EncounteredStateMachines.StateStack)
 			{
 				const FAnimNode_StateMachine* StateMachineInstance = AnimInstance->GetStateMachineInstance(Entry.StateMachineIndex);
 				if(StateMachineInstance->GetStateInfo(Entry.StateIndex).StateName == StateName)
diff --git a/Engine/Source/Runtime/Engine/Private/Animation/AnimSequenceBase.cpp b/Engine/Source/Runtime/Engine/Private/Animation/AnimSequenceBase.cpp
--- a/Engine/Source/Runtime/Engine/Private/Animation/AnimSequenceBase.cpp
+++ b/Engine/Source/Runtime/Engine/Private/Animation/AnimSequenceBase.cpp
@@ -396,7 +396,7 @@ void UAnimSequenceBase::GetAnimNotifiesFromDeltaPositions(const float& PreviousP
 				if (NotifyContext.TickRecord)
 				{
 					NotifyContext.ActiveNotifies.Emplace(&AnimNotifyEvent, this, NotifyContext.TickRecord->MirrorDataTable);
-					NotifyContext.ActiveNotifies.Top().SetContextDataInterfaces(NotifyContext.TickRecord->ContextData);
+					NotifyContext.ActiveNotifies.Top().GatherTickRecordData(*NotifyContext.TickRecord);
 				}
 				else
 				{
@@ -418,7 +418,7 @@ void UAnimSequenceBase::GetAnimNotifiesFromDeltaPositions(const float& PreviousP
 				if (NotifyContext.TickRecord)
 				{
 					NotifyContext.ActiveNotifies.Emplace(&AnimNotifyEvent, this, NotifyContext.TickRecord->MirrorDataTable);
-					NotifyContext.ActiveNotifies.Top().SetContextDataInterfaces(NotifyContext.TickRecord->ContextData); 
+					NotifyContext.ActiveNotifies.Top().GatherTickRecordData(*NotifyContext.TickRecord); 
 				}
 				else
 				{
diff --git a/Engine/Source/Runtime/Engine/Private/Animation/MirrorSyncScope.cpp b/Engine/Source/Runtime/Engine/Private/Animation/MirrorSyncScope.cpp
--- a/Engine/Source/Runtime/Engine/Private/Animation/MirrorSyncScope.cpp
+++ b/Engine/Source/Runtime/Engine/Private/Animation/MirrorSyncScope.cpp
@@ -72,9 +72,9 @@ int32 FMirrorSyncScope::GetMirrorScopeDepth() const
 	return MirrorScopeDepth;
 }
 
-TSharedPtr<const IAnimNotifyEventContextDataInterface> FMirrorSyncScope::MakeEventContextData() const
+TUniquePtr<const IAnimNotifyEventContextDataInterface> FMirrorSyncScope::MakeUniqueEventContextData() const
 {
-	return MakeShared<FAnimNotifyMirrorContext>(MirrorDataTable);
+	return MakeUnique<FAnimNotifyMirrorContext>(MirrorDataTable);
 }
 
 }}	// namespace UE::Anim
\ No newline at end of file
diff --git a/Engine/Source/Runtime/Engine/Private/AnimationAsset.cpp b/Engine/Source/Runtime/Engine/Private/AnimationAsset.cpp
--- a/Engine/Source/Runtime/Engine/Private/AnimationAsset.cpp
+++ b/Engine/Source/Runtime/Engine/Private/AnimationAsset.cpp
@@ -11,6 +11,7 @@
 #include "Animation/BlendSpace.h"
 #include "Animation/PoseAsset.h"
 #include "Animation/AnimSequenceHelpers.h"
+#include "Animation/AnimNodeBase.h"
 
 #define LOCTEXT_NAMESPACE "AnimationAsset"
 
@@ -231,6 +232,24 @@ FAnimTickRecord::FAnimTickRecord(UPoseAsset* InPoseAsset, float InFinalBlendWeig
 	EffectiveBlendWeight = InFinalBlendWeight;
 }
 
+void FAnimTickRecord::GatherContextData(const FAnimationUpdateContext& InContext)
+{
+	if(InContext.GetSharedContext())
+	{
+		TArray<TUniquePtr<const UE::Anim::IAnimNotifyEventContextDataInterface>> NewContextData;
+		InContext.GetSharedContext()->MessageStack.MakeEventContextData(NewContextData);
+		if(NewContextData.Num())
+		{
+			if(!ContextData.IsValid())
+			{
+				ContextData = MakeShared<TArray<TUniquePtr<const UE::Anim::IAnimNotifyEventContextDataInterface>>>();
+			}
+
+			ContextData->Append(MoveTemp(NewContextData));
+		}
+	}
+}
+
 //////////////////////////////////////////////////////////////////////////
 // UAnimationAsset
 
diff --git a/Engine/Source/Runtime/Engine/Public/Animation/ActiveStateMachineScope.h b/Engine/Source/Runtime/Engine/Public/Animation/ActiveStateMachineScope.h
--- a/Engine/Source/Runtime/Engine/Public/Animation/ActiveStateMachineScope.h
+++ b/Engine/Source/Runtime/Engine/Public/Animation/ActiveStateMachineScope.h
@@ -31,8 +31,8 @@ struct ENGINE_API FEncounteredStateMachineStack
 		int32 StateIndex = INDEX_NONE;;
 	};
 	TArray<FStateMachineEntry, TInlineAllocator<4>> StateStack;
-	static TSharedPtr<const FEncounteredStateMachineStack> InitStack(const FEncounteredStateMachineStack& ParentStack, int32 InStateMachineIndex, int32 InStateIndex);
-	static TSharedPtr<const FEncounteredStateMachineStack> InitStack(int32 InStateMachineIndex, int32 InStateIndex);
+	static FEncounteredStateMachineStack InitStack(const FEncounteredStateMachineStack& ParentStack, int32 InStateMachineIndex, int32 InStateIndex);
+	static FEncounteredStateMachineStack InitStack(int32 InStateMachineIndex, int32 InStateIndex);
 };
 
 namespace UE { namespace Anim {
@@ -40,23 +40,23 @@ namespace UE { namespace Anim {
 class ENGINE_API FAnimNotifyStateMachineContext : public UE::Anim::IAnimNotifyEventContextDataInterface
 {
 	DECLARE_NOTIFY_CONTEXT_INTERFACE(FAnimNotifyStateMachineContext)
-	public:
-	FAnimNotifyStateMachineContext(const TSharedPtr<const FEncounteredStateMachineStack>& InEncounteredStateMachines);
+public:
+	FAnimNotifyStateMachineContext(const FEncounteredStateMachineStack& InEncounteredStateMachines);
 	bool IsStateMachineInContext(int32 StateMachineIndex) const;
 	bool IsStateInStateMachineInContext(int32 StateMachineIndex, int32 StateIndex) const; 
-	TSharedPtr<const FEncounteredStateMachineStack> EncounteredStateMachines; 
+	const FEncounteredStateMachineStack EncounteredStateMachines;
 };
 
 class ENGINE_API FActiveStateMachineScope : public UE::Anim::IGraphMessage
 {
 	DECLARE_ANIMGRAPH_MESSAGE(FActiveStateMachineScope);
-	public:
+public:
 	FActiveStateMachineScope(const FAnimationBaseContext& InContext, FAnimNode_StateMachine* StateMachine, int32 InStateIndex);
-	virtual ~FActiveStateMachineScope();
+
 	static int32 GetStateMachineIndex(FAnimNode_StateMachine* StateMachine, const FAnimationBaseContext& Context);
-	virtual TSharedPtr<const IAnimNotifyEventContextDataInterface> MakeEventContextData() const override;
+	virtual TUniquePtr<const IAnimNotifyEventContextDataInterface> MakeUniqueEventContextData() const override;
 	
-	TSharedPtr<const FEncounteredStateMachineStack> ActiveStateMachines; 
+	const FEncounteredStateMachineStack ActiveStateMachines; 
 };
 
 }}	// namespace UE::Anim
diff --git a/Engine/Source/Runtime/Engine/Public/Animation/AnimNodeMessages.h b/Engine/Source/Runtime/Engine/Public/Animation/AnimNodeMessages.h
--- a/Engine/Source/Runtime/Engine/Public/Animation/AnimNodeMessages.h
+++ b/Engine/Source/Runtime/Engine/Public/Animation/AnimNodeMessages.h
@@ -99,12 +99,18 @@ public:
 		return *static_cast<Type*>(this);
 	}
 
+	UE_DEPRECATED(5.0, "No longer in use, please use MakeUniqueEventContextData")
 	virtual TSharedPtr<const IAnimNotifyEventContextDataInterface> MakeEventContextData() const
 	{
 		TSharedPtr<const IAnimNotifyEventContextDataInterface> NullPtr;
 		return NullPtr; 
 	}
 
+	virtual TUniquePtr<const IAnimNotifyEventContextDataInterface> MakeUniqueEventContextData() const
+	{
+		return TUniquePtr<const IAnimNotifyEventContextDataInterface>(); 
+	}
+
 private:
 	static const FName BaseTypeName;
 };
@@ -272,7 +278,7 @@ public:
 
 	// Call MakeEventContextData for the top entry of each MessageType, returning interfaces for types that return event data
 	// @param	ContextData		An array of valid IAnimNotifyEventContextDataInterface.  Only one entry will exist per message type
-	void MakeEventContextData(TArray<TSharedPtr<const IAnimNotifyEventContextDataInterface>>& ContextData) const;
+	void MakeEventContextData(TArray<TUniquePtr<const IAnimNotifyEventContextDataInterface>>& ContextData) const;
 
 private:
 	// Push a message onto the stack
diff --git a/Engine/Source/Runtime/Engine/Public/Animation/AnimNotifyQueue.h b/Engine/Source/Runtime/Engine/Public/Animation/AnimNotifyQueue.h
--- a/Engine/Source/Runtime/Engine/Public/Animation/AnimNotifyQueue.h
+++ b/Engine/Source/Runtime/Engine/Public/Animation/AnimNotifyQueue.h
@@ -81,20 +81,25 @@ struct FAnimNotifyEventReference
 	template<typename Type> 
 	const Type* GetContextData() const 
 	{
-		for(const TSharedPtr<const UE::Anim::IAnimNotifyEventContextDataInterface>& DataInterface :  ContextData)
-        {
-        	if (DataInterface->Is<Type>())
-        	{
-        		return &(DataInterface->As<Type>()); 
-        	}
-        }
-        return nullptr; 
+		if(ContextData.IsValid())
+		{
+			for(const TUniquePtr<const UE::Anim::IAnimNotifyEventContextDataInterface>& DataInterface : *ContextData)
+			{
+				if (DataInterface->Is<Type>())
+				{
+					return &(DataInterface->As<Type>()); 
+				}
+			}
+		}
+		return nullptr; 
 	}
 
-	void SetContextDataInterfaces(const TArray<TSharedPtr<const UE::Anim::IAnimNotifyEventContextDataInterface>>& InContextData);
-private:
+	// Pulls relevant data from the tick record
+	void GatherTickRecordData(const FAnimTickRecord& InTickRecord);
 
-	TArray<TSharedPtr<const UE::Anim::IAnimNotifyEventContextDataInterface>> ContextData;
+private:
+	// Context data gleaned from the tick record
+	TSharedPtr<TArray<TUniquePtr<const UE::Anim::IAnimNotifyEventContextDataInterface>>> ContextData;
 	
 	const FAnimNotifyEvent* Notify;
 
diff --git a/Engine/Source/Runtime/Engine/Public/Animation/MirrorSyncScope.h b/Engine/Source/Runtime/Engine/Public/Animation/MirrorSyncScope.h
--- a/Engine/Source/Runtime/Engine/Public/Animation/MirrorSyncScope.h
+++ b/Engine/Source/Runtime/Engine/Public/Animation/MirrorSyncScope.h
@@ -29,7 +29,7 @@ public:
 	FMirrorSyncScope(const FAnimationBaseContext& InContext, const UMirrorDataTable* InMirrorDataTable);
 	virtual ~FMirrorSyncScope();
 	int32 GetMirrorScopeDepth() const;
-	TSharedPtr<const IAnimNotifyEventContextDataInterface> MakeEventContextData() const override;
+	virtual TUniquePtr<const IAnimNotifyEventContextDataInterface> MakeUniqueEventContextData() const override;
 private:
 	const UMirrorDataTable* MirrorDataTable = nullptr;
 	int32 MirrorScopeDepth = 1; 
