Fixed a crash opening MirrorDataTables

diff --git a/Engine/Source/Runtime/Engine/Private/Animation/MirrorDataTable.cpp b/Engine/Source/Runtime/Engine/Private/Animation/MirrorDataTable.cpp
--- a/Engine/Source/Runtime/Engine/Private/Animation/MirrorDataTable.cpp
+++ b/Engine/Source/Runtime/Engine/Private/Animation/MirrorDataTable.cpp
@@ -63,12 +63,12 @@ UMirrorDataTable::UMirrorDataTable(const FObjectInitializer& ObjectInitializer)
 void UMirrorDataTable::PostLoad()
 {
 	Super::PostLoad();
-	FillMirrorArrays(); 
 	if (Skeleton)
 	{
 		Skeleton->ConditionalPostLoad();
 		Skeleton->OnSmartNamesChangedEvent.AddUObject(this, &UMirrorDataTable::FillMirrorArrays);
 	}
+	FillMirrorArrays(); 
 }
 
 #if WITH_EDITOR
@@ -391,9 +391,13 @@ void UMirrorDataTable::FillMirrorArrays()
 		}
 	);
 
+	// Ensure post load prepares the smart names
+	Skeleton->ConditionalPostLoad();
+
 	//ensure that pairs always appear beside each other in the arrays
 	TSet<SmartName::UID_Type> AddedSourceUIDs; 
 	const FSmartNameMapping* CurveSmartNames = Skeleton->GetSmartNameContainer(USkeleton::AnimCurveMappingName);
+
 	if (CurveSmartNames)
 	{
 		CurveMirrorSourceUIDArray.Reset(CurveToMirrorCurveMap.Num());
