Fix crash when saving a newly created partitioned level during PIE: make sure the world doesn't keep data during the PIE session, clear everything once the original world has been duplicated for PIE. That way, if we ever save as the world while in PIE, that information isn't lost in the process and the system stays in a valid state.

diff --git a/Engine/Source/Runtime/Engine/Private/World.cpp b/Engine/Source/Runtime/Engine/Private/World.cpp
--- a/Engine/Source/Runtime/Engine/Private/World.cpp
+++ b/Engine/Source/Runtime/Engine/Private/World.cpp
@@ -3419,6 +3419,13 @@ UWorld* UWorld::GetDuplicatedWorldForPIE(UWorld* InWorld, UPackage* InPIEackage,
 
 	DuplicatedWorld->StreamingLevelsPrefix = UWorld::BuildPIEPackagePrefix(PIEInstanceID);
 
+#if WITH_EDITOR
+	if (UWorldPartition* WorldPartition = InWorld->GetWorldPartition())
+	{
+		WorldPartition->PostDuplicateWorldForPIE();
+	}
+#endif
+
 	return DuplicatedWorld;
 }
 
diff --git a/Engine/Source/Runtime/Engine/Private/WorldPartition/WorldPartition.cpp b/Engine/Source/Runtime/Engine/Private/WorldPartition/WorldPartition.cpp
--- a/Engine/Source/Runtime/Engine/Private/WorldPartition/WorldPartition.cpp
+++ b/Engine/Source/Runtime/Engine/Private/WorldPartition/WorldPartition.cpp
@@ -253,7 +253,6 @@ UWorldPartition::UWorldPartition(const FObjectInitializer& ObjectInitializer)
 	, WorldPartitionEditor(nullptr)
 	, bForceGarbageCollection(false)
 	, bForceGarbageCollectionPurge(false)
-	, bIsPIE(false)
 #endif
 	, InitState(EWorldPartitionInitState::Uninitialized)
 	, StreamingPolicy(nullptr)
@@ -314,42 +313,18 @@ void UWorldPartition::OnGCPostReachabilityAnalysis()
 
 void UWorldPartition::OnPreBeginPIE(bool bStartSimulate)
 {
-	check(!bIsPIE);
-	bIsPIE = true;
-
-	check(IsMainWorldPartition());
-
-	OnBeginPlay();
-}
-
-void UWorldPartition::OnPrePIEEnded(bool bWasSimulatingInEditor)
-{
-	check(bIsPIE);
-	bIsPIE = false;
+	OnBeginPlay(true);
 }
 
-void UWorldPartition::OnBeginPlay()
+void UWorldPartition::OnBeginPlay(bool bIsPIE)
 {
-	GenerateStreaming();
+	GenerateStreaming(nullptr, bIsPIE);
 	RuntimeHash->OnBeginPlay();
 }
 
 void UWorldPartition::OnCancelPIE()
 {
-	// No check here since CancelPIE can be called after PrePIEEnded
-	bIsPIE = false;
-	// Call OnEndPlay here since EndPlayMapDelegate is not called when cancelling PIE
-	OnEndPlay();
-}
-
-void UWorldPartition::OnEndPlay()
-{
-	check(IsMainWorldPartition());
-
-	RuntimeHash->FlushStreaming();
-	RuntimeHash->OnEndPlay();
-
-	StreamingPolicy = nullptr;
+	PostDuplicateWorldForPIE();
 }
 
 FName UWorldPartition::GetWorldPartitionEditorName() const
@@ -383,8 +358,6 @@ void UWorldPartition::Initialize(UWorld* InWorld, const FTransform& InTransform)
 	UWorld* OuterWorld = GetTypedOuter<UWorld>();
 	check(OuterWorld);
 
-	check(IsMainWorldPartition());
-
 	RegisterDelegates();
 
 	AWorldPartitionReplay::Initialize(World);
@@ -506,7 +479,7 @@ void UWorldPartition::Initialize(UWorld* InWorld, const FTransform& InTransform)
 	{
 		if (bIsGame || bPIEWorldTravel)
 		{
-			OnBeginPlay();
+			OnBeginPlay(false);
 		}
 
 		// Apply remapping of Persistent Level's SoftObjectPaths
@@ -558,11 +531,7 @@ void UWorldPartition::Uninitialize()
 			GEditor->OnEditorClose().RemoveAll(this);
 		}
 
-		if (World->IsGameWorld())
-		{
-			OnEndPlay();
-		}
-		else
+		if (!World->IsGameWorld())
 		{
 			if (!IsEngineExitRequested())
 			{
@@ -591,12 +560,6 @@ bool UWorldPartition::IsInitialized() const
 	return InitState == EWorldPartitionInitState::Initialized;
 }
 
-bool UWorldPartition::IsMainWorldPartition() const
-{
-	check(World);
-	return World == GetTypedOuter<UWorld>();
-}
-
 void UWorldPartition::OnPostBugItGoCalled(const FVector& Loc, const FRotator& Rot)
 {
 #if WITH_EDITOR
@@ -617,9 +580,7 @@ void UWorldPartition::RegisterDelegates()
 	if (GEditor && !IsTemplate() && !World->IsGameWorld())
 	{
 		FEditorDelegates::PreBeginPIE.AddUObject(this, &UWorldPartition::OnPreBeginPIE);
-		FEditorDelegates::PrePIEEnded.AddUObject(this, &UWorldPartition::OnPrePIEEnded);
 		FEditorDelegates::CancelPIE.AddUObject(this, &UWorldPartition::OnCancelPIE);
-		FGameDelegates::Get().GetEndPlayMapDelegate().AddUObject(this, &UWorldPartition::OnEndPlay);
 
 		FCoreUObjectDelegates::PostReachabilityAnalysis.AddUObject(this, &UWorldPartition::OnGCPostReachabilityAnalysis);
 
@@ -645,9 +606,7 @@ void UWorldPartition::UnregisterDelegates()
 	if (GEditor && !IsTemplate() && !World->IsGameWorld())
 	{
 		FEditorDelegates::PreBeginPIE.RemoveAll(this);
-		FEditorDelegates::PrePIEEnded.RemoveAll(this);
 		FEditorDelegates::CancelPIE.RemoveAll(this);
-		FGameDelegates::Get().GetEndPlayMapDelegate().RemoveAll(this);
 
 		if (!IsEngineExitRequested())
 		{
@@ -1244,6 +1203,14 @@ EWorldPartitionStreamingPerformance UWorldPartition::GetStreamingPerformance() c
 }
 
 #if WITH_EDITOR
+void UWorldPartition::PostDuplicateWorldForPIE()
+{
+	RuntimeHash->FlushStreaming();
+	RuntimeHash->OnEndPlay();
+
+	StreamingPolicy = nullptr;
+}
+
 void UWorldPartition::DrawRuntimeHashPreview()
 {
 	RuntimeHash->DrawPreview();
diff --git a/Engine/Source/Runtime/Engine/Private/WorldPartition/WorldPartitionStreamingGeneration.cpp b/Engine/Source/Runtime/Engine/Private/WorldPartition/WorldPartitionStreamingGeneration.cpp
--- a/Engine/Source/Runtime/Engine/Private/WorldPartition/WorldPartitionStreamingGeneration.cpp
+++ b/Engine/Source/Runtime/Engine/Private/WorldPartition/WorldPartitionStreamingGeneration.cpp
@@ -542,7 +542,7 @@ private:
 	TMap<FGuid, FActorContainerID> ActorDescViewsContainersMap;
 };
 
-bool UWorldPartition::GenerateStreaming(TArray<FString>* OutPackagesToGenerate)
+bool UWorldPartition::GenerateStreaming(TArray<FString>* OutPackagesToGenerate, bool bIsPIE)
 {
 	FActorDescList* ModifiedActorsDescList = nullptr;
 
diff --git a/Engine/Source/Runtime/Engine/Public/WorldPartition/WorldPartition.h b/Engine/Source/Runtime/Engine/Public/WorldPartition/WorldPartition.h
--- a/Engine/Source/Runtime/Engine/Public/WorldPartition/WorldPartition.h
+++ b/Engine/Source/Runtime/Engine/Public/WorldPartition/WorldPartition.h
@@ -91,10 +91,8 @@ private:
 
 	// PIE/Game Methods
 	void OnPreBeginPIE(bool bStartSimulate);
-	void OnPrePIEEnded(bool bWasSimulatingInEditor);
+	void OnBeginPlay(bool bIsPIE);
 	void OnCancelPIE();
-	void OnBeginPlay();
-	void OnEndPlay();
 
 	// UActorDescContainer events
 	virtual void OnActorDescAdded(FWorldPartitionActorDesc* NewActorDesc) override;
@@ -126,7 +124,7 @@ public:
 	bool RefreshLoadedEditorCells(bool bIsFromUserChange);
 
 	// PIE/Game/Cook Methods
-	bool GenerateStreaming(TArray<FString>* OutPackagesToGenerate = nullptr);
+	bool GenerateStreaming(TArray<FString>* OutPackagesToGenerate = nullptr, bool bIsPIE = false);
 	void RemapSoftObjectPath(FSoftObjectPath& ObjectPath);
 
 	// Cook Methods
@@ -174,6 +172,10 @@ public:
 
 	EWorldPartitionStreamingPerformance GetStreamingPerformance() const;
 
+#if WITH_EDITOR
+	void PostDuplicateWorldForPIE();
+#endif
+
 #if WITH_EDITORONLY_DATA
 	UPROPERTY(DuplicateTransient)
 	TObjectPtr<UWorldPartitionEditorHash> EditorHash;
@@ -188,13 +190,10 @@ public:
 	UPROPERTY(EditAnywhere, Category=RuntimeHash, NoClear, meta = (NoResetToDefault), Instanced)
 	TObjectPtr<UWorldPartitionRuntimeHash> RuntimeHash;
 
-#if WITH_EDITOR
+#if WITH_EDITORONLY_DATA
 	bool bForceGarbageCollection;
 	bool bForceGarbageCollectionPurge;
-	bool bIsPIE;
-#endif
 
-#if WITH_EDITORONLY_DATA
 	// Default HLOD layer
 	UPROPERTY(EditAnywhere, Category=HLOD)
 	TObjectPtr<class UHLODLayer> DefaultHLODLayer;
@@ -205,7 +204,7 @@ public:
 private:
 	EWorldPartitionInitState InitState;
 
-	UPROPERTY()
+	UPROPERTY(NonPIEDuplicateTransient)
 	mutable TObjectPtr<UWorldPartitionStreamingPolicy> StreamingPolicy;
 
 	TArray<IWorldPartitionStreamingSourceProvider*> StreamingSourceProviders;
@@ -217,7 +216,6 @@ private:
 	FWorldPartitionReference WorldDataLayersActor;
 #endif
 
-	bool IsMainWorldPartition() const;
 	void OnWorldMatchStarting();
 
 	void OnPostBugItGoCalled(const FVector& Loc, const FRotator& Rot);
