CoerceValue() allows truncating to Float1 in addition to Float, can fix invalid errors related to Lerp expressions

diff --git a/Engine/Source/Runtime/Engine/Private/Materials/HLSLMaterialTranslator.cpp b/Engine/Source/Runtime/Engine/Private/Materials/HLSLMaterialTranslator.cpp
--- a/Engine/Source/Runtime/Engine/Private/Materials/HLSLMaterialTranslator.cpp
+++ b/Engine/Source/Runtime/Engine/Private/Materials/HLSLMaterialTranslator.cpp
@@ -3242,7 +3242,7 @@ int32 FHLSLMaterialTranslator::AccessUniformExpression(int32 Index)
 FString FHLSLMaterialTranslator::CoerceValue(const FString& Code, EMaterialValueType SourceType, EMaterialValueType DestType)
 {
 	EMaterialCastFlags CastFlags = EMaterialCastFlags::ReplicateScalar;
-	if (DestType == MCT_Float || DestType == MCT_LWCScalar)
+	if (DestType == MCT_Float || DestType == MCT_Float1 || DestType == MCT_LWCScalar)
 	{
 		// CoerceValue allows truncating to scalar types only
 		CastFlags |= EMaterialCastFlags::AllowTruncate;
