Create FSceneViewState with correct FeatureLevel

diff --git a/Engine/Plugins/MovieScene/MoviePipelineMaskRenderPass/Source/MoviePipelineMaskRenderPass/Private/MoviePipelinePanoramicPass.cpp b/Engine/Plugins/MovieScene/MoviePipelineMaskRenderPass/Source/MoviePipelineMaskRenderPass/Private/MoviePipelinePanoramicPass.cpp
--- a/Engine/Plugins/MovieScene/MoviePipelineMaskRenderPass/Source/MoviePipelineMaskRenderPass/Private/MoviePipelinePanoramicPass.cpp
+++ b/Engine/Plugins/MovieScene/MoviePipelineMaskRenderPass/Source/MoviePipelineMaskRenderPass/Private/MoviePipelinePanoramicPass.cpp
@@ -162,7 +162,7 @@ void UMoviePipelinePanoramicPass::SetupImpl(const MoviePipeline::FMoviePipelineR
 
 		for (int32 Index = 0; Index < OptionalPaneViewStates.Num(); Index++)
 		{
-			OptionalPaneViewStates[Index].Allocate();
+			OptionalPaneViewStates[Index].Allocate(InPassInitSettings.FeatureLevel);
 		}
 	}
 
diff --git a/Engine/Plugins/MovieScene/MovieRenderPipeline/Source/MovieRenderPipelineCore/Private/MoviePipelineRendering.cpp b/Engine/Plugins/MovieScene/MovieRenderPipeline/Source/MovieRenderPipelineCore/Private/MoviePipelineRendering.cpp
--- a/Engine/Plugins/MovieScene/MovieRenderPipeline/Source/MovieRenderPipelineCore/Private/MoviePipelineRendering.cpp
+++ b/Engine/Plugins/MovieScene/MovieRenderPipeline/Source/MovieRenderPipelineCore/Private/MoviePipelineRendering.cpp
@@ -85,10 +85,10 @@ void UMoviePipeline::SetupRenderingPipelineForShot(UMoviePipelineExecutorShot* I
 	// Note how many tiles we wish to render with.
 	BackbufferTileCount = FIntPoint(HighResSettings->TileCount, HighResSettings->TileCount);
 
+	const ERHIFeatureLevel::Type FeatureLevel = InShot->GetWorld()->FeatureLevel;
+
 	// Initialize our render pass. This is a copy of the settings to make this less coupled to the Settings UI.
-	MoviePipeline::FMoviePipelineRenderPassInitSettings RenderPassInitSettings;
-	RenderPassInitSettings.BackbufferResolution = BackbufferResolution;
-	RenderPassInitSettings.TileCount = BackbufferTileCount;
+	const MoviePipeline::FMoviePipelineRenderPassInitSettings RenderPassInitSettings(FeatureLevel, BackbufferResolution, BackbufferTileCount);
 
 	// Code expects at least a 1x1 tile.
 	ensure(RenderPassInitSettings.TileCount.X > 0 && RenderPassInitSettings.TileCount.Y > 0);
diff --git a/Engine/Plugins/MovieScene/MovieRenderPipeline/Source/MovieRenderPipelineCore/Public/MovieRenderPipelineDataTypes.h b/Engine/Plugins/MovieScene/MovieRenderPipeline/Source/MovieRenderPipelineCore/Public/MovieRenderPipelineDataTypes.h
--- a/Engine/Plugins/MovieScene/MovieRenderPipeline/Source/MovieRenderPipelineCore/Public/MovieRenderPipelineDataTypes.h
+++ b/Engine/Plugins/MovieScene/MovieRenderPipeline/Source/MovieRenderPipelineCore/Public/MovieRenderPipelineDataTypes.h
@@ -950,18 +950,27 @@ namespace MoviePipeline
 	struct FMoviePipelineRenderPassInitSettings
 	{
 	public:
+		UE_DEPRECATED(5.0, "FMoviePipelineRenderPassInitSettings must be constructed with arguments")
 		FMoviePipelineRenderPassInitSettings()
-			: BackbufferResolution(0, 0)
-			, TileCount(0, 0)
+		{
+			FeatureLevel = GMaxRHIFeatureLevel;
+		}
+
+		FMoviePipelineRenderPassInitSettings(ERHIFeatureLevel::Type InFeatureLevel, const FIntPoint& InBackbufferResolution, const FIntPoint& InTileCount)
+			:	BackbufferResolution(InBackbufferResolution)
+			,	TileCount(InTileCount)
+			,	FeatureLevel(InFeatureLevel)
 		{
 		}
 
 	public:
 		/** This takes into account any padding needed for tiled rendering overlap. Different than the output resolution of the final image. */
-		FIntPoint BackbufferResolution;
+		FIntPoint BackbufferResolution = FIntPoint(0, 0);
 
 		/** How many tiles (in each direction) are we rendering with. */
-		FIntPoint TileCount;
+		FIntPoint TileCount = FIntPoint(0, 0);
+
+		ERHIFeatureLevel::Type FeatureLevel;
 	};
 
 }
diff --git a/Engine/Plugins/MovieScene/MovieRenderPipeline/Source/MovieRenderPipelineRenderPasses/Private/MoviePipelineDeferredPasses.cpp b/Engine/Plugins/MovieScene/MovieRenderPipeline/Source/MovieRenderPipelineRenderPasses/Private/MoviePipelineDeferredPasses.cpp
--- a/Engine/Plugins/MovieScene/MovieRenderPipeline/Source/MovieRenderPipelineRenderPasses/Private/MoviePipelineDeferredPasses.cpp
+++ b/Engine/Plugins/MovieScene/MovieRenderPipeline/Source/MovieRenderPipelineRenderPasses/Private/MoviePipelineDeferredPasses.cpp
@@ -133,7 +133,7 @@ void UMoviePipelineDeferredPassBase::SetupImpl(const MoviePipeline::FMoviePipeli
 
 	for (int32 Index = 0; Index < StencilLayerViewStates.Num(); Index++)
 	{
-		StencilLayerViewStates[Index].Allocate();
+		StencilLayerViewStates[Index].Allocate(InPassInitSettings.FeatureLevel);
 	}
 
 
diff --git a/Engine/Plugins/MovieScene/MovieRenderPipeline/Source/MovieRenderPipelineRenderPasses/Private/MoviePipelineImagePassBase.cpp b/Engine/Plugins/MovieScene/MovieRenderPipeline/Source/MovieRenderPipelineRenderPasses/Private/MoviePipelineImagePassBase.cpp
--- a/Engine/Plugins/MovieScene/MovieRenderPipeline/Source/MovieRenderPipelineRenderPasses/Private/MoviePipelineImagePassBase.cpp
+++ b/Engine/Plugins/MovieScene/MovieRenderPipeline/Source/MovieRenderPipelineRenderPasses/Private/MoviePipelineImagePassBase.cpp
@@ -36,7 +36,7 @@ void UMoviePipelineImagePassBase::SetupImpl(const MoviePipeline::FMoviePipelineR
 	Super::SetupImpl(InPassInitSettings);
 
 	// Allocate 
-	ViewState.Allocate();
+	ViewState.Allocate(InPassInitSettings.FeatureLevel);
 }
 
 void UMoviePipelineImagePassBase::TeardownImpl()
diff --git a/Engine/Source/Editor/MovieSceneTools/Private/TrackEditorThumbnail/TrackEditorThumbnail.cpp b/Engine/Source/Editor/MovieSceneTools/Private/TrackEditorThumbnail/TrackEditorThumbnail.cpp
--- a/Engine/Source/Editor/MovieSceneTools/Private/TrackEditorThumbnail/TrackEditorThumbnail.cpp
+++ b/Engine/Source/Editor/MovieSceneTools/Private/TrackEditorThumbnail/TrackEditorThumbnail.cpp
@@ -428,7 +428,7 @@ void FTrackEditorThumbnailCache::DrawViewportThumbnail(FTrackEditorThumbnail& Tr
 		FSceneViewStateInterface* Ref = ThumbnailSettings->ViewState.GetReference();
 		if (!Ref)
 		{
-			ThumbnailSettings->ViewState.Allocate();
+			ThumbnailSettings->ViewState.Allocate(ViewFamily.GetFeatureLevel());
 		}
 		ViewStateInterface = ThumbnailSettings->ViewState.GetReference();
 		break;
diff --git a/Engine/Source/Editor/UnrealEd/Private/EditorViewportClient.cpp b/Engine/Source/Editor/UnrealEd/Private/EditorViewportClient.cpp
--- a/Engine/Source/Editor/UnrealEd/Private/EditorViewportClient.cpp
+++ b/Engine/Source/Editor/UnrealEd/Private/EditorViewportClient.cpp
@@ -464,7 +464,8 @@ FEditorViewportClient::FEditorViewportClient(FEditorModeTools* InModeTools, FPre
 		ModeTools = MakeShared<FAssetEditorModeManager>();
 	}
 
-	ViewState.Allocate();
+	FSceneInterface* Scene = GetScene();
+	ViewState.Allocate(Scene ? Scene->GetFeatureLevel() : GMaxRHIFeatureLevel);
 
 	// NOTE: StereoViewState will be allocated on demand, for viewports than end up drawing in stereo
 
@@ -1196,7 +1197,8 @@ FSceneView* FEditorViewportClient::CalcSceneView(FSceneViewFamily* ViewFamily, c
 
 		if (StereoViewStates[ViewStateIndex].GetReference() == nullptr)
 		{
-			StereoViewStates[ViewStateIndex].Allocate();
+			FSceneInterface* Scene = GetScene();
+			StereoViewStates[ViewStateIndex].Allocate(Scene ? Scene->GetFeatureLevel() : GMaxRHIFeatureLevel);
 		}
 	}
 
diff --git a/Engine/Source/Runtime/Engine/Private/Components/SceneCaptureComponent.cpp b/Engine/Source/Runtime/Engine/Private/Components/SceneCaptureComponent.cpp
--- a/Engine/Source/Runtime/Engine/Private/Components/SceneCaptureComponent.cpp
+++ b/Engine/Source/Runtime/Engine/Private/Components/SceneCaptureComponent.cpp
@@ -316,7 +316,8 @@ FSceneViewStateInterface* USceneCaptureComponent::GetViewState(int32 ViewIndex)
 	FSceneViewStateInterface* ViewStateInterface = ViewStates[ViewIndex].GetReference();
 	if ((bCaptureEveryFrame || bAlwaysPersistRenderingState) && ViewStateInterface == NULL)
 	{
-		ViewStates[ViewIndex].Allocate();
+		const ERHIFeatureLevel::Type FeatureLevel = GetScene() ? GetScene()->GetFeatureLevel() : GMaxRHIFeatureLevel;
+		ViewStates[ViewIndex].Allocate(FeatureLevel);
 		ViewStateInterface = ViewStates[ViewIndex].GetReference();
 	}
 	else if (!bCaptureEveryFrame && ViewStateInterface && !bAlwaysPersistRenderingState)
@@ -999,9 +1000,11 @@ void UPlanarReflectionComponent::PostEditChangeProperty(FPropertyChangedEvent& P
 
 	for (int32 ViewIndex = 0; ViewIndex < ViewStates.Num(); ViewIndex++)
 	{
+		const ERHIFeatureLevel::Type FeatureLevel = GetScene() ? GetScene()->GetFeatureLevel() : GMaxRHIFeatureLevel;
+
 		// Recreate the view state to reset temporal history so that property changes can be seen immediately
 		ViewStates[ViewIndex].Destroy();
-		ViewStates[ViewIndex].Allocate();
+		ViewStates[ViewIndex].Allocate(FeatureLevel);
 	}
 
 	if (ProxyMeshComponent)
diff --git a/Engine/Source/Runtime/Engine/Private/LocalPlayer.cpp b/Engine/Source/Runtime/Engine/Private/LocalPlayer.cpp
--- a/Engine/Source/Runtime/Engine/Private/LocalPlayer.cpp
+++ b/Engine/Source/Runtime/Engine/Private/LocalPlayer.cpp
@@ -252,7 +252,7 @@ void ULocalPlayer::PostInitProperties()
 		ViewStates.SetNum(NumViews);
 		for (auto& State : ViewStates)
 		{
-			State.Allocate();
+			State.Allocate(GetWorld()->FeatureLevel);
 		}		
 	}
 }
@@ -784,7 +784,7 @@ bool ULocalPlayer::CalcSceneViewInitOptions(
 	if (!ViewStates.IsValidIndex(ViewIndex))
 	{
 		ViewStates.EmplaceAt(ViewIndex);
-		ViewStates[ViewIndex].Allocate();
+		ViewStates[ViewIndex].Allocate(GetWorld()->FeatureLevel);
 	}
 
 	ViewInitOptions.SceneViewStateInterface = ViewStates[ViewIndex].GetReference();
diff --git a/Engine/Source/Runtime/Engine/Private/RendererSupport.cpp b/Engine/Source/Runtime/Engine/Private/RendererSupport.cpp
--- a/Engine/Source/Runtime/Engine/Private/RendererSupport.cpp
+++ b/Engine/Source/Runtime/Engine/Private/RendererSupport.cpp
@@ -153,7 +153,10 @@ static void RestoreReferencesToRendererModuleClasses(
 
 	IRendererModule& RendererModule = GetRendererModule();
 
-	FSceneViewStateReference::AllocateAll();
+	TMap<UWorld*, bool>::TConstIterator WorldsIt(WorldsToUpdate);
+	const ERHIFeatureLevel::Type FirstFeatureLevel = WorldsIt ? WorldsIt.Key()->FeatureLevel.GetValue() : GMaxRHIFeatureLevel;
+
+	FSceneViewStateReference::AllocateAll(FirstFeatureLevel);
 
 	// Recreate all renderer scenes
 	for (TMap<UWorld*, bool>::TConstIterator It(WorldsToUpdate); It; ++It)
diff --git a/Engine/Source/Runtime/Engine/Private/SceneView.cpp b/Engine/Source/Runtime/Engine/Private/SceneView.cpp
--- a/Engine/Source/Runtime/Engine/Private/SceneView.cpp
+++ b/Engine/Source/Runtime/Engine/Private/SceneView.cpp
@@ -350,14 +350,19 @@ FSceneViewStateReference::~FSceneViewStateReference()
 	Destroy();
 }
 
-void FSceneViewStateReference::Allocate()
+void FSceneViewStateReference::Allocate(ERHIFeatureLevel::Type FeatureLevel)
 {
 	check(!Reference);
-	Reference = GetRendererModule().AllocateViewState();
+	Reference = GetRendererModule().AllocateViewState(FeatureLevel);
 	GlobalListLink = TLinkedList<FSceneViewStateReference*>(this);
 	GlobalListLink.LinkHead(GetSceneViewStateList());
 }
 
+void FSceneViewStateReference::Allocate()
+{
+	Allocate(GMaxRHIFeatureLevel);
+}
+
 void FSceneViewStateReference::Destroy()
 {
 	GlobalListLink.Unlink();
@@ -379,15 +384,20 @@ void FSceneViewStateReference::DestroyAll()
 	}
 }
 
-void FSceneViewStateReference::AllocateAll()
+void FSceneViewStateReference::AllocateAll(ERHIFeatureLevel::Type FeatureLevel)
 {
 	for(TLinkedList<FSceneViewStateReference*>::TIterator ViewStateIt(FSceneViewStateReference::GetSceneViewStateList());ViewStateIt;ViewStateIt.Next())
 	{
 		FSceneViewStateReference* ViewStateReference = *ViewStateIt;
-		ViewStateReference->Reference = GetRendererModule().AllocateViewState();
+		ViewStateReference->Reference = GetRendererModule().AllocateViewState(FeatureLevel);
 	}
 }
 
+void FSceneViewStateReference::AllocateAll()
+{
+	AllocateAll(GMaxRHIFeatureLevel);
+}
+
 TLinkedList<FSceneViewStateReference*>*& FSceneViewStateReference::GetSceneViewStateList()
 {
 	static TLinkedList<FSceneViewStateReference*>* List = NULL;
diff --git a/Engine/Source/Runtime/Engine/Public/SceneTypes.h b/Engine/Source/Runtime/Engine/Public/SceneTypes.h
--- a/Engine/Source/Runtime/Engine/Public/SceneTypes.h
+++ b/Engine/Source/Runtime/Engine/Public/SceneTypes.h
@@ -3,6 +3,7 @@
 #pragma once
 
 #include "CoreMinimal.h"
+#include "RHIDefinitions.h"
 #include "UObject/ObjectMacros.h"
 #include "Templates/RefCounting.h"
 #include "Containers/List.h"
@@ -75,6 +76,9 @@ public:
 	ENGINE_API virtual ~FSceneViewStateReference();
 
 	/** Allocates the Scene view state. */
+	ENGINE_API void Allocate(ERHIFeatureLevel::Type FeatureLevel);
+
+	UE_DEPRECATED(5.0, "Allocate must be called with an appropriate RHI Feature Level")
 	ENGINE_API void Allocate();
 
 	/** Destorys the Scene view state. */
@@ -84,6 +88,9 @@ public:
 	ENGINE_API static void DestroyAll();
 
 	/** Recreates all view states in the global list. */
+	ENGINE_API static void AllocateAll(ERHIFeatureLevel::Type FeatureLevel);
+
+	UE_DEPRECATED(5.0, "AllocateAll must be called with an appropriate RHI Feature Level")
 	ENGINE_API static void AllocateAll();
 
 	FSceneViewStateInterface* GetReference()
diff --git a/Engine/Source/Runtime/RenderCore/Public/RendererInterface.h b/Engine/Source/Runtime/RenderCore/Public/RendererInterface.h
--- a/Engine/Source/Runtime/RenderCore/Public/RendererInterface.h
+++ b/Engine/Source/Runtime/RenderCore/Public/RendererInterface.h
@@ -772,7 +772,13 @@ public:
 	virtual void UpdateStaticDrawListsForMaterials(const TArray<const FMaterial*>& Materials) = 0;
 
 	/** Allocates a new instance of the private scene manager implementation of FSceneViewStateInterface */
-	virtual class FSceneViewStateInterface* AllocateViewState() = 0;
+	virtual class FSceneViewStateInterface* AllocateViewState(ERHIFeatureLevel::Type FeatureLevel) = 0;
+
+	UE_DEPRECATED(5.0, "AllocateViewState must be called with an appropriate RHI Feature Level")
+	inline FSceneViewStateInterface* AllocateViewState()
+	{
+		return AllocateViewState(GMaxRHIFeatureLevel);
+	}
 
 	/** @return The number of lights that affect a primitive. */
 	virtual uint32 GetNumDynamicLightsAffectingPrimitive(const class FPrimitiveSceneInfo* PrimitiveSceneInfo,const class FLightCacheInterface* LCI) = 0;
diff --git a/Engine/Source/Runtime/Renderer/Private/RendererModule.h b/Engine/Source/Runtime/Renderer/Private/RendererModule.h
--- a/Engine/Source/Runtime/Renderer/Private/RendererModule.h
+++ b/Engine/Source/Runtime/Renderer/Private/RendererModule.h
@@ -43,7 +43,7 @@ public:
 	virtual void RemoveScene(FSceneInterface* Scene) override;
 	virtual void UpdateStaticDrawLists() override;
 	virtual void UpdateStaticDrawListsForMaterials(const TArray<const FMaterial*>& Materials) override;
-	virtual FSceneViewStateInterface* AllocateViewState() override;
+	virtual FSceneViewStateInterface* AllocateViewState(ERHIFeatureLevel::Type FeatureLevel) override;
 	virtual uint32 GetNumDynamicLightsAffectingPrimitive(const FPrimitiveSceneInfo* PrimitiveSceneInfo,const FLightCacheInterface* LCI) override;
 	virtual void OnWorldCleanup(UWorld* World, bool bSessionEnded, bool bCleanupResources, bool bWorldChanged) override;
 	virtual void InitializeSystemTextures(FRHICommandListImmediate& RHICmdList);
diff --git a/Engine/Source/Runtime/Renderer/Private/RendererScene.cpp b/Engine/Source/Runtime/Renderer/Private/RendererScene.cpp
--- a/Engine/Source/Runtime/Renderer/Private/RendererScene.cpp
+++ b/Engine/Source/Runtime/Renderer/Private/RendererScene.cpp
@@ -191,7 +191,7 @@ TRefCountPtr<IPooledRenderTarget> FPersistentSkyAtmosphereData::GetCurrentCamera
 }
 
 /** Default constructor. */
-FSceneViewState::FSceneViewState()
+FSceneViewState::FSceneViewState(ERHIFeatureLevel::Type FeatureLevel)
 	: OcclusionQueryPool(RHICreateRenderQueryPool(RQT_Occlusion))
 	, TimerQueryPool(RHICreateRenderQueryPool(RQT_AbsoluteTime, FLatentGPUTimer::NumBufferedFrames * 2 * 2 * 2))
 	, TranslucencyTimer(TimerQueryPool)
@@ -200,7 +200,7 @@ FSceneViewState::FSceneViewState()
 	, PostMotionBlurTranslucencyTimer(TimerQueryPool)
 {
 	// Set FeatureLevel to a valid value, so we get Init/ReleaseDynamicRHI calls on FeatureLevel changes
-	SetFeatureLevel(GMaxRHIFeatureLevel);
+	SetFeatureLevel(FeatureLevel);
 	
 	UniqueID = FSceneViewState_UniqueID.Increment();
 	OcclusionFrameCounter = 0;
@@ -5299,9 +5299,9 @@ void FRendererModule::UpdateStaticDrawListsForMaterials(const TArray<const FMate
 	UpdateStaticMeshesForMaterials(Materials);
 }
 
-FSceneViewStateInterface* FRendererModule::AllocateViewState()
+FSceneViewStateInterface* FRendererModule::AllocateViewState(ERHIFeatureLevel::Type FeatureLevel)
 {
-	return new FSceneViewState();
+	return new FSceneViewState(FeatureLevel);
 }
 
 void FRendererModule::InvalidatePathTracedOutput()
diff --git a/Engine/Source/Runtime/Renderer/Private/ScenePrivate.h b/Engine/Source/Runtime/Renderer/Private/ScenePrivate.h
--- a/Engine/Source/Runtime/Renderer/Private/ScenePrivate.h
+++ b/Engine/Source/Runtime/Renderer/Private/ScenePrivate.h
@@ -1175,7 +1175,7 @@ public:
 	}
 
 	/** Default constructor. */
-	FSceneViewState();
+	FSceneViewState(ERHIFeatureLevel::Type FeatureLevel);
 
 	virtual ~FSceneViewState();
 
diff --git a/Engine/Source/Runtime/UMG/Private/Components/Viewport.cpp b/Engine/Source/Runtime/UMG/Private/Components/Viewport.cpp
--- a/Engine/Source/Runtime/UMG/Private/Components/Viewport.cpp
+++ b/Engine/Source/Runtime/UMG/Private/Components/Viewport.cpp
@@ -96,7 +96,8 @@ FUMGViewportClient::FUMGViewportClient(FPreviewScene* InPreviewScene)
 	, Viewport(nullptr)
 	, EngineShowFlags(ESFIM_Game)
 {
-	ViewState.Allocate();
+	FSceneInterface* Scene = GetScene();
+	ViewState.Allocate(Scene ? Scene->GetFeatureLevel() : GMaxRHIFeatureLevel);
 
 	BackgroundColor = FColor(55, 55, 55);
 }
