Fix for sudden mix jump when a push mix interrupts during a pop mix

diff --git a/Engine/Source/Runtime/Engine/Private/AudioDevice.cpp b/Engine/Source/Runtime/Engine/Private/AudioDevice.cpp
--- a/Engine/Source/Runtime/Engine/Private/AudioDevice.cpp
+++ b/Engine/Source/Runtime/Engine/Private/AudioDevice.cpp
@@ -2528,6 +2528,20 @@ void FAudioDevice::UpdateSoundMix(USoundMix* SoundMix, FSoundMixState* SoundMixS
 					SoundMixState->EndTime = SoundMixState->FadeOutStartTime + SoundMix->FadeOutTime;
 				}
 			}
+			else if (SoundMixState->CurrentState == ESoundMixState::FadingOut || SoundMixState->CurrentState == ESoundMixState::AwaitingRemoval)
+			{
+				// Pretend our fade time is starting now
+				SoundMixState->StartTime = GetAudioClock();
+
+				// Invert the current fade-out interp value to get a fade-in interp value
+				float TargetInterpValue = 1.0f - SoundMixState->InterpValue;
+
+				// Compute a fade in start time that would result in this target interp value so we avoid jumps
+				SoundMixState->FadeInStartTime = SoundMixState->StartTime - TargetInterpValue * SoundMix->FadeInTime;
+
+				SoundMixState->CurrentState = ESoundMixState::FadingIn;
+				SoundMixState->FadeInEndTime = SoundMixState->FadeInStartTime + SoundMix->FadeInTime;
+			}
 		}
 	}
 }
