Changing default LFE bleed in sound classes to 0.0

diff --git a/Engine/Source/Runtime/Engine/Private/SoundClass.cpp b/Engine/Source/Runtime/Engine/Private/SoundClass.cpp
--- a/Engine/Source/Runtime/Engine/Private/SoundClass.cpp
+++ b/Engine/Source/Runtime/Engine/Private/SoundClass.cpp
@@ -25,7 +25,7 @@ FSoundClassProperties::FSoundClassProperties()
 	, Pitch(1.0f)
 	, LowPassFilterFrequency(MAX_FILTER_FREQUENCY)
 	, AttenuationDistanceScale(1.0f)
-	, LFEBleed(0.5f)
+	, LFEBleed(0.0f)
 	, VoiceCenterChannelVolume(0.0f)
 	, RadioFilterVolume(0.0f)
 	, RadioFilterVolumeThreshold(0.0f)
