When an app requires login also add the MCP and build ID override to the commandline.
Adding Find widget by tag function for increased editor testablity
Adding Simulation button click for better testing opportunities

diff --git a/Engine/Source/Programs/AutomationTool/Gauntlet/Unreal/Game/EpicGame.TestConfig.cs b/Engine/Source/Programs/AutomationTool/Gauntlet/Unreal/Game/EpicGame.TestConfig.cs
--- a/Engine/Source/Programs/AutomationTool/Gauntlet/Unreal/Game/EpicGame.TestConfig.cs
+++ b/Engine/Source/Programs/AutomationTool/Gauntlet/Unreal/Game/EpicGame.TestConfig.cs
@@ -128,7 +128,7 @@ namespace EpicGame
 		{
 			base.ApplyToConfig(AppConfig, ConfigRole, OtherRoles);
 
-			if (ConfigRole.RoleType.IsClient() || ConfigRole.RoleType.IsServer())
+			if (ConfigRole.RoleType.IsClient() || ConfigRole.RoleType.IsServer() || RequiresLogin)
 			{
 				string McpString = "";
 				bool bIsBuildMachine = CommandUtils.IsBuildMachine; 
diff --git a/Engine/Source/Runtime/Engine/Private/Tests/AutomationCommon.cpp b/Engine/Source/Runtime/Engine/Private/Tests/AutomationCommon.cpp
--- a/Engine/Source/Runtime/Engine/Private/Tests/AutomationCommon.cpp
+++ b/Engine/Source/Runtime/Engine/Private/Tests/AutomationCommon.cpp
@@ -195,6 +195,46 @@ namespace AutomationCommon
 		return FrameTrace;
 	}
 
+	SWidget* FindWidgetByTag(const FName Tag)
+	{
+		const FTagMetaData UniqueMetaData(Tag);
+		// Get a list of all the current slate windows
+		TArray<TSharedRef<SWindow>> Windows;
+		FSlateApplication::Get().GetAllVisibleWindowsOrdered(/*OUT*/Windows);
+
+		TArray<SWidget*> Stack;
+		for (const TSharedRef<SWindow>& Window : Windows)
+		{
+			Stack.Push(&Window.Get());
+		}
+
+		while (Stack.Num() > 0)
+		{
+			SWidget* Widget = Stack.Pop();
+			const int32 NumChildren = Widget->GetChildren()->Num();
+			for (int32 ChildIndex = 0; ChildIndex < NumChildren; ChildIndex++)
+			{
+				SWidget& ChildWidget = Widget->GetChildren()->GetChildAt(ChildIndex).Get();
+				const TArray<TSharedRef<FTagMetaData>> AllMetaData = ChildWidget.GetAllMetaData<FTagMetaData>();
+				for (int32 MetaDataIndex = 0; MetaDataIndex < AllMetaData.Num(); ++MetaDataIndex)
+				{
+					TSharedRef<FTagMetaData> MetaData = AllMetaData[MetaDataIndex];
+					if (MetaData->Tag == UniqueMetaData.Tag)
+					{
+						// Done! found the widget
+						return &ChildWidget;
+					}
+				}
+
+				// If we got here we didn't match the widget so push this child on the stack.
+				Stack.Push(&ChildWidget);
+			}
+
+		}
+
+		return nullptr;
+	}
+
 	class FAutomationImageComparisonRequest : public IAutomationLatentCommand
 	{
 	public:
diff --git a/Engine/Source/Runtime/Engine/Public/Tests/AutomationCommon.h b/Engine/Source/Runtime/Engine/Public/Tests/AutomationCommon.h
--- a/Engine/Source/Runtime/Engine/Public/Tests/AutomationCommon.h
+++ b/Engine/Source/Runtime/Engine/Public/Tests/AutomationCommon.h
@@ -7,6 +7,7 @@
 
 class AMatineeActor;
 class SWindow;
+class SWidget;
 
 #if WITH_AUTOMATION_TESTS
 
@@ -52,6 +53,13 @@ namespace AutomationCommon
 
 	ENGINE_API TArray<uint8> CaptureFrameTrace(const FString& MapOrContext, const FString& TestName);
 
+	/**
+	 * Given the FName of a FTagMetaData will find all the corresponding widgets.
+	 * @param Tag the meta data tag searched for
+	 * @return the found widget or nullptr
+	 */
+	ENGINE_API SWidget* FindWidgetByTag(const FName Tag);
+
 #endif
 }
 
diff --git a/Engine/Source/Runtime/Slate/Private/Widgets/Input/SButton.cpp b/Engine/Source/Runtime/Slate/Private/Widgets/Input/SButton.cpp
--- a/Engine/Source/Runtime/Slate/Private/Widgets/Input/SButton.cpp
+++ b/Engine/Source/Runtime/Slate/Private/Widgets/Input/SButton.cpp
@@ -637,6 +637,13 @@ void SButton::SetPressMethod(EButtonPressMethod::Type InPressMethod)
 	PressMethod = InPressMethod;
 }
 
+#if !UE_BUILD_SHIPPING
+void SButton::SimulateClick()
+{
+	ExecuteOnClick();
+}
+#endif // !UE_BUILD_SHIPPING
+
 FSlateColor SButton::GetDisabledForegroundColor() const
 {
 	return Style->DisabledForeground;
diff --git a/Engine/Source/Runtime/Slate/Public/Widgets/Input/SButton.h b/Engine/Source/Runtime/Slate/Public/Widgets/Input/SButton.h
--- a/Engine/Source/Runtime/Slate/Public/Widgets/Input/SButton.h
+++ b/Engine/Source/Runtime/Slate/Public/Widgets/Input/SButton.h
@@ -185,6 +185,10 @@ public:
 	void SetTouchMethod(EButtonTouchMethod::Type InTouchMethod);
 	void SetPressMethod(EButtonPressMethod::Type InPressMethod);
 
+#if !UE_BUILD_SHIPPING
+	void SimulateClick();
+#endif // !UE_BUILD_SHIPPING
+
 public:
 
 	//~ SWidget overrides
