Chaos Immediate Physics (RBAN) : Add option to disable bodies that are no longer active based on the index rather than just enabling the first N (enabled by default).

#tests PIE

#localization none

diff --git a/Engine/Source/Runtime/AnimGraphRuntime/Private/BoneControllers/AnimNode_RigidBody.cpp b/Engine/Source/Runtime/AnimGraphRuntime/Private/BoneControllers/AnimNode_RigidBody.cpp
--- a/Engine/Source/Runtime/AnimGraphRuntime/Private/BoneControllers/AnimNode_RigidBody.cpp
+++ b/Engine/Source/Runtime/AnimGraphRuntime/Private/BoneControllers/AnimNode_RigidBody.cpp
@@ -1737,7 +1737,7 @@ void FAnimNode_RigidBody::InitializeBoneReferences(const FBoneContainer& Require
 	OutputBoneData.Empty(NumBodies);
 
 	int32 NumSimulatedBodies = 0;
-
+	TArray<int32> SimulatedBodyIndices;
 	// if no name is entered, use root
 	if (BaseBoneRef.BoneName == NAME_None)
 	{
@@ -1775,6 +1775,7 @@ void FAnimNode_RigidBody::InitializeBoneReferences(const FBoneContainer& Require
 			if (BodyAnimData[BodyIndex].bIsSimulated)
 			{
 				++NumSimulatedBodies;
+				SimulatedBodyIndices.AddUnique(BodyIndex);
 			}
 
 			OutputData->BoneIndicesToParentBody.Add(CompactPoseBoneIndex);
@@ -1825,7 +1826,11 @@ void FAnimNode_RigidBody::InitializeBoneReferences(const FBoneContainer& Require
 
 		if (PhysicsSimulation)
 		{
+#if WITH_CHAOS
+			PhysicsSimulation->SetNumActiveBodies(NumSimulatedBodies, SimulatedBodyIndices);
+#else
 			PhysicsSimulation->SetNumActiveBodies(NumSimulatedBodies);
+#endif
 		}
 
 		// We're switching to a new LOD, this invalidates our captured poses.
diff --git a/Engine/Source/Runtime/Engine/Private/PhysicsEngine/ImmediatePhysics/ImmediatePhysicsChaos/ImmediatePhysicsSimulation_Chaos.cpp b/Engine/Source/Runtime/Engine/Private/PhysicsEngine/ImmediatePhysics/ImmediatePhysicsChaos/ImmediatePhysicsSimulation_Chaos.cpp
--- a/Engine/Source/Runtime/Engine/Private/PhysicsEngine/ImmediatePhysics/ImmediatePhysicsChaos/ImmediatePhysicsSimulation_Chaos.cpp
+++ b/Engine/Source/Runtime/Engine/Private/PhysicsEngine/ImmediatePhysics/ImmediatePhysicsChaos/ImmediatePhysicsSimulation_Chaos.cpp
@@ -39,6 +39,7 @@ int32 ChaosImmediate_Evolution_DeltaTimeCount = 100;
 int32 ChaosImmediate_Evolution_Iterations = -1;
 int32 ChaosImmediate_Evolution_PushOutIterations = -1;
 Chaos::FRealSingle ChaosImmediate_Evolution_BoundsExtension = 0.0f;
+int32 ChaosImmediate_DisableInactiveByIndex = 1;
 FAutoConsoleVariableRef CVarChaosImmPhysStepTime(TEXT("p.Chaos.ImmPhys.StepTime"), ChaosImmediate_Evolution_StepTime, TEXT("Override step time (if not zero)"));
 FAutoConsoleVariableRef CVarChaosImmPhysNumSteps(TEXT("p.Chaos.ImmPhys.NumSteps"), ChaosImmediate_Evolution_NumSteps, TEXT("Override num steps (if not zero)"));
 FAutoConsoleVariableRef CVarChaosImmPhysInitialStepTime(TEXT("p.Chaos.ImmPhys.InitialStepTime"), ChaosImmediate_Evolution_InitialStepTime, TEXT("Initial step time (then calculated from rolling average)"));
@@ -46,6 +47,7 @@ FAutoConsoleVariableRef CVarChaosImmPhysDeltaTimeCount(TEXT("p.Chaos.ImmPhys.Del
 FAutoConsoleVariableRef CVarChaosImmPhysIterations(TEXT("p.Chaos.ImmPhys.Iterations"), ChaosImmediate_Evolution_Iterations, TEXT("Override number of constraint solver loops in immediate physics (if >= 0)"));
 FAutoConsoleVariableRef CVarChaosImmPhysPushOutIterations(TEXT("p.Chaos.ImmPhys.PushOutIterations"), ChaosImmediate_Evolution_PushOutIterations, TEXT("Override number of solver push-out loops (if >= 0)"));
 FAutoConsoleVariableRef CVarChaosImmPhysBoundsExtension(TEXT("p.Chaos.ImmPhys.BoundsExtension"), ChaosImmediate_Evolution_BoundsExtension, TEXT("Bounds are grown by this fraction of their size (should be >= 0.0)"));
+FAutoConsoleVariableRef CVarChaosImmPhysDisableInactiveByIndex(TEXT("p.Chaos.ImmPhys.DisableInactiveByIndex"), ChaosImmediate_DisableInactiveByIndex, TEXT("Disable bodies that are no longer active based on the index, rather than just count."));
 
 Chaos::FRealSingle ChaosImmediate_Evolution_SimSpaceCentrifugalAlpha = 1.0f;
 Chaos::FRealSingle ChaosImmediate_Evolution_SimSpaceCoriolisAlpha = 0.5f;
@@ -503,9 +505,9 @@ namespace ImmediatePhysics_Chaos
 		delete JointHandle;
 	}
 
-	void FSimulation::SetNumActiveBodies(int32 InNumActiveActorHandles)
+	void FSimulation::SetNumActiveBodies(int32 InNumActiveActorHandles, TArray<int32> ActiveBodyIndices)
 	{
-		if (InNumActiveActorHandles == Implementation->NumActiveDynamicActorHandles)
+		if (InNumActiveActorHandles == Implementation->NumActiveDynamicActorHandles && ChaosImmediate_DisableInactiveByIndex == 0)
 		{
 			return;
 		}
@@ -520,18 +522,31 @@ namespace ImmediatePhysics_Chaos
 			{
 				continue;
 			}
-
-			if (Implementation->NumActiveDynamicActorHandles < InNumActiveActorHandles)
+			if (ChaosImmediate_DisableInactiveByIndex != 0)
 			{
-				Handle->SetEnabled(true);
-				++Implementation->NumActiveDynamicActorHandles;
+				if (ActiveBodyIndices.Contains(ActorHandleIndex))
+				{
+					Handle->SetEnabled(true);
+					++Implementation->NumActiveDynamicActorHandles;
+				}
+				else
+				{
+					Handle->SetEnabled(false);
+				}
 			}
 			else
 			{
-				Handle->SetEnabled(false);
+				if (Implementation->NumActiveDynamicActorHandles < InNumActiveActorHandles)
+				{
+					Handle->SetEnabled(true);
+					++Implementation->NumActiveDynamicActorHandles;
+				}
+				else
+				{
+					Handle->SetEnabled(false);
+				}
 			}
 		}
-
 		Implementation->bActorsDirty = true;
 	}
 
diff --git a/Engine/Source/Runtime/Engine/Public/Physics/ImmediatePhysics/ImmediatePhysicsChaos/ImmediatePhysicsSimulation_Chaos.h b/Engine/Source/Runtime/Engine/Public/Physics/ImmediatePhysics/ImmediatePhysicsChaos/ImmediatePhysicsSimulation_Chaos.h
--- a/Engine/Source/Runtime/Engine/Public/Physics/ImmediatePhysics/ImmediatePhysicsChaos/ImmediatePhysicsSimulation_Chaos.h
+++ b/Engine/Source/Runtime/Engine/Public/Physics/ImmediatePhysics/ImmediatePhysicsChaos/ImmediatePhysicsSimulation_Chaos.h
@@ -38,7 +38,7 @@ namespace ImmediatePhysics_Chaos
 		void DestroyJoint(FJointHandle* JointHandle);
 
 		/** Sets the number of active bodies. This number is reset any time a new simulated body is created */
-		void SetNumActiveBodies(int32 NumActiveBodies);
+		void SetNumActiveBodies(int32 NumActiveBodies, TArray<int32> ActiveBodyIndices);
 
 		/** An array of actors to ignore. */
 		struct FIgnorePair
