Add option to not create kinematic kinematic joint constraints (enabled by default).

diff --git a/Engine/Source/Runtime/Engine/Private/PhysicsEngine/ConstraintInstance.cpp b/Engine/Source/Runtime/Engine/Private/PhysicsEngine/ConstraintInstance.cpp
--- a/Engine/Source/Runtime/Engine/Private/PhysicsEngine/ConstraintInstance.cpp
+++ b/Engine/Source/Runtime/Engine/Private/PhysicsEngine/ConstraintInstance.cpp
@@ -52,6 +52,9 @@ TAutoConsoleVariable<float> CVarConstraintAngularStiffnessScale(
 bool bEnableSkeletalMeshConstraints = true;
 FAutoConsoleVariableRef CVarEnableSkeletalMeshConstraints(TEXT("p.EnableSkeletalMeshConstraints"), bEnableSkeletalMeshConstraints, TEXT("Enable skeletal mesh constraints defined within the Physics Asset Editor"));
 
+bool bAllowKinematicKinematicConstraints = false;
+FAutoConsoleVariableRef CVarAllowKinematicKinematicConstraints(TEXT("p.AllowKinematicKinematicConstraints"), bAllowKinematicKinematicConstraints, TEXT("Do not create constraints between two rigid kinematics."));
+
 /** Handy macro for setting BIT of VAR based on the bool CONDITION */
 #define SET_DRIVE_PARAM(VAR, CONDITION, BIT)   (VAR) = (CONDITION) ? ((VAR) | (BIT)) : ((VAR) & ~(BIT))
 
@@ -489,6 +492,11 @@ void FConstraintInstance::InitConstraint(FBodyInstance* Body1, FBodyInstance* Bo
 			return;
 		}
 
+		if (!bAllowKinematicKinematicConstraints && FPhysicsInterface::IsKinematic(Actor1) && FPhysicsInterface::IsKinematic(Actor2))
+		{
+			return;
+		}
+
 		FPhysicsCommand::ExecuteWrite(Actor1, Actor2, [&](const FPhysicsActorHandle& ActorA, const FPhysicsActorHandle& ActorB)
 		{
 			InitConstraint_AssumesLocked(ActorA, ActorB, InScale, InConstraintBrokenDelegate, InPlasticDeformationDelegate);
