Check for valid actor handles in FConstraintInstance::InitConstraint (fix crash).

diff --git a/Engine/Source/Runtime/Engine/Private/PhysicsEngine/ConstraintInstance.cpp b/Engine/Source/Runtime/Engine/Private/PhysicsEngine/ConstraintInstance.cpp
--- a/Engine/Source/Runtime/Engine/Private/PhysicsEngine/ConstraintInstance.cpp
+++ b/Engine/Source/Runtime/Engine/Private/PhysicsEngine/ConstraintInstance.cpp
@@ -492,7 +492,7 @@ void FConstraintInstance::InitConstraint(FBodyInstance* Body1, FBodyInstance* Bo
 			return;
 		}
 
-		if (!bAllowKinematicKinematicConstraints && FPhysicsInterface::IsKinematic(Actor1) && FPhysicsInterface::IsKinematic(Actor2))
+		if (!bAllowKinematicKinematicConstraints && (!FPhysicsInterface::IsValid(Actor1) || FPhysicsInterface::IsKinematic(Actor1)) && (!FPhysicsInterface::IsValid(Actor2) || FPhysicsInterface::IsKinematic(Actor2)))
 		{
 			return;
 		}
