Speculative fix & logging analog cursor crash

diff --git a/Engine/Plugins/Experimental/CommonUI/Source/CommonUI/Private/Input/CommonUIActionRouterBase.cpp b/Engine/Plugins/Experimental/CommonUI/Source/CommonUI/Private/Input/CommonUIActionRouterBase.cpp
--- a/Engine/Plugins/Experimental/CommonUI/Source/CommonUI/Private/Input/CommonUIActionRouterBase.cpp
+++ b/Engine/Plugins/Experimental/CommonUI/Source/CommonUI/Private/Input/CommonUIActionRouterBase.cpp
@@ -188,13 +188,20 @@ bool UCommonUIActionRouterBase::RegisterLinkedPreprocessor(const UWidget& Widget
 void UCommonUIActionRouterBase::Initialize(FSubsystemCollectionBase& Collection)
 {
 	Super::Initialize(Collection);
-	Collection.InitializeDependency(UCommonInputSubsystem::StaticClass());
+	UCommonInputSubsystem* InputSubsystem = Collection.InitializeDependency<UCommonInputSubsystem>();
 
 	UCommonActivatableWidget::OnRebuilding.AddUObject(this, &UCommonUIActionRouterBase::HandleActivatableWidgetRebuilding);
 	FCoreUObjectDelegates::GetPostGarbageCollect().AddUObject(this, &UCommonUIActionRouterBase::HandlePostGarbageCollect);
 
-	AnalogCursor = MakeAnalogCursor();
-	PostAnalogCursorCreate();
+	if (ensure(InputSubsystem))
+	{
+		AnalogCursor = MakeAnalogCursor();
+		PostAnalogCursorCreate();
+	}
+	else
+	{
+		UE_LOG(LogUIActionRouter, Warning, TEXT("Input system not initialized before action router!"));
+	}
 
 	FSlateApplication::Get().OnFocusChanging().AddUObject(this, &UCommonUIActionRouterBase::HandleSlateFocusChanging);
 }
@@ -234,6 +241,8 @@ bool UCommonUIActionRouterBase::ShouldCreateSubsystem(UObject* Outer) const
 	TArray<UClass*> ChildClasses;
 	GetDerivedClasses(GetClass(), ChildClasses, false);
 
+	UE_LOG(LogUIActionRouter, Warning, TEXT("Found %i derived classes when attemping to create action router (%s)"), ChildClasses.Num(), *GetClass()->GetName());
+
 	// Only create an instance if there is no override implementation defined elsewhere
 	return ChildClasses.Num() == 0;
 }
diff --git a/Engine/Source/Runtime/Engine/Private/Subsystems/SubsystemCollection.cpp b/Engine/Source/Runtime/Engine/Private/Subsystems/SubsystemCollection.cpp
--- a/Engine/Source/Runtime/Engine/Private/Subsystems/SubsystemCollection.cpp
+++ b/Engine/Source/Runtime/Engine/Private/Subsystems/SubsystemCollection.cpp
@@ -8,6 +8,8 @@
 #include "Modules/ModuleManager.h"
 #include "Misc/PackageName.h"
 
+DEFINE_LOG_CATEGORY(LogSubsystemCollection);
+
 /** FSubsystemModuleWatcher class to hide the implementation of keeping the DynamicSystemModuleMap up to date*/
 class FSubsystemModuleWatcher
 {
@@ -187,13 +189,18 @@ void FSubsystemCollectionBase::Deinitialize()
 
 USubsystem* FSubsystemCollectionBase::InitializeDependency(TSubclassOf<USubsystem> SubsystemClass)
 {
+	UE_LOG(LogSubsystemCollection, Warning, TEXT("Attempting to initialize subsystem dependency (%s)"), *SubsystemClass->GetName());
+
+	USubsystem* Subsystem = nullptr;
 	if (ensureMsgf(SubsystemClass, TEXT("Attempting to add invalid subsystem as dependancy."))
 		&& ensureMsgf(bPopulating, TEXT("InitializeDependancy() should only be called from System USubsystem::Initialization() implementations."))
 		&& ensureMsgf(SubsystemClass->IsChildOf(BaseType), TEXT("ClassType (%s) must be a subclass of BaseType(%s)."), *SubsystemClass->GetName(), *BaseType->GetName()))
 	{
-		return AddAndInitializeSubsystem(SubsystemClass);
+		Subsystem = AddAndInitializeSubsystem(SubsystemClass);
 	}
-	return nullptr;
+
+	UE_CLOG(!Subsystem, LogSubsystemCollection, Warning, TEXT("Failed to initialize subsystem dependency (%s)"), *SubsystemClass->GetName());
+	return Subsystem;
 }
 
 void FSubsystemCollectionBase::AddReferencedObjects(FReferenceCollector& Collector)
@@ -233,10 +240,13 @@ USubsystem* FSubsystemCollectionBase::AddAndInitializeSubsystem(UClass* Subsyste
 				Subsystem->Initialize(*this);
 				return Subsystem;
 			}
+
+			UE_LOG(LogSubsystemCollection, Warning, TEXT("Subsystem DNE, but CDO choose to not create (%s)"), *SubsystemClass->GetName());
 		}
 		return nullptr;
 	}
 
+	UE_LOG(LogSubsystemCollection, Warning, TEXT("Subsystem already exists (%s)"), *SubsystemClass->GetName());
 	return SubsystemMap.FindRef(SubsystemClass);
 }
 
diff --git a/Engine/Source/Runtime/Engine/Public/Subsystems/SubsystemCollection.h b/Engine/Source/Runtime/Engine/Public/Subsystems/SubsystemCollection.h
--- a/Engine/Source/Runtime/Engine/Public/Subsystems/SubsystemCollection.h
+++ b/Engine/Source/Runtime/Engine/Public/Subsystems/SubsystemCollection.h
@@ -9,6 +9,8 @@
 class USubsystem;
 class UDynamicSubsystem;
 
+DECLARE_LOG_CATEGORY_EXTERN(LogSubsystemCollection, Log, All);
+
 class ENGINE_API FSubsystemCollectionBase : public FGCObject
 {
 public:
