-Fix Attached Actor dirty propagation to also work with multiple nested attached actors.

diff --git a/Engine/Source/Editor/UnrealEd/Private/Elements/Actor/ActorElementEditorViewportInteractionCustomization.cpp b/Engine/Source/Editor/UnrealEd/Private/Elements/Actor/ActorElementEditorViewportInteractionCustomization.cpp
--- a/Engine/Source/Editor/UnrealEd/Private/Elements/Actor/ActorElementEditorViewportInteractionCustomization.cpp
+++ b/Engine/Source/Editor/UnrealEd/Private/Elements/Actor/ActorElementEditorViewportInteractionCustomization.cpp
@@ -140,7 +140,7 @@ void FActorElementEditorViewportInteractionCustomization::ApplyDeltaToActor(AAct
 	AffectedActors.Add(InActor);
 	
 	// grab all attached actors to invalidate everybody 	
-	InActor->GetAttachedActors(AffectedActors, /*bResetArray*/ false );
+	InActor->GetAttachedActors(AffectedActors, /*bResetArray*/ false, /*bRecursivelyIncludeAttachedActors*/ true );
 
 	if (GEditor->IsDeltaModificationEnabled())
 	{
diff --git a/Engine/Source/Runtime/Engine/Classes/GameFramework/Actor.h b/Engine/Source/Runtime/Engine/Classes/GameFramework/Actor.h
--- a/Engine/Source/Runtime/Engine/Classes/GameFramework/Actor.h
+++ b/Engine/Source/Runtime/Engine/Classes/GameFramework/Actor.h
@@ -2980,7 +2980,7 @@ public:
 	
 	/** Find all Actors which are attached directly to a component in this actor */
 	UFUNCTION(BlueprintPure, Category = "Actor")
-	void GetAttachedActors(TArray<AActor*>& OutActors, bool bResetArray = true) const;
+	void GetAttachedActors(TArray<AActor*>& OutActors, bool bResetArray = true, bool bRecursivelyIncludeAttachedActors = false) const;
 
 	/**
 	 * Sets the ticking group for this actor.
diff --git a/Engine/Source/Runtime/Engine/Private/Actor.cpp b/Engine/Source/Runtime/Engine/Private/Actor.cpp
--- a/Engine/Source/Runtime/Engine/Private/Actor.cpp
+++ b/Engine/Source/Runtime/Engine/Private/Actor.cpp
@@ -2203,13 +2203,26 @@ void AActor::ForEachAttachedActors(TFunctionRef<bool(class AActor*)> Functor) co
 	}
 }
 
-void AActor::GetAttachedActors(TArray<class AActor*>& OutActors, bool bResetArray) const
+void AActor::GetAttachedActors(TArray<class AActor*>& OutActors, bool bResetArray, bool bRecursivelyIncludeAttachedActors) const
 {
 	if (bResetArray)
 	{
 		OutActors.Reset();
 	}
-	ForEachAttachedActors([&OutActors](AActor * Actor) { OutActors.AddUnique(Actor); return true; });
+
+	ForEachAttachedActors([bRecursivelyIncludeAttachedActors, &OutActors](AActor* AttachedActor)
+	{
+		int32 OriginalNumActors = OutActors.Num();
+
+		if ((OriginalNumActors <= OutActors.AddUnique(AttachedActor)) && bRecursivelyIncludeAttachedActors)
+		{
+			AttachedActor->GetAttachedActors(OutActors, false, true);
+		}
+		return true;
+	});
+
+
+
 }
 
 bool AActor::ActorHasTag(FName Tag) const
