Change log category of subsystem init messages.

diff --git a/Engine/Source/Runtime/Engine/Private/Subsystems/SubsystemCollection.cpp b/Engine/Source/Runtime/Engine/Private/Subsystems/SubsystemCollection.cpp
--- a/Engine/Source/Runtime/Engine/Private/Subsystems/SubsystemCollection.cpp
+++ b/Engine/Source/Runtime/Engine/Private/Subsystems/SubsystemCollection.cpp
@@ -189,7 +189,7 @@ void FSubsystemCollectionBase::Deinitialize()
 
 USubsystem* FSubsystemCollectionBase::InitializeDependency(TSubclassOf<USubsystem> SubsystemClass)
 {
-	UE_LOG(LogSubsystemCollection, Warning, TEXT("Attempting to initialize subsystem dependency (%s)"), *SubsystemClass->GetName());
+	UE_LOG(LogSubsystemCollection, Log, TEXT("Attempting to initialize subsystem dependency (%s)"), *SubsystemClass->GetName());
 
 	USubsystem* Subsystem = nullptr;
 	if (ensureMsgf(SubsystemClass, TEXT("Attempting to add invalid subsystem as dependancy."))
@@ -199,7 +199,7 @@ USubsystem* FSubsystemCollectionBase::InitializeDependency(TSubclassOf<USubsyste
 		Subsystem = AddAndInitializeSubsystem(SubsystemClass);
 	}
 
-	UE_CLOG(!Subsystem, LogSubsystemCollection, Warning, TEXT("Failed to initialize subsystem dependency (%s)"), *SubsystemClass->GetName());
+	UE_CLOG(!Subsystem, LogSubsystemCollection, Log, TEXT("Failed to initialize subsystem dependency (%s)"), *SubsystemClass->GetName());
 	return Subsystem;
 }
 
@@ -241,12 +241,12 @@ USubsystem* FSubsystemCollectionBase::AddAndInitializeSubsystem(UClass* Subsyste
 				return Subsystem;
 			}
 
-			UE_LOG(LogSubsystemCollection, Warning, TEXT("Subsystem DNE, but CDO choose to not create (%s)"), *SubsystemClass->GetName());
+			UE_LOG(LogSubsystemCollection, Log, TEXT("Subsystem DNE, but CDO choose to not create (%s)"), *SubsystemClass->GetName());
 		}
 		return nullptr;
 	}
 
-	UE_LOG(LogSubsystemCollection, Warning, TEXT("Subsystem already exists (%s)"), *SubsystemClass->GetName());
+	UE_LOG(LogSubsystemCollection, Log, TEXT("Subsystem already exists (%s)"), *SubsystemClass->GetName());
 	return SubsystemMap.FindRef(SubsystemClass);
 }
 
