Implemented processing granularity to ULevel::RouteActorInitialization so that we can call it piecewise until we exceed the AddToWorld budget.

diff --git a/Engine/Source/Runtime/Engine/Classes/Engine/CoreSettings.h b/Engine/Source/Runtime/Engine/Classes/Engine/CoreSettings.h
--- a/Engine/Source/Runtime/Engine/Classes/Engine/CoreSettings.h
+++ b/Engine/Source/Runtime/Engine/Classes/Engine/CoreSettings.h
@@ -142,6 +142,8 @@ extern ENGINE_API int32 GLevelStreamingComponentsRegistrationGranularity;
 extern ENGINE_API int32 GLevelStreamingAddPrimitiveGranularity;
 /** Batching granularity used to unregister actor components during level streaming.  */
 extern ENGINE_API int32 GLevelStreamingComponentsUnregistrationGranularity;
+/** Batching granularity used to initialize actors during level streaming. If this is zero, we process all actors and stages in one pass. */
+extern ENGINE_API int32 GLevelStreamingRouteActorInitializationGranularity;
 /** Maximum allowed time to spend for actor unregistration steps during level streaming (ms per frame). If this is 0.0 then we don't timeslice.*/
 extern ENGINE_API float GLevelStreamingUnregisterComponentsTimeLimit;
 /** Whether to force a GC after levels are streamed out to instantly reclaim the memory at the expensive of a hitch. */
diff --git a/Engine/Source/Runtime/Engine/Classes/Engine/Level.h b/Engine/Source/Runtime/Engine/Classes/Engine/Level.h
--- a/Engine/Source/Runtime/Engine/Classes/Engine/Level.h
+++ b/Engine/Source/Runtime/Engine/Classes/Engine/Level.h
@@ -791,6 +791,16 @@ private:
 	TSoftObjectPtr<UObject> OwnerLevelPartition;
 #endif // #if WITH_EDITORONLY_DATA
 
+	enum class ERouteActorInitializationState : uint8
+	{
+		Preinitialize,
+		Initialize,
+		BeginPlay,
+		Finished
+	};
+	ERouteActorInitializationState RouteActorInitializationState;
+	int32 RouteActorInitializationIndex;
+
 public:
 	// Used internally to determine which actors should go on the world's NetworkActor list
 	ENGINE_API static bool IsNetActor(const AActor* Actor);
@@ -959,12 +969,18 @@ public:
 	/** Releases rendering resources for this level. */
 	ENGINE_API void ReleaseRenderingResources();
 
+	/**
+	 * Old implementation for routing actor initialization in full.
+	 */
+	void RouteActorInitializeOld();
+
 	/**
 	 * Routes pre and post initialize to actors and also sets volumes.
 	 *
+	 * @param NumActorsToProcess	The maximum number of actors to update in this pass, 0 to process all actors.
 	 * @todo seamless worlds: this doesn't correctly handle volumes in the multi- level case
 	 */
-	void RouteActorInitialize();
+	void RouteActorInitialize(int32 NumActorsToProcess);
 
 	/**
 	 * Rebuilds static streaming data for all levels in the specified UWorld.
diff --git a/Engine/Source/Runtime/Engine/Private/CoreSettings.cpp b/Engine/Source/Runtime/Engine/Private/CoreSettings.cpp
--- a/Engine/Source/Runtime/Engine/Private/CoreSettings.cpp
+++ b/Engine/Source/Runtime/Engine/Private/CoreSettings.cpp
@@ -20,6 +20,7 @@ float GLevelStreamingUnregisterComponentsTimeLimit = 1.0f;
 int32 GLevelStreamingComponentsRegistrationGranularity = 10;
 int32 GLevelStreamingAddPrimitiveGranularity = 120;
 int32 GLevelStreamingComponentsUnregistrationGranularity = 5;
+int32 GLevelStreamingRouteActorInitializationGranularity = 10;
 int32 GLevelStreamingForceGCAfterLevelStreamedOut = 1;
 int32 GLevelStreamingContinuouslyIncrementalGCWhileLevelsPendingPurge = 1;
 int32 GLevelStreamingContinuouslyIncrementalGCWhileLevelsPendingPurgeOverride = 0;
@@ -98,6 +99,13 @@ static FAutoConsoleVariableRef CVarLevelStreamingComponentsUnregistrationGranula
 	ECVF_Default
 	);
 
+static FAutoConsoleVariableRef CVarLevelStreamingRouteActorInitializationGranularity(
+	TEXT("s.LevelStreamingRouteActorInitializationGranularity"),
+	GLevelStreamingRouteActorInitializationGranularity,
+	TEXT("Batching granularity used to initialize actors during level streaming. If this is zero, we process all actors and stages in one pass."),
+	ECVF_Default
+);
+
 static FAutoConsoleVariableRef CVarForceGCAfterLevelStreamedOut(
 	TEXT("s.ForceGCAfterLevelStreamedOut"),
 	GLevelStreamingForceGCAfterLevelStreamedOut,
diff --git a/Engine/Source/Runtime/Engine/Private/Level.cpp b/Engine/Source/Runtime/Engine/Private/Level.cpp
--- a/Engine/Source/Runtime/Engine/Private/Level.cpp
+++ b/Engine/Source/Runtime/Engine/Private/Level.cpp
@@ -91,6 +91,14 @@ static FAutoConsoleVariableRef CVarActorClusteringEnabled(
 	ECVF_Default
 );
 
+bool GUseLegacyRouteActorInitialization = false;
+static FAutoConsoleVariableRef CVarUseLegacyRouteActorInitialization(
+	TEXT("s.UseLegacyRouteActorInitialization"),
+	GUseLegacyRouteActorInitialization,
+	TEXT("Toggle for whether to use the old non-granular implementation of route actor initialization."),
+	ECVF_Default
+);
+
 #if WITH_EDITOR
 FLevelPartitionOperationScope::FLevelPartitionOperationScope(ULevel* InLevel)
 {
@@ -333,6 +341,8 @@ ULevel::ULevel( const FObjectInitializer& ObjectInitializer )
 	,	TickTaskLevel(FTickTaskManagerInterface::Get().AllocateTickTaskLevel())
 	,	PrecomputedLightVolume(new FPrecomputedLightVolume())
 	,	PrecomputedVolumetricLightmap(new FPrecomputedVolumetricLightmap())
+	,	RouteActorInitializationState(ERouteActorInitializationState::Preinitialize)
+	,	RouteActorInitializationIndex(0)
 {
 #if WITH_EDITORONLY_DATA
 	LevelColor = FLinearColor::White;
@@ -2377,30 +2387,28 @@ void ULevel::ReleaseRenderingResources()
 	}
 }
 
-void ULevel::RouteActorInitialize()
+void ULevel::RouteActorInitializeOld()
 {
-	TRACE_OBJECT_EVENT(this, RouteActorInitialize);
-
 	// Send PreInitializeComponents and collect volumes.
-	for( int32 Index = 0; Index < Actors.Num(); ++Index )
+	for (int32 Index = 0; Index < Actors.Num(); ++Index)
 	{
 		AActor* const Actor = Actors[Index];
-		if( Actor && !Actor->IsActorInitialized() )
+		if (Actor && !Actor->IsActorInitialized())
 		{
 			Actor->PreInitializeComponents();
 		}
 	}
 
 	const bool bCallBeginPlay = OwningWorld->HasBegunPlay();
-	TArray<AActor *> ActorsToBeginPlay;
+	TArray<AActor*> ActorsToBeginPlay;
 
 	// Send InitializeComponents on components and PostInitializeComponents.
-	for( int32 Index = 0; Index < Actors.Num(); ++Index )
+	for (int32 Index = 0; Index < Actors.Num(); ++Index)
 	{
 		AActor* const Actor = Actors[Index];
-		if( Actor )
+		if (Actor)
 		{
-			if( !Actor->IsActorInitialized() )
+			if (!Actor->IsActorInitialized())
 			{
 				// Call Initialize on Components.
 				Actor->InitializeComponents();
@@ -2408,7 +2416,7 @@ void ULevel::RouteActorInitialize()
 				Actor->PostInitializeComponents(); // should set Actor->bActorInitialized = true
 				if (!Actor->IsActorInitialized() && IsValidChecked(Actor))
 				{
-					UE_LOG(LogActor, Fatal, TEXT("%s failed to route PostInitializeComponents.  Please call Super::PostInitializeComponents() in your <className>::PostInitializeComponents() function. "), *Actor->GetFullName() );
+					UE_LOG(LogActor, Fatal, TEXT("%s failed to route PostInitializeComponents.  Please call Super::PostInitializeComponents() in your <className>::PostInitializeComponents() function. "), *Actor->GetFullName());
 				}
 
 				if (bCallBeginPlay && !Actor->IsChildActor())
@@ -2426,6 +2434,111 @@ void ULevel::RouteActorInitialize()
 		SCOPE_CYCLE_COUNTER(STAT_ActorBeginPlay);
 		Actor->DispatchBeginPlay(/*bFromLevelStreaming*/ true);
 	}
+
+	bAlreadyRoutedActorInitialize = true;
+}
+
+void ULevel::RouteActorInitialize(int32 NumActorsToProcess)
+{
+	TRACE_OBJECT_EVENT(this, RouteActorInitialize);
+	if (GUseLegacyRouteActorInitialization)
+	{
+		RouteActorInitializeOld();
+		return;
+	}
+
+	const bool bFullProcessing = (NumActorsToProcess <= 0);
+	switch (RouteActorInitializationState)
+	{
+		case ERouteActorInitializationState::Preinitialize:
+		{
+			// Actor pre-initialization may spawn new actors so we need to incrementally process until actor count stabilizes
+			while (RouteActorInitializationIndex < Actors.Num())
+			{
+				AActor* const Actor = Actors[RouteActorInitializationIndex];
+				if (Actor && !Actor->IsActorInitialized())
+				{
+					Actor->PreInitializeComponents();
+				}
+
+				++RouteActorInitializationIndex;
+				if (!bFullProcessing && (--NumActorsToProcess == 0))
+				{
+					return;
+				}
+			}
+
+			RouteActorInitializationIndex = 0;
+			RouteActorInitializationState = ERouteActorInitializationState::Initialize;
+		}
+
+		// Intentional fall-through, proceeding if we haven't expired our actor count budget
+		case ERouteActorInitializationState::Initialize:
+		{
+			while (RouteActorInitializationIndex < Actors.Num())
+			{
+				AActor* const Actor = Actors[RouteActorInitializationIndex];
+				if (Actor)
+				{
+					if (!Actor->IsActorInitialized())
+					{
+						Actor->InitializeComponents();
+						Actor->PostInitializeComponents();
+						if (!Actor->IsActorInitialized() && IsValidChecked(Actor))
+						{
+							UE_LOG(LogActor, Fatal, TEXT("%s failed to route PostInitializeComponents. Please call Super::PostInitializeComponents() in your <className>::PostInitializeComponents() function. "), *Actor->GetFullName());
+						}
+					}
+				}
+
+				++RouteActorInitializationIndex;
+				if (!bFullProcessing && (--NumActorsToProcess == 0))
+				{
+					return;
+				}
+			}
+
+			RouteActorInitializationIndex = 0;
+			RouteActorInitializationState = ERouteActorInitializationState::BeginPlay;
+		}
+
+		// Intentional fall-through, proceeding if we haven't expired our actor count budget
+		case ERouteActorInitializationState::BeginPlay:
+		{
+			if (OwningWorld->HasBegunPlay())
+			{
+				while (RouteActorInitializationIndex < Actors.Num())
+				{
+					// Child actors have play begun explicitly by their parents
+					AActor* const Actor = Actors[RouteActorInitializationIndex];
+					if (Actor && !Actor->IsChildActor())
+					{
+						// This will no-op if the actor has already begun play
+						SCOPE_CYCLE_COUNTER(STAT_ActorBeginPlay);
+						const bool bFromLevelStreaming = true;
+						Actor->DispatchBeginPlay(bFromLevelStreaming);
+					}
+
+					++RouteActorInitializationIndex;
+					if (!bFullProcessing && (--NumActorsToProcess == 0))
+					{
+						return;
+					}
+				}
+			}
+
+			// Should only have to clean up once, so do it at the end of this phase in case this method gets called again after we're done
+			RouteActorInitializationState = ERouteActorInitializationState::Finished;
+			RouteActorInitializationIndex = 0;
+			bAlreadyRoutedActorInitialize = true;
+		}
+
+		// Intentional fall-through if we're done
+		case ERouteActorInitializationState::Finished:
+		{
+			break;
+		}
+	}
 }
 
 UPackage* ULevel::CreateMapBuildDataPackage() const
diff --git a/Engine/Source/Runtime/Engine/Private/World.cpp b/Engine/Source/Runtime/Engine/Private/World.cpp
--- a/Engine/Source/Runtime/Engine/Private/World.cpp
+++ b/Engine/Source/Runtime/Engine/Private/World.cpp
@@ -2829,12 +2829,15 @@ void UWorld::AddToWorld( ULevel* Level, const FTransform& LevelTransform, bool b
 		{
 			QUICK_SCOPE_CYCLE_COUNTER(STAT_AddToWorldTime_RouteActorInitialize);
 			SCOPE_TIME_TO_VAR(&RouteActorInitializeTime);
+			const int32 NumActorsToProcess = (!bConsiderTimeLimit || !IsGameWorld() || IsRunningCommandlet()) ? 0 : GLevelStreamingRouteActorInitializationGranularity;
 			bStartup = 1;
-			Level->RouteActorInitialize();
-			Level->bAlreadyRoutedActorInitialize = true;
+			do 
+			{
+				Level->RouteActorInitialize(NumActorsToProcess);
+			} while (!Level->bAlreadyRoutedActorInitialize && !IsTimeLimitExceeded(TEXT("routing Initialize on actors"), StartTime, Level, TimeLimit));
 			bStartup = 0;
 
-			bExecuteNextStep = (!bConsiderTimeLimit || !IsTimeLimitExceeded( TEXT("routing Initialize on actors"), StartTime, Level, TimeLimit ));
+			bExecuteNextStep = Level->bAlreadyRoutedActorInitialize && (!bConsiderTimeLimit || !IsTimeLimitExceeded( TEXT("routing Initialize on actors"), StartTime, Level, TimeLimit ));
 		}
 
 		// Sort the actor list; can't do this on save as the relevant properties for sorting might have been changed by code
@@ -4713,10 +4716,11 @@ void UWorld::InitializeActorsForPlay(const FURL& InURL, bool bResetTime, FRegist
 		}
 
 		// Route various initialization functions and set volumes.
-		for( int32 LevelIndex=0; LevelIndex<Levels.Num(); LevelIndex++ )
+		const int32 ProcessAllRouteActorInitializationGranularity = 0;
+		for (int32 LevelIndex = 0; LevelIndex < Levels.Num(); LevelIndex++)
 		{
-			ULevel*	const Level = Levels[LevelIndex];
-			Level->RouteActorInitialize();
+			ULevel* const Level = Levels[LevelIndex];
+			Level->RouteActorInitialize(ProcessAllRouteActorInitializationGranularity);
 		}
 
 		// Let server know client sub-levels visibility state
