Editor-only public access to FPrimaryAssetTypeInfo::SpecificAssets property

diff --git a/Engine/Source/Runtime/Engine/Classes/Engine/AssetManagerTypes.h b/Engine/Source/Runtime/Engine/Classes/Engine/AssetManagerTypes.h
--- a/Engine/Source/Runtime/Engine/Classes/Engine/AssetManagerTypes.h
+++ b/Engine/Source/Runtime/Engine/Classes/Engine/AssetManagerTypes.h
@@ -131,6 +131,8 @@ public:
 
 #if WITH_EDITOR
 	TArray<FDirectoryPath>& GetDirectories() { return Directories; }
+
+	TArray<FSoftObjectPath>& GetSpecificAssets() { return SpecificAssets; }
 #endif //if WITH_EDITOR
 
 private:
