Disabling new RouteActorInitialize path by default as it's causing hangs when leaving the match. I've not yet been able to reproduce the issue, but I'll re-enable once I get to the bottom of it.

diff --git a/Engine/Source/Runtime/Engine/Private/Level.cpp b/Engine/Source/Runtime/Engine/Private/Level.cpp
--- a/Engine/Source/Runtime/Engine/Private/Level.cpp
+++ b/Engine/Source/Runtime/Engine/Private/Level.cpp
@@ -91,7 +91,7 @@ static FAutoConsoleVariableRef CVarActorClusteringEnabled(
 	ECVF_Default
 );
 
-bool GUseLegacyRouteActorInitialization = false;
+bool GUseLegacyRouteActorInitialization = true;
 static FAutoConsoleVariableRef CVarUseLegacyRouteActorInitialization(
 	TEXT("s.UseLegacyRouteActorInitialization"),
 	GUseLegacyRouteActorInitialization,
