Allowing volumes the option of skipping the check for a valid collision component in CheckForErrors.
+ This was previously done for DefaultPhysicsVolumes, and we'd like to extend this option to other volumes.

diff --git a/Engine/Plugins/AudioGameplayVolume/Source/AudioGameplayVolume/Public/AudioGameplayVolume.h b/Engine/Plugins/AudioGameplayVolume/Source/AudioGameplayVolume/Public/AudioGameplayVolume.h
--- a/Engine/Plugins/AudioGameplayVolume/Source/AudioGameplayVolume/Public/AudioGameplayVolume.h
+++ b/Engine/Plugins/AudioGameplayVolume/Source/AudioGameplayVolume/Public/AudioGameplayVolume.h
@@ -38,6 +38,7 @@ public:
 	//~ Begin UObject Interface
 #if WITH_EDITOR
 	virtual void PostEditChangeProperty(FPropertyChangedEvent& PropertyChangedEvent) override;
+	virtual bool ShouldCheckCollisionComponentForErrors() const override { return false; }
 #endif // WITH_EDITOR
 	virtual void GetLifetimeReplicatedProps(TArray<FLifetimeProperty>& OutLifetimeProps) const override;
 	//~ End UObject Interface
diff --git a/Engine/Source/Runtime/Engine/Classes/GameFramework/DefaultPhysicsVolume.h b/Engine/Source/Runtime/Engine/Classes/GameFramework/DefaultPhysicsVolume.h
--- a/Engine/Source/Runtime/Engine/Classes/GameFramework/DefaultPhysicsVolume.h
+++ b/Engine/Source/Runtime/Engine/Classes/GameFramework/DefaultPhysicsVolume.h
@@ -21,6 +21,10 @@ UCLASS(notplaceable, transient, MinimalAPI)
 class ADefaultPhysicsVolume : public APhysicsVolume
 {
 	GENERATED_UCLASS_BODY()
+
+#if WITH_EDITOR
+		virtual bool ShouldCheckCollisionComponentForErrors() const override { return false; }
+#endif // WITH_EDITOR
 };
 
 
diff --git a/Engine/Source/Runtime/Engine/Classes/GameFramework/Volume.h b/Engine/Source/Runtime/Engine/Classes/GameFramework/Volume.h
--- a/Engine/Source/Runtime/Engine/Classes/GameFramework/Volume.h
+++ b/Engine/Source/Runtime/Engine/Classes/GameFramework/Volume.h
@@ -37,6 +37,7 @@ public:
 	* for any potential errors and register them with map check dialog.
 	*/
 	virtual void CheckForErrors() override;
+	virtual bool ShouldCheckCollisionComponentForErrors() const { return true; }
 #endif // WITH_EDITOR
 	
 	virtual bool IsLevelBoundsRelevant() const override;
diff --git a/Engine/Source/Runtime/Engine/Private/ErrorChecking.cpp b/Engine/Source/Runtime/Engine/Private/ErrorChecking.cpp
--- a/Engine/Source/Runtime/Engine/Private/ErrorChecking.cpp
+++ b/Engine/Source/Runtime/Engine/Private/ErrorChecking.cpp
@@ -11,7 +11,6 @@
 #include "Engine/Brush.h"
 #include "GameFramework/Volume.h"
 #include "UObject/Package.h"
-#include "GameFramework/DefaultPhysicsVolume.h"
 #include "Logging/TokenizedMessage.h"
 #include "Logging/MessageLog.h"
 #include "Misc/UObjectToken.h"
@@ -190,28 +189,27 @@ void AVolume::CheckForErrors()
 {
 	Super::CheckForErrors();
 
-	// The default physics volume can have zero area; it's extents aren't used, only the physics properties
-	if (IsA(ADefaultPhysicsVolume::StaticClass()))
+	// Some volumes do not need a valid collision component; for example:
+	// the default physics volume only uses the physics properties and not the extents - so it will have an area of zero
+	if (ShouldCheckCollisionComponentForErrors())
 	{
-		return;
-	}
-
-	if (GetRootComponent() == NULL)
-	{
-		FMessageLog("MapCheck").Warning()
-			->AddToken(FUObjectToken::Create(this))
-			->AddToken(FTextToken::Create(LOCTEXT( "MapCheck_Message_VolumeActorCollisionComponentNULL", "Volume actor has NULL collision component - please delete")))
-			->AddToken(FMapErrorToken::Create(FMapErrors::VolumeActorCollisionComponentNULL));
-	}
-	else
-	{
-		if (GetRootComponent()->Bounds.SphereRadius <= SMALL_NUMBER)
+		if (GetRootComponent() == NULL)
 		{
-			FFormatNamedArguments Arguments;
-			Arguments.Add(TEXT("ActorName"), FText::FromString(GetName()));
 			FMessageLog("MapCheck").Warning()
 				->AddToken(FUObjectToken::Create(this))
-				->AddToken(FTextToken::Create(LOCTEXT( "MapCheck_Message_VolumeActorZeroRadius", "Volume actor has a collision component with 0 radius - please delete")));
+				->AddToken(FTextToken::Create(LOCTEXT("MapCheck_Message_VolumeActorCollisionComponentNULL", "Volume actor has NULL collision component - please delete")))
+				->AddToken(FMapErrorToken::Create(FMapErrors::VolumeActorCollisionComponentNULL));
+		}
+		else
+		{
+			if (GetRootComponent()->Bounds.SphereRadius <= SMALL_NUMBER)
+			{
+				FFormatNamedArguments Arguments;
+				Arguments.Add(TEXT("ActorName"), FText::FromString(GetName()));
+				FMessageLog("MapCheck").Warning()
+					->AddToken(FUObjectToken::Create(this))
+					->AddToken(FTextToken::Create(LOCTEXT("MapCheck_Message_VolumeActorZeroRadius", "Volume actor has a collision component with 0 radius - please delete")));
+			}
 		}
 	}
 }
