Fixed an issue with granular route actor initialization where a subsequent stream on a level that had previously completed streaming would not make progress. This would result in a hang the next time the level attempted to shut down as it would wait on streaming to complete and never do so.

When streaming is complete, we were resetting bAlreadyRoutedActorInitialization but leaving the state at Finished, so the next time we attempted to route actor initialization, we would no-op in the Finished state and the bool would never progress. We only really need one variable to keep track of the progress in that stage, so I'm eliminating the bool.

diff --git a/Engine/Source/Runtime/Engine/Classes/Engine/Level.h b/Engine/Source/Runtime/Engine/Classes/Engine/Level.h
--- a/Engine/Source/Runtime/Engine/Classes/Engine/Level.h
+++ b/Engine/Source/Runtime/Engine/Classes/Engine/Level.h
@@ -635,8 +635,6 @@ public:
 	uint8										bAlreadyInitializedNetworkActors:1;
 	/** Whether we already cleared AActor::bActorSeamlessTraveled.							*/
 	uint8										bAlreadyClearedActorsSeamlessTravelFlag:1;
-	/** Whether we already routed initialize on actors.										*/
-	uint8										bAlreadyRoutedActorInitialize:1;
 	/** Whether we already sorted the actor list.											*/
 	uint8										bAlreadySortedActorList:1;
 	/** Whether this level is in the process of being associated with its world	(i.e. we are within AddToWorld for this level */
@@ -969,6 +967,16 @@ public:
 	/** Releases rendering resources for this level. */
 	ENGINE_API void ReleaseRenderingResources();
 
+	/**
+	 * Returns whether the level has completed routing actor initialization.
+	 */
+	bool IsFinishedRouteActorInitialization() const { return RouteActorInitializationState == ERouteActorInitializationState::Finished; }
+
+	/**
+	 * Method for resetting routing actor initialization for the next time this level is streamed.
+	 */
+	void ResetRouteActorInitializationState();
+
 	/**
 	 * Old implementation for routing actor initialization in full.
 	 */
diff --git a/Engine/Source/Runtime/Engine/Private/Level.cpp b/Engine/Source/Runtime/Engine/Private/Level.cpp
--- a/Engine/Source/Runtime/Engine/Private/Level.cpp
+++ b/Engine/Source/Runtime/Engine/Private/Level.cpp
@@ -2387,6 +2387,12 @@ void ULevel::ReleaseRenderingResources()
 	}
 }
 
+void ULevel::ResetRouteActorInitializationState()
+{
+	 RouteActorInitializationState = ERouteActorInitializationState::Preinitialize;
+	 RouteActorInitializationIndex = 0;
+}
+
 void ULevel::RouteActorInitializeOld()
 {
 	// Send PreInitializeComponents and collect volumes.
@@ -2435,7 +2441,7 @@ void ULevel::RouteActorInitializeOld()
 		Actor->DispatchBeginPlay(/*bFromLevelStreaming*/ true);
 	}
 
-	bAlreadyRoutedActorInitialize = true;
+	RouteActorInitializationState = ERouteActorInitializationState::Finished;
 }
 
 void ULevel::RouteActorInitialize(int32 NumActorsToProcess)
@@ -2527,10 +2533,7 @@ void ULevel::RouteActorInitialize(int32 NumActorsToProcess)
 				}
 			}
 
-			// Should only have to clean up once, so do it at the end of this phase in case this method gets called again after we're done
 			RouteActorInitializationState = ERouteActorInitializationState::Finished;
-			RouteActorInitializationIndex = 0;
-			bAlreadyRoutedActorInitialize = true;
 		}
 
 		// Intentional fall-through if we're done
diff --git a/Engine/Source/Runtime/Engine/Private/World.cpp b/Engine/Source/Runtime/Engine/Private/World.cpp
--- a/Engine/Source/Runtime/Engine/Private/World.cpp
+++ b/Engine/Source/Runtime/Engine/Private/World.cpp
@@ -2825,7 +2825,7 @@ void UWorld::AddToWorld( ULevel* Level, const FTransform& LevelTransform, bool b
 		}
 
 		// Route various initialization functions and set volumes.
-		if( bExecuteNextStep && !Level->bAlreadyRoutedActorInitialize )
+		if (bExecuteNextStep && !Level->IsFinishedRouteActorInitialization())
 		{
 			QUICK_SCOPE_CYCLE_COUNTER(STAT_AddToWorldTime_RouteActorInitialize);
 			SCOPE_TIME_TO_VAR(&RouteActorInitializeTime);
@@ -2834,10 +2834,10 @@ void UWorld::AddToWorld( ULevel* Level, const FTransform& LevelTransform, bool b
 			do 
 			{
 				Level->RouteActorInitialize(NumActorsToProcess);
-			} while (!Level->bAlreadyRoutedActorInitialize && !IsTimeLimitExceeded(TEXT("routing Initialize on actors"), StartTime, Level, TimeLimit));
+			} while (!Level->IsFinishedRouteActorInitialization() && !IsTimeLimitExceeded(TEXT("routing Initialize on actors"), StartTime, Level, TimeLimit));
 			bStartup = 0;
 
-			bExecuteNextStep = Level->bAlreadyRoutedActorInitialize && (!bConsiderTimeLimit || !IsTimeLimitExceeded( TEXT("routing Initialize on actors"), StartTime, Level, TimeLimit ));
+			bExecuteNextStep = Level->IsFinishedRouteActorInitialization() && (!bConsiderTimeLimit || !IsTimeLimitExceeded( TEXT("routing Initialize on actors"), StartTime, Level, TimeLimit ));
 		}
 
 		// Sort the actor list; can't do this on save as the relevant properties for sorting might have been changed by code
@@ -2865,9 +2865,9 @@ void UWorld::AddToWorld( ULevel* Level, const FTransform& LevelTransform, bool b
 		
 		Level->bAlreadyShiftedActors					= false;
 		Level->bAlreadyUpdatedComponents				= false;
-		Level->bAlreadyRoutedActorInitialize			= false;
 		Level->bAlreadySortedActorList					= false;
 		Level->bAlreadyClearedActorsSeamlessTravelFlag	= false;
+		Level->ResetRouteActorInitializationState();
 
 		// Finished making level visible - allow other levels to be added to the world.
 		CurrentLevelPendingVisibility = nullptr;
