Fixes for NaN checks fired when copying between uninitialized FVector/FMatrix variant types.

diff --git a/Engine/Source/Runtime/Core/Public/Math/Matrix.inl b/Engine/Source/Runtime/Core/Public/Math/Matrix.inl
--- a/Engine/Source/Runtime/Core/Public/Math/Matrix.inl
+++ b/Engine/Source/Runtime/Core/Public/Math/Matrix.inl
@@ -26,11 +26,6 @@ namespace Math
 template<typename T>
 FORCEINLINE TMatrix<T>::TMatrix()
 {
-#if ENABLE_NAN_DIAGNOSTIC
-	// TODO: Remove? Rendering code performs basic (throw away) operations on uninitialized matrices 
-	// making diagnostic ensure ineffective without this.
-	FMemory::Memzero(this, sizeof(*this));
-#endif
 }
 
 template<typename T>
diff --git a/Engine/Source/Runtime/Engine/Public/PrimitiveSceneProxy.h b/Engine/Source/Runtime/Engine/Public/PrimitiveSceneProxy.h
--- a/Engine/Source/Runtime/Engine/Public/PrimitiveSceneProxy.h
+++ b/Engine/Source/Runtime/Engine/Public/PrimitiveSceneProxy.h
@@ -1324,7 +1324,7 @@ private:
 	FBoxSphereBounds Bounds;
 
 	/** The primitive's local space bounds. */
-	FBoxSphereBounds LocalBounds;
+	FBoxSphereBounds LocalBounds = FBoxSphereBounds(ForceInit);
 
 	/** The component's actor's position. */
 	FVector ActorPosition;
diff --git a/Engine/Source/Runtime/Engine/Public/SceneView.h b/Engine/Source/Runtime/Engine/Public/SceneView.h
--- a/Engine/Source/Runtime/Engine/Public/SceneView.h
+++ b/Engine/Source/Runtime/Engine/Public/SceneView.h
@@ -228,9 +228,17 @@ struct FViewMatrices
 	FViewMatrices()
 	{
 		ProjectionMatrix.SetIdentity();
+		ProjectionNoAAMatrix.SetIdentity();
+		InvProjectionMatrix.SetIdentity();
 		ViewMatrix.SetIdentity();
+		InvViewMatrix.SetIdentity();
+		ViewProjectionMatrix.SetIdentity();
+		InvViewProjectionMatrix.SetIdentity();
 		HMDViewMatrixNoRoll.SetIdentity();
 		TranslatedViewMatrix.SetIdentity();
+		InvTranslatedViewMatrix.SetIdentity();
+		OverriddenTranslatedViewMatrix.SetIdentity();
+		OverriddenInvTranslatedViewMatrix.SetIdentity();
 		TranslatedViewProjectionMatrix.SetIdentity();
 		InvTranslatedViewProjectionMatrix.SetIdentity();
 		PreViewTranslation = FVector::ZeroVector;
diff --git a/Engine/Source/Runtime/RenderCore/Public/RenderTransform.h b/Engine/Source/Runtime/RenderCore/Public/RenderTransform.h
--- a/Engine/Source/Runtime/RenderCore/Public/RenderTransform.h
+++ b/Engine/Source/Runtime/RenderCore/Public/RenderTransform.h
@@ -406,7 +406,7 @@ struct FCompressedTransform
 	uint16		Scale_SharedExp[4];	// 1B padding if SignMantissaBits == 16
 
 	FORCEINLINE FCompressedTransform() {}
-	FORCEINLINE FCompressedTransform( const FRenderTransform& In )
+	FORCEINLINE_DEBUGGABLE FCompressedTransform( const FRenderTransform& In )
 	{
 		Translation = In.Origin;
 
@@ -416,7 +416,7 @@ struct FCompressedTransform
 		{
 			Axis[i] = In.TransformRows[i];
 			Scale[i] = Axis[i].Size();
-			Axis[i] /= Scale[i];
+			Axis[i] /= Scale[i] ? Scale[i] : 1.f;
 		}
 
 		// Rotation
