Convert SimpleSuspension ComputeSprungMasses() to use FReal to avoid precision issues.

diff --git a/Engine/Source/Runtime/Engine/Private/PhysicsEngine/SimpleSuspension.cpp b/Engine/Source/Runtime/Engine/Private/PhysicsEngine/SimpleSuspension.cpp
--- a/Engine/Source/Runtime/Engine/Private/PhysicsEngine/SimpleSuspension.cpp
+++ b/Engine/Source/Runtime/Engine/Private/PhysicsEngine/SimpleSuspension.cpp
@@ -38,6 +38,7 @@ void FSimpleSuspension::Update(const FTransform& LocalToWorld, const FVector& Li
 
 bool FSimpleSuspensionHelpers::ComputeSprungMasses(const TArray<FVector>& MassSpringPositions, const float TotalMass, TArray<float>& OutSprungMasses)
 {
+	using Chaos::FReal;
 	/*
 
 	For a body which is supported by a collection of parallel springs,
@@ -97,14 +98,14 @@ bool FSimpleSuspensionHelpers::ComputeSprungMasses(const TArray<FVector>& MassSp
 
 		*/
 
-		const float AX = MassSpringPositions[0].X;
-		const float AY = MassSpringPositions[0].Y;
-		const float DiffX = MassSpringPositions[1].X - AX;
-		const float DiffY = MassSpringPositions[1].Y - AY;
+		const FReal AX = MassSpringPositions[0].X;
+		const FReal AY = MassSpringPositions[0].Y;
+		const FReal DiffX = MassSpringPositions[1].X - AX;
+		const FReal DiffY = MassSpringPositions[1].Y - AY;
 
 		// If the springs are close together, just divide the mass in 2.
-		const float DistSquared = (DiffX * DiffX) + (DiffY * DiffY);
-		const float Dist = FMath::Sqrt(DistSquared);
+		const FReal DistSquared = (DiffX * DiffX) + (DiffY * DiffY);
+		const FReal Dist = FMath::Sqrt(DistSquared);
 		if (Dist <= SMALL_NUMBER)
 		{
 			OutSprungMasses[0] = OutSprungMasses[1] = TotalMass * .5f;
@@ -112,10 +113,10 @@ bool FSimpleSuspensionHelpers::ComputeSprungMasses(const TArray<FVector>& MassSp
 		}
 
 		// The springs are far enough apart, compute the distribution
-		const float DistInv = 1.f / Dist;
-		const float DirX = DiffX * DistInv;
-		const float DirY = DiffY * DistInv;
-		const float DirDotA = (DirX * AX) + (DirY * AY);
+		const FReal DistInv = 1.f / Dist;
+		const FReal DirX = DiffX * DistInv;
+		const FReal DirY = DiffY * DistInv;
+		const FReal DirDotA = (DirX * AX) + (DirY * AY);
 		OutSprungMasses[0] = -TotalMass * DirDotA * DistInv;
 		OutSprungMasses[1] = TotalMass - OutSprungMasses[0];
 		if (ensureMsgf(OutSprungMasses[0] >= 0.f, TEXT("Spring configuration is invalid! Please make sure the center of mass is located between the springs.")) &&
@@ -191,21 +192,21 @@ bool FSimpleSuspensionHelpers::ComputeSprungMasses(const TArray<FVector>& MassSp
 	*/
 
 	// Cache values we'll need later, and clear out the results array
-	const float CountN = (float)Count;
-	const float CountInverse = 1.f / CountN;
-	const float AverageMass = TotalMass * CountInverse;
-	float SumX = 0.f;
-	float SumY = 0.f;
-	float XDotX = 0.f;
-	float YDotY = 0.f;
-	float XDotY = 0.f;
-	float B0 = 0.f;
-	float B1 = 0.f;
-	float B2 = 0.f;
+	const FReal CountN = (FReal)Count;
+	const FReal CountInverse = 1.f / CountN;
+	const FReal AverageMass = TotalMass * CountInverse;
+	FReal SumX = 0.f;
+	FReal SumY = 0.f;
+	FReal XDotX = 0.f;
+	FReal YDotY = 0.f;
+	FReal XDotY = 0.f;
+	FReal B0 = 0.f;
+	FReal B1 = 0.f;
+	FReal B2 = 0.f;
 	for (uint32 Index = 0; Index < Count; ++Index)
 	{
-		const float X = MassSpringPositions[Index].X;
-		const float Y = MassSpringPositions[Index].Y;
+		const FReal X = MassSpringPositions[Index].X;
+		const FReal Y = MassSpringPositions[Index].Y;
 		SumX += X;
 		SumY += Y;
 		XDotX += X * X;
@@ -214,7 +215,7 @@ bool FSimpleSuspensionHelpers::ComputeSprungMasses(const TArray<FVector>& MassSp
 	}
 
 	// Compute determinant of system matrix, in prep for inversion
-	const float DetLL
+	const FReal DetLL
 		= (XDotX * YDotY * Count)
 		+ (2.f * XDotY * SumX * SumY)
 		- (YDotY * SumX * SumX)
@@ -228,31 +229,31 @@ bool FSimpleSuspensionHelpers::ComputeSprungMasses(const TArray<FVector>& MassSp
 	}
 
 	// Compute the elements of the inverse matrix
-	const float DetLLInv = 1.f / DetLL;
-	const float InvLL00 = ((Count * YDotY) - (SumY * SumY)) * DetLLInv;
-	const float InvLL01 = ((SumX * SumY) - (Count * XDotY)) * DetLLInv;
-	const float InvLL02 = ((SumY * XDotY) - (SumX * YDotY)) * DetLLInv;
-	const float InvLL10 = InvLL01; // Symmetry!
-	const float InvLL11 = ((Count * XDotX) - (SumX * SumX)) * DetLLInv;
-	const float InvLL12 = ((SumX * XDotY) - (SumY * XDotX)) * DetLLInv; // = InvLL21. Symmetry!
-	const float InvLL20 = InvLL02; // Symmetry!
-	const float InvLL21 = InvLL12; // Symmetry!
-	const float InvLL22 = ((XDotX * YDotY) - (XDotY * XDotY)) * DetLLInv;
+	const FReal DetLLInv = 1.f / DetLL;
+	const FReal InvLL00 = ((Count * YDotY) - (SumY * SumY)) * DetLLInv;
+	const FReal InvLL01 = ((SumX * SumY) - (Count * XDotY)) * DetLLInv;
+	const FReal InvLL02 = ((SumY * XDotY) - (SumX * YDotY)) * DetLLInv;
+	const FReal InvLL10 = InvLL01; // Symmetry!
+	const FReal InvLL11 = ((Count * XDotX) - (SumX * SumX)) * DetLLInv;
+	const FReal InvLL12 = ((SumX * XDotY) - (SumY * XDotX)) * DetLLInv; // = InvLL21. Symmetry!
+	const FReal InvLL20 = InvLL02; // Symmetry!
+	const FReal InvLL21 = InvLL12; // Symmetry!
+	const FReal InvLL22 = ((XDotX * YDotY) - (XDotY * XDotY)) * DetLLInv;
 
 	// Compute the Lagrange multipliers
-	const float LambdaB0 = 2.f * AverageMass * SumX;
-	const float LambdaB1 = 2.f * AverageMass * SumY;
-	const float LambdaB2 = (2.f * AverageMass * CountN) - (2.f * TotalMass);
-	const float Lambda0 = (InvLL00 * LambdaB0) + (InvLL01 * LambdaB1) + (InvLL02 * LambdaB2);
-	const float Lambda1 = (InvLL10 * LambdaB0) + (InvLL11 * LambdaB1) + (InvLL12 * LambdaB2);
-	const float Lambda2 = (InvLL20 * LambdaB0) + (InvLL21 * LambdaB1) + (InvLL22 * LambdaB2);
+	const FReal LambdaB0 = 2.f * AverageMass * SumX;
+	const FReal LambdaB1 = 2.f * AverageMass * SumY;
+	const FReal LambdaB2 = (2.f * AverageMass * CountN) - (2.f * TotalMass);
+	const FReal Lambda0 = (InvLL00 * LambdaB0) + (InvLL01 * LambdaB1) + (InvLL02 * LambdaB2);
+	const FReal Lambda1 = (InvLL10 * LambdaB0) + (InvLL11 * LambdaB1) + (InvLL12 * LambdaB2);
+	const FReal Lambda2 = (InvLL20 * LambdaB0) + (InvLL21 * LambdaB1) + (InvLL22 * LambdaB2);
 
 	// Compute the masses
 	for (uint32 Index = 0; Index < Count; ++Index)
 	{
-		const float X = MassSpringPositions[Index].X;
-		const float Y = MassSpringPositions[Index].Y;
-		const float LLambda = (X * Lambda0) + (Y * Lambda1) + Lambda2;
+		const FReal X = MassSpringPositions[Index].X;
+		const FReal Y = MassSpringPositions[Index].Y;
+		const FReal LLambda = (X * Lambda0) + (Y * Lambda1) + Lambda2;
 		OutSprungMasses[Index] = AverageMass - (0.5f * LLambda);
 		if (!ensureMsgf(OutSprungMasses[Index] >= 0.f, TEXT("Spring configuration is invalid! Please make sure the center of mass is located inside the area covered by the springs.")))
 		{
