Fixing crash in ULocalPlayer::PostInitProperties caused by assuming the World has been set.

diff --git a/Engine/Source/Runtime/Engine/Private/LocalPlayer.cpp b/Engine/Source/Runtime/Engine/Private/LocalPlayer.cpp
--- a/Engine/Source/Runtime/Engine/Private/LocalPlayer.cpp
+++ b/Engine/Source/Runtime/Engine/Private/LocalPlayer.cpp
@@ -248,11 +248,14 @@ void ULocalPlayer::PostInitProperties()
 			NumViews = GEngine->StereoRenderingDevice->GetDesiredNumberOfViews(true);
 			check(NumViews > 0);			
 		}
-				
+
+		const UWorld* CurrentWorld = GetWorld();
+		const ERHIFeatureLevel::Type FeatureLevel = CurrentWorld ? CurrentWorld->FeatureLevel.GetValue() : GMaxRHIFeatureLevel;
+
 		ViewStates.SetNum(NumViews);
 		for (auto& State : ViewStates)
 		{
-			State.Allocate(GetWorld()->FeatureLevel);
+			State.Allocate(FeatureLevel);
 		}		
 	}
 }
