[Backout] - CL18431296
Original CL Desc
-----------------------------------------------------------------
Fix crash when saving a newly created partitioned level during PIE: make sure the world doesn't keep data during the PIE session, clear everything once the original world has been duplicated for PIE. That way, if we ever save as the world while in PIE, that information isn't lost in the process and the system stays in a valid state.

diff --git a/Engine/Source/Runtime/Engine/Private/World.cpp b/Engine/Source/Runtime/Engine/Private/World.cpp
--- a/Engine/Source/Runtime/Engine/Private/World.cpp
+++ b/Engine/Source/Runtime/Engine/Private/World.cpp
@@ -3422,13 +3422,6 @@ UWorld* UWorld::GetDuplicatedWorldForPIE(UWorld* InWorld, UPackage* InPIEackage,
 
 	DuplicatedWorld->StreamingLevelsPrefix = UWorld::BuildPIEPackagePrefix(PIEInstanceID);
 
-#if WITH_EDITOR
-	if (UWorldPartition* WorldPartition = InWorld->GetWorldPartition())
-	{
-		WorldPartition->PostDuplicateWorldForPIE();
-	}
-#endif
-
 	return DuplicatedWorld;
 }
 
diff --git a/Engine/Source/Runtime/Engine/Private/WorldPartition/WorldPartition.cpp b/Engine/Source/Runtime/Engine/Private/WorldPartition/WorldPartition.cpp
--- a/Engine/Source/Runtime/Engine/Private/WorldPartition/WorldPartition.cpp
+++ b/Engine/Source/Runtime/Engine/Private/WorldPartition/WorldPartition.cpp
@@ -253,6 +253,7 @@ UWorldPartition::UWorldPartition(const FObjectInitializer& ObjectInitializer)
 	, WorldPartitionEditor(nullptr)
 	, bForceGarbageCollection(false)
 	, bForceGarbageCollectionPurge(false)
+	, bIsPIE(false)
 #endif
 	, InitState(EWorldPartitionInitState::Uninitialized)
 	, StreamingPolicy(nullptr)
@@ -313,18 +314,42 @@ void UWorldPartition::OnGCPostReachabilityAnalysis()
 
 void UWorldPartition::OnPreBeginPIE(bool bStartSimulate)
 {
-	OnBeginPlay(true);
+	check(!bIsPIE);
+	bIsPIE = true;
+
+	check(IsMainWorldPartition());
+
+	OnBeginPlay();
+}
+
+void UWorldPartition::OnPrePIEEnded(bool bWasSimulatingInEditor)
+{
+	check(bIsPIE);
+	bIsPIE = false;
 }
 
-void UWorldPartition::OnBeginPlay(bool bIsPIE)
+void UWorldPartition::OnBeginPlay()
 {
-	GenerateStreaming(nullptr, bIsPIE);
+	GenerateStreaming();
 	RuntimeHash->OnBeginPlay();
 }
 
 void UWorldPartition::OnCancelPIE()
 {
-	PostDuplicateWorldForPIE();
+	// No check here since CancelPIE can be called after PrePIEEnded
+	bIsPIE = false;
+	// Call OnEndPlay here since EndPlayMapDelegate is not called when cancelling PIE
+	OnEndPlay();
+}
+
+void UWorldPartition::OnEndPlay()
+{
+	check(IsMainWorldPartition());
+
+	RuntimeHash->FlushStreaming();
+	RuntimeHash->OnEndPlay();
+
+	StreamingPolicy = nullptr;
 }
 
 FName UWorldPartition::GetWorldPartitionEditorName() const
@@ -358,6 +383,8 @@ void UWorldPartition::Initialize(UWorld* InWorld, const FTransform& InTransform)
 	UWorld* OuterWorld = GetTypedOuter<UWorld>();
 	check(OuterWorld);
 
+	check(IsMainWorldPartition());
+
 	RegisterDelegates();
 
 	AWorldPartitionReplay::Initialize(World);
@@ -479,7 +506,7 @@ void UWorldPartition::Initialize(UWorld* InWorld, const FTransform& InTransform)
 	{
 		if (bIsGame || bPIEWorldTravel)
 		{
-			OnBeginPlay(false);
+			OnBeginPlay();
 		}
 
 		// Apply remapping of Persistent Level's SoftObjectPaths
@@ -531,7 +558,11 @@ void UWorldPartition::Uninitialize()
 			GEditor->OnEditorClose().RemoveAll(this);
 		}
 
-		if (!World->IsGameWorld())
+		if (World->IsGameWorld())
+		{
+			OnEndPlay();
+		}
+		else
 		{
 			if (!IsEngineExitRequested())
 			{
@@ -560,6 +591,12 @@ bool UWorldPartition::IsInitialized() const
 	return InitState == EWorldPartitionInitState::Initialized;
 }
 
+bool UWorldPartition::IsMainWorldPartition() const
+{
+	check(World);
+	return World == GetTypedOuter<UWorld>();
+}
+
 void UWorldPartition::OnPostBugItGoCalled(const FVector& Loc, const FRotator& Rot)
 {
 #if WITH_EDITOR
@@ -580,7 +617,9 @@ void UWorldPartition::RegisterDelegates()
 	if (GEditor && !IsTemplate() && !World->IsGameWorld())
 	{
 		FEditorDelegates::PreBeginPIE.AddUObject(this, &UWorldPartition::OnPreBeginPIE);
+		FEditorDelegates::PrePIEEnded.AddUObject(this, &UWorldPartition::OnPrePIEEnded);
 		FEditorDelegates::CancelPIE.AddUObject(this, &UWorldPartition::OnCancelPIE);
+		FGameDelegates::Get().GetEndPlayMapDelegate().AddUObject(this, &UWorldPartition::OnEndPlay);
 
 		FCoreUObjectDelegates::PostReachabilityAnalysis.AddUObject(this, &UWorldPartition::OnGCPostReachabilityAnalysis);
 
@@ -606,7 +645,9 @@ void UWorldPartition::UnregisterDelegates()
 	if (GEditor && !IsTemplate() && !World->IsGameWorld())
 	{
 		FEditorDelegates::PreBeginPIE.RemoveAll(this);
+		FEditorDelegates::PrePIEEnded.RemoveAll(this);
 		FEditorDelegates::CancelPIE.RemoveAll(this);
+		FGameDelegates::Get().GetEndPlayMapDelegate().RemoveAll(this);
 
 		if (!IsEngineExitRequested())
 		{
@@ -1203,14 +1244,6 @@ EWorldPartitionStreamingPerformance UWorldPartition::GetStreamingPerformance() c
 }
 
 #if WITH_EDITOR
-void UWorldPartition::PostDuplicateWorldForPIE()
-{
-	RuntimeHash->FlushStreaming();
-	RuntimeHash->OnEndPlay();
-
-	StreamingPolicy = nullptr;
-}
-
 void UWorldPartition::DrawRuntimeHashPreview()
 {
 	RuntimeHash->DrawPreview();
diff --git a/Engine/Source/Runtime/Engine/Private/WorldPartition/WorldPartitionStreamingGeneration.cpp b/Engine/Source/Runtime/Engine/Private/WorldPartition/WorldPartitionStreamingGeneration.cpp
--- a/Engine/Source/Runtime/Engine/Private/WorldPartition/WorldPartitionStreamingGeneration.cpp
+++ b/Engine/Source/Runtime/Engine/Private/WorldPartition/WorldPartitionStreamingGeneration.cpp
@@ -542,7 +542,7 @@ private:
 	TMap<FGuid, FActorContainerID> ActorDescViewsContainersMap;
 };
 
-bool UWorldPartition::GenerateStreaming(TArray<FString>* OutPackagesToGenerate, bool bIsPIE)
+bool UWorldPartition::GenerateStreaming(TArray<FString>* OutPackagesToGenerate)
 {
 	FActorDescList* ModifiedActorsDescList = nullptr;
 
diff --git a/Engine/Source/Runtime/Engine/Public/WorldPartition/WorldPartition.h b/Engine/Source/Runtime/Engine/Public/WorldPartition/WorldPartition.h
--- a/Engine/Source/Runtime/Engine/Public/WorldPartition/WorldPartition.h
+++ b/Engine/Source/Runtime/Engine/Public/WorldPartition/WorldPartition.h
@@ -91,8 +91,10 @@ private:
 
 	// PIE/Game Methods
 	void OnPreBeginPIE(bool bStartSimulate);
-	void OnBeginPlay(bool bIsPIE);
+	void OnPrePIEEnded(bool bWasSimulatingInEditor);
 	void OnCancelPIE();
+	void OnBeginPlay();
+	void OnEndPlay();
 
 	// UActorDescContainer events
 	virtual void OnActorDescAdded(FWorldPartitionActorDesc* NewActorDesc) override;
@@ -124,7 +126,7 @@ public:
 	bool RefreshLoadedEditorCells(bool bIsFromUserChange);
 
 	// PIE/Game/Cook Methods
-	bool GenerateStreaming(TArray<FString>* OutPackagesToGenerate = nullptr, bool bIsPIE = false);
+	bool GenerateStreaming(TArray<FString>* OutPackagesToGenerate = nullptr);
 	void RemapSoftObjectPath(FSoftObjectPath& ObjectPath);
 
 	// Cook Methods
@@ -172,10 +174,6 @@ public:
 
 	EWorldPartitionStreamingPerformance GetStreamingPerformance() const;
 
-#if WITH_EDITOR
-	void PostDuplicateWorldForPIE();
-#endif
-
 #if WITH_EDITORONLY_DATA
 	UPROPERTY(DuplicateTransient)
 	TObjectPtr<UWorldPartitionEditorHash> EditorHash;
@@ -190,10 +188,13 @@ public:
 	UPROPERTY(EditAnywhere, Category=RuntimeHash, NoClear, meta = (NoResetToDefault), Instanced)
 	TObjectPtr<UWorldPartitionRuntimeHash> RuntimeHash;
 
-#if WITH_EDITORONLY_DATA
+#if WITH_EDITOR
 	bool bForceGarbageCollection;
 	bool bForceGarbageCollectionPurge;
+	bool bIsPIE;
+#endif
 
+#if WITH_EDITORONLY_DATA
 	// Default HLOD layer
 	UPROPERTY(EditAnywhere, Category=HLOD)
 	TObjectPtr<class UHLODLayer> DefaultHLODLayer;
@@ -204,7 +205,7 @@ public:
 private:
 	EWorldPartitionInitState InitState;
 
-	UPROPERTY(NonPIEDuplicateTransient)
+	UPROPERTY()
 	mutable TObjectPtr<UWorldPartitionStreamingPolicy> StreamingPolicy;
 
 	TArray<IWorldPartitionStreamingSourceProvider*> StreamingSourceProviders;
@@ -216,6 +217,7 @@ private:
 	FWorldPartitionReference WorldDataLayersActor;
 #endif
 
+	bool IsMainWorldPartition() const;
 	void OnWorldMatchStarting();
 
 	void OnPostBugItGoCalled(const FVector& Loc, const FRotator& Rot);
