Fix non-unity error in simplesuspension.cpp

diff --git a/Engine/Source/Runtime/Engine/Private/PhysicsEngine/SimpleSuspension.cpp b/Engine/Source/Runtime/Engine/Private/PhysicsEngine/SimpleSuspension.cpp
--- a/Engine/Source/Runtime/Engine/Private/PhysicsEngine/SimpleSuspension.cpp
+++ b/Engine/Source/Runtime/Engine/Private/PhysicsEngine/SimpleSuspension.cpp
@@ -1,6 +1,7 @@
 // Copyright Epic Games, Inc. All Rights Reserved.
 
 #include "Physics/SimpleSuspension.h"
+#include "Chaos/Real.h"
 
 void FSimpleSuspension::Setup(const FSimpleSuspensionParams& InSuspensionParams)
 {
