Property access optimizations

Cache function return value property, as calling GetReturnProperty can be expensive (it iterates the property chain)
Move from using mem stack allocator to a inline-allocated array of bytes, manually aligning the base ptr

diff --git a/Engine/Source/Runtime/Engine/Private/Animation/AnimBlueprintGeneratedClass.cpp b/Engine/Source/Runtime/Engine/Private/Animation/AnimBlueprintGeneratedClass.cpp
--- a/Engine/Source/Runtime/Engine/Private/Animation/AnimBlueprintGeneratedClass.cpp
+++ b/Engine/Source/Runtime/Engine/Private/Animation/AnimBlueprintGeneratedClass.cpp
@@ -327,7 +327,7 @@ UAnimBlueprintGeneratedClass::UAnimBlueprintGeneratedClass(const FObjectInitiali
 const void* UAnimBlueprintGeneratedClass::GetConstantNodeValueRaw(int32 InIndex) const
 {
 	const FProperty* Property = ConstantProperties[InIndex];
-	check(Property->GetOwner<UStruct>() && GetSparseClassDataStruct()->IsChildOf(Property->GetOwner<UStruct>()));
+	checkSlow(Property->GetOwner<UStruct>() && GetSparseClassDataStruct()->IsChildOf(Property->GetOwner<UStruct>()));
 	const void* ConstantData = GetConstantNodeData();
 	check(ConstantData);
 	return Property->ContainerPtrToValuePtr<void>(static_cast<const void*>(ConstantData));
@@ -336,7 +336,7 @@ const void* UAnimBlueprintGeneratedClass::GetConstantNodeValueRaw(int32 InIndex)
 const void* UAnimBlueprintGeneratedClass::GetMutableNodeValueRaw(int32 InIndex, const UObject* InObject) const
 {
 	const FProperty* Property = MutableProperties[InIndex];
-	check(Property->GetOwner<UStruct>() && Property->GetOwner<UStruct>()->IsChildOf(FAnimBlueprintMutableData::StaticStruct()));
+	checkSlow(Property->GetOwner<UStruct>() && Property->GetOwner<UStruct>()->IsChildOf(FAnimBlueprintMutableData::StaticStruct()));
 	const FAnimBlueprintMutableData* MutableData = GetMutableNodeData(InObject);
 	check(MutableData);	
 	return Property->ContainerPtrToValuePtr<void>(static_cast<const void*>(MutableData));
@@ -346,8 +346,8 @@ FAnimBlueprintMutableData* UAnimBlueprintGeneratedClass::GetMutableNodeData(UObj
 {
 	if(MutableNodeDataProperty != nullptr)
 	{
-		check(MutableNodeDataProperty->Struct->IsChildOf(FAnimBlueprintMutableData::StaticStruct()));
-		check(MutableNodeDataProperty->GetOwner<UClass>() && InObject->GetClass()->IsChildOf(MutableNodeDataProperty->GetOwner<UClass>()));
+		checkSlow(MutableNodeDataProperty->Struct->IsChildOf(FAnimBlueprintMutableData::StaticStruct()));
+		checkSlow(MutableNodeDataProperty->GetOwner<UClass>() && InObject->GetClass()->IsChildOf(MutableNodeDataProperty->GetOwner<UClass>()));
 		return MutableNodeDataProperty->ContainerPtrToValuePtr<FAnimBlueprintMutableData>(InObject);
 	}
 
@@ -358,8 +358,8 @@ const FAnimBlueprintMutableData* UAnimBlueprintGeneratedClass::GetMutableNodeDat
 {
 	if(MutableNodeDataProperty != nullptr)
 	{
-		check(MutableNodeDataProperty->Struct->IsChildOf(FAnimBlueprintMutableData::StaticStruct()));
-		check(MutableNodeDataProperty->GetOwner<UClass>() && InObject->GetClass()->IsChildOf(MutableNodeDataProperty->GetOwner<UClass>()));
+		checkSlow(MutableNodeDataProperty->Struct->IsChildOf(FAnimBlueprintMutableData::StaticStruct()));
+		checkSlow(MutableNodeDataProperty->GetOwner<UClass>() && InObject->GetClass()->IsChildOf(MutableNodeDataProperty->GetOwner<UClass>()));
 		return MutableNodeDataProperty->ContainerPtrToValuePtr<FAnimBlueprintMutableData>(InObject);
 	}
 
@@ -376,7 +376,7 @@ void UAnimBlueprintGeneratedClass::ForEachSubsystem(TFunctionRef<EAnimSubsystemE
 	for(int32 SubsystemIndex = 0; SubsystemIndex < ConstantSubsystemProperties.Num(); ++SubsystemIndex)
 	{
 		FStructProperty* ConstantSubsystemProperty = ConstantSubsystemProperties[SubsystemIndex];
-		check(ConstantSubsystemProperty->GetOwner<UStruct>() && GetSparseClassDataStruct()->IsChildOf(ConstantSubsystemProperty->GetOwner<UStruct>()));
+		checkSlow(ConstantSubsystemProperty->GetOwner<UStruct>() && GetSparseClassDataStruct()->IsChildOf(ConstantSubsystemProperty->GetOwner<UStruct>()));
 		FAnimSubsystem& Subsystem = *ConstantSubsystemProperty->ContainerPtrToValuePtr<FAnimSubsystem>(const_cast<void*>(GetConstantNodeData()));
 
 		if(InFunction(FAnimSubsystemContext(Subsystem, ConstantSubsystemProperty->Struct)) == EAnimSubsystemEnumeration::Stop)
@@ -395,10 +395,10 @@ void UAnimBlueprintGeneratedClass::ForEachSubsystem(UObject* InObject, TFunction
 		FStructProperty* ConstantSubsystemProperty = ConstantSubsystemProperties[SubsystemIndex];
 		FStructProperty* MutableSubsystemProperty = MutableSubsystemProperties[SubsystemIndex];
 		
-		check(ConstantSubsystemProperty->GetOwner<UStruct>() && GetSparseClassDataStruct()->IsChildOf(ConstantSubsystemProperty->GetOwner<UStruct>()));
+		checkSlow(ConstantSubsystemProperty->GetOwner<UStruct>() && GetSparseClassDataStruct()->IsChildOf(ConstantSubsystemProperty->GetOwner<UStruct>()));
 		const FAnimSubsystem& Subsystem = *ConstantSubsystemProperty->ContainerPtrToValuePtr<FAnimSubsystem>(GetConstantNodeData());
 
-		check(MutableSubsystemProperty->GetOwner<UClass>() && InObject->GetClass()->IsChildOf(MutableSubsystemProperty->GetOwner<UClass>()));
+		checkSlow(MutableSubsystemProperty->GetOwner<UClass>() && InObject->GetClass()->IsChildOf(MutableSubsystemProperty->GetOwner<UClass>()));
 		FAnimSubsystemInstance& SubsystemInstance = *MutableSubsystemProperty->ContainerPtrToValuePtr<FAnimSubsystemInstance>(InObject);
 		
 		if(InFunction(FAnimSubsystemInstanceContext(Subsystem, ConstantSubsystemProperty->Struct, SubsystemInstance, MutableSubsystemProperty->Struct)) == EAnimSubsystemEnumeration::Stop)
diff --git a/Engine/Source/Runtime/Engine/Private/PropertyAccess.cpp b/Engine/Source/Runtime/Engine/Private/PropertyAccess.cpp
--- a/Engine/Source/Runtime/Engine/Private/PropertyAccess.cpp
+++ b/Engine/Source/Runtime/Engine/Private/PropertyAccess.cpp
@@ -74,6 +74,7 @@ struct FPropertyAccessSystem
 				}
 
 				Indirection.Function = Segment.Function;
+				Indirection.Property = Segment.Function->GetReturnProperty();
 				Indirection.ReturnBufferSize = Segment.Property.Get()->GetSize();
 				Indirection.ReturnBufferAlignment = Segment.Property.Get()->GetMinAlignment();
 			}
@@ -127,7 +128,7 @@ struct FPropertyAccessSystem
 
 					Indirection.Offset = GetPropertyOffset(Segment.Property.Get());
 					Indirection.Type = EPropertyAccessIndirectionType::Array;
-					Indirection.ArrayProperty = CastFieldChecked<FArrayProperty>(Segment.Property.Get());
+					Indirection.Property = CastFieldChecked<FArrayProperty>(Segment.Property.Get());
 					Indirection.ArrayIndex = ArrayIndex;
 					if(UnmodifiedFlags == EPropertyAccessSegmentFlags::ArrayOfObjects)
 					{
@@ -139,7 +140,7 @@ struct FPropertyAccessSystem
 							ExtraIndirection.Offset = 0;
 							ExtraIndirection.Type = EPropertyAccessIndirectionType::Object;
 
-							FProperty* InnerProperty = Indirection.ArrayProperty.Get()->Inner;
+							FProperty* InnerProperty = CastFieldChecked<FArrayProperty>(Indirection.Property.Get())->Inner;
 							if(FObjectProperty* ObjectProperty = CastField<FObjectProperty>(InnerProperty))
 							{
 								ExtraIndirection.ObjectType = EPropertyAccessObjectType::Object;
@@ -395,8 +396,15 @@ struct FPropertyAccessSystem
 		InFunction->Invoke(InObject, Stack, OutRetValue);
 	}
 
-	static void IterateAccess(void* InContainer, const FPropertyAccessLibrary& InLibrary, const FPropertyAccessIndirectionChain& InAccess, TFunctionRef<void(void*)> InAddressFunction)
+	// Gets the address that corresponds to a property access.
+	// Forwards the address onto the passed-in function. This callback-style approach is used because in some cases 
+	// (e.g. functions), the address may be memory allocated on the stack.
+	template<typename PredicateType>
+	static void GetAccessAddress(void* InContainer, const FPropertyAccessLibrary& InLibrary, const FPropertyAccessIndirectionChain& InAccess, PredicateType InAddressFunction)
 	{
+		// Buffer for function return values
+		TArray<uint8, TNonRelocatableInlineAllocator<64>> ArrayBuffer;
+	
 		void* Address = InContainer;
 
 		for(int32 IndirectionIndex = InAccess.IndirectionStartIndex; Address != nullptr && IndirectionIndex < InAccess.IndirectionEndIndex; ++IndirectionIndex)
@@ -439,7 +447,7 @@ struct FPropertyAccessSystem
 			}
 			case EPropertyAccessIndirectionType::Array:
 			{
-				if(FArrayProperty* ArrayProperty = Indirection.ArrayProperty.Get())
+				if(FArrayProperty* ArrayProperty = CastFieldChecked<FArrayProperty>(Indirection.Property.Get()))
 				{
 					FScriptArrayHelper Helper(ArrayProperty, static_cast<uint8*>(Address) + Indirection.Offset);
 					if(Helper.IsValidIndex(Indirection.ArrayIndex))
@@ -464,12 +472,14 @@ struct FPropertyAccessSystem
 				{
 					UObject* CalleeObject = static_cast<UObject*>(Address);
 
-					// Allocate an aligned buffer for the return value
-					Address = (uint8*)FMemStack::Get().Alloc(Indirection.ReturnBufferSize, Indirection.ReturnBufferAlignment);
+					// Allocate buffer + alignment slack for the return value
+					ArrayBuffer.SetNumUninitialized(Indirection.ReturnBufferSize + Indirection.ReturnBufferAlignment, false);
+					Address = Align(ArrayBuffer.GetData(), Indirection.ReturnBufferAlignment);
 
 					// Init value
 					check(Indirection.Function->GetReturnProperty());
-					Indirection.Function->GetReturnProperty()->InitializeValue(Address);
+					check(Indirection.Function->GetReturnProperty() == Indirection.Property.Get());
+					Indirection.Property.Get()->InitializeValue(Address);
 
 					if(Indirection.Type == EPropertyAccessIndirectionType::NativeFunction)
 					{
@@ -523,14 +533,6 @@ struct FPropertyAccessSystem
 			InAddressFunction(Address);
 		}
 	}
-	
-	// Gets the address that corresponds to a property access.
-	// Forwards the address onto the passed-in function. This callback-style approach is used because in some cases 
-	// (e.g. functions), the address may be memory allocated on the stack.
-	static void GetAccessAddress(void* InContainer, const FPropertyAccessLibrary& InLibrary, const FPropertyAccessIndirectionChain& InAccess, TFunctionRef<void(void*)> InAddressFunction)
-	{
-		IterateAccess(InContainer, InLibrary, InAccess, InAddressFunction);
-	}
 
 	// Process a single copy
 	static void ProcessCopy(UStruct* InStruct, void* InContainer, const FPropertyAccessLibrary& InLibrary, int32 InCopyIndex, int32 InBatchId, TFunctionRef<void(const FProperty*, void*)> InPostCopyOperation)
@@ -566,8 +568,6 @@ struct FPropertyAccessSystem
 	{
 		if(InLibrary.bHasBeenPostLoaded)
 		{
-			FMemMark Mark(FMemStack::Get());
-
 			if(InLibrary.CopyBatchArray.IsValidIndex(InBatchId))
 			{
 				// Copy all valid properties
diff --git a/Engine/Source/Runtime/Engine/Public/PropertyAccess.h b/Engine/Source/Runtime/Engine/Public/PropertyAccess.h
--- a/Engine/Source/Runtime/Engine/Public/PropertyAccess.h
+++ b/Engine/Source/Runtime/Engine/Public/PropertyAccess.h
@@ -130,9 +130,9 @@ struct FPropertyAccessIndirection
 private:
 	friend struct ::FPropertyAccessSystem;
 
-	// Array property if this is an array indirection
+	// Property of this indirection. Used for arrays and functions (holds the return value property for functions)
 	UPROPERTY()
-	TFieldPath<FArrayProperty> ArrayProperty;
+	TFieldPath<FProperty> Property;
 
 	// Function if this is a script of native function indirection
 	UPROPERTY()
