Non-Unity Fix

diff --git a/Engine/Source/Runtime/Engine/Private/SkeletalMeshDeformerHelpers.cpp b/Engine/Source/Runtime/Engine/Private/SkeletalMeshDeformerHelpers.cpp
--- a/Engine/Source/Runtime/Engine/Private/SkeletalMeshDeformerHelpers.cpp
+++ b/Engine/Source/Runtime/Engine/Private/SkeletalMeshDeformerHelpers.cpp
@@ -5,7 +5,6 @@
 #include "SkeletalMeshDeformerHelpers.h"
 
 #include "GPUSkinVertexFactory.h"
-#include "RenderGraphResources.h"
 #include "SkeletalRenderGPUSkin.h"
 
 FRHIShaderResourceView* FSkeletalMeshDeformerHelpers::GetBoneBufferForReading(
diff --git a/Engine/Source/Runtime/Engine/Public/SkeletalMeshDeformerHelpers.h b/Engine/Source/Runtime/Engine/Public/SkeletalMeshDeformerHelpers.h
--- a/Engine/Source/Runtime/Engine/Public/SkeletalMeshDeformerHelpers.h
+++ b/Engine/Source/Runtime/Engine/Public/SkeletalMeshDeformerHelpers.h
@@ -3,9 +3,11 @@
 #pragma once
 
 #include "CoreMinimal.h"
+#include "RenderGraphResources.h"
 
 class FSkeletalMeshObject;
 class FRDGPooledBuffer;
+class FRHIShaderResourceView;
 
 /** Functions that expose some functionality of FSkeletalMeshObject required by MeshDeformer systems. */
 class FSkeletalMeshDeformerHelpers
