Added anim node data flags to reduce overhead introduced by anim node functions

Now we dont have to fetch the data out of folded constants to determine that it is invalid

diff --git a/Engine/Source/Editor/AnimGraph/Private/AnimBlueprintCompiler.cpp b/Engine/Source/Editor/AnimGraph/Private/AnimBlueprintCompiler.cpp
--- a/Engine/Source/Editor/AnimGraph/Private/AnimBlueprintCompiler.cpp
+++ b/Engine/Source/Editor/AnimGraph/Private/AnimBlueprintCompiler.cpp
@@ -2014,6 +2014,29 @@ void FAnimBlueprintCompilerContext::ProcessFoldedPropertyRecords()
 
 	BuildAnimNodeData(ConstantPropertyRecords, NumConstantProperties);
 	BuildAnimNodeData(MutablePropertyRecords, NumMutableProperties);
+
+	const int32 NumNodes = AllocatedAnimNodeIndices.Num();
+	
+	// Patch anim node data flags
+	for(const TPair<UAnimGraphNode_Base*, int32>& GraphNodeIndexPair : AllocatedAnimNodeIndices)
+	{
+		UAnimGraphNode_Base* AnimGraphNode = GraphNodeIndexPair.Key;
+		const int32 NodeIndex = NumNodes - 1 - GraphNodeIndexPair.Value;
+		FAnimNodeData& NodeData = NewAnimBlueprintClass->AnimNodeData[NodeIndex];
+		UClass* ClassToUse = AnimGraphNode->GetBlueprintClassFromNode();
+		if(AnimGraphNode->InitialUpdateFunction.ResolveMember<UFunction>(ClassToUse))
+		{
+			NodeData.SetNodeFlags(EAnimNodeDataFlags::HasInitialUpdateFunction);
+		}
+		if(AnimGraphNode->UpdateFunction.ResolveMember<UFunction>(ClassToUse))
+		{
+			NodeData.SetNodeFlags(EAnimNodeDataFlags::HasUpdateFunction);
+		}
+		if(AnimGraphNode->BecomeRelevantFunction.ResolveMember<UFunction>(ClassToUse))
+		{
+			NodeData.SetNodeFlags(EAnimNodeDataFlags::HasBecomeRelevantFunction);
+		}
+	}
 }
 
 bool FAnimBlueprintCompilerContext::IsAnimGraphNodeFolded(UAnimGraphNode_Base* InNode) const
diff --git a/Engine/Source/Runtime/Engine/Classes/Animation/AnimClassInterface.h b/Engine/Source/Runtime/Engine/Classes/Animation/AnimClassInterface.h
--- a/Engine/Source/Runtime/Engine/Classes/Animation/AnimClassInterface.h
+++ b/Engine/Source/Runtime/Engine/Classes/Animation/AnimClassInterface.h
@@ -22,6 +22,7 @@ struct FAnimBlueprintConstantData;
 struct FAnimBlueprintMutableData;
 struct FAnimNode_Base;
 struct FAnimNodeData;
+enum class EAnimNodeDataFlags : uint32;
 namespace UE { namespace Anim { struct FNodeDataId; } }
 
 /** Describes the input and output of an anim blueprint 'function' */
@@ -227,7 +228,10 @@ public:
 		check(Subsystem);
 		return static_cast<const SubsystemType&>(*Subsystem);
 	}
-	
+
+	// Check whether a node at the specified index has the specified flags
+	static bool HasNodeAnyFlags(IAnimClassInterface* InAnimClassInterface, int32 InNodeIndex, EAnimNodeDataFlags InNodeDataFlags);
+
 protected:
 	// These direct accessors are here to allow internal access that doesnt redirect to the root class
 	virtual const TArray<FBakedAnimationStateMachine>& GetBakedStateMachines_Direct() const = 0;
@@ -297,6 +301,21 @@ public:
 
 	// Get the object ptr given an anim node
 	static const UObject* GetObjectPtrFromAnimNode(const IAnimClassInterface* InAnimClassInterface, const FAnimNode_Base* InNode);
+
+	// Get an anim node of the specified type given the object & node index
+	// Asserts if InObject is nullptr
+	// @return nullptr if the node index was out of bounds or the incorrect type
+	static const FAnimNode_Base* GetAnimNodeFromObjectPtr(const UObject* InObject, int32 InNodeIndex, UScriptStruct* InNodeType);
+
+	// Get an anim node of the specified type given the object & node index
+	// Asserts if InObject is nullptr, the node index is out of bounds or the node is the incorrect type
+	template<typename NodeType>
+	static const NodeType& GetAnimNodeFromObjectPtrChecked(const UObject* InObject, int32 InNodeIndex)
+	{
+		const FAnimNode_Base* NodePtr = GetAnimNodeFromObjectPtr(InObject, InNodeIndex, NodeType::StaticStruct());
+		check(NodePtr);
+		return static_cast<NodeType&>(*NodePtr);
+	}
 	
 	UE_DEPRECATED(4.23, "Please use GetAnimBlueprintFunctions()")
 	virtual int32 GetRootAnimNodeIndex() const { return INDEX_NONE; }
diff --git a/Engine/Source/Runtime/Engine/Private/Animation/AnimClassInterface.cpp b/Engine/Source/Runtime/Engine/Private/Animation/AnimClassInterface.cpp
--- a/Engine/Source/Runtime/Engine/Private/Animation/AnimClassInterface.cpp
+++ b/Engine/Source/Runtime/Engine/Private/Animation/AnimClassInterface.cpp
@@ -150,3 +150,25 @@ const UObject* IAnimClassInterface::GetObjectPtrFromAnimNode(const IAnimClassInt
 	FStructProperty* NodeProperty = InAnimClassInterface->GetAnimNodeProperties()[NodeIndex];
 	return ValuePtrToContainerUObjectPtr(NodeProperty, InNode, 0);
 }
+
+const FAnimNode_Base* IAnimClassInterface::GetAnimNodeFromObjectPtr(const UObject* InObject, int32 InNodeIndex, UScriptStruct* InNodeType)
+{
+	if(IAnimClassInterface* AnimClassInterface = IAnimClassInterface::GetFromClass(InObject->GetClass()))
+	{
+		check(AnimClassInterface->GetAnimNodeProperties()[InNodeIndex]);
+		FStructProperty* NodeProperty = AnimClassInterface->GetAnimNodeProperties()[InNodeIndex];
+		if(NodeProperty->Struct->IsChildOf(InNodeType))
+		{
+			return NodeProperty->ContainerPtrToValuePtr<const FAnimNode_Base>(InObject);
+		}
+	}
+
+	return nullptr;
+}
+
+bool IAnimClassInterface::HasNodeAnyFlags(IAnimClassInterface* InAnimClassInterface, int32 InNodeIndex, EAnimNodeDataFlags InNodeDataFlags)
+{
+	const TArrayView<const FAnimNodeData> AnimNodeData = InAnimClassInterface->GetNodeData();
+	check(AnimNodeData.IsValidIndex(InNodeIndex));
+	return AnimNodeData[InNodeIndex].HasNodeAnyFlags(InNodeDataFlags);
+}
\ No newline at end of file
diff --git a/Engine/Source/Runtime/Engine/Private/Animation/AnimNodeFunctionRef.cpp b/Engine/Source/Runtime/Engine/Private/Animation/AnimNodeFunctionRef.cpp
--- a/Engine/Source/Runtime/Engine/Private/Animation/AnimNodeFunctionRef.cpp
+++ b/Engine/Source/Runtime/Engine/Private/Animation/AnimNodeFunctionRef.cpp
@@ -97,7 +97,7 @@ static void CallFunctionHelper(const FAnimNodeFunctionRef& InFunction, ContextTy
 
 void FNodeFunctionCaller::InitialUpdate(const FAnimationUpdateContext& InContext, FAnimNode_Base& InNode)
 {
-	if(InNode.NodeData != nullptr)
+	if(InNode.NodeData != nullptr && InNode.NodeData->HasNodeAnyFlags(EAnimNodeDataFlags::HasInitialUpdateFunction))
 	{
 		const FAnimNodeFunctionRef& Function = InNode.GetInitialUpdateFunction();
 		if(Function.IsValid())
@@ -114,7 +114,7 @@ void FNodeFunctionCaller::InitialUpdate(const FAnimationUpdateContext& InContext
 	
 void FNodeFunctionCaller::BecomeRelevant(const FAnimationUpdateContext& InContext, FAnimNode_Base& InNode)
 {
-	if(InNode.NodeData != nullptr)
+	if(InNode.NodeData != nullptr && InNode.NodeData->HasNodeAnyFlags(EAnimNodeDataFlags::HasBecomeRelevantFunction))
 	{
 		const FAnimNodeFunctionRef& Function = InNode.GetBecomeRelevantFunction();
 		if(Function.IsValid())
@@ -131,7 +131,7 @@ void FNodeFunctionCaller::BecomeRelevant(const FAnimationUpdateContext& InContex
 	
 void FNodeFunctionCaller::Update(const FAnimationUpdateContext& InContext, FAnimNode_Base& InNode)
 {
-	if(InNode.NodeData != nullptr)
+	if(InNode.NodeData != nullptr && InNode.NodeData->HasNodeAnyFlags(EAnimNodeDataFlags::HasUpdateFunction))
 	{
 		CallFunctionHelper<FAnimUpdateContext>(InNode.GetUpdateFunction(), InContext, InNode);
 	}
diff --git a/Engine/Source/Runtime/Engine/Public/Animation/AnimNodeData.h b/Engine/Source/Runtime/Engine/Public/Animation/AnimNodeData.h
--- a/Engine/Source/Runtime/Engine/Public/Animation/AnimNodeData.h
+++ b/Engine/Source/Runtime/Engine/Public/Animation/AnimNodeData.h
@@ -66,6 +66,25 @@ private:
 
 } }	// namespace UE::Anim
 
+// The flags field of FAnimNodeData
+// Primarily this is used to prevent the extra work asscicated with recovering folded properties for anim node functions
+UENUM()
+enum class EAnimNodeDataFlags : uint32
+{
+	None						= 0x00000000,
+	
+	// This node binds its initial update function
+	HasInitialUpdateFunction	= 0x00000001,
+
+	// This node binds its become relevant function
+	HasBecomeRelevantFunction	= 0x00000002,
+
+	// This node binds its update function
+	HasUpdateFunction			= 0x00000004,
+};
+
+ENUM_CLASS_FLAGS(EAnimNodeDataFlags);
+
 // Any constant/folded class data an anim node can be accessed via this struct
 USTRUCT(BlueprintInternalUseOnly)
 struct ENGINE_API FAnimNodeData
@@ -99,10 +118,16 @@ public:
 	/** Get the node index for this constant data block. */
 	int32 GetNodeIndex() const { return NodeIndex; }
 
+	/** Check whether the node has the specified flags */
+	bool HasNodeAnyFlags(EAnimNodeDataFlags InFlags) const { return EnumHasAnyFlags<EAnimNodeDataFlags>(Flags, InFlags); }
+	
 private:
 	friend class FAnimBlueprintCompilerContext;
 	friend struct UE::Anim::FNodeDataId;
 
+	/** Set whether the node has the specified flags */
+	void SetNodeFlags(EAnimNodeDataFlags InFlags) { return EnumAddFlags<EAnimNodeDataFlags>(Flags, InFlags); }
+	
 	/** The class we are part of */
 	UPROPERTY()
 	TScriptInterface<IAnimClassInterface> AnimClassInterface = nullptr;
@@ -120,6 +145,10 @@ private:
 	 */
 	UPROPERTY()
 	int32 NodeIndex = INDEX_NONE;
+
+	/** Common flags for this node */
+	UPROPERTY()
+	EAnimNodeDataFlags Flags = EAnimNodeDataFlags::None;
 };
 
 /**
