Optimizations around reference/identity poses

Minor improvements, but simple changes.
Avoid excessive inner branching when setting a compact pose to reference pose
Avoid setting scale twice when resetting a transform to additive identity

diff --git a/Engine/Source/Runtime/Core/Public/Math/TransformNonVectorized.h b/Engine/Source/Runtime/Core/Public/Math/TransformNonVectorized.h
--- a/Engine/Source/Runtime/Core/Public/Math/TransformNonVectorized.h
+++ b/Engine/Source/Runtime/Core/Public/Math/TransformNonVectorized.h
@@ -770,6 +770,19 @@ public:
 		Scale3D = TVector<T>(1, 1, 1);
 	}
 
+	/**
+	* Sets the components to the 'additive' identity transform:
+	*   Rotation = (0,0,0,1)
+	*   Translation = (0,0,0)
+	*   Scale3D = (0,0,0)
+	*/
+	FORCEINLINE void SetIdentityZeroScale()
+	{
+		Rotation = TQuat<T>::Identity;
+		Translation = TVector<T>::ZeroVector;
+		Scale3D = TVector<T>::ZeroVector;
+	}
+	
 	/**
 	* Scales the Scale3D component by a new factor
 	* @param Scale3DMultiplier The value to multiply Scale3D with
diff --git a/Engine/Source/Runtime/Core/Public/Math/TransformVectorized.h b/Engine/Source/Runtime/Core/Public/Math/TransformVectorized.h
--- a/Engine/Source/Runtime/Core/Public/Math/TransformVectorized.h
+++ b/Engine/Source/Runtime/Core/Public/Math/TransformVectorized.h
@@ -825,6 +825,22 @@ public:
 		Scale3D = GlobalVectorConstants::Float1110; //VectorSet_W0( VectorOne() );
 	}
 
+	/**
+	 * Sets the components to the 'additive' identity transform:
+	 *   Rotation = (0,0,0,1)
+	 *   Translation = (0,0,0)
+	 *   Scale3D = (0,0,0)
+	 */
+	FORCEINLINE void SetIdentityZeroScale()
+	{
+		// Rotation = {0,0,0,1)
+		Rotation = GlobalVectorConstants::Float0001; //VectorSet_W1( VectorZero() );
+		// Translation = {0,0,0,0)
+		Translation = GlobalVectorConstants::FloatZero;
+		// Scale3D = {0,0,0,0);
+		Scale3D = GlobalVectorConstants::FloatZero;
+	}
+	
 	/**
 	 * Scales the Scale3D component by a new factor
 	 * @param Scale3DMultiplier The value to multiply Scale3D with
diff --git a/Engine/Source/Runtime/Engine/Private/Animation/BonePose.cpp b/Engine/Source/Runtime/Engine/Private/Animation/BonePose.cpp
--- a/Engine/Source/Runtime/Engine/Private/Animation/BonePose.cpp
+++ b/Engine/Source/Runtime/Engine/Private/Animation/BonePose.cpp
@@ -41,8 +41,7 @@ void FCompactPose::ResetToAdditiveIdentity()
 	{
 		for (FTransform& Bone : this->Bones)
 		{
-			Bone.SetIdentity();
-			Bone.SetScale3D(FVector::ZeroVector);
+			Bone.SetIdentityZeroScale();
 		}
 	}
 }
@@ -78,8 +77,7 @@ void FCompactHeapPose::ResetToAdditiveIdentity()
 	{
 		for (FTransform& Bone : this->Bones)
 		{
-			Bone.SetIdentity();
-			Bone.SetScale3D(FVector::ZeroVector);
+			Bone.SetIdentityZeroScale();
 		}
 	}
 }
diff --git a/Engine/Source/Runtime/Engine/Public/BoneContainer.h b/Engine/Source/Runtime/Engine/Public/BoneContainer.h
--- a/Engine/Source/Runtime/Engine/Public/BoneContainer.h
+++ b/Engine/Source/Runtime/Engine/Public/BoneContainer.h
@@ -360,6 +360,27 @@ public:
 		return CompactPoseParentBones;
 	}
 
+	// Fill the supplied buffer with the compact pose reference pose
+	template<typename ArrayType>
+	void FillWithCompactRefPose(ArrayType& OutTransforms) const
+	{
+		const int32 CompactPoseBoneCount = GetCompactPoseNumBones();
+		OutTransforms.Reset(CompactPoseBoneCount);
+		if (RefPoseOverride.IsValid())
+		{
+			OutTransforms.Append(RefPoseOverride->RefBonePoses);
+		}
+		else
+		{
+			OutTransforms.SetNumUninitialized(CompactPoseBoneCount);
+			const TArray<FTransform>& RefPoseTransforms = RefSkeleton->GetRefBonePose();
+			for (int32 CompactBoneIndex = 0; CompactBoneIndex < CompactPoseBoneCount; ++CompactBoneIndex)
+			{
+				OutTransforms[CompactBoneIndex] = RefPoseTransforms[BoneIndicesArray[CompactBoneIndex]];
+			}
+		}
+	}
+	
 	const FTransform& GetRefPoseTransform(const FCompactPoseBoneIndex& BoneIndex) const
 	{
 		if (RefPoseOverride.IsValid())
diff --git a/Engine/Source/Runtime/Engine/Public/BonePose.h b/Engine/Source/Runtime/Engine/Public/BonePose.h
--- a/Engine/Source/Runtime/Engine/Public/BonePose.h
+++ b/Engine/Source/Runtime/Engine/Public/BonePose.h
@@ -247,13 +247,8 @@ public:
 	// Sets this pose to the supplied BoneContainers ref pose
 	void ResetToRefPose(const FBoneContainer& RequiredBones)
 	{
-		const int32 CompactPoseBoneCount = RequiredBones.GetCompactPoseNumBones();
-		this->Bones.Reset(CompactPoseBoneCount);
-		for (int32 CompactBoneIndex = 0; CompactBoneIndex < CompactPoseBoneCount; ++CompactBoneIndex)
-		{
-			this->Bones.Add(RequiredBones.GetRefPoseTransform(FCompactPoseBoneIndex(CompactBoneIndex)));
-		}
-		
+		RequiredBones.FillWithCompactRefPose(this->Bones);
+
 		this->BoneContainer = &RequiredBones;
 
 		// If retargeting is disabled, copy ref pose from Skeleton, rather than mesh.
@@ -283,8 +278,7 @@ public:
 	{
 		for (FTransform& Bone : this->Bones)
 		{
-			Bone.SetIdentity();
-			Bone.SetScale3D(FVector::ZeroVector);
+			Bone.SetIdentityZeroScale();
 		}
 	}
 
