Fix missing brackets.

diff --git a/Engine/Source/Runtime/Engine/Private/WorldPartition/WorldPartitionSubsystem.cpp b/Engine/Source/Runtime/Engine/Private/WorldPartition/WorldPartitionSubsystem.cpp
--- a/Engine/Source/Runtime/Engine/Private/WorldPartition/WorldPartitionSubsystem.cpp
+++ b/Engine/Source/Runtime/Engine/Private/WorldPartition/WorldPartitionSubsystem.cpp
@@ -134,7 +134,9 @@ void UWorldPartitionSubsystem::PostInitialize()
 
 #if WITH_EDITOR
 	if (IsRunningConvertWorldPartitionCommandlet())
+	{
 		return;
+	}
 #endif
 
 	if (UWorldPartition* WorldPartition = GetWorldPartition())
