Fix a crash when running engine test
Project.Functional Tests./Game/Tests/MeshMerge/SkelMeshMergeMap.AnimationScreenShot

#rnx

diff --git a/Engine/Source/Runtime/Engine/Private/SkeletalMeshLODSettings.cpp b/Engine/Source/Runtime/Engine/Private/SkeletalMeshLODSettings.cpp
--- a/Engine/Source/Runtime/Engine/Private/SkeletalMeshLODSettings.cpp
+++ b/Engine/Source/Runtime/Engine/Private/SkeletalMeshLODSettings.cpp
@@ -81,24 +81,27 @@ bool USkeletalMeshLODSettings::SetLODSettingsToMesh(USkeletalMesh* InMesh, int32
 			}
 		}
 
-#if WITH_EDITOR
 		// copy the shared setting to mesh setting
-		// make sure we have the section in the mesh 
-		const FSkeletalMeshLODModel& LodModel = InMesh->GetImportedModel()->LODModels[LODIndex];
+#if WITH_EDITOR
+		const FSkeletalMeshLODModel* LodModel = InMesh->GetImportedModel()->LODModels.IsValidIndex(LODIndex) ? &(InMesh->GetImportedModel()->LODModels[LODIndex]) : nullptr;
 #endif
 		LODInfo->SectionsToPrioritize.Reset();
 		for (const int32 SettingSectionIndex : Setting.SectionsToPrioritize)
 		{
 			int32 FinalSettingSectionIndex = SettingSectionIndex;
 #if WITH_EDITOR
-			FinalSettingSectionIndex = INDEX_NONE;
-			for (int32 SectionIndex = 0; SectionIndex < LodModel.Sections.Num(); ++SectionIndex)
+			if (LodModel)
 			{
-				const FSkelMeshSection& Section = LodModel.Sections[SectionIndex];
-				if (Section.ChunkedParentSectionIndex == INDEX_NONE && SettingSectionIndex == Section.OriginalDataSectionIndex)
+				//Make sure all sections to prioritize are valid LODModel sections
+				FinalSettingSectionIndex = INDEX_NONE;
+				for (int32 SectionIndex = 0; SectionIndex < LodModel->Sections.Num(); ++SectionIndex)
 				{
-					FinalSettingSectionIndex = Section.OriginalDataSectionIndex;
-					break;
+					const FSkelMeshSection& Section = LodModel->Sections[SectionIndex];
+					if (Section.ChunkedParentSectionIndex == INDEX_NONE && SettingSectionIndex == Section.OriginalDataSectionIndex)
+					{
+						FinalSettingSectionIndex = Section.OriginalDataSectionIndex;
+						break;
+					}
 				}
 			}
 #endif
