Build fix to slate enum in InputChord

diff --git a/Engine/Plugins/FX/Niagara/Source/NiagaraEditor/Private/Widgets/SNiagaraScriptGraph.cpp b/Engine/Plugins/FX/Niagara/Source/NiagaraEditor/Private/Widgets/SNiagaraScriptGraph.cpp
--- a/Engine/Plugins/FX/Niagara/Source/NiagaraEditor/Private/Widgets/SNiagaraScriptGraph.cpp
+++ b/Engine/Plugins/FX/Niagara/Source/NiagaraEditor/Private/Widgets/SNiagaraScriptGraph.cpp
@@ -342,7 +342,7 @@ FReply SNiagaraScriptGraph::OnSpawnGraphNodeByShortcut(FInputChord InChord, cons
 
 	for (int32 i = 0; i < Settings->GraphCreationShortcuts.Num(); i++)
 	{
-		if (Settings->GraphCreationShortcuts[i].Input.GetRelationship(InChord) == FInputChord::Same)
+		if (Settings->GraphCreationShortcuts[i].Input.GetRelationship(InChord) == FInputChord::ERelationshipType::Same)
 		{
 			const UEdGraphSchema_Niagara* Schema = GetDefault<UEdGraphSchema_Niagara>();
 
diff --git a/Engine/Source/Runtime/Engine/Private/UserInterface/PlayerInput.cpp b/Engine/Source/Runtime/Engine/Private/UserInterface/PlayerInput.cpp
--- a/Engine/Source/Runtime/Engine/Private/UserInterface/PlayerInput.cpp
+++ b/Engine/Source/Runtime/Engine/Private/UserInterface/PlayerInput.cpp
@@ -836,12 +836,12 @@ void UPlayerInput::GetChordsForKeyMapping(const FInputActionKeyMapping& KeyMappi
 		{
 			FInputChord::ERelationshipType ChordRelationship = Chord.GetRelationship(FoundChords[ChordIndex].Chord);
 
-			if (ChordRelationship == FInputChord::Masks)
+			if (ChordRelationship == FInputChord::ERelationshipType::Masks)
 			{
 				// If we mask the found one, then remove it from the list
 				FoundChords.RemoveAtSwap(ChordIndex, 1, false);
 			}
-			else if (ChordRelationship == FInputChord::Masked)
+			else if (ChordRelationship == FInputChord::ERelationshipType::Masked)
 			{
 				bAddDelegate = false;
 				break;
@@ -955,12 +955,12 @@ void UPlayerInput::GetChordForKey(const FInputKeyBinding& KeyBinding, const bool
 				{
 					FInputChord::ERelationshipType ChordRelationship = KeyBinding.Chord.GetRelationship(FoundChords[ChordIndex].Chord);
 
-					if (ChordRelationship == FInputChord::Masks)
+					if (ChordRelationship == FInputChord::ERelationshipType::Masks)
 					{
 						// If we mask the found one, then remove it from the list
 						FoundChords.RemoveAtSwap(ChordIndex, 1, false);
 					}
-					else if (ChordRelationship == FInputChord::Masked)
+					else if (ChordRelationship == FInputChord::ERelationshipType::Masked)
 					{
 						bAddDelegate = false;
 						break;
diff --git a/Engine/Source/Runtime/Slate/Private/Framework/Commands/InputChord.cpp b/Engine/Source/Runtime/Slate/Private/Framework/Commands/InputChord.cpp
--- a/Engine/Source/Runtime/Slate/Private/Framework/Commands/InputChord.cpp
+++ b/Engine/Source/Runtime/Slate/Private/Framework/Commands/InputChord.cpp
@@ -89,7 +89,7 @@ FText FInputChord::GetModifierText(TOptional<FText> ModifierAppender) const
 
 FInputChord::ERelationshipType FInputChord::GetRelationship( const FInputChord& OtherChord ) const
 {
-	ERelationshipType Relationship = None;
+	ERelationshipType Relationship = ERelationshipType::None;
 
 	if (Key == OtherChord.Key)
 	{
@@ -98,21 +98,21 @@ FInputChord::ERelationshipType FInputChord::GetRelationship( const FInputChord&
 			(bShift == OtherChord.bShift) &&
 			(bCmd == OtherChord.bCmd))
 		{
-			Relationship = Same;
+			Relationship = ERelationshipType::Same;
 		}
 		else if ((bAlt || !OtherChord.bAlt) &&
 				(bCtrl || !OtherChord.bCtrl) &&
 				(bShift || !OtherChord.bShift) &&
 				(bCmd || !OtherChord.bCmd))
 		{
-			Relationship = Masks;
+			Relationship = ERelationshipType::Masks;
 		}
 		else if ((!bAlt || OtherChord.bAlt) &&
 				(!bCtrl || OtherChord.bCtrl) &&
 				(!bShift || OtherChord.bShift) &&
 				(!bCmd || OtherChord.bCmd))
 		{
-			Relationship = Masked;
+			Relationship = ERelationshipType::Masked;
 		}
 	}
 
diff --git a/Engine/Source/Runtime/Slate/Public/Framework/Commands/InputChord.h b/Engine/Source/Runtime/Slate/Public/Framework/Commands/InputChord.h
--- a/Engine/Source/Runtime/Slate/Public/Framework/Commands/InputChord.h
+++ b/Engine/Source/Runtime/Slate/Public/Framework/Commands/InputChord.h
@@ -42,7 +42,7 @@ struct SLATE_API FInputChord
 	  * The ways two chords can be related to each other. A chord is considered masking 
 	  * when it has all the same modifier keys as another chord plus more.
 	  */
-	enum ERelationshipType
+	enum class ERelationshipType
 	{
 		None,
 		Same,
