World Partition Streaming Generation
- Added FStreamingGenerationNullErrorHandler to log into.
- Enforced having an error handler during generation to cleanup code.

diff --git a/Engine/Source/Runtime/Engine/Private/WorldPartition/WorldPartitionStreamingGeneration.cpp b/Engine/Source/Runtime/Engine/Private/WorldPartition/WorldPartitionStreamingGeneration.cpp
--- a/Engine/Source/Runtime/Engine/Private/WorldPartition/WorldPartitionStreamingGeneration.cpp
+++ b/Engine/Source/Runtime/Engine/Private/WorldPartition/WorldPartitionStreamingGeneration.cpp
@@ -9,6 +9,7 @@
 #include "WorldPartition/WorldPartitionRuntimeHash.h"
 #include "WorldPartition/WorldPartitionStreamingPolicy.h"
 #include "WorldPartition/WorldPartitionActorCluster.h"
+#include "WorldPartition/ErrorHandling/WorldPartitionStreamingGenerationNullErrorHandler.h"
 #include "WorldPartition/ErrorHandling/WorldPartitionStreamingGenerationLogErrorHandler.h"
 #include "WorldPartition/ErrorHandling/WorldPartitionStreamingGenerationMapCheckErrorHandler.h"
 #include "WorldPartition/HLOD/HLODActor.h"
@@ -181,7 +182,7 @@ class FWorldPartitionStreamingGenerator
 			const FActorContainerID& ContainerID = It.Key();
 			FContainerDescriptor& ContainerDescriptor = It.Value();
 
-			if (ContainerID.IsMainContainer() && ContainerDescriptor.Container->GetWorld() != nullptr)
+			if (ContainerID.IsMainContainer() && ContainerDescriptor.Container->GetWorld())
 			{
 				// Gather all references to external actors from the level script and make them always loaded
 				if (ULevelScriptBlueprint* LevelScriptBlueprint = ContainerDescriptor.Container->GetWorld()->PersistentLevel->GetLevelScriptBlueprint(true))
@@ -194,21 +195,13 @@ class FWorldPartitionStreamingGenerator
 
 						if (ActorDescView.GetIsSpatiallyLoaded())
 						{
-							if (ErrorHandler)
-							{
-								ErrorHandler->OnInvalidReferenceLevelScriptStreamed(ActorDescView);
-							}
-
+							ErrorHandler->OnInvalidReferenceLevelScriptStreamed(ActorDescView);
 							ActorDescView.SetForcedNonSpatiallyLoaded();
 						}
 
 						if (ActorDescView.GetDataLayers().Num())
 						{
-							if (ErrorHandler)
-							{
-								ErrorHandler->OnInvalidReferenceLevelScriptDataLayers(ActorDescView);
-							}
-
+							ErrorHandler->OnInvalidReferenceLevelScriptDataLayers(ActorDescView);
 							ActorDescView.SetInvalidDataLayers();
 						}
 					}
@@ -222,10 +215,11 @@ class FWorldPartitionStreamingGenerator
 			// 
 			// This works because fixes are deterministic and always apply the same way to both Actors being modified, so there's no ordering issues possible
 			uint32 NbValidationPasses = 0;
+			uint32 NbErrorsDetected;
 
-			while(true)
+			do
 			{
-				uint32 NbErrorsDetected = 0;
+				NbErrorsDetected = 0;
 
 				for (auto ActorDescIt = ContainerDescriptor.ActorDescViewMap.CreateIterator(); ActorDescIt; ++ActorDescIt)
 				{
@@ -304,7 +298,7 @@ class FWorldPartitionStreamingGenerator
 
 							if (bIsActorDescSpatiallyLoaded != bIsActorDescRefSpatiallyLoaded)
 							{
-								if (ErrorHandler && !NbValidationPasses)
+								if (!NbValidationPasses)
 								{
 									ErrorHandler->OnInvalidReferenceGridPlacement(*RefererActorDescView, *ReferenceActorDescView);									
 								}
@@ -319,7 +313,7 @@ class FWorldPartitionStreamingGenerator
 
 							if (!IsReferenceDataLayersValid(*RefererActorDescView, *ReferenceActorDescView))
 							{
-								if (ErrorHandler && !NbValidationPasses)
+								if (!NbValidationPasses)
 								{
 									ErrorHandler->OnInvalidReferenceDataLayers(*RefererActorDescView, *ReferenceActorDescView);									
 								}
@@ -334,7 +328,7 @@ class FWorldPartitionStreamingGenerator
 
 							if (ReferenceActorDescView->GetRuntimeGrid() != RefererActorDescView->GetRuntimeGrid())
 							{
-								if (ErrorHandler && !NbValidationPasses)
+								if (!NbValidationPasses)
 								{
 									ErrorHandler->OnInvalidReferenceRuntimeGrid(*RefererActorDescView, *ReferenceActorDescView);
 								}
@@ -348,7 +342,7 @@ class FWorldPartitionStreamingGenerator
 							}
 
 						}
-						else if (ErrorHandler)
+						else
 						{
 							if (!NbValidationPasses)
 							{
@@ -360,11 +354,8 @@ class FWorldPartitionStreamingGenerator
 				}		
 
 				NbValidationPasses++;
-				if (NbErrorsDetected == 0)
-				{
-					break;
-				}
-			} 
+			}
+			while (NbErrorsDetected);
 		}
 	}
 
@@ -400,9 +391,9 @@ class FWorldPartitionStreamingGenerator
 	}
 
 public:
-	FWorldPartitionStreamingGenerator(FActorDescList* InModifiedActorsDescList, IStreamingGenerationErrorHandler* InErrorHandler = nullptr)
+	FWorldPartitionStreamingGenerator(FActorDescList* InModifiedActorsDescList, IStreamingGenerationErrorHandler* InErrorHandler)
 	: ModifiedActorsDescList(InModifiedActorsDescList)
-	, ErrorHandler(InErrorHandler)
+	, ErrorHandler(InErrorHandler ? InErrorHandler : &NullErrorHandler)
 	{}
 
 	void PreparationPhase(const UActorDescContainer* Container)
@@ -515,6 +506,7 @@ public:
 private:
 	FActorDescList* ModifiedActorsDescList;
 	IStreamingGenerationErrorHandler* ErrorHandler;
+	FStreamingGenerationNullErrorHandler NullErrorHandler;
 
 	struct FContainerDescriptor
 	{
diff --git a/Engine/Source/Runtime/Engine/Public/WorldPartition/ErrorHandling/WorldPartitionStreamingGenerationNullErrorHandler.h b/Engine/Source/Runtime/Engine/Public/WorldPartition/ErrorHandling/WorldPartitionStreamingGenerationNullErrorHandler.h
new file mode 100644
--- /dev/null
+++ b/Engine/Source/Runtime/Engine/Public/WorldPartition/ErrorHandling/WorldPartitionStreamingGenerationNullErrorHandler.h
@@ -0,0 +1,17 @@
+// Copyright Epic Games, Inc. All Rights Reserved.
+#pragma once
+
+#if WITH_EDITOR
+#include "WorldPartition/ErrorHandling/WorldPartitionStreamingGenerationErrorHandler.h"
+
+class ENGINE_API FStreamingGenerationNullErrorHandler : public IStreamingGenerationErrorHandler
+{
+public:
+	virtual void OnInvalidReference(const FWorldPartitionActorDescView& ActorDescView, const FGuid& ReferenceGuid) override {}
+	virtual void OnInvalidReferenceGridPlacement(const FWorldPartitionActorDescView& ActorDescView, const FWorldPartitionActorDescView& ReferenceActorDescView) override {}
+	virtual void OnInvalidReferenceDataLayers(const FWorldPartitionActorDescView& ActorDescView, const FWorldPartitionActorDescView& ReferenceActorDescView) override {}
+	virtual void OnInvalidReferenceLevelScriptStreamed(const FWorldPartitionActorDescView& ActorDescView) override {}
+	virtual void OnInvalidReferenceLevelScriptDataLayers(const FWorldPartitionActorDescView& ActorDescView) override {}
+	virtual void OnInvalidReferenceRuntimeGrid(const FWorldPartitionActorDescView& ActorDescView, const FWorldPartitionActorDescView& ReferenceActorDescView) override {}
+};
+#endif
