-Fix WorldPartition map load/run in editor dedicated server mode.

PR #8424: WorldPartition "UnrealEditor.exe -server" crash (Contributed by [REDACTED_PERSON])

diff --git a/Engine/Source/Runtime/Engine/Private/WorldPartition/WorldPartition.cpp b/Engine/Source/Runtime/Engine/Private/WorldPartition/WorldPartition.cpp
--- a/Engine/Source/Runtime/Engine/Private/WorldPartition/WorldPartition.cpp
+++ b/Engine/Source/Runtime/Engine/Private/WorldPartition/WorldPartition.cpp
@@ -393,6 +393,7 @@ void UWorldPartition::Initialize(UWorld* InWorld, const FTransform& InTransform)
 	const bool bIsGame = IsRunningGame();
 	const bool bIsEditor = !World->IsGameWorld();
 	const bool bIsCooking = IsRunningCookCommandlet();
+	const bool bIsDedicatedServer = IsRunningDedicatedServer();
 
 	UE_LOG(LogWorldPartition, Log, TEXT("UWorldPartition::Initialize(IsEditor=%d, IsGame=%d, IsCooking=%d)"), bIsEditor ? 1 : 0, bIsGame ? 1 : 0, bIsCooking ? 1 : 0);
 
@@ -416,7 +417,7 @@ void UWorldPartition::Initialize(UWorld* InWorld, const FTransform& InTransform)
 	// Did we travel into a WP map in PIE (null StreamingPolicy means GenerateStreaming wasn't called)
 	const bool bPIEWorldTravel = World->WorldType == EWorldType::PIE && !StreamingPolicy;
 
-	if (bIsEditor || bIsGame || bPIEWorldTravel)
+	if (bIsEditor || bIsGame || bPIEWorldTravel || bIsDedicatedServer)
 	{
 		UPackage* LevelPackage = OuterWorld->PersistentLevel->GetOutermost();
 		const FName PackageName = LevelPackage->GetLoadedPath().GetPackageFName();
@@ -504,7 +505,7 @@ void UWorldPartition::Initialize(UWorld* InWorld, const FTransform& InTransform)
 #if WITH_EDITOR
 	if (!bIsEditor)
 	{
-		if (bIsGame || bPIEWorldTravel)
+		if (bIsGame || bPIEWorldTravel || bIsDedicatedServer)
 		{
 			OnBeginPlay();
 		}
diff --git a/Engine/Source/Runtime/Engine/Private/WorldPartition/WorldPartitionLevelStreamingDynamic.cpp b/Engine/Source/Runtime/Engine/Private/WorldPartition/WorldPartitionLevelStreamingDynamic.cpp
--- a/Engine/Source/Runtime/Engine/Private/WorldPartition/WorldPartitionLevelStreamingDynamic.cpp
+++ b/Engine/Source/Runtime/Engine/Private/WorldPartition/WorldPartitionLevelStreamingDynamic.cpp
@@ -190,7 +190,7 @@ bool UWorldPartitionLevelStreamingDynamic::RequestLevel(UWorld* InPersistentWorl
 			check(CellLevelPackage);
 			check(UWorld::FindWorldInPackage(CellLevelPackage));
 			check(RuntimeLevel->OwningWorld);
-			check(RuntimeLevel->OwningWorld->WorldType == EWorldType::PIE || (IsRunningGame() && RuntimeLevel->OwningWorld->WorldType == EWorldType::Game));
+			check(RuntimeLevel->OwningWorld->WorldType == EWorldType::PIE || ((IsRunningGame() || IsRunningDedicatedServer()) && RuntimeLevel->OwningWorld->WorldType == EWorldType::Game));
 	
 			if (IssueLoadRequests())
 			{
@@ -346,7 +346,7 @@ void UWorldPartitionLevelStreamingDynamic::FinalizeRuntimeLevel()
 		}
 		else if (OuterWorld->IsGameWorld())
 		{
-			check(IsRunningGame());
+			check(IsRunningGame() || IsRunningDedicatedServer());
 			// Remap Runtime Level's SoftObjectPaths
 			FWorldPartitionLevelHelper::RemapLevelSoftObjectPaths(RuntimeLevel, OuterWorldPartition.Get());
 		}
@@ -477,4 +477,4 @@ UWorld* UWorldPartitionLevelStreamingDynamic::GetOuterWorld() const
 	return OuterWorldPartition->GetTypedOuter<UWorld>();
 }
 
-#undef LOCTEXT_NAMESPACE
\ No newline at end of file
+#undef LOCTEXT_NAMESPACE
