Fix ChildActor Being Set To Null Via Replication.  We dont need to worry about Name Remapping or Destruction Info when we RerunContructionScripts.

diff --git a/Engine/Source/Runtime/Engine/Private/NetDriver.cpp b/Engine/Source/Runtime/Engine/Private/NetDriver.cpp
--- a/Engine/Source/Runtime/Engine/Private/NetDriver.cpp
+++ b/Engine/Source/Runtime/Engine/Private/NetDriver.cpp
@@ -3360,7 +3360,7 @@ void UNetDriver::NotifyActorDestroyed( AActor* ThisActor, bool IsSeamlessTravel
 
 		const bool bIsActorStatic = !GuidCache->IsDynamicObject( ThisActor );
 		const bool bActorHasRole = ThisActor->GetRemoteRole() != ROLE_None;
-		const bool bShouldCreateDestructionInfo = bIsServer && bIsActorStatic && bActorHasRole && !IsSeamlessTravel;
+		const bool bShouldCreateDestructionInfo = bIsServer && bIsActorStatic && bActorHasRole && !IsSeamlessTravel && !GIsReconstructingBlueprintInstances;
 
 		if (bShouldCreateDestructionInfo)
 		{
@@ -3434,7 +3434,8 @@ void UNetDriver::NotifyActorRenamed(AActor* ThisActor, FName PreviousName)
 #if WITH_EDITOR
 	// When recompiling and reinstancing a Blueprint, we rename the old actor out of the way, which would cause this code to emit a warning during PIE
 	// Since that old actor is about to die (on both client and server) in the reinstancing case, it's safe to skip
-	if (GCompilingBlueprint)
+	// We also want to skip when we are reconstructing Blueprint Instances as we dont want to store renaming as components are destroyed and created again
+	if (GCompilingBlueprint || GIsReconstructingBlueprintInstances)
 	{
 		return;
 	}
