Removing invalid check in SoundWaveProxy ctor

diff --git a/Engine/Source/Runtime/Engine/Private/SoundWave.cpp b/Engine/Source/Runtime/Engine/Private/SoundWave.cpp
--- a/Engine/Source/Runtime/Engine/Private/SoundWave.cpp
+++ b/Engine/Source/Runtime/Engine/Private/SoundWave.cpp
@@ -3395,7 +3395,6 @@ FSoundWaveProxy::FSoundWaveProxy(USoundWave* InWave)
 
 	if (InWave->SoundWaveDataPtr->RunningPlatformData.NumChunks > 0)
 	{
-		check(SoundWaveDataPtr->LoadingBehavior != ESoundWaveLoadingBehavior::ForceInline);
 		NumChunks = InWave->GetNumChunks();
 	}
 
