eTool tip fixes for Audio Stream Caching statics

diff --git a/Engine/Source/Runtime/Engine/Classes/Kismet/GameplayStatics.h b/Engine/Source/Runtime/Engine/Classes/Kismet/GameplayStatics.h
--- a/Engine/Source/Runtime/Engine/Classes/Kismet/GameplayStatics.h
+++ b/Engine/Source/Runtime/Engine/Classes/Kismet/GameplayStatics.h
@@ -884,11 +884,11 @@ public:
 	UFUNCTION(BlueprintCallable, Category = "Audio")
 	static void PrimeSound(USoundBase* InSound);
 
-	/** Primes the sound sound waves in the given USoundClass, caching the first chunk of streamed audio. **/
+	/** Primes the sound waves in the given USoundClass, caching the first chunk of streamed audio. */
 	UFUNCTION(BlueprintCallable, Category = "Audio")
 	static void PrimeAllSoundsInSoundClass(class USoundClass* InSoundClass);
 
-	/** Iterate through all soundwaves an release and handles to retained chunks. (if the chunk is not being played, the chunk will be up for eviction) **/
+	/** Iterate through all sound waves and releases handles to retained chunks. (If the chunk is not being played it will be up for eviction) */
 	UFUNCTION(BlueprintCallable, Category = "Audio")
 	static void UnRetainAllSoundsInSoundClass(class USoundClass* InSoundClass);
 
