Force at least a single element to be in the Audio Streaming Cache to avoid crashes

diff --git a/Engine/Source/Runtime/Engine/Private/AudioCompressionSettingsUtils.cpp b/Engine/Source/Runtime/Engine/Private/AudioCompressionSettingsUtils.cpp
--- a/Engine/Source/Runtime/Engine/Private/AudioCompressionSettingsUtils.cpp
+++ b/Engine/Source/Runtime/Engine/Private/AudioCompressionSettingsUtils.cpp
@@ -395,7 +395,7 @@ FCachedAudioStreamingManagerParams FPlatformCompressionUtilities::BuildCachedStr
 	// Primary cache defined here:
 	CacheDimensions.MaxChunkSize = 256 * 1024; // max possible chunk size (hard coded for legacy streaming path)
 	CacheDimensions.MaxMemoryInBytes = CacheSettings.CacheSizeKB * 1024;
-	CacheDimensions.NumElements = NumElements;
+	CacheDimensions.NumElements = FMath::Max(NumElements, 1); // force at least a single cache element to avoid crashes
 	Params.Caches.Add(CacheDimensions);
 
 	// TODO: When settings are added to support multiple sub-caches, add it here.
