GitHub 6075 : Make COND_Initial replication work on per UObject, not per AActor
#6075

diff --git a/Engine/Source/Runtime/Engine/Private/DataChannel.cpp b/Engine/Source/Runtime/Engine/Private/DataChannel.cpp
--- a/Engine/Source/Runtime/Engine/Private/DataChannel.cpp
+++ b/Engine/Source/Runtime/Engine/Private/DataChannel.cpp
@@ -94,6 +94,12 @@ static FAutoConsoleVariableRef CVarDormancyHysteresis(
 	TEXT("When > 0, represents the time we'll wait before letting a channel become fully dormant (in seconds). This can prevent churn when objects are going in and out of dormant more frequently than normal.")
 );
 
+static TAutoConsoleVariable<bool> CVarEnableNetInitialSubObjects(
+	TEXT("net.EnableNetInitialSubObjects"),
+	true,
+	TEXT("Enables new SubObjects to set bNetInitial to true to make sure all replicated properties are replicated.")
+);
+
 template<typename T>
 static const bool IsBunchTooLarge(UNetConnection* Connection, T* Bunch)
 {
@@ -4377,6 +4383,7 @@ bool UActorChannel::ReplicateSubobject(UObject *Obj, FOutBunch &Bunch, const FRe
 
 	bool NewSubobject = false;
 	
+	FReplicationFlags ObjRepFlags = RepFlags;
 	TSharedRef<FObjectReplicator>& ObjectReplicator = !bFoundReplicator ? CreateReplicator(Obj, !bFoundInvalidReplicator) : *FoundReplicator;
 	if (!bFoundReplicator)
 	{
@@ -4387,9 +4394,13 @@ bool UActorChannel::ReplicateSubobject(UObject *Obj, FOutBunch &Bunch, const FRe
 		// to spawn this on the client).
 		Bunch.bReliable = true;
 		NewSubobject = true;
+		if (CVarEnableNetInitialSubObjects.GetValueOnAnyThread())
+		{
+			ObjRepFlags.bNetInitial = true;
+		}	
 	}
 	UE_NET_TRACE_OBJECT_SCOPE(ObjectReplicator->ObjectNetGUID, Bunch, GetTraceCollector(Bunch), ENetTraceVerbosity::Trace);
-	bool bWroteSomething = ObjectReplicator.Get().ReplicateProperties(Bunch, RepFlags);
+	bool bWroteSomething = ObjectReplicator.Get().ReplicateProperties(Bunch, ObjRepFlags);
 	if (NewSubobject && !bWroteSomething)
 	{
 		// Write empty payload to force object creation
