Adding SourceSettings support for SourceDataOverride audio extension.

diff --git a/Engine/Source/Developer/AudioSettingsEditor/Private/TargetPlatformAudioCustomization.cpp b/Engine/Source/Developer/AudioSettingsEditor/Private/TargetPlatformAudioCustomization.cpp
--- a/Engine/Source/Developer/AudioSettingsEditor/Private/TargetPlatformAudioCustomization.cpp
+++ b/Engine/Source/Developer/AudioSettingsEditor/Private/TargetPlatformAudioCustomization.cpp
@@ -306,7 +306,7 @@ void FAudioPluginWidgetManager::BuildAudioCategory(IDetailLayoutBuilder& DetailL
 			MakeAudioPluginSelectorWidget(AudioOcclusionPropertyHandle, EAudioPlugin::OCCLUSION, PlatformName)
 		];
 		
-		TSharedPtr<IPropertyHandle> AudioPropagationPropertyHandle = DetailLayout.GetProperty("PropagationPlugin", ClassOuterMost);
+		TSharedPtr<IPropertyHandle> AudioPropagationPropertyHandle = DetailLayout.GetProperty("SourceDataOverridePlugin", ClassOuterMost);
 		IDetailPropertyRow& AudioPropagationPropertyRow = AudioCategory.AddProperty(AudioPropagationPropertyHandle);
 
 		AudioPropagationPropertyRow.CustomWidget()
diff --git a/Engine/Source/Editor/DetailCustomizations/Private/AttenuationSettingsCustomizations.cpp b/Engine/Source/Editor/DetailCustomizations/Private/AttenuationSettingsCustomizations.cpp
--- a/Engine/Source/Editor/DetailCustomizations/Private/AttenuationSettingsCustomizations.cpp
+++ b/Engine/Source/Editor/DetailCustomizations/Private/AttenuationSettingsCustomizations.cpp
@@ -428,6 +428,7 @@ void FSoundAttenuationSettingsCustomization::CustomizeChildren(TSharedRef<IPrope
 	bIsReverbSendEnabledHandle = PropertyHandles.FindChecked(GET_MEMBER_NAME_CHECKED(FSoundAttenuationSettings, bEnableReverbSend)).ToSharedRef();
 	bIsPriorityAttenuationEnabledHandle = PropertyHandles.FindChecked(GET_MEMBER_NAME_CHECKED(FSoundAttenuationSettings, bEnablePriorityAttenuation)).ToSharedRef();
 	bIsSubmixSendAttenuationEnabledHandle = PropertyHandles.FindChecked(GET_MEMBER_NAME_CHECKED(FSoundAttenuationSettings, bEnableSubmixSends)).ToSharedRef();
+	bIsSourceDataOverrideEnabledHandle = PropertyHandles.FindChecked(GET_MEMBER_NAME_CHECKED(FSoundAttenuationSettings, bEnableSourceDataOverride)).ToSharedRef();
 	ReverbSendMethodHandle = PropertyHandles.FindChecked(GET_MEMBER_NAME_CHECKED(FSoundAttenuationSettings, ReverbSendMethod)).ToSharedRef();
 	PriorityAttenuationMethodHandle = PropertyHandles.FindChecked(GET_MEMBER_NAME_CHECKED(FSoundAttenuationSettings, PriorityAttenuationMethod)).ToSharedRef();
 	AbsorptionMethodHandle = PropertyHandles.FindChecked(GET_MEMBER_NAME_CHECKED(FSoundAttenuationSettings, AbsorptionMethod)).ToSharedRef();
@@ -665,6 +666,10 @@ void FSoundAttenuationSettingsCustomization::CustomizeChildren(TSharedRef<IPrope
 		.Visibility(IsAttenuationOverriddenVisibleAttribute())
 		.EditCondition(GetIsSubmixSendAttenuationEnabledAttribute(), nullptr);
 
+	LayoutBuilder.AddPropertyToCategory(bIsSourceDataOverrideEnabledHandle)
+		.Visibility(IsAttenuationOverriddenVisibleAttribute())
+		.EditCondition(IsAttenuationOverriddenAttribute(), nullptr);
+
 	if (bIsAudioMixerEnabled)
 	{
 		TSharedPtr<IPropertyHandle> PluginProperty = PropertyHandles.FindChecked(GET_MEMBER_NAME_CHECKED(FSoundAttenuationSettings, PluginSettings));
@@ -688,9 +693,10 @@ void FSoundAttenuationSettingsCustomization::CustomizeChildren(TSharedRef<IPrope
 	LayoutBuilder.EditCategory("AttenuationPriority", LOCTEXT("AttenuationPriority", "Attenuation (Priority)"));
 	LayoutBuilder.EditCategory("AttenuationAirAbsorption", LOCTEXT("AttenuationAirAbsorption", "Attenuation (Air Absorption)"));
 	LayoutBuilder.EditCategory("AttenuationPluginSettings", LOCTEXT("AttenuationPluginSettings", "Attenuation (Plugin Settings)"));
+	LayoutBuilder.EditCategory("AttenuationSourceDataOverride", LOCTEXT("AttenuationSourceDataOverride", "Attenuation (Source Data Override)"));
 	LayoutBuilder.SortCategories(AttenuationSettingsUtils::SortCategories);
 
-	if (PropertyHandles.Num() != 65)
+	if (PropertyHandles.Num() != 66)
 	{
 		ensureMsgf(false, TEXT("Unexpected property handle(s) customizing FSoundAttenuationSettings. %d handles found"), PropertyHandles.Num());
 	}
diff --git a/Engine/Source/Editor/DetailCustomizations/Private/AttenuationSettingsCustomizations.h b/Engine/Source/Editor/DetailCustomizations/Private/AttenuationSettingsCustomizations.h
--- a/Engine/Source/Editor/DetailCustomizations/Private/AttenuationSettingsCustomizations.h
+++ b/Engine/Source/Editor/DetailCustomizations/Private/AttenuationSettingsCustomizations.h
@@ -58,6 +58,7 @@ protected:
 	TSharedPtr<IPropertyHandle> bIsReverbSendEnabledHandle;
 	TSharedPtr<IPropertyHandle> bIsPriorityAttenuationEnabledHandle;
 	TSharedPtr<IPropertyHandle> bIsSubmixSendAttenuationEnabledHandle;
+	TSharedPtr<IPropertyHandle> bIsSourceDataOverrideEnabledHandle;
 	TSharedPtr<IPropertyHandle> ReverbSendMethodHandle;
 	TSharedPtr<IPropertyHandle> PriorityAttenuationMethodHandle;
 	TSharedPtr<IPropertyHandle> AbsorptionMethodHandle;
diff --git a/Engine/Source/Runtime/AudioExtensions/Public/IAudioExtensionPlugin.h b/Engine/Source/Runtime/AudioExtensions/Public/IAudioExtensionPlugin.h
--- a/Engine/Source/Runtime/AudioExtensions/Public/IAudioExtensionPlugin.h
+++ b/Engine/Source/Runtime/AudioExtensions/Public/IAudioExtensionPlugin.h
@@ -367,6 +367,14 @@ public:
 	}
 };
 
+
+/** This is a class which should be overridden to provide users with settings to use for individual sounds */
+UCLASS(config = Engine, abstract, editinlinenew, BlueprintType)
+class AUDIOEXTENSIONS_API USourceDataOverridePluginSourceSettingsBase : public UObject
+{
+	GENERATED_BODY()
+};
+
 /************************************************************************/
 /* IAudioSourceDataOverrideFactory										*/
 /* Implement this modular feature to make your SourceDataOverride plugin*/
@@ -395,6 +403,14 @@ public:
 	}
 	/* End IAudioPluginWithMetadata implementation */
 
+		/**
+	* @return the UClass type of your settings for source data overrides. This allows us to only pass in user settings for your plugin.
+	*/
+	virtual UClass* GetCustomSourceDataOverrideSettingsClass() const
+	{
+		return nullptr;
+	}
+
 	/**
 	* @return a new instance of your source data override plugin, owned by a shared pointer.
 	*/
@@ -417,7 +433,7 @@ public:
 	}
 
 	/** Called when a source is assigned to a voice. */
-	virtual void OnInitSource(const uint32 SourceId)
+	virtual void OnInitSource(const uint32 SourceId, const FName& AudioComponentUserId, USourceDataOverridePluginSourceSettingsBase* InSettings)
 	{
 	}
 
diff --git a/Engine/Source/Runtime/AudioMixer/Private/AudioMixerSource.cpp b/Engine/Source/Runtime/AudioMixer/Private/AudioMixerSource.cpp
--- a/Engine/Source/Runtime/AudioMixer/Private/AudioMixerSource.cpp
+++ b/Engine/Source/Runtime/AudioMixer/Private/AudioMixerSource.cpp
@@ -570,6 +570,9 @@ namespace Audio
 			// Grab the source's reverb plugin settings
 			InitParams.ReverbPluginSettings = UseReverbPlugin() ? WaveInstance->ReverbPluginSettings : nullptr;
 
+			// Grab the source's source data override plugin settings
+			InitParams.SourceDataOverridePluginSettings = UseSourceDataOverridePlugin() ? WaveInstance->SourceDataOverridePluginSettings : nullptr;
+
 			// Update the buffer sample rate to the wave instance sample rate in case it was serialized incorrectly
 			MixerBuffer->InitSampleRate(WaveData->GetSampleRateForCurrentPlatform());
 
@@ -613,7 +616,7 @@ namespace Audio
 				if (AudioDevice->SourceDataOverridePluginInterface)
 				{
 					uint32 SourceId = MixerSourceVoice->GetSourceId();
-					AudioDevice->SourceDataOverridePluginInterface->OnInitSource(SourceId);
+					AudioDevice->SourceDataOverridePluginInterface->OnInitSource(SourceId, InitParams.AudioComponentUserID, InitParams.SourceDataOverridePluginSettings);
 				}
 
 				InitializationState = EMixerSourceInitializationState::Initialized;
@@ -758,7 +761,7 @@ namespace Audio
 		++TickCount;
 
 		// Allow plugins to override any data in a waveinstance
-		if (AudioDevice->SourceDataOverridePluginInterface)
+		if (AudioDevice->SourceDataOverridePluginInterface && WaveInstance->bEnableSourceDataOverride)
 		{
 			uint32 SourceId = MixerSourceVoice->GetSourceId();
 			int32 ListenerIndex = WaveInstance->ActiveSound->GetClosestListenerIndex();
@@ -1746,4 +1749,11 @@ namespace Audio
 			AudioDevice->IsReverbPluginEnabled() &&
 			WaveInstance->ReverbPluginSettings != nullptr;
 	}
+
+	bool FMixerSource::UseSourceDataOverridePlugin() const
+	{
+		return (Buffer->NumChannels == 1 || Buffer->NumChannels == 2) &&
+			AudioDevice->IsSourceDataOverridePluginEnabled() &&
+			WaveInstance->SourceDataOverridePluginSettings != nullptr;
+	}
 }
diff --git a/Engine/Source/Runtime/AudioMixer/Private/AudioMixerSource.h b/Engine/Source/Runtime/AudioMixer/Private/AudioMixerSource.h
--- a/Engine/Source/Runtime/AudioMixer/Private/AudioMixerSource.h
+++ b/Engine/Source/Runtime/AudioMixer/Private/AudioMixerSource.h
@@ -109,6 +109,9 @@ namespace Audio
 		/** Whether or not to use the reverb plugin. */
 		bool UseReverbPlugin() const;
 
+		/** Whether or not to use the source data override plugin */
+		bool UseSourceDataOverridePlugin() const;
+
 	private:
 
 		FMixerDevice* MixerDevice;
diff --git a/Engine/Source/Runtime/AudioMixer/Private/AudioMixerSourceManager.cpp b/Engine/Source/Runtime/AudioMixer/Private/AudioMixerSourceManager.cpp
--- a/Engine/Source/Runtime/AudioMixer/Private/AudioMixerSourceManager.cpp
+++ b/Engine/Source/Runtime/AudioMixer/Private/AudioMixerSourceManager.cpp
@@ -271,6 +271,7 @@ namespace Audio
 			SourceInfo.bUseHRTFSpatializer = false;
 			SourceInfo.bUseOcclusionPlugin = false;
 			SourceInfo.bUseReverbPlugin = false;
+			SourceInfo.bUseSourceDataOverridePlugin = false;
 			SourceInfo.bHasStarted = false;
 			SourceInfo.bEnableBusSends = false;
 			SourceInfo.bEnableBaseSubmix = false;
@@ -533,6 +534,11 @@ namespace Audio
 			MixerDevice->ReverbPluginInterface->OnReleaseSource(SourceId);
 		}
 
+		if (SourceInfo.bUseSourceDataOverridePlugin)
+		{
+			MixerDevice->SourceDataOverridePluginInterface->OnReleaseSource(SourceId);
+		}
+
 		// Delete the source effects
 		SourceInfo.SourceEffectChainId = INDEX_NONE;
 		ResetSourceEffectChain(SourceId);
@@ -596,6 +602,7 @@ namespace Audio
 		SourceInfo.bIsExternalSend = false;
 		SourceInfo.bUseOcclusionPlugin = false;
 		SourceInfo.bUseReverbPlugin = false;
+		SourceInfo.bUseSourceDataOverridePlugin = false;
 		SourceInfo.bHasStarted = false;
 		SourceInfo.bEnableBusSends = false;
 		SourceInfo.bEnableBaseSubmix = false;
@@ -854,6 +861,12 @@ namespace Audio
 			SourceInfo.SourceBufferListener = InitParams.SourceBufferListener;
 			SourceInfo.bShouldSourceBufferListenerZeroBuffer = InitParams.bShouldSourceBufferListenerZeroBuffer;
 
+			if (InitParams.SourceDataOverridePluginSettings != nullptr)
+			{
+				MixerDevice->SourceDataOverridePluginInterface->OnInitSource(SourceId, InitParams.AudioComponentUserID, InitParams.SourceDataOverridePluginSettings);
+				SourceInfo.bUseSourceDataOverridePlugin = true;
+			}
+
 			// Default all sounds to not consider effect chain tails when playing
 			SourceInfo.bEffectTailsDone = true;
 
diff --git a/Engine/Source/Runtime/AudioMixer/Private/AudioMixerSourceManager.h b/Engine/Source/Runtime/AudioMixer/Private/AudioMixerSourceManager.h
--- a/Engine/Source/Runtime/AudioMixer/Private/AudioMixerSourceManager.h
+++ b/Engine/Source/Runtime/AudioMixer/Private/AudioMixerSourceManager.h
@@ -122,6 +122,7 @@ namespace Audio
 		USpatializationPluginSourceSettingsBase* SpatializationPluginSettings = nullptr;
 		UOcclusionPluginSourceSettingsBase* OcclusionPluginSettings = nullptr;
 		UReverbPluginSourceSettingsBase* ReverbPluginSettings = nullptr;
+		USourceDataOverridePluginSourceSettingsBase* SourceDataOverridePluginSettings = nullptr;
 
 		FSoundModulationDefaultSettings ModulationSettings;
 
@@ -452,6 +453,7 @@ namespace Audio
 			uint8 bIsExternalSend:1;
 			uint8 bUseOcclusionPlugin:1;
 			uint8 bUseReverbPlugin:1;
+			uint8 bUseSourceDataOverridePlugin:1;
 			uint8 bIsDone:1;
 			uint8 bIsLastBuffer:1;
 			uint8 bEnableBusSends : 1;
diff --git a/Engine/Source/Runtime/Engine/Classes/Sound/SoundAttenuation.h b/Engine/Source/Runtime/Engine/Classes/Sound/SoundAttenuation.h
--- a/Engine/Source/Runtime/Engine/Classes/Sound/SoundAttenuation.h
+++ b/Engine/Source/Runtime/Engine/Classes/Sound/SoundAttenuation.h
@@ -92,13 +92,17 @@ struct ENGINE_API FSoundAttenuationPluginSettings
 	UPROPERTY(EditAnywhere, BlueprintReadWrite, Category = AttenuationSpatialization, meta = (DisplayName = "Spatialization Plugin Settings"))
 	TArray<TObjectPtr<USpatializationPluginSourceSettingsBase>> SpatializationPluginSettingsArray;
 
-	/** Settings to use with occlusion audio plugin. These are defined by the plugin creator. Not all audio plugins utilize this feature. This  is an array so multiple plugins can have settings. */
+	/** Settings to use with occlusion audio plugin. These are defined by the plugin creator. Not all audio plugins utilize this feature. This is an array so multiple plugins can have settings. */
 	UPROPERTY(EditAnywhere, BlueprintReadWrite, Category = AttenuationOcclusion, meta = (DisplayName = "Occlusion Plugin Settings"))
 	TArray<TObjectPtr<UOcclusionPluginSourceSettingsBase>> OcclusionPluginSettingsArray;
 
-	/** Settings to use with reverb audio plugin. These are defined by the plugin creator. Not all audio plugins utilize this feature. This  is an array so multiple plugins can have settings. */
+	/** Settings to use with reverb audio plugin. These are defined by the plugin creator. Not all audio plugins utilize this feature. This is an array so multiple plugins can have settings. */
 	UPROPERTY(EditAnywhere, BlueprintReadWrite, Category = AttenuationReverbSend, meta = (DisplayName = "Reverb Plugin Settings"))
 	TArray<TObjectPtr<UReverbPluginSourceSettingsBase>> ReverbPluginSettingsArray;
+
+	/** Settings to use with source data override audio plugin. These are defined by the plugin creator. Not all audio plugins utilize this feature. This is an array so multiple plugins can have settings. */
+	UPROPERTY(EditAnywhere, BlueprintReadWrite, Category = AttenuationSourceDataOverride, meta = (DisplayName = "Source Data Override Plugin Settings"))
+	TArray<TObjectPtr<USourceDataOverridePluginSourceSettingsBase>> SourceDataOverridePluginSettingsArray;
 };
 
 USTRUCT(BlueprintType)
@@ -195,6 +199,10 @@ struct ENGINE_API FSoundAttenuationSettings : public FBaseAttenuationSettings
 	UPROPERTY(EditAnywhere, BlueprintReadWrite, Category = AttenuationSubmixSend, meta = (DisplayName = "Enable Submix Send"))
 	uint8 bEnableSubmixSends : 1;
 
+	/** Enables overriding WaveInstance data using source data override plugin */
+	UPROPERTY(EditAnywhere, BlueprintReadWrite, Category = AttenuationSourceDataOverride, meta = (DisplayName = "Enable Source Data Override"))
+	uint8 bEnableSourceDataOverride : 1;
+
 	/** What method we use to spatialize the sound. */
 	UPROPERTY(EditAnywhere, BlueprintReadWrite, Category = AttenuationSpatialization, meta = (ClampMin = "0", EditCondition = "bSpatialize", DisplayName = "Spatialization Method"))
 	TEnumAsByte<enum ESoundSpatializationAlgorithm> SpatializationAlgorithm;
@@ -404,6 +412,7 @@ struct ENGINE_API FSoundAttenuationSettings : public FBaseAttenuationSettings
 		, bApplyNormalizationToStereoSounds(false)
 		, bEnableLogFrequencyScaling(false)
 		, bEnableSubmixSends(false)
+		, bEnableSourceDataOverride(false)
 		, SpatializationAlgorithm(ESoundSpatializationAlgorithm::SPATIALIZATION_Default)
 		, BinauralRadius(0.0f)
 		, AbsorptionMethod(EAirAbsorptionMethod::Linear)
diff --git a/Engine/Source/Runtime/Engine/Private/ActiveSound.cpp b/Engine/Source/Runtime/Engine/Private/ActiveSound.cpp
--- a/Engine/Source/Runtime/Engine/Private/ActiveSound.cpp
+++ b/Engine/Source/Runtime/Engine/Private/ActiveSound.cpp
@@ -1742,6 +1742,22 @@ void FActiveSound::UpdateAttenuation(float DeltaTime, FSoundParseParameters& Par
 		}
 	}
 
+	if (Settings->PluginSettings.SourceDataOverridePluginSettingsArray.Num() > 0)
+	{
+		UClass* PluginClass = GetAudioPluginCustomSettingsClass(EAudioPlugin::SOURCEDATAOVERRIDE);
+		if (PluginClass)
+		{
+			for (USourceDataOverridePluginSourceSettingsBase* SettingsBase : Settings->PluginSettings.SourceDataOverridePluginSettingsArray)
+			{
+				if (SettingsBase != nullptr && SettingsBase->IsA(PluginClass))
+				{
+					ParseParams.SourceDataOverridePluginSettings = SettingsBase;
+					break;
+				}
+			}
+		}
+	}
+
 	// Attenuate with the absorption filter if necessary
 	if (Settings->bAttenuateWithLPF)
 	{
@@ -1768,6 +1784,7 @@ void FActiveSound::UpdateAttenuation(float DeltaTime, FSoundParseParameters& Par
 	ParseParams.StereoSpread = Settings->StereoSpread;
 	ParseParams.bApplyNormalizationToStereoSounds = Settings->bApplyNormalizationToStereoSounds;
 	ParseParams.bUseSpatialization |= Settings->bSpatialize;
+	ParseParams.bEnableSourceDataOverride |= Settings->bEnableSourceDataOverride;
 
 	// Check the binaural radius to determine if we're going to HRTF spatialize
 	if (ListenerData.ListenerToSoundDistance < Settings->BinauralRadius)
diff --git a/Engine/Source/Runtime/Engine/Private/Audio.cpp b/Engine/Source/Runtime/Engine/Private/Audio.cpp
--- a/Engine/Source/Runtime/Engine/Private/Audio.cpp
+++ b/Engine/Source/Runtime/Engine/Private/Audio.cpp
@@ -154,7 +154,10 @@ UClass* GetAudioPluginCustomSettingsClass(EAudioPlugin PluginType)
 
 		case EAudioPlugin::SOURCEDATAOVERRIDE:
 		{
-			return nullptr;
+			if (IAudioSourceDataOverrideFactory* Factory = AudioPluginUtilities::GetDesiredSourceDataOverridePlugin())
+			{
+				return Factory->GetCustomSourceDataOverrideSettingsClass();
+			}		
 		}
 		break;
 
@@ -810,6 +813,7 @@ FWaveInstance::FWaveInstance(const UPTRINT InWaveInstanceHash, FActiveSound& InA
 	, SpatializationPluginSettings(nullptr)
 	, OcclusionPluginSettings(nullptr)
 	, ReverbPluginSettings(nullptr)
+	, SourceDataOverridePluginSettings(nullptr)
 	, OutputTarget(EAudioOutputTarget::Speaker)
 	, LowPassFilterFrequency(MAX_FILTER_FREQUENCY)
 	, SoundClassFilterFrequency(MAX_FILTER_FREQUENCY)
@@ -971,6 +975,16 @@ float FWaveInstance::GetDistanceAndOcclusionAttenuation() const
 	return DistanceAttenuation * OcclusionAttenuation;
 }
 
+float FWaveInstance::GetDistanceAttenuation() const
+{
+	return DistanceAttenuation;
+}
+
+float FWaveInstance::GetOcclusionAttenuation() const
+{
+	return OcclusionAttenuation;
+}
+
 float FWaveInstance::GetDynamicVolume() const
 {
 	float OutVolume = 1.0f;
diff --git a/Engine/Source/Runtime/Engine/Private/AudioDevice.cpp b/Engine/Source/Runtime/Engine/Private/AudioDevice.cpp
--- a/Engine/Source/Runtime/Engine/Private/AudioDevice.cpp
+++ b/Engine/Source/Runtime/Engine/Private/AudioDevice.cpp
@@ -266,6 +266,7 @@ FAudioDevice::FAudioDevice()
 	, NumPrecacheFrames(MONO_PCM_BUFFER_SAMPLES)
 	, DeviceID(static_cast<Audio::FDeviceId>(INDEX_NONE))
 	, SpatializationPluginInterface(nullptr)
+	, SourceDataOverridePluginInterface(nullptr)
 	, ReverbPluginInterface(nullptr)
 	, OcclusionInterface(nullptr)
 	, MaxSources(0)
@@ -303,6 +304,7 @@ FAudioDevice::FAudioDevice()
 	, bSpatializationInterfaceEnabled(false)
 	, bOcclusionInterfaceEnabled(false)
 	, bReverbInterfaceEnabled(false)
+	, bSourceDataOverrideInterfaceEnabled(false)
 	, bModulationInterfaceEnabled(false)
 	, bPluginListenersInitialized(false)
 	, bHRTFEnabledForAll(false)
@@ -495,6 +497,7 @@ bool FAudioDevice::Init(Audio::FDeviceId InDeviceID, int32 InMaxSources, int32 I
 	if (SourceDataOverridePluginFactory)
 	{
 		SourceDataOverridePluginInterface = SourceDataOverridePluginFactory->CreateNewSourceDataOverridePlugin(this);
+		bSourceDataOverrideInterfaceEnabled = true;
 		check(IsAudioMixerEnabled());
 		UE_LOG(LogAudio, Display, TEXT("Audio Source Data Override Plugin: %s"), *(SourceDataOverridePluginFactory->GetDisplayName()));
 	}
@@ -797,6 +800,12 @@ void FAudioDevice::Teardown()
 		bReverbInterfaceEnabled = false;
 	}
 
+	if (SourceDataOverridePluginInterface.IsValid())
+	{
+		SourceDataOverridePluginInterface.Reset();
+		bSourceDataOverrideInterfaceEnabled = false;
+	}
+
 	if (OcclusionInterface.IsValid())
 	{
 		OcclusionInterface->Shutdown();
@@ -863,6 +872,11 @@ void FAudioDevice::UpdateAudioPluginSettingsObjectCache()
 	{
 		PluginSettingsObjects.Add(*It);
 	}
+
+	for (TObjectIterator<USourceDataOverridePluginSourceSettingsBase> It; It; ++It)
+	{
+		PluginSettingsObjects.Add(*It);
+	}
 }
 
 void FAudioDevice::AddReferencedObjects(FReferenceCollector& Collector)
diff --git a/Engine/Source/Runtime/Engine/Private/AudioPluginUtilities.cpp b/Engine/Source/Runtime/Engine/Private/AudioPluginUtilities.cpp
--- a/Engine/Source/Runtime/Engine/Private/AudioPluginUtilities.cpp
+++ b/Engine/Source/Runtime/Engine/Private/AudioPluginUtilities.cpp
@@ -35,7 +35,7 @@ FORCEINLINE const TCHAR* GetPluginConfigName(EAudioPlugin PluginType)
 			return TEXT("ModulationPlugin");
 
 		case EAudioPlugin::SOURCEDATAOVERRIDE:
-			return TEXT("PropagationPlugin");
+			return TEXT("SourceDataOverridePlugin");
 			
 		default:
 			checkf(false, TEXT("Undefined audio plugin type."));
diff --git a/Engine/Source/Runtime/Engine/Private/SoundAttenuation.cpp b/Engine/Source/Runtime/Engine/Private/SoundAttenuation.cpp
--- a/Engine/Source/Runtime/Engine/Private/SoundAttenuation.cpp
+++ b/Engine/Source/Runtime/Engine/Private/SoundAttenuation.cpp
@@ -126,6 +126,7 @@ bool FSoundAttenuationSettings::operator==(const FSoundAttenuationSettings& Othe
 			&& PluginSettings.OcclusionPluginSettingsArray	== Other.PluginSettings.OcclusionPluginSettingsArray
 			&& bEnableReverbSend		== Other.bEnableReverbSend
 			&& PluginSettings.ReverbPluginSettingsArray		== Other.PluginSettings.ReverbPluginSettingsArray
+			&& PluginSettings.SourceDataOverridePluginSettingsArray == Other.PluginSettings.SourceDataOverridePluginSettingsArray
 			&& ReverbWetLevelMin		== Other.ReverbWetLevelMin
 			&& ReverbWetLevelMax		== Other.ReverbWetLevelMax
 			&& ReverbDistanceMin		== Other.ReverbDistanceMin
diff --git a/Engine/Source/Runtime/Engine/Private/SoundWave.cpp b/Engine/Source/Runtime/Engine/Private/SoundWave.cpp
--- a/Engine/Source/Runtime/Engine/Private/SoundWave.cpp
+++ b/Engine/Source/Runtime/Engine/Private/SoundWave.cpp
@@ -2462,6 +2462,9 @@ void USoundWave::Parse(FAudioDevice* AudioDevice, const UPTRINT NodeWaveInstance
 	WaveInstance->SpatializationPluginSettings = ParseParams.SpatializationPluginSettings;
 	WaveInstance->OcclusionPluginSettings = ParseParams.OcclusionPluginSettings;
 	WaveInstance->ReverbPluginSettings = ParseParams.ReverbPluginSettings;
+	WaveInstance->SourceDataOverridePluginSettings = ParseParams.SourceDataOverridePluginSettings;
+
+	WaveInstance->bEnableSourceDataOverride = ParseParams.bEnableSourceDataOverride;
 
 	if (WaveInstance->IsPlaying())
 	{
diff --git a/Engine/Source/Runtime/Engine/Public/ActiveSound.h b/Engine/Source/Runtime/Engine/Public/ActiveSound.h
--- a/Engine/Source/Runtime/Engine/Public/ActiveSound.h
+++ b/Engine/Source/Runtime/Engine/Public/ActiveSound.h
@@ -165,6 +165,9 @@ struct FSoundParseParameters
 	// What reverb plugin source settings to use
 	UReverbPluginSourceSettingsBase* ReverbPluginSettings;
 
+	// What source data override plugin source settings to use
+	USourceDataOverridePluginSourceSettingsBase* SourceDataOverridePluginSettings;
+
 	// What source effect chain to use
 	USoundEffectSourcePresetChain* SourceEffectChain;
 
@@ -192,6 +195,8 @@ struct FSoundParseParameters
 	/** Whether or not to enable Submix Sends in addition to the Main Submix*/
 	uint32 bEnableSubmixSends : 1;
 
+	uint32 bEnableSourceDataOverride : 1;
+
 	// Whether the sound should be spatialized
 	uint8 bUseSpatialization:1;
 
@@ -241,6 +246,7 @@ struct FSoundParseParameters
 		, SpatializationPluginSettings(nullptr)
 		, OcclusionPluginSettings(nullptr)
 		, ReverbPluginSettings(nullptr)
+		, SourceDataOverridePluginSettings(nullptr)
 		, SourceEffectChain(nullptr)
 		, LowPassFilterFrequency(MAX_FILTER_FREQUENCY)
 		, AttenuationLowpassFilterFrequency(MAX_FILTER_FREQUENCY)
@@ -250,6 +256,7 @@ struct FSoundParseParameters
 		, bEnableBusSends(false)
 		, bEnableBaseSubmix(false)
 		, bEnableSubmixSends(false)
+		, bEnableSourceDataOverride(false)
 		, bUseSpatialization(false)
 		, bLooping(false)
 		, bEnableLowPassFilter(false)
diff --git a/Engine/Source/Runtime/Engine/Public/Audio.h b/Engine/Source/Runtime/Engine/Public/Audio.h
--- a/Engine/Source/Runtime/Engine/Public/Audio.h
+++ b/Engine/Source/Runtime/Engine/Public/Audio.h
@@ -274,6 +274,9 @@ public:
 	/** Whether or not to enable Submix Sends in addition to the Main Submix*/
 	uint32 bEnableSubmixSends : 1;
 
+	/** Whether or not to use source data overrides */
+	uint32 bEnableSourceDataOverride : 1;
+
 	/** Set to true if the sound nodes state that the radio filter should be applied */
 	uint32 bApplyRadioFilter:1;
 
@@ -336,6 +339,9 @@ public:
 	/** The occlusion plugin settings to use for the wave instance. */
 	UReverbPluginSourceSettingsBase* ReverbPluginSettings;
 
+	/** The source data override plugin settings to use for the wave instance. */
+	USourceDataOverridePluginSourceSettingsBase* SourceDataOverridePluginSettings;
+
 	/** Which output target the sound should play on. */
 	EAudioOutputTarget::Type OutputTarget;
 
@@ -447,9 +453,15 @@ public:
 	/** Returns the volume of the sound including distance attenuation. */
 	float GetVolumeWithDistanceAndOcclusionAttenuation() const;
 
-	/** Returns the distance attenuation of the source voice. */
+	/** Returns the combined distance and occlusion attenuation of the source voice. */
 	float GetDistanceAndOcclusionAttenuation() const;
 
+	/** Returns the distance attenuation of the source voice */
+	float GetDistanceAttenuation() const;
+
+	/** Returns the occlusion attenuation of the source voice */
+	float GetOcclusionAttenuation() const;
+
 	/** Returns the dynamic volume of the sound */
 	float GetDynamicVolume() const;
 
diff --git a/Engine/Source/Runtime/Engine/Public/AudioDevice.h b/Engine/Source/Runtime/Engine/Public/AudioDevice.h
--- a/Engine/Source/Runtime/Engine/Public/AudioDevice.h
+++ b/Engine/Source/Runtime/Engine/Public/AudioDevice.h
@@ -1331,6 +1331,20 @@ public:
 		return false;
 	}
 
+	bool IsSourceDataOverridePluginEnabled() const
+	{
+		return bSourceDataOverrideInterfaceEnabled;
+	}
+
+	static bool IsSourceDataOverridePluginLoaded()
+	{
+		if (FAudioDeviceHandle MainAudioDevice = GEngine->GetMainAudioDevice())
+		{
+			return MainAudioDevice->bSourceDataOverrideInterfaceEnabled;
+		}
+		return false;
+	}
+
 	/** Returns if this is the multi-platform audio mixer. */
 	bool IsAudioMixerEnabled() const
 	{
@@ -2086,6 +2100,7 @@ private:
 	uint8 bSpatializationInterfaceEnabled:1;
 	uint8 bOcclusionInterfaceEnabled:1;
 	uint8 bReverbInterfaceEnabled:1;
+	uint8 bSourceDataOverrideInterfaceEnabled:1;
 	uint8 bModulationInterfaceEnabled:1;
 
 	/** Whether or not we've initialized plugin listeners array. */
