Fix crash compiling an empty state machine

diff --git a/Engine/Source/Runtime/Engine/Private/Animation/AnimNode_StateMachine.cpp b/Engine/Source/Runtime/Engine/Private/Animation/AnimNode_StateMachine.cpp
--- a/Engine/Source/Runtime/Engine/Private/Animation/AnimNode_StateMachine.cpp
+++ b/Engine/Source/Runtime/Engine/Private/Animation/AnimNode_StateMachine.cpp
@@ -1172,7 +1172,7 @@ const FPoseContext& FAnimNode_StateMachine::EvaluateState(int32 StateIndex, cons
 
 bool FAnimNode_StateMachine::IsAConduitState(int32 StateIndex) const
 {
-	return ((PRIVATE_MachineDescription != NULL) && (StateIndex < PRIVATE_MachineDescription->States.Num())) ? GetStateInfo(StateIndex).bIsAConduit : false;
+	return ((PRIVATE_MachineDescription != nullptr) && (PRIVATE_MachineDescription->States.IsValidIndex(StateIndex))) ? GetStateInfo(StateIndex).bIsAConduit : false;
 }
 
 bool FAnimNode_StateMachine::IsValidTransitionIndex(int32 TransitionIndex) const
