PackedLevelActor: Choose ISM Component type based on if the mesh is nanite or not. (ISM vs HISM)

diff --git a/Engine/Source/Runtime/Engine/Private/PackedLevelActor/PackedLevelActor.cpp b/Engine/Source/Runtime/Engine/Private/PackedLevelActor/PackedLevelActor.cpp
--- a/Engine/Source/Runtime/Engine/Private/PackedLevelActor/PackedLevelActor.cpp
+++ b/Engine/Source/Runtime/Engine/Private/PackedLevelActor/PackedLevelActor.cpp
@@ -10,7 +10,6 @@
 APackedLevelActor::APackedLevelActor(const FObjectInitializer& ObjectInitializer)
 	: Super(ObjectInitializer)
 #if WITH_EDITORONLY_DATA
-	, ISMComponentClass(UInstancedStaticMeshComponent::StaticClass())
 	, bChildChanged(false)
 #endif
 {
diff --git a/Engine/Source/Runtime/Engine/Private/PackedLevelActor/PackedLevelActorISMBuilder.cpp b/Engine/Source/Runtime/Engine/Private/PackedLevelActor/PackedLevelActorISMBuilder.cpp
--- a/Engine/Source/Runtime/Engine/Private/PackedLevelActor/PackedLevelActorISMBuilder.cpp
+++ b/Engine/Source/Runtime/Engine/Private/PackedLevelActor/PackedLevelActorISMBuilder.cpp
@@ -11,6 +11,7 @@
 
 #include "Components/StaticMeshComponent.h"
 #include "Components/InstancedStaticMeshComponent.h"
+#include "Components/HierarchicalInstancedStaticMeshComponent.h"
 
 FPackedLevelActorBuilderID FPackedLevelActorISMBuilder::BuilderID = 'ISMP';
 
@@ -38,15 +39,7 @@ void FPackedLevelActorISMBuilder::GetPackClusters(FPackedLevelActorBuilderContex
 void FPackedLevelActorISMBuilder::PackActors(FPackedLevelActorBuilderContext& InContext, APackedLevelActor* InPackingActor, const FPackedLevelActorBuilderClusterID& InClusterID, const TArray<UActorComponent*>& InComponents) const
 {
 	check(InClusterID.GetBuilderID() == GetID());
-	UInstancedStaticMeshComponent* PackComponent = InPackingActor->AddPackedComponent<UInstancedStaticMeshComponent>();
-	
 	FTransform ActorTransform = InPackingActor->GetActorTransform();
-	PackComponent->AttachToComponent(InPackingActor->GetRootComponent(), FAttachmentTransformRules::SnapToTargetIncludingScale);
-	
-	FPackedLevelActorISMBuilderCluster* ISMCluster = (FPackedLevelActorISMBuilderCluster*)InClusterID.GetData();
-	check(ISMCluster);
-
-	ISMCluster->ISMDescriptor.InitComponent(PackComponent);
 
 	TArray<FTransform> InstanceTransforms;
 	for (UActorComponent* Component : InComponents)
@@ -73,6 +66,23 @@ void FPackedLevelActorISMBuilder::PackActors(FPackedLevelActorBuilderContext& In
 			InstanceTransforms.Add(StaticMeshComponent->GetComponentTransform().GetRelativeTransform(ActorTransform) * InContext.GetRelativePivotTransform());
 		}
 	}
+
+	FPackedLevelActorISMBuilderCluster* ISMCluster = (FPackedLevelActorISMBuilderCluster*)InClusterID.GetData();
+	check(ISMCluster);
+
+	TSubclassOf<UInstancedStaticMeshComponent> ComponentClass = UInstancedStaticMeshComponent::StaticClass();
+	if (InstanceTransforms.Num() > 1 && ISMCluster->ISMDescriptor.StaticMesh && !ISMCluster->ISMDescriptor.StaticMesh->NaniteSettings.bEnabled)
+	{
+		// Use HISM for non-nanite when there is more than one transform (no use in using HISM cpu occlusion for a single instance)
+		ComponentClass = UHierarchicalInstancedStaticMeshComponent::StaticClass();
+	}
+		
+	UInstancedStaticMeshComponent* PackComponent = InPackingActor->AddPackedComponent<UInstancedStaticMeshComponent>(ComponentClass);
+		
+	PackComponent->AttachToComponent(InPackingActor->GetRootComponent(), FAttachmentTransformRules::SnapToTargetIncludingScale);
+
+	ISMCluster->ISMDescriptor.InitComponent(PackComponent);
+
 	PackComponent->AddInstances(InstanceTransforms, /*bShouldReturnIndices*/false, /*bWorldSpace*/false);
 	PackComponent->RegisterComponent();
 }
@@ -81,6 +91,8 @@ FPackedLevelActorISMBuilderCluster::FPackedLevelActorISMBuilderCluster(FPackedLe
 	: FPackedLevelActorBuilderCluster(InBuilderID)
 {
 	ISMDescriptor.InitFrom(InComponent, /** bInitBodyInstance= */ false);
+	// ComponentClass will be determined based on StaticMesh nanite settings and number of instances and we don't want it to be part of the cluster hash
+	ISMDescriptor.ComponentClass = nullptr;
 	// Component descriptor should be considered hidden if original actor owner was.
 	ISMDescriptor.bHiddenInGame |= InComponent->GetOwner()->IsHidden();
 	ISMDescriptor.BodyInstance.CopyRuntimeBodyInstancePropertiesFrom(&InComponent->BodyInstance);
diff --git a/Engine/Source/Runtime/Engine/Public/PackedLevelActor/PackedLevelActor.h b/Engine/Source/Runtime/Engine/Public/PackedLevelActor/PackedLevelActor.h
--- a/Engine/Source/Runtime/Engine/Public/PackedLevelActor/PackedLevelActor.h
+++ b/Engine/Source/Runtime/Engine/Public/PackedLevelActor/PackedLevelActor.h
@@ -60,11 +60,11 @@ public:
 	virtual void RerunConstructionScripts() override;
 
 	template<class T>
-	T* AddPackedComponent()
+	T* AddPackedComponent(TSubclassOf<T> ComponentClass)
 	{
 		Modify();
-		FName NewComponentName = *FComponentEditorUtils::GenerateValidVariableName(T::StaticClass(), this);
-		T* NewComponent = NewObject<T>(this, NewComponentName, RF_Transactional);
+		FName NewComponentName = *FComponentEditorUtils::GenerateValidVariableName(ComponentClass, this);
+		T* NewComponent = NewObject<T>(this, ComponentClass, NewComponentName, RF_Transactional);
 		AddInstanceComponent(NewComponent);
 		NewComponent->ComponentTags.Add(GetPackedComponentTag());
 		return NewComponent;
@@ -72,9 +72,6 @@ public:
 #endif
 
 #if WITH_EDITORONLY_DATA
-	UPROPERTY(EditAnywhere, Category = Packed, meta = (DisplayName = "ISM Component Class"))
-	TSubclassOf<UInstancedStaticMeshComponent> ISMComponentClass;
-
 	UPROPERTY(VisibleAnywhere, Category = Packed)
 	TSoftObjectPtr<UBlueprint> BlueprintAsset;
 
