Additional debug info to track down crash on mac editor.

diff --git a/Engine/Source/Runtime/Engine/Private/Rendering/MorphTargetVertexInfoBuffers.cpp b/Engine/Source/Runtime/Engine/Private/Rendering/MorphTargetVertexInfoBuffers.cpp
--- a/Engine/Source/Runtime/Engine/Private/Rendering/MorphTargetVertexInfoBuffers.cpp
+++ b/Engine/Source/Runtime/Engine/Private/Rendering/MorphTargetVertexInfoBuffers.cpp
@@ -146,7 +146,7 @@ void FMorphTargetVertexInfoBuffers::InitMorphResources(EShaderPlatform ShaderPla
 	bResourcesInitialized = true;
 
 	// UseGPUMorphTargets() can be toggled only on SM5 atm
-	if (!IsFeatureLevelSupported(ShaderPlatform, ERHIFeatureLevel::SM5) || MorphTargets.Num() == 0)
+	if (!IsPlatformShaderSupported(ShaderPlatform) || MorphTargets.Num() == 0)
 	{
 		return;
 	}
@@ -447,3 +447,8 @@ void FMorphTargetVertexInfoBuffers::InitMorphResources(EShaderPlatform ShaderPla
 
 	ValidateVertexBuffers(true);
 }
+
+bool FMorphTargetVertexInfoBuffers::IsPlatformShaderSupported(EShaderPlatform ShaderPlatform)
+{
+	return IsFeatureLevelSupported(ShaderPlatform, ERHIFeatureLevel::SM5);
+}
diff --git a/Engine/Source/Runtime/Engine/Private/SkeletalMeshLODRenderData.cpp b/Engine/Source/Runtime/Engine/Private/SkeletalMeshLODRenderData.cpp
--- a/Engine/Source/Runtime/Engine/Private/SkeletalMeshLODRenderData.cpp
+++ b/Engine/Source/Runtime/Engine/Private/SkeletalMeshLODRenderData.cpp
@@ -669,7 +669,7 @@ void FSkeletalMeshLODRenderData::SerializeStreamedData(FArchive& Ar, USkeletalMe
 			for (FName ShaderFormat : ShaderFormats)
 			{
 				const EShaderPlatform ShaderPlatform = ShaderFormatToLegacyShaderPlatform(ShaderFormat);
-				if (IsFeatureLevelSupported(ShaderPlatform, ERHIFeatureLevel::SM5))
+				if (FMorphTargetVertexInfoBuffers::IsPlatformShaderSupported(ShaderPlatform))
 				{
 					MorphTargetShaderPlatform = ShaderPlatform;
 					bSerializeCompressedMorphTargets = true;
@@ -698,16 +698,19 @@ void FSkeletalMeshLODRenderData::SerializeStreamedData(FArchive& Ar, USkeletalMe
 				{
 					TargetMorphTargetVertexInfoBuffers = &MorphTargetVertexInfoBuffers;
 					TargetMorphTargetVertexInfoBuffers->InitMorphResources(MorphTargetShaderPlatform, RenderSections, MorphTargets, StaticVertexBuffers.StaticMeshVertexBuffer.GetNumVertices(), LODIdx, SkeletalMeshLODInfo->MorphTargetPositionErrorTolerance);
+					check(TargetMorphTargetVertexInfoBuffers->IsMorphCPUDataValid());
 				}
 				else if (MorphTargetVertexInfoBuffers.IsMorphCPUDataValid() && !MorphTargetVertexInfoBuffers.IsRHIIntialized())
 				{
 					TargetMorphTargetVertexInfoBuffers = &MorphTargetVertexInfoBuffers;
+					check(TargetMorphTargetVertexInfoBuffers->IsMorphCPUDataValid());
 				}
 				else
 				{
 					// Fallback for when the RHI data has already been created.  Create a local version of the morph data, and serialize that
 					TargetMorphTargetVertexInfoBuffers = &LocalMorphTargetVertexInfoBuffers;
 					TargetMorphTargetVertexInfoBuffers->InitMorphResources(MorphTargetShaderPlatform, RenderSections, MorphTargets, StaticVertexBuffers.StaticMeshVertexBuffer.GetNumVertices(), LODIdx, SkeletalMeshLODInfo->MorphTargetPositionErrorTolerance);
+					check(TargetMorphTargetVertexInfoBuffers->IsMorphCPUDataValid());
 				}
 
 				Ar << *TargetMorphTargetVertexInfoBuffers;
diff --git a/Engine/Source/Runtime/Engine/Public/Rendering/MorphTargetVertexInfoBuffers.h b/Engine/Source/Runtime/Engine/Public/Rendering/MorphTargetVertexInfoBuffers.h
--- a/Engine/Source/Runtime/Engine/Public/Rendering/MorphTargetVertexInfoBuffers.h
+++ b/Engine/Source/Runtime/Engine/Public/Rendering/MorphTargetVertexInfoBuffers.h
@@ -71,6 +71,8 @@ public:
 		return TangentZPrecision;
 	}
 
+	static bool IsPlatformShaderSupported(EShaderPlatform ShaderPlatform);
+
 	FBufferRHIRef MorphDataBuffer;
 	FShaderResourceViewRHIRef MorphDataSRV;
 
