PackedLevelActor: Prevent changing world asset (set once when creating BP)

diff --git a/Engine/Source/Runtime/Engine/Private/PackedLevelActor/PackedLevelActor.cpp b/Engine/Source/Runtime/Engine/Private/PackedLevelActor/PackedLevelActor.cpp
--- a/Engine/Source/Runtime/Engine/Private/PackedLevelActor/PackedLevelActor.cpp
+++ b/Engine/Source/Runtime/Engine/Private/PackedLevelActor/PackedLevelActor.cpp
@@ -208,10 +208,10 @@ bool APackedLevelActor::CanEditChange(const FProperty* InProperty) const
 		return false;
 	}
 
-	// Disallow editing of the World if we are a BP instance
+	// Disallow change of the World on Packed Level Actors as it is set once when creating the blueprint.
 	if (InProperty->GetFName() == GET_MEMBER_NAME_CHECKED(APackedLevelActor, WorldAsset))
 	{
-		return GetClass()->IsNative();
+		return false;
 	}
 
 	return true;
