More tests for ios cool failure.

diff --git a/Engine/Source/Runtime/Engine/Private/SkeletalMeshLODRenderData.cpp b/Engine/Source/Runtime/Engine/Private/SkeletalMeshLODRenderData.cpp
--- a/Engine/Source/Runtime/Engine/Private/SkeletalMeshLODRenderData.cpp
+++ b/Engine/Source/Runtime/Engine/Private/SkeletalMeshLODRenderData.cpp
@@ -696,6 +696,9 @@ void FSkeletalMeshLODRenderData::SerializeStreamedData(FArchive& Ar, USkeletalMe
 				// The CPU data could have already been destroyed by this point, which happens when the RHI is initialized.  If possible, use the MorphTargetVertexInfoBuffers.
 				if (!MorphTargetVertexInfoBuffers.IsMorphResourcesInitialized())
 				{
+					check(FMorphTargetVertexInfoBuffers::IsPlatformShaderSupported(MorphTargetShaderPlatform));
+					check(MorphTargets.Num() > 0);
+
 					TargetMorphTargetVertexInfoBuffers = &MorphTargetVertexInfoBuffers;
 					TargetMorphTargetVertexInfoBuffers->InitMorphResources(MorphTargetShaderPlatform, RenderSections, MorphTargets, StaticVertexBuffers.StaticMeshVertexBuffer.GetNumVertices(), LODIdx, SkeletalMeshLODInfo->MorphTargetPositionErrorTolerance);
 					check(TargetMorphTargetVertexInfoBuffers->IsMorphCPUDataValid());
