Level Instance:
- Move conversion to OFPA + MoveActorsToLevel to UEditorLevelUtils::CreateNewStreamingLevelForWorld
- Add Scan when creating new level so that FLevelInstanceActorDesc can register its container properly
- Fix existing Scan on commit (was calling ScanPath instead of ScanFiles)

diff --git a/Engine/Source/Editor/UnrealEd/Private/EditorLevelUtils.cpp b/Engine/Source/Editor/UnrealEd/Private/EditorLevelUtils.cpp
--- a/Engine/Source/Editor/UnrealEd/Private/EditorLevelUtils.cpp
+++ b/Engine/Source/Editor/UnrealEd/Private/EditorLevelUtils.cpp
@@ -681,6 +681,24 @@ UWorld* GetWorldForNewStreamingLevel(UWorld* InTemplateWorld)
 };
 
 ULevelStreaming* UEditorLevelUtils::CreateNewStreamingLevelForWorld(UWorld& InWorld, TSubclassOf<ULevelStreaming> LevelStreamingClass, const FString& DefaultFilename /* = TEXT( "" ) */, bool bMoveSelectedActorsIntoNewLevel /* = false */, UWorld* InTemplateWorld /* = nullptr */)
+{
+	TArray<AActor*> ActorsToMove;
+	if (bMoveSelectedActorsIntoNewLevel)
+	{
+		ActorsToMove.Reserve(GEditor->GetSelectedActorCount());
+		for (FSelectionIterator It(GEditor->GetSelectedActorIterator()); It; ++It)
+		{
+			if (AActor* Actor = Cast<AActor>(*It))
+			{
+				ActorsToMove.Add(Actor);
+			}
+		}
+	}
+
+	return CreateNewStreamingLevelForWorld(InWorld, LevelStreamingClass, /*bUseExternalActors=*/false, DefaultFilename, &ActorsToMove, InTemplateWorld);
+}
+
+ULevelStreaming* UEditorLevelUtils::CreateNewStreamingLevelForWorld(UWorld& InWorld, TSubclassOf<ULevelStreaming> LevelStreamingClass, bool bUseExternalActors, const FString& DefaultFilename /* = TEXT("") */, const TArray<AActor*>* ActorsToMove /* = nullptr */, UWorld* InTemplateWorld /* = nullptr */)
 {
 	// Editor modes cannot be active when any level saving occurs.
 	if (!IsRunningCommandlet())
@@ -695,6 +713,11 @@ ULevelStreaming* UEditorLevelUtils::CreateNewStreamingLevelForWorld(UWorld& InWo
 	// This is the new streaming level's world not the persistent level world
 	UWorld* NewLevelWorld = GetWorldForNewStreamingLevel(InTemplateWorld);
 	ULevel* PersistentLevel = GetPersistentLevelForNewStreamingLevel(NewLevelWorld, InTemplateWorld);
+	if (bUseExternalActors)
+	{
+		PersistentLevel->ConvertAllActorsToPackaging(true);
+		PersistentLevel->bUseExternalActors = true;
+	}
 
 	bool bNewWorldSaved = false;
 	FString NewPackageName = DefaultFilename;
@@ -723,9 +746,9 @@ ULevelStreaming* UEditorLevelUtils::CreateNewStreamingLevelForWorld(UWorld& InWo
 		{
 			NewLevel = NewStreamingLevel->GetLoadedLevel();
 			// If we are moving the selected actors to the new level move them now
-			if (bMoveSelectedActorsIntoNewLevel)
+			if (ActorsToMove)
 			{
-				MoveSelectedActorsToLevel(NewStreamingLevel);
+				MoveActorsToLevel(*ActorsToMove, NewLevel);
 			}
 
 			// Finally make the new level the current one
diff --git a/Engine/Source/Editor/UnrealEd/Public/EditorLevelUtils.h b/Engine/Source/Editor/UnrealEd/Public/EditorLevelUtils.h
--- a/Engine/Source/Editor/UnrealEd/Public/EditorLevelUtils.h
+++ b/Engine/Source/Editor/UnrealEd/Public/EditorLevelUtils.h
@@ -85,11 +85,42 @@ public:
 	 */
 	static UNREALED_API void MakeLevelCurrent(ULevel* InLevel, bool bEvenIfLocked = false);
 
-	
+	/**
+	 * Move provided actors to destination level.
+	 *
+	 * @param	ActorsToMove						Actors to move.
+	 * @param	DestLevel							Destination level to move the actors to.
+	 * @param	bWarnAboutReferences				Whether or not to show a modal warning about referenced actors that may no longer function after being moved.
+	 * @param	bWarnAboutRenaming					Whether or not to show a model warning for asset rename.
+	 * @param	bMoveAllOrFail						Whether operation should fail if any of the actors fails to be moved.
+	 * @param	OutActors							Optional, if not null the array will be filled with the newly moved actors.
+	 *
+	 * @return	Returns the number of actors moved.
+	 */
 	static UNREALED_API int32 MoveActorsToLevel(const TArray<AActor*>& ActorsToMove, ULevel* DestLevel, bool bWarnAboutReferences = true, bool bWarnAboutRenaming = true, bool bMoveAllOrFail = false, TArray<AActor*>* OutActors = nullptr);
 
+	/**
+	 * Copy provided actors to destination level.
+	 *
+	 * @param	ActorsToMove						Actors to copy.
+	 * @param	DestLevel							Destination level to copy the actors to.
+	 * @param	bWarnAboutReferences				Whether or not to show a modal warning about referenced actors that may no longer function after being copied.
+	 * @param	bWarnAboutRenaming					Whether or not to show a model warning for asset rename.
+	 * @param	bMoveAllOrFail						Whether operation should fail if any of the actors fails to be copied.
+	 * @param	OutActors							Optional, if not null the array will be filled with the newly copied actors.
+	 *
+	 * @return	Returns the number of actors copied.
+	 */
 	static UNREALED_API int32 CopyActorsToLevel(const TArray<AActor*>& ActorsToMove, ULevel* DestLevel, bool bWarnAboutReferences = true, bool bWarnAboutRenaming = true, bool bMoveAllOrFail = false, TArray<AActor*>* OutActors = nullptr);
 
+	/**
+	 * Move selected actors to destination level.
+	 *
+	 * @param	DestLevel							Destination level to move the actors to.
+	 * @param	bWarnAboutReferences				Whether or not to show a modal warning about referenced actors that may no longer function after being copied.
+	 *
+	 * @return	Returns the number of actors copied.
+	 */
 	static UNREALED_API int32 MoveSelectedActorsToLevel(ULevel* DestLevel, bool bWarnAboutReferences = true);
 
 	/** 
@@ -97,7 +128,7 @@ public:
 	*/
 	DECLARE_MULTICAST_DELEGATE_ThreeParams(FCanMoveActorToLevelDelegate, const AActor* /* ActorToMove */, const ULevel* /* DestLevel */, bool& /* bOutCanMove */);
 	static UNREALED_API FCanMoveActorToLevelDelegate CanMoveActorToLevelDelegate;
-
+		
 	/**
 	 * Creates a new streaming level and adds it to a world
 	 *
@@ -111,6 +142,19 @@ public:
 	 */
 	static UNREALED_API ULevelStreaming* CreateNewStreamingLevelForWorld(UWorld& World, TSubclassOf<ULevelStreaming> LevelStreamingClass, const FString& DefaultFilename = TEXT(""), bool bMoveSelectedActorsIntoNewLevel = false, UWorld* InTemplateWorld = nullptr);
 
+	/**
+	 * Creates a new streaming level and adds it to a world
+	 *
+	 * @param	InWorld								The world to add the streaming level to
+	 * @param	LevelStreamingClass					The streaming class type instead to use for the level.
+	 * @param	bUseExternalActors					If level should use external actors.
+	 * @param	DefaultFilename						file name for level.  If empty, the user will be prompted during the save process.
+	 * @param	ActorsToMove						Optional, move provided actors into the new level.
+	 * @param	InTemplateWorld						If valid, the new level will be a copy of the template world.
+	 *
+	 * @return	Returns the newly created level, or NULL on failure
+	 */
+	static UNREALED_API ULevelStreaming* CreateNewStreamingLevelForWorld(UWorld& World, TSubclassOf<ULevelStreaming> LevelStreamingClass, bool bUseExternalActors, const FString& DefaultFilename, const TArray<AActor*>* ActorsToMove = nullptr, UWorld* InTemplateWorld = nullptr);
 
 	/**
 	 * Adds the named level packages to the world.  Does nothing if all the levels already exist in the world.
diff --git a/Engine/Source/Runtime/Engine/Classes/Engine/Level.h b/Engine/Source/Runtime/Engine/Classes/Engine/Level.h
--- a/Engine/Source/Runtime/Engine/Classes/Engine/Level.h
+++ b/Engine/Source/Runtime/Engine/Classes/Engine/Level.h
@@ -1087,6 +1087,11 @@ public:
 	 */
 	static ENGINE_API FString GetExternalActorsPath(UPackage* InLevelPackage, const FString& InPackageShortName = FString());
 
+	/**
+	 * Scans/Updates all Level Assets (level package and external packages)
+	 */
+	static ENGINE_API void ScanLevelAssets(const FString& InLevelPackageName);
+
 	/**
 	 * Get the folder name from which all external actors paths are created
 	 * @return folder name
diff --git a/Engine/Source/Runtime/Engine/Private/Level.cpp b/Engine/Source/Runtime/Engine/Private/Level.cpp
--- a/Engine/Source/Runtime/Engine/Private/Level.cpp
+++ b/Engine/Source/Runtime/Engine/Private/Level.cpp
@@ -1948,8 +1948,7 @@ namespace LevelAssetRegistryHelper
 		IAssetRegistry& AssetRegistry = FModuleManager::LoadModuleChecked<FAssetRegistryModule>(TEXT("AssetRegistry")).Get();
 		if (AssetRegistry.IsLoadingAssets() || IsRunningGame())
 		{
-			const FString PackagePath = FPackageName::GetLongPackagePath(LevelPackage.ToString());
-			AssetRegistry.ScanPathsSynchronous({ PackagePath });
+			AssetRegistry.ScanFilesSynchronous({ LevelPackage.ToString()});
 		}
 
 		TArray<FAssetData> LevelPackageAssets;
@@ -1966,6 +1965,17 @@ namespace LevelAssetRegistryHelper
 
 		return false;
 	}
+
+	static void ScanLevelAssets(const FString& LevelPackage)
+	{
+		IAssetRegistry& AssetRegistry = FModuleManager::LoadModuleChecked<FAssetRegistryModule>(TEXT("AssetRegistry")).Get();
+		AssetRegistry.ScanModifiedAssetFiles({ LevelPackage });
+
+		if (ULevel::GetIsLevelUsingExternalActorsFromPackage(FName(*LevelPackage)))
+		{
+			AssetRegistry.ScanPathsSynchronous({ ULevel::GetExternalActorsPath(LevelPackage) }, /*bForceRescan=*/true);
+		}
+	}
 }
 
 bool ULevel::GetIsLevelPartitionedFromPackage(FName LevelPackage)
@@ -2696,6 +2706,11 @@ FString ULevel::GetExternalActorsPath(UPackage* InLevelPackage, const FString& I
 	return GetExternalActorsPath(InLevelPackage->GetName(), InPackageShortName);
 }
 
+void ULevel::ScanLevelAssets(const FString& InLevelPackageName)
+{
+	LevelAssetRegistryHelper::ScanLevelAssets(InLevelPackageName);
+}
+
 const TCHAR* ULevel::GetExternalActorsFolderName()
 {
 	return FPackagePath::GetExternalActorsFolderName();
diff --git a/Engine/Source/Runtime/Engine/Private/LevelInstance/LevelInstanceSubsystem.cpp b/Engine/Source/Runtime/Engine/Private/LevelInstance/LevelInstanceSubsystem.cpp
--- a/Engine/Source/Runtime/Engine/Private/LevelInstance/LevelInstanceSubsystem.cpp
+++ b/Engine/Source/Runtime/Engine/Private/LevelInstance/LevelInstanceSubsystem.cpp
@@ -811,20 +811,6 @@ ALevelInstance* ULevelInstanceSubsystem::CreateLevelInstanceFrom(const TArray<AA
 		LevelFilename = FPackageName::LongPackageNameToFilename(CreationParams.LevelPackageName, FPackageName::GetMapPackageExtension());
 	}
 
-	ULevelStreamingLevelInstanceEditor* LevelStreaming = StaticCast<ULevelStreamingLevelInstanceEditor*>(EditorLevelUtils::CreateNewStreamingLevelForWorld(*GetWorld(), ULevelStreamingLevelInstanceEditor::StaticClass(), LevelFilename, false, CreationParams.TemplateWorld));
-	if (!LevelStreaming)
-	{
-		UE_LOG(LogLevelInstance, Warning, TEXT("Failed to create new Level Instance level"));
-		return nullptr;
-	}
-
-	ULevel* LoadedLevel = LevelStreaming->GetLoadedLevel();
-	check(LoadedLevel);
-
-	const bool bWarnAboutReferences = true;
-	const bool bWarnAboutRenaming = false;
-	const bool bMoveAllOrFail = true;
-
 	TSet<FName> DirtyPackages;
 
 	// Capture Packages before Moving actors as they can get GCed in the process
@@ -837,20 +823,17 @@ ALevelInstance* ULevelInstanceSubsystem::CreateLevelInstanceFrom(const TArray<AA
 		}
 	}
 
-	if (!EditorLevelUtils::MoveActorsToLevel(ActorsToMove, LoadedLevel, bWarnAboutReferences, bWarnAboutRenaming, bMoveAllOrFail))
+	ULevelStreamingLevelInstanceEditor* LevelStreaming = StaticCast<ULevelStreamingLevelInstanceEditor*>(EditorLevelUtils::CreateNewStreamingLevelForWorld(
+		*GetWorld(), ULevelStreamingLevelInstanceEditor::StaticClass(), CreationParams.UseExternalActors(), LevelFilename, &ActorsToMove, CreationParams.TemplateWorld));
+	if (!LevelStreaming)
 	{
-		ULevelStreamingLevelInstanceEditor::Unload(LevelStreaming);
-		UE_LOG(LogLevelInstance, Warning, TEXT("Failed to create Level Instance because some actors couldn't be moved"));
+		UE_LOG(LogLevelInstance, Warning, TEXT("Failed to create new Level"));
 		return nullptr;
 	}
-	
-	// Convert to OFPA
-	if (CreationParams.UseExternalActors())
-	{
-		LoadedLevel->ConvertAllActorsToPackaging(true);
-		LoadedLevel->bUseExternalActors = true;
-	}
-
+		
+	ULevel* LoadedLevel = LevelStreaming->GetLoadedLevel();
+	check(LoadedLevel);
+		
 	// Take all actors out of any folders they may have been in since we don't support folders inside of level instances
 	for (AActor* Actor : LoadedLevel->Actors)
 	{
@@ -865,6 +848,9 @@ ALevelInstance* ULevelInstanceSubsystem::CreateLevelInstanceFrom(const TArray<AA
 	ALevelInstance* NewLevelInstanceActor = nullptr;
 	TSoftObjectPtr<UWorld> WorldPtr(LoadedLevel->GetTypedOuter<UWorld>());
 	
+	// Make sure newly created level asset gets scanned
+	ULevel::ScanLevelAssets(LoadedLevel->GetPackage()->GetName());
+
 	if (CreationParams.Type == ELevelInstanceCreationType::LevelInstance)
 	{
 		NewLevelInstanceActor = GetWorld()->SpawnActor<ALevelInstance>(ALevelInstance::StaticClass(), SpawnParams);
@@ -903,7 +889,7 @@ ALevelInstance* ULevelInstanceSubsystem::CreateLevelInstanceFrom(const TArray<AA
 	NewLevelInstanceActor->SetWorldAsset(WorldPtr);
 	NewLevelInstanceActor->SetActorLocation(LevelInstanceLocation);
 	
-	// Actors were moved and kept their World positions so when saving we want their positions to actually be relative to the FounationActor/LevelTransform
+	// Actors were moved and kept their World positions so when saving we want their positions to actually be relative to the LevelInstance Actor
 	// so we set the LevelTransform and we mark the level as having moved its actors. 
 	// On Level save FLevelUtils::RemoveEditorTransform will fixup actor transforms to make them relative to the LevelTransform.
 	LevelStreaming->LevelTransform = NewLevelInstanceActor->GetActorTransform();
@@ -1656,9 +1642,7 @@ ALevelInstance* ULevelInstanceSubsystem::CommitLevelInstanceInternal(TUniquePtr<
 
 	if (bChangesCommitted)
 	{
-		// Sync the AssetData so that the updated instances have the latest Actor Registry Data
-		IAssetRegistry& AssetRegistry = FModuleManager::LoadModuleChecked<FAssetRegistryModule>(TEXT("AssetRegistry")).Get();
-		AssetRegistry.ScanPathsSynchronous({ EditPackage }, true);
+		ULevel::ScanLevelAssets(EditPackage);
 	}
 	
 	// Backup ID on Commit in case Actor gets recreated
