Texture: Log when the new texture build path is being used

#rnx

diff --git a/Engine/Source/Runtime/Engine/Private/TextureDerivedData.cpp b/Engine/Source/Runtime/Engine/Private/TextureDerivedData.cpp
--- a/Engine/Source/Runtime/Engine/Private/TextureDerivedData.cpp
+++ b/Engine/Source/Runtime/Engine/Private/TextureDerivedData.cpp
@@ -1216,6 +1216,7 @@ bool FTexturePlatformData::IsUsingNewDerivedData()
 			{
 				GConfig->GetBool(TEXT("TextureBuild"), TEXT("NewTextureBuilds"), bUseNewDerivedData, GEditorIni);
 			}
+			UE_CLOG(bUseNewDerivedData, LogTexture, Log, TEXT("Using new texture derived data builds."));
 		}
 		bool bUseNewDerivedData;
 	};
