Fix a crash on editor shutdown when the typed element registry object is destroyed before element objects.

#rnx

diff --git a/Engine/Source/Runtime/Engine/Private/Elements/Framework/EngineElementsLibrary.cpp b/Engine/Source/Runtime/Engine/Private/Elements/Framework/EngineElementsLibrary.cpp
--- a/Engine/Source/Runtime/Engine/Private/Elements/Framework/EngineElementsLibrary.cpp
+++ b/Engine/Source/Runtime/Engine/Private/Elements/Framework/EngineElementsLibrary.cpp
@@ -301,12 +301,20 @@ void UEngineElementsLibrary::DestroyEditorObjectElement(const UObject* Object)
 
 void UEngineElementsLibrary::DestroyUnreachableEditorObjectElements()
 {
-	auto IsObjectElementUnreachable = [](const TTypedElementOwner<FObjectElementData>& EditorElement)
+	UTypedElementRegistry* Registry = UTypedElementRegistry::GetInstance();
+	auto IsObjectElementUnreachable = [Registry](const TTypedElementOwner<FObjectElementData>& EditorElement)
 	{
+		// In the case we are shutting down, the registry may have already been cleaned up when we get here via UTypedElementRegistry::FinishDestroy,
+		// and our get checked element access will fail. We always return true if the registry is null, so that the elements can be destroyed during shutdown GC.
+		if (!Registry)
+		{
+			return true;
+		}
+
 		return EditorElement.GetDataChecked().Object->IsUnreachable();
 	};
 
-	auto DestroyUnreachableObjectElement = [Registry = UTypedElementRegistry::GetInstance()](TTypedElementOwner<FObjectElementData>&& EditorElement)
+	auto DestroyUnreachableObjectElement = [Registry](TTypedElementOwner<FObjectElementData>&& EditorElement)
 	{
 		if (Registry)
 		{
