Only sort 1 WorldSettings to the front of the list

diff --git a/Engine/Source/Runtime/Engine/Private/Level.cpp b/Engine/Source/Runtime/Engine/Private/Level.cpp
--- a/Engine/Source/Runtime/Engine/Private/Level.cpp
+++ b/Engine/Source/Runtime/Engine/Private/Level.cpp
@@ -1152,9 +1152,11 @@ static void SortActorsHierarchy(TArray<AActor*>& Actors, ULevel* Level)
 
 	DepthMap.Reserve(Actors.Num());
 
+	bool bFoundWorldSettings = false;
+
 	// This imitates the internals of GetDefaultBrush() without the sometimes problematic checks
 	ABrush* DefaultBrush = (Actors.Num() >= 2 ? Cast<ABrush>(Actors[1]) : nullptr); 
-	TFunction<int32(AActor*)> CalcAttachDepth = [&DepthMap, &VisitedActors, &CalcAttachDepth, DefaultBrush](AActor* Actor)
+	TFunction<int32(AActor*)> CalcAttachDepth = [&DepthMap, &VisitedActors, &CalcAttachDepth, &bFoundWorldSettings, DefaultBrush](AActor* Actor)
 	{
 		int32 Depth = 0;
 		if (int32* FoundDepth = DepthMap.Find(Actor))
@@ -1164,9 +1166,19 @@ static void SortActorsHierarchy(TArray<AActor*>& Actors, ULevel* Level)
 		else
 		{
 			// WorldSettings is expected to be the first element in the sorted Actors array
+			// To accomodate for the known issue where two world settings can exist, we only sort the
+			// first one we find to the 0 index
 			if (Actor->IsA<AWorldSettings>())
 			{
-				Depth = TNumericLimits<int32>::Lowest();
+				if (!bFoundWorldSettings)
+				{
+					Depth = TNumericLimits<int32>::Lowest();
+					bFoundWorldSettings = true;
+				}
+				else
+				{
+					UE_LOG(LogLevel, Warning, TEXT("Detected duplicate WorldSettings actor - UE-62934"));
+				}
 			}
 			// The default brush is expected to be the second element in the sorted Actors array
 			else if (Actor == DefaultBrush)
