Refactor the HLSLTree EmitValueShader API.  Rather than returning a string, now return FEmitShaderCode objects directly.  This allows implementation to more easily build operations out of component parts (which can be defined as helper functions that operate on FEmitShaderCode objects)

diff --git a/Engine/Source/Runtime/Engine/Private/HLSLTree/HLSLTree.cpp b/Engine/Source/Runtime/Engine/Private/HLSLTree/HLSLTree.cpp
--- a/Engine/Source/Runtime/Engine/Private/HLSLTree/HLSLTree.cpp
+++ b/Engine/Source/Runtime/Engine/Private/HLSLTree/HLSLTree.cpp
@@ -112,166 +112,87 @@ const TCHAR* FEmitContext::AcquireLocalDeclarationCode()
 	return MemStack::AllocateStringf(*Allocator, TEXT("Local%d"), NumExpressionLocals++);
 }
 
-FEmitShaderValue* FEmitContext::CastShaderValue(FNode* Node, FScope* Scope, FEmitShaderValue* ShaderValue, const Shader::FType& DestType)
+namespace Private
+{
+void MoveToScope(FEmitShaderCode* ShaderValue, FScope* Scope)
 {
-	if (ShaderValue->Type == DestType)
+	if (ShaderValue->Scope != Scope)
 	{
-		return ShaderValue;
+		FScope* NewScope = FScope::FindSharedParent(ShaderValue->Scope, Scope);
+		check(NewScope);
+
+		ShaderValue->Scope = NewScope;
+		for (FEmitShaderCode* Dependency : ShaderValue->Dependencies)
+		{
+			MoveToScope(Dependency, NewScope);
+		}
 	}
+}
 
-	const Shader::FValueTypeDescription SourceTypeDesc = Shader::GetValueTypeDescription(ShaderValue->Type);
-	const Shader::FValueTypeDescription DestTypeDesc = Shader::GetValueTypeDescription(DestType);
+void FormatArg_ShaderValue(FEmitShaderCode* ShaderValue, FEmitShaderValueDependencies& OutDependencies, FStringBuilderBase& OutCode)
+{
+	OutDependencies.Add(ShaderValue);
+	OutCode.Append(ShaderValue->Reference);
+}
 
-	TStringBuilder<1024> FormattedCode;
-	Shader::FType IntermediateType = DestType;
+} // namespace Private
 
-	if (SourceTypeDesc.NumComponents > 0 && DestTypeDesc.NumComponents > 0)
-	{
-		const bool bIsSourceLWC = SourceTypeDesc.ComponentType == Shader::EValueComponentType::Double;
-		const bool bIsLWC = DestTypeDesc.ComponentType == Shader::EValueComponentType::Double;
+FEmitShaderCode* FEmitContext::EmitFormatCodeInternal(const Shader::FType& Type, const TCHAR* Format, bool bInline, const FFormatArgList& ArgList)
+{
+	check(!Type.IsVoid());
 
-		if (bIsLWC != bIsSourceLWC)
+	TStringBuilder<1024> FormattedCode;
+	FEmitShaderValueDependencies Dependencies;
+	int32 ArgIndex = 0;
+	while (true)
+	{
+		const TCHAR Char = *Format++;
+		if (Char == 0)
 		{
-			if (bIsLWC)
-			{
-				// float->LWC
-				ShaderValue = CastShaderValue(Node, Scope, ShaderValue, Shader::MakeValueType(Shader::EValueComponentType::Float, DestTypeDesc.NumComponents));
-				FormattedCode.Appendf(TEXT("LWCPromote(%s)"), ShaderValue->Reference);
-			}
-			else
+			break;
+		}
+		else if (Char == TEXT('%'))
+		{
+			const FFormatArgVariant& Arg = ArgList[ArgIndex++];
+			switch (Arg.Type)
 			{
-				//LWC->float
-				FormattedCode.Appendf(TEXT("LWCToFloat(%s)"), ShaderValue->Reference);
-				IntermediateType = Shader::MakeValueType(Shader::EValueComponentType::Float, SourceTypeDesc.NumComponents);
+			case EFormatArgType::ShaderValue: Private::FormatArg_ShaderValue(Arg.ShaderValue, Dependencies, FormattedCode); break;
+			case EFormatArgType::String: FormattedCode.Append(Arg.String); break;
+			case EFormatArgType::Int: FormattedCode.Appendf(TEXT("%d"), Arg.Int); break;
+			default:
+				checkNoEntry();
+				break;
 			}
 		}
 		else
 		{
-			const bool bReplicateScalar = (SourceTypeDesc.NumComponents == 1);
-
-			int32 NumComponents = 0;
-			bool bNeedClosingParen = false;
-			if (bIsLWC)
-			{
-				FormattedCode.Append(TEXT("MakeLWCVector("));
-				bNeedClosingParen = true;
-			}
-			else
-			{
-				if (SourceTypeDesc.NumComponents == 1 || SourceTypeDesc.NumComponents == DestTypeDesc.NumComponents)
-				{
-					NumComponents = DestTypeDesc.NumComponents;
-					// Cast the scalar to the correct type, HLSL language will replicate the scalar if needed when performing this cast
-					FormattedCode.Appendf(TEXT("((%s)%s)"), DestTypeDesc.Name, ShaderValue->Reference);
-				}
-				else
-				{
-					NumComponents = FMath::Min(SourceTypeDesc.NumComponents, DestTypeDesc.NumComponents);
-					if (NumComponents < DestTypeDesc.NumComponents)
-					{
-						FormattedCode.Appendf(TEXT("%s("), DestTypeDesc.Name);
-						bNeedClosingParen = true;
-					}
-					if (NumComponents == SourceTypeDesc.NumComponents && SourceTypeDesc.ComponentType == DestTypeDesc.ComponentType)
-					{
-						// If we're taking all the components from the source, can avoid adding a swizzle
-						FormattedCode.Append(ShaderValue->Reference);
-					}
-					else
-					{
-						// Use a cast to truncate the source to the correct number of types
-						const Shader::EValueType LocalType = Shader::MakeValueType(DestTypeDesc.ComponentType, NumComponents);
-						FormattedCode.Appendf(TEXT("((%s)%s)"), Shader::GetValueTypeDescription(LocalType).Name, ShaderValue->Reference);
-					}
-				}
-			}
-
-			if (bNeedClosingParen)
-			{
-				const Shader::FValue ZeroValue(DestTypeDesc.ComponentType, 1);
-				for (int32 ComponentIndex = NumComponents; ComponentIndex < DestTypeDesc.NumComponents; ++ComponentIndex)
-				{
-					if (ComponentIndex > 0u)
-					{
-						FormattedCode.Append(TEXT(","));
-					}
-					if (bIsLWC)
-					{
-						if (!bReplicateScalar && ComponentIndex >= SourceTypeDesc.NumComponents)
-						{
-							FormattedCode.Append(TEXT("LWCPromote(0.0f)"));
-						}
-						else
-						{
-							FormattedCode.Appendf(TEXT("LWCGetComponent(%s, %d)"), ShaderValue->Reference, bReplicateScalar ? 0 : ComponentIndex);
-						}
-					}
-					else
-					{
-						// Non-LWC case should only be zero-filling here, other cases should have already been handled
-						check(!bReplicateScalar);
-						check(ComponentIndex >= SourceTypeDesc.NumComponents);
-						ZeroValue.ToString(Shader::EValueStringFormat::HLSL, FormattedCode);
-					}
-				}
-				NumComponents = DestTypeDesc.NumComponents;
-				FormattedCode.Append(TEXT(")"));
-			}
-
-			check(NumComponents == DestTypeDesc.NumComponents);
+			FormattedCode.Append(Char);
 		}
 	}
-	else
-	{
-		Errors.AddErrorf(Node, TEXT("Cannot cast between non-numeric types %s to %s."), SourceTypeDesc.Name, DestTypeDesc.Name);
-		FormattedCode.Appendf(TEXT("((%s)0)"), DestType.GetName());
-	}
+	checkf(ArgIndex == ArgList.Num(), TEXT("%d args were provided, but %d were used"), ArgList.Num(), ArgIndex);
 
-	check(IntermediateType != ShaderValue->Type);
-	const bool bInline = true;
-	ShaderValue = AcquireShader(Scope, IntermediateType, FormattedCode.ToView(), bInline, MakeArrayView(&ShaderValue, 1));
-	if (ShaderValue->Type != DestType)
-	{
-		// May need to cast through multiple intermediate types to reach our destination type
-		ShaderValue = CastShaderValue(Node, Scope, ShaderValue, DestType);
-	}
-	return ShaderValue;
+	return EmitCodeInternal(Type, FormattedCode.ToView(), bInline, Dependencies);
 }
 
-namespace Private
+FEmitShaderCode* FEmitContext::EmitCodeInternal(const Shader::FType& Type, FStringView Code, bool bInline, TArrayView<FEmitShaderCode*> Dependencies)
 {
-void MoveToScope(FEmitShaderValue* ShaderValue, FScope* Scope)
-{
-	if (ShaderValue->Scope != Scope)
-	{
-		FScope* NewScope = FScope::FindSharedParent(ShaderValue->Scope, Scope);
-		check(NewScope);
-
-		ShaderValue->Scope = NewScope;
-		for (FEmitShaderValue* Dependency : ShaderValue->Dependencies)
-		{
-			MoveToScope(Dependency, NewScope);
-		}
-	}
-}
-}
+	FScope* CurrentScope = ScopeStack.Last();
+	check(IsScopeLive(CurrentScope));
 
-FEmitShaderValue* FEmitContext::AcquireShader(FScope* Scope, const Shader::FType& Type, FStringView Code, bool bInline, TArrayView<FEmitShaderValue*> Dependencies)
-{
-	FEmitShaderValue* ShaderValue = nullptr;
+	FEmitShaderCode* ShaderValue = nullptr;
 
 	// Check to see if we've already generated code for an equivalent expression
 	const FSHAHash ShaderHash = HashString(Code);
-	FEmitShaderValue** const PrevShaderValue = ShaderValueMap.Find(ShaderHash);
+	FEmitShaderCode** const PrevShaderValue = ShaderValueMap.Find(ShaderHash);
 	if (PrevShaderValue)
 	{
 		ShaderValue = *PrevShaderValue;
 		check(ShaderValue->Type == Type);
-		Private::MoveToScope(ShaderValue, Scope);
+		Private::MoveToScope(ShaderValue, CurrentScope);
 	}
 	else
 	{
-		ShaderValue = new(*Allocator) FEmitShaderValue(Scope, Type);
+		ShaderValue = new(*Allocator) FEmitShaderCode(CurrentScope, Type);
 		ShaderValue->Hash = ShaderHash;
 		ShaderValue->Dependencies = MemStack::AllocateArrayView(*Allocator, Dependencies);
 		if (bInline)
@@ -349,7 +270,7 @@ void WriteMaterialUniformAccess(Shader::EValueComponentType ComponentType, uint3
 }
 } // namespace Private
 
-FEmitShaderValue* FEmitContext::AcquirePreshaderOrConstant(const FRequestedType& RequestedType, FScope* Scope, FExpression* Expression)
+FEmitShaderCode* FEmitContext::EmitPreshaderOrConstant(const FRequestedType& RequestedType, FExpression* Expression)
 {
 	Shader::FPreshaderData LocalPreshader;
 	Expression->EmitValuePreshader(*this, RequestedType, LocalPreshader);
@@ -360,12 +281,15 @@ FEmitShaderValue* FEmitContext::AcquirePreshaderOrConstant(const FRequestedType&
 	Hasher.Update((uint8*)&Type, sizeof(Type));
 	LocalPreshader.AppendHash(Hasher);
 	const FSHAHash Hash = Hasher.Finalize();
-	FEmitShaderValue* const* PrevShaderValue = PreshaderValueMap.Find(Hash);
+	FEmitShaderCode* const* PrevShaderValue = PreshaderValueMap.Find(Hash);
 	if (PrevShaderValue)
 	{
-		FEmitShaderValue* ShaderValue = *PrevShaderValue;
+		FScope* CurrentScope = ScopeStack.Last();
+		check(IsScopeLive(CurrentScope));
+
+		FEmitShaderCode* ShaderValue = *PrevShaderValue;
 		check(ShaderValue->Type == Type);
-		Private::MoveToScope(ShaderValue, Scope);
+		Private::MoveToScope(ShaderValue, CurrentScope);
 		return ShaderValue;
 	}
 
@@ -548,27 +472,25 @@ FEmitShaderValue* FEmitContext::AcquirePreshaderOrConstant(const FRequestedType&
 	}
 
 	const bool bInline = !Type.IsStruct(); // struct declarations can't be inline, due to HLSL syntax
-	FEmitShaderValue* ShaderValue = AcquireShader(Scope, Type, FormattedCode.ToView(), bInline, TArrayView<FEmitShaderValue*>());
+	FEmitShaderCode* ShaderValue = EmitCodeInternal(Type, FormattedCode.ToView(), bInline, TArrayView<FEmitShaderCode*>());
 	PreshaderValueMap.Add(Hash, ShaderValue);
 
 	return ShaderValue;
 }
 
-FEmitShaderValue* FEmitContext::AcquireConstantZero(FScope* Scope, const Shader::FType& Type)
+FEmitShaderCode* FEmitContext::EmitConstantZero(const Shader::FType& Type)
 {
-	TStringBuilder<256> FormattedCode;
-	FormattedCode.Appendf(TEXT("((%s)0)"), Type.GetName());
-	return AcquireShader(Scope, Type, FormattedCode.ToView(), true, TArrayView<FEmitShaderValue*>());
+	return EmitInlineCode(Type, TEXT("((%)0)"), Type.GetName());
 }
 
 namespace Private
 {
-void EmitShaderValue(FEmitContext& Context, FEmitShaderValue* ShaderValue)
+void EmitShaderValue(FEmitContext& Context, FEmitShaderCode* ShaderValue)
 {
 	if (ShaderValue->Scope)
 	{
 		// Emit dependencies first
-		for (FEmitShaderValue* Dependency : ShaderValue->Dependencies)
+		for (FEmitShaderCode* Dependency : ShaderValue->Dependencies)
 		{
 			EmitShaderValue(Context, Dependency);
 		}
@@ -588,8 +510,6 @@ void EmitShaderValue(FEmitContext& Context, FEmitShaderValue* ShaderValue)
 void FEmitContext::Finalize()
 {
 	check(ScopeStack.Num() == 0);
-	check(ShaderValueStack.Num() == 0);
-
 	for (const auto& It : ShaderValueMap)
 	{
 		Private::EmitShaderValue(*this, It.Value);
@@ -727,7 +647,7 @@ void FExpressionLocalPHI::PrepareValue(FEmitContext& Context, const FRequestedTy
 	}
 }
 
-void FExpressionLocalPHI::EmitValueShader(FEmitContext& Context, const FRequestedType& RequestedType, FEmitShaderResult& OutShader) const
+void FExpressionLocalPHI::EmitValueShader(FEmitContext& Context, const FRequestedType& RequestedType, FEmitShaderValues& OutResult) const
 {
 	int32 LocalPHIIndex = Context.LocalPHIs.Find(this);
 	if (LocalPHIIndex == INDEX_NONE)
@@ -758,7 +678,7 @@ void FExpressionLocalPHI::EmitValueShader(FEmitContext& Context, const FRequeste
 			check(IsScopeLive(ValueScope));
 
 			Context.ScopeStack.Add(ValueScope);
-			const FShaderValue ShaderValue = Values[i]->GetValueShader(Context, LocalType);
+			const FEmitShaderValues ShaderValue = Values[i]->GetValueShader(Context, LocalType);
 			Context.ScopeStack.Pop();
 
 			if (ValueScope == DeclarationScope)
@@ -766,17 +686,17 @@ void FExpressionLocalPHI::EmitValueShader(FEmitContext& Context, const FRequeste
 				ValueScope->EmitDeclarationf(Context, TEXT("%s LocalPHI%d = %s;"),
 					LocalType.GetName(),
 					LocalPHIIndex,
-					ShaderValue.Code);
+					ShaderValue.Code->Reference);
 				if (Derivative == EExpressionDerivative::Valid)
 				{
 					ValueScope->EmitDeclarationf(Context, TEXT("%s LocalPHI%dDdx = %s;"),
 						LocalDerivativeType.GetName(),
 						LocalPHIIndex,
-						ShaderValue.CodeDdx);
+						ShaderValue.CodeDdx->Reference);
 					ValueScope->EmitDeclarationf(Context, TEXT("%s LocalPHI%dDdy = %s;"),
 						LocalDerivativeType.GetName(),
 						LocalPHIIndex,
-						ShaderValue.CodeDdy);
+						ShaderValue.CodeDdy->Reference);
 				}
 
 				bNeedToAddDeclaration = false;
@@ -785,15 +705,15 @@ void FExpressionLocalPHI::EmitValueShader(FEmitContext& Context, const FRequeste
 			{
 				ValueScope->EmitStatementf(Context, TEXT("LocalPHI%d = %s;"),
 					LocalPHIIndex,
-					ShaderValue.Code);
+					ShaderValue.Code->Reference);
 				if (Derivative == EExpressionDerivative::Valid)
 				{
 					ValueScope->EmitStatementf(Context, TEXT("LocalPHI%dDdx = %s;"),
 						LocalPHIIndex,
-						ShaderValue.CodeDdx);
+						ShaderValue.CodeDdx->Reference);
 					ValueScope->EmitStatementf(Context, TEXT("LocalPHI%dDdy = %s;"),
 						LocalPHIIndex,
-						ShaderValue.CodeDdy);
+						ShaderValue.CodeDdy->Reference);
 				}
 			}
 		}
@@ -816,11 +736,14 @@ void FExpressionLocalPHI::EmitValueShader(FEmitContext& Context, const FRequeste
 		}
 	}
 
-	OutShader.Code.Appendf(TEXT("LocalPHI%d"), LocalPHIIndex);
-	OutShader.CodeDdx.Appendf(TEXT("LocalPHI%dDdx"), LocalPHIIndex);
-	OutShader.CodeDdy.Appendf(TEXT("LocalPHI%dDdy"), LocalPHIIndex);
-	OutShader.Type = GetType();
-	OutShader.bInline = true;
+	const Shader::FType LocalType = GetType();
+	OutResult.Code = Context.EmitInlineCode(LocalType, TEXT("LocalPHI%"), LocalPHIIndex);
+	if (GetDerivative(RequestedType) == EExpressionDerivative::Valid)
+	{
+		const Shader::FType LocalDerivativeType = LocalType.GetDerivativeType();
+		OutResult.CodeDdx = Context.EmitInlineCode(LocalDerivativeType, TEXT("LocalPHI%Ddx"), LocalPHIIndex);
+		OutResult.CodeDdy = Context.EmitInlineCode(LocalDerivativeType, TEXT("LocalPHI%Ddy"), LocalPHIIndex);
+	}
 }
 
 void FStatement::Reset()
@@ -1324,7 +1247,7 @@ void FPrepareValueResult::SetForwardValue(FEmitContext& Context, const FRequeste
 	}
 }
 
-void FExpression::EmitValueShader(FEmitContext& Context, const FRequestedType& RequestedType, FEmitShaderResult& OutShader) const
+void FExpression::EmitValueShader(FEmitContext& Context, const FRequestedType& RequestedType, FEmitShaderValues& OutResult) const
 {
 	check(false);
 }
@@ -1334,7 +1257,7 @@ void FExpression::EmitValuePreshader(FEmitContext& Context, const FRequestedType
 	check(false);
 }
 
-FShaderValue FExpression::GetValueShader(FEmitContext& Context, const FRequestedType& RequestedType, const Shader::FType& ResultType)
+FEmitShaderValues FExpression::GetValueShader(FEmitContext& Context, const FRequestedType& RequestedType, const Shader::FType& ResultType)
 {
 	if (PrepareValueResult.ForwardValue)
 	{
@@ -1344,108 +1267,46 @@ FShaderValue FExpression::GetValueShader(FEmitContext& Context, const FRequested
 	const FPreparedData Data = PrepareValueResult.PreparedType.GetData(RequestedType);
 	check(Data.IsValid());
 
-	FScope* CurrentScope = Context.ScopeStack.Last();
-	check(IsScopeLive(CurrentScope));
-
-	const EExpressionDerivative Derivative = Data.Derivative;
-	FEmitShaderValue* ShaderValue = nullptr;
-	FEmitShaderValue* ShaderValueDdx = nullptr;
-	FEmitShaderValue* ShaderValueDdy = nullptr;
+	FEmitShaderValues Result;
 	if (Data.Evaluation == EExpressionEvaluation::Constant || Data.Evaluation == EExpressionEvaluation::Preshader)
 	{
-		ShaderValue = Context.AcquirePreshaderOrConstant(RequestedType, CurrentScope, this);
-		check(Derivative != EExpressionDerivative::Valid); // If a constant has a valid derivative, it should be 'Zero'
+		Result.Code = Context.EmitPreshaderOrConstant(RequestedType, this);
+		check(Data.Derivative != EExpressionDerivative::Valid); // If a constant has a valid derivative, it should be 'Zero'
+		check(!Result.HasDerivatives());
 	}
 	else
 	{
 		check(Data.Evaluation == EExpressionEvaluation::Shader);
-
-		TStringBuilder<512> FormattedCode;
-		TStringBuilder<512> FormattedCodeDdx;
-		TStringBuilder<512> FormattedCodeDdy;
-		FEmitShaderResult ShaderResult(FormattedCode, FormattedCodeDdx, FormattedCodeDdy);
-		FEmitShaderValueDependencies Dependencies;
+		EmitValueShader(Context, RequestedType, Result);
+		if (Result.HasDerivatives())
 		{
-			FEmitShaderValueContext& ValueContext = Context.ShaderValueStack.AddDefaulted_GetRef();
-
-			EmitValueShader(Context, RequestedType, ShaderResult);
-			check(!ShaderResult.Type.IsVoid());
-
-			Dependencies = MoveTemp(ValueContext.Dependencies);
-			if (Derivative == EExpressionDerivative::Valid)
-			{
-				Dependencies.Append(MoveTemp(ValueContext.DerivativeDependencies));
-			}
-			Context.ShaderValueStack.Pop(false);
+			checkf(Data.Derivative == EExpressionDerivative::Valid, TEXT("Expression emitted derivatives, but didn't request them during PrepareValue"));
 		}
-		
-		ShaderValue = Context.AcquireShader(CurrentScope, ShaderResult.Type, ShaderResult.Code, ShaderResult.bInline, Dependencies);
-		if (Derivative == EExpressionDerivative::Valid)
+		else
 		{
-			check(ShaderResult.CodeDdx.Len() > 0);
-			check(ShaderResult.CodeDdy.Len() > 0);
-			const Shader::FType DerivativeResultType = ShaderResult.Type.GetDerivativeType();
-			check(!DerivativeResultType.IsVoid());
-			ShaderValueDdx = Context.AcquireShader(CurrentScope, DerivativeResultType, ShaderResult.CodeDdx, ShaderResult.bInline, Dependencies);
-			ShaderValueDdy = Context.AcquireShader(CurrentScope, DerivativeResultType, ShaderResult.CodeDdy, ShaderResult.bInline, Dependencies);
+			checkf(Data.Derivative != EExpressionDerivative::Valid, TEXT("Expression requested derivatives during PrepareValue, but didn't emit them"));
 		}
 	}
 
-	const Shader::FType DerivativeResultType = ResultType.GetDerivativeType();
-	check(ShaderValue);
-	ShaderValue = Context.CastShaderValue(this, CurrentScope, ShaderValue, ResultType);
-	bool bHasDerivatives = false;
-	if (Derivative == EExpressionDerivative::Valid)
-	{
-		check(!DerivativeResultType.IsVoid());
-		ShaderValueDdx = Context.CastShaderValue(this, CurrentScope, ShaderValueDdx, DerivativeResultType);
-		ShaderValueDdy = Context.CastShaderValue(this, CurrentScope, ShaderValueDdy, DerivativeResultType);
-		bHasDerivatives = true;
-	}
-	else if (Derivative == EExpressionDerivative::Zero)
+	Result = Context.EmitCast(Result, ResultType);
+	if (Data.Derivative == EExpressionDerivative::Zero)
 	{
+		const Shader::FType DerivativeResultType = ResultType.GetDerivativeType();
 		check(!DerivativeResultType.IsVoid());
-		ShaderValueDdx = Context.AcquireConstantZero(CurrentScope, DerivativeResultType);
-		ShaderValueDdy = ShaderValueDdx;
-		bHasDerivatives = true;
+		check(!Result.HasDerivatives());
+		Result.CodeDdx = Context.EmitConstantZero(DerivativeResultType);
+		Result.CodeDdy = Result.CodeDdx;
 	}
 
-	if (Context.ShaderValueStack.Num() > 0)
-	{
-		FEmitShaderValueContext& PrevValueContext = Context.ShaderValueStack.Last();
-		PrevValueContext.Dependencies.Add(ShaderValue);
-		if (bHasDerivatives)
-		{
-			check(ShaderValueDdx && ShaderValueDdy); // must have both
-			if (ShaderValueDdx != ShaderValue)
-			{
-				PrevValueContext.DerivativeDependencies.Add(ShaderValueDdx);
-				if (ShaderValueDdy != ShaderValueDdx)
-				{
-					PrevValueContext.DerivativeDependencies.Add(ShaderValueDdy);
-				}
-			}
-		}
-	}
-
-	FShaderValue Result;
-	Result.Type = ShaderValue->Type;
-	Result.Code = ShaderValue->Reference;
-	Result.bHasDerivatives = bHasDerivatives;
-	if (bHasDerivatives)
-	{
-		Result.CodeDdx = ShaderValueDdx->Reference;
-		Result.CodeDdy = ShaderValueDdy->Reference;
-	}
 	return Result;
 }
 
-FShaderValue FExpression::GetValueShader(FEmitContext& Context, const FRequestedType& RequestedType)
+FEmitShaderValues FExpression::GetValueShader(FEmitContext& Context, const FRequestedType& RequestedType)
 {
 	return GetValueShader(Context, RequestedType, RequestedType.GetType());
 }
 
-FShaderValue FExpression::GetValueShader(FEmitContext& Context)
+FEmitShaderValues FExpression::GetValueShader(FEmitContext& Context)
 {
 	return GetValueShader(Context, GetRequestedType());
 }
diff --git a/Engine/Source/Runtime/Engine/Private/HLSLTree/HLSLTreeIntrinsics.cpp b/Engine/Source/Runtime/Engine/Private/HLSLTree/HLSLTreeIntrinsics.cpp
new file mode 100644
--- /dev/null
+++ b/Engine/Source/Runtime/Engine/Private/HLSLTree/HLSLTreeIntrinsics.cpp
@@ -0,0 +1,280 @@
+// Copyright Epic Games, Inc. All Rights Reserved.
+#include "HLSLTree/HLSLTree.h"
+#include "Misc/StringBuilder.h"
+#include "Misc/MemStack.h"
+#include "Misc/MemStackUtility.h"
+#include "Shader/ShaderTypes.h"
+#include "Shader/Preshader.h"
+#include "Misc/LargeWorldRenderPosition.h"
+
+namespace UE
+{
+namespace HLSLTree
+{
+
+FBinaryOpDescription::FBinaryOpDescription()
+	: Name(nullptr), Operator(nullptr), PreshaderOpcode(Shader::EPreshaderOpcode::Nop)
+{}
+
+FBinaryOpDescription::FBinaryOpDescription(const TCHAR* InName, const TCHAR* InOperator, Shader::EPreshaderOpcode InOpcode)
+	: Name(InName), Operator(InOperator), PreshaderOpcode(InOpcode)
+{}
+
+FBinaryOpDescription GetBinaryOpDesription(EBinaryOp Op)
+{
+	switch (Op)
+	{
+	case EBinaryOp::None: return FBinaryOpDescription(TEXT("None"), TEXT(""), Shader::EPreshaderOpcode::Nop); break;
+	case EBinaryOp::Add: return FBinaryOpDescription(TEXT("Add"), TEXT("+"), Shader::EPreshaderOpcode::Add); break;
+	case EBinaryOp::Sub: return FBinaryOpDescription(TEXT("Subtract"), TEXT("-"), Shader::EPreshaderOpcode::Sub); break;
+	case EBinaryOp::Mul: return FBinaryOpDescription(TEXT("Multiply"), TEXT("*"), Shader::EPreshaderOpcode::Mul); break;
+	case EBinaryOp::Div: return FBinaryOpDescription(TEXT("Divide"), TEXT("/"), Shader::EPreshaderOpcode::Div); break;
+	case EBinaryOp::Less: return FBinaryOpDescription(TEXT("Less"), TEXT("<"), Shader::EPreshaderOpcode::Nop); break;
+	default: checkNoEntry(); return FBinaryOpDescription();
+	}
+}
+
+FEmitShaderCode* FEmitContext::EmitCast(FEmitShaderCode* ShaderValue, const Shader::FType& DestType)
+{
+	check(ShaderValue);
+	check(!DestType.IsVoid());
+
+	if (ShaderValue->Type == DestType)
+	{
+		return ShaderValue;
+	}
+
+	const Shader::FValueTypeDescription SourceTypeDesc = Shader::GetValueTypeDescription(ShaderValue->Type);
+	const Shader::FValueTypeDescription DestTypeDesc = Shader::GetValueTypeDescription(DestType);
+
+	TStringBuilder<1024> FormattedCode;
+	Shader::FType IntermediateType = DestType;
+
+	if (SourceTypeDesc.NumComponents > 0 && DestTypeDesc.NumComponents > 0)
+	{
+		const bool bIsSourceLWC = SourceTypeDesc.ComponentType == Shader::EValueComponentType::Double;
+		const bool bIsLWC = DestTypeDesc.ComponentType == Shader::EValueComponentType::Double;
+
+		if (bIsLWC != bIsSourceLWC)
+		{
+			if (bIsLWC)
+			{
+				// float->LWC
+				ShaderValue = EmitCast(ShaderValue, Shader::MakeValueType(Shader::EValueComponentType::Float, DestTypeDesc.NumComponents));
+				FormattedCode.Appendf(TEXT("LWCPromote(%s)"), ShaderValue->Reference);
+			}
+			else
+			{
+				//LWC->float
+				FormattedCode.Appendf(TEXT("LWCToFloat(%s)"), ShaderValue->Reference);
+				IntermediateType = Shader::MakeValueType(Shader::EValueComponentType::Float, SourceTypeDesc.NumComponents);
+			}
+		}
+		else
+		{
+			const bool bReplicateScalar = (SourceTypeDesc.NumComponents == 1);
+
+			int32 NumComponents = 0;
+			bool bNeedClosingParen = false;
+			if (bIsLWC)
+			{
+				FormattedCode.Append(TEXT("MakeLWCVector("));
+				bNeedClosingParen = true;
+			}
+			else
+			{
+				if (SourceTypeDesc.NumComponents == 1 || SourceTypeDesc.NumComponents == DestTypeDesc.NumComponents)
+				{
+					NumComponents = DestTypeDesc.NumComponents;
+					// Cast the scalar to the correct type, HLSL language will replicate the scalar if needed when performing this cast
+					FormattedCode.Appendf(TEXT("((%s)%s)"), DestTypeDesc.Name, ShaderValue->Reference);
+				}
+				else
+				{
+					NumComponents = FMath::Min(SourceTypeDesc.NumComponents, DestTypeDesc.NumComponents);
+					if (NumComponents < DestTypeDesc.NumComponents)
+					{
+						FormattedCode.Appendf(TEXT("%s("), DestTypeDesc.Name);
+						bNeedClosingParen = true;
+					}
+					if (NumComponents == SourceTypeDesc.NumComponents && SourceTypeDesc.ComponentType == DestTypeDesc.ComponentType)
+					{
+						// If we're taking all the components from the source, can avoid adding a swizzle
+						FormattedCode.Append(ShaderValue->Reference);
+					}
+					else
+					{
+						// Use a cast to truncate the source to the correct number of types
+						const Shader::EValueType LocalType = Shader::MakeValueType(DestTypeDesc.ComponentType, NumComponents);
+						FormattedCode.Appendf(TEXT("((%s)%s)"), Shader::GetValueTypeDescription(LocalType).Name, ShaderValue->Reference);
+					}
+				}
+			}
+
+			if (bNeedClosingParen)
+			{
+				const Shader::FValue ZeroValue(DestTypeDesc.ComponentType, 1);
+				for (int32 ComponentIndex = NumComponents; ComponentIndex < DestTypeDesc.NumComponents; ++ComponentIndex)
+				{
+					if (ComponentIndex > 0u)
+					{
+						FormattedCode.Append(TEXT(","));
+					}
+					if (bIsLWC)
+					{
+						if (!bReplicateScalar && ComponentIndex >= SourceTypeDesc.NumComponents)
+						{
+							FormattedCode.Append(TEXT("LWCPromote(0.0f)"));
+						}
+						else
+						{
+							FormattedCode.Appendf(TEXT("LWCGetComponent(%s, %d)"), ShaderValue->Reference, bReplicateScalar ? 0 : ComponentIndex);
+						}
+					}
+					else
+					{
+						// Non-LWC case should only be zero-filling here, other cases should have already been handled
+						check(!bReplicateScalar);
+						check(ComponentIndex >= SourceTypeDesc.NumComponents);
+						ZeroValue.ToString(Shader::EValueStringFormat::HLSL, FormattedCode);
+					}
+				}
+				NumComponents = DestTypeDesc.NumComponents;
+				FormattedCode.Append(TEXT(")"));
+			}
+
+			check(NumComponents == DestTypeDesc.NumComponents);
+		}
+	}
+	else
+	{
+		Errors.AddErrorf(nullptr, TEXT("Cannot cast between non-numeric types %s to %s."), SourceTypeDesc.Name, DestTypeDesc.Name);
+		FormattedCode.Appendf(TEXT("((%s)0)"), DestType.GetName());
+	}
+
+	check(IntermediateType != ShaderValue->Type);
+	const bool bInline = true;
+	ShaderValue = EmitCodeInternal(IntermediateType, FormattedCode.ToView(), bInline, MakeArrayView(&ShaderValue, 1));
+	if (ShaderValue->Type != DestType)
+	{
+		// May need to cast through multiple intermediate types to reach our destination type
+		ShaderValue = EmitCast(ShaderValue, DestType);
+	}
+	return ShaderValue;
+}
+
+FEmitShaderValues FEmitContext::EmitCast(FEmitShaderValues ShaderValue, const Shader::FType& DestType)
+{
+	FEmitShaderValues Result;
+	Result.Code = EmitCast(ShaderValue.Code, DestType);
+	if (ShaderValue.HasDerivatives())
+	{
+		const Shader::FType DerivativeType = DestType.GetDerivativeType();
+		Result.CodeDdx = EmitCast(ShaderValue.CodeDdx, DerivativeType);
+		Result.CodeDdy = EmitCast(ShaderValue.CodeDdy, DerivativeType);
+	}
+	return Result;
+}
+
+FEmitShaderCode* FEmitContext::EmitBinaryOp(EBinaryOp Op, FEmitShaderCode* Lhs, FEmitShaderCode* Rhs)
+{
+	const Shader::FValueTypeDescription LhsTypeDesc = Shader::GetValueTypeDescription(Lhs->Type);
+	const Shader::FValueTypeDescription RhsTypeDesc = Shader::GetValueTypeDescription(Rhs->Type);
+	const int8 NumComponents = (LhsTypeDesc.NumComponents == 1 || RhsTypeDesc.NumComponents == 1) ? FMath::Max(LhsTypeDesc.NumComponents, RhsTypeDesc.NumComponents) : FMath::Min(LhsTypeDesc.NumComponents, RhsTypeDesc.NumComponents);
+	const Shader::EValueComponentType InputComponentType = Shader::CombineComponentTypes(LhsTypeDesc.ComponentType, RhsTypeDesc.ComponentType);
+	const Shader::EValueType InputType = Shader::MakeValueType(InputComponentType, NumComponents);
+	Shader::EValueType ResultType = InputType;
+
+	FEmitShaderCode* LhsCast = EmitCast(Lhs, InputType);
+	FEmitShaderCode* RhsCast = EmitCast(Rhs, InputType);
+
+	FEmitShaderCode* Result = nullptr;
+	switch (Op)
+	{
+	case EBinaryOp::Add:
+		if (InputComponentType == Shader::EValueComponentType::Double)
+		{
+			Result = EmitCode(ResultType, TEXT("LWCAdd(%, %)"), LhsCast, RhsCast);
+		}
+		else
+		{
+			Result = EmitCode(ResultType, TEXT("(% + %)"), LhsCast, RhsCast);
+		}
+		break;
+	case EBinaryOp::Sub:
+		if (InputComponentType == Shader::EValueComponentType::Double)
+		{
+			Result = EmitCode(ResultType, TEXT("LWCSubtract(%, %)"), LhsCast, RhsCast);
+		}
+		else
+		{
+			Result = EmitCode(ResultType, TEXT("(% - %)"), LhsCast, RhsCast);
+		}
+		break;
+	case EBinaryOp::Mul:
+		if (InputComponentType == Shader::EValueComponentType::Double)
+		{
+			Result = EmitCode(ResultType, TEXT("LWCMultiply(%, %)"), LhsCast, RhsCast);
+		}
+		else
+		{
+			Result = EmitCode(ResultType, TEXT("(% * %)"), LhsCast, RhsCast);
+		}
+		break;
+	case EBinaryOp::Div:
+		if (InputComponentType == Shader::EValueComponentType::Double)
+		{
+			Result = EmitCode(ResultType, TEXT("LWCDivide(%, %)"), LhsCast, RhsCast);
+		}
+		else
+		{
+			Result = EmitCode(ResultType, TEXT("(% / %)"), LhsCast, RhsCast);
+		}
+		break;
+	case EBinaryOp::Less:
+		ResultType = Shader::MakeValueType(Shader::EValueComponentType::Bool, NumComponents);
+		if (InputComponentType == Shader::EValueComponentType::Double)
+		{
+			Result = EmitCode(ResultType, TEXT("LWCLess(%, %)"), LhsCast, RhsCast);
+		}
+		else
+		{
+			Result = EmitCode(ResultType, TEXT("(% < %)"), LhsCast, RhsCast);
+		}
+		break;
+	default:
+		checkNoEntry();
+		break;
+	}
+
+	return Result;
+}
+
+FEmitShaderValues FEmitContext::EmitBinaryOp(EBinaryOp Op, FEmitShaderValues Lhs, FEmitShaderValues Rhs, EExpressionDerivative Derivative)
+{
+	FEmitShaderValues Result;
+	Result.Code = EmitBinaryOp(Op, Lhs.Code, Rhs.Code);
+	if (Derivative == EExpressionDerivative::Valid && Lhs.HasDerivatives() && Rhs.HasDerivatives())
+	{
+		switch (Op)
+		{
+		case EBinaryOp::Add:
+			Result.CodeDdx = EmitAdd(Lhs.CodeDdx, Rhs.CodeDdx);
+			Result.CodeDdy = EmitAdd(Lhs.CodeDdy, Rhs.CodeDdy);
+			break;
+		case EBinaryOp::Sub:
+			Result.CodeDdx = EmitSub(Lhs.CodeDdx, Rhs.CodeDdx);
+			Result.CodeDdy = EmitSub(Lhs.CodeDdy, Rhs.CodeDdy);
+			break;
+		case EBinaryOp::Mul:
+			Result.CodeDdx = EmitAdd(EmitMul(Lhs.CodeDdx, Rhs.Code), EmitMul(Rhs.CodeDdx, Lhs.Code));
+			Result.CodeDdy = EmitAdd(EmitMul(Lhs.CodeDdy, Rhs.Code), EmitMul(Rhs.CodeDdy, Lhs.Code));
+			break;
+		}
+	}
+
+	return Result;
+}
+
+} // namespace HLSLTree
+} // namespace UE
diff --git a/Engine/Source/Runtime/Engine/Public/HLSLTree/HLSLTree.h b/Engine/Source/Runtime/Engine/Public/HLSLTree/HLSLTree.h
--- a/Engine/Source/Runtime/Engine/Public/HLSLTree/HLSLTree.h
+++ b/Engine/Source/Runtime/Engine/Public/HLSLTree/HLSLTree.h
@@ -35,8 +35,8 @@ class FNode;
 class FScope;
 class FExpressionLocalPHI;
 class FRequestedType;
-struct FEmitShaderValue;
-struct FEmitShaderResult;
+struct FEmitShaderCode;
+struct FEmitShaderValues;
 
 static constexpr int32 MaxNumPreviousScopes = 2;
 
@@ -82,14 +82,66 @@ enum class ECastFlags : uint32
 };
 ENUM_CLASS_FLAGS(ECastFlags);
 
-using FEmitShaderValueDependencies = TArray<FEmitShaderValue*, TInlineAllocator<8>>;
+struct FEmitShaderCode
+{
+	FEmitShaderCode(FScope* InScope, const Shader::FType& InType) : Scope(InScope), Type(InType) {}
+
+	inline bool IsInline() const { return Value == nullptr; }
+
+	FScope* Scope = nullptr;
+	const TCHAR* Reference = nullptr;
+	const TCHAR* Value = nullptr;
+	Shader::FType Type;
+	TArrayView<FEmitShaderCode*> Dependencies;
+	FSHAHash Hash;
+};
+
+struct FEmitShaderValues
+{
+	FEmitShaderCode* Code = nullptr;
+	FEmitShaderCode* CodeDdx = nullptr;
+	FEmitShaderCode* CodeDdy = nullptr;
+
+	bool HasDerivatives() const { return (bool)CodeDdx && (bool)CodeDdy; }
+};
+
+using FEmitShaderValueDependencies = TArray<FEmitShaderCode*, TInlineAllocator<8>>;
 
-struct FEmitShaderValueContext
+enum class EFormatArgType : uint8
 {
-	FEmitShaderValueDependencies Dependencies;
-	FEmitShaderValueDependencies DerivativeDependencies;
+	Void,
+	ShaderValue,
+	String,
+	Int,
 };
 
+struct FFormatArgVariant
+{
+	FFormatArgVariant() {}
+	FFormatArgVariant(FEmitShaderCode* InValue) : Type(EFormatArgType::ShaderValue), ShaderValue(InValue) { check(InValue); }
+	FFormatArgVariant(const TCHAR* InValue) : Type(EFormatArgType::String), String(InValue) { check(InValue); }
+	FFormatArgVariant(int32 InValue) : Type(EFormatArgType::Int), Int(InValue) {}
+
+	EFormatArgType Type = EFormatArgType::Void;
+	union
+	{
+		FEmitShaderCode* ShaderValue;
+		const TCHAR* String;
+		int32 Int;
+	};
+};
+
+using FFormatArgList = TArray<FFormatArgVariant, TInlineAllocator<8>>;
+
+inline void BuildFormatArgList(FFormatArgList&) {}
+
+template<typename Type, typename... Types>
+inline void BuildFormatArgList(FFormatArgList& OutList, Type Arg, Types... Args)
+{
+	OutList.Add(Arg);
+	BuildFormatArgList(OutList, Forward<Types>(Args)...);
+}
+
 /** Tracks shared state while emitting HLSL code */
 class FEmitContext
 {
@@ -102,16 +154,45 @@ public:
 	/** Get a unique local variable name */
 	const TCHAR* AcquireLocalDeclarationCode();
 
-	FEmitShaderValue* AcquireShader(FScope* Scope, const Shader::FType& Type, FStringView Code, bool bInline, TArrayView<FEmitShaderValue*> Dependencies);
-	FEmitShaderValue* AcquirePreshaderOrConstant(const FRequestedType& RequestedType, FScope* Scope, FExpression* Expression);
-	FEmitShaderValue* AcquireConstantZero(FScope* Scope, const Shader::FType& Type);
+	FEmitShaderCode* EmitCodeInternal(const Shader::FType& Type, FStringView Code, bool bInline, TArrayView<FEmitShaderCode*> Dependencies);
+	FEmitShaderCode* EmitFormatCodeInternal(const Shader::FType& Type, const TCHAR* Format, bool bInline, const FFormatArgList& ArgList);
 
-	FEmitShaderValue* CastShaderValue(FNode* Node, FScope* Scope, FEmitShaderValue* ShaderValue, const Shader::FType& DestType);
+	template<typename FormatType, typename... Types>
+	FEmitShaderCode* EmitCode(const Shader::FType& Type, const FormatType& Format, Types... Args)
+	{
+		FFormatArgList ArgList;
+		BuildFormatArgList(ArgList, Forward<Types>(Args)...);
+		return EmitFormatCodeInternal(Type, Format, false, ArgList);
+	}
+
+	template<typename FormatType, typename... Types>
+	FEmitShaderCode* EmitInlineCode(const Shader::FType& Type, const FormatType& Format, Types... Args)
+	{
+		FFormatArgList ArgList;
+		BuildFormatArgList(ArgList, Forward<Types>(Args)...);
+		return EmitFormatCodeInternal(Type, Format, true, ArgList);
+	}
+
+	
+	FEmitShaderCode* EmitPreshaderOrConstant(const FRequestedType& RequestedType, FExpression* Expression);
+	FEmitShaderCode* EmitConstantZero(const Shader::FType& Type);
+
+	FEmitShaderCode* EmitCast(FEmitShaderCode* ShaderValue, const Shader::FType& DestType);
+	FEmitShaderValues EmitCast(FEmitShaderValues ShaderValue, const Shader::FType& DestType);
+
+	FEmitShaderCode* EmitBinaryOp(EBinaryOp Op, FEmitShaderCode* Lhs, FEmitShaderCode* Rhs);
+	FEmitShaderValues EmitBinaryOp(EBinaryOp Op, FEmitShaderValues Lhs, FEmitShaderValues Rhs, EExpressionDerivative Derivative);
+
+	template<typename T> inline T EmitAdd(T Lhs, T Rhs) { return EmitBinaryOp(EBinaryOp::Add, Lhs, Rhs); }
+	template<typename T> inline T EmitSub(T Lhs, T Rhs) { return EmitBinaryOp(EBinaryOp::Sub, Lhs, Rhs); }
+	template<typename T> inline T EmitMul(T Lhs, T Rhs) { return EmitBinaryOp(EBinaryOp::Mul, Lhs, Rhs); }
+	template<typename T> inline T EmitDiv(T Lhs, T Rhs) { return EmitBinaryOp(EBinaryOp::Div, Lhs, Rhs); }
 
 	FMemStackBase* Allocator = nullptr;
 	const Shader::FStructTypeRegistry* TypeRegistry = nullptr;
-	TMap<FSHAHash, FEmitShaderValue*> ShaderValueMap;
-	TMap<FSHAHash, FEmitShaderValue*> PreshaderValueMap;
+	TArray<FScope*, TInlineAllocator<16>> ScopeStack;
+	TMap<FSHAHash, FEmitShaderCode*> ShaderValueMap;
+	TMap<FSHAHash, FEmitShaderCode*> PreshaderValueMap;
 	TArray<const FExpressionLocalPHI*> LocalPHIs;
 	FErrors Errors;
 
@@ -120,8 +201,6 @@ public:
 	const FStaticParameterSet* StaticParameters = nullptr;
 	FMaterialCompilationOutput* MaterialCompilationOutput = nullptr;
 	TMap<Shader::FValue, uint32> DefaultUniformValues;
-	TArray<FScope*, TInlineAllocator<16>> ScopeStack;
-	TArray<FEmitShaderValueContext, TInlineAllocator<16>> ShaderValueStack;
 	uint32 UniformPreshaderOffset = 0u;
 	bool bReadMaterialNormal = false;
 
@@ -298,6 +377,7 @@ public:
 	FPreparedData GetComponentData(int32 Index) const;
 
 	inline EExpressionEvaluation GetEvaluation(const FRequestedType& RequestedType) const { return GetData(RequestedType).Evaluation; }
+	inline EExpressionDerivative GetDerivative(const FRequestedType& RequestedType) const { return GetData(RequestedType).Derivative; }
 
 	void SetComponentData(int32 Index, const FPreparedData& Data);
 	void MergeComponentData(int32 Index, EComponentRequest Request, const FPreparedData& Data);
@@ -312,21 +392,6 @@ public:
 
 FPreparedType MergePreparedTypes(const FPreparedType& Lhs, const FPreparedType& Rhs);
 
-struct FEmitShaderResult
-{
-	FEmitShaderResult(FStringBuilderBase& InCode, FStringBuilderBase& InCodeDdx, FStringBuilderBase& InCodeDdy)
-		: Code(InCode)
-		, CodeDdx(InCodeDdx)
-		, CodeDdy(InCodeDdy)
-	{}
-
-	FStringBuilderBase& Code;
-	FStringBuilderBase& CodeDdx;
-	FStringBuilderBase& CodeDdy;
-	Shader::FType Type;
-	bool bInline = false;
-};
-
 class FPrepareValueResult
 {
 public:
@@ -351,29 +416,6 @@ private:
 	friend class FExpression;
 };
 
-struct FEmitShaderValue
-{
-	FEmitShaderValue(FScope* InScope, const Shader::FType& InType) : Scope(InScope), Type(InType) {}
-
-	inline bool IsInline() const { return Value == nullptr; }
-
-	FScope* Scope = nullptr;
-	const TCHAR* Reference = nullptr;
-	const TCHAR* Value = nullptr;
-	Shader::FType Type;
-	TArrayView<FEmitShaderValue*> Dependencies;
-	FSHAHash Hash;
-};
-
-struct FShaderValue
-{
-	const TCHAR* Code = nullptr;
-	const TCHAR* CodeDdx = nullptr;
-	const TCHAR* CodeDdy = nullptr;
-	Shader::FType Type;
-	bool bHasDerivatives = false;
-};
-
 /**
  * Represents an HLSL expression.  This is a piece of code that evaluates to a value, but has no side effects.
  * Unlike statements, expressions are not expected to execute in any particular order.  They may be cached (or not) in generated code, without the underlying implementation needing to care.
@@ -387,21 +429,22 @@ public:
 	FRequestedType GetRequestedType() const { return PrepareValueResult.PreparedType.GetRequestedType(); }
 	Shader::FType GetType() const { return PrepareValueResult.PreparedType.GetType(); }
 	EExpressionEvaluation GetEvaluation(const FRequestedType& RequestedType) const { return PrepareValueResult.PreparedType.GetEvaluation(RequestedType); }
+	EExpressionDerivative GetDerivative(const FRequestedType& RequestedType) const { return PrepareValueResult.PreparedType.GetDerivative(RequestedType); }
 
 	virtual void Reset() override;
 
 	friend const FPreparedType& PrepareExpressionValue(FEmitContext& Context, FExpression* InExpression, const FRequestedType& RequestedType);
 
-	FShaderValue GetValueShader(FEmitContext& Context, const FRequestedType& RequestedType, const Shader::FType& ResultType);
+	FEmitShaderValues GetValueShader(FEmitContext& Context, const FRequestedType& RequestedType, const Shader::FType& ResultType);
 	void GetValuePreshader(FEmitContext& Context, const FRequestedType& RequestedType, Shader::FPreshaderData& OutPreshader);
 	Shader::FValue GetValueConstant(FEmitContext& Context, const FRequestedType& RequestedType);
 
-	FShaderValue GetValueShader(FEmitContext& Context, const FRequestedType& RequestedType);
-	FShaderValue GetValueShader(FEmitContext& Context);
+	FEmitShaderValues GetValueShader(FEmitContext& Context, const FRequestedType& RequestedType);
+	FEmitShaderValues GetValueShader(FEmitContext& Context);
 
 protected:
 	virtual void PrepareValue(FEmitContext& Context, const FRequestedType& RequestedType, FPrepareValueResult& OutResult) const = 0;
-	virtual void EmitValueShader(FEmitContext& Context, const FRequestedType& RequestedType, FEmitShaderResult& OutShader) const;
+	virtual void EmitValueShader(FEmitContext& Context, const FRequestedType& RequestedType, FEmitShaderValues& OutResult) const;
 	virtual void EmitValuePreshader(FEmitContext& Context, const FRequestedType& RequestedType, Shader::FPreshaderData& OutPreshader) const;
 
 private:
@@ -422,7 +465,7 @@ class FExpressionLocalPHI final : public FExpression
 {
 public:
 	virtual void PrepareValue(FEmitContext& Context, const FRequestedType& RequestedType, FPrepareValueResult& OutResult) const override;
-	virtual void EmitValueShader(FEmitContext& Context, const FRequestedType& RequestedType, FEmitShaderResult& OutShader) const override;
+	virtual void EmitValueShader(FEmitContext& Context, const FRequestedType& RequestedType, FEmitShaderValues& OutResult) const override;
 
 	FName LocalName;
 	FScope* Scopes[MaxNumPreviousScopes];
diff --git a/Engine/Source/Runtime/Engine/Public/HLSLTree/HLSLTreeCommon.h b/Engine/Source/Runtime/Engine/Public/HLSLTree/HLSLTreeCommon.h
--- a/Engine/Source/Runtime/Engine/Public/HLSLTree/HLSLTreeCommon.h
+++ b/Engine/Source/Runtime/Engine/Public/HLSLTree/HLSLTreeCommon.h
@@ -8,36 +8,9 @@ enum class EMaterialParameterType : uint8;
 namespace UE
 {
 
-namespace Shader
-{
-enum class EPreshaderOpcode : uint8;
-} // namespace Shader
-
 namespace HLSLTree
 {
 
-enum class EBinaryOp : uint8
-{
-	None,
-	Add,
-	Sub,
-	Mul,
-	Div,
-	Less,
-};
-
-struct FBinaryOpDescription
-{
-	FBinaryOpDescription();
-	FBinaryOpDescription(const TCHAR* InName, const TCHAR* InOperator, Shader::EPreshaderOpcode InOpcode);
-
-	const TCHAR* Name;
-	const TCHAR* Operator;
-	Shader::EPreshaderOpcode PreshaderOpcode;
-};
-
-FBinaryOpDescription GetBinaryOpDesription(EBinaryOp Op);
-
 class FExpressionConstant : public FExpression
 {
 public:
@@ -125,7 +98,7 @@ public:
 	EExternalInputType InputType;
 
 	virtual void PrepareValue(FEmitContext& Context, const FRequestedType& RequestedType, FPrepareValueResult& OutResult) const override;
-	virtual void EmitValueShader(FEmitContext& Context, const FRequestedType& RequestedType, FEmitShaderResult& OutShader) const override;
+	virtual void EmitValueShader(FEmitContext& Context, const FRequestedType& RequestedType, FEmitShaderValues& OutResult) const override;
 };
 
 class FExpressionTextureSample : public FExpression
@@ -144,7 +117,7 @@ public:
 	ETextureMipValueMode MipValueMode;
 
 	virtual void PrepareValue(FEmitContext& Context, const FRequestedType& RequestedType, FPrepareValueResult& OutResult) const override;
-	virtual void EmitValueShader(FEmitContext& Context, const FRequestedType& RequestedType, FEmitShaderResult& OutShader) const override;
+	virtual void EmitValueShader(FEmitContext& Context, const FRequestedType& RequestedType, FEmitShaderValues& OutResult) const override;
 };
 
 class FExpressionGetStructField : public FExpression
@@ -163,7 +136,7 @@ public:
 	FExpression* StructExpression;
 
 	virtual void PrepareValue(FEmitContext& Context, const FRequestedType& RequestedType, FPrepareValueResult& OutResult) const override;
-	virtual void EmitValueShader(FEmitContext& Context, const FRequestedType& RequestedType, FEmitShaderResult& OutShader) const override;
+	virtual void EmitValueShader(FEmitContext& Context, const FRequestedType& RequestedType, FEmitShaderValues& OutResult) const override;
 	virtual void EmitValuePreshader(FEmitContext& Context, const FRequestedType& RequestedType, Shader::FPreshaderData& OutPreshader) const override;
 };
 
@@ -185,7 +158,7 @@ public:
 	FExpression* FieldExpression;
 
 	virtual void PrepareValue(FEmitContext& Context, const FRequestedType& RequestedType, FPrepareValueResult& OutResult) const override;
-	virtual void EmitValueShader(FEmitContext& Context, const FRequestedType& RequestedType, FEmitShaderResult& OutShader) const override;
+	virtual void EmitValueShader(FEmitContext& Context, const FRequestedType& RequestedType, FEmitShaderValues& OutResult) const override;
 	virtual void EmitValuePreshader(FEmitContext& Context, const FRequestedType& RequestedType, Shader::FPreshaderData& OutPreshader) const override;
 };
 
@@ -203,7 +176,7 @@ public:
 	FExpression* FalseExpression;
 
 	virtual void PrepareValue(FEmitContext& Context, const FRequestedType& RequestedType, FPrepareValueResult& OutResult) const override;
-	virtual void EmitValueShader(FEmitContext& Context, const FRequestedType& RequestedType, FEmitShaderResult& OutShader) const override;
+	virtual void EmitValueShader(FEmitContext& Context, const FRequestedType& RequestedType, FEmitShaderValues& OutResult) const override;
 	virtual void EmitValuePreshader(FEmitContext& Context, const FRequestedType& RequestedType, Shader::FPreshaderData& OutPreshader) const override;
 };
 
@@ -221,7 +194,7 @@ public:
 	FExpression* Rhs;
 
 	virtual void PrepareValue(FEmitContext& Context, const FRequestedType& RequestedType, FPrepareValueResult& OutResult) const override;
-	virtual void EmitValueShader(FEmitContext& Context, const FRequestedType& RequestedType, FEmitShaderResult& OutShader) const override;
+	virtual void EmitValueShader(FEmitContext& Context, const FRequestedType& RequestedType, FEmitShaderValues& OutResult) const override;
 	virtual void EmitValuePreshader(FEmitContext& Context, const FRequestedType& RequestedType, Shader::FPreshaderData& OutPreshader) const override;
 };
 
@@ -250,7 +223,7 @@ public:
 	FExpression* Input;
 
 	virtual void PrepareValue(FEmitContext& Context, const FRequestedType& RequestedType, FPrepareValueResult& OutResult) const override;
-	virtual void EmitValueShader(FEmitContext& Context, const FRequestedType& RequestedType, FEmitShaderResult& OutShader) const override;
+	virtual void EmitValueShader(FEmitContext& Context, const FRequestedType& RequestedType, FEmitShaderValues& OutResult) const override;
 	virtual void EmitValuePreshader(FEmitContext& Context, const FRequestedType& RequestedType, Shader::FPreshaderData& OutPreshader) const override;
 };
 
@@ -266,7 +239,7 @@ public:
 	FExpression* Rhs;
 
 	virtual void PrepareValue(FEmitContext& Context, const FRequestedType& RequestedType, FPrepareValueResult& OutResult) const override;
-	virtual void EmitValueShader(FEmitContext& Context, const FRequestedType& RequestedType, FEmitShaderResult& OutShader) const override;
+	virtual void EmitValueShader(FEmitContext& Context, const FRequestedType& RequestedType, FEmitShaderValues& OutResult) const override;
 	virtual void EmitValuePreshader(FEmitContext& Context, const FRequestedType& RequestedType, Shader::FPreshaderData& OutPreshader) const override;
 };
 
@@ -274,7 +247,7 @@ class FExpressionReflectionVector : public FExpression
 {
 public:
 	virtual void PrepareValue(FEmitContext& Context, const FRequestedType& RequestedType, FPrepareValueResult& OutResult) const override;
-	virtual void EmitValueShader(FEmitContext& Context, const FRequestedType& RequestedType, FEmitShaderResult& OutShader) const override;
+	virtual void EmitValueShader(FEmitContext& Context, const FRequestedType& RequestedType, FEmitShaderValues& OutResult) const override;
 };
 
 class FStatementReturn : public FStatement
diff --git a/Engine/Source/Runtime/Engine/Public/HLSLTree/HLSLTreeTypes.h b/Engine/Source/Runtime/Engine/Public/HLSLTree/HLSLTreeTypes.h
--- a/Engine/Source/Runtime/Engine/Public/HLSLTree/HLSLTreeTypes.h
+++ b/Engine/Source/Runtime/Engine/Public/HLSLTree/HLSLTreeTypes.h
@@ -9,6 +9,12 @@ class UTexture;
 
 namespace UE
 {
+
+namespace Shader
+{
+enum class EPreshaderOpcode : uint8;
+} // namespace Shader
+
 namespace HLSLTree
 {
 
@@ -56,6 +62,28 @@ enum class EExpressionDerivative : uint8
 
 EExpressionDerivative CombineDerivatives(EExpressionDerivative Lhs, EExpressionDerivative Rhs);
 
+enum class EBinaryOp : uint8
+{
+	None,
+	Add,
+	Sub,
+	Mul,
+	Div,
+	Less,
+};
+
+struct FBinaryOpDescription
+{
+	FBinaryOpDescription();
+	FBinaryOpDescription(const TCHAR* InName, const TCHAR* InOperator, Shader::EPreshaderOpcode InOpcode);
+
+	const TCHAR* Name;
+	const TCHAR* Operator;
+	Shader::EPreshaderOpcode PreshaderOpcode;
+};
+
+FBinaryOpDescription GetBinaryOpDesription(EBinaryOp Op);
+
 struct FTextureDescription
 {
 	FTextureDescription()
