WIP commit to fix previous partial commit

diff --git a/Engine/Source/Runtime/Engine/Private/HLSLTree/HLSLTreeCommon.cpp b/Engine/Source/Runtime/Engine/Private/HLSLTree/HLSLTreeCommon.cpp
--- a/Engine/Source/Runtime/Engine/Private/HLSLTree/HLSLTreeCommon.cpp
+++ b/Engine/Source/Runtime/Engine/Private/HLSLTree/HLSLTreeCommon.cpp
@@ -9,28 +9,6 @@ namespace UE
 namespace HLSLTree
 {
 
-FBinaryOpDescription::FBinaryOpDescription()
-	: Name(nullptr), Operator(nullptr), PreshaderOpcode(Shader::EPreshaderOpcode::Nop)
-{}
-
-FBinaryOpDescription::FBinaryOpDescription(const TCHAR* InName, const TCHAR* InOperator, Shader::EPreshaderOpcode InOpcode)
-	: Name(InName), Operator(InOperator), PreshaderOpcode(InOpcode)
-{}
-
-FBinaryOpDescription GetBinaryOpDesription(EBinaryOp Op)
-{
-	switch (Op)
-	{
-	case EBinaryOp::None: return FBinaryOpDescription(TEXT("None"), TEXT(""), Shader::EPreshaderOpcode::Nop); break;
-	case EBinaryOp::Add: return FBinaryOpDescription(TEXT("Add"), TEXT("+"), Shader::EPreshaderOpcode::Add); break;
-	case EBinaryOp::Sub: return FBinaryOpDescription(TEXT("Subtract"), TEXT("-"), Shader::EPreshaderOpcode::Sub); break;
-	case EBinaryOp::Mul: return FBinaryOpDescription(TEXT("Multiply"), TEXT("*"), Shader::EPreshaderOpcode::Mul); break;
-	case EBinaryOp::Div: return FBinaryOpDescription(TEXT("Divide"), TEXT("/"), Shader::EPreshaderOpcode::Div); break;
-	case EBinaryOp::Less: return FBinaryOpDescription(TEXT("Less"), TEXT("<"), Shader::EPreshaderOpcode::Nop); break;
-	default: checkNoEntry(); return FBinaryOpDescription();
-	}
-}
-
 FSwizzleParameters::FSwizzleParameters(int8 InR, int8 InG, int8 InB, int8 InA) : NumComponents(0)
 {
 	ComponentIndex[0] = InR;
@@ -160,45 +138,50 @@ void FExpressionExternalInput::PrepareValue(FEmitContext& Context, const FReques
 	OutResult.SetTypeWithDerivative(Context, RequestedType, EExpressionEvaluation::Shader, GetInputExpressionType(InputType));
 }
 
-void FExpressionExternalInput::EmitValueShader(FEmitContext& Context, const FRequestedType& RequestedType, FEmitShaderResult& OutShader) const
+void FExpressionExternalInput::EmitValueShader(FEmitContext& Context, const FRequestedType& RequestedType, FEmitShaderValues& OutResult) const
 {
-	OutShader.Type = GetInputExpressionType(InputType);
-	OutShader.bInline = true;
-
 	if (IsTexCoord(InputType))
 	{
 		const int32 TypeIndex = (int32)InputType;
 		const int32 TexCoordIndex = TypeIndex - (int32)EExternalInputType::TexCoord0;
 
 		Context.NumTexCoords = FMath::Max(Context.NumTexCoords, TexCoordIndex + 1);
-		OutShader.Code.Appendf(TEXT("Parameters.TexCoords[%u].xy"), TexCoordIndex);
-		OutShader.CodeDdx.Appendf(TEXT("Parameters.TexCoords_DDX[%u].xy"), TexCoordIndex);
-		OutShader.CodeDdy.Appendf(TEXT("Parameters.TexCoords_DDY[%u].xy"), TexCoordIndex);
+		OutResult.Code = Context.EmitInlineCode(Shader::EValueType::Float2, TEXT("Parameters.TexCoords[%].xy"), TexCoordIndex);
+		if (GetDerivative(RequestedType) == EExpressionDerivative::Valid)
+		{
+			OutResult.CodeDdx = Context.EmitInlineCode(Shader::EValueType::Float2, TEXT("Parameters.TexCoords_DDX[%].xy"), TexCoordIndex);
+			OutResult.CodeDdy = Context.EmitInlineCode(Shader::EValueType::Float2, TEXT("Parameters.TexCoords_DDY[%].xy"), TexCoordIndex);
+		}
 	}
 	else
 	{
+		const Shader::EValueType Type = GetInputExpressionType(InputType);
+
 		// TODO - handle PrevFrame position
 		switch (InputType)
 		{
 		case EExternalInputType::WorldPosition:
-			OutShader.Code.Append(TEXT("GetWorldPosition(Parameters)"));
+			OutResult.Code = Context.EmitInlineCode(Type, TEXT("GetWorldPosition(Parameters)"));
 			break;
 		case EExternalInputType::WorldPosition_NoOffsets:
-			OutShader.Code.Append(TEXT("GetWorldPosition_NoMaterialOffsets(Parameters)"));
+			OutResult.Code = Context.EmitInlineCode(Type, TEXT("GetWorldPosition_NoMaterialOffsets(Parameters)"));
 			break;
 		case EExternalInputType::TranslatedWorldPosition:
-			OutShader.Code.Append(TEXT("GetTranslatedWorldPosition(Parameters)"));
+			OutResult.Code = Context.EmitInlineCode(Type, TEXT("GetTranslatedWorldPosition(Parameters)"));
 			break;
 		case EExternalInputType::TranslatedWorldPosition_NoOffsets:
-			OutShader.Code.Append(TEXT("GetTranslatedWorldPosition_NoMaterialOffsets(Parameters)"));
+			OutResult.Code = Context.EmitInlineCode(Type, TEXT("GetTranslatedWorldPosition_NoMaterialOffsets(Parameters)"));
 			break;
 		default:
 			checkNoEntry();
 			break;
 		}
 
-		OutShader.CodeDdx.Append(TEXT("Parameters.WorldPosition_DDX"));
-		OutShader.CodeDdy.Append(TEXT("Parameters.WorldPosition_DDY"));
+		if (GetDerivative(RequestedType) == EExpressionDerivative::Valid)
+		{
+			OutResult.CodeDdx = Context.EmitInlineCode(Shader::EValueType::Float3, TEXT("Parameters.WorldPosition_DDX"));
+			OutResult.CodeDdy = Context.EmitInlineCode(Shader::EValueType::Float3, TEXT("Parameters.WorldPosition_DDY"));
+		}
 	}
 }
 
@@ -209,7 +192,7 @@ void FExpressionTextureSample::PrepareValue(FEmitContext& Context, const FReques
 	OutResult.SetType(Context, RequestedType, EExpressionEvaluation::Shader, Shader::EValueType::Float4);
 }
 
-void FExpressionTextureSample::EmitValueShader(FEmitContext& Context, const FRequestedType& RequestedType, FEmitShaderResult& OutShader) const
+void FExpressionTextureSample::EmitValueShader(FEmitContext& Context, const FRequestedType& RequestedType, FEmitShaderValues& OutResult) const
 {
 	const FTextureDescription& Desc = Declaration->Description;
 	const EMaterialValueType MaterialValueType = Desc.Texture->GetMaterialType();
@@ -336,9 +319,8 @@ void FExpressionTextureSample::EmitValueShader(FEmitContext& Context, const FReq
 	}
 
 	const FRequestedType TexCoordRequestedType(Shader::EValueType::Float2, EComponentRequest::RequestedWithDerivative);
-	const FShaderValue TexCoordValue = TexCoordExpression->GetValueShader(Context, TexCoordRequestedType);
-	OutShader.Type = Shader::EValueType::Float4;
-	OutShader.Code.Appendf(TEXT("%s(%sGrad(%s, %s, %s, %s, %s))"),
+	const FEmitShaderValues TexCoordValue = TexCoordExpression->GetValueShader(Context, TexCoordRequestedType);
+	OutResult.Code = Context.EmitCode(Shader::EValueType::Float4, TEXT("%(%Grad(%, %, %, %, %))"),
 		SamplerTypeFunction,
 		SampleFunctionName,
 		*TextureName,
@@ -362,24 +344,24 @@ void FExpressionGetStructField::PrepareValue(FEmitContext& Context, const FReque
 	OutResult.SetType(Context, RequestedType, StructPreparedType.GetFieldType(Field));
 }
 
-void FExpressionGetStructField::EmitValueShader(FEmitContext& Context, const FRequestedType& RequestedType, FEmitShaderResult& OutShader) const
+void FExpressionGetStructField::EmitValueShader(FEmitContext& Context, const FRequestedType& RequestedType, FEmitShaderValues& OutResult) const
 {
 	FRequestedType RequestedStructType(StructType, EComponentRequest::None);
 	RequestedStructType.SetField(Field, RequestedType);
 
-	const FShaderValue StructValue = StructExpression->GetValueShader(Context, RequestedStructType);
+	const FEmitShaderValues StructValue = StructExpression->GetValueShader(Context, RequestedStructType);
 
-	OutShader.bInline = true;
-	OutShader.Type = Field->Type;
-	OutShader.Code.Appendf(TEXT("%s.%s"),
+	OutResult.Code = Context.EmitInlineCode(Field->Type, TEXT("%.%"),
 		StructValue.Code,
 		Field->Name);
-	if (StructValue.bHasDerivatives)
+
+	if (StructValue.HasDerivatives() && GetDerivative(RequestedType) == EExpressionDerivative::Valid)
 	{
-		OutShader.CodeDdx.Appendf(TEXT("%s.%s"),
+		const Shader::FType DerivativeType = Field->Type.GetDerivativeType();
+		OutResult.CodeDdx = Context.EmitInlineCode(DerivativeType, TEXT("%.%"),
 			StructValue.CodeDdx,
 			Field->Name);
-		OutShader.CodeDdy.Appendf(TEXT("%s.%s"),
+		OutResult.CodeDdy = Context.EmitInlineCode(DerivativeType, TEXT("%.%"),
 			StructValue.CodeDdy,
 			Field->Name);
 	}
@@ -422,7 +404,7 @@ void FExpressionSetStructField::PrepareValue(FEmitContext& Context, const FReque
 	}
 }
 
-void FExpressionSetStructField::EmitValueShader(FEmitContext& Context, const FRequestedType& RequestedType, FEmitShaderResult& OutShader) const
+void FExpressionSetStructField::EmitValueShader(FEmitContext& Context, const FRequestedType& RequestedType, FEmitShaderValues& OutResult) const
 {
 	FRequestedType RequestedStructType(RequestedType);
 	RequestedStructType.ClearFieldRequested(Field);
@@ -431,30 +413,30 @@ void FExpressionSetStructField::EmitValueShader(FEmitContext& Context, const FRe
 	const FRequestedType RequestedFieldType = RequestedType.GetField(Field);
 	const EExpressionEvaluation FieldEvaluation = FieldExpression->GetEvaluation(RequestedFieldType);
 
-	const FShaderValue FieldValue = FieldExpression->GetValueShader(Context, RequestedFieldType, Field->Type);
+	const FEmitShaderValues FieldValue = FieldExpression->GetValueShader(Context, RequestedFieldType, Field->Type);
 
 	const Shader::FStructType* DerivativeStructType = StructType->DerivativeType;
+	const EExpressionDerivative Derivative = GetDerivative(RequestedType);
 
-	OutShader.Type = StructType;
 	if (StructEvaluation == EExpressionEvaluation::None)
 	{
 		check(FieldEvaluation != EExpressionEvaluation::None);
 
 		// StructExpression is not used, so default to a zero-initialized struct
 		// This will happen if all the accessed struct fields are explicitly defined
-		OutShader.Code.Appendf(TEXT("%s_Set%s((%s)0, %s)"),
+		OutResult.Code = Context.EmitCode(StructType, TEXT("%_Set%((%)0, %)"),
 			StructType->Name,
 			Field->Name,
 			StructType->Name,
 			FieldValue.Code);
-		if (DerivativeStructType && FieldValue.bHasDerivatives)
+		if (Derivative == EExpressionDerivative::Valid && DerivativeStructType && FieldValue.HasDerivatives())
 		{
-			OutShader.CodeDdx.Appendf(TEXT("%s_Set%s((%s)0, %s)"),
+			OutResult.CodeDdx = Context.EmitCode(DerivativeStructType, TEXT("%_Set%((%)0, %)"),
 				DerivativeStructType->Name,
 				Field->Name,
 				DerivativeStructType->Name,
 				FieldValue.CodeDdx);
-			OutShader.CodeDdy.Appendf(TEXT("%s_Set%s((%s)0, %s)"),
+			OutResult.CodeDdy = Context.EmitCode(DerivativeStructType, TEXT("%_Set%((%)0, %)"),
 				DerivativeStructType->Name,
 				Field->Name,
 				DerivativeStructType->Name,
@@ -463,20 +445,20 @@ void FExpressionSetStructField::EmitValueShader(FEmitContext& Context, const FRe
 	}
 	else
 	{
-		const FShaderValue StructValue = StructExpression->GetValueShader(Context, RequestedStructType);
-		OutShader.Code.Appendf(TEXT("%s_Set%s(%s, %s)"),
+		const FEmitShaderValues StructValue = StructExpression->GetValueShader(Context, RequestedStructType);
+		OutResult.Code = Context.EmitCode(StructType, TEXT("%_Set%(%, %)"),
 			StructType->Name,
 			Field->Name,
 			StructValue.Code,
 			FieldValue.Code);
-		if (DerivativeStructType && StructValue.bHasDerivatives && FieldValue.bHasDerivatives)
+		if (Derivative == EExpressionDerivative::Valid && DerivativeStructType && StructValue.HasDerivatives() && FieldValue.HasDerivatives())
 		{
-			OutShader.CodeDdx.Appendf(TEXT("%s_Set%s(%s, %s)"),
+			OutResult.CodeDdx = Context.EmitCode(DerivativeStructType, TEXT("%_Set%(%, %)"),
 				DerivativeStructType->Name,
 				Field->Name,
 				StructValue.CodeDdx,
 				FieldValue.CodeDdx);
-			OutShader.CodeDdy.Appendf(TEXT("%s_Set%s(%s, %s)"),
+			OutResult.CodeDdy = Context.EmitCode(DerivativeStructType, TEXT("%_Set%(%, %)"),
 				DerivativeStructType->Name,
 				Field->Name,
 				StructValue.CodeDdy,
@@ -528,13 +510,12 @@ void FExpressionSelect::PrepareValue(FEmitContext& Context, const FRequestedType
 	OutResult.SetType(Context, RequestedType, MergePreparedTypes(LhsType, RhsType));
 }
 
-void FExpressionSelect::EmitValueShader(FEmitContext& Context, const FRequestedType& RequestedType, FEmitShaderResult& OutShader) const
+void FExpressionSelect::EmitValueShader(FEmitContext& Context, const FRequestedType& RequestedType, FEmitShaderValues& OutResult) const
 {
-	const FShaderValue TrueValue = TrueExpression->GetValueShader(Context, RequestedType);
-	const FShaderValue FalseValue = FalseExpression->GetValueShader(Context, RequestedType);
+	const FEmitShaderValues TrueValue = TrueExpression->GetValueShader(Context, RequestedType);
+	const FEmitShaderValues FalseValue = FalseExpression->GetValueShader(Context, RequestedType);
 
-	OutShader.Type = RequestedType.GetType();
-	OutShader.Code.Appendf(TEXT("(%s ? %s : %s)"),
+	OutResult.Code = Context.EmitCode(RequestedType.GetType(), TEXT("(% ? % : %)"),
 		ConditionExpression->GetValueShader(Context, Shader::EValueType::Bool1).Code,
 		TrueValue.Code,
 		FalseValue.Code);
@@ -545,6 +526,46 @@ void FExpressionSelect::EmitValuePreshader(FEmitContext& Context, const FRequest
 	check(false); // TODO
 }
 
+void FExpressionUnaryOp::PrepareValue(FEmitContext& Context, const FRequestedType& RequestedType, FPrepareValueResult& OutResult) const
+{
+	const FPreparedType& InputType = PrepareExpressionValue(Context, Input, RequestedType);
+	if (InputType.IsVoid())
+	{
+		return;
+	}
+
+	if (!InputType.IsNumeric())
+	{
+		return Context.Errors.AddError(this, TEXT("Invalid operation on non-numeric type"));
+	}
+
+	const FUnaryOpDescription OpDesc = GetUnaryOpDesription(Op);
+	FPreparedType ResultType = InputType;
+	if (OpDesc.PreshaderOpcode == Shader::EPreshaderOpcode::Nop)
+	{
+		// No preshader support
+		ResultType.SetEvaluation(EExpressionEvaluation::Shader);
+	}
+
+	OutResult.SetType(Context, RequestedType, ResultType);
+}
+
+void FExpressionUnaryOp::EmitValueShader(FEmitContext& Context, const FRequestedType& RequestedType, FEmitShaderValues& OutResult) const
+{
+	OutResult = Context.EmitUnaryOp(Op,
+		Input->GetValueShader(Context, RequestedType),
+		GetDerivative(RequestedType));
+}
+
+void FExpressionUnaryOp::EmitValuePreshader(FEmitContext& Context, const FRequestedType& RequestedType, Shader::FPreshaderData& OutPreshader) const
+{
+	const FUnaryOpDescription OpDesc = GetUnaryOpDesription(Op);
+	check(OpDesc.PreshaderOpcode != Shader::EPreshaderOpcode::Nop);
+
+	Input->GetValuePreshader(Context, RequestedType, OutPreshader);
+	OutPreshader.WriteOpcode(OpDesc.PreshaderOpcode);
+}
+
 void FExpressionBinaryOp::PrepareValue(FEmitContext& Context, const FRequestedType& RequestedType, FPrepareValueResult& OutResult) const
 {
 	const FPreparedType& LhsType = PrepareExpressionValue(Context, Lhs, RequestedType);
@@ -577,122 +598,31 @@ void FExpressionBinaryOp::PrepareValue(FEmitContext& Context, const FRequestedTy
 
 namespace Private
 {
-struct FBinaryOpTypes
-{
-	Shader::FType ResultType;
-	FRequestedType LhsType;
-	FRequestedType RhsType;
-	bool bIsLWC;
-};
-FBinaryOpTypes GetBinaryOpTypes(const FRequestedType& RequestedType, EBinaryOp Op, const FRequestedType& LhsType, const FRequestedType& RhsType)
+FRequestedType GetBinaryOpInputType(const FRequestedType& RequestedType, EBinaryOp Op, const FRequestedType& LhsType, const FRequestedType& RhsType)
 {
 	const int32 NumRequestedComponents = RequestedType.GetNumComponents();
 	const Shader::EValueComponentType InputComponentType = Shader::CombineComponentTypes(LhsType.ValueComponentType, RhsType.ValueComponentType);
-	const FRequestedType InputType = MakeRequestedType(InputComponentType, RequestedType);
-
-	FBinaryOpTypes Result;
-	Result.bIsLWC = (InputComponentType == Shader::EValueComponentType::Double);
-	switch (Op)
-	{
-	case EBinaryOp::Less:
-		Result.LhsType = InputType;
-		Result.RhsType = InputType;
-		Result.ResultType = Shader::MakeValueType(Shader::EValueComponentType::Bool, NumRequestedComponents);
-		break;
-	default:
-		Result.LhsType = InputType;
-		Result.RhsType = InputType;
-		Result.ResultType = Shader::MakeValueType(InputComponentType, NumRequestedComponents);
-		break;
-	}
-	return Result;
+	return MakeRequestedType(InputComponentType, RequestedType);
 }
 } // namespace Private
 
-void FExpressionBinaryOp::EmitValueShader(FEmitContext& Context, const FRequestedType& RequestedType, FEmitShaderResult& OutShader) const
+void FExpressionBinaryOp::EmitValueShader(FEmitContext& Context, const FRequestedType& RequestedType, FEmitShaderValues& OutResult) const
 {
-	const Private::FBinaryOpTypes Types = Private::GetBinaryOpTypes(RequestedType, Op, Lhs->GetRequestedType(), Rhs->GetRequestedType());
-	const FShaderValue LhsValue = Lhs->GetValueShader(Context, Types.LhsType);
-	const FShaderValue RhsValue = Rhs->GetValueShader(Context, Types.RhsType);
-	const bool bHasDerivatives = LhsValue.bHasDerivatives && RhsValue.bHasDerivatives;
-
-	OutShader.Type = Types.ResultType;
-	switch (Op)
-	{
-	case EBinaryOp::Add:
-		if (Types.bIsLWC)
-		{
-			OutShader.Code.Appendf(TEXT("LWCAdd(%s, %s)"), LhsValue.Code, RhsValue.Code);
-		}
-		else
-		{
-			OutShader.Code.Appendf(TEXT("(%s + %s)"), LhsValue.Code, RhsValue.Code);
-		}
-		if (bHasDerivatives)
-		{
-			OutShader.CodeDdx.Appendf(TEXT("(%s + %s)"), LhsValue.CodeDdx, RhsValue.CodeDdx);
-			OutShader.CodeDdy.Appendf(TEXT("(%s + %s)"), LhsValue.CodeDdy, RhsValue.CodeDdy);
-		}
-		break;
-	case EBinaryOp::Sub:
-		if (Types.bIsLWC)
-		{
-			OutShader.Code.Appendf(TEXT("LWCSub(%s, %s)"), LhsValue.Code, RhsValue.Code);
-		}
-		else
-		{
-			OutShader.Code.Appendf(TEXT("(%s - %s)"), LhsValue.Code, RhsValue.Code);
-		}
-		break;
-	case EBinaryOp::Mul:
-		if (Types.bIsLWC)
-		{
-			OutShader.Code.Appendf(TEXT("LWCMultiply(%s, %s)"), LhsValue.Code, RhsValue.Code);
-		}
-		else
-		{
-			OutShader.Code.Appendf(TEXT("(%s * %s)"), LhsValue.Code, RhsValue.Code);
-		}
-		if (bHasDerivatives)
-		{
-			OutShader.CodeDdx.Appendf(TEXT("(%s * %s + %s * %s)"), LhsValue.CodeDdx, RhsValue.Code, RhsValue.CodeDdx, LhsValue.Code);
-			OutShader.CodeDdy.Appendf(TEXT("(%s * %s + %s * %s)"), LhsValue.CodeDdy, RhsValue.Code, RhsValue.CodeDdy, LhsValue.Code);
-		}
-		break;
-	case EBinaryOp::Div:
-		if (Types.bIsLWC)
-		{
-			OutShader.Code.Appendf(TEXT("LWCDivide(%s, %s)"), LhsValue.Code, RhsValue.Code);
-		}
-		else
-		{
-			OutShader.Code.Appendf(TEXT("(%s / %s)"), LhsValue.Code, RhsValue.Code);
-		}
-		break;
-	case EBinaryOp::Less:
-		if (Types.bIsLWC)
-		{
-			OutShader.Code.Appendf(TEXT("LWCLess(%s, %s)"), LhsValue.Code, RhsValue.Code);
-		}
-		else
-		{
-			OutShader.Code.Appendf(TEXT("(%s < %s)"), LhsValue.Code, RhsValue.Code);
-		}
-		break;
-	default:
-		checkNoEntry();
-		break;
-	}
+	const FRequestedType InputType = Private::GetBinaryOpInputType(RequestedType, Op, Lhs->GetRequestedType(), Rhs->GetRequestedType());
+	OutResult = Context.EmitBinaryOp(Op,
+		Lhs->GetValueShader(Context, InputType),
+		Rhs->GetValueShader(Context, InputType),
+		GetDerivative(RequestedType));
 }
 
 void FExpressionBinaryOp::EmitValuePreshader(FEmitContext& Context, const FRequestedType& RequestedType, Shader::FPreshaderData& OutPreshader) const
 {
-	const Private::FBinaryOpTypes Types = Private::GetBinaryOpTypes(RequestedType, Op, Lhs->GetRequestedType(), Rhs->GetRequestedType());
+	const FRequestedType InputType = Private::GetBinaryOpInputType(RequestedType, Op, Lhs->GetRequestedType(), Rhs->GetRequestedType());
 	const FBinaryOpDescription OpDesc = GetBinaryOpDesription(Op);
 	check(OpDesc.PreshaderOpcode != Shader::EPreshaderOpcode::Nop);
 
-	Lhs->GetValuePreshader(Context, Types.LhsType, OutPreshader);
-	Rhs->GetValuePreshader(Context, Types.RhsType, OutPreshader);
+	Lhs->GetValuePreshader(Context, InputType, OutPreshader);
+	Rhs->GetValuePreshader(Context, InputType, OutPreshader);
 	OutPreshader.WriteOpcode(OpDesc.PreshaderOpcode);
 }
 
@@ -711,7 +641,7 @@ void FExpressionSwizzle::PrepareValue(FEmitContext& Context, const FRequestedTyp
 	OutResult.SetType(Context, RequestedType, ResultType);
 }
 
-void FExpressionSwizzle::EmitValueShader(FEmitContext& Context, const FRequestedType& RequestedType, FEmitShaderResult& OutShader) const
+void FExpressionSwizzle::EmitValueShader(FEmitContext& Context, const FRequestedType& RequestedType, FEmitShaderValues& OutResult) const
 {
 	const FRequestedType RequestedInputType = Parameters.GetRequestedInputType(RequestedType);
 
@@ -729,19 +659,16 @@ void FExpressionSwizzle::EmitValueShader(FEmitContext& Context, const FRequested
 		}
 	}
 
-	const FShaderValue InputValue = Input->GetValueShader(Context, RequestedInputType);
-
-	OutShader.Type = RequestedType.GetType();
-	OutShader.bInline = true;
+	const FEmitShaderValues InputValue = Input->GetValueShader(Context, RequestedInputType);
 	if (bHasSwizzleReorder)
 	{
-		OutShader.Code.Appendf(TEXT("%s.%s"),
+		OutResult.Code = Context.EmitInlineCode(RequestedType.GetType(), TEXT("%.%"),
 			InputValue.Code,
 			Swizzle);
 	}
 	else
 	{
-		OutShader.Code.Appendf(TEXT("%s"), InputValue.Code);
+		OutResult.Code = InputValue.Code;
 	}
 }
 
@@ -821,21 +748,19 @@ FAppendTypes GetAppendTypes(const FRequestedType& RequestedType, const FRequeste
 }
 } // namespace Private
 
-void FExpressionAppend::EmitValueShader(FEmitContext& Context, const FRequestedType& RequestedType, FEmitShaderResult& OutShader) const
+void FExpressionAppend::EmitValueShader(FEmitContext& Context, const FRequestedType& RequestedType, FEmitShaderValues& OutResult) const
 {
 	const Private::FAppendTypes Types = Private::GetAppendTypes(RequestedType, Lhs->GetRequestedType(), Rhs->GetRequestedType());
-	const FShaderValue LhsValue = Lhs->GetValueShader(Context, Types.LhsType);
+	const FEmitShaderValues LhsValue = Lhs->GetValueShader(Context, Types.LhsType);
 
-	OutShader.bInline = true;
-	OutShader.Type = Types.ResultType;
 	if (Types.RhsType.IsVoid())
 	{
-		OutShader.Code.Appendf(TEXT("%s"), LhsValue.Code);
+		OutResult.Code = LhsValue.Code;
 	}
 	else
 	{
-		const FShaderValue RhsValue = Rhs->GetValueShader(Context, Types.RhsType);
-		OutShader.Code.Appendf(TEXT("%s(%s, %s)"),
+		const FEmitShaderValues RhsValue = Rhs->GetValueShader(Context, Types.RhsType);
+		OutResult.Code = Context.EmitInlineCode(Types.ResultType, TEXT("%(%, %)"),
 			RequestedType.GetName(),
 			LhsValue.Code,
 			RhsValue.Code);
@@ -859,11 +784,9 @@ void FExpressionReflectionVector::PrepareValue(FEmitContext& Context, const FReq
 	OutResult.SetType(Context, RequestedType, EExpressionEvaluation::Shader, Shader::EValueType::Float3);
 }
 
-void FExpressionReflectionVector::EmitValueShader(FEmitContext& Context, const FRequestedType& RequestedType, FEmitShaderResult& OutShader) const
+void FExpressionReflectionVector::EmitValueShader(FEmitContext& Context, const FRequestedType& RequestedType, FEmitShaderValues& OutResult) const
 {
-	OutShader.bInline = true;
-	OutShader.Type = Shader::EValueType::Float3;
-	OutShader.Code.Append(TEXT("Parameters.ReflectionVector"));
+	OutResult.Code = Context.EmitInlineCode(Shader::EValueType::Float3, TEXT("Parameters.ReflectionVector"));
 }
 
 void FStatementBreak::Prepare(FEmitContext& Context) const
@@ -881,7 +804,7 @@ void FStatementReturn::Prepare(FEmitContext& Context) const
 
 void FStatementReturn::EmitShader(FEmitContext& Context) const
 {
-	ParentScope->EmitStatementf(Context, TEXT("return %s;"), Expression->GetValueShader(Context).Code);
+	ParentScope->EmitStatementf(Context, TEXT("return %s;"), Expression->GetValueShader(Context).Code->Reference);
 }
 
 void FStatementIf::Prepare(FEmitContext& Context) const
@@ -926,8 +849,8 @@ void FStatementIf::EmitShader(FEmitContext& Context) const
 	{
 		if (IsScopeLive(ThenScope) || IsScopeLive(ElseScope))
 		{
-			const FShaderValue ConditionValue = ConditionExpression->GetValueShader(Context, Shader::EValueType::Bool1);
-			ParentScope->EmitNestedScopef(Context, ThenScope, TEXT("if (%s)"), ConditionValue.Code);
+			const FEmitShaderValues ConditionValue = ConditionExpression->GetValueShader(Context, Shader::EValueType::Bool1);
+			ParentScope->EmitNestedScopef(Context, ThenScope, TEXT("if (%s)"), ConditionValue.Code->Reference);
 			if (IsScopeLive(ElseScope))
 			{
 				ParentScope->EmitNestedScopef(Context, ElseScope, TEXT("else"));
diff --git a/Engine/Source/Runtime/Engine/Private/HLSLTree/HLSLTreeIntrinsics.cpp b/Engine/Source/Runtime/Engine/Private/HLSLTree/HLSLTreeIntrinsics.cpp
--- a/Engine/Source/Runtime/Engine/Private/HLSLTree/HLSLTreeIntrinsics.cpp
+++ b/Engine/Source/Runtime/Engine/Private/HLSLTree/HLSLTreeIntrinsics.cpp
@@ -12,6 +12,14 @@ namespace UE
 namespace HLSLTree
 {
 
+FUnaryOpDescription::FUnaryOpDescription()
+	: Name(nullptr), Operator(nullptr), PreshaderOpcode(Shader::EPreshaderOpcode::Nop)
+{}
+
+FUnaryOpDescription::FUnaryOpDescription(const TCHAR* InName, const TCHAR* InOperator, Shader::EPreshaderOpcode InOpcode)
+	: Name(InName), Operator(InOperator), PreshaderOpcode(InOpcode)
+{}
+
 FBinaryOpDescription::FBinaryOpDescription()
 	: Name(nullptr), Operator(nullptr), PreshaderOpcode(Shader::EPreshaderOpcode::Nop)
 {}
@@ -20,6 +28,16 @@ FBinaryOpDescription::FBinaryOpDescription(const TCHAR* InName, const TCHAR* InO
 	: Name(InName), Operator(InOperator), PreshaderOpcode(InOpcode)
 {}
 
+FUnaryOpDescription GetUnaryOpDesription(EUnaryOp Op)
+{
+	switch (Op)
+	{
+	case EUnaryOp::None: return FUnaryOpDescription(TEXT("None"), TEXT(""), Shader::EPreshaderOpcode::Nop); break;
+	case EUnaryOp::Rcp: return FUnaryOpDescription(TEXT("Rcp"), TEXT("/"), Shader::EPreshaderOpcode::Rcp); break;
+	default: checkNoEntry(); return FUnaryOpDescription();
+	}
+}
+
 FBinaryOpDescription GetBinaryOpDesription(EBinaryOp Op)
 {
 	switch (Op)
@@ -176,6 +194,45 @@ FEmitShaderValues FEmitContext::EmitCast(FEmitShaderValues ShaderValue, const Sh
 	return Result;
 }
 
+FEmitShaderCode* FEmitContext::EmitUnaryOp(EUnaryOp Op, FEmitShaderCode* Input)
+{
+	const Shader::FValueTypeDescription InputTypeDesc = Shader::GetValueTypeDescription(Input->Type);
+	const Shader::EValueComponentType InputComponentType = InputTypeDesc.ComponentType;
+	const int8 NumComponents = InputTypeDesc.NumComponents;
+	Shader::EValueType ResultType = Input->Type;
+
+	FEmitShaderCode* Result = nullptr;
+	switch (Op)
+	{
+	case EUnaryOp::Rcp:
+		if (InputComponentType == Shader::EValueComponentType::Double)
+		{
+			ResultType = Shader::MakeValueType(Shader::EValueComponentType::Float, NumComponents);
+			Result = EmitCode(ResultType, TEXT("LWCRcp(%)"), Input);
+		}
+		else
+		{
+			Result = EmitCode(ResultType, TEXT("rcp(%)"), Input);
+		}
+		break;
+	default:
+		checkNoEntry();
+	}
+	return Result;
+}
+
+FEmitShaderValues FEmitContext::EmitUnaryOp(EUnaryOp Op, FEmitShaderValues Input, EExpressionDerivative Derivative)
+{
+	FEmitShaderValues Result;
+	Result.Code = EmitUnaryOp(Op, Input.Code);
+	if (Derivative == EExpressionDerivative::Valid && Input.HasDerivatives())
+	{
+		
+	}
+
+	return Result;
+}
+
 FEmitShaderCode* FEmitContext::EmitBinaryOp(EBinaryOp Op, FEmitShaderCode* Lhs, FEmitShaderCode* Rhs)
 {
 	const Shader::FValueTypeDescription LhsTypeDesc = Shader::GetValueTypeDescription(Lhs->Type);
diff --git a/Engine/Source/Runtime/Engine/Public/HLSLTree/HLSLTree.h b/Engine/Source/Runtime/Engine/Public/HLSLTree/HLSLTree.h
--- a/Engine/Source/Runtime/Engine/Public/HLSLTree/HLSLTree.h
+++ b/Engine/Source/Runtime/Engine/Public/HLSLTree/HLSLTree.h
@@ -180,6 +180,9 @@ public:
 	FEmitShaderCode* EmitCast(FEmitShaderCode* ShaderValue, const Shader::FType& DestType);
 	FEmitShaderValues EmitCast(FEmitShaderValues ShaderValue, const Shader::FType& DestType);
 
+	FEmitShaderCode* EmitUnaryOp(EUnaryOp Op, FEmitShaderCode* Input);
+	FEmitShaderValues EmitUnaryOp(EUnaryOp Op, FEmitShaderValues Input, EExpressionDerivative Derivative);
+
 	FEmitShaderCode* EmitBinaryOp(EBinaryOp Op, FEmitShaderCode* Lhs, FEmitShaderCode* Rhs);
 	FEmitShaderValues EmitBinaryOp(EBinaryOp Op, FEmitShaderValues Lhs, FEmitShaderValues Rhs, EExpressionDerivative Derivative);
 
diff --git a/Engine/Source/Runtime/Engine/Public/HLSLTree/HLSLTreeCommon.h b/Engine/Source/Runtime/Engine/Public/HLSLTree/HLSLTreeCommon.h
--- a/Engine/Source/Runtime/Engine/Public/HLSLTree/HLSLTreeCommon.h
+++ b/Engine/Source/Runtime/Engine/Public/HLSLTree/HLSLTreeCommon.h
@@ -180,6 +180,22 @@ public:
 	virtual void EmitValuePreshader(FEmitContext& Context, const FRequestedType& RequestedType, Shader::FPreshaderData& OutPreshader) const override;
 };
 
+class FExpressionUnaryOp : public FExpression
+{
+public:
+	FExpressionUnaryOp(EUnaryOp InOp, FExpression* InInput)
+		: Op(InOp)
+		, Input(InInput)
+	{}
+
+	EUnaryOp Op;
+	FExpression* Input;
+
+	virtual void PrepareValue(FEmitContext& Context, const FRequestedType& RequestedType, FPrepareValueResult& OutResult) const override;
+	virtual void EmitValueShader(FEmitContext& Context, const FRequestedType& RequestedType, FEmitShaderValues& OutResult) const override;
+	virtual void EmitValuePreshader(FEmitContext& Context, const FRequestedType& RequestedType, Shader::FPreshaderData& OutPreshader) const override;
+};
+
 class FExpressionBinaryOp : public FExpression
 {
 public:
diff --git a/Engine/Source/Runtime/Engine/Public/HLSLTree/HLSLTreeTypes.h b/Engine/Source/Runtime/Engine/Public/HLSLTree/HLSLTreeTypes.h
--- a/Engine/Source/Runtime/Engine/Public/HLSLTree/HLSLTreeTypes.h
+++ b/Engine/Source/Runtime/Engine/Public/HLSLTree/HLSLTreeTypes.h
@@ -62,6 +62,24 @@ enum class EExpressionDerivative : uint8
 
 EExpressionDerivative CombineDerivatives(EExpressionDerivative Lhs, EExpressionDerivative Rhs);
 
+enum class EUnaryOp : uint8
+{
+	None,
+	Rcp,
+};
+
+struct FUnaryOpDescription
+{
+	FUnaryOpDescription();
+	FUnaryOpDescription(const TCHAR* InName, const TCHAR* InOperator, Shader::EPreshaderOpcode InOpcode);
+
+	const TCHAR* Name;
+	const TCHAR* Operator;
+	Shader::EPreshaderOpcode PreshaderOpcode;
+};
+
+FUnaryOpDescription GetUnaryOpDesription(EUnaryOp Op);
+
 enum class EBinaryOp : uint8
 {
 	None,
