Fixing split screen for spatialization plugins

diff --git a/Engine/Source/Runtime/Engine/Private/AudioDevice.cpp b/Engine/Source/Runtime/Engine/Private/AudioDevice.cpp
--- a/Engine/Source/Runtime/Engine/Private/AudioDevice.cpp
+++ b/Engine/Source/Runtime/Engine/Private/AudioDevice.cpp
@@ -570,7 +570,10 @@ bool FAudioDevice::Init(Audio::FDeviceId InDeviceID, int32 InMaxSources, int32 I
 	// Make sure the Listeners array has at least one entry, so we don't have to check for Listeners.Num() == 0 all the time
 	Listeners.Add(FListener(this));
 	ListenerProxies.AddDefaulted();
-	InverseListenerTransform.SetIdentity();
+	FMatrix Transform;
+	Transform.SetIdentity();
+	InverseListenerTransforms.Reset();
+	InverseListenerTransforms.Add(Transform);
 
 	if (!bDeferStartupPrecache)
 	{
@@ -4471,15 +4474,23 @@ void FAudioDevice::Update(bool bGameTicking)
 
 	ProcessingPendingActiveSoundStops();
 
-	// Update listener transform
-	if (Listeners.Num() > 0)
+	// Update listener transforms
+	if (Listeners.Num() != InverseListenerTransforms.Num())
+	{
+		InverseListenerTransforms.SetNum(Listeners.Num());
+	}
+
+	for (int32 ListenerIndex = 0; ListenerIndex < Listeners.Num(); ++ListenerIndex)
 	{
 		// Caches the matrix used to transform a sounds position into local space so we can just look
 		// at the Y component after normalization to determine spatialization.
-		const FVector Up = Listeners[0].GetUp();
-		const FVector Right = Listeners[0].GetFront();
-		InverseListenerTransform = FMatrix(Up, Right, Up ^ Right, Listeners[0].Transform.GetTranslation()).Inverse();
-		ensure(!InverseListenerTransform.ContainsNaN());
+		const FListener& Listener = Listeners[ListenerIndex];
+		FMatrix& InverseTransform = InverseListenerTransforms[ListenerIndex];
+
+		const FVector Up = Listener.GetUp();
+		const FVector Right = Listener.GetFront();
+		InverseTransform = FMatrix(Up, Right, Up ^ Right, Listener.Transform.GetTranslation()).Inverse();
+		ensure(!InverseTransform.ContainsNaN());
 	}
 
 	int32 FirstActiveIndex = INDEX_NONE;
@@ -6725,8 +6736,18 @@ void FAudioDevice::SetDeviceMuted(const bool bMuted)
 
 FVector FAudioDevice::GetListenerTransformedDirection(const FVector& Position, float* OutDistance)
 {
+	float DistSquared;
+	const bool bAllowAttenuationOverrides = true;
+
+	int32 ClosestListenerIndex = FindClosestListenerIndex(Position, DistSquared, bAllowAttenuationOverrides);
+	if (ClosestListenerIndex == INDEX_NONE || ClosestListenerIndex >= InverseListenerTransforms.Num())
+	{
+		UE_LOG(LogAudioMixer, Display, TEXT("Invalid listener index (%d) when trying to find listener-transformed direction"), ClosestListenerIndex);
+		return FVector::ForwardVector;
+	}
+
 	check(IsInAudioThread());
-	FVector UnnormalizedDirection = InverseListenerTransform.TransformPosition(Position);
+	FVector UnnormalizedDirection = InverseListenerTransforms[ClosestListenerIndex].TransformPosition(Position);
 	if (OutDistance)
 	{
 		*OutDistance = UnnormalizedDirection.Size();
diff --git a/Engine/Source/Runtime/Engine/Public/AudioDevice.h b/Engine/Source/Runtime/Engine/Public/AudioDevice.h
--- a/Engine/Source/Runtime/Engine/Public/AudioDevice.h
+++ b/Engine/Source/Runtime/Engine/Public/AudioDevice.h
@@ -2195,8 +2195,8 @@ private:
 	friend class FSoundConcurrencyManager;
 	FSoundConcurrencyManager ConcurrencyManager;
 
-	/** Inverse listener transformation, used for spatialization */
-	FMatrix InverseListenerTransform;
+	/** Inverse listener transformations, used for spatialization */
+	TArray<FMatrix> InverseListenerTransforms;
 
 	/** A count of the number of one-shot active sounds. */
 	uint32 OneShotCount;
