Expand support for analytic derivatives in new HLSL generator

diff --git a/Engine/Source/Runtime/Engine/Private/HLSLTree/HLSLTree.cpp b/Engine/Source/Runtime/Engine/Private/HLSLTree/HLSLTree.cpp
--- a/Engine/Source/Runtime/Engine/Private/HLSLTree/HLSLTree.cpp
+++ b/Engine/Source/Runtime/Engine/Private/HLSLTree/HLSLTree.cpp
@@ -1120,6 +1120,18 @@ void FPreparedType::SetEvaluation(EExpressionEvaluation Evaluation)
 	}
 }
 
+void FPreparedType::SetDerivative(EExpressionDerivative Derivative)
+{
+	for (int32 Index = 0; Index < PreparedComponents.Num(); ++Index)
+	{
+		FPreparedData& ComponentData = PreparedComponents[Index];
+		if (ComponentData.IsValid())
+		{
+			ComponentData.Derivative = Derivative;
+		}
+	}
+}
+
 void FPreparedType::SetField(const Shader::FStructField* Field, const FPreparedType& FieldType)
 {
 	for (int32 Index = 0; Index < Field->GetNumComponents(); ++Index)
diff --git a/Engine/Source/Runtime/Engine/Private/HLSLTree/HLSLTreeCommon.cpp b/Engine/Source/Runtime/Engine/Private/HLSLTree/HLSLTreeCommon.cpp
--- a/Engine/Source/Runtime/Engine/Private/HLSLTree/HLSLTreeCommon.cpp
+++ b/Engine/Source/Runtime/Engine/Private/HLSLTree/HLSLTreeCommon.cpp
@@ -584,6 +584,7 @@ void FExpressionBinaryOp::PrepareValue(FEmitContext& Context, const FRequestedTy
 	if (Op == EBinaryOp::Less)
 	{
 		ResultType.ValueComponentType = Shader::EValueComponentType::Bool;
+		ResultType.SetDerivative(EExpressionDerivative::Zero); // Comparisons result in 'zero' derivative...is this correct?
 	}
 
 	const FBinaryOpDescription OpDesc = GetBinaryOpDesription(Op);
diff --git a/Engine/Source/Runtime/Engine/Private/HLSLTree/HLSLTreeIntrinsics.cpp b/Engine/Source/Runtime/Engine/Private/HLSLTree/HLSLTreeIntrinsics.cpp
--- a/Engine/Source/Runtime/Engine/Private/HLSLTree/HLSLTreeIntrinsics.cpp
+++ b/Engine/Source/Runtime/Engine/Private/HLSLTree/HLSLTreeIntrinsics.cpp
@@ -33,6 +33,7 @@ FUnaryOpDescription GetUnaryOpDesription(EUnaryOp Op)
 	switch (Op)
 	{
 	case EUnaryOp::None: return FUnaryOpDescription(TEXT("None"), TEXT(""), Shader::EPreshaderOpcode::Nop); break;
+	case EUnaryOp::Neg: return FUnaryOpDescription(TEXT("Neg"), TEXT("-"), Shader::EPreshaderOpcode::Neg); break;
 	case EUnaryOp::Rcp: return FUnaryOpDescription(TEXT("Rcp"), TEXT("/"), Shader::EPreshaderOpcode::Rcp); break;
 	default: checkNoEntry(); return FUnaryOpDescription();
 	}
@@ -204,6 +205,16 @@ FEmitShaderCode* FEmitContext::EmitUnaryOp(EUnaryOp Op, FEmitShaderCode* Input)
 	FEmitShaderCode* Result = nullptr;
 	switch (Op)
 	{
+	case EUnaryOp::Neg:
+		if (InputComponentType == Shader::EValueComponentType::Double)
+		{
+			Result = EmitCode(ResultType, TEXT("LWCNegate(%)"), Input);
+		}
+		else
+		{
+			Result = EmitInlineCode(ResultType, TEXT("(-%)"), Input);
+		}
+		break;
 	case EUnaryOp::Rcp:
 		if (InputComponentType == Shader::EValueComponentType::Double)
 		{
@@ -217,6 +228,7 @@ FEmitShaderCode* FEmitContext::EmitUnaryOp(EUnaryOp Op, FEmitShaderCode* Input)
 		break;
 	default:
 		checkNoEntry();
+		break;
 	}
 	return Result;
 }
@@ -227,9 +239,24 @@ FEmitShaderValues FEmitContext::EmitUnaryOp(EUnaryOp Op, FEmitShaderValues Input
 	Result.Code = EmitUnaryOp(Op, Input.Code);
 	if (Derivative == EExpressionDerivative::Valid && Input.HasDerivatives())
 	{
-		
+		switch (Op)
+		{
+		case EUnaryOp::Neg:
+			Result.CodeDdx = EmitNeg(Input.CodeDdx);
+			Result.CodeDdy = EmitNeg(Input.CodeDdy);
+			break;
+		case EUnaryOp::Rcp:
+		{
+			FEmitShaderCode* dFdA = EmitNeg(EmitMul(Result.Code, Result.Code));
+			Result.CodeDdx = EmitMul(dFdA, Input.CodeDdx);
+			Result.CodeDdy = EmitMul(dFdA, Input.CodeDdy);
+			break;
+		}
+		default:
+			checkNoEntry();
+			break;
+		}
 	}
-
 	return Result;
 }
 
@@ -313,6 +340,9 @@ FEmitShaderValues FEmitContext::EmitBinaryOp(EBinaryOp Op, FEmitShaderValues Lhs
 	Result.Code = EmitBinaryOp(Op, Lhs.Code, Rhs.Code);
 	if (Derivative == EExpressionDerivative::Valid && Lhs.HasDerivatives() && Rhs.HasDerivatives())
 	{
+		const Shader::FValueTypeDescription ResultTypeDesc = Shader::GetValueTypeDescription(Result.Code->Type);
+		const Shader::EValueType DerivativeType = Shader::MakeValueType(Shader::EValueComponentType::Float, ResultTypeDesc.NumComponents);
+
 		switch (Op)
 		{
 		case EBinaryOp::Add:
@@ -327,6 +357,18 @@ FEmitShaderValues FEmitContext::EmitBinaryOp(EBinaryOp Op, FEmitShaderValues Lhs
 			Result.CodeDdx = EmitAdd(EmitMul(Lhs.CodeDdx, Rhs.Code), EmitMul(Rhs.CodeDdx, Lhs.Code));
 			Result.CodeDdy = EmitAdd(EmitMul(Lhs.CodeDdy, Rhs.Code), EmitMul(Rhs.CodeDdy, Lhs.Code));
 			break;
+		case EBinaryOp::Div:
+		{
+			FEmitShaderCode* Denom = EmitRcp(EmitMul(Rhs.Code, Rhs.Code));
+			FEmitShaderCode* dFdA = EmitMul(Rhs.Code, Denom);
+			FEmitShaderCode* dFdB = EmitNeg(EmitMul(Lhs.Code, Denom));
+			Result.CodeDdx = EmitAdd(EmitMul(dFdA, Lhs.CodeDdx), EmitMul(dFdB, Rhs.CodeDdx));
+			Result.CodeDdy = EmitAdd(EmitMul(dFdA, Lhs.CodeDdy), EmitMul(dFdB, Rhs.CodeDdy));
+			break;
+		}
+		default:
+			checkNoEntry();
+			break;
 		}
 	}
 
diff --git a/Engine/Source/Runtime/Engine/Private/Materials/MaterialUniformExpressions.cpp b/Engine/Source/Runtime/Engine/Private/Materials/MaterialUniformExpressions.cpp
--- a/Engine/Source/Runtime/Engine/Private/Materials/MaterialUniformExpressions.cpp
+++ b/Engine/Source/Runtime/Engine/Private/Materials/MaterialUniformExpressions.cpp
@@ -730,11 +730,14 @@ void FUniformExpressionSet::FillUniformBuffer(const FMaterialRenderContext& Mate
 				const FValueTypeDescription UniformTypeDesc = GetValueTypeDescription(PreshaderField.Type);
 				const int32 NumFieldComponents = UniformTypeDesc.NumComponents;
 
+				const int32 NumResultComponents = FMath::Min<int32>(NumFieldComponents, Result.Component.Num() - PreshaderField.ComponentIndex);
+				check(NumResultComponents > 0);
+
 				// The type generated by the preshader might not match the expected type
 				// In the future, with new HLSLTree, preshader could potentially include explicit cast opcodes, and avoid implicit conversions
 				// Making this change is difficult while also maintaining backwards compatibility however
-				FValue FieldValue(Result.Type.GetComponentType(PreshaderField.ComponentIndex), NumFieldComponents);
-				for (int32 i = 0; i < NumFieldComponents; ++i)
+				FValue FieldValue(Result.Type.GetComponentType(PreshaderField.ComponentIndex), NumResultComponents);
+				for (int32 i = 0; i < NumResultComponents; ++i)
 				{
 					FieldValue.Component[i] = Result.Component[PreshaderField.ComponentIndex + i];
 				}
diff --git a/Engine/Source/Runtime/Engine/Private/Shader/Preshader.cpp b/Engine/Source/Runtime/Engine/Private/Shader/Preshader.cpp
--- a/Engine/Source/Runtime/Engine/Private/Shader/Preshader.cpp
+++ b/Engine/Source/Runtime/Engine/Private/Shader/Preshader.cpp
@@ -514,6 +514,7 @@ FPreshaderValue EvaluatePreshader(const FUniformExpressionSet* UniformExpression
 		case EPreshaderOpcode::Clamp: EvaluateTernaryOp(Stack, Clamp); break;
 		case EPreshaderOpcode::Dot: EvaluateBinaryOp(Stack, Dot); break;
 		case EPreshaderOpcode::Cross: EvaluateBinaryOp(Stack, Cross); break;
+		case EPreshaderOpcode::Neg: EvaluateUnaryOp(Stack, Neg); break;
 		case EPreshaderOpcode::Sqrt: EvaluateUnaryOp(Stack, Sqrt); break;
 		case EPreshaderOpcode::Rcp: EvaluateUnaryOp(Stack, Rcp); break;
 		case EPreshaderOpcode::Length: EvaluateUnaryOp(Stack, [](const FValue& Value) { return Sqrt(Dot(Value, Value)); }); break;
diff --git a/Engine/Source/Runtime/Engine/Private/Shader/ShaderValue.cpp b/Engine/Source/Runtime/Engine/Private/Shader/ShaderValue.cpp
--- a/Engine/Source/Runtime/Engine/Private/Shader/ShaderValue.cpp
+++ b/Engine/Source/Runtime/Engine/Private/Shader/ShaderValue.cpp
@@ -894,6 +894,7 @@ struct FOpBaseNoInt : public FOpBase
 	static constexpr bool SupportsInt = false;
 };
 
+struct FOpNeg : public FOpBase { template<typename T> T operator()(T Value) const { return -Value; } };
 struct FOpAbs : public FOpBase { template<typename T> T operator()(T Value) const { return FMath::Abs(Value); } };
 struct FOpSign : public FOpBase { template<typename T> T operator()(T Value) const { return FMath::Sign(Value); } };
 struct FOpSaturate : public FOpBaseNoInt { template<typename T> T operator()(T Value) const { return FMath::Clamp(Value, (T)0, (T)1); } };
@@ -1086,6 +1087,11 @@ uint32 GetTypeHash(const FValue& Value)
 	return Result;
 }
 
+FValue Neg(const FValue& Value)
+{
+	return Private::UnaryOp(Private::FOpNeg(), Value);
+}
+
 FValue Abs(const FValue& Value)
 {
 	return Private::UnaryOp(Private::FOpAbs(), Value);
diff --git a/Engine/Source/Runtime/Engine/Public/HLSLTree/HLSLTree.h b/Engine/Source/Runtime/Engine/Public/HLSLTree/HLSLTree.h
--- a/Engine/Source/Runtime/Engine/Public/HLSLTree/HLSLTree.h
+++ b/Engine/Source/Runtime/Engine/Public/HLSLTree/HLSLTree.h
@@ -186,6 +186,9 @@ public:
 	FEmitShaderCode* EmitBinaryOp(EBinaryOp Op, FEmitShaderCode* Lhs, FEmitShaderCode* Rhs);
 	FEmitShaderValues EmitBinaryOp(EBinaryOp Op, FEmitShaderValues Lhs, FEmitShaderValues Rhs, EExpressionDerivative Derivative);
 
+	template<typename T> inline T EmitNeg(T Input) { return EmitUnaryOp(EUnaryOp::Neg, Input); }
+	template<typename T> inline T EmitRcp(T Input) { return EmitUnaryOp(EUnaryOp::Rcp, Input); }
+
 	template<typename T> inline T EmitAdd(T Lhs, T Rhs) { return EmitBinaryOp(EBinaryOp::Add, Lhs, Rhs); }
 	template<typename T> inline T EmitSub(T Lhs, T Rhs) { return EmitBinaryOp(EBinaryOp::Sub, Lhs, Rhs); }
 	template<typename T> inline T EmitMul(T Lhs, T Rhs) { return EmitBinaryOp(EBinaryOp::Mul, Lhs, Rhs); }
@@ -362,6 +365,7 @@ public:
 	FPreparedType(const Shader::FType& InType);
 
 	void SetEvaluation(EExpressionEvaluation Evaluation);
+	void SetDerivative(EExpressionDerivative Derivative);
 
 	void SetField(const Shader::FStructField* Field, const FPreparedType& FieldType);
 	FPreparedType GetFieldType(const Shader::FStructField* Field) const;
diff --git a/Engine/Source/Runtime/Engine/Public/HLSLTree/HLSLTreeTypes.h b/Engine/Source/Runtime/Engine/Public/HLSLTree/HLSLTreeTypes.h
--- a/Engine/Source/Runtime/Engine/Public/HLSLTree/HLSLTreeTypes.h
+++ b/Engine/Source/Runtime/Engine/Public/HLSLTree/HLSLTreeTypes.h
@@ -65,6 +65,7 @@ EExpressionDerivative CombineDerivatives(EExpressionDerivative Lhs, EExpressionD
 enum class EUnaryOp : uint8
 {
 	None,
+	Neg,
 	Rcp,
 };
 
diff --git a/Engine/Source/Runtime/Engine/Public/Shader/Preshader.h b/Engine/Source/Runtime/Engine/Public/Shader/Preshader.h
--- a/Engine/Source/Runtime/Engine/Public/Shader/Preshader.h
+++ b/Engine/Source/Runtime/Engine/Public/Shader/Preshader.h
@@ -61,6 +61,7 @@ enum class EPreshaderOpcode : uint8
 	RuntimeVirtualTextureUniform,
 	GetField,
 	SetField,
+	Neg,
 };
 
 struct FPreshaderStructType
diff --git a/Engine/Source/Runtime/Engine/Public/Shader/ShaderTypes.h b/Engine/Source/Runtime/Engine/Public/Shader/ShaderTypes.h
--- a/Engine/Source/Runtime/Engine/Public/Shader/ShaderTypes.h
+++ b/Engine/Source/Runtime/Engine/Public/Shader/ShaderTypes.h
@@ -400,6 +400,7 @@ inline bool operator!=(const FValue& Lhs, const FValue& Rhs) { return !operator=
 ENGINE_API uint32 GetTypeHash(const FType& Type);
 ENGINE_API uint32 GetTypeHash(const FValue& Value);
 
+ENGINE_API FValue Neg(const FValue& Value);
 ENGINE_API FValue Abs(const FValue& Value);
 ENGINE_API FValue Saturate(const FValue& Value);
 ENGINE_API FValue Floor(const FValue& Value);
