Lowered r.DistanceFields.MaxPerMeshResolution to 256

diff --git a/Engine/Source/Runtime/Engine/Private/DistanceFieldAtlas.cpp b/Engine/Source/Runtime/Engine/Private/DistanceFieldAtlas.cpp
--- a/Engine/Source/Runtime/Engine/Private/DistanceFieldAtlas.cpp
+++ b/Engine/Source/Runtime/Engine/Private/DistanceFieldAtlas.cpp
@@ -52,7 +52,7 @@ namespace DistanceFieldCookStats
 static TAutoConsoleVariable<int32> CVarDistField(
 	TEXT("r.GenerateMeshDistanceFields"),
 	0,	
-	TEXT("Whether to build distance fields of static meshes, needed for distance field AO, which is used to implement Movable SkyLight shadows.\n")
+	TEXT("Whether to build distance fields of static meshes, needed for Lumen Software Ray Tracing and Distance Field AO, which is used to implement Movable SkyLight shadows.\n")
 	TEXT("Enabling will increase mesh build times and memory usage.  Changing this value will cause a rebuild of all static meshes."),
 	ECVF_ReadOnly);
 
@@ -65,7 +65,7 @@ static TAutoConsoleVariable<int32> CVarDistFieldSupportWhenHardwareRayTracingEna
 
 static TAutoConsoleVariable<int32> CVarDistFieldRes(
 	TEXT("r.DistanceFields.MaxPerMeshResolution"),
-	512,	
+	256,	
 	TEXT("Highest resolution (in one dimension) allowed for a single static mesh asset, used to cap the memory usage of meshes with a large scale.\n")
 	TEXT("Changing this will cause all distance fields to be rebuilt.  Large values such as 512 can consume memory very quickly! (64Mb for one asset at 512)"),
 	ECVF_ReadOnly);
