Fix for crash while serializing morph targets on a editor platform that doesn't itself support morph targets.

diff --git a/Engine/Source/Runtime/Engine/Private/SkeletalMeshLODRenderData.cpp b/Engine/Source/Runtime/Engine/Private/SkeletalMeshLODRenderData.cpp
--- a/Engine/Source/Runtime/Engine/Private/SkeletalMeshLODRenderData.cpp
+++ b/Engine/Source/Runtime/Engine/Private/SkeletalMeshLODRenderData.cpp
@@ -653,7 +653,7 @@ void FSkeletalMeshLODRenderData::SerializeStreamedData(FArchive& Ar, USkeletalMe
 
 	// Determine if morph target data should be cooked out for this platform
 	EShaderPlatform MorphTargetShaderPlatform = GMaxRHIShaderPlatform;
-	bool bSerializeCompressedMorphTargets = Ar.IsSaving() && Owner->GetMorphTargets().Num() > 0;
+	bool bSerializeCompressedMorphTargets = Ar.IsSaving() && FMorphTargetVertexInfoBuffers::IsPlatformShaderSupported(MorphTargetShaderPlatform) && Owner->GetMorphTargets().Num() > 0;
 	if (Ar.IsCooking())
 	{
 		const ITargetPlatform* Platform = Ar.CookingTarget();
@@ -696,12 +696,8 @@ void FSkeletalMeshLODRenderData::SerializeStreamedData(FArchive& Ar, USkeletalMe
 				// The CPU data could have already been destroyed by this point, which happens when the RHI is initialized.  If possible, use the MorphTargetVertexInfoBuffers.
 				if (!MorphTargetVertexInfoBuffers.IsMorphResourcesInitialized())
 				{
-					check(FMorphTargetVertexInfoBuffers::IsPlatformShaderSupported(MorphTargetShaderPlatform));
-					check(MorphTargets.Num() > 0);
-
 					TargetMorphTargetVertexInfoBuffers = &MorphTargetVertexInfoBuffers;
 					TargetMorphTargetVertexInfoBuffers->InitMorphResources(MorphTargetShaderPlatform, RenderSections, MorphTargets, StaticVertexBuffers.StaticMeshVertexBuffer.GetNumVertices(), LODIdx, SkeletalMeshLODInfo->MorphTargetPositionErrorTolerance);
-					check(TargetMorphTargetVertexInfoBuffers->IsMorphCPUDataValid());
 				}
 				else if (MorphTargetVertexInfoBuffers.IsMorphCPUDataValid() && !MorphTargetVertexInfoBuffers.IsRHIIntialized())
 				{
