Ensure each VT tile has a unique DebugTexturePathName

diff --git a/Engine/Source/Runtime/Engine/Private/VT/VirtualTextureDataBuilder.cpp b/Engine/Source/Runtime/Engine/Private/VT/VirtualTextureDataBuilder.cpp
--- a/Engine/Source/Runtime/Engine/Private/VT/VirtualTextureDataBuilder.cpp
+++ b/Engine/Source/Runtime/Engine/Private/VT/VirtualTextureDataBuilder.cpp
@@ -750,10 +750,13 @@ void FVirtualTextureDataBuilder::BuildTiles(const TArray<FVTSourceTileEntry>& Ti
 			}
 #endif // SAVE_TILES
 
+			// give each tile a unique DebugTexturePathName for DebugDump option :
+			FString DebugTilePathName = FString::Printf(TEXT("%s_VT%04d"), *DebugTexturePathName, TileIndex);
+
 			TArray<FCompressedImage2D> CompressedMip;
 			TArray<FImage> EmptyList;
 			uint32 NumMipsInTail, ExtData;
-			if (!ensure(Compressor->BuildTexture(TileImages, EmptyList, TBSettings, DebugTexturePathName, CompressedMip, NumMipsInTail, ExtData)))
+			if (!ensure(Compressor->BuildTexture(TileImages, EmptyList, TBSettings, DebugTilePathName, CompressedMip, NumMipsInTail, ExtData)))
 			{
 				bCompressionError = true;
 			}
