Fix PhysX compiler error

diff --git a/Engine/Source/Runtime/Engine/Private/PhysicsEngine/ConstraintInstance.cpp b/Engine/Source/Runtime/Engine/Private/PhysicsEngine/ConstraintInstance.cpp
--- a/Engine/Source/Runtime/Engine/Private/PhysicsEngine/ConstraintInstance.cpp
+++ b/Engine/Source/Runtime/Engine/Private/PhysicsEngine/ConstraintInstance.cpp
@@ -492,10 +492,12 @@ void FConstraintInstance::InitConstraint(FBodyInstance* Body1, FBodyInstance* Bo
 			return;
 		}
 
+#if WITH_CHAOS
 		if (!bAllowKinematicKinematicConstraints && (!FPhysicsInterface::IsValid(Actor1) || FPhysicsInterface::IsKinematic(Actor1)) && (!FPhysicsInterface::IsValid(Actor2) || FPhysicsInterface::IsKinematic(Actor2)))
 		{
 			return;
 		}
+#endif
 
 		FPhysicsCommand::ExecuteWrite(Actor1, Actor2, [&](const FPhysicsActorHandle& ActorA, const FPhysicsActorHandle& ActorB)
 		{
