Fix for distance attenuation stomp

diff --git a/Engine/Source/Runtime/Engine/Private/ActiveSound.cpp b/Engine/Source/Runtime/Engine/Private/ActiveSound.cpp
--- a/Engine/Source/Runtime/Engine/Private/ActiveSound.cpp
+++ b/Engine/Source/Runtime/Engine/Private/ActiveSound.cpp
@@ -1702,7 +1702,6 @@ void FActiveSound::UpdateAttenuation(float DeltaTime, FSoundParseParameters& Par
 			CheckOcclusion(ListenerPosition, ParseParams.Transform.GetTranslation(), Settings);
 
 			// Apply the volume attenuation due to occlusion (using the interpolating dynamic parameter)
-			ParseParams.DistanceAttenuation = CurrentOcclusionVolumeAttenuation.GetValue();
 			ParseParams.OcclusionAttenuation = CurrentOcclusionVolumeAttenuation.GetValue();
 			ParseParams.bIsOccluded = bIsOccluded;
 			ParseParams.OcclusionFilterFrequency = CurrentOcclusionFilterFrequency.GetValue();
